Merging USD changes from Dev-Main to Release-Staging-4.17. Accepting version from Main, since 4.17 changes seem to be all edigrates.

diff --git a/Engine/Plugins/Editor/USDImporter/Source/UnrealUSDWrapper/Source/Private/UnrealUSDWrapper.cpp b/Engine/Plugins/Editor/USDImporter/Source/UnrealUSDWrapper/Source/Private/UnrealUSDWrapper.cpp
--- a/Engine/Plugins/Editor/USDImporter/Source/UnrealUSDWrapper/Source/Private/UnrealUSDWrapper.cpp
+++ b/Engine/Plugins/Editor/USDImporter/Source/UnrealUSDWrapper/Source/Private/UnrealUSDWrapper.cpp
@@ -31,19 +31,37 @@
 using std::vector;
 using std::string;
 
+using namespace pxr;
+
 #ifdef _WINDOWS
 #pragma warning(disable:4267)
-#define PXR_NAMESPACE
+
 #define USDWRAPPER_USE_XFORMACHE	1
 #else
-#define PXR_NAMESPACE		pxr
 #define USDWRAPPER_USE_XFORMACHE	0
 #endif
 
+
 #if USDWRAPPER_USE_XFORMACHE
 static UsdGeomXformCache XFormCache;
 #endif // USDWRAPPER_USE_XFORMACHE
 
+namespace UnrealIdentifiers
+{
+	/**
+	* Identifies the LOD variant set on a primitive which means this primitive has child prims that LOD meshes
+	* named LOD0, LOD1, LOD2, etc
+	*/
+	static const TfToken LOD("LOD");
+
+	static const TfToken AssetPath("unrealAssetPath");
+
+	static const TfToken ActorClass("unrealActorClass");
+
+	static const TfToken PropertyPath("unrealPropertyPath");
+}
+
+
 void Log(const char* Format, ...)
 {
 	static char Buffer[2048];
@@ -68,13 +86,13 @@ class USDHelpers
 {
 
 public:
-	static void LogPrimTree(const PXR_NAMESPACE::UsdPrim& Root)
+	static void LogPrimTree(const UsdPrim& Root)
 	{
 		LogPrimTreeHelper("", Root);
 	}
 
 private:
-	static void LogPrimTreeHelper(const string& Concat, const PXR_NAMESPACE::UsdPrim& Prim)
+	static void LogPrimTreeHelper(const string& Concat, const UsdPrim& Prim)
 	{
 		string TypeName = Prim.GetTypeName().GetString();
 		bool bIsModel = Prim.IsModel();
@@ -87,7 +105,7 @@ private:
 		Log(string(Concat + "Prim: [%s] %s Model:%d Abstract:%d Group:%d Instance:%d Active:%d InMaster:%d IsMaster:%d\n").c_str(),
 			TypeName.c_str(), Prim.GetName().GetText(), bIsModel, bIsAbstract, bIsGroup, bIsInstance, bIsActive, bInMaster, bIsMaster);
 		{
-			/*UsdMetadataValueMap Metadata = Prim.GetAllMetadata();
+			UsdMetadataValueMap Metadata = Prim.GetAllMetadata();
 			if(Metadata.size())
 			{
 				Log(string(Concat+"\tMetaData:\n").c_str());
@@ -95,9 +113,11 @@ private:
 				{
 					Log(string(Concat+"\t\t[%s] %s\n").c_str(), KeyValue.second.GetTypeName().c_str(), KeyValue.first.GetText());
 				}
-			}*/
+			}
 
-			/*vector<UsdAttribute> Attributes = Prim.GetAttributes();
+		
+
+			vector<UsdAttribute> Attributes = Prim.GetAttributes();
 			if(Attributes.size())
 			{
 				Log(string(Concat+"\tAttributes:\n").c_str());
@@ -108,9 +128,9 @@ private:
 						Log(string(Concat + "\t\t[%s] %s %s\n").c_str(), Attribute.GetTypeName().GetAsToken().GetText(), Attribute.GetBaseName().GetText(), Attribute.GetDisplayName().c_str());
 					}
 				}
-			}*/
+			}
 
-			/*if (Prim.HasVariantSets())
+			if (Prim.HasVariantSets())
 			{
 				Log(string(Concat + "\tVariant Sets:\n").c_str());
 				UsdVariantSets VariantSets = Prim.GetVariantSets();
@@ -132,11 +152,11 @@ private:
 						Log(string(Concat + "\t\t\t%s%c\n").c_str(), VariantName.c_str(), ActiveChar);
 					}
 				}
-			}*/
+			}
 		}
 
 
-		for(const PXR_NAMESPACE::UsdPrim& Child : Prim.GetChildren())
+		for(const UsdPrim& Child : Prim.GetChildren())
 		{
 			LogPrimTreeHelper(Concat+"\t", Child);
 		}
@@ -152,33 +172,311 @@ void InitWrapper()
 
 struct FPrimAndData
 {
-	PXR_NAMESPACE::UsdPrim Prim;
+	UsdPrim Prim;
 	class FUsdPrim* PrimData;
 
-	FPrimAndData(const PXR_NAMESPACE::UsdPrim& InPrim)
+	FPrimAndData(const UsdPrim& InPrim)
 		: Prim(InPrim)
 		, PrimData(nullptr)
 	{}
 };
 
+
+class FAttribInternalData
+{
+public:
+	FAttribInternalData(UsdAttribute& InAttribute)
+		: Attribute(InAttribute)
+	{
+		VtValue CustomData = Attribute.GetCustomDataByKey(UnrealIdentifiers::PropertyPath);
+
+		AttributeName = Attribute.GetBaseName().GetString();
+		TypeName = Attribute.GetTypeName().GetAsToken().GetString();
+
+		if (CustomData.IsHolding<std::string>())
+		{
+			UnrealPropertyPath = CustomData.Get<std::string>();
+		}
+	}
+
+	std::string UnrealPropertyPath;
+	std::string AttributeName;
+	std::string TypeName;
+	UsdAttribute Attribute;
+};
+
+
+FUsdAttribute::FUsdAttribute(std::shared_ptr<FAttribInternalData> InInternalData)
+	: InternalData(InInternalData)
+{
+}
+
+FUsdAttribute::~FUsdAttribute()
+{
+}
+
+const char* FUsdAttribute::GetAttributeName() const
+{
+	return InternalData->AttributeName.c_str();
+}
+
+const char* FUsdAttribute::GetTypeName() const
+{
+	return InternalData->TypeName.c_str();
+}
+
+const char* FUsdAttribute::GetUnrealPropertyPath() const
+{
+	return InternalData->UnrealPropertyPath.c_str();
+}
+
+template<typename T>
+bool GetValue(T& OutVal, UsdAttribute Attrib, int ArrayIndex, double Time)
+{
+	bool bResult = false;
+
+	if (ArrayIndex != -1)
+	{
+		// Note: VtArray is copy on write so this is cheap
+		VtArray<T> Array;
+		if (Attrib.Get(&Array, Time))
+		{
+			OutVal = Array[ArrayIndex];
+			bResult = true;
+		}
+	}
+	else
+	{
+		bResult = Attrib.Get(&OutVal, Time);
+	}
+
+	return bResult;
+}
+
+template<typename T>
+bool IsHolding(const VtValue& Value)
+{
+	return Value.IsHolding<T>() || Value.IsHolding<VtArray<T>>();
+}
+
+bool FUsdAttribute::AsInt(int64_t& OutVal, int ArrayIndex, double Time) const
+{
+	// We test multiple types of ints here. int64 is always returned as it can hold all other types
+	// Unreal expects this
+	VtValue Value;
+	bool bResult = InternalData->Attribute.Get(&Value, Time);
+	if (IsHolding<int8_t>(Value))
+	{
+		uint8_t Val;
+		bResult = GetValue(Val, InternalData->Attribute, ArrayIndex, Time);
+		OutVal = Val;
+	}
+	else if (IsHolding<int32_t>(Value))
+	{
+		int32_t Val;
+		bResult = GetValue(Val, InternalData->Attribute, ArrayIndex, Time);
+		OutVal = Val;
+	}
+	else if (IsHolding<int64_t>(Value))
+	{
+		int64_t Val;
+		bResult = GetValue(Val, InternalData->Attribute, ArrayIndex, Time);
+		OutVal = Val;
+	}
+
+	return bResult;
+}
+
+bool FUsdAttribute::AsUnsignedInt(uint64_t& OutVal, int ArrayIndex, double Time) const
+{
+	// We test multiple types of ints here. uint64 is always returned as it can hold all other types
+	// Unreal expects this
+	VtValue Value;
+	bool bResult = InternalData->Attribute.Get(&Value, Time);
+	if (IsHolding<uint8_t>(Value))
+	{
+		uint8_t Val;
+		bResult = GetValue(Val, InternalData->Attribute, ArrayIndex, Time);
+		OutVal = Val;
+	}
+	else if (IsHolding<uint32_t>(Value))
+	{
+		uint32_t Val;
+		bResult = GetValue(Val, InternalData->Attribute, ArrayIndex, Time);
+		OutVal = Val;
+	}
+	else if (IsHolding<uint64_t>(Value))
+	{
+		uint64_t Val;
+		bResult = GetValue(Val, InternalData->Attribute, ArrayIndex, Time);
+		OutVal = Val;
+	}
+
+	return bResult;
+}
+
+bool FUsdAttribute::AsDouble(double& OutVal, int ArrayIndex, double Time) const
+{
+	bool bResult = false;
+
+
+	bResult = GetValue<double>(OutVal, InternalData->Attribute, ArrayIndex, Time);
+
+	if (!bResult)
+	{
+		float Val = 0.0f;
+		bResult = GetValue<float>(Val, InternalData->Attribute, ArrayIndex, Time);
+		OutVal = Val;
+	}
+
+	return bResult;
+}
+
+bool FUsdAttribute::AsString(const char*& OutVal, int ArrayIndex, double Time) const
+{
+	// this method is very hacky to return temp strings
+	// designed to have the string copied immediately
+	bool bResult = false;
+
+	VtValue Value;
+	InternalData->Attribute.Get(&Value);
+	// mem leak
+	static std::string Temp;
+	if (IsHolding<std::string>(Value))
+	{
+		bResult = GetValue(Temp, InternalData->Attribute, ArrayIndex, Time);
+
+		OutVal = Temp.c_str();
+	}
+	else if (IsHolding<TfToken>(Value))
+	{
+		TfToken Token;
+		bResult = GetValue(Token, InternalData->Attribute, ArrayIndex, Time);
+
+		Temp = Token.GetString();
+
+		OutVal = Temp.c_str();
+	}
+
+	return bResult;
+}
+
+bool FUsdAttribute::AsBool(bool& OutVal, int ArrayIndex, double Time) const
+{
+	return GetValue(OutVal, InternalData->Attribute, ArrayIndex, Time);
+}
+
+bool FUsdAttribute::AsVector2(FUsdVector2Data& OutVal, int ArrayIndex, double Time) const
+{
+	GfVec2f Value;
+	const bool bResult = GetValue(Value, InternalData->Attribute, ArrayIndex, Time);
+
+	OutVal.X = Value[0];
+	OutVal.Y = Value[1];
+
+	return bResult;
+}
+
+bool FUsdAttribute::AsVector3(FUsdVectorData& OutVal, int ArrayIndex, double Time) const
+{
+	GfVec3f Value;
+	const bool bResult = GetValue(Value, InternalData->Attribute, ArrayIndex, Time);
+
+	OutVal.X = Value[0];
+	OutVal.Y = Value[1];
+	OutVal.Z = Value[2];
+
+	return bResult;
+}
+
+bool FUsdAttribute::AsVector4(FUsdVector4Data& OutVal, int ArrayIndex, double Time) const
+{
+	GfVec4f Value;
+	const bool bResult = GetValue(Value, InternalData->Attribute, ArrayIndex, Time);
+
+	OutVal.X = Value[0];
+	OutVal.Y = Value[1];
+	OutVal.Z = Value[2];
+	OutVal.W = Value[3];
+
+	return bResult;
+}
+
+bool FUsdAttribute::AsColor(FUsdVector4Data& OutVal, int ArrayIndex, double Time) const
+{
+	GfVec4f Value;
+	bool bResult = GetValue(Value, InternalData->Attribute, ArrayIndex, Time);
+
+	if (bResult)
+	{
+		OutVal.X = Value[0];
+		OutVal.Y = Value[1];
+		OutVal.Z = Value[2];
+		OutVal.W = Value[3];
+	}
+	else
+	{
+		// Try color 3 with a = 1;
+		GfVec3f Value;
+		bResult = GetValue<GfVec3f>(Value, InternalData->Attribute, ArrayIndex, Time);
+		OutVal.X = Value[0];
+		OutVal.Y = Value[1];
+		OutVal.Z = Value[2];
+		OutVal.W = 1;
+	}
+
+	return bResult;
+}
+
+bool FUsdAttribute::IsUnsigned() const
+{
+	VtValue Value;
+	InternalData->Attribute.Get(&Value);
+
+	return IsHolding<uint8_t>(Value)
+		|| IsHolding<uint32_t>(Value)
+		|| IsHolding<uint64_t>(Value);
+}
+
+int FUsdAttribute::GetArraySize() const
+{
+	VtValue Value;
+	InternalData->Attribute.Get(&Value);
+
+	return Value.IsArrayValued() ? (int)Value.GetArraySize() : -1;
+}
+
+
 class FUsdPrim : public IUsdPrim
 {
 public:
-	FUsdPrim(const PXR_NAMESPACE::UsdPrim& InPrim)
+	FUsdPrim(const UsdPrim& InPrim)
 		: Prim(InPrim)
 		, GeomData(nullptr)
 	{
 		PrimName = Prim.GetName().GetString();
 		PrimPath = Prim.GetPath().GetString();
 
-		static const PXR_NAMESPACE::TfToken AssetPathToken = PXR_NAMESPACE::TfToken(UnrealIdentifiers::AssetPath);
-		PXR_NAMESPACE::UsdAttribute UnrealAssetPathAttr = Prim.GetAttribute(AssetPathToken);
+		UsdAttribute UnrealAssetPathAttr = Prim.GetAttribute(UnrealIdentifiers::AssetPath);
 		if (UnrealAssetPathAttr.HasValue())
 		{
 			UnrealAssetPathAttr.Get(&UnrealAssetPath);
 		}
 
-		for (const PXR_NAMESPACE::UsdPrim& Child : Prim.GetChildren())
+		UsdAttribute UnrealActorClassAttr = Prim.GetAttribute(UnrealIdentifiers::ActorClass);
+		if (UnrealActorClassAttr.HasValue())
+		{
+			UnrealActorClassAttr.Get(&UnrealActorClass);
+		}
+	
+		VtValue CustomData = Prim.GetCustomDataByKey(UnrealIdentifiers::PropertyPath);
+		if (CustomData.IsHolding<std::string>())
+		{
+			UnrealPropertyPath = CustomData.Get<std::string>();
+
+		}
+
+		for (const UsdPrim& Child : Prim.GetChildren())
 		{
 			Children.push_back(FPrimAndData(Child));
 		}
@@ -223,34 +521,43 @@ public:
 		return PrimPath.c_str();
 	}
 
+	virtual const char* GetUnrealPropertyPath() const override
+	{
+		return UnrealPropertyPath.c_str();
+	}
 
 	virtual bool IsGroup() const override
 	{
 		return Prim.IsGroup();
 	}
 
+	virtual bool IsUnrealProperty() const override
+	{
+		return Prim.HasCustomDataKey(UnrealIdentifiers::PropertyPath);
+	}
+
 	virtual bool HasTransform() const override
 	{
-		return PXR_NAMESPACE::UsdGeomXformable(Prim);
+		return UsdGeomXformable(Prim);
 	}
 
 	virtual FUsdMatrixData GetLocalToWorldTransform() const override
 	{
-		return GetLocalToWorldTransform(PXR_NAMESPACE::UsdTimeCode::Default().GetValue());
+		return GetLocalToWorldTransform(UsdTimeCode::Default().GetValue());
 	}
 
 #if !USDWRAPPER_USE_XFORMACHE
-	static PXR_NAMESPACE::GfMatrix4d GetLocalToWorldTransform(const PXR_NAMESPACE::UsdPrim& Prim, double Time, const PXR_NAMESPACE::SdfPath& AbsoluteRootPath)
+	static GfMatrix4d GetLocalToWorldTransform(const UsdPrim& Prim, double Time, const SdfPath& AbsoluteRootPath)
 	{
-		PXR_NAMESPACE::SdfPath PrimPath = Prim.GetPath();
+		SdfPath PrimPath = Prim.GetPath();
 		if (!Prim || PrimPath == AbsoluteRootPath)
 		{
-			return PXR_NAMESPACE::GfMatrix4d(1);
+			return GfMatrix4d(1);
 		}
 
-		PXR_NAMESPACE::GfMatrix4d AccumulatedTransform(1.);
+		GfMatrix4d AccumulatedTransform(1.);
 		bool bResetsXFormStack = false;
-		PXR_NAMESPACE::UsdGeomXformable XFormable(Prim);
+		UsdGeomXformable XFormable(Prim);
 		// silently ignoring errors
 		XFormable.GetLocalTransformation(&AccumulatedTransform, &bResetsXFormStack, Time);
 
@@ -267,14 +574,14 @@ public:
 	{
 #if USDWRAPPER_USE_XFORMACHE
 		XFormCache.SetTime(Time);
-		PXR_NAMESPACE::GfMatrix4d LocalToWorld = XFormCache.GetLocalToWorldTransform(Prim);
+		GfMatrix4d LocalToWorld = XFormCache.GetLocalToWorldTransform(Prim);
 #else
-		PXR_NAMESPACE::UsdGeomXformable XFormable(Prim);
-		PXR_NAMESPACE::SdfPath WorldPath = PXR_NAMESPACE::SdfPath::AbsoluteRootPath();
-		PXR_NAMESPACE::GfMatrix4d LocalToWorld;
+		UsdGeomXformable XFormable(Prim);
+		SdfPath WorldPath = SdfPath::AbsoluteRootPath();
+		GfMatrix4d LocalToWorld;
 		if (!Prim.GetPath().HasPrefix(WorldPath))
 		{
-			LocalToWorld = PXR_NAMESPACE::GfMatrix4d(1);
+			LocalToWorld = GfMatrix4d(1);
 		}
 		else
 		{
@@ -291,7 +598,7 @@ public:
 
 	virtual FUsdMatrixData GetLocalToParentTransform() const override
 	{
-		return GetLocalToParentTransform(PXR_NAMESPACE::UsdTimeCode::Default().GetValue());
+		return GetLocalToParentTransform(UsdTimeCode::Default().GetValue());
 	}
 
 	virtual FUsdMatrixData GetLocalToParentTransform(double Time) const override
@@ -299,10 +606,10 @@ public:
 		bool bResetsXFormStack = false;
 #if USDWRAPPER_USE_XFORMACHE
 		XFormCache.SetTime(Time);
-		PXR_NAMESPACE::GfMatrix4d LocalToParent = XFormCache.GetLocalTransformation(Prim, &bResetsXFormStack);
+		GfMatrix4d LocalToParent = XFormCache.GetLocalTransformation(Prim, &bResetsXFormStack);
 #else
-		PXR_NAMESPACE::UsdGeomXformable XFormable(Prim);
-		PXR_NAMESPACE::GfMatrix4d LocalToParent;
+		UsdGeomXformable XFormable(Prim);
+		GfMatrix4d LocalToParent;
 		// silently ignoring errors
 		XFormable.GetLocalTransformation(&LocalToParent, &bResetsXFormStack, Time);
 #endif
@@ -314,7 +621,7 @@ public:
 
 	virtual int GetNumChildren() const override
 	{
-		return Children.size();
+		return (int)Children.size();
 	}
 
 	virtual IUsdPrim* GetChild(int ChildIndex) override
@@ -333,22 +640,27 @@ public:
 		return UnrealAssetPath.length() > 0 ? UnrealAssetPath.c_str() : nullptr;
 	}
 
-	bool HasGeometryData() const override
+	virtual const char* GetUnrealActorClass() const override
 	{
-		PXR_NAMESPACE::UsdGeomMesh Mesh(Prim);
+		return UnrealActorClass.length() > 0 ? UnrealActorClass.c_str() : nullptr;
+	}
+
+	virtual bool HasGeometryData() const override
+	{
+		UsdGeomMesh Mesh(Prim);
 
 		return Mesh;
 	}
 
 	virtual const FUsdGeomData* GetGeometryData() override
 	{
-		return GetGeometryData(PXR_NAMESPACE::UsdTimeCode::Default().GetValue());
+		return GetGeometryData(UsdTimeCode::Default().GetValue());
 	}
 
 	virtual const FUsdGeomData* GetGeometryData(double Time) override
 	{
 		bool bValid = false;
-		PXR_NAMESPACE::UsdGeomMesh Mesh(Prim);
+		UsdGeomMesh Mesh(Prim);
 		if (Mesh)
 		{
 			if (GeomData)
@@ -360,78 +672,78 @@ public:
 
 			// Faces and points
 			{
-				PXR_NAMESPACE::UsdAttribute FaceCounts = Mesh.GetFaceVertexCountsAttr();
+				UsdAttribute FaceCounts = Mesh.GetFaceVertexCountsAttr();
 				if (FaceCounts)
 				{
 					GeomData->FaceVertexCounts.clear();
 
-					PXR_NAMESPACE::VtArray<int> FaceCountArray;
+					VtArray<int> FaceCountArray;
 					FaceCounts.Get(&FaceCountArray, Time);
 					GeomData->FaceVertexCounts.assign(FaceCountArray.begin(), FaceCountArray.end());
 				}
 
-				PXR_NAMESPACE::UsdAttribute FaceIndices = Mesh.GetFaceVertexIndicesAttr();
+				UsdAttribute FaceIndices = Mesh.GetFaceVertexIndicesAttr();
 				if (FaceIndices)
 				{
 					GeomData->FaceIndices.clear();
 
-					PXR_NAMESPACE::VtArray<int> FaceIndicesArray;
+					VtArray<int> FaceIndicesArray;
 					FaceIndices.Get(&FaceIndicesArray, Time);
 					GeomData->FaceIndices.assign(FaceIndicesArray.begin(), FaceIndicesArray.end());
 				}
 
-				PXR_NAMESPACE::UsdAttribute Points = Mesh.GetPointsAttr();
+				UsdAttribute Points = Mesh.GetPointsAttr();
 				if (Points)
 				{
 					GeomData->Points.clear();
 
-					PXR_NAMESPACE::VtArray<PXR_NAMESPACE::GfVec3f> PointsArray;
+					VtArray<GfVec3f> PointsArray;
 					Points.Get(&PointsArray, Time);
 
 					// Bug??  Usd returns nothing for UsdTimeCode::Default if there are animated points
-					if (PointsArray.size() == 0 && Time == PXR_NAMESPACE::UsdTimeCode::Default())
+					if (PointsArray.size() == 0 && Time == UsdTimeCode::Default())
 					{
 						// Try to get at time = 0
-						Points.Get(&PointsArray, PXR_NAMESPACE::UsdTimeCode(0));
+						Points.Get(&PointsArray, UsdTimeCode(0));
 					}
 					
 					GeomData->Points.resize(PointsArray.size());
-					memcpy(&GeomData->Points[0], &PointsArray[0], PointsArray.size() * sizeof(PXR_NAMESPACE::GfVec3f));
+					memcpy(&GeomData->Points[0], &PointsArray[0], PointsArray.size() * sizeof(GfVec3f));
 				}
 
-				PXR_NAMESPACE::UsdAttribute Normals = Mesh.GetNormalsAttr();
+				UsdAttribute Normals = Mesh.GetNormalsAttr();
 				if (Normals)
 				{
 					GeomData->Normals.clear();
 
-					PXR_NAMESPACE::VtArray<PXR_NAMESPACE::GfVec3f> NormalsArray;
+					VtArray<GfVec3f> NormalsArray;
 					Normals.Get(&NormalsArray, Time);
 
 					GeomData->Normals.resize(NormalsArray.size());
-					memcpy(&GeomData->Normals[0], &NormalsArray[0], NormalsArray.size() * sizeof(PXR_NAMESPACE::GfVec3f));
+					memcpy(&GeomData->Normals[0], &NormalsArray[0], NormalsArray.size() * sizeof(GfVec3f));
 				} 
 
-				PXR_NAMESPACE::UsdGeomPrimvar DisplayColorPrimVar = Mesh.GetDisplayColorPrimvar();
+				UsdGeomPrimvar DisplayColorPrimVar = Mesh.GetDisplayColorPrimvar();
 				if (DisplayColorPrimVar)
 				{
 					GeomData->VertexColors.clear();
 
-					PXR_NAMESPACE::TfToken Interpolation = DisplayColorPrimVar.GetInterpolation();
+					TfToken Interpolation = DisplayColorPrimVar.GetInterpolation();
 
-					if (Interpolation == PXR_NAMESPACE::UsdGeomTokens->faceVarying || Interpolation == PXR_NAMESPACE::UsdGeomTokens->uniform)
+					if (Interpolation == UsdGeomTokens->faceVarying || Interpolation == UsdGeomTokens->uniform)
 					{
-						GeomData->VertexColorInterpMethod = Interpolation == PXR_NAMESPACE::UsdGeomTokens->faceVarying ? EUsdInterpolationMethod::FaceVarying : EUsdInterpolationMethod::Uniform;
-						PXR_NAMESPACE::VtArray<PXR_NAMESPACE::GfVec3f> DisplayColorArray;
+						GeomData->VertexColorInterpMethod = Interpolation == UsdGeomTokens->faceVarying ? EUsdInterpolationMethod::FaceVarying : EUsdInterpolationMethod::Uniform;
+						VtArray<GfVec3f> DisplayColorArray;
 						DisplayColorPrimVar.ComputeFlattened(&DisplayColorArray, Time);
 		
 						GeomData->VertexColors.resize(DisplayColorArray.size());
-						memcpy(&GeomData->VertexColors[0], &DisplayColorArray[0], DisplayColorArray.size() * sizeof(PXR_NAMESPACE::GfVec3f));
+						memcpy(&GeomData->VertexColors[0], &DisplayColorArray[0], DisplayColorArray.size() * sizeof(GfVec3f));
 					}
-					else if (Interpolation == PXR_NAMESPACE::UsdGeomTokens->vertex)
+					else if (Interpolation == UsdGeomTokens->vertex)
 					{
 						GeomData->VertexColorInterpMethod = EUsdInterpolationMethod::Vertex;
-						PXR_NAMESPACE::VtIntArray VertexColorIndices;
-						PXR_NAMESPACE::VtArray<PXR_NAMESPACE::GfVec3f> DisplayColorArray;
+						VtIntArray VertexColorIndices;
+						VtArray<GfVec3f> DisplayColorArray;
 						DisplayColorPrimVar.GetIndices(&VertexColorIndices, Time);
 						DisplayColorPrimVar.Get(&DisplayColorArray, Time);
 
@@ -440,7 +752,7 @@ public:
 							GeomData->VertexColors.resize(VertexColorIndices.size());
 							for (int PointIdx = 0; PointIdx < GeomData->Points.size(); ++PointIdx)
 							{
-								PXR_NAMESPACE::GfVec3f Color = DisplayColorArray[VertexColorIndices[PointIdx]];
+								GfVec3f Color = DisplayColorArray[VertexColorIndices[PointIdx]];
 								GeomData->VertexColors[PointIdx] = FUsdVectorData(Color[0], Color[1], Color[2]);
 							}
 						}
@@ -448,30 +760,30 @@ public:
 						{
 							// Assume mapping is identical
 							GeomData->VertexColors.resize(DisplayColorArray.size());
-							memcpy(&GeomData->VertexColors[0], &DisplayColorArray[0], DisplayColorArray.size() * sizeof(PXR_NAMESPACE::GfVec3f));
+							memcpy(&GeomData->VertexColors[0], &DisplayColorArray[0], DisplayColorArray.size() * sizeof(GfVec3f));
 						}
 					}
-					else if (Interpolation == PXR_NAMESPACE::UsdGeomTokens->constant)
+					else if (Interpolation == UsdGeomTokens->constant)
 					{
-						PXR_NAMESPACE::VtArray<PXR_NAMESPACE::GfVec3f> DisplayColorArray;
+						VtArray<GfVec3f> DisplayColorArray;
 						DisplayColorPrimVar.Get(&DisplayColorArray, Time);
 
 						if(DisplayColorArray.size() > 0)
 						{
-							PXR_NAMESPACE::GfVec3f Color = DisplayColorArray[0];
+							GfVec3f Color = DisplayColorArray[0];
 							GeomData->VertexColors.push_back(FUsdVectorData(Color[0], Color[1], Color[2]));
 						}
 					}
 				}
 
 				GeomData->Orientation = EUsdGeomOrientation::RightHanded;
-				PXR_NAMESPACE::UsdAttribute Orientation = Mesh.GetOrientationAttr();
+				UsdAttribute Orientation = Mesh.GetOrientationAttr();
 				if(Orientation)
 				{ 
-					static PXR_NAMESPACE::TfToken RightHanded("rightHanded");
-					static PXR_NAMESPACE::TfToken LeftHanded("leftHanded");
+					static TfToken RightHanded("rightHanded");
+					static TfToken LeftHanded("leftHanded");
 
-					PXR_NAMESPACE::TfToken OrientationValue;
+					TfToken OrientationValue;
 					Orientation.Get(&OrientationValue, Time);
 
 					GeomData->Orientation = OrientationValue == RightHanded ? EUsdGeomOrientation::RightHanded : EUsdGeomOrientation::LeftHanded;
@@ -481,9 +793,9 @@ public:
 
 			// UVs
 			{
-				static PXR_NAMESPACE::TfToken UVSetName("primvars:st");
+				static TfToken UVSetName("primvars:st");
 
-				PXR_NAMESPACE::UsdGeomPrimvar STPrimvar = Mesh.GetPrimvar(UVSetName);
+				UsdGeomPrimvar STPrimvar = Mesh.GetPrimvar(UVSetName);
 
 				int UVIndex = GeomData->NumUVs;
 				if(STPrimvar)
@@ -493,22 +805,22 @@ public:
 
 					GeomData->UVs[UVIndex].Coords.clear();
 
-					if (STPrimvar.GetInterpolation() == PXR_NAMESPACE::UsdGeomTokens->faceVarying)
+					if (STPrimvar.GetInterpolation() == UsdGeomTokens->faceVarying)
 					{
 						GeomData->UVs[UVIndex].UVInterpMethod = EUsdInterpolationMethod::FaceVarying;
-						PXR_NAMESPACE::VtVec2fArray UVs;
+						VtVec2fArray UVs;
 						STPrimvar.ComputeFlattened(&UVs, Time);
 						if (UVs.size() == GeomData->FaceIndices.size())
 						{
 							GeomData->UVs[UVIndex].Coords.resize(UVs.size());
-							memcpy(&GeomData->UVs[UVIndex].Coords[0], &UVs[0], UVs.size() * sizeof(PXR_NAMESPACE::GfVec2f));
+							memcpy(&GeomData->UVs[UVIndex].Coords[0], &UVs[0], UVs.size() * sizeof(GfVec2f));
 						}
 					}
-					else if (STPrimvar.GetInterpolation() == PXR_NAMESPACE::UsdGeomTokens->vertex)
+					else if (STPrimvar.GetInterpolation() == UsdGeomTokens->vertex)
 					{
 						GeomData->UVs[UVIndex].UVInterpMethod = EUsdInterpolationMethod::Vertex;
-						PXR_NAMESPACE::VtIntArray UVIndices;
-						PXR_NAMESPACE::VtVec2fArray UVs;
+						VtIntArray UVIndices;
+						VtVec2fArray UVs;
 						STPrimvar.GetIndices(&UVIndices, Time);
 						STPrimvar.Get(&UVs, Time);
 
@@ -517,7 +829,7 @@ public:
 							GeomData->UVs[UVIndex].Coords.resize(UVIndices.size());
 							for (int PointIdx = 0; PointIdx < GeomData->Points.size(); ++PointIdx)
 							{
-								PXR_NAMESPACE::GfVec2f UV = UVs[UVIndices[PointIdx]];
+								GfVec2f UV = UVs[UVIndices[PointIdx]];
 								GeomData->UVs[UVIndex].Coords[PointIdx] = FUsdVector2Data(UV[0], UV[1]);
 
 							}
@@ -530,7 +842,7 @@ public:
 			// @todo time not supported yet
 			if(GeomData->FaceMaterialIndices.size() == 0)
 			{
-				vector<PXR_NAMESPACE::UsdGeomFaceSetAPI> FaceSets = PXR_NAMESPACE::UsdGeomFaceSetAPI::GetFaceSets(Prim);
+				vector<UsdGeomFaceSetAPI> FaceSets = UsdGeomFaceSetAPI::GetFaceSets(Prim);
 
 				GeomData->FaceMaterialIndices.resize(GeomData->FaceVertexCounts.size());
 				memset(&GeomData->FaceMaterialIndices[0], 0, sizeof(int)*GeomData->FaceMaterialIndices.size());
@@ -546,24 +858,24 @@ public:
 				//GeomData->MaterialNames.resize(FaceSets)
 				for (int FaceSetIdx = 0; FaceSetIdx < FaceSets.size(); ++FaceSetIdx)
 				{
-					const PXR_NAMESPACE::UsdGeomFaceSetAPI& FaceSet = FaceSets[FaceSetIdx];
+					const UsdGeomFaceSetAPI& FaceSet = FaceSets[FaceSetIdx];
 					
-					PXR_NAMESPACE::SdfPathVector BindingTargets;
+					SdfPathVector BindingTargets;
 					FaceSet.GetBindingTargets(&BindingTargets);
 
 					
-					PXR_NAMESPACE::UsdStageWeakPtr Stage = Prim.GetStage();
-					for(const PXR_NAMESPACE::SdfPath& Path : BindingTargets)
+					UsdStageWeakPtr Stage = Prim.GetStage();
+					for(const SdfPath& Path : BindingTargets)
 					{
 						// load each material at the material path; 
-						PXR_NAMESPACE::UsdPrim MaterialPrim = Stage->Load(Path);
+						UsdPrim MaterialPrim = Stage->Load(Path);
 
 						// Default to using the prim path name as the path for this material in Unreal
 						std::string MaterialName = MaterialPrim.GetName().GetString();
 
 						// See if the material has an "unrealAssetPath" attribute.  This should be the full name of the material
-						static const PXR_NAMESPACE::TfToken AssetPathToken = PXR_NAMESPACE::TfToken(UnrealIdentifiers::AssetPath);
-						PXR_NAMESPACE::UsdAttribute UnrealAssetPathAttr = MaterialPrim.GetAttribute(AssetPathToken);
+						static const TfToken AssetPathToken = TfToken(UnrealIdentifiers::AssetPath);
+						UsdAttribute UnrealAssetPathAttr = MaterialPrim.GetAttribute(AssetPathToken);
 						if (UnrealAssetPathAttr.HasValue())
 						{
 							UnrealAssetPathAttr.Get(&MaterialName);
@@ -575,11 +887,11 @@ public:
 					if (FaceSet.GetIsPartition())
 					{
 						// Get the list of faces in the face set.  The size of this list determines the number of materials in this set
-						PXR_NAMESPACE::VtIntArray FaceCounts;
+						VtIntArray FaceCounts;
 						FaceSet.GetFaceCounts(&FaceCounts, Time);
 
 						// Get the list of global face indices mapped in this set
-						PXR_NAMESPACE::VtIntArray FaceIndices;
+						VtIntArray FaceIndices;
 						FaceSet.GetFaceIndices(&FaceIndices, Time);
 
 						// How far we are into the face indices list
@@ -609,14 +921,14 @@ public:
 			// SubD
 			{
 				GeomData->SubdivisionScheme = EUsdSubdivisionScheme::CatmullClark;
-				PXR_NAMESPACE::UsdAttribute SubDivScheme = Mesh.GetSubdivisionSchemeAttr();
+				UsdAttribute SubDivScheme = Mesh.GetSubdivisionSchemeAttr();
 				if (SubDivScheme)
 				{
-					static const PXR_NAMESPACE::TfToken CatmullClark("catmullClark");
-					static const PXR_NAMESPACE::TfToken Loop("loop");
-					static const PXR_NAMESPACE::TfToken Bilinear("bilinear");
-					static const PXR_NAMESPACE::TfToken None("none");
-					PXR_NAMESPACE::TfToken SchemeName;
+					static const TfToken CatmullClark("catmullClark");
+					static const TfToken Loop("loop");
+					static const TfToken Bilinear("bilinear");
+					static const TfToken None("none");
+					TfToken SchemeName;
 					SubDivScheme.Get(&SchemeName, Time);
 
 					if (SchemeName == CatmullClark)
@@ -637,52 +949,52 @@ public:
 					}
 
 				}
-				PXR_NAMESPACE::UsdAttribute CreaseIndices = Mesh.GetCreaseIndicesAttr();
+				UsdAttribute CreaseIndices = Mesh.GetCreaseIndicesAttr();
 				if (CreaseIndices)
 				{
 					GeomData->CreaseIndices.clear();
 
-					PXR_NAMESPACE::VtIntArray CreaseIndicesArray;
+					VtIntArray CreaseIndicesArray;
 					CreaseIndices.Get(&CreaseIndicesArray, Time);
 					GeomData->CreaseIndices.assign(CreaseIndicesArray.begin(), CreaseIndicesArray.end());
 				}
 
-				PXR_NAMESPACE::UsdAttribute CreaseLengths = Mesh.GetCreaseLengthsAttr();
+				UsdAttribute CreaseLengths = Mesh.GetCreaseLengthsAttr();
 				if (CreaseLengths)
 				{
 					GeomData->CreaseLengths.clear();
 
-					PXR_NAMESPACE::VtIntArray CreaseLengthsArray;
+					VtIntArray CreaseLengthsArray;
 					CreaseLengths.Get(&CreaseLengthsArray);
 					GeomData->CreaseLengths.assign(CreaseLengthsArray.begin(), CreaseLengthsArray.end());
 				}
 
-				PXR_NAMESPACE::UsdAttribute CreaseSharpnesses = Mesh.GetCreaseSharpnessesAttr();
+				UsdAttribute CreaseSharpnesses = Mesh.GetCreaseSharpnessesAttr();
 				if (CreaseSharpnesses)
 				{
 					GeomData->CreaseSharpnesses.clear();
 
-					PXR_NAMESPACE::VtFloatArray CreaseSharpnessesArray;
+					VtFloatArray CreaseSharpnessesArray;
 					CreaseSharpnesses.Get(&CreaseSharpnessesArray);
 					GeomData->CreaseSharpnesses.assign(CreaseSharpnessesArray.begin(), CreaseSharpnessesArray.end());
 				}
 
-				PXR_NAMESPACE::UsdAttribute CornerCreaseIndices = Mesh.GetCornerIndicesAttr();
+				UsdAttribute CornerCreaseIndices = Mesh.GetCornerIndicesAttr();
 				if (CornerCreaseIndices)
 				{
 					GeomData->CornerCreaseIndices.clear();
 
-					PXR_NAMESPACE::VtIntArray CornerCreaseIndicesArray;
+					VtIntArray CornerCreaseIndicesArray;
 					CornerCreaseIndices.Get(&CornerCreaseIndicesArray);
 					GeomData->CornerCreaseIndices.assign(CornerCreaseIndicesArray.begin(), CornerCreaseIndicesArray.end());
 				}
 				
-				PXR_NAMESPACE::UsdAttribute CornerSharpnesses = Mesh.GetCornerSharpnessesAttr();
+				UsdAttribute CornerSharpnesses = Mesh.GetCornerSharpnessesAttr();
 				if (CornerSharpnesses)
 				{
 					GeomData->CornerSharpnesses.clear();
 
-					PXR_NAMESPACE::VtFloatArray CornerSharpnessesArray;
+					VtFloatArray CornerSharpnessesArray;
 					CornerSharpnesses.Get(&CornerSharpnessesArray);
 					GeomData->CornerSharpnesses.assign(CornerSharpnessesArray.begin(), CornerSharpnessesArray.end());
 				}
@@ -700,7 +1012,7 @@ public:
 		int NumLODs = 0;
 		if (Prim.HasVariantSets())
 		{
-			PXR_NAMESPACE::UsdVariantSet LODVariantSet = Prim.GetVariantSet(UnrealIdentifiers::LOD);
+			UsdVariantSet LODVariantSet = Prim.GetVariantSet(UnrealIdentifiers::LOD);
 			if(LODVariantSet.IsValid())
 			{
 				vector<string> VariantNames = LODVariantSet.GetVariantNames();
@@ -715,7 +1027,7 @@ public:
 	{
 		if (Prim.HasVariantSets())
 		{
-			PXR_NAMESPACE::UsdVariantSet LODVariantSet = Prim.GetVariantSet(UnrealIdentifiers::LOD);
+			UsdVariantSet LODVariantSet = Prim.GetVariantSet(UnrealIdentifiers::LOD);
 			if (LODVariantSet.IsValid())
 			{
 				char LODName[10];
@@ -726,7 +1038,7 @@ public:
 #endif
 				LODVariantSet.SetVariantSelection(string(LODName));
 
-				PXR_NAMESPACE::UsdPrim LODChild = Prim.GetChild(PXR_NAMESPACE::TfToken(LODName));
+				UsdPrim LODChild = Prim.GetChild(TfToken(LODName));
 				if (LODChild)
 				{
 					
@@ -758,20 +1070,59 @@ public:
 
 		return nullptr;
 	}
+
+	virtual const std::vector<FUsdAttribute>& GetAttributes() const override
+	{
+		AllAttributes.clear();
+		PrivateGetAttributes(AllAttributes, TfToken());
+
+		return AllAttributes;
+	}
+
+	virtual const std::vector<FUsdAttribute>& GetUnrealPropertyAttributes() const override
+	{
+		UnrealPropAttributes.clear();
+		PrivateGetAttributes(UnrealPropAttributes, UnrealIdentifiers::PropertyPath);
+
+		return UnrealPropAttributes;
+	}
+
+	UsdPrim GetUSDPrim() { return Prim; }
+
 private:
-	PXR_NAMESPACE::UsdPrim Prim;
+	void PrivateGetAttributes(std::vector<FUsdAttribute>& OutAttributes, const TfToken& ByMetadata) const
+	{
+		std::vector<UsdAttribute> Attributes = Prim.GetAttributes();
+		OutAttributes.reserve(Attributes.size());
+
+		for (UsdAttribute& Attr : Attributes)
+		{
+			if (ByMetadata.IsEmpty() || Attr.HasCustomDataKey(ByMetadata))
+			{
+				OutAttributes.push_back(
+					FUsdAttribute(std::shared_ptr<FAttribInternalData>(new FAttribInternalData(Attr)))
+				);
+			}
+		}
+	}
+private:
+	UsdPrim Prim;
 	vector<FPrimAndData> Children;
 	vector<FPrimAndData> VariantData;
+	mutable std::vector<FUsdAttribute> AllAttributes;
+	mutable std::vector<FUsdAttribute> UnrealPropAttributes;
 	string PrimName;
 	string PrimPath;
 	string UnrealAssetPath;
+	string UnrealActorClass;
+	string UnrealPropertyPath;
 	FUsdGeomData* GeomData;
 };
 
 class FUsdStage : public IUsdStage
 {
 public:
-	FUsdStage(PXR_NAMESPACE::UsdStageRefPtr& InStage)
+	FUsdStage(UsdStageRefPtr& InStage)
 		: Stage(InStage)
 		, RootPrim(nullptr)
 	{
@@ -788,7 +1139,7 @@ public:
 	EUsdUpAxis GetUpAxis() const override
 	{
 		// note USD does not support X up
-		return UsdGeomGetStageUpAxis(Stage) == PXR_NAMESPACE::UsdGeomTokens->y ? EUsdUpAxis::YAxis : EUsdUpAxis::ZAxis;
+		return UsdGeomGetStageUpAxis(Stage) == UsdGeomTokens->y ? EUsdUpAxis::YAxis : EUsdUpAxis::ZAxis;
 	}
 
 	IUsdPrim* GetRootPrim() override
@@ -827,32 +1178,27 @@ public:
 	}
 
 private:
-	PXR_NAMESPACE::UsdStageRefPtr Stage;
+	UsdStageRefPtr Stage;
 	FUsdPrim* RootPrim;
 };
 
 bool UnrealUSDWrapper::bInitialized = false;
 FUsdStage* UnrealUSDWrapper::CurrentStage = nullptr;
+std::string UnrealUSDWrapper::Errors;
 
 
-
-void UnrealUSDWrapper::Initialize(const char* PathToModule)
+void UnrealUSDWrapper::Initialize(const std::vector<std::string>& PluginDirectories)
 {
 	bInitialized = true;
 
-	string PluginPath = PathToModule;
-#ifdef _WINDOWS
-	PluginPath += "/Resources/UsdResources/plugins";
-#elif defined(__linux__)
-	PluginPath += "/Resources/UsdResources/Linux/plugins";
-#endif
-
 	// Needed to find plugins in non-standard places
-	PXR_NAMESPACE::PlugRegistry::GetInstance().RegisterPlugins(PluginPath);
+	PlugRegistry::GetInstance().RegisterPlugins(PluginDirectories);
 }
 
 IUsdStage* UnrealUSDWrapper::ImportUSDFile(const char* Path, const char* Filename)
 {
+	Errors.clear();
+
 	if (!bInitialized)
 	{
 		return nullptr;
@@ -863,29 +1209,28 @@ IUsdStage* UnrealUSDWrapper::ImportUSDFile(const char* Path, const char* Filenam
 	// Clean up any old data
 	CleanUp();
 
-	PXR_NAMESPACE::TfErrorMark ErrorMark;
+	TfErrorMark ErrorMark;
 
 	string PathAndFilename = string(Path) + string(Filename);
 
-	bool bIsSupported = PXR_NAMESPACE::UsdStage::IsSupportedFile(PathAndFilename);
+	bool bIsSupported = UsdStage::IsSupportedFile(PathAndFilename);
 
-	PXR_NAMESPACE::UsdStageRefPtr Stage = PXR_NAMESPACE::UsdStage::Open(PathAndFilename);
+	UsdStageRefPtr Stage = UsdStage::Open(PathAndFilename);
 
-	if(Stage)
+	if (Stage)
 	{
 		CurrentStage = new FUsdStage(Stage);
-
-
-		//USDHelpers::LogPrimTree(Stage->GetPseudoRoot());
 	}
 
-
 	if (!ErrorMark.IsClean())
 	{
-		PXR_NAMESPACE::TfErrorMark::Iterator i;
+		TfErrorMark::Iterator i;
 		for (i = ErrorMark.GetBegin(); i != ErrorMark.GetEnd(); ++i)
 		{
-			Log("%s %s\n", i->GetErrorCodeAsString().c_str(), i->GetCommentary().c_str());
+			Errors += i->GetErrorCodeAsString();
+			Errors += " ";
+			Errors += i->GetCommentary();
+			Errors += "\n";
 		}
 	}
 
@@ -907,5 +1252,11 @@ void UnrealUSDWrapper::CleanUp()
 
 double UnrealUSDWrapper::GetDefaultTimeCode()
 {
-	return PXR_NAMESPACE::UsdTimeCode::Default().GetValue();
+	return UsdTimeCode::Default().GetValue();
 }
+
+const char* UnrealUSDWrapper::GetErrors()
+{
+	return Errors.length() > 0 ? Errors.c_str() : nullptr;
+}
+
