- Updated MeshDescription attribute calls to fix deprecation warnings.
- Removed TMeshAttributeArraySet::AddArray because the functionality was already available as part of SetNumIndices.
- Renamed TMeshAttributeArraySet::InsertArray, RemoveArray to InsertIndex, RemoveIndex for naming convention consistency (these methods deal with attribute indices, not with arrays). Added support for them in other attribute classes, and made them virtual so they can be called as part of an AttributesView.
- Removed redundant code in FUSDStaticMeshImportState::AddPolygons, when determining the number of UVs in the mesh description.

diff --git a/Engine/Plugins/Editor/MeshEditor/Source/MeshEditor/MeshEditorSelectionModifiers.cpp b/Engine/Plugins/Editor/MeshEditor/Source/MeshEditor/MeshEditorSelectionModifiers.cpp
--- a/Engine/Plugins/Editor/MeshEditor/Source/MeshEditor/MeshEditorSelectionModifiers.cpp
+++ b/Engine/Plugins/Editor/MeshEditor/Source/MeshEditor/MeshEditorSelectionModifiers.cpp
@@ -204,7 +204,7 @@ bool USelectPolygonsBySmoothingGroup::ModifySelection(TMap< UEditableMesh*, TArr
 
 		TSet< FPolygonID > CheckedPolygons;		// set of polygons that have already been checked
 
-		const TEdgeAttributeArray<bool>& EdgeHardnesses = EditableMesh->GetMeshDescription()->EdgeAttributes().GetAttributes<bool>(MeshAttribute::Edge::IsHard);
+		TEdgeAttributesConstRef<bool> EdgeHardnesses = MeshDescription->EdgeAttributes().GetAttributesRef<bool>(MeshAttribute::Edge::IsHard);
 		for (const FMeshElement& PolygonElement : Polygons)
 		{
 			FPolygonID PolygonID(PolygonElement.ElementAddress.ElementID);
@@ -232,8 +232,7 @@ bool USelectPolygonsBySmoothingGroup::ModifySelection(TMap< UEditableMesh*, TArr
 					{
 						if (!EdgeHardnesses[EdgeID])
 						{
-							const FMeshEdge& Edge = MeshDescription->GetEdge(EdgeID);
-							for (const FPolygonID& ConnectedPolygonID : Edge.ConnectedPolygons)
+							for (const FPolygonID ConnectedPolygonID : MeshDescription->GetEdgeConnectedPolygons(EdgeID))
 							{
 								if (!CheckedPolygons.Contains(ConnectedPolygonID))
 								{
diff --git a/Engine/Plugins/Editor/USDImporter/Source/USDImporter/Private/StaticMeshImporter.cpp b/Engine/Plugins/Editor/USDImporter/Source/USDImporter/Private/StaticMeshImporter.cpp
--- a/Engine/Plugins/Editor/USDImporter/Source/USDImporter/Private/StaticMeshImporter.cpp
+++ b/Engine/Plugins/Editor/USDImporter/Source/USDImporter/Private/StaticMeshImporter.cpp
@@ -23,34 +23,34 @@ struct FMeshDescriptionWrapper
 	struct FVertexAttributes
 	{
 		FVertexAttributes(FMeshDescription* MeshDescription) :
-			Positions(MeshDescription->VertexAttributes().GetAttributes<FVector>(MeshAttribute::Vertex::Position)),
-			Normals(MeshDescription->VertexInstanceAttributes().GetAttributes<FVector>(MeshAttribute::VertexInstance::Normal)),
-			Tangents(MeshDescription->VertexInstanceAttributes().GetAttributes<FVector>(MeshAttribute::VertexInstance::Tangent)),
-			BinormalSigns(MeshDescription->VertexInstanceAttributes().GetAttributes<float>(MeshAttribute::VertexInstance::BinormalSign)),
-			Colors(MeshDescription->VertexInstanceAttributes().GetAttributes<FVector4>(MeshAttribute::VertexInstance::Color)),
-			UVs(MeshDescription->VertexInstanceAttributes().GetAttributesSet<FVector2D>(MeshAttribute::VertexInstance::TextureCoordinate))
+			Positions(MeshDescription->VertexAttributes().GetAttributesRef<FVector>(MeshAttribute::Vertex::Position)),
+			Normals(MeshDescription->VertexInstanceAttributes().GetAttributesRef<FVector>(MeshAttribute::VertexInstance::Normal)),
+			Tangents(MeshDescription->VertexInstanceAttributes().GetAttributesRef<FVector>(MeshAttribute::VertexInstance::Tangent)),
+			BinormalSigns(MeshDescription->VertexInstanceAttributes().GetAttributesRef<float>(MeshAttribute::VertexInstance::BinormalSign)),
+			Colors(MeshDescription->VertexInstanceAttributes().GetAttributesRef<FVector4>(MeshAttribute::VertexInstance::Color)),
+			UVs(MeshDescription->VertexInstanceAttributes().GetAttributesRef<FVector2D>(MeshAttribute::VertexInstance::TextureCoordinate))
 		{
 		}
-		TVertexAttributeArray<FVector>& Positions;
-		TVertexInstanceAttributeArray<FVector>& Normals;
-		TVertexInstanceAttributeArray<FVector>& Tangents;
-		TVertexInstanceAttributeArray<float>& BinormalSigns;
-		TVertexInstanceAttributeArray<FVector4>& Colors;
-		TVertexInstanceAttributeIndicesArray<FVector2D>& UVs;
+		TVertexAttributesRef<FVector> Positions;
+		TVertexInstanceAttributesRef<FVector> Normals;
+		TVertexInstanceAttributesRef<FVector> Tangents;
+		TVertexInstanceAttributesRef<float> BinormalSigns;
+		TVertexInstanceAttributesRef<FVector4> Colors;
+		TVertexInstanceAttributesRef<FVector2D> UVs;
 	};
 
 	FMeshDescriptionWrapper(FMeshDescription* MeshDescription) :
 		Vertex(MeshDescription),
-		EdgeHardnesses(MeshDescription->EdgeAttributes().GetAttributes<bool>(MeshAttribute::Edge::IsHard)),
-		EdgeCreaseSharpnesses(MeshDescription->EdgeAttributes().GetAttributes<float>(MeshAttribute::Edge::CreaseSharpness)),
-		PolygonGroupImportedMaterialSlotNames(MeshDescription->PolygonGroupAttributes().GetAttributes<FName>(MeshAttribute::PolygonGroup::ImportedMaterialSlotName))
+		EdgeHardnesses(MeshDescription->EdgeAttributes().GetAttributesRef<bool>(MeshAttribute::Edge::IsHard)),
+		EdgeCreaseSharpnesses(MeshDescription->EdgeAttributes().GetAttributesRef<float>(MeshAttribute::Edge::CreaseSharpness)),
+		PolygonGroupImportedMaterialSlotNames(MeshDescription->PolygonGroupAttributes().GetAttributesRef<FName>(MeshAttribute::PolygonGroup::ImportedMaterialSlotName))
 	{
 	}
 
 	FVertexAttributes Vertex;
-	TEdgeAttributeArray<bool>& EdgeHardnesses;
-	TEdgeAttributeArray<float>& EdgeCreaseSharpnesses;
-	TPolygonGroupAttributeArray<FName>& PolygonGroupImportedMaterialSlotNames;
+	TEdgeAttributesRef<bool> EdgeHardnesses;
+	TEdgeAttributesRef<float> EdgeCreaseSharpnesses;
+	TPolygonGroupAttributesRef<FName> PolygonGroupImportedMaterialSlotNames;
 };
 
 struct FUSDImportMaterialInfo
@@ -117,7 +117,6 @@ void FUSDStaticMeshImportState::ProcessStaticUSDGeometry(IUsdPrim* GeomPrim, int
 
 void FUSDStaticMeshImportState::AddVertexPositions(FMeshDescriptionWrapper& DestMeshWrapper, const FUsdGeomData& GeomData)
 {
-	TVertexAttributeArray<FVector>& VertexPositions = DestMeshWrapper.Vertex.Positions;
 	for (int32 LocalPointIndex = 0; LocalPointIndex < GeomData.Points.size(); ++LocalPointIndex)
 	{
 		const FUsdVectorData& Point = GeomData.Points[LocalPointIndex];
@@ -125,7 +124,7 @@ void FUSDStaticMeshImportState::AddVertexPositions(FMeshDescriptionWrapper& Dest
 		Pos = FinalTransform.TransformPosition(Pos);
 
 		FVertexID AddedVertexId = MeshDescription->CreateVertex();
-		VertexPositions[AddedVertexId] = Pos;
+		DestMeshWrapper.Vertex.Positions[AddedVertexId] = Pos;
 	}
 }
 
@@ -133,18 +132,7 @@ bool FUSDStaticMeshImportState::AddPolygons(FMeshDescriptionWrapper& DestMeshWra
 {
 	// When importing multiple mesh pieces to the same static mesh.  Ensure each mesh piece has the same number of Uv's
 	{
-		int32 ExistingUVCount = 0;
-		for (int32 UVChannelIndex = 0; UVChannelIndex < DestMeshWrapper.Vertex.UVs.GetNumIndices(); ++UVChannelIndex)
-		{
-			if (DestMeshWrapper.Vertex.UVs.GetArrayForIndex(UVChannelIndex).Num() > 0)
-			{
-				ExistingUVCount++;
-			}
-			else
-			{
-				break;
-			}
-		}
+		int32 ExistingUVCount = DestMeshWrapper.Vertex.UVs.GetNumIndices();
 		int32 NumUVs = FMath::Max(GeomData.NumUVs, ExistingUVCount);
 		NumUVs = FMath::Min<int32>(MAX_MESH_TEXTURE_COORDS, NumUVs);
 		// At least one UV set must exist.  
@@ -211,7 +199,7 @@ bool FUSDStaticMeshImportState::AddPolygons(FMeshDescriptionWrapper& DestMeshWra
 
 				// Flip V for Unreal uv's which match directx
 				FVector2D FinalUVVector(UV.X, 1.f - UV.Y);
-				DestMeshWrapper.Vertex.UVs.GetArrayForIndex(UVLayerIndex)[AddedVertexInstanceId] = FinalUVVector;
+				DestMeshWrapper.Vertex.UVs.Set(AddedVertexInstanceId, UVLayerIndex, FinalUVVector);
 			}
 		}
 
diff --git a/Engine/Source/Developer/MeshDescriptionOperations/Private/MeshDescriptionOperations.cpp b/Engine/Source/Developer/MeshDescriptionOperations/Private/MeshDescriptionOperations.cpp
--- a/Engine/Source/Developer/MeshDescriptionOperations/Private/MeshDescriptionOperations.cpp
+++ b/Engine/Source/Developer/MeshDescriptionOperations/Private/MeshDescriptionOperations.cpp
@@ -1201,20 +1201,20 @@ bool FMeshDescriptionOperations::GenerateUniqueUVsForStaticMesh(const FMeshDescr
 
 bool FMeshDescriptionOperations::AddUVChannel(FMeshDescription& MeshDescription)
 {
-	TVertexInstanceAttributeIndicesArray<FVector2D>& VertexInstanceUVs = MeshDescription.VertexInstanceAttributes().GetAttributesSet<FVector2D>(MeshAttribute::VertexInstance::TextureCoordinate);
+	TVertexInstanceAttributesRef<FVector2D> VertexInstanceUVs = MeshDescription.VertexInstanceAttributes().GetAttributesRef<FVector2D>(MeshAttribute::VertexInstance::TextureCoordinate);
 	if (VertexInstanceUVs.GetNumIndices() >= MAX_MESH_TEXTURE_COORDS)
 	{
 		UE_LOG(LogMeshDescriptionOperations, Error, TEXT("AddUVChannel: Cannot add UV channel. Maximum number of UV channels reached (%d)."), MAX_MESH_TEXTURE_COORDS);
 		return false;
 	}
 
-	VertexInstanceUVs.AddArray();
+	VertexInstanceUVs.SetNumIndices(VertexInstanceUVs.GetNumIndices() + 1);
 	return true;
 }
 
 bool FMeshDescriptionOperations::InsertUVChannel(FMeshDescription& MeshDescription, int32 UVChannelIndex)
 {
-	TVertexInstanceAttributeIndicesArray<FVector2D>& VertexInstanceUVs = MeshDescription.VertexInstanceAttributes().GetAttributesSet<FVector2D>(MeshAttribute::VertexInstance::TextureCoordinate);
+	TVertexInstanceAttributesRef<FVector2D> VertexInstanceUVs = MeshDescription.VertexInstanceAttributes().GetAttributesRef<FVector2D>(MeshAttribute::VertexInstance::TextureCoordinate);
 	if (UVChannelIndex < 0 || UVChannelIndex > VertexInstanceUVs.GetNumIndices())
 	{
 		UE_LOG(LogMeshDescriptionOperations, Error, TEXT("InsertUVChannel: Cannot insert UV channel. Given UV channel index %d is out of bounds."), UVChannelIndex);
@@ -1227,13 +1227,13 @@ bool FMeshDescriptionOperations::InsertUVChannel(FMeshDescription& MeshDescripti
 		return false;
 	}
 
-	VertexInstanceUVs.InsertArray(UVChannelIndex);
+	VertexInstanceUVs.InsertIndex(UVChannelIndex);
 	return true;
 }
 
 bool FMeshDescriptionOperations::RemoveUVChannel(FMeshDescription& MeshDescription, int32 UVChannelIndex)
 {
-	TVertexInstanceAttributeIndicesArray<FVector2D>& VertexInstanceUVs = MeshDescription.VertexInstanceAttributes().GetAttributesSet<FVector2D>(MeshAttribute::VertexInstance::TextureCoordinate);
+	TVertexInstanceAttributesRef<FVector2D> VertexInstanceUVs = MeshDescription.VertexInstanceAttributes().GetAttributesRef<FVector2D>(MeshAttribute::VertexInstance::TextureCoordinate);
 	if (VertexInstanceUVs.GetNumIndices() == 1)
 	{
 		UE_LOG(LogMeshDescriptionOperations, Error, TEXT("RemoveUVChannel: Cannot remove UV channel. There must be at least one channel."));
@@ -1246,7 +1246,7 @@ bool FMeshDescriptionOperations::RemoveUVChannel(FMeshDescription& MeshDescripti
 		return false;
 	}
 
-	VertexInstanceUVs.RemoveArray(UVChannelIndex);
+	VertexInstanceUVs.RemoveIndex(UVChannelIndex);
 	return true;
 }
 
diff --git a/Engine/Source/Runtime/Engine/Private/StaticMesh.cpp b/Engine/Source/Runtime/Engine/Private/StaticMesh.cpp
--- a/Engine/Source/Runtime/Engine/Private/StaticMesh.cpp
+++ b/Engine/Source/Runtime/Engine/Private/StaticMesh.cpp
@@ -3169,8 +3169,7 @@ int32 UStaticMesh::GetNumUVChannels(int32 LODIndex)
 	FMeshDescription* MeshDescription = GetOriginalMeshDescription(LODIndex);
 	if (MeshDescription)
 	{
-		TVertexInstanceAttributeIndicesArray<FVector2D>& VertexInstanceUVs = MeshDescription->VertexInstanceAttributes().GetAttributesSet<FVector2D>(MeshAttribute::VertexInstance::TextureCoordinate);
-		NumUVChannels = VertexInstanceUVs.GetNumIndices();
+		NumUVChannels = MeshDescription->VertexInstanceAttributes().GetAttributeIndexCount<FVector2D>(MeshAttribute::VertexInstance::TextureCoordinate);
 	}
 #endif
 	return NumUVChannels;
diff --git a/Engine/Source/Runtime/MeshDescription/Public/MeshAttributeArray.h b/Engine/Source/Runtime/MeshDescription/Public/MeshAttributeArray.h
--- a/Engine/Source/Runtime/MeshDescription/Public/MeshAttributeArray.h
+++ b/Engine/Source/Runtime/MeshDescription/Public/MeshAttributeArray.h
@@ -249,6 +249,8 @@ public:
 	virtual void Remap( const TSparseArray<int32>& IndexRemap ) = 0;
 	virtual int32 GetNumIndices() const = 0;
 	virtual void SetNumIndices( const int32 NumIndices ) = 0;
+	virtual void InsertIndex( const int32 Index ) = 0;
+	virtual void RemoveIndex( const int32 Index ) = 0;
 
 	/** Determine whether this attribute array set is of the given type */
 	template <typename T>
@@ -321,13 +323,6 @@ public:
 		}
 	}
 
-	/** Add an array filled with the default value */
-	FORCEINLINE void AddArray()
-	{
-		TMeshAttributeArrayBase<AttributeType>& NewArray = ArrayForIndices.AddDefaulted_GetRef();
-		NewArray.Initialize(NumElements, DefaultValue);
-	}
-
 	/** Sets the number of elements to the exact number provided, and initializes them to the default value */
 	virtual void Initialize( const int32 Count )
 	{
@@ -338,20 +333,6 @@ public:
 		}
 	}
 
-	/** Insert an array at the given index, filled with the default value */
-	FORCEINLINE void InsertArray(int32 Index)
-	{
-		TMeshAttributeArrayBase<AttributeType>& NewArray = ArrayForIndices.InsertDefaulted_GetRef(Index);
-		NewArray.Initialize(NumElements, DefaultValue);
-	}
-
-	/** Remove the array at the given index */
-	FORCEINLINE void RemoveArray(int32 Index)
-	{
-		ArrayForIndices.RemoveAt(Index);
-	}
-
-
 	/** Polymorphic serialization */
 	virtual void Serialize( FArchive& Ar )
 	{
@@ -384,6 +365,20 @@ public:
 		}
 	}
 
+	/** Insert a new attribute index */
+	virtual void InsertIndex( const int32 Index )
+	{
+		ArrayForIndices.InsertDefaulted( Index );
+		ArrayForIndices[ Index ].Initialize( NumElements, DefaultValue );
+	}
+
+	/** Remove the array at the given index */
+	virtual void RemoveIndex( const int32 Index )
+	{
+		ArrayForIndices.RemoveAt( Index );
+	}
+
+
 	/** Return the TMeshAttributeArrayBase corresponding to the given attribute index */
 	FORCEINLINE const TMeshAttributeArrayBase<AttributeType>& GetArrayForIndex( const int32 Index ) const { return ArrayForIndices[ Index ]; }
 	FORCEINLINE TMeshAttributeArrayBase<AttributeType>& GetArrayForIndex( const int32 Index ) { return ArrayForIndices[ Index ]; }
@@ -530,10 +525,22 @@ public:
 	}
 
 	/** Sets number of indices this attribute has */
-	FORCEINLINE void SetNumIndices( const int32 NumIndices )
+	FORCEINLINE void SetNumIndices( const int32 NumIndices ) const
 	{
 		const_cast<ArrayType*>( this->ArrayPtr )->ArrayType::SetNumIndices( NumIndices );	// note: override virtual dispatch
 	}
+
+	/** Inserts an attribute index */
+	FORCEINLINE void InsertIndex( const int32 Index ) const
+	{
+		const_cast<ArrayType*>( this->ArrayPtr )->ArrayType::InsertIndex( Index );	// note: override virtual dispatch
+	}
+
+	/** Removes an attribute index */
+	FORCEINLINE void RemoveIndex( const int32 Index ) const
+	{
+		const_cast<ArrayType*>( this->ArrayPtr )->ArrayType::RemoveIndex( Index );	// note: override virtual dispatch
+	}
 };
 
 
@@ -632,6 +639,18 @@ public:
 	{
 		const_cast<FMeshAttributeArraySetBase*>( this->ArrayPtr )->SetNumIndices( NumIndices );
 	}
+
+	/** Inserts an attribute index */
+	FORCEINLINE void InsertIndex( const int32 Index ) const
+	{
+		const_cast<FMeshAttributeArraySetBase*>( this->ArrayPtr )->InsertIndex( Index );
+	}
+
+	/** Removes an attribute index */
+	FORCEINLINE void RemoveIndex( const int32 Index )
+	{
+		const_cast<FMeshAttributeArraySetBase*>( this->ArrayPtr )->RemoveIndex( Index );
+	}
 };
 
 
@@ -955,6 +974,34 @@ public:
 		}
 	}
 
+	/** Insert a new index for the attribute with the given name */
+	template <typename AttributeType>
+	void InsertAttributeIndex( const FName AttributeName, const int32 Index )
+	{
+		if( FAttributesSetEntry* ArraySetPtr = this->Map.Find( AttributeName ) )
+		{
+			if( ( *ArraySetPtr )->HasType<AttributeType>() )
+			{
+				using ArrayType = TMeshAttributeArraySet<AttributeType>;
+				static_cast<ArrayType*>( ArraySetPtr->Get() )->ArrayType::InsertIndex( Index );	// note: override virtual dispatch
+			}
+		}
+	}
+
+	/** Remove an existing index from the attribute with the given name */
+	template <typename AttributeType>
+	void RemoveAttributeIndex( const FName AttributeName, const int32 Index )
+	{
+		if( FAttributesSetEntry* ArraySetPtr = this->Map.Find( AttributeName ) )
+		{
+			if( ( *ArraySetPtr )->HasType<AttributeType>() )
+			{
+				using ArrayType = TMeshAttributeArraySet<AttributeType>;
+				static_cast<ArrayType*>( ArraySetPtr->Get() )->ArrayType::RemoveIndex( Index );	// note: override virtual dispatch
+			}
+		}
+	}
+
 	/** Returns an array of all the attribute names registered for this attribute type */
 	template <typename AttributeType, typename Allocator>
 	void GetAttributeNames( TArray<FName, Allocator>& OutAttributeNames ) const
