Remove support for USD in editor monolithic. It relies on a different set of new and delete overrides which isn't available in monolithic.
We can use USD in a monolithic editor by using the ANSI allocator so that USD allocs and UE allocs all end up in the stl allocator.

diff --git a/Engine/Plugins/Importers/USDImporter/Source/UnrealUSDWrapper/Public/USDMemory.h b/Engine/Plugins/Importers/USDImporter/Source/UnrealUSDWrapper/Public/USDMemory.h
--- a/Engine/Plugins/Importers/USDImporter/Source/UnrealUSDWrapper/Public/USDMemory.h
+++ b/Engine/Plugins/Importers/USDImporter/Source/UnrealUSDWrapper/Public/USDMemory.h
@@ -191,7 +191,7 @@ private:
 template< typename T >
 TUsdStore< T > MakeUsdStore( T&& UsdObject ) { return TUsdStore< T >( MoveTemp( UsdObject ) ); }
 
-#if !FORCE_ANSI_ALLOCATOR
+#if !FORCE_ANSI_ALLOCATOR && !IS_MONOLITHIC
 	#define REPLACEMENT_OPERATOR_NEW_AND_DELETE_USD \
 		OPERATOR_NEW_MSVC_PRAGMA void* operator new  ( size_t Size                        ) OPERATOR_NEW_THROW_SPEC      { return FUsdMemoryManager::Malloc( Size ); } \
 		OPERATOR_NEW_MSVC_PRAGMA void* operator new[]( size_t Size                        ) OPERATOR_NEW_THROW_SPEC      { return FUsdMemoryManager::Malloc( Size ); } \
@@ -209,7 +209,7 @@ TUsdStore< T > MakeUsdStore( T&& UsdObject ) { return TUsdStore< T >( MoveTemp(
 	#define REPLACEMENT_OPERATOR_NEW_AND_DELETE_USD
 #endif
 
-#if !FORCE_ANSI_ALLOCATOR
+#if !FORCE_ANSI_ALLOCATOR && !IS_MONOLITHIC
 	#define IMPLEMENT_MODULE_USD( ModuleImplClass, ModuleName ) \
 		\
 		/**/ \
diff --git a/Engine/Plugins/Importers/USDImporter/Source/UnrealUSDWrapper/UnrealUSDWrapper.Build.cs b/Engine/Plugins/Importers/USDImporter/Source/UnrealUSDWrapper/UnrealUSDWrapper.Build.cs
--- a/Engine/Plugins/Importers/USDImporter/Source/UnrealUSDWrapper/UnrealUSDWrapper.Build.cs
+++ b/Engine/Plugins/Importers/USDImporter/Source/UnrealUSDWrapper/UnrealUSDWrapper.Build.cs
@@ -19,7 +19,8 @@ namespace UnrealBuildTool.Rules
 			var EngineDir = Path.GetFullPath(Target.RelativeEnginePath);
 			var PythonSourceTPSDir = Path.Combine(EngineDir, "Source", "ThirdParty", "Python");
 
-			if (Target.WindowsPlatform.Compiler != WindowsCompiler.Clang && Target.WindowsPlatform.StaticAnalyzer == WindowsStaticAnalyzer.None)
+			if (Target.WindowsPlatform.Compiler != WindowsCompiler.Clang && Target.WindowsPlatform.StaticAnalyzer == WindowsStaticAnalyzer.None &&
+				Target.LinkType != TargetLinkType.Monolithic) // If you want to use USD in a monolithic target, you'll have to use the ANSI allocator and remove this condition
 			{
 				PublicDefinitions.Add("USE_USD_SDK=1");
 
