USD importer: Added parent prim validity check to fix crash

#rnx

diff --git a/Engine/Plugins/Importers/USDImporter/Source/USDImporter/Private/USDPrimResolver.cpp b/Engine/Plugins/Importers/USDImporter/Source/USDImporter/Private/USDPrimResolver.cpp
--- a/Engine/Plugins/Importers/USDImporter/Source/USDImporter/Private/USDPrimResolver.cpp
+++ b/Engine/Plugins/Importers/USDImporter/Source/USDImporter/Private/USDPrimResolver.cpp
@@ -231,7 +231,7 @@ AActor* UUSDPrimResolver::SpawnActor(FUSDSceneImportContext& ImportContext, cons
 
 					TArray< pxr::UsdPrim > ParentPrims;
 
-					while ( ParentPrim != SpawnData.ActorPrim.Get() )
+					while ( ParentPrim && ParentPrim != SpawnData.ActorPrim.Get() )
 					{
 						ParentPrims.Add( ParentPrim );
 						ParentPrim = ParentPrim.GetParent();
