Fixed USD transform issue when the mesh prim is at the root of the stage.

diff --git a/Engine/Plugins/Importers/USDImporter/Source/USDImporter/Private/USDPrimResolver.cpp b/Engine/Plugins/Importers/USDImporter/Source/USDImporter/Private/USDPrimResolver.cpp
--- a/Engine/Plugins/Importers/USDImporter/Source/USDImporter/Private/USDPrimResolver.cpp
+++ b/Engine/Plugins/Importers/USDImporter/Source/USDImporter/Private/USDPrimResolver.cpp
@@ -251,7 +251,11 @@ AActor* UUSDPrimResolver::SpawnActor(FUSDSceneImportContext& ImportContext, cons
 					UStaticMeshComponent* StaticMeshComponent = NewObject< UStaticMeshComponent >( SpawnedActor, ComponentName );
 					StaticMeshComponent->SetStaticMesh( ImportedStaticMesh );
 
-					LocalTransform = USDToUnreal::ConvertMatrix( *ImportContext.Stage, IUsdPrim::GetLocalTransform( UsdAssetPrimToImport.Prim.Get() ) ) * LocalTransform;
+					// Don't add the prim transform if its the same prim used for the actor as it's already accounted for in the ActorTransform
+					if ( UsdAssetPrimToImport.Prim.Get() != SpawnData.ActorPrim.Get() )
+					{
+						LocalTransform = USDToUnreal::ConvertMatrix( *ImportContext.Stage, IUsdPrim::GetLocalTransform( UsdAssetPrimToImport.Prim.Get() ) ) * LocalTransform;
+					}
 
 					StaticMeshComponent->SetRelativeTransform( LocalTransform );
 
