Fix USD Linux compiler warning
#rnx

diff --git a/Engine/Plugins/Editor/USDImporter/Source/USDImporter/Private/StaticMeshImporter.cpp b/Engine/Plugins/Editor/USDImporter/Source/USDImporter/Private/StaticMeshImporter.cpp
--- a/Engine/Plugins/Editor/USDImporter/Source/USDImporter/Private/StaticMeshImporter.cpp
+++ b/Engine/Plugins/Editor/USDImporter/Source/USDImporter/Private/StaticMeshImporter.cpp
@@ -67,9 +67,9 @@ struct FUSDImportMaterialInfo
 struct FUSDStaticMeshImportState
 {
 public:
-	FUSDStaticMeshImportState(FUsdImportContext& ImportContext, TArray<FUSDImportMaterialInfo>& Materials) :
-		ImportContext(ImportContext),
-		Materials(Materials),
+	FUSDStaticMeshImportState(FUsdImportContext& InImportContext, TArray<FUSDImportMaterialInfo>& InMaterials) :
+		ImportContext(InImportContext),
+		Materials(InMaterials),
 		MeshDescription(nullptr)
 	{
 	}
