Fixed crash with Usd Attributes being deallocated by the wrong allocator.

diff --git a/Engine/Plugins/Importers/USDImporter/Source/USDImporter/Private/PropertySetter.cpp b/Engine/Plugins/Importers/USDImporter/Source/USDImporter/Private/PropertySetter.cpp
--- a/Engine/Plugins/Importers/USDImporter/Source/USDImporter/Private/PropertySetter.cpp
+++ b/Engine/Plugins/Importers/USDImporter/Source/USDImporter/Private/PropertySetter.cpp
@@ -140,12 +140,12 @@ void FUSDPropertySetter::RegisterStructSetter(FName StructName, FStructSetterFun
 
 void FUSDPropertySetter::ApplyPropertiesFromUsdAttributes(const pxr::UsdPrim& Prim, AActor* SpawnedActor, const FString& StartingPropertyPath)
 {
-	const std::vector<pxr::UsdAttribute> Attributes = IUsdPrim::GetUnrealPropertyAttributes( Prim );
+	const TUsdStore< std::vector<pxr::UsdAttribute> > Attributes = IUsdPrim::GetUnrealPropertyAttributes( Prim );
 
 	// For map properties.  Ignore prims set by maps
 	std::set< std::string > AttribsToIgnore;
 
-	for (const pxr::UsdAttribute& Attribute : Attributes)
+	for (const pxr::UsdAttribute& Attribute : Attributes.Get())
 	{
 		if ( AttribsToIgnore.find( Attribute.GetName().GetString() ) != AttribsToIgnore.end() )
 		{
@@ -455,10 +455,10 @@ void FUSDPropertySetter::SetFromUSDValue(PropertyHelpers::FPropertyAddress& Prop
 
 bool FUSDPropertySetter::FindMapKeyAndValues(const pxr::UsdPrim& Prim, pxr::UsdAttribute& OutKey, TArray<pxr::UsdAttribute>& OutValues)
 {
-	const std::vector<pxr::UsdAttribute> Attributes = IUsdPrim::GetUnrealPropertyAttributes( Prim );
+	const TUsdStore< std::vector<pxr::UsdAttribute> > Attributes = IUsdPrim::GetUnrealPropertyAttributes( Prim );
 
 	bool bFoundKey = false;
-	for (const pxr::UsdAttribute& Attrib : Attributes)
+	for (const pxr::UsdAttribute& Attrib : Attributes.Get())
 	{
 		if (USDToUnreal::ConvertString( FUsdAttribute::GetUnrealPropertyPath( Attrib ).c_str() ) == TEXT("_KEY"))
 		{
diff --git a/Engine/Plugins/Importers/USDImporter/Source/UnrealUSDWrapper/Private/UnrealUSDWrapper.cpp b/Engine/Plugins/Importers/USDImporter/Source/UnrealUSDWrapper/Private/UnrealUSDWrapper.cpp
--- a/Engine/Plugins/Importers/USDImporter/Source/UnrealUSDWrapper/Private/UnrealUSDWrapper.cpp
+++ b/Engine/Plugins/Importers/USDImporter/Source/UnrealUSDWrapper/Private/UnrealUSDWrapper.cpp
@@ -572,8 +572,10 @@ std::string IUsdPrim::GetUnrealPropertyPath(const pxr::UsdPrim& Prim)
 	return {};
 }
 
-std::vector<UsdAttribute> PrivateGetAttributes(const pxr::UsdPrim& Prim, const TfToken& ByMetadata)
+TUsdStore< std::vector<UsdAttribute> > PrivateGetAttributes(const pxr::UsdPrim& Prim, const TfToken& ByMetadata)
 {
+	FScopedUsdAllocs UsdAllocs;
+
 	std::vector<UsdAttribute> Attributes = Prim.GetAttributes();
 
 	std::vector<UsdAttribute> OutAttributes;
@@ -590,7 +592,7 @@ std::vector<UsdAttribute> PrivateGetAttributes(const pxr::UsdPrim& Prim, const T
 	return OutAttributes;
 }
 
-std::vector<UsdAttribute> IUsdPrim::GetUnrealPropertyAttributes(const pxr::UsdPrim& Prim)
+TUsdStore< std::vector<UsdAttribute> > IUsdPrim::GetUnrealPropertyAttributes(const pxr::UsdPrim& Prim)
 {
 	return PrivateGetAttributes(Prim, UnrealIdentifiers::PropertyPath);
 }
diff --git a/Engine/Plugins/Importers/USDImporter/Source/UnrealUSDWrapper/Public/UnrealUSDWrapper.h b/Engine/Plugins/Importers/USDImporter/Source/UnrealUSDWrapper/Public/UnrealUSDWrapper.h
--- a/Engine/Plugins/Importers/USDImporter/Source/UnrealUSDWrapper/Public/UnrealUSDWrapper.h
+++ b/Engine/Plugins/Importers/USDImporter/Source/UnrealUSDWrapper/Public/UnrealUSDWrapper.h
@@ -206,7 +206,7 @@ public:
 	static UNREALUSDWRAPPER_API bool HasTransform(const pxr::UsdPrim& Prim);
 	static UNREALUSDWRAPPER_API std::string GetUnrealPropertyPath(const pxr::UsdPrim& Prim);
 
-	static UNREALUSDWRAPPER_API std::vector<pxr::UsdAttribute> GetUnrealPropertyAttributes(const pxr::UsdPrim& Prim);
+	static UNREALUSDWRAPPER_API TUsdStore< std::vector<pxr::UsdAttribute> > GetUnrealPropertyAttributes(const pxr::UsdPrim& Prim);
 
 	static UNREALUSDWRAPPER_API std::string GetUnrealAssetPath(const pxr::UsdPrim& Prim);
 	static UNREALUSDWRAPPER_API std::string GetUnrealActorClass(const pxr::UsdPrim& Prim);
