USD importer
- Fixed crash when ignoring existing actors on a reimport

diff --git a/Engine/Plugins/Importers/USDImporter/Source/USDImporter/Private/PropertySetter.cpp b/Engine/Plugins/Importers/USDImporter/Source/USDImporter/Private/PropertySetter.cpp
--- a/Engine/Plugins/Importers/USDImporter/Source/USDImporter/Private/PropertySetter.cpp
+++ b/Engine/Plugins/Importers/USDImporter/Source/USDImporter/Private/PropertySetter.cpp
@@ -99,6 +99,11 @@ FUSDPropertySetter::FUSDPropertySetter(FUsdImportContext& InImportContext)
 
 void FUSDPropertySetter::ApplyPropertiesToActor(AActor* SpawnedActor, const pxr::UsdPrim& Prim, const FString& StartingPropertyPath)
 {
+	if ( !SpawnedActor )
+	{
+		return;
+	}
+
 	ApplyPropertiesFromUsdAttributes(Prim, SpawnedActor, StartingPropertyPath);
 
 	// find prims that represent complicated properties 
