Fixed USD memory allocation issue when importing geometry with lods.

diff --git a/Engine/Plugins/Importers/USDImporter/Source/UnrealUSDWrapper/Private/UnrealUSDWrapper.cpp b/Engine/Plugins/Importers/USDImporter/Source/UnrealUSDWrapper/Private/UnrealUSDWrapper.cpp
--- a/Engine/Plugins/Importers/USDImporter/Source/UnrealUSDWrapper/Private/UnrealUSDWrapper.cpp
+++ b/Engine/Plugins/Importers/USDImporter/Source/UnrealUSDWrapper/Private/UnrealUSDWrapper.cpp
@@ -466,6 +466,8 @@ bool IUsdPrim::HasGeometryDataOrLODVariants(const UsdPrim& Prim)
 
 int IUsdPrim::GetNumLODs(const UsdPrim& Prim)
 {
+	FScopedUsdAllocs UsdAllocs;
+
 	// 0 indicates no variant or no lods in variant. 
 	int NumLODs = 0;
 	if (Prim.HasVariantSets())
@@ -868,6 +870,8 @@ bool IUsdPrim::HasTransform(const pxr::UsdPrim& Prim)
 
 bool IUsdPrim::SetActiveLODIndex(const pxr::UsdPrim& Prim, int LODIndex)
 {
+	FScopedUsdAllocs UsdAllocs;
+
 	if (Prim.HasVariantSets())
 	{
 		UsdVariantSet LODVariantSet = Prim.GetVariantSet(UnrealIdentifiers::LOD);
