Copied cl 9420671, 9466794, 9532476, 9602026

USD
- Fixed UV crash with house
- Fixed skeletal mesh UVs for quads
- Added support for material opacity
- Linux build fixes
- Fixed crashes when we have empty or degenerate meshes.

diff --git a/Engine/Plugins/Importers/USDImporter/Source/USDStage/Private/USDStageActor.cpp b/Engine/Plugins/Importers/USDImporter/Source/USDStage/Private/USDStageActor.cpp
--- a/Engine/Plugins/Importers/USDImporter/Source/USDStage/Private/USDStageActor.cpp
+++ b/Engine/Plugins/Importers/USDImporter/Source/USDStage/Private/USDStageActor.cpp
@@ -248,7 +248,7 @@ bool AUsdStageActor::LoadStaticMesh( const pxr::UsdGeomMesh& UsdMesh, UStaticMes
 
 	StaticMesh = MeshCache.FindRef( MeshHash.ToString() );
 
-	if ( !StaticMesh )
+	if ( !StaticMesh && !MeshDescription.IsEmpty() )
 	{
 		StaticMesh = NewObject< UStaticMesh >( GetTransientPackage(), NAME_None, DefaultObjFlag | EObjectFlags::RF_Public );
 
diff --git a/Engine/Plugins/Importers/USDImporter/Source/USDStageEditor/Private/SUSDLayersTreeView.cpp b/Engine/Plugins/Importers/USDImporter/Source/USDStageEditor/Private/SUSDLayersTreeView.cpp
--- a/Engine/Plugins/Importers/USDImporter/Source/USDStageEditor/Private/SUSDLayersTreeView.cpp
+++ b/Engine/Plugins/Importers/USDImporter/Source/USDStageEditor/Private/SUSDLayersTreeView.cpp
@@ -366,7 +366,7 @@ void SUsdLayersTreeView::SetupColumns()
 	AddColumn( TEXT("Edit"), LOCTEXT( "Edit", "Edit" ), LayerEditColumn );
 }
 
-TSharedPtr< SWidget > SUsdLayersTreeView::ConstructLayerContextMenu() const
+TSharedPtr< SWidget > SUsdLayersTreeView::ConstructLayerContextMenu()
 {
 	TSharedRef< SWidget > MenuWidget = SNullWidget::NullWidget;
 
diff --git a/Engine/Plugins/Importers/USDImporter/Source/USDStageEditor/Private/SUSDLayersTreeView.h b/Engine/Plugins/Importers/USDImporter/Source/USDStageEditor/Private/SUSDLayersTreeView.h
--- a/Engine/Plugins/Importers/USDImporter/Source/USDStageEditor/Private/SUSDLayersTreeView.h
+++ b/Engine/Plugins/Importers/USDImporter/Source/USDStageEditor/Private/SUSDLayersTreeView.h
@@ -40,7 +40,7 @@ private:
 
 	void BuildUsdLayersEntries( AUsdStageActor* UsdStageActor );
 
-	TSharedPtr< SWidget > ConstructLayerContextMenu() const;
+	TSharedPtr< SWidget > ConstructLayerContextMenu();
 
 	bool CanEditLayer( FUsdLayersTreeItemRef LayerItem ) const;
 	bool CanEditSelectedLayer() const;
diff --git a/Engine/Plugins/Importers/USDImporter/Source/USDStageEditor/Private/SUSDStageTreeView.cpp b/Engine/Plugins/Importers/USDImporter/Source/USDStageEditor/Private/SUSDStageTreeView.cpp
--- a/Engine/Plugins/Importers/USDImporter/Source/USDStageEditor/Private/SUSDStageTreeView.cpp
+++ b/Engine/Plugins/Importers/USDImporter/Source/USDStageEditor/Private/SUSDStageTreeView.cpp
@@ -495,7 +495,7 @@ void SUsdStageTreeView::SetupColumns()
 	AddColumn( TEXT("Payload"), LOCTEXT( "Payload", "Payload" ), MakeShared< FUsdStagePayloadColumn >(), PayloadColumnArguments );
 }
 
-TSharedPtr< SWidget > SUsdStageTreeView::ConstructPrimContextMenu() const
+TSharedPtr< SWidget > SUsdStageTreeView::ConstructPrimContextMenu()
 {
 	TSharedRef< SWidget > MenuWidget = SNullWidget::NullWidget;
 
diff --git a/Engine/Plugins/Importers/USDImporter/Source/USDStageEditor/Private/SUSDStageTreeView.h b/Engine/Plugins/Importers/USDImporter/Source/USDStageEditor/Private/SUSDStageTreeView.h
--- a/Engine/Plugins/Importers/USDImporter/Source/USDStageEditor/Private/SUSDStageTreeView.h
+++ b/Engine/Plugins/Importers/USDImporter/Source/USDStageEditor/Private/SUSDStageTreeView.h
@@ -53,7 +53,7 @@ private:
 	void OnPrimNameCommitted( const FUsdStageTreeItemRef& TreeItem, const FText& InPrimName );
 
 	virtual void SetupColumns() override;
-	TSharedPtr< SWidget > ConstructPrimContextMenu() const;
+	TSharedPtr< SWidget > ConstructPrimContextMenu();
 
 	void OnToggleAllPayloads( EPayloadsTrigger PayloadsTrigger );
 
diff --git a/Engine/Plugins/Importers/USDImporter/Source/USDStageEditor/Private/Widgets/SUSDTreeView.h b/Engine/Plugins/Importers/USDImporter/Source/USDStageEditor/Private/Widgets/SUSDTreeView.h
--- a/Engine/Plugins/Importers/USDImporter/Source/USDStageEditor/Private/Widgets/SUSDTreeView.h
+++ b/Engine/Plugins/Importers/USDImporter/Source/USDStageEditor/Private/Widgets/SUSDTreeView.h
@@ -80,7 +80,7 @@ inline void SUsdTreeView< ItemType >::Construct( const FArguments& InArgs )
 
 	STreeView< ItemType >::Construct
 	(
-		STreeView< ItemType >::FArguments()
+		typename STreeView< ItemType >::FArguments()
 		.TreeItemsSource( &RootItems )
 		.OnGenerateRow( this, &SUsdTreeView< ItemType >::OnGenerateRow )
 		.OnGetChildren( this, &SUsdTreeView< ItemType >::OnGetChildren )
@@ -101,7 +101,7 @@ public:
 		TreeItem = InTreeItem;
 		SharedData = InSharedData;
 
-		SMultiColumnTableRow< ItemType >::Construct( SMultiColumnTableRow< ItemType >::FArguments(), OwnerTable );
+		SMultiColumnTableRow< ItemType >::Construct( typename SMultiColumnTableRow< ItemType >::FArguments(), OwnerTable );
 	}
 
 	/** Overridden from SMultiColumnTableRow.  Generates a widget for this column of the tree row. */
diff --git a/Engine/Plugins/Importers/USDImporter/Source/USDUtilities/Private/USDGeomMeshConversion.cpp b/Engine/Plugins/Importers/USDImporter/Source/USDUtilities/Private/USDGeomMeshConversion.cpp
--- a/Engine/Plugins/Importers/USDImporter/Source/USDUtilities/Private/USDGeomMeshConversion.cpp
+++ b/Engine/Plugins/Importers/USDImporter/Source/USDUtilities/Private/USDGeomMeshConversion.cpp
@@ -157,16 +157,22 @@ bool UsdToUnreal::ConvertGeomMesh( const pxr::UsdGeomMesh& UsdMesh, FMeshDescrip
 				{
 					UVSet.UVIndices.Emplace();
 
-					if ( PrimvarST.GetIndices( &UVSet.UVIndices.GetValue() ) && PrimvarST.Get( &UVSet.UVs ) )
+					if ( PrimvarST.GetIndices( &UVSet.UVIndices.GetValue(), TimeCode ) && PrimvarST.Get( &UVSet.UVs, TimeCode ) )
 					{
-						UVSets.Add( MoveTemp( UVSet ) );
+						if ( UVSet.UVs.size() > 0 )
+						{
+							UVSets.Add( MoveTemp( UVSet ) );
+						}
 					}
 				}
 				else
 				{
 					if ( PrimvarST.Get( &UVSet.UVs ) )
 					{
-						UVSets.Add( MoveTemp( UVSet ) );
+						if ( UVSet.UVs.size() > 0 )
+						{
+							UVSets.Add( MoveTemp( UVSet ) );
+						}
 					}
 				}
 			}
@@ -422,6 +428,35 @@ bool UsdToUnreal::ConvertMaterial( const pxr::UsdShadeMaterial& UsdShadeMaterial
 
 	bool bHasMaterialInfo = false;
 
+	auto ParseFloatInput = [ &Material, &SurfaceShader ]( const ANSICHAR* InputName, float DefaultValue ) -> UMaterialExpression*
+	{
+		pxr::UsdShadeInput Input = SurfaceShader.GetInput( pxr::TfToken(InputName) );
+		UMaterialExpression* InputExpression = nullptr;
+
+		if ( Input )
+		{
+			if ( !Input.HasConnectedSource() )
+			{
+				float InputValue = DefaultValue;
+				Input.Get< float >( &InputValue );
+
+				if ( !FMath::IsNearlyEqual( InputValue, DefaultValue ) )
+				{
+					UMaterialExpressionConstant* ExpressionConstant = Cast< UMaterialExpressionConstant >( UMaterialEditingLibrary::CreateMaterialExpression( &Material, UMaterialExpressionConstant::StaticClass() ) );
+					ExpressionConstant->R = InputValue;
+
+					InputExpression = ExpressionConstant;
+				}
+			}
+			else
+			{
+				InputExpression = UsdGeomMeshConversionImpl::ParseInputTexture( Input, Material );
+			}
+		}
+
+		return InputExpression;
+	};
+
 	// Base color
 	{
 		pxr::UsdShadeInput DiffuseInput = SurfaceShader.GetInput( pxr::TfToken("diffuseColor") );
@@ -457,63 +492,35 @@ bool UsdToUnreal::ConvertMaterial( const pxr::UsdShadeMaterial& UsdShadeMaterial
 	
 	// Metallic
 	{
-		pxr::UsdShadeInput MetallicInput = SurfaceShader.GetInput( pxr::TfToken("metallic") );
+		UMaterialExpression* MetallicExpression = ParseFloatInput( "metallic", 0.f );
 
-		if ( MetallicInput )
+		if ( MetallicExpression )
 		{
-			UMaterialExpression* MetallicExpression = nullptr;
-
-			if ( !MetallicInput.HasConnectedSource() )
-			{
-				float Metallic = 0.f;
-				MetallicInput.Get< float >( &Metallic );
-
-				UMaterialExpressionConstant* ExpressionConstant = Cast< UMaterialExpressionConstant >( UMaterialEditingLibrary::CreateMaterialExpression( &Material, UMaterialExpressionConstant::StaticClass() ) );
-				ExpressionConstant->R = Metallic;
-
-				MetallicExpression = ExpressionConstant;
-			}
-			else
-			{
-				MetallicExpression = UsdGeomMeshConversionImpl::ParseInputTexture( MetallicInput, Material );
-			}
-
-			if ( MetallicExpression )
-			{
-				Material.Metallic.Expression = MetallicExpression;
-				bHasMaterialInfo = true;
-			}
+			Material.Metallic.Expression = MetallicExpression;
+			bHasMaterialInfo = true;
 		}
 	}
 
 	// Roughness
 	{
-		pxr::UsdShadeInput RoughnessInput = SurfaceShader.GetInput( pxr::TfToken("roughness") );
+		UMaterialExpression* RoughnessExpression = ParseFloatInput( "roughness", 1.f );
 
-		if ( RoughnessInput )
+		if ( RoughnessExpression )
 		{
-			UMaterialExpression* RoughnessExpression = nullptr;
-
-			if ( !RoughnessInput.HasConnectedSource() )
-			{
-				float Roughness = 1.f;
-				SurfaceShader.GetInput( pxr::TfToken("roughness") ).Get< float >( &Roughness );
-
-				UMaterialExpressionConstant* ExpressionConstant = Cast< UMaterialExpressionConstant >( UMaterialEditingLibrary::CreateMaterialExpression( &Material, UMaterialExpressionConstant::StaticClass() ) );
-				ExpressionConstant->R = Roughness;
+			Material.Roughness.Expression = RoughnessExpression;
+			bHasMaterialInfo = true;
+		}
+	}
 
-				RoughnessExpression = ExpressionConstant;
-			}
-			else
-			{
-				RoughnessExpression = UsdGeomMeshConversionImpl::ParseInputTexture( RoughnessInput, Material );
-			}
+	// Opacity
+	{
+		UMaterialExpression* OpacityExpression = ParseFloatInput( "opacity", 1.f );;
 
-			if ( RoughnessExpression )
-			{
-				Material.Roughness.Expression = RoughnessExpression;
-				bHasMaterialInfo = true;
-			}
+		if ( OpacityExpression )
+		{
+			Material.Opacity.Expression = OpacityExpression;
+			Material.BlendMode = BLEND_Translucent;
+			bHasMaterialInfo = true;
 		}
 	}
 
diff --git a/Engine/Plugins/Importers/USDImporter/Source/USDUtilities/Private/USDSkeletalDataConversion.cpp b/Engine/Plugins/Importers/USDImporter/Source/USDUtilities/Private/USDSkeletalDataConversion.cpp
--- a/Engine/Plugins/Importers/USDImporter/Source/USDUtilities/Private/USDSkeletalDataConversion.cpp
+++ b/Engine/Plugins/Importers/USDImporter/Source/USDUtilities/Private/USDSkeletalDataConversion.cpp
@@ -172,11 +172,11 @@ bool UsdToUnreal::ConvertSkinnedMesh(const pxr::UsdSkelSkinningQuery& SkinningQu
 	}
 
 	// Face indices
-	VtArray<int> FaceIndices;
+	VtArray<int> OriginalFaceIndices;
 	UsdAttribute FaceIndicesAttribute = UsdMesh.GetFaceVertexIndicesAttr();
 	if (FaceIndicesAttribute)
 	{
-		FaceIndicesAttribute.Get(&FaceIndices, UsdTimeCode::Default());
+		FaceIndicesAttribute.Get(&OriginalFaceIndices, UsdTimeCode::Default());
 	}
 
 	// Normals
@@ -191,7 +191,7 @@ bool UsdToUnreal::ConvertSkinnedMesh(const pxr::UsdSkelSkinningQuery& SkinningQu
 	uint32 NumExistingWedges = SkelMeshImportData.Wedges.Num();
 
 	uint32 NumFaces = FaceCounts.size();
-	SkelMeshImportData.Faces.AddUninitialized(NumFaces);
+	SkelMeshImportData.Faces.Reserve( NumFaces * 2 );
 
 	// Retrieve prim materials
 	TTuple< TArray< FString >, TArray< int32 > > GeometryMaterials = IUsdPrim::GetGeometryMaterials( UsdTimeCode::Default().GetValue(), SkinningPrim );
@@ -244,220 +244,152 @@ bool UsdToUnreal::ConvertSkinnedMesh(const pxr::UsdSkelSkinningQuery& SkinningQu
 	SkelMeshImportData.NumTexCoords = 0;
 
 	bool bReverseOrder = IUsdPrim::GetGeometryOrientation(UsdMesh) == EUsdGeomOrientation::LeftHanded;
-	uint32 CurrentVertexInstanceIndex = 0;
 
-	// SkeletalMeshImportData uses triangle faces so quads will have to be split into triangles
-	SkeletalMeshImportData::FVertex TmpWedges[3];
-	TArray<SkeletalMeshImportData::FTriangle> SplitQuads;
-
-	for (uint32 PolygonIndex = NumExistingFaces, LocalIndex = 0; PolygonIndex < NumExistingFaces + NumFaces; ++PolygonIndex, ++LocalIndex)
+	struct FUVSet
 	{
-		uint32 NumFaceVertices = FaceCounts[LocalIndex];
-
-		// Face must be processed as triangle
-		bool bIsQuad = NumFaceVertices == 4;
-		NumFaceVertices = 3;
+		TOptional< VtIntArray > UVIndices; // UVs might be indexed or they might be flat (one per vertex)
+		VtVec2fArray UVs;
 
-		SkeletalMeshImportData::FTriangle& Triangle = SkelMeshImportData.Faces[PolygonIndex];
+		EUsdInterpolationMethod InterpolationMethod = EUsdInterpolationMethod::FaceVarying;
+	};
 
-		// Set the face smoothing by default. It could be any number, but not zero
-		Triangle.SmoothingGroups = 255;
+	TArray< FUVSet > UVSets;
 
-		// #ueent_todo: Convert normals to TangentZ (TangentX/TangentY are tangents/binormals)
+	int32 UVChannelIndex = 0;
+	while ( true )
+	{
+		pxr::TfToken UsdUVSetName = UsdUtils::GetUVSetName( UVChannelIndex ).Get();
+		UsdGeomPrimvar PrimvarST = UsdMesh.GetPrimvar( UsdUVSetName );
 
-		// Manage materials
-		int32 LocalMaterialIndex = 0;
-		if (LocalIndex >= 0 && LocalIndex < (uint32) FaceMaterialIndices.Num())
+		if ( PrimvarST )
 		{
-			LocalMaterialIndex = FaceMaterialIndices[LocalIndex];
-			if (LocalMaterialIndex < 0 || LocalMaterialIndex > MaterialNames.Num())
+			FUVSet UVSet;
+
+			if ( PrimvarST.GetInterpolation() == UsdGeomTokens->vertex )
 			{
-				LocalMaterialIndex = 0;
+				UVSet.InterpolationMethod = EUsdInterpolationMethod::Vertex;
 			}
-		}
-
-		int32 RealMaterialIndex = 0;
-		if (LocalMaterialIndex < MaterialNames.Num())
-		{
-			FString MaterialName = MaterialNames[LocalMaterialIndex];
-			if (!ExistingMaterialNames.Contains(MaterialName))
+			else if ( PrimvarST.GetInterpolation() == UsdGeomTokens->faceVarying )
 			{
-				// If new material, add it to the list
-				SkeletalMeshImportData::FMaterial NewMaterial;
-
-				NewMaterial.MaterialImportName = MaterialName;
-				RealMaterialIndex = SkelMeshImportData.Materials.Add(NewMaterial);
-				ExistingMaterialNames.Add(MaterialName, RealMaterialIndex);
+				UVSet.InterpolationMethod = EUsdInterpolationMethod::FaceVarying;
 			}
-			else
+			else if (  PrimvarST.GetInterpolation() == UsdGeomTokens->uniform )
 			{
-				RealMaterialIndex = ExistingMaterialNames[MaterialName];
+				UVSet.InterpolationMethod = EUsdInterpolationMethod::Uniform;
 			}
-		}
-
-		Triangle.MatIndex = RealMaterialIndex;
-		SkelMeshImportData.MaxMaterialIndex = FMath::Max<uint32>(SkelMeshImportData.MaxMaterialIndex, Triangle.MatIndex);
-
-		Triangle.AuxMatIndex = 0;
-
-		struct FUVSet
-		{
-			TOptional< VtIntArray > UVIndices; // UVs might be indexed or they might be flat (one per vertex)
-			VtVec2fArray UVs;
-
-			EUsdInterpolationMethod InterpolationMethod = EUsdInterpolationMethod::FaceVarying;
-		};
-
-		TArray< FUVSet > UVSets;
-
-		int32 UVChannelIndex = 0;
-		while ( true )
-		{
-			pxr::TfToken UsdUVSetName = UsdUtils::GetUVSetName( UVChannelIndex ).Get();
-			UsdGeomPrimvar PrimvarST = UsdMesh.GetPrimvar( UsdUVSetName );
-
-			if ( PrimvarST )
+			else if ( PrimvarST.GetInterpolation() == UsdGeomTokens->constant )
 			{
-				FUVSet UVSet;
+				UVSet.InterpolationMethod = EUsdInterpolationMethod::Constant;
+			}
 
-				if ( PrimvarST.GetInterpolation() == UsdGeomTokens->vertex )
-				{
-					UVSet.InterpolationMethod = EUsdInterpolationMethod::Vertex;
-				}
-				else if ( PrimvarST.GetInterpolation() == UsdGeomTokens->faceVarying )
-				{
-					UVSet.InterpolationMethod = EUsdInterpolationMethod::FaceVarying;
-				}
-				else if (  PrimvarST.GetInterpolation() == UsdGeomTokens->uniform )
-				{
-					UVSet.InterpolationMethod = EUsdInterpolationMethod::Uniform;
-				}
-				else if ( PrimvarST.GetInterpolation() == UsdGeomTokens->constant )
-				{
-					UVSet.InterpolationMethod = EUsdInterpolationMethod::Constant;
-				}
+			if ( PrimvarST.IsIndexed() )
+			{
+				UVSet.UVIndices.Emplace();
 
-				if ( PrimvarST.IsIndexed() )
+				if ( PrimvarST.GetIndices( &UVSet.UVIndices.GetValue() ) && PrimvarST.Get( &UVSet.UVs ) )
 				{
-					UVSet.UVIndices.Emplace();
-
-					if ( PrimvarST.GetIndices( &UVSet.UVIndices.GetValue() ) && PrimvarST.Get( &UVSet.UVs ) )
+					if ( UVSet.UVs.size() > 0 )
 					{
 						UVSets.Add( MoveTemp( UVSet ) );
 					}
 				}
-				else
+			}
+			else
+			{
+				if ( PrimvarST.Get( &UVSet.UVs ) )
 				{
-					if ( PrimvarST.Get( &UVSet.UVs ) )
+					if ( UVSet.UVs.size() > 0 )
 					{
 						UVSets.Add( MoveTemp( UVSet ) );
 					}
 				}
 			}
-			else
-			{
-				break;
-			}
-
-			++UVChannelIndex;
+		}
+		else
+		{
+			break;
 		}
 
-		// Manage vertex colors
-		FColor FaceColor = SkelMeshImportData.bHasVertexColors ? FaceColors[bIsConstantColor ? 0 : LocalIndex] : FColor::White;
+		++UVChannelIndex;
+	}
 
-		for (uint32 VertexIndex = 0; VertexIndex < NumFaceVertices; ++VertexIndex, ++CurrentVertexInstanceIndex)
+	// SkeletalMeshImportData uses triangle faces so quads will have to be split into triangles
+	SkeletalMeshImportData::FVertex TmpWedges[3];
+
+	for (uint32 PolygonIndex = NumExistingFaces, LocalIndex = 0; PolygonIndex < NumExistingFaces + NumFaces; ++PolygonIndex, ++LocalIndex)
+	{
+		const uint32 NumOriginalFaceVertices = FaceCounts[LocalIndex];
+		const uint32 NumFinalFaceVertices = 3;
+
+		// Face must be processed as triangle
+		bool bIsQuad = ( NumOriginalFaceVertices == 4 );
+
+		uint32 NumTriangles = bIsQuad ? 2 : 1;
+
+		for ( uint32 TriangleIndex = 0; TriangleIndex < NumTriangles; ++TriangleIndex )
 		{
-			int32 PointIndex = FaceIndices[CurrentVertexInstanceIndex];
+			int32 TriangleFaceIndex = SkelMeshImportData.Faces.AddUninitialized();
 
-			int32 UnrealVertexIndex = bReverseOrder ? NumFaceVertices - 1 - VertexIndex : VertexIndex;
+			SkeletalMeshImportData::FTriangle& Triangle = SkelMeshImportData.Faces[ TriangleFaceIndex ];
 
-			TmpWedges[UnrealVertexIndex].MatIndex = Triangle.MatIndex;
-			TmpWedges[UnrealVertexIndex].VertexIndex = NumExistingPoints + PointIndex;
-			
-			TmpWedges[UnrealVertexIndex].Color = FaceColor;
+			// Set the face smoothing by default. It could be any number, but not zero
+			Triangle.SmoothingGroups = 255;
 
-			int32 UVLayerIndex = 0;
-			for ( const FUVSet& UVSet : UVSets )
-			{
-				int32 ValueIndex = 0;
+			// #ueent_todo: Convert normals to TangentZ (TangentX/TangentY are tangents/binormals)
 
-				if ( UVSet.InterpolationMethod == EUsdInterpolationMethod::Vertex )
-				{
-					ValueIndex = UnrealVertexIndex;
-				}
-				else if ( UVSet.InterpolationMethod == EUsdInterpolationMethod::FaceVarying )
-				{
-					ValueIndex = CurrentVertexInstanceIndex;
-				}
-				else if ( UVSet.InterpolationMethod == EUsdInterpolationMethod::Uniform )
-				{
-					ValueIndex = PolygonIndex;
-				}
-				else if ( UVSet.InterpolationMethod == EUsdInterpolationMethod::Constant )
+			// Manage materials
+			int32 LocalMaterialIndex = 0;
+			if (LocalIndex >= 0 && LocalIndex < (uint32) FaceMaterialIndices.Num())
+			{
+				LocalMaterialIndex = FaceMaterialIndices[LocalIndex];
+				if (LocalMaterialIndex < 0 || LocalMaterialIndex > MaterialNames.Num())
 				{
-					ValueIndex = 0;
+					LocalMaterialIndex = 0;
 				}
+			}
 
-				GfVec2f UV( 0.f, 0.f );
-
-				if ( UVSet.UVIndices.IsSet() )
+			int32 RealMaterialIndex = 0;
+			if (LocalMaterialIndex < MaterialNames.Num())
+			{
+				FString MaterialName = MaterialNames[LocalMaterialIndex];
+				if (!ExistingMaterialNames.Contains(MaterialName))
 				{
-					if ( ensure( UVSet.UVIndices.GetValue().size() > ValueIndex ) )
-					{
-						UV = UVSet.UVs[ UVSet.UVIndices.GetValue()[ ValueIndex ] ];
-					}
+					// If new material, add it to the list
+					SkeletalMeshImportData::FMaterial NewMaterial;
+
+					NewMaterial.MaterialImportName = MaterialName;
+					RealMaterialIndex = SkelMeshImportData.Materials.Add(NewMaterial);
+					ExistingMaterialNames.Add(MaterialName, RealMaterialIndex);
 				}
-				else if ( ensure( UVSet.UVs.size() > ValueIndex ) )
+				else
 				{
-					UV = UVSet.UVs[ ValueIndex ];
+					RealMaterialIndex = ExistingMaterialNames[MaterialName];
 				}
-
-				// Flip V for Unreal uv's which match directx
-				FVector2D FinalUVVector( UV[0], 1.f - UV[1] );
-				TmpWedges[ UnrealVertexIndex ].UVs[ UVLayerIndex ] = FinalUVVector;
-
-				++UVLayerIndex;
 			}
 
-			Triangle.TangentX[UnrealVertexIndex] = FVector::ZeroVector;
-			Triangle.TangentY[UnrealVertexIndex] = FVector::ZeroVector;
-			Triangle.TangentZ[UnrealVertexIndex] = FVector::ZeroVector;
-		}
+			Triangle.MatIndex = RealMaterialIndex;
+			SkelMeshImportData.MaxMaterialIndex = FMath::Max<uint32>(SkelMeshImportData.MaxMaterialIndex, Triangle.MatIndex);
 
-		// Fill the wedge data and complete the triangle setup with the wedge indices
-		SkelMeshImportData.Wedges.Reserve( SkelMeshImportData.Wedges.Num() + NumFaceVertices );
-		for (uint32 VertexIndex = 0; VertexIndex < NumFaceVertices; ++VertexIndex)
-		{
-			uint32 WedgeIndex = SkelMeshImportData.Wedges.AddUninitialized();
-			SkelMeshImportData.Wedges[WedgeIndex].VertexIndex = TmpWedges[VertexIndex].VertexIndex;
-			SkelMeshImportData.Wedges[WedgeIndex].MatIndex = TmpWedges[VertexIndex].MatIndex;
-			SkelMeshImportData.Wedges[WedgeIndex].Color = TmpWedges[VertexIndex].Color;
-			SkelMeshImportData.Wedges[WedgeIndex].Reserved = 0;
-			FMemory::Memcpy(SkelMeshImportData.Wedges[WedgeIndex].UVs, TmpWedges[VertexIndex].UVs, sizeof(FVector2D) * MAX_TEXCOORDS);
-
-			Triangle.WedgeIndex[VertexIndex] = WedgeIndex;
-		}
+			Triangle.AuxMatIndex = 0;
 
-		// Do the same processing for the other half of the quad
-		if (bIsQuad)
-		{
-			SkeletalMeshImportData::FTriangle& SplitTriangle = SplitQuads.AddDefaulted_GetRef();
+			// Manage vertex colors
+			FColor FaceColor = SkelMeshImportData.bHasVertexColors ? FaceColors[bIsConstantColor ? 0 : LocalIndex] : FColor::White;
 
-			SplitTriangle.SmoothingGroups = 255;
-			SplitTriangle.MatIndex = RealMaterialIndex;
-			SplitTriangle.AuxMatIndex = 0;
+			// Fill the wedge data and complete the triangle setup with the wedge indices
+			SkelMeshImportData.Wedges.Reserve( SkelMeshImportData.Wedges.Num() + NumFinalFaceVertices );
 
-			uint32 SplitCurrentVertexInstanceIndex = CurrentVertexInstanceIndex - 3;
-			for (uint32 VertexIndex = 0; VertexIndex < NumFaceVertices; ++VertexIndex)
+			for ( uint32 CornerIndex = 0; CornerIndex < NumFinalFaceVertices; ++CornerIndex )
 			{
-				int32 LocalVertexInstanceIndex = SplitCurrentVertexInstanceIndex + VertexIndex + (VertexIndex > 0 ? 1 : 0);
-				int32 PointIndex = FaceIndices[LocalVertexInstanceIndex];
+				uint32 OriginalCornerIndex = ( ( TriangleIndex * ( NumOriginalFaceVertices - 2 ) ) + CornerIndex ) % NumOriginalFaceVertices;
+				uint32 OriginalVertexInstanceIndex = ( LocalIndex * NumOriginalFaceVertices ) + OriginalCornerIndex;
+				int32 OriginalVertexIndex = OriginalFaceIndices[ OriginalVertexInstanceIndex ];
 
-				int32 UnrealVertexIndex = bReverseOrder ? NumFaceVertices - 1 - VertexIndex : VertexIndex;
+				int32 FinalCornerIndex = bReverseOrder ? NumFinalFaceVertices - 1 - CornerIndex : CornerIndex;
 
-				TmpWedges[UnrealVertexIndex].MatIndex = SplitTriangle.MatIndex;
-				TmpWedges[UnrealVertexIndex].VertexIndex = NumExistingPoints + PointIndex;
-				TmpWedges[UnrealVertexIndex].Color = FaceColor;
+				TmpWedges[ FinalCornerIndex ].MatIndex = Triangle.MatIndex;
+				TmpWedges[ FinalCornerIndex ].VertexIndex = NumExistingPoints + OriginalVertexIndex;
+			
+				TmpWedges[ FinalCornerIndex ].Color = FaceColor;
 
 				int32 UVLayerIndex = 0;
 				for ( const FUVSet& UVSet : UVSets )
@@ -466,11 +398,11 @@ bool UsdToUnreal::ConvertSkinnedMesh(const pxr::UsdSkelSkinningQuery& SkinningQu
 
 					if ( UVSet.InterpolationMethod == EUsdInterpolationMethod::Vertex )
 					{
-						ValueIndex = UnrealVertexIndex;
+						ValueIndex = OriginalVertexIndex;
 					}
 					else if ( UVSet.InterpolationMethod == EUsdInterpolationMethod::FaceVarying )
 					{
-						ValueIndex = CurrentVertexInstanceIndex;
+						ValueIndex = OriginalVertexInstanceIndex;
 					}
 					else if ( UVSet.InterpolationMethod == EUsdInterpolationMethod::Uniform )
 					{
@@ -497,35 +429,29 @@ bool UsdToUnreal::ConvertSkinnedMesh(const pxr::UsdSkelSkinningQuery& SkinningQu
 
 					// Flip V for Unreal uv's which match directx
 					FVector2D FinalUVVector( UV[0], 1.f - UV[1] );
-					TmpWedges[ UnrealVertexIndex ].UVs[ UVLayerIndex ] = FinalUVVector;
+					TmpWedges[ FinalCornerIndex ].UVs[ UVLayerIndex ] = FinalUVVector;
 
 					++UVLayerIndex;
 				}
 
-				SplitTriangle.TangentX[VertexIndex] = FVector::ZeroVector;
-				SplitTriangle.TangentY[VertexIndex] = FVector::ZeroVector;
-				SplitTriangle.TangentZ[VertexIndex] = FVector::ZeroVector;
-			}
+				Triangle.TangentX[ FinalCornerIndex ] = FVector::ZeroVector;
+				Triangle.TangentY[ FinalCornerIndex ] = FVector::ZeroVector;
+				Triangle.TangentZ[ FinalCornerIndex ] = FVector::ZeroVector;
 
-			SkelMeshImportData.Wedges.Reserve( SkelMeshImportData.Wedges.Num() + NumFaceVertices );
-			for (uint32 VertexIndex = 0; VertexIndex < NumFaceVertices; ++VertexIndex)
-			{
-				uint32 WedgeIndex = SkelMeshImportData.Wedges.AddUninitialized();
-				SkelMeshImportData.Wedges[WedgeIndex].VertexIndex = TmpWedges[VertexIndex].VertexIndex;
-				SkelMeshImportData.Wedges[WedgeIndex].MatIndex = TmpWedges[VertexIndex].MatIndex;
-				SkelMeshImportData.Wedges[WedgeIndex].Color = TmpWedges[VertexIndex].Color;
-				SkelMeshImportData.Wedges[WedgeIndex].Reserved = 0;
-				FMemory::Memcpy(SkelMeshImportData.Wedges[WedgeIndex].UVs, TmpWedges[VertexIndex].UVs, sizeof(FVector2D)*MAX_TEXCOORDS);
-
-				SplitTriangle.WedgeIndex[VertexIndex] = WedgeIndex;
+				{
+					uint32 WedgeIndex = SkelMeshImportData.Wedges.AddUninitialized();
+					SkelMeshImportData.Wedges[ WedgeIndex ].VertexIndex = TmpWedges[ FinalCornerIndex ].VertexIndex;
+					SkelMeshImportData.Wedges[ WedgeIndex ].MatIndex = TmpWedges[ FinalCornerIndex ].MatIndex;
+					SkelMeshImportData.Wedges[ WedgeIndex ].Color = TmpWedges[ FinalCornerIndex ].Color;
+					SkelMeshImportData.Wedges[ WedgeIndex ].Reserved = 0;
+					FMemory::Memcpy( SkelMeshImportData.Wedges[ WedgeIndex ].UVs, TmpWedges[ FinalCornerIndex ].UVs, sizeof(FVector2D) * MAX_TEXCOORDS );
+
+					Triangle.WedgeIndex[ FinalCornerIndex ] = WedgeIndex;
+				}
 			}
-			++CurrentVertexInstanceIndex;
 		}
 	}
 
-	// Add the triangles that were split off from the quads
-	SkelMeshImportData.Faces.Append(SplitQuads);
-
 	// Convert joint influences into the SkeletalMeshImportData
 
 	// ComputeJointInfluences returns the influences per bone that applies to all the points of the mesh
diff --git a/Engine/Plugins/Importers/USDImporter/Source/UnrealUSDWrapper/Public/UnrealUSDWrapper.h b/Engine/Plugins/Importers/USDImporter/Source/UnrealUSDWrapper/Public/UnrealUSDWrapper.h
--- a/Engine/Plugins/Importers/USDImporter/Source/UnrealUSDWrapper/Public/UnrealUSDWrapper.h
+++ b/Engine/Plugins/Importers/USDImporter/Source/UnrealUSDWrapper/Public/UnrealUSDWrapper.h
@@ -35,7 +35,7 @@ PXR_NAMESPACE_CLOSE_SCOPE
 
 class IUsdPrim;
 
-enum EUsdInterpolationMethod
+enum class EUsdInterpolationMethod
 {
 	/** Each element in a buffer maps directly to a specific vertex */
 	Vertex,
@@ -47,7 +47,7 @@ enum EUsdInterpolationMethod
 	Constant
 };
 
-enum EUsdGeomOrientation
+enum class EUsdGeomOrientation
 {
 	/** Right handed coordinate system */
 	RightHanded,
@@ -55,7 +55,7 @@ enum EUsdGeomOrientation
 	LeftHanded,
 };
 
-enum EUsdSubdivisionScheme
+enum class EUsdSubdivisionScheme
 {
 	None,
 	CatmullClark,
@@ -64,7 +64,7 @@ enum EUsdSubdivisionScheme
 
 };
 
-enum EUsdUpAxis
+enum class EUsdUpAxis
 {
 	XAxis,
 	YAxis,
diff --git a/Engine/Source/Developer/MeshBuilder/Private/BuildOptimizationHelper.h b/Engine/Source/Developer/MeshBuilder/Private/BuildOptimizationHelper.h
--- a/Engine/Source/Developer/MeshBuilder/Private/BuildOptimizationHelper.h
+++ b/Engine/Source/Developer/MeshBuilder/Private/BuildOptimizationHelper.h
@@ -291,6 +291,11 @@ public:
 		TArray<int32>& WedgeMap
 	)
 	{
+		if (Vertices.Num() <= 0)
+		{
+			return;
+		}
+
 		// Copy the vertices since we will be reordering them
 		TArray<FStaticMeshBuildVertex> OriginalVertices = Vertices;
 
