USD
- Support display color and display opacity through material instances.

diff --git a/Engine/Plugins/Importers/USDImporter/Source/USDStage/Private/USDStageActor.cpp b/Engine/Plugins/Importers/USDImporter/Source/USDStage/Private/USDStageActor.cpp
--- a/Engine/Plugins/Importers/USDImporter/Source/USDStage/Private/USDStageActor.cpp
+++ b/Engine/Plugins/Importers/USDImporter/Source/USDStage/Private/USDStageActor.cpp
@@ -26,6 +26,8 @@
 #include "MaterialEditingLibrary.h"
 #include "Materials/Material.h"
 #include "Materials/MaterialInterface.h"
+#include "Materials/MaterialInstanceConstant.h"
+
 #include "MeshDescriptionOperations.h"
 #include "Misc/Paths.h"
 #include "Misc/ScopedSlowTask.h"
@@ -85,7 +87,7 @@ FMeshDescription LoadMeshDescription( const pxr::UsdGeomMesh& UsdMesh, const pxr
 	return MeshDescription;
 }
 
-void ProcessMaterials( const pxr::UsdStageRefPtr& Stage, UStaticMesh* StaticMesh, const FMeshDescription& MeshDescription, TMap< FString, UMaterial* >& MaterialsCache, bool bHasPrimDisplayColor )
+void ProcessMaterials( const pxr::UsdPrim& UsdPrim, UStaticMesh* StaticMesh, const FMeshDescription& MeshDescription, TMap< FString, UMaterial* >& MaterialsCache, bool bHasPrimDisplayColor, float Time )
 {
 	FStaticMeshConstAttributes StaticMeshAttributes( MeshDescription );
 
@@ -114,17 +116,17 @@ void ProcessMaterials( const pxr::UsdStageRefPtr& Stage, UStaticMesh* StaticMesh
 		}
 
 		UMaterialInterface* Material = nullptr;
-		pxr::UsdPrim MaterialPrim = Stage->GetPrimAtPath( UnrealToUsd::ConvertPath( *ImportedMaterialSlotName.ToString() ).Get() );
+		TUsdStore< pxr::UsdPrim > MaterialPrim = UsdPrim.GetStage()->GetPrimAtPath( UnrealToUsd::ConvertPath( *ImportedMaterialSlotName.ToString() ).Get() );
 
-		if ( MaterialPrim )
+		if ( MaterialPrim.Get() )
 		{
-			UMaterial*& CachedMaterial = MaterialsCache.FindOrAdd( UsdToUnreal::ConvertPath( MaterialPrim.GetPrimPath() ) );
+			UMaterial*& CachedMaterial = MaterialsCache.FindOrAdd( UsdToUnreal::ConvertPath( MaterialPrim.Get().GetPrimPath() ) );
 
 			if ( !CachedMaterial )
 			{
 				CachedMaterial = NewObject< UMaterial >();
 
-				if ( UsdToUnreal::ConvertMaterial( pxr::UsdShadeMaterial( MaterialPrim ), *CachedMaterial ) )
+				if ( UsdToUnreal::ConvertMaterial( pxr::UsdShadeMaterial( MaterialPrim.Get() ), *CachedMaterial ) )
 				{
 					//UMaterialEditingLibrary::RecompileMaterial( CachedMaterial ); // Too slow
 					CachedMaterial->PostEditChange();
@@ -140,8 +142,11 @@ void ProcessMaterials( const pxr::UsdStageRefPtr& Stage, UStaticMesh* StaticMesh
 
 		if ( Material == nullptr && bHasPrimDisplayColor )
 		{
-			FSoftObjectPath VertexColorMaterialPath( TEXT("Material'/Engine/EngineDebugMaterials/VertexColorMaterial.VertexColorMaterial'") );
-			Material = Cast< UMaterialInterface >( VertexColorMaterialPath.TryLoad() );
+			UMaterialInstanceConstant* MaterialInstance = NewObject< UMaterialInstanceConstant >();
+			if ( UsdToUnreal::ConvertDisplayColor( pxr::UsdGeomMesh( UsdPrim ), *MaterialInstance, pxr::UsdTimeCode( Time ) ) )
+			{
+				Material = MaterialInstance;
+			}
 		}
 
 		FStaticMaterial StaticMaterial( Material, MaterialSlotName, ImportedMaterialSlotName );
@@ -154,14 +159,14 @@ void ProcessMaterials( const pxr::UsdStageRefPtr& Stage, UStaticMesh* StaticMesh
 	}
 }
 
-void ProcessMaterials( const pxr::UsdStageRefPtr& Stage, FSkeletalMeshImportData& SkelMeshImportData, TMap< FString, UMaterial* >& MaterialsCache, bool bHasPrimDisplayColor )
+void ProcessMaterials( const pxr::UsdPrim& UsdPrim, FSkeletalMeshImportData& SkelMeshImportData, TMap< FString, UMaterial* >& MaterialsCache, bool bHasPrimDisplayColor, float Time )
 {
 	for (SkeletalMeshImportData::FMaterial& ImportedMaterial : SkelMeshImportData.Materials)
 	{
 		if (!ImportedMaterial.Material.IsValid())
 		{
 			UMaterialInterface* Material = nullptr;
-			TUsdStore< pxr::UsdPrim > MaterialPrim = Stage->GetPrimAtPath( UnrealToUsd::ConvertPath( *ImportedMaterial.MaterialImportName ).Get() );
+			TUsdStore< pxr::UsdPrim > MaterialPrim = UsdPrim.GetStage()->GetPrimAtPath( UnrealToUsd::ConvertPath( *ImportedMaterial.MaterialImportName ).Get() );
 
 			if ( MaterialPrim.Get() )
 			{
@@ -189,9 +194,13 @@ void ProcessMaterials( const pxr::UsdStageRefPtr& Stage, FSkeletalMeshImportData
 
 			if ( Material == nullptr && bHasPrimDisplayColor )
 			{
-				FSoftObjectPath VertexColorMaterialPath( TEXT("Material'/Engine/EngineDebugMaterials/VertexColorMaterial.VertexColorMaterial'") );
-				Material = Cast< UMaterialInterface >( VertexColorMaterialPath.TryLoad() );
+				UMaterialInstanceConstant* MaterialInstance = NewObject< UMaterialInstanceConstant >();
+				if ( UsdToUnreal::ConvertDisplayColor( pxr::UsdGeomMesh( UsdPrim ), *MaterialInstance, pxr::UsdTimeCode( Time ) ) )
+				{
+					Material = MaterialInstance;
+				}
 			}
+
 			ImportedMaterial.Material = Material;
 		}
 	}
@@ -265,7 +274,7 @@ bool AUsdStageActor::LoadStaticMesh( const pxr::UsdGeomMesh& UsdMesh, UStaticMes
 		StaticMesh->CommitMeshDescription(0);
 
 		const bool bHasPrimDisplayColor = UsdMesh.GetDisplayColorPrimvar().IsDefined();
-		ProcessMaterials( GetUsdStage(), StaticMesh, *StaticMeshDescription, MaterialsCache, bHasPrimDisplayColor );
+		ProcessMaterials( UsdMesh.GetPrim(), StaticMesh, *StaticMeshDescription, MaterialsCache, bHasPrimDisplayColor, Time );
 
 		// Create render data
 		if ( !StaticMesh->RenderData )
@@ -737,7 +746,7 @@ bool AUsdStageActor::ProcessSkeletonRoot( const pxr::UsdPrim& Prim, USkinnedMesh
 				PrimDelegates.Add(UsdToUnreal::ConvertPath(PrimPath), Handle);
 			}
 
-			ProcessMaterials(GetUsdStage(), SkelMeshImportData, MaterialsCache, bHasPrimDisplayColor);
+			ProcessMaterials(Prim, SkelMeshImportData, MaterialsCache, bHasPrimDisplayColor, Time);
 
 			break;
 		}
diff --git a/Engine/Plugins/Importers/USDImporter/Source/USDUtilities/Private/USDGeomMeshConversion.cpp b/Engine/Plugins/Importers/USDImporter/Source/USDUtilities/Private/USDGeomMeshConversion.cpp
--- a/Engine/Plugins/Importers/USDImporter/Source/USDUtilities/Private/USDGeomMeshConversion.cpp
+++ b/Engine/Plugins/Importers/USDImporter/Source/USDUtilities/Private/USDGeomMeshConversion.cpp
@@ -15,6 +15,7 @@
 #include "Materials/Material.h"
 #include "Materials/MaterialExpressionConstant.h"
 #include "Materials/MaterialExpressionConstant4Vector.h"
+#include "Materials/MaterialInstanceConstant.h"
 #include "MeshDescription.h"
 #include "Misc/Paths.h"
 #include "StaticMeshAttributes.h"
@@ -341,27 +342,61 @@ bool UsdToUnreal::ConvertGeomMesh( const pxr::UsdGeomMesh& UsdMesh, FMeshDescrip
 				return FLinearColor( FLinearColor( UsdToUnreal::ConvertColor( UsdColor ) ).ToFColor( false ) );
 			};
 
-			pxr::TfToken USDInterpType = ColorPrimvar.GetInterpolation();
-			if ( USDInterpType == pxr::UsdGeomTokens->faceVarying && NumColors >= MeshDescriptionColors.GetNumElements() )
+			pxr::TfToken UsdInterpType = ColorPrimvar.GetInterpolation();
+			if ( UsdInterpType == pxr::UsdGeomTokens->faceVarying && NumColors >= MeshDescriptionColors.GetNumElements() )
 			{
 				for(int Index = 0; Index < MeshDescriptionColors.GetNumElements(); ++Index)
 				{
-					MeshDescriptionColors[FVertexInstanceID(Index)] = ConvertToLinear(UsdColors[Index]);
+					MeshDescriptionColors[ FVertexInstanceID(Index) ] = ConvertToLinear( UsdColors[Index] );
 				}
 			}
-			else if ( USDInterpType == pxr::UsdGeomTokens->vertex && NumColors >= MeshDescription.Vertices().Num() )
+			else if ( UsdInterpType == pxr::UsdGeomTokens->vertex && NumColors >= MeshDescription.Vertices().Num() )
 			{
-				for( auto VertexInstID : MeshDescription.VertexInstances().GetElementIDs() )
+				for( FVertexInstanceID VertexInstID : MeshDescription.VertexInstances().GetElementIDs() )
 				{
-					FVertexID VertexID = MeshDescription.GetVertexInstanceVertex(VertexInstID);
-					MeshDescriptionColors[VertexInstID] = ConvertToLinear(UsdColors[VertexID.GetValue()]);
+					FVertexID VertexID = MeshDescription.GetVertexInstanceVertex( VertexInstID );
+					MeshDescriptionColors[ VertexInstID ] = ConvertToLinear( UsdColors[ VertexID.GetValue() ] );
 				}
 			}
-			else if ( USDInterpType == pxr::UsdGeomTokens->constant && NumColors == 1 )
+			else if ( UsdInterpType == pxr::UsdGeomTokens->constant && NumColors == 1 )
 			{
 				for( int Index = 0; Index < MeshDescriptionColors.GetNumElements(); ++Index )
 				{
-					MeshDescriptionColors[FVertexInstanceID(Index)] = ConvertToLinear(UsdColors[0]);
+					MeshDescriptionColors[ FVertexInstanceID(Index) ] = ConvertToLinear( UsdColors[0] );
+				}
+			}
+		}
+
+		// Vertex opacity
+		UsdGeomPrimvar OpacityPrimvar = UsdMesh.GetDisplayOpacityPrimvar();
+		if ( OpacityPrimvar )
+		{
+			pxr::VtArray< float > UsdOpacities;
+			OpacityPrimvar.ComputeFlattened( &UsdOpacities );
+
+			const int32 NumOpacities = UsdOpacities.size();
+
+			pxr::TfToken UsdInterpType = OpacityPrimvar.GetInterpolation();
+			if ( UsdInterpType == pxr::UsdGeomTokens->faceVarying && NumOpacities >= MeshDescriptionColors.GetNumElements() )
+			{
+				for(int Index = 0; Index < MeshDescriptionColors.GetNumElements(); ++Index)
+				{
+					MeshDescriptionColors[ FVertexInstanceID(Index) ][3] = UsdOpacities[Index];
+				}
+			}
+			else if ( UsdInterpType == pxr::UsdGeomTokens->vertex && NumOpacities >= MeshDescription.Vertices().Num() )
+			{
+				for( FVertexInstanceID VertexInstID : MeshDescription.VertexInstances().GetElementIDs() )
+				{
+					FVertexID VertexID = MeshDescription.GetVertexInstanceVertex( VertexInstID );
+					MeshDescriptionColors[ VertexInstID ][3] = UsdOpacities[ VertexID.GetValue() ];
+				}
+			}
+			else if ( UsdInterpType == pxr::UsdGeomTokens->constant && NumOpacities == 1 )
+			{
+				for( int Index = 0; Index < MeshDescriptionColors.GetNumElements(); ++Index )
+				{
+					MeshDescriptionColors[ FVertexInstanceID(Index) ][3] = UsdOpacities[0];
 				}
 			}
 		}
@@ -406,6 +441,7 @@ namespace UsdGeomMeshConversionImpl
 					{
 						UMaterialExpressionTextureSample* TextureSample = Cast< UMaterialExpressionTextureSample >( UMaterialEditingLibrary::CreateMaterialExpression( &Material, UMaterialExpressionTextureSample::StaticClass() ) );
 						TextureSample->Texture = Texture2D;
+						TextureSample->SamplerType = UMaterialExpressionTextureBase::GetSamplerTypeForTexture( Texture2D );
 
 						Result = TextureSample;
 					}
@@ -527,6 +563,57 @@ bool UsdToUnreal::ConvertMaterial( const pxr::UsdShadeMaterial& UsdShadeMaterial
 	return bHasMaterialInfo;
 }
 
+bool UsdToUnreal::ConvertDisplayColor( const pxr::UsdGeomMesh& UsdMesh, UMaterialInstanceConstant& MaterialInstance, const pxr::UsdTimeCode TimeCode )
+{
+	FScopedUsdAllocs UsdAllocs;
+
+	pxr::VtArray< pxr::GfVec3f > UsdDisplayColors = UsdUtils::GetUsdValue< pxr::VtArray< pxr::GfVec3f > >( UsdMesh.GetDisplayColorAttr(), TimeCode );
+
+	if ( !UsdDisplayColors.empty() )
+	{
+		pxr::VtArray< float > UsdOpacities = UsdUtils::GetUsdValue< pxr::VtArray< float > >( UsdMesh.GetDisplayOpacityAttr(), TimeCode );
+
+		bool bHasTransparency = false;
+
+		for ( float Opacity : UsdOpacities )
+		{
+			bHasTransparency = !FMath::IsNearlyEqual( Opacity, 1.f );
+
+			if ( bHasTransparency )
+			{
+				break;
+			}
+		}
+
+		{
+			FScopedUnrealAllocs UnrealAllocs;
+
+			FSoftObjectPath DisplayColorMaterialPath( TEXT("Material'/USDImporter/Materials/DisplayColorAndOpacity.DisplayColorAndOpacity'") );
+			UMaterialInterface* DisplayColorAndOpacityMaterial = Cast< UMaterialInterface >( DisplayColorMaterialPath.TryLoad() );
+
+			UMaterialEditingLibrary::SetMaterialInstanceParent( &MaterialInstance, DisplayColorAndOpacityMaterial );
+
+			if ( bHasTransparency )
+			{
+				MaterialInstance.BasePropertyOverrides.bOverride_BlendMode = true;
+				MaterialInstance.BasePropertyOverrides.BlendMode = EBlendMode::BLEND_Translucent;
+			}
+
+			if ( UsdMesh.GetDoubleSidedAttr().IsDefined() )
+			{
+				MaterialInstance.BasePropertyOverrides.bOverride_TwoSided = true;
+				MaterialInstance.BasePropertyOverrides.TwoSided = UsdUtils::GetUsdValue< bool >( UsdMesh.GetDoubleSidedAttr(), TimeCode );
+			}
+		}
+
+		return true;
+	}
+	else
+	{
+		return false;
+	}
+}
+
 bool UnrealToUsd::ConvertStaticMesh( const UStaticMesh* StaticMesh, pxr::UsdGeomMesh& UsdMesh, const pxr::UsdTimeCode TimeCode )
 {
 	FScopedUsdAllocs UsdAllocs;
diff --git a/Engine/Plugins/Importers/USDImporter/Source/USDUtilities/Private/USDSkeletalDataConversion.cpp b/Engine/Plugins/Importers/USDImporter/Source/USDUtilities/Private/USDSkeletalDataConversion.cpp
--- a/Engine/Plugins/Importers/USDImporter/Source/USDUtilities/Private/USDSkeletalDataConversion.cpp
+++ b/Engine/Plugins/Importers/USDImporter/Source/USDUtilities/Private/USDSkeletalDataConversion.cpp
@@ -228,6 +228,7 @@ bool UsdToUnreal::ConvertSkinnedMesh(const pxr::UsdSkelSkinningQuery& SkinningQu
 
 		if(/*USDInterpType == pxr::UsdGeomTokens->vertex &&*/ NumColors >= NumFaces)
 		{
+			FaceColors.Reserve( NumColors );
 			for (uint32 Index = 0; Index < NumColors; ++Index)
 			{
 				FaceColors.Add(ConvertToColor(UsdColors[Index]));
@@ -241,6 +242,39 @@ bool UsdToUnreal::ConvertSkinnedMesh(const pxr::UsdSkelSkinningQuery& SkinningQu
 		SkelMeshImportData.bHasVertexColors = true;
 	}
 
+	// Retrieve vertex opacity
+	UsdGeomPrimvar OpacityPrimvar = UsdMesh.GetDisplayOpacityPrimvar();
+	if ( OpacityPrimvar )
+	{
+		pxr::VtArray< float > UsdOpacities;
+		OpacityPrimvar.ComputeFlattened( &UsdOpacities );
+
+		const uint32 NumOpacities = UsdOpacities.size();
+
+		pxr::TfToken UsdInterpType = OpacityPrimvar.GetInterpolation();
+		if ( /*UsdInterpType == pxr::UsdGeomTokens->vertex &&*/ NumOpacities >= NumFaces )
+		{
+			for (uint32 Index = 0; Index < NumOpacities; ++Index)
+			{
+				if ( !FaceColors.IsValidIndex( Index ) )
+				{
+					FaceColors.Add( FColor::White );
+				}
+
+				FaceColors[ Index ].A = UsdOpacities[ Index ];
+			}
+		}
+		else if ( /*UsdInterpType == pxr::UsdGeomTokens->constant &&*/ NumOpacities == 1 )
+		{
+			if ( FaceColors.Num() < 1 )
+			{
+				FaceColors.Add( FColor::White );
+			}
+
+			FaceColors[ 0 ].A = UsdOpacities[ 0 ];
+		}
+	}
+
 	SkelMeshImportData.NumTexCoords = 0;
 
 	bool bReverseOrder = IUsdPrim::GetGeometryOrientation(UsdMesh) == EUsdGeomOrientation::LeftHanded;
diff --git a/Engine/Plugins/Importers/USDImporter/Source/USDUtilities/Public/USDGeomMeshConversion.h b/Engine/Plugins/Importers/USDImporter/Source/USDUtilities/Public/USDGeomMeshConversion.h
--- a/Engine/Plugins/Importers/USDImporter/Source/USDUtilities/Public/USDGeomMeshConversion.h
+++ b/Engine/Plugins/Importers/USDImporter/Source/USDUtilities/Public/USDGeomMeshConversion.h
@@ -20,12 +20,16 @@ PXR_NAMESPACE_CLOSE_SCOPE
 struct FMeshDescription;
 struct FStaticMeshLODResources;
 class UMaterial;
+class UMaterialInstanceConstant;
 class UStaticMesh;
 
 namespace UsdToUnreal
 {
 	USDUTILITIES_API bool ConvertGeomMesh( const pxr::UsdGeomMesh& UsdMesh, FMeshDescription& MeshDescription, const pxr::UsdTimeCode TimeCode = pxr::UsdTimeCode::EarliestTime() );
 	USDUTILITIES_API bool ConvertMaterial( const pxr::UsdShadeMaterial& UsdMesh, UMaterial& Material );
+
+	/** Reads the first display color and opacity value and assigns it as the base color and opacity of the material */
+	USDUTILITIES_API bool ConvertDisplayColor( const pxr::UsdGeomMesh& UsdMesh, UMaterialInstanceConstant& MaterialInstance, const pxr::UsdTimeCode TimeCode = pxr::UsdTimeCode::EarliestTime() );
 }
 
 namespace UnrealToUsd
