Copied cl 9880835 from Release-4.24

diff --git a/Engine/Plugins/Importers/USDImporter/Source/USDStage/Private/USDStageActor.cpp b/Engine/Plugins/Importers/USDImporter/Source/USDStage/Private/USDStageActor.cpp
--- a/Engine/Plugins/Importers/USDImporter/Source/USDStage/Private/USDStageActor.cpp
+++ b/Engine/Plugins/Importers/USDImporter/Source/USDStage/Private/USDStageActor.cpp
@@ -222,11 +222,6 @@ AUsdStageActor::AUsdStageActor()
 	UsdListener.OnPrimChanged.AddLambda(
 		[ this ]( const FString& PrimPath, bool bResync )
 			{
-				if ( this->PrimBeingUpdated.IsSet() && this->PrimBeingUpdated.GetValue() == PrimPath )
-				{
-					return;
-				}
-
 				TUsdStore< pxr::SdfPath > UsdPrimPath = UnrealToUsd::ConvertPath( *PrimPath );
 				this->UpdatePrim( UsdPrimPath.Get(), bResync );
 
@@ -1079,12 +1074,10 @@ void AUsdStageActor::OnPrimObjectPropertyChanged( UObject* ObjectBeingModified,
 
 			if ( UsdStage )
 			{
-				PrimBeingUpdated = PrimPath;
+				FScopedBlockNotices BlockNotices( UsdListener );
 
 				TUsdStore< pxr::UsdPrim > UsdPrim = UsdStage->GetPrimAtPath( UnrealToUsd::ConvertPath( *PrimPath ).Get() );
 				UnrealToUsd::ConvertSceneComponent( UsdStage, PrimSceneComponent, UsdPrim.Get() );
-
-				PrimBeingUpdated.Reset();
 			}
 		}
 	}
diff --git a/Engine/Plugins/Importers/USDImporter/Source/USDStage/Public/USDStageActor.h b/Engine/Plugins/Importers/USDImporter/Source/USDStage/Public/USDStageActor.h
--- a/Engine/Plugins/Importers/USDImporter/Source/USDStage/Public/USDStageActor.h
+++ b/Engine/Plugins/Importers/USDImporter/Source/USDStage/Public/USDStageActor.h
@@ -114,8 +114,6 @@ private:
 
 	TMap< UObject*, FString > ObjectsToWatch;
 
-	TOptional< FString > PrimBeingUpdated; // Ignore prim updates if we are the one updating the USD prim
-
 private:
 	UPROPERTY( NonPIEDuplicateTransient )
 	TMap< FString, UStaticMesh* > MeshCache;
diff --git a/Engine/Plugins/Importers/USDImporter/Source/USDStageEditor/Private/SUSDStageInfo.cpp b/Engine/Plugins/Importers/USDImporter/Source/USDStageEditor/Private/SUSDStageInfo.cpp
--- a/Engine/Plugins/Importers/USDImporter/Source/USDStageEditor/Private/SUSDStageInfo.cpp
+++ b/Engine/Plugins/Importers/USDImporter/Source/USDStageEditor/Private/SUSDStageInfo.cpp
@@ -26,6 +26,8 @@
 
 void SUsdStageInfo::Construct( const FArguments& InArgs, AUsdStageActor* UsdStageActor )
 {
+	OnInitialLoadSetChanged = InArgs._OnInitialLoadSetChanged;
+
 	InitialLoadSetStrings.Reset();
 	InitialLoadSetStrings.Add( MakeShared< FString >( TEXT("Load All") ) );
 	InitialLoadSetStrings.Add( MakeShared< FString >( TEXT("Load None") ) );
