USD Stage Editor: Fixed selected color of composition arcs being unreadable.

diff --git a/Engine/Plugins/Importers/USDImporter/Source/USDStageEditor/Private/SUSDStageTreeView.cpp b/Engine/Plugins/Importers/USDImporter/Source/USDStageEditor/Private/SUSDStageTreeView.cpp
--- a/Engine/Plugins/Importers/USDImporter/Source/USDStageEditor/Private/SUSDStageTreeView.cpp
+++ b/Engine/Plugins/Importers/USDImporter/Source/USDStageEditor/Private/SUSDStageTreeView.cpp
@@ -163,6 +163,8 @@ public:
 	DECLARE_DELEGATE_TwoParams( FOnPrimNameCommitted, const FUsdStageTreeItemRef&, const FText& );
 	FOnPrimNameCommitted OnPrimNameCommitted;
 
+	TWeakPtr< SUsdStageTreeView > OwnerTree;
+
 	virtual TSharedRef< SWidget > GenerateWidget( const TSharedPtr< IUsdTreeViewItem > InTreeItem ) override
 	{
 		if ( !InTreeItem )
@@ -205,7 +207,15 @@ protected:
 
 		if ( TreeItem && TreeItem->RowData->HasCompositionArcs() )
 		{
-			TextColor = FLinearColor( FColor::Orange );
+			const TSharedPtr< SUsdStageTreeView > OwnerTreePtr = OwnerTree.Pin();
+			if ( OwnerTreePtr && OwnerTreePtr->IsItemSelected( TreeItem.ToSharedRef() ) )
+			{
+				TextColor = FSlateColor::UseForeground();
+			}
+			else
+			{
+				TextColor = FLinearColor( FColor::Orange );
+			}
 		}
 		
 		return TextColor;
@@ -477,6 +487,7 @@ void SUsdStageTreeView::SetupColumns()
 	
 	{
 		TSharedRef< FUsdStageNameColumn > PrimNameColumn = MakeShared< FUsdStageNameColumn >();
+		PrimNameColumn->OwnerTree = SharedThis( this );
 		PrimNameColumn->bIsMainColumn = true;
 		PrimNameColumn->OnPrimNameCommitted.BindRaw( this, &SUsdStageTreeView::OnPrimNameCommitted );
 
diff --git a/Engine/Plugins/Importers/USDImporter/Source/USDStageEditor/Private/Widgets/SUSDTreeView.h b/Engine/Plugins/Importers/USDImporter/Source/USDStageEditor/Private/Widgets/SUSDTreeView.h
--- a/Engine/Plugins/Importers/USDImporter/Source/USDStageEditor/Private/Widgets/SUSDTreeView.h
+++ b/Engine/Plugins/Importers/USDImporter/Source/USDStageEditor/Private/Widgets/SUSDTreeView.h
@@ -24,7 +24,7 @@ public:
 	virtual TSharedRef< SWidget > GenerateWidget( const TSharedPtr< IUsdTreeViewItem > TreeItem ) = 0;
 
 public:
-	bool bIsMainColumn;
+	bool bIsMainColumn = false;
 };
 
 struct FSharedUsdTreeData
