USD
- Multithreaded the creation of the static meshes.
- Don't spawn actors/components for GeomScope.

diff --git a/Engine/Plugins/Enterprise/StaticMeshEditorExtension/Source/StaticMeshEditorExtension/Private/UVTools/SUVGenerationTool.cpp b/Engine/Plugins/Enterprise/StaticMeshEditorExtension/Source/StaticMeshEditorExtension/Private/UVTools/SUVGenerationTool.cpp
--- a/Engine/Plugins/Enterprise/StaticMeshEditorExtension/Source/StaticMeshEditorExtension/Private/UVTools/SUVGenerationTool.cpp
+++ b/Engine/Plugins/Enterprise/StaticMeshEditorExtension/Source/StaticMeshEditorExtension/Private/UVTools/SUVGenerationTool.cpp
@@ -416,6 +416,12 @@ void SGenerateUV::FitSettings()
 		
 		const int32 CurrentLOD = FMath::Max(0, GetSelectedLOD(EditorPtr));
 		const FMeshDescription* MeshDescription = StaticMesh->GetMeshDescription(CurrentLOD);
+
+		if ( !MeshDescription )
+		{
+			return;
+		}
+
 		FStaticMeshConstAttributes Attributes(*MeshDescription);
 		TMeshAttributesConstRef<FVertexID, FVector> VertexPositions = Attributes.GetVertexPositions();
 		TArray<FVector> RotatedVertexPositions;
diff --git a/Engine/Plugins/Importers/USDImporter/Source/USDImporter/Private/StaticMeshImporter.cpp b/Engine/Plugins/Importers/USDImporter/Source/USDImporter/Private/StaticMeshImporter.cpp
--- a/Engine/Plugins/Importers/USDImporter/Source/USDImporter/Private/StaticMeshImporter.cpp
+++ b/Engine/Plugins/Importers/USDImporter/Source/USDImporter/Private/StaticMeshImporter.cpp
@@ -164,7 +164,9 @@ UStaticMesh* FUSDStaticMeshImporter::ImportStaticMesh(FUsdImportContext& ImportC
 {
 	const pxr::UsdPrim& Prim = *PrimToImport.Prim;
 
-	FTransform PrimToWorld = ImportContext.bApplyWorldTransformToGeometry ? UsdToUnreal::ConvertMatrix(ImportContext.Stage.Get(), IUsdPrim::GetLocalTransform( Prim )) : FTransform::Identity;
+	const pxr::TfToken StageUpAxis = UsdUtils::GetUsdStageAxis( ImportContext.Stage.Get() );
+
+	FTransform PrimToWorld = ImportContext.bApplyWorldTransformToGeometry ? UsdToUnreal::ConvertMatrix(StageUpAxis, IUsdPrim::GetLocalTransform( Prim )) : FTransform::Identity;
 
 	FTransform FinalTransform = PrimToWorld;
 	if (ImportContext.ImportOptions->Scale != 1.0f)
diff --git a/Engine/Plugins/Importers/USDImporter/Source/USDImporter/Private/USDPrimResolver.cpp b/Engine/Plugins/Importers/USDImporter/Source/USDImporter/Private/USDPrimResolver.cpp
--- a/Engine/Plugins/Importers/USDImporter/Source/USDImporter/Private/USDPrimResolver.cpp
+++ b/Engine/Plugins/Importers/USDImporter/Source/USDImporter/Private/USDPrimResolver.cpp
@@ -225,6 +225,8 @@ AActor* UUSDPrimResolver::SpawnActor(FUSDSceneImportContext& ImportContext, cons
 					SpawnedActor = ImportContext.World->SpawnActor( AActor::StaticClass() );
 				}
 
+				const pxr::TfToken StageUpAxis = UsdUtils::GetUsdStageAxis( *ImportContext.Stage );
+
 				int32 AssetIndex = 0;
 				for ( UObject* ImportedAsset : ImportedAssets )
 				{
@@ -246,7 +248,7 @@ AActor* UUSDPrimResolver::SpawnActor(FUSDSceneImportContext& ImportContext, cons
 					for ( int32 ParentPrimIndex = ParentPrims.Num() - 1; ParentPrimIndex >= 0; --ParentPrimIndex )
 					{
 						ParentPrim = ParentPrims[ ParentPrimIndex ];
-						LocalTransform = UsdToUnreal::ConvertMatrix( *ImportContext.Stage, IUsdPrim::GetLocalTransform( ParentPrim ) ) * LocalTransform;
+						LocalTransform = UsdToUnreal::ConvertMatrix( StageUpAxis, IUsdPrim::GetLocalTransform( ParentPrim ) ) * LocalTransform;
 					}
 
 					FName ComponentBaseName = *UsdToUnreal::ConvertString( UsdAssetPrimToImport.Prim.Get().GetName().GetString().c_str() );
@@ -258,7 +260,7 @@ AActor* UUSDPrimResolver::SpawnActor(FUSDSceneImportContext& ImportContext, cons
 					// Don't add the prim transform if its the same prim used for the actor as it's already accounted for in the ActorTransform
 					if ( UsdAssetPrimToImport.Prim.Get() != SpawnData.ActorPrim.Get() )
 					{
-						LocalTransform = UsdToUnreal::ConvertMatrix( *ImportContext.Stage, IUsdPrim::GetLocalTransform( UsdAssetPrimToImport.Prim.Get() ) ) * LocalTransform;
+						LocalTransform = UsdToUnreal::ConvertMatrix( StageUpAxis, IUsdPrim::GetLocalTransform( UsdAssetPrimToImport.Prim.Get() ) ) * LocalTransform;
 					}
 
 					StaticMeshComponent->SetRelativeTransform( LocalTransform );
@@ -436,7 +438,7 @@ void UUSDPrimResolver::FindActorsToSpawn_Recursive(FUSDSceneImportContext& Impor
 
 		FName PrimName = UsdToUnreal::ConvertName(Prim.Get().GetName().GetString());
 		SpawnData.ActorName = PrimName;
-		SpawnData.WorldTransform = UsdToUnreal::ConvertMatrix( *ImportContext.Stage, IUsdPrim::GetLocalTransform( *Prim ) );
+		SpawnData.WorldTransform = UsdToUnreal::ConvertMatrix( UsdUtils::GetUsdStageAxis( *ImportContext.Stage ), IUsdPrim::GetLocalTransform( *Prim ) );
 		SpawnData.AttachParentPrim = ParentPrim;
 		SpawnData.ActorPrim = Prim;
 
diff --git a/Engine/Plugins/Importers/USDImporter/Source/USDImporter/Private/USDPrimResolverKind.cpp b/Engine/Plugins/Importers/USDImporter/Source/USDImporter/Private/USDPrimResolverKind.cpp
--- a/Engine/Plugins/Importers/USDImporter/Source/USDImporter/Private/USDPrimResolverKind.cpp
+++ b/Engine/Plugins/Importers/USDImporter/Source/USDImporter/Private/USDPrimResolverKind.cpp
@@ -65,7 +65,7 @@ void UUSDPrimResolverKind::FindActorsToSpawn_Recursive(FUSDSceneImportContext& I
 			UE_LOG(LogUSDImport, Log, TEXT("Adding %s Actor with %d meshes"), *PrimName.ToString(), SpawnData.AssetsToImport.Num());
 		}
 
-		SpawnData.WorldTransform = UsdToUnreal::ConvertMatrix( *ImportContext.Stage, IUsdPrim::GetLocalTransform( *Prim ) );
+		SpawnData.WorldTransform = UsdToUnreal::ConvertMatrix( UsdUtils::GetUsdStageAxis( *ImportContext.Stage ), IUsdPrim::GetLocalTransform( *Prim ) );
 		SpawnData.ActorPrim = Prim;
 		SpawnData.ActorName = PrimName;
 		SpawnData.AttachParentPrim = GroupParent;
@@ -87,7 +87,7 @@ void UUSDPrimResolverKind::FindActorsToSpawn_Recursive(FUSDSceneImportContext& I
 		SpawnData.ActorPrim = Prim;
 		SpawnData.ActorName = PrimName;
 
-		SpawnData.WorldTransform = UsdToUnreal::ConvertMatrix( *ImportContext.Stage, IUsdPrim::GetLocalTransform( *Prim ) );
+		SpawnData.WorldTransform = UsdToUnreal::ConvertMatrix( UsdUtils::GetUsdStageAxis( *ImportContext.Stage ), IUsdPrim::GetLocalTransform( *Prim ) );
 		SpawnData.AttachParentPrim = GroupParent;
 
 		// New parent of all children is this prim
diff --git a/Engine/Plugins/Importers/USDImporter/Source/USDStage/Private/USDListener.cpp b/Engine/Plugins/Importers/USDImporter/Source/USDStage/Private/USDListener.cpp
--- a/Engine/Plugins/Importers/USDImporter/Source/USDStage/Private/USDListener.cpp
+++ b/Engine/Plugins/Importers/USDImporter/Source/USDStage/Private/USDListener.cpp
@@ -48,28 +48,30 @@ void FUsdListener::HandleUsdNotice( const pxr::UsdNotice::ObjectsChanged& Notice
 		return;
 	}
 
+	auto BroadcastPrimUpdate = [ this ]( const pxr::SdfPath& PrimPath, bool bResync )
+	{
+		FScopedUnrealAllocs UnrealAllocs;
+
+		//FPlatformMisc::LowLevelOutputDebugStringf( TEXT("Prim %s changed (%s)\n"), *UsdToUnreal::ConvertPath( *It ), bResync ? TEXT("resync") : TEXT("update") );
+		OnPrimChanged.Broadcast( UsdToUnreal::ConvertPath( PrimPath ), bResync );
+	};
+
+	FScopedUsdAllocs UsdAllocs;
+
 	PathRange PathsToUpdate = Notice.GetResyncedPaths();
 
 	for ( PathRange::const_iterator It = PathsToUpdate.begin(); It != PathsToUpdate.end(); ++It )
 	{
-		//FPlatformMisc::LowLevelOutputDebugStringf( TEXT("Prim %s changed (resync)\n"), *UsdToUnreal::ConvertPath( *It ) );
-
-		FScopedUnrealAllocs UnrealAllocs;
-
 		constexpr bool bResync = true;
-		OnPrimChanged.Broadcast( UsdToUnreal::ConvertPath( *It ), bResync );
+		BroadcastPrimUpdate( It->GetAbsoluteRootOrPrimPath(), bResync );
 	}
 
 	PathsToUpdate = Notice.GetChangedInfoOnlyPaths();
 
 	for ( PathRange::const_iterator It = PathsToUpdate.begin(); It != PathsToUpdate.end(); ++It )
 	{
-		//FPlatformMisc::LowLevelOutputDebugStringf( TEXT("Prim %s changed (update)\n"), *UsdToUnreal::ConvertPath( *It ) );
-
-		FScopedUnrealAllocs UnrealAllocs;
-
 		constexpr bool bResync = false;
-		OnPrimChanged.Broadcast( UsdToUnreal::ConvertPath( *It ), bResync );
+		BroadcastPrimUpdate( It->GetAbsoluteRootOrPrimPath(), bResync );
 	}
 }
 
diff --git a/Engine/Plugins/Importers/USDImporter/Source/USDStage/Private/USDStageActor.cpp b/Engine/Plugins/Importers/USDImporter/Source/USDStage/Private/USDStageActor.cpp
--- a/Engine/Plugins/Importers/USDImporter/Source/USDStage/Private/USDStageActor.cpp
+++ b/Engine/Plugins/Importers/USDImporter/Source/USDStage/Private/USDStageActor.cpp
@@ -10,6 +10,7 @@
 #include "USDTypesConversion.h"
 #include "UnrealUSDWrapper.h"
 
+#include "Async/ParallelFor.h"
 #include "CineCameraActor.h"
 #include "CineCameraComponent.h"
 #include "Components/PoseableMeshComponent.h"
@@ -27,7 +28,7 @@
 #include "Materials/Material.h"
 #include "Materials/MaterialInterface.h"
 #include "Materials/MaterialInstanceConstant.h"
-
+#include "MeshDescription.h"
 #include "MeshDescriptionOperations.h"
 #include "Misc/Paths.h"
 #include "Misc/ScopedSlowTask.h"
@@ -38,7 +39,6 @@
 #include "Channels/MovieSceneChannelProxy.h"
 #include "LevelSequence.h"
 #include "LevelSequenceActor.h"
-#include "MeshDescription.h"
 #include "Misc/FrameNumber.h"
 #include "MovieScene.h"
 #include "Sections/MovieSceneFloatSection.h"
@@ -68,9 +68,16 @@
 
 #define LOCTEXT_NAMESPACE "USDStageActor"
 
+DEFINE_LOG_CATEGORY( LogUsdStage );
+
 static const EObjectFlags DefaultObjFlag = EObjectFlags::RF_Transactional | EObjectFlags::RF_Transient;
 
+AUsdStageActor::FOnActorLoaded AUsdStageActor::OnActorLoaded{};
+
 #if USE_USD_SDK
+
+namespace UsdStageActorImpl
+{
 FMeshDescription LoadMeshDescription( const pxr::UsdGeomMesh& UsdMesh, const pxr::UsdTimeCode TimeCode )
 {
 	if ( !UsdMesh )
@@ -87,86 +94,43 @@ FMeshDescription LoadMeshDescription( const pxr::UsdGeomMesh& UsdMesh, const pxr
 	return MeshDescription;
 }
 
-void ProcessMaterials( const pxr::UsdPrim& UsdPrim, UStaticMesh* StaticMesh, const FMeshDescription& MeshDescription, TMap< FString, UMaterial* >& MaterialsCache, bool bHasPrimDisplayColor, float Time )
-{
-	FStaticMeshConstAttributes StaticMeshAttributes( MeshDescription );
-
-	for ( const FPolygonGroupID PolygonGroupID : MeshDescription.PolygonGroups().GetElementIDs() )
+	void ProcessMaterials( const pxr::UsdPrim& UsdPrim, UStaticMesh& StaticMesh, TMap< FString, UMaterial* >& MaterialsCache, bool bHasPrimDisplayColor, float Time )
 	{
-		const FName& ImportedMaterialSlotName = StaticMeshAttributes.GetPolygonGroupMaterialSlotNames()[ PolygonGroupID ];
-		const FName MaterialSlotName = ImportedMaterialSlotName;
+		const FMeshDescription* MeshDescription = StaticMesh.GetMeshDescription( 0 );
 
-		int32 MaterialIndex = INDEX_NONE;
-		int32 MeshMaterialIndex = 0;
-
-		for ( FStaticMaterial& StaticMaterial : StaticMesh->StaticMaterials )
+		if ( !MeshDescription )
 		{
-			if ( StaticMaterial.MaterialSlotName.IsEqual( ImportedMaterialSlotName ) )
-			{
-				MaterialIndex = MeshMaterialIndex;
-				break;
-			}
-
-			++MeshMaterialIndex;
+			return;
 		}
 
-		if ( MaterialIndex == INDEX_NONE )
-		{
-			MaterialIndex = PolygonGroupID.GetValue();
-		}
+		FStaticMeshConstAttributes StaticMeshAttributes( *MeshDescription );
 
-		UMaterialInterface* Material = nullptr;
-		TUsdStore< pxr::UsdPrim > MaterialPrim = UsdPrim.GetStage()->GetPrimAtPath( UnrealToUsd::ConvertPath( *ImportedMaterialSlotName.ToString() ).Get() );
-
-		if ( MaterialPrim.Get() )
+		for ( const FPolygonGroupID PolygonGroupID : MeshDescription->PolygonGroups().GetElementIDs() )
 		{
-			UMaterial*& CachedMaterial = MaterialsCache.FindOrAdd( UsdToUnreal::ConvertPath( MaterialPrim.Get().GetPrimPath() ) );
+			const FName& ImportedMaterialSlotName = StaticMeshAttributes.GetPolygonGroupMaterialSlotNames()[ PolygonGroupID ];
+			const FName MaterialSlotName = ImportedMaterialSlotName;
 
-			if ( !CachedMaterial )
-			{
-				CachedMaterial = NewObject< UMaterial >();
+			int32 MaterialIndex = INDEX_NONE;
+			int32 MeshMaterialIndex = 0;
 
-				if ( UsdToUnreal::ConvertMaterial( pxr::UsdShadeMaterial( MaterialPrim.Get() ), *CachedMaterial ) )
-				{
-					//UMaterialEditingLibrary::RecompileMaterial( CachedMaterial ); // Too slow
-					CachedMaterial->PostEditChange();
-				}
-				else
+				for ( FStaticMaterial& StaticMaterial : StaticMesh.StaticMaterials )
+			{
+				if ( StaticMaterial.MaterialSlotName.IsEqual( ImportedMaterialSlotName ) )
 				{
-					CachedMaterial = nullptr;
+					MaterialIndex = MeshMaterialIndex;
+					break;
 				}
-			}
 
-			Material = CachedMaterial;
-		}
+				++MeshMaterialIndex;
+			}
 
-		if ( Material == nullptr && bHasPrimDisplayColor )
-		{
-			UMaterialInstanceConstant* MaterialInstance = NewObject< UMaterialInstanceConstant >();
-			if ( UsdToUnreal::ConvertDisplayColor( pxr::UsdGeomMesh( UsdPrim ), *MaterialInstance, pxr::UsdTimeCode( Time ) ) )
+			if ( MaterialIndex == INDEX_NONE )
 			{
-				Material = MaterialInstance;
+				MaterialIndex = PolygonGroupID.GetValue();
 			}
-		}
 
-		FStaticMaterial StaticMaterial( Material, MaterialSlotName, ImportedMaterialSlotName );
-		StaticMesh->StaticMaterials.Add( StaticMaterial );
-
-		FMeshSectionInfo MeshSectionInfo;
-		MeshSectionInfo.MaterialIndex = MaterialIndex;
-
-		StaticMesh->GetSectionInfoMap().Set( 0, PolygonGroupID.GetValue(), MeshSectionInfo );
-	}
-}
-
-void ProcessMaterials( const pxr::UsdPrim& UsdPrim, FSkeletalMeshImportData& SkelMeshImportData, TMap< FString, UMaterial* >& MaterialsCache, bool bHasPrimDisplayColor, float Time )
-{
-	for (SkeletalMeshImportData::FMaterial& ImportedMaterial : SkelMeshImportData.Materials)
-	{
-		if (!ImportedMaterial.Material.IsValid())
-		{
 			UMaterialInterface* Material = nullptr;
-			TUsdStore< pxr::UsdPrim > MaterialPrim = UsdPrim.GetStage()->GetPrimAtPath( UnrealToUsd::ConvertPath( *ImportedMaterial.MaterialImportName ).Get() );
+			TUsdStore< pxr::UsdPrim > MaterialPrim = UsdPrim.GetStage()->GetPrimAtPath( UnrealToUsd::ConvertPath( *ImportedMaterialSlotName.ToString() ).Get() );
 
 			if ( MaterialPrim.Get() )
 			{
@@ -178,10 +142,7 @@ void ProcessMaterials( const pxr::UsdPrim& UsdPrim, FSkeletalMeshImportData& Ske
 
 					if ( UsdToUnreal::ConvertMaterial( pxr::UsdShadeMaterial( MaterialPrim.Get() ), *CachedMaterial ) )
 					{
-						CachedMaterial->bUsedWithSkeletalMesh = true;
-						bool bNeedsRecompile = false;
-						CachedMaterial->GetMaterial()->SetMaterialUsage(bNeedsRecompile, MATUSAGE_SkeletalMesh);
-
+						//UMaterialEditingLibrary::RecompileMaterial( CachedMaterial ); // Too slow
 						CachedMaterial->PostEditChange();
 					}
 					else
@@ -189,6 +150,7 @@ void ProcessMaterials( const pxr::UsdPrim& UsdPrim, FSkeletalMeshImportData& Ske
 						CachedMaterial = nullptr;
 					}
 				}
+
 				Material = CachedMaterial;
 			}
 
@@ -201,7 +163,65 @@ void ProcessMaterials( const pxr::UsdPrim& UsdPrim, FSkeletalMeshImportData& Ske
 				}
 			}
 
-			ImportedMaterial.Material = Material;
+			FStaticMaterial StaticMaterial( Material, MaterialSlotName, ImportedMaterialSlotName );
+				StaticMesh.StaticMaterials.Add( StaticMaterial );
+
+			FMeshSectionInfo MeshSectionInfo;
+			MeshSectionInfo.MaterialIndex = MaterialIndex;
+
+			StaticMesh.GetSectionInfoMap().Set( 0, PolygonGroupID.GetValue(), MeshSectionInfo );
+		}
+	}
+
+	void ProcessMaterials( const pxr::UsdPrim& UsdPrim, FSkeletalMeshImportData& SkelMeshImportData, TMap< FString, UMaterial* >& MaterialsCache, bool bHasPrimDisplayColor, float Time )
+	{
+		for (SkeletalMeshImportData::FMaterial& ImportedMaterial : SkelMeshImportData.Materials)
+		{
+			if (!ImportedMaterial.Material.IsValid())
+			{
+				UMaterialInterface* Material = nullptr;
+				TUsdStore< pxr::UsdPrim > MaterialPrim = UsdPrim.GetStage()->GetPrimAtPath( UnrealToUsd::ConvertPath( *ImportedMaterial.MaterialImportName ).Get() );
+
+				if ( MaterialPrim.Get() )
+				{
+					UMaterial*& CachedMaterial = MaterialsCache.FindOrAdd( UsdToUnreal::ConvertPath( MaterialPrim.Get().GetPrimPath() ) );
+
+					if ( !CachedMaterial )
+					{
+						CachedMaterial = NewObject< UMaterial >();
+
+						if ( UsdToUnreal::ConvertMaterial( pxr::UsdShadeMaterial( MaterialPrim.Get() ), *CachedMaterial ) )
+						{
+							CachedMaterial->bUsedWithSkeletalMesh = true;
+							bool bNeedsRecompile = false;
+							CachedMaterial->GetMaterial()->SetMaterialUsage(bNeedsRecompile, MATUSAGE_SkeletalMesh);
+
+							CachedMaterial->PostEditChange();
+						}
+						else
+						{
+							CachedMaterial = nullptr;
+						}
+					}
+					Material = CachedMaterial;
+				}
+
+				if ( Material == nullptr && bHasPrimDisplayColor )
+				{
+					UMaterialInstanceConstant* MaterialInstance = NewObject< UMaterialInstanceConstant >();
+					if ( UsdToUnreal::ConvertDisplayColor( pxr::UsdGeomMesh( UsdPrim ), *MaterialInstance, pxr::UsdTimeCode( Time ) ) )
+					{
+						if ( MaterialInstance && MaterialInstance->HasOverridenBaseProperties() )
+						{
+							MaterialInstance->ForceRecompileForRendering();
+						}
+						
+						Material = MaterialInstance;
+					}
+				}
+
+				ImportedMaterial.Material = Material;
+			}
 		}
 	}
 }
@@ -223,7 +243,10 @@ AUsdStageActor::AUsdStageActor()
 		[ this ]( const FString& PrimPath, bool bResync )
 			{
 				TUsdStore< pxr::SdfPath > UsdPrimPath = UnrealToUsd::ConvertPath( *PrimPath );
-				this->UpdatePrim( UsdPrimPath.Get(), bResync );
+
+				TMap< FString, UStaticMesh* > PrimPathsToStaticMeshes = this->LoadStaticMeshes( this->GetUsdStage()->GetPrimAtPath( UsdPrimPath.Get() ) );
+
+				this->UpdatePrim( UsdPrimPath.Get(), bResync, PrimPathsToStaticMeshes );
 
 				this->OnPrimChanged.Broadcast( PrimPath, bResync );
 			}
@@ -236,84 +259,20 @@ AUsdStageActor::AUsdStageActor()
 	FCoreUObjectDelegates::OnObjectPropertyChanged.AddUObject( this, &AUsdStageActor::OnPrimObjectPropertyChanged );
 }
 
-#if USE_USD_SDK
-
-bool AUsdStageActor::LoadStaticMesh( const pxr::UsdGeomMesh& UsdMesh, UStaticMeshComponent& MeshComponent )
+AUsdStageActor::~AUsdStageActor()
 {
-	UStaticMesh* StaticMesh = nullptr;
+#if USE_USD_SDK
+	UnrealUSDWrapper::GetUsdStageCache().Erase( UsdStageStore.Get() );
+#endif // #if USE_USD_SDK
+}
 
-	pxr::UsdGeomXformable Xformable( UsdMesh.GetPrim() );
+#if USE_USD_SDK
 
+bool AUsdStageActor::SetStaticMesh( const pxr::UsdGeomMesh& UsdMesh, UStaticMeshComponent& MeshComponent, const TMap< FString, UStaticMesh* >& PrimPathsToStaticMeshes )
+{
 	FString MeshPrimPath = UsdToUnreal::ConvertPath( UsdMesh.GetPrim().GetPrimPath() );
 
-	FMeshDescription MeshDescription = LoadMeshDescription( UsdMesh, pxr::UsdTimeCode( Time ) );
-
-	FSHAHash MeshHash = FMeshDescriptionOperations::ComputeSHAHash( MeshDescription );
-
-	StaticMesh = MeshCache.FindRef( MeshHash.ToString() );
-
-	if ( !StaticMesh && !MeshDescription.IsEmpty() )
-	{
-		StaticMesh = NewObject< UStaticMesh >( GetTransientPackage(), NAME_None, DefaultObjFlag | EObjectFlags::RF_Public );
-
-		FStaticMeshSourceModel& SourceModel = StaticMesh->AddSourceModel();
-		SourceModel.BuildSettings.bGenerateLightmapUVs = false;
-		SourceModel.BuildSettings.bRecomputeNormals = false;
-		SourceModel.BuildSettings.bRecomputeTangents = false;
-		SourceModel.BuildSettings.bBuildAdjacencyBuffer = false;
-		SourceModel.BuildSettings.bBuildReversedIndexBuffer = false;
-
-		FMeshDescription* StaticMeshDescription = StaticMesh->CreateMeshDescription(0);
-		check( StaticMeshDescription );
-		*StaticMeshDescription = MoveTemp( MeshDescription );
-		StaticMesh->CommitMeshDescription(0);
-
-		const bool bHasPrimDisplayColor = UsdMesh.GetDisplayColorPrimvar().IsDefined();
-		ProcessMaterials( UsdMesh.GetPrim(), StaticMesh, *StaticMeshDescription, MaterialsCache, bHasPrimDisplayColor, Time );
-
-		// Create render data
-		if ( !StaticMesh->RenderData )
-		{
-			StaticMesh->RenderData.Reset(new(FMemory::Malloc(sizeof(FStaticMeshRenderData)))FStaticMeshRenderData());
-		}
-
-		ITargetPlatformManagerModule& TargetPlatformManager = GetTargetPlatformManagerRef();
-		ITargetPlatform* RunningPlatform = TargetPlatformManager.GetRunningTargetPlatform();
-		check(RunningPlatform);
-		const FStaticMeshLODSettings& LODSettings = RunningPlatform->GetStaticMeshLODSettings();
-
-		const FStaticMeshLODGroup& LODGroup = LODSettings.GetLODGroup(StaticMesh->LODGroup);
-
-		IMeshBuilderModule& MeshBuilderModule = FModuleManager::LoadModuleChecked< IMeshBuilderModule >( TEXT("MeshBuilder") );
-		if ( !MeshBuilderModule.BuildMesh( *StaticMesh->RenderData, StaticMesh, LODGroup ) )
-		{
-			UE_LOG(LogStaticMesh, Error, TEXT("Failed to build static mesh. See previous line(s) for details."));
-			return false;
-		}
-
-		for ( FStaticMeshLODResources& LODResources : StaticMesh->RenderData->LODResources )
-		{
-			LODResources.bHasColorVertexData = true;
-		}
-
-		// Bounds
-		StaticMesh->CalculateExtendedBounds();
-
-		// Recreate resources
-		StaticMesh->InitResources();
-
-		// Other setup
-		//StaticMesh->CreateBodySetup();
-
-		// Collision
-		//StaticMesh->BodySetup->CollisionTraceFlag = ECollisionTraceFlag::CTF_UseComplexAsSimple;
-
-		MeshCache.Add( MeshHash.ToString() ) = StaticMesh;
-	}
-	/*else
-	{
-		FPlatformMisc::LowLevelOutputDebugStringf( TEXT("Mesh found in cache %s\n"), *StaticMesh->GetName() );
-	}*/
+	UStaticMesh* StaticMesh = PrimPathsToStaticMeshes.FindRef( MeshPrimPath );
 
 	if ( StaticMesh != MeshComponent.GetStaticMesh() )
 	{
@@ -327,7 +286,7 @@ bool AUsdStageActor::LoadStaticMesh( const pxr::UsdGeomMesh& UsdMesh, UStaticMes
 		MeshComponent.RegisterComponent();
 	}
 
-	return true;
+	return ( StaticMesh != nullptr );
 }
 
 FUsdPrimTwin* AUsdStageActor::SpawnPrim( const pxr::SdfPath& UsdPrimPath )
@@ -360,41 +319,11 @@ FUsdPrimTwin* AUsdStageActor::SpawnPrim( const pxr::SdfPath& UsdPrimPath )
 				this->OnUsdPrimTwinDestroyed( UsdPrimTwin );
 			} );
 
-		if ( !UsdPrimTwin->AnimationHandle.IsValid() )
-		{
-			bool bHasXformbaleTimeSamples = false;
-			{
-				pxr::UsdGeomXformable Xformable( Prim );
-
-				if ( Xformable )
-				{
-					std::vector< double > TimeSamples;
-					Xformable.GetTimeSamples( &TimeSamples );
-
-					bHasXformbaleTimeSamples = TimeSamples.size() > 0;
-				}
-			}
-
-			bool bHasAttributesTimeSamples = false;
-			{
-				std::vector< pxr::UsdAttribute > Attributes = Prim.GetAttributes();
-
-				for ( pxr::UsdAttribute& Attribute : Attributes )
-				{
-					bHasAttributesTimeSamples = Attribute.ValueMightBeTimeVarying();
-					if ( bHasAttributesTimeSamples )
-					{
-						break;
-					}
-				}
-			}
-
-			if ( ( bHasXformbaleTimeSamples || bHasAttributesTimeSamples ) )
+		if ( !UsdPrimTwin->AnimationHandle.IsValid() && UsdUtils::IsAnimated( Prim ) )
 			{
 				PrimsToAnimate.Add( PrimPath );
 			}
 		}
-	}
 
 	UClass* ComponentClass = UsdUtils::GetComponentTypeForPrim( Prim );
 
@@ -403,7 +332,9 @@ FUsdPrimTwin* AUsdStageActor::SpawnPrim( const pxr::SdfPath& UsdPrimPath )
 		return nullptr;
 	}
 
-	const bool bNeedsActor = ( ParentUsdPrimTwin->SceneComponent == nullptr || Prim.IsModel() || UsdUtils::HasCompositionArcs( Prim ) || Prim.IsGroup() || Prim.IsA< pxr::UsdGeomScope >() || Prim.IsA< pxr::UsdGeomCamera >() );
+	FScopedUnrealAllocs UnrealAllocs;
+
+	const bool bNeedsActor = ( ParentUsdPrimTwin->SceneComponent == nullptr || Prim.IsModel() || UsdUtils::HasCompositionArcs( Prim ) || Prim.IsGroup() || Prim.IsA< pxr::UsdGeomCamera >() );
 
 	if ( bNeedsActor && !UsdPrimTwin->SpawnedActor.IsValid() )
 	{
@@ -502,7 +433,7 @@ FUsdPrimTwin* AUsdStageActor::SpawnPrim( const pxr::SdfPath& UsdPrimPath )
 	return UsdPrimTwin;
 }
 
-FUsdPrimTwin* AUsdStageActor::LoadPrim( const pxr::SdfPath& Path )
+FUsdPrimTwin* AUsdStageActor::LoadPrim( const pxr::SdfPath& Path, const TMap< FString, UStaticMesh* >& PrimPathsToStaticMeshes )
 {
 	TRACE_CPUPROFILER_EVENT_SCOPE( AUsdStageActor::LoadPrim );
 
@@ -537,7 +468,7 @@ FUsdPrimTwin* AUsdStageActor::LoadPrim( const pxr::SdfPath& Path )
 
 	if ( UStaticMeshComponent* MeshComponent = Cast< UStaticMeshComponent >( UsdPrimTwin->SceneComponent.Get() ) )
 	{
-		LoadStaticMesh( pxr::UsdGeomMesh( Prim ), *MeshComponent );
+		SetStaticMesh( pxr::UsdGeomMesh( Prim ), *MeshComponent, PrimPathsToStaticMeshes );
 	}
 	else if ( UCineCameraComponent* CameraComponent = Cast< UCineCameraComponent >( UsdPrimTwin->SceneComponent.Get() ) )
 	{
@@ -560,7 +491,16 @@ FUsdPrimTwin* AUsdStageActor::LoadPrim( const pxr::SdfPath& Path )
 		}
 
 		ComponentToAttach->AttachToComponent( AttachToComponent, FAttachmentTransformRules::KeepRelativeTransform );
-		UsdPrimTwin->SceneComponent->SetMobility( EComponentMobility::Movable );
+		
+		if ( PrimsToAnimate.Contains( UsdToUnreal::ConvertPath( Path.GetPrimPath() ) ) )
+		{
+			ComponentToAttach->SetMobility( EComponentMobility::Movable );
+		}
+		else
+		{
+			ComponentToAttach->SetMobility( EComponentMobility::Static );
+		}
+
 		UsdPrimTwin->SceneComponent->GetOwner()->AddInstanceComponent( UsdPrimTwin->SceneComponent.Get() );
 	}
 
@@ -602,7 +542,7 @@ FUsdPrimTwin* AUsdStageActor::LoadPrim( const pxr::SdfPath& Path )
 	return UsdPrimTwin;
 }
 
-FUsdPrimTwin* AUsdStageActor::ExpandPrim( const pxr::UsdPrim& Prim )
+FUsdPrimTwin* AUsdStageActor::ExpandPrim( const pxr::UsdPrim& Prim, const TMap< FString, UStaticMesh* >& PrimPathsToStaticMeshes )
 {
 	if ( !Prim )
 	{
@@ -611,7 +551,7 @@ FUsdPrimTwin* AUsdStageActor::ExpandPrim( const pxr::UsdPrim& Prim )
 
 	TRACE_CPUPROFILER_EVENT_SCOPE( AUsdStageActor::ExpandPrim );
 
-	FUsdPrimTwin* UsdPrimTwin = LoadPrim( Prim.GetPrimPath() );
+	FUsdPrimTwin* UsdPrimTwin = LoadPrim( Prim.GetPrimPath(), PrimPathsToStaticMeshes );
 
 	if ( Prim.IsA<pxr::UsdSkelRoot>() )
 	{
@@ -626,7 +566,7 @@ FUsdPrimTwin* AUsdStageActor::ExpandPrim( const pxr::UsdPrim& Prim )
 		pxr::UsdPrimSiblingRange PrimChildren = Prim.GetFilteredChildren( pxr::UsdTraverseInstanceProxies() );
 		for ( TUsdStore< pxr::UsdPrim > ChildStore : PrimChildren )
 		{
-			ExpandPrim( ChildStore.Get() );
+			ExpandPrim( ChildStore.Get(), PrimPathsToStaticMeshes );
 		}
 	}
 
@@ -658,6 +598,8 @@ bool AUsdStageActor::ProcessSkeletonRoot( const pxr::UsdPrim& Prim, USkinnedMesh
 			return false;
 		}
 
+		const pxr::TfToken StageUpAxis = UsdUtils::GetUsdStageAxis( GetUsdStage() );
+
 		pxr::UsdGeomXformable Xformable(Prim);
 
 		std::vector<double> TimeSamples;
@@ -707,7 +649,7 @@ bool AUsdStageActor::ProcessSkeletonRoot( const pxr::UsdPrim& Prim, USkinnedMesh
 
 			if (TimeSamples.size() > 0)
 			{
-				auto UpdateBoneTransforms = [&, SkelQuery]()
+				auto UpdateBoneTransforms = [&, SkelQuery, StageUpAxis]()
 				{
 					FScopedUnrealAllocs UnrealAllocs;
 
@@ -725,7 +667,7 @@ bool AUsdStageActor::ProcessSkeletonRoot( const pxr::UsdPrim& Prim, USkinnedMesh
 						for (uint32 Index = 0; Index < UsdBoneTransforms.size(); ++Index)
 						{
 							const pxr::GfMatrix4d& UsdMatrix = UsdBoneTransforms[Index];
-							FTransform BoneTransform = UsdToUnreal::ConvertMatrix(GetUsdStage(), UsdMatrix);
+							FTransform BoneTransform = UsdToUnreal::ConvertMatrix( StageUpAxis, UsdMatrix );
 							BoneTransforms.Add(BoneTransform);
 						}
 					}
@@ -741,7 +683,7 @@ bool AUsdStageActor::ProcessSkeletonRoot( const pxr::UsdPrim& Prim, USkinnedMesh
 				PrimDelegates.Add(UsdToUnreal::ConvertPath(PrimPath), Handle);
 			}
 
-			ProcessMaterials(Prim, SkelMeshImportData, MaterialsCache, bHasPrimDisplayColor, Time);
+			UsdStageActorImpl::ProcessMaterials(Prim, SkelMeshImportData, MaterialsCache, bHasPrimDisplayColor, Time);
 
 			break;
 		}
@@ -762,7 +704,7 @@ bool AUsdStageActor::ProcessSkeletonRoot( const pxr::UsdPrim& Prim, USkinnedMesh
 	return SkeletalMesh != nullptr;
 }
 
-void AUsdStageActor::UpdatePrim( const pxr::SdfPath& InUsdPrimPath, bool bResync )
+void AUsdStageActor::UpdatePrim( const pxr::SdfPath& InUsdPrimPath, bool bResync, const TMap< FString, UStaticMesh* >& PrimPathsToStaticMeshes )
 {
 	TRACE_CPUPROFILER_EVENT_SCOPE( AUsdStageActor::UpdatePrim );
 
@@ -789,7 +731,7 @@ void AUsdStageActor::UpdatePrim( const pxr::SdfPath& InUsdPrimPath, bool bResync
 		}
 
 		TUsdStore< pxr::UsdPrim > PrimToExpand = GetUsdStage()->GetPrimAtPath( UsdPrimPath );
-		FUsdPrimTwin* UsdPrimTwin = ExpandPrim( PrimToExpand.Get() );
+		FUsdPrimTwin* UsdPrimTwin = ExpandPrim( PrimToExpand.Get(), PrimPathsToStaticMeshes );
 
 		GEditor->BroadcastLevelActorListChanged();
 		GEditor->RedrawLevelEditingViewports();
@@ -814,12 +756,16 @@ void AUsdStageActor::SetTime(float InTime)
 
 void AUsdStageActor::PostEditChangeProperty(FPropertyChangedEvent& PropertyChangedEvent)
 {
+	Super::PostEditChangeProperty( PropertyChangedEvent );
+
 	UProperty* PropertyThatChanged = PropertyChangedEvent.MemberProperty;
 	const FName PropertyName = PropertyThatChanged ? PropertyThatChanged->GetFName() : NAME_None;
 	
 	if ( PropertyName == GET_MEMBER_NAME_CHECKED( AUsdStageActor, RootLayer ) )
 	{
 #if USE_USD_SDK
+		UnrealUSDWrapper::GetUsdStageCache().Erase( UsdStageStore.Get() );
+
 		UsdStageStore = TUsdStore< pxr::UsdStageRefPtr >();
 		LoadUsdStage();
 #endif // #if USE_USD_SDK
@@ -967,6 +913,8 @@ void AUsdStageActor::SetupLevelSequence()
 void AUsdStageActor::LoadUsdStage()
 {
 #if USE_USD_SDK
+	double StartTime = FPlatformTime::Cycles64();
+
 	FScopedSlowTask SlowTask( 1.f, LOCTEXT( "LoadingUDStage", "Loading USD Stage") );
 	SlowTask.MakeDialog();
 	SlowTask.EnterProgressFrame();
@@ -990,7 +938,9 @@ void AUsdStageActor::LoadUsdStage()
 
 	SetupLevelSequence();
 
-	UpdatePrim( UsdStage->GetPseudoRoot().GetPrimPath(), true );
+	TMap< FString, UStaticMesh* > PrimPathsToStaticMeshes = LoadStaticMeshes( UsdStage->GetPseudoRoot() );
+
+	UpdatePrim( UsdStage->GetPseudoRoot().GetPrimPath(), true, PrimPathsToStaticMeshes );
 
 	SetTime( StartTimeCode );
 
@@ -1001,7 +951,10 @@ void AUsdStageActor::LoadUsdStage()
 	{
 		for ( const FString& PrimToAnimate : PrimsToAnimate )
 		{
-			this->LoadPrim( UnrealToUsd::ConvertPath( *PrimToAnimate ).Get() );
+			TUsdStore< pxr::SdfPath > PrimPath = UnrealToUsd::ConvertPath( *PrimToAnimate );
+
+			TMap< FString, UStaticMesh* > PrimPathsToStaticMeshes = LoadStaticMeshes( this->GetUsdStage()->GetPrimAtPath( PrimPath.Get() ) );
+			this->LoadPrim( PrimPath.Get(), PrimPathsToStaticMeshes );
 		}
 
 		GEditor->BroadcastLevelActorListChanged();
@@ -1009,6 +962,14 @@ void AUsdStageActor::LoadUsdStage()
 	};
 
 	OnTimeChanged.AddLambda( AnimatePrims );
+
+	// Log time spent to load the stage
+	double ElapsedSeconds = FPlatformTime::ToSeconds64(FPlatformTime::Cycles64() - StartTime);
+
+	int ElapsedMin = int(ElapsedSeconds / 60.0);
+	ElapsedSeconds -= 60.0 * (double)ElapsedMin;
+
+	UE_LOG( LogUsdStage, Log, TEXT("%s %s in [%d min %.3f s]"), TEXT("Stage loaded"), *FPaths::GetBaseFilename( RootLayer.FilePath ), ElapsedMin, ElapsedSeconds );
 #endif // #if USE_USD_SDK
 }
 
@@ -1026,11 +987,19 @@ void AUsdStageActor::PostRegisterAllComponents()
 
 	if ( UsdStage )
 	{
-		UpdatePrim( UsdStage->GetPseudoRoot().GetPrimPath(), true );
+		TMap< FString, UStaticMesh* > PrimPathsToStaticMeshes = LoadStaticMeshes( UsdStage->GetPseudoRoot() );
+		UpdatePrim( UsdStage->GetPseudoRoot().GetPrimPath(), true, PrimPathsToStaticMeshes );
 	}
 #endif // #if USE_USD_SDK
 }
 
+void AUsdStageActor::PostLoad()
+{
+	Super::PostLoad();
+
+	OnActorLoaded.Broadcast( this );
+}
+
 void AUsdStageActor::OnUsdPrimTwinDestroyed( const FUsdPrimTwin& UsdPrimTwin )
 {
 	TArray<FDelegateHandle> DelegateHandles;
@@ -1051,7 +1020,7 @@ void AUsdStageActor::OnUsdPrimTwinDestroyed( const FUsdPrimTwin& UsdPrimTwin )
 void AUsdStageActor::OnPrimObjectPropertyChanged( UObject* ObjectBeingModified, FPropertyChangedEvent& PropertyChangedEvent )
 {
 #if USE_USD_SDK
-	if ( !ObjectsToWatch.Contains( ObjectBeingModified ) )
+	if ( ObjectBeingModified == this || !ObjectsToWatch.Contains( ObjectBeingModified ) )
 	{
 		return;
 	}
@@ -1084,4 +1053,198 @@ void AUsdStageActor::OnPrimObjectPropertyChanged( UObject* ObjectBeingModified,
 #endif // #if USE_USD_SDK
 }
 
+namespace UsdStageActorImpl
+{
+	void PreBuildStaticMesh( const pxr::UsdGeomMesh& UsdMesh, UStaticMesh& StaticMesh, TMap< FString, UMaterial* >& MaterialsCache, float Time )
+	{
+		const bool bHasPrimDisplayColor = UsdMesh.GetDisplayColorPrimvar().IsDefined();
+		ProcessMaterials( UsdMesh.GetPrim(), StaticMesh, MaterialsCache, bHasPrimDisplayColor, Time );
+
+		StaticMesh.RenderData = MakeUnique< FStaticMeshRenderData >();
+	}
+
+	bool BuildStaticMesh( UStaticMesh& StaticMesh, IMeshBuilderModule& MeshBuilderModule )
+	{
+		ITargetPlatformManagerModule& TargetPlatformManager = GetTargetPlatformManagerRef();
+		ITargetPlatform* RunningPlatform = TargetPlatformManager.GetRunningTargetPlatform();
+		check(RunningPlatform);
+
+		const FStaticMeshLODSettings& LODSettings = RunningPlatform->GetStaticMeshLODSettings();
+
+		const FStaticMeshLODGroup& LODGroup = LODSettings.GetLODGroup(StaticMesh.LODGroup);
+
+		if ( !MeshBuilderModule.BuildMesh( *StaticMesh.RenderData, &StaticMesh, LODGroup ) )
+		{
+			return false;
+		}
+
+		for ( FStaticMeshLODResources& LODResources : StaticMesh.RenderData->LODResources )
+		{
+			LODResources.bHasColorVertexData = true;
+		}
+
+		return true;
+	}
+
+	void PostBuildStaticMesh( UStaticMesh& StaticMesh )
+	{
+		StaticMesh.InitResources();
+
+		if ( const FMeshDescription* MeshDescription = StaticMesh.GetMeshDescription( 0 ) )
+		{
+			StaticMesh.RenderData->Bounds = MeshDescription->GetBounds();
+		}
+	
+		StaticMesh.CalculateExtendedBounds();
+	}
+}
+
+UStaticMesh* AUsdStageActor::CreateStaticMesh( const pxr::UsdGeomMesh& UsdMesh, FMeshDescription&& MeshDescription, bool& bOutIsNew )
+{
+	UStaticMesh* StaticMesh = nullptr;
+
+	FSHAHash MeshHash = FMeshDescriptionOperations::ComputeSHAHash( MeshDescription );
+
+	StaticMesh = MeshCache.FindRef( MeshHash.ToString() );
+
+	if ( !StaticMesh && !MeshDescription.IsEmpty() )
+	{
+		bOutIsNew = true;
+
+		StaticMesh = NewObject< UStaticMesh >( GetTransientPackage(), NAME_None, DefaultObjFlag | EObjectFlags::RF_Public );
+
+		FStaticMeshSourceModel& SourceModel = StaticMesh->AddSourceModel();
+		SourceModel.BuildSettings.bGenerateLightmapUVs = false;
+		SourceModel.BuildSettings.bRecomputeNormals = false;
+		SourceModel.BuildSettings.bRecomputeTangents = false;
+		SourceModel.BuildSettings.bBuildAdjacencyBuffer = false;
+		SourceModel.BuildSettings.bBuildReversedIndexBuffer = false;
+
+		FMeshDescription* StaticMeshDescription = StaticMesh->CreateMeshDescription(0);
+		check( StaticMeshDescription );
+		*StaticMeshDescription = MoveTemp( MeshDescription );
+
+		MeshCache.Add( MeshHash.ToString() ) = StaticMesh;
+	}
+	else
+	{
+		//FPlatformMisc::LowLevelOutputDebugStringf( TEXT("Mesh found in cache %s\n"), *StaticMesh->GetName() );
+		bOutIsNew = false;
+	}
+
+	return StaticMesh;
+}
+
+TMap< FString, UStaticMesh* > AUsdStageActor::LoadStaticMeshes( const pxr::UsdPrim& StartPrim )
+{
+	TMap< FString, UStaticMesh* > PrimPathsToStaticMeshes;
+
+	FScopedSlowTask LoadAllProgress( 6.f, LOCTEXT( "LoadingStaticMeshes", "Loading Static Meshes" ) );
+	LoadAllProgress.MakeDialog();
+
+	TArray< TUsdStore< pxr::UsdPrim > > AllGeomMeshes = UsdUtils::GetAllPrimsOfType( StartPrim, pxr::TfType::Find< pxr::UsdGeomMesh >(), { pxr::TfType::Find< pxr::UsdSkelRoot >() } );
+	TArray< FMeshDescription > MeshDescriptions;
+	MeshDescriptions.AddDefaulted( AllGeomMeshes.Num() );
+
+	// Create all mesh descriptions
+	{
+		LoadAllProgress.EnterProgressFrame();
+
+		ParallelFor( AllGeomMeshes.Num(), 
+			[ &AllGeomMeshes, &MeshDescriptions, TimeCode = pxr::UsdTimeCode( Time ) ]( int32 GeomMeshIndex )
+			{
+				MeshDescriptions[ GeomMeshIndex ] = UsdStageActorImpl::LoadMeshDescription( pxr::UsdGeomMesh( AllGeomMeshes[ GeomMeshIndex ].Get() ), TimeCode );
+			} );
+	}
+
+	TArray< TPair< int32, UStaticMesh* > > StaticMeshesToBuild;
+	StaticMeshesToBuild.Reserve( AllGeomMeshes.Num() );
+
+	// Create all static meshes
+	{
+		LoadAllProgress.EnterProgressFrame();
+
+		for ( int32 PrimIndex = 0; PrimIndex < AllGeomMeshes.Num(); ++PrimIndex )
+		{
+			bool bIsNew = true;
+			UStaticMesh* StaticMesh = CreateStaticMesh( pxr::UsdGeomMesh( AllGeomMeshes[ PrimIndex ].Get() ), MoveTemp( MeshDescriptions[ PrimIndex ] ), bIsNew );
+
+			PrimPathsToStaticMeshes.Add( UsdToUnreal::ConvertPath( AllGeomMeshes[ PrimIndex ].Get().GetPrimPath() ), StaticMesh );
+
+			if ( bIsNew )
+			{
+				StaticMeshesToBuild.Emplace( PrimIndex, StaticMesh );
+			}
+		}
+	}
+
+	// Commit all mesh descriptions
+	{
+		LoadAllProgress.EnterProgressFrame();
+
+		ParallelFor( StaticMeshesToBuild.Num(), 
+			[ &StaticMeshesToBuild ]( int32 MeshIndex )
+			{
+				UStaticMesh::FCommitMeshDescriptionParams Params;
+				Params.bMarkPackageDirty = false;
+				Params.bUseHashAsGuid = true;
+
+				TPair< int32, UStaticMesh* >& MeshToBuild = StaticMeshesToBuild[ MeshIndex ];
+
+				if ( UStaticMesh* StaticMesh = MeshToBuild.Value )
+				{
+					StaticMesh->CommitMeshDescription( 0, Params );
+				}
+			} );
+	}
+
+	// PreBuild all static meshes
+	{
+		LoadAllProgress.EnterProgressFrame();
+
+		for ( TPair< int32, UStaticMesh* >& MeshToBuild : StaticMeshesToBuild )
+		{
+			if ( UStaticMesh* StaticMesh = MeshToBuild.Value )
+			{
+				UsdStageActorImpl::PreBuildStaticMesh( pxr::UsdGeomMesh( AllGeomMeshes[ MeshToBuild.Key ].Get() ), *StaticMesh, MaterialsCache, Time );
+			}
+		}
+	}
+
+	// Build all static meshes
+	{
+		LoadAllProgress.EnterProgressFrame();
+
+		IMeshBuilderModule& MeshBuilderModule = FModuleManager::LoadModuleChecked< IMeshBuilderModule >( TEXT("MeshBuilder") );
+
+		ParallelFor( StaticMeshesToBuild.Num(), 
+			[ &StaticMeshesToBuild, &MeshBuilderModule ]( int32 MeshIndex )
+			{
+				if ( UStaticMesh* StaticMesh = StaticMeshesToBuild[ MeshIndex ].Value )
+				{
+					if ( !UsdStageActorImpl::BuildStaticMesh( *StaticMesh, MeshBuilderModule ) )
+					{
+						// Build failed, discard the mesh
+						StaticMeshesToBuild[ MeshIndex ].Value = nullptr;
+					}
+				}
+			} );
+	}
+
+	// PostBuild all static meshes
+	{
+		LoadAllProgress.EnterProgressFrame();
+
+		for ( TPair< int32, UStaticMesh* >& MeshToBuild : StaticMeshesToBuild )
+		{
+			if ( UStaticMesh* StaticMesh = MeshToBuild.Value )
+			{
+				UsdStageActorImpl::PostBuildStaticMesh( *StaticMesh );
+			}
+		}
+	}
+
+	return PrimPathsToStaticMeshes;
+}
+
 #undef LOCTEXT_NAMESPACE
diff --git a/Engine/Plugins/Importers/USDImporter/Source/USDStage/Public/USDStageActor.h b/Engine/Plugins/Importers/USDImporter/Source/USDStage/Public/USDStageActor.h
--- a/Engine/Plugins/Importers/USDImporter/Source/USDStage/Public/USDStageActor.h
+++ b/Engine/Plugins/Importers/USDImporter/Source/USDStage/Public/USDStageActor.h
@@ -26,10 +26,14 @@
 #include "USDStageActor.generated.h"
 
 class ALevelSequenceActor;
+struct FMeshDescription;
+class IMeshBuilderModule;
 class ULevelSequence;
 class UMaterial;
 class UUsdAsset;
 
+DECLARE_LOG_CATEGORY_EXTERN( LogUsdStage, Log, All );
+
 UENUM()
 enum class EUsdInitialLoadSet
 {
@@ -73,6 +77,9 @@ private:
 	TWeakObjectPtr< ULevelSequence > LevelSequence;
 
 public:
+	DECLARE_EVENT_OneParam( AUsdStageActor, FOnActorLoaded, AUsdStageActor* );
+	USDSTAGE_API static FOnActorLoaded OnActorLoaded;
+
 	DECLARE_EVENT( AUsdStageActor, FOnStageChanged );
 	FOnStageChanged OnStageChanged;
 
@@ -84,12 +91,14 @@ public:
 
 public:
 	AUsdStageActor();
+	virtual ~AUsdStageActor();
 
 	void Refresh() const;
 
 public:
 	virtual void PostEditChangeProperty( FPropertyChangedEvent& PropertyChangedEvent ) override;
 	virtual void PostRegisterAllComponents() override;
+	virtual void PostLoad() override;
 
 private:
 	void Clear();
@@ -110,7 +119,7 @@ private:
 
 	TMultiMap< FString, FDelegateHandle > PrimDelegates;
 
-	TArray< FString > PrimsToAnimate;
+	TSet< FString > PrimsToAnimate;
 
 	TMap< UObject*, FString > ObjectsToWatch;
 
@@ -129,12 +138,16 @@ public:
 	const FUsdListener& GetUsdListener() const { return UsdListener; }
 
 	FUsdPrimTwin* SpawnPrim( const pxr::SdfPath& UsdPrimPath );
-	FUsdPrimTwin* LoadPrim( const pxr::SdfPath& Path );
-	FUsdPrimTwin* ExpandPrim( const pxr::UsdPrim& Prim );
-	void UpdatePrim( const pxr::SdfPath& UsdPrimPath, bool bResync );
-	bool LoadStaticMesh( const pxr::UsdGeomMesh& UsdMesh, UStaticMeshComponent& MeshComponent );
+	FUsdPrimTwin* LoadPrim( const pxr::SdfPath& Path, const TMap< FString, UStaticMesh* >& PrimPathsToStaticMeshes );
+	FUsdPrimTwin* ExpandPrim( const pxr::UsdPrim& Prim, const TMap< FString, UStaticMesh* >& PrimPathsToStaticMeshes );
+	void UpdatePrim( const pxr::SdfPath& UsdPrimPath, bool bResync, const TMap< FString, UStaticMesh* >& PrimPathsToStaticMeshes );
+	bool SetStaticMesh( const pxr::UsdGeomMesh& UsdMesh, UStaticMeshComponent& MeshComponent, const TMap< FString, UStaticMesh* >& PrimPathsToStaticMeshes );
 	bool ProcessSkeletonRoot( const pxr::UsdPrim& Prim, USkinnedMeshComponent& SkinnedMeshComponent );
 
+protected:
+	TMap< FString, UStaticMesh* > LoadStaticMeshes( const pxr::UsdPrim& StartPrim );
+	UStaticMesh* CreateStaticMesh( const pxr::UsdGeomMesh& UsdMesh, FMeshDescription&& MeshDescription, bool& bOutIsNew );
+
 private:
 	TUsdStore< pxr::UsdStageRefPtr > UsdStageStore;
 	FUsdListener UsdListener;
diff --git a/Engine/Plugins/Importers/USDImporter/Source/USDStageEditor/Private/SUSDStage.cpp b/Engine/Plugins/Importers/USDImporter/Source/USDStageEditor/Private/SUSDStage.cpp
--- a/Engine/Plugins/Importers/USDImporter/Source/USDStageEditor/Private/SUSDStage.cpp
+++ b/Engine/Plugins/Importers/USDImporter/Source/USDStageEditor/Private/SUSDStage.cpp
@@ -19,6 +19,7 @@
 #include "Engine/World.h"
 #include "Framework/MultiBox/MultiBoxBuilder.h"
 #include "Modules/ModuleManager.h"
+#include "ScopedTransaction.h"
 #include "UObject/GCObjectScopeGuard.h"
 #include "UObject/StrongObjectPtr.h"
 #include "Widgets/Layout/SSplitter.h"
@@ -50,6 +51,9 @@ namespace SUSDStageConstants
 
 void SUsdStage::Construct( const FArguments& InArgs )
 {
+	OnStageActorPropertyChangedHandle = FCoreUObjectDelegates::OnObjectPropertyChanged.AddSP( SharedThis( this ), &SUsdStage::OnStageActorPropertyChanged );
+	OnActorLoadedHandle = AUsdStageActor::OnActorLoaded.AddSP( SharedThis( this ), &SUsdStage::OnStageActorLoaded );
+
 	IUsdStageModule& UsdStageModule = FModuleManager::Get().LoadModuleChecked< IUsdStageModule >( "UsdStage" );
 	UsdStageActor = &UsdStageModule.GetUsdStageActor( GWorld );
 
@@ -58,6 +62,11 @@ void SUsdStage::Construct( const FArguments& InArgs )
 	if ( UsdStageActor.IsValid() )
 	{
 		UsdStage = UsdStageActor->GetUsdStage();
+
+		if ( !UsdStage.Get() )
+		{
+			UsdStageActor.Reset();
+		}
 	}
 
 	ChildSlot
@@ -150,37 +159,45 @@ void SUsdStage::SetupStageActorDelegates()
 				{
 					this->UsdStageTreeView->RefreshPrim( PrimPath, bResync );
 				}
+
+				if ( this->UsdPrimInfoWidget )
+				{
+					this->UsdPrimInfoWidget->SetPrimPath( UsdStageActor->GetUsdStage(), *PrimPath );
+				}
 			}
 		);
 
-		OnStageChangedHandle = UsdStageActor->OnStageChanged.AddLambda(
-			[ this ]()
-			{
-				if ( UsdStageActor.IsValid() )
+		if ( !OnStageChangedHandle.IsValid() )
+		{
+			OnStageChangedHandle = UsdStageActor->OnStageChanged.AddLambda(
+				[ this ]()
 				{
-					if ( this->UsdStageInfoWidget )
-					{
-						this->UsdStageInfoWidget->RefreshStageInfos( UsdStageActor.Get() );
-					}
-
-					if ( this->UsdStageTreeView )
+					if ( UsdStageActor.IsValid() )
 					{
-						this->UsdStageTreeView->Refresh( UsdStageActor.Get() );
-						this->UsdStageTreeView->RequestTreeRefresh();
-					}
-
-					if ( this->UsdPrimInfoWidget )
-					{
-						this->UsdPrimInfoWidget->SetPrimPath( UsdStageActor->GetUsdStage(), TEXT("/") );
-					}
-
-					if ( this->UsdLayersTreeView )
-					{
-						this->UsdLayersTreeView->Refresh( UsdStageActor.Get(), true );
+						if ( this->UsdStageInfoWidget )
+						{
+							this->UsdStageInfoWidget->RefreshStageInfos( UsdStageActor.Get() );
+						}
+
+						if ( this->UsdStageTreeView )
+						{
+							this->UsdStageTreeView->Refresh( UsdStageActor.Get() );
+							this->UsdStageTreeView->RequestTreeRefresh();
+						}
+
+						if ( this->UsdPrimInfoWidget )
+						{
+							this->UsdPrimInfoWidget->SetPrimPath( UsdStageActor->GetUsdStage(), TEXT("/") );
+						}
+
+						if ( this->UsdLayersTreeView )
+						{
+							this->UsdLayersTreeView->Refresh( UsdStageActor.Get(), true );
+						}
 					}
 				}
-			}
-		);
+			);
+		}
 
 		OnStageEditTargetChangedHandle = UsdStageActor->GetUsdListener().OnStageEditTargetChanged.AddLambda(
 			[ this ]()
@@ -196,6 +213,9 @@ void SUsdStage::SetupStageActorDelegates()
 
 SUsdStage::~SUsdStage()
 {
+	FCoreUObjectDelegates::OnObjectPropertyChanged.Remove( OnStageActorPropertyChangedHandle );
+	AUsdStageActor::OnActorLoaded.Remove( OnActorLoadedHandle );
+
 	if ( UsdStageActor.IsValid() )
 	{
 		UsdStageActor->OnStageChanged.Remove( OnStageChangedHandle );
@@ -308,25 +328,27 @@ void SUsdStage::OnNew()
 
 	if ( UsdFilePath )
 	{
-		FScopedUsdAllocs UsdAllocs;
+		{
+			FScopedUsdAllocs UsdAllocs;
 
-		pxr::UsdStageCacheContext UsdStageCacheContext( UnrealUSDWrapper::GetUsdStageCache() );
-		pxr::UsdStageRefPtr UsdStage = pxr::UsdStage::CreateNew( UnrealToUsd::ConvertString( *UsdFilePath.GetValue() ).Get() );
+			pxr::UsdStageCacheContext UsdStageCacheContext( UnrealUSDWrapper::GetUsdStageCache() );
+			pxr::UsdStageRefPtr UsdStage = pxr::UsdStage::CreateNew( UnrealToUsd::ConvertString( *UsdFilePath.GetValue() ).Get() );
 
-		if ( !UsdStage )
-		{
-			return;
-		}
+			if ( !UsdStage )
+			{
+				return;
+			}
 
-		// Create default prim
-		pxr::UsdGeomXform RootPrim = pxr::UsdGeomXform::Define( UsdStage, UnrealToUsd::ConvertPath( TEXT("/Root") ).Get() );
-		pxr::UsdModelAPI( RootPrim ).SetKind( pxr::TfToken("component") );
+			// Create default prim
+			pxr::UsdGeomXform RootPrim = pxr::UsdGeomXform::Define( UsdStage, UnrealToUsd::ConvertPath( TEXT("/Root") ).Get() );
+			pxr::UsdModelAPI( RootPrim ).SetKind( pxr::TfToken("component") );
 		
-		// Set default prim
-		UsdStage->SetDefaultPrim( RootPrim.GetPrim() );
+			// Set default prim
+			UsdStage->SetDefaultPrim( RootPrim.GetPrim() );
 
-		// Set up axis
-		UsdUtils::SetUsdStageAxis( UsdStage, pxr::UsdGeomTokens->z );
+			// Set up axis
+			UsdUtils::SetUsdStageAxis( UsdStage, pxr::UsdGeomTokens->z );
+		}
 
 		OpenStage( *UsdFilePath.GetValue() );
 	}
@@ -464,8 +486,10 @@ void SUsdStage::OnPrimSelected( FString PrimPath )
 
 void SUsdStage::OnInitialLoadSetChanged( EUsdInitialLoadSet InitialLoadSet )
 {
-	if ( UsdStageActor.IsValid() )
+	if ( UsdStageActor.IsValid() && UsdStageActor->InitialLoadSet != InitialLoadSet )
 	{
+		const FScopedTransaction Transaction( LOCTEXT("EditInitialLoadSetTransaction", "Edit Initial Load Set") );
+		UsdStageActor->Modify();
 		UsdStageActor->InitialLoadSet = InitialLoadSet;
 	}
 }
@@ -477,6 +501,11 @@ void SUsdStage::OpenStage( const TCHAR* FilePath )
 		IUsdStageModule& UsdStageModule = FModuleManager::Get().LoadModuleChecked< IUsdStageModule >( "UsdStage" );
 		UsdStageActor = &UsdStageModule.GetUsdStageActor( GWorld );
 
+		if ( UsdStageInfoWidget )
+		{
+			UsdStageActor->InitialLoadSet = UsdStageInfoWidget->GetInitialLoadSet();
+		}
+
 		SetupStageActorDelegates();
 	}
 
@@ -487,6 +516,28 @@ void SUsdStage::OpenStage( const TCHAR* FilePath )
 	UsdStageActor->PostEditChangeProperty( RootLayerPropertyChangedEvent );
 }
 
+void SUsdStage::OnStageActorLoaded( AUsdStageActor* InUsdStageActor )
+{
+	if ( UsdStageActor == InUsdStageActor )
+	{
+		return;
+	}
+
+	UsdStageActor = InUsdStageActor;
+	SetupStageActorDelegates();
+}
+
+void SUsdStage::OnStageActorPropertyChanged( UObject* ObjectBeingModified, FPropertyChangedEvent& PropertyChangedEvent )
+{
+	if ( ObjectBeingModified == UsdStageActor )
+	{
+		if ( this->UsdStageInfoWidget )
+		{
+			this->UsdStageInfoWidget->RefreshStageInfos( UsdStageActor.Get() );
+		}
+	}
+}
+
 #endif // #if USE_USD_SDK
 
 #undef LOCTEXT_NAMESPACE
diff --git a/Engine/Plugins/Importers/USDImporter/Source/USDStageEditor/Private/SUSDStage.h b/Engine/Plugins/Importers/USDImporter/Source/USDStageEditor/Private/SUSDStage.h
--- a/Engine/Plugins/Importers/USDImporter/Source/USDStageEditor/Private/SUSDStage.h
+++ b/Engine/Plugins/Importers/USDImporter/Source/USDStageEditor/Private/SUSDStage.h
@@ -47,6 +47,9 @@ protected:
 
 	void OpenStage( const TCHAR* FilePath );
 
+	void OnStageActorLoaded( AUsdStageActor* InUsdStageActor );
+	void OnStageActorPropertyChanged( UObject* ObjectBeingModified, FPropertyChangedEvent& PropertyChangedEvent );
+
 protected:
 	TSharedPtr< class SUsdStageInfo > UsdStageInfoWidget;
 	TSharedPtr< class SUsdStageTreeView > UsdStageTreeView;
@@ -55,6 +58,8 @@ protected:
 
 	TWeakObjectPtr< AUsdStageActor > UsdStageActor;
 
+	FDelegateHandle OnActorLoadedHandle;
+	FDelegateHandle OnStageActorPropertyChangedHandle;
 	FDelegateHandle OnStageChangedHandle;
 	FDelegateHandle OnStageEditTargetChangedHandle;
 	FDelegateHandle OnPrimChangedHandle;
diff --git a/Engine/Plugins/Importers/USDImporter/Source/USDStageEditor/Private/SUSDStageInfo.cpp b/Engine/Plugins/Importers/USDImporter/Source/USDStageEditor/Private/SUSDStageInfo.cpp
--- a/Engine/Plugins/Importers/USDImporter/Source/USDStageEditor/Private/SUSDStageInfo.cpp
+++ b/Engine/Plugins/Importers/USDImporter/Source/USDStageEditor/Private/SUSDStageInfo.cpp
@@ -63,7 +63,7 @@ void SUsdStageInfo::Construct( const FArguments& InArgs, AUsdStageActor* UsdStag
 			.HAlign( HAlign_Left )
 			.Padding( 2.f, 2.f )
 			[
-				SNew( STextComboBox )
+				SAssignNew( InitialLoadSetWidget, STextComboBox )
 				.OptionsSource( &InitialLoadSetStrings )
 				.InitiallySelectedItem( InitialLoadSetStrings[ (int32)StageInfos.InitialLoadSet ] )
 				.OnSelectionChanged( this, &SUsdStageInfo::OnInitialLoadSetSelectionChanged )
@@ -81,6 +81,13 @@ void SUsdStageInfo::RefreshStageInfos( AUsdStageActor* UsdStageActor )
 		return;
 	}
 
+	StageInfos.InitialLoadSet = UsdStageActor->InitialLoadSet;
+
+	if ( InitialLoadSetWidget && InitialLoadSetStrings.IsValidIndex( (int32)StageInfos.InitialLoadSet ) )
+	{
+		InitialLoadSetWidget->SetSelectedItem( InitialLoadSetStrings[ (int32)StageInfos.InitialLoadSet ] );
+	}
+
 	if ( const pxr::UsdStageRefPtr& UsdStage = UsdStageActor->GetUsdStage() )
 	{
 		TUsdStore< std::string > UsdDisplayName = UsdStage->GetRootLayer()->GetDisplayName();
diff --git a/Engine/Plugins/Importers/USDImporter/Source/USDStageEditor/Private/SUSDStageInfo.h b/Engine/Plugins/Importers/USDImporter/Source/USDStageEditor/Private/SUSDStageInfo.h
--- a/Engine/Plugins/Importers/USDImporter/Source/USDStageEditor/Private/SUSDStageInfo.h
+++ b/Engine/Plugins/Importers/USDImporter/Source/USDStageEditor/Private/SUSDStageInfo.h
@@ -7,6 +7,7 @@
 
 class AUsdStageActor;
 enum class EUsdInitialLoadSet;
+class STextComboBox;
 
 #if USE_USD_SDK
 
@@ -22,6 +23,8 @@ public:
 	void Construct( const FArguments& InArgs, AUsdStageActor* UsdStageActor );
 	void RefreshStageInfos( AUsdStageActor* UsdStageActor );
 
+	EUsdInitialLoadSet GetInitialLoadSet() const { return StageInfos.InitialLoadSet; }
+
 private:
 	void OnInitialLoadSetSelectionChanged( TSharedPtr<FString> NewValue, ESelectInfo::Type SelectInfo );
 
@@ -37,8 +40,8 @@ private:
 	FStageInfos StageInfos;
 
 	FText GetRootLayerDisplayName() const { return StageInfos.RootLayerDisplayName; }
-	FString GetInitialLoadSet() const { return *InitialLoadSetStrings[ (int32)StageInfos.InitialLoadSet ]; }
 
+	TSharedPtr< STextComboBox > InitialLoadSetWidget;
 };
 
 #endif // #if USE_USD_SDK
diff --git a/Engine/Plugins/Importers/USDImporter/Source/USDStageEditor/Private/SUSDStageTreeView.cpp b/Engine/Plugins/Importers/USDImporter/Source/USDStageEditor/Private/SUSDStageTreeView.cpp
--- a/Engine/Plugins/Importers/USDImporter/Source/USDStageEditor/Private/SUSDStageTreeView.cpp
+++ b/Engine/Plugins/Importers/USDImporter/Source/USDStageEditor/Private/SUSDStageTreeView.cpp
@@ -34,6 +34,8 @@
 
 #include "USDIncludesEnd.h"
 
+#include <iterator>
+
 enum class EPayloadsTrigger
 {
 	Load,
@@ -67,7 +69,7 @@ public:
 	{
 		UsdPrim = InUsdPrim;
 
-		RefreshData();
+		RefreshData( false );
 		FillChildren();
 	}
 
@@ -78,27 +80,52 @@ public:
 	{
 	}
 
-	TArray< FUsdStageTreeItemRef > GetChildren()
+	TArray< FUsdStageTreeItemRef >& UpdateChildren()
 	{
 		if ( !UsdPrim.Get() )
 		{
-			return {};
+			return Children;
 		}
 
 		bool bNeedsRefresh = false;
 
-		int32 ChildIndex = 0;
 		pxr::UsdPrimSiblingRange PrimChildren = UsdPrim.Get().GetFilteredChildren( pxr::UsdTraverseInstanceProxies( pxr::UsdPrimAllPrimsPredicate ) );
-		for ( pxr::UsdPrim Child : PrimChildren )
+
+		const int32 NumUsdChildren = (TArray< FUsdStageTreeItemRef >::SizeType )std::distance( PrimChildren.begin(), PrimChildren.end() );
+		const int32 NumUnrealChildren = [&]()
+			{
+				int32 ValidPrims = 0;
+				for ( const FUsdStageTreeItemRef& Child : Children )
+				{
+					if ( !Child->RowData->Name.IsEmpty() )
+					{
+						++ValidPrims;
+					}
+				}
+
+				return ValidPrims;
+			}();
+
+		if ( NumUsdChildren != NumUnrealChildren )
 		{
-			if ( !Children.IsValidIndex( ChildIndex ) || Children[ ChildIndex ]->UsdPrim.Get().GetPrimPath() != Child.GetPrimPath() )
+			Children.Reset();
+			bNeedsRefresh = true;
+		}
+		else
+		{
+			int32 ChildIndex = 0;
+
+			for ( const pxr::UsdPrim& Child : PrimChildren )
 			{
-				Children.Reset();
-				bNeedsRefresh = true;
-				break;
-			}
+				if ( !Children.IsValidIndex( ChildIndex ) || Children[ ChildIndex ]->UsdPrim.Get().GetPrimPath() != Child.GetPrimPath() )
+				{
+					Children.Reset();
+					bNeedsRefresh = true;
+					break;
+				}
 				
-			++ChildIndex;
+				++ChildIndex;
+			}
 		}
 
 		if ( bNeedsRefresh )
@@ -118,12 +145,17 @@ public:
 		}
 	}
 
-	void RefreshData()
+	void RefreshData( bool bRefreshChildren )
 	{
 		FString PrimName;
 
 		if ( !UsdPrim.Get() )
 		{
+			if ( ParentItem )
+			{
+				ParentItem->RefreshData( true );
+			}
+
 			RowData.Get() = FUsdStageTreeRowData();
 			return;
 		}
@@ -139,6 +171,14 @@ public:
 		{
 			RowData->bIsVisible = ( UsdGeomImageable.ComputeVisibility() != pxr::UsdGeomTokens->invisible );
 		}
+
+		if ( bRefreshChildren )
+		{
+			for ( FUsdStageTreeItemRef& Child : UpdateChildren() )
+			{
+				Child->RefreshData( bRefreshChildren );
+			}
+		}
 	}
 
 	/** Delegate for hooking up an inline editable text block to be notified that a rename is requested. */
@@ -387,7 +427,7 @@ TSharedRef< ITableRow > SUsdStageTreeView::OnGenerateRow( FUsdStageTreeItemRef I
 
 void SUsdStageTreeView::OnGetChildren( FUsdStageTreeItemRef InParent, TArray< FUsdStageTreeItemRef >& OutChildren ) const
 {
-	for ( FUsdStageTreeItemRef& Child : InParent->GetChildren() )
+	for ( const FUsdStageTreeItemRef& Child : InParent->UpdateChildren() )
 	{
 		OutChildren.Add( Child );
 	}
@@ -424,8 +464,8 @@ void SUsdStageTreeView::RefreshPrim( const FString& PrimPath, bool bResync )
 
 	TUsdStore< pxr::SdfPath > UsdPrimPath = UnrealToUsd::ConvertPath( *PrimPath );
 
-	TFunction< FUsdStageTreeItemPtr( FUsdStageTreeItemRef ) > FindTreeItemFromPrimPath;
-	FindTreeItemFromPrimPath = [ &UsdPrimPath, &FindTreeItemFromPrimPath ]( FUsdStageTreeItemRef ItemRef ) -> FUsdStageTreeItemPtr
+	TFunction< FUsdStageTreeItemPtr( const TUsdStore< pxr::SdfPath >&, const FUsdStageTreeItemRef& ) > FindTreeItemFromPrimPath;
+	FindTreeItemFromPrimPath = [ &FindTreeItemFromPrimPath ]( const TUsdStore< pxr::SdfPath >& UsdPrimPath, const FUsdStageTreeItemRef& ItemRef ) -> FUsdStageTreeItemPtr
 	{
 		if ( ItemRef->UsdPrim.Get().GetPrimPath() == UsdPrimPath.Get() )
 		{
@@ -435,7 +475,7 @@ void SUsdStageTreeView::RefreshPrim( const FString& PrimPath, bool bResync )
 		{
 			for ( FUsdStageTreeItemRef ChildItem : ItemRef->Children )
 			{
-				if ( FUsdStageTreeItemPtr ChildValue = FindTreeItemFromPrimPath( ChildItem ) )
+				if ( FUsdStageTreeItemPtr ChildValue = FindTreeItemFromPrimPath( UsdPrimPath, ChildItem ) )
 				{
 					return ChildValue;
 				}
@@ -449,17 +489,35 @@ void SUsdStageTreeView::RefreshPrim( const FString& PrimPath, bool bResync )
 
 	for ( FUsdStageTreeItemRef RootItem : this->RootItems )
 	{
-		if ( FUsdStageTreeItemPtr FoundItem = FindTreeItemFromPrimPath( RootItem ) )
+		TUsdStore< pxr::SdfPath > PrimPathToSearch = UsdPrimPath;
+
+		bool bFoundPrimToRefresh = false;
+
+		while ( !bFoundPrimToRefresh )
 		{
-			FoundItem->RefreshData();
+			if ( FUsdStageTreeItemPtr FoundItem = FindTreeItemFromPrimPath( PrimPathToSearch, RootItem ) )
+			{
+				FoundItem->RefreshData( true );
+				bFoundPrimToRefresh = true;
 
-			if ( bResync )
+				break;
+			}
+			else
 			{
-				FoundItem->Children.Reset();
+				TUsdStore< pxr::SdfPath > ParentPrimPath = PrimPathToSearch.Get().GetParentPath();
+
+				if ( ParentPrimPath.Get() == PrimPathToSearch.Get() )
+				{
+					break;
+				}
+				else
+				{
+					PrimPathToSearch = MoveTemp( ParentPrimPath );
+				}
 			}
-			break;
 		}
-		else
+
+		if ( !bFoundPrimToRefresh )
 		{
 			bNeedsFullRefresh = true; // Prim not found, refresh the whole tree
 		}
@@ -670,8 +728,10 @@ void SUsdStageTreeView::OnAddReference()
 
 		pxr::UsdReferences References = SelectedItem->UsdPrim.Get().GetReferences();
 
-		FPaths::MakePathRelativeTo( PickedFile.GetValue(), *UsdStageActor->RootLayer.FilePath );
-		References.AddReference( UnrealToUsd::ConvertString( *PickedFile.GetValue() ).Get() );
+		FString AbsoluteFilePath = FPaths::ConvertRelativePathToFull( PickedFile.GetValue() );
+
+		FPaths::MakePathRelativeTo( AbsoluteFilePath, *UsdStageActor->RootLayer.FilePath );
+		References.AddReference( UnrealToUsd::ConvertString( *AbsoluteFilePath ).Get() );
 	}
 }
 
@@ -762,7 +822,7 @@ void SUsdStageTreeView::OnToggleAllPayloads( EPayloadsTrigger PayloadsTrigger )
 				}
 				else
 				{
-					for ( FUsdStageTreeItemRef Child : InSelectedItem->GetChildren() )
+					for ( FUsdStageTreeItemRef Child : InSelectedItem->UpdateChildren() )
 					{
 						RecursiveTogglePayloads( Child );
 					}
diff --git a/Engine/Plugins/Importers/USDImporter/Source/USDUtilities/Private/USDConversionUtils.cpp b/Engine/Plugins/Importers/USDImporter/Source/USDUtilities/Private/USDConversionUtils.cpp
--- a/Engine/Plugins/Importers/USDImporter/Source/USDUtilities/Private/USDConversionUtils.cpp
+++ b/Engine/Plugins/Importers/USDImporter/Source/USDUtilities/Private/USDConversionUtils.cpp
@@ -4,6 +4,7 @@
 
 #include "USDTypesConversion.h"
 
+#include "Algo/Copy.h"
 #include "CineCameraActor.h"
 #include "CineCameraComponent.h"
 
@@ -13,6 +14,8 @@
 #if USE_USD_SDK
 #include "USDIncludesStart.h"
 
+#include "pxr/usd/usd/attribute.h"
+#include "pxr/usd/usd/primRange.h"
 #include "pxr/usd/usdGeom/camera.h"
 #include "pxr/usd/usdGeom/scope.h"
 #include "pxr/usd/usdGeom/xform.h"
@@ -47,7 +50,7 @@ UClass* UsdUtils::GetComponentTypeForPrim( const pxr::UsdPrim& Prim )
 	{
 		return UCineCameraComponent::StaticClass();
 	}
-	else if ( Prim.IsA< pxr::UsdGeomXformable >() || Prim.IsA< pxr::UsdGeomScope >() )
+	else if ( Prim.IsA< pxr::UsdGeomXformable >() )
 	{
 		return USceneComponent::StaticClass();
 	}
@@ -73,4 +76,68 @@ TUsdStore< pxr::TfToken > UsdUtils::GetUVSetName( int32 UVChannelIndex )
 	return UVSetNameToken;
 }
 
+bool UsdUtils::IsAnimated( const pxr::UsdPrim& Prim )
+{
+	FScopedUsdAllocs UsdAllocs;
+
+	bool bHasXformbaleTimeSamples = false;
+	{
+		pxr::UsdGeomXformable Xformable( Prim );
+
+		if ( Xformable )
+		{
+			std::vector< double > TimeSamples;
+			Xformable.GetTimeSamples( &TimeSamples );
+
+			bHasXformbaleTimeSamples = TimeSamples.size() > 0;
+		}
+	}
+
+	bool bHasAttributesTimeSamples = false;
+	{
+		if ( !bHasXformbaleTimeSamples )
+		{
+			const std::vector< pxr::UsdAttribute >& Attributes = Prim.GetAttributes();
+
+			for ( const pxr::UsdAttribute& Attribute : Attributes )
+			{
+				bHasAttributesTimeSamples = Attribute.ValueMightBeTimeVarying();
+				if ( bHasAttributesTimeSamples )
+				{
+					break;
+				}
+			}
+		}
+	}
+
+	return bHasXformbaleTimeSamples || bHasAttributesTimeSamples;
+}
+
+TArray< TUsdStore< pxr::UsdPrim > > UsdUtils::GetAllPrimsOfType( const pxr::UsdPrim& StartPrim, const pxr::TfType& SchemaType, const TArray< TUsdStore< pxr::TfType > >& ExcludeSchemaTypes )
+{
+    TArray< TUsdStore< pxr::UsdPrim > > Result;
+
+	pxr::UsdPrimRange PrimRange( StartPrim, pxr::UsdTraverseInstanceProxies() );
+
+	for ( pxr::UsdPrimRange::iterator PrimRangeIt = PrimRange.begin(); PrimRangeIt != PrimRange.end(); ++PrimRangeIt )
+	{
+		if ( PrimRangeIt->IsA( SchemaType ) )
+		{
+			Result.Add( *PrimRangeIt );
+		}
+		else
+		{
+			for ( const TUsdStore< pxr::TfType >& SchemaToExclude : ExcludeSchemaTypes )
+			{
+				if ( PrimRangeIt->IsA( SchemaToExclude.Get() ) )
+				{
+					PrimRangeIt.PruneChildren();
+				}
+			}
+		}
+	}
+
+    return Result;
+}
+
 #endif // #if USE_USD_SDK
diff --git a/Engine/Plugins/Importers/USDImporter/Source/USDUtilities/Private/USDGeomMeshConversion.cpp b/Engine/Plugins/Importers/USDImporter/Source/USDUtilities/Private/USDGeomMeshConversion.cpp
--- a/Engine/Plugins/Importers/USDImporter/Source/USDUtilities/Private/USDGeomMeshConversion.cpp
+++ b/Engine/Plugins/Importers/USDImporter/Source/USDUtilities/Private/USDGeomMeshConversion.cpp
@@ -37,6 +37,8 @@ bool UsdToUnreal::ConvertGeomMesh( const pxr::UsdGeomMesh& UsdMesh, FMeshDescrip
 	FScopedUsdAllocs UsdAllocs;
 
 	pxr::UsdStageRefPtr Stage = UsdMesh.GetPrim().GetStage();
+	const pxr::TfToken StageUpAxis = UsdUtils::GetUsdStageAxis( Stage );
+
 	const double TimeCodeValue = TimeCode.GetValue();
 
 	TTuple< TArray< FString >, TArray< int32 > > GeometryMaterials = IUsdPrim::GetGeometryMaterials( TimeCodeValue, UsdMesh.GetPrim() );
@@ -65,7 +67,7 @@ bool UsdToUnreal::ConvertGeomMesh( const pxr::UsdGeomMesh& UsdMesh, FMeshDescrip
 			{
 				const GfVec3f& Point = PointsArray[ LocalPointIndex ];
 
-				FVector Pos = UsdToUnreal::ConvertVector( Stage, Point );
+				FVector Pos = UsdToUnreal::ConvertVector( StageUpAxis, Point );
 
 				FVertexID AddedVertexId = MeshDescription.CreateVertex();
 				MeshDescriptionVertexPositions[ AddedVertexId ] = Pos;
@@ -229,7 +231,7 @@ bool UsdToUnreal::ConvertGeomMesh( const pxr::UsdGeomMesh& UsdMesh, FMeshDescrip
 					const int32 NormalIndex = Normals.size() != FaceIndices.size() ? FaceIndices[CurrentVertexInstanceIndex] : CurrentVertexInstanceIndex;
 					check(NormalIndex < Normals.size());
 					const GfVec3f& Normal = Normals[NormalIndex];
-					FVector TransformedNormal = UsdToUnreal::ConvertVector( Stage, Normal );
+					FVector TransformedNormal = UsdToUnreal::ConvertVector( StageUpAxis, Normal );
 
 					ensure( !TransformedNormal.IsNearlyZero() );
 					MeshDescriptionNormals[AddedVertexInstanceId] = TransformedNormal.GetSafeNormal();
@@ -411,6 +413,8 @@ namespace UsdGeomMeshConversionImpl
 	{
 		UMaterialExpression* Result = nullptr;
 
+		FScopedUsdAllocs UsdAllocs;
+
 		pxr::UsdShadeConnectableAPI Source;
 		pxr::TfToken SourceName;
 		pxr::UsdShadeAttributeType AttributeType;
@@ -426,6 +430,8 @@ namespace UsdGeomMeshConversionImpl
 					pxr::SdfAssetPath TexturePath;
 					FileInput.Get< pxr::SdfAssetPath >( &TexturePath );
 
+					FScopedUnrealAllocs UnrealAllocs;
+
 					FString TextureFileName = UsdToUnreal::ConvertString( TexturePath.GetResolvedPath() );
 					TextureFileName.RemoveFromStart( TEXT("\\\\?\\") );
 
@@ -625,6 +631,8 @@ bool UnrealToUsd::ConvertStaticMesh( const UStaticMesh* StaticMesh, pxr::UsdGeom
 		return false;
 	}
 
+	const pxr::TfToken StageUpAxis = UsdUtils::GetUsdStageAxis( Stage );
+
 	const FStaticMeshLODResources& RenderMesh = StaticMesh->GetLODForExport( 0 /*ExportLOD*/ );
 
 	// Verify the integrity of the static mesh.
@@ -653,7 +661,7 @@ bool UnrealToUsd::ConvertStaticMesh( const UStaticMesh* StaticMesh, pxr::UsdGeom
 				for ( int32 VertexIndex = 0; VertexIndex < VertexCount; ++VertexIndex )
 				{
 					FVector VertexPosition = RenderMesh.VertexBuffers.PositionVertexBuffer.VertexPosition( VertexIndex );
-					PointsArray.push_back( UnrealToUsd::ConvertVector( Stage, VertexPosition ) );
+					PointsArray.push_back( UnrealToUsd::ConvertVector( StageUpAxis, VertexPosition ) );
 				}
 
 				Points.Set( PointsArray, TimeCode );
@@ -671,7 +679,7 @@ bool UnrealToUsd::ConvertStaticMesh( const UStaticMesh* StaticMesh, pxr::UsdGeom
 				for ( int32 VertexIndex = 0; VertexIndex < VertexCount; ++VertexIndex )
 				{
 					FVector VertexNormal = RenderMesh.VertexBuffers.StaticMeshVertexBuffer.VertexTangentZ( VertexIndex );
-					Normals.push_back( UnrealToUsd::ConvertVector( Stage, VertexNormal ) );
+					Normals.push_back( UnrealToUsd::ConvertVector( StageUpAxis, VertexNormal ) );
 				}
 
 				NormalsAttribute.Set( Normals, TimeCode );
diff --git a/Engine/Plugins/Importers/USDImporter/Source/USDUtilities/Private/USDSkeletalDataConversion.cpp b/Engine/Plugins/Importers/USDImporter/Source/USDUtilities/Private/USDSkeletalDataConversion.cpp
--- a/Engine/Plugins/Importers/USDImporter/Source/USDUtilities/Private/USDSkeletalDataConversion.cpp
+++ b/Engine/Plugins/Importers/USDImporter/Source/USDUtilities/Private/USDSkeletalDataConversion.cpp
@@ -70,7 +70,7 @@ bool UsdToUnreal::ConvertSkeleton(const pxr::UsdSkelSkeletonQuery& SkeletonQuery
 		for (uint32 Index = 0; Index < UsdBoneTransforms.size(); ++Index)
 		{
 			const GfMatrix4d& UsdMatrix = UsdBoneTransforms[Index];
-			FTransform BoneTransform = UsdToUnreal::ConvertMatrix(Stage, UsdMatrix);
+			FTransform BoneTransform = UsdToUnreal::ConvertMatrix( UsdUtils::GetUsdStageAxis( Stage ), UsdMatrix );
 			BoneTransforms.Add(BoneTransform);
 		}
 	}
@@ -126,7 +126,9 @@ bool UsdToUnreal::ConvertSkinnedMesh(const pxr::UsdSkelSkinningQuery& SkinningQu
 		GeomBindingAttribute.Get(&GeomBindingTransform, UsdTimeCode::Default());
 	}
 
-	FTransform GeomTransform = UsdToUnreal::ConvertMatrix(SkinningPrim.GetStage(), GeomBindingTransform);
+	const pxr::TfToken StageUpAxis = UsdUtils::GetUsdStageAxis( SkinningPrim.GetStage() );
+
+	FTransform GeomTransform = UsdToUnreal::ConvertMatrix(StageUpAxis, GeomBindingTransform);
 
 	// Ref. FFbxImporter::FillSkelMeshImporterFromFbx
 	UsdGeomMesh UsdMesh = UsdGeomMesh(SkinningPrim);
@@ -149,7 +151,7 @@ bool UsdToUnreal::ConvertSkinnedMesh(const pxr::UsdSkelSkinningQuery& SkinningQu
 			const GfVec3f& Point = UsdPoints[PointIndex];
 
 			// Convert the USD vertex to Unreal and apply the GeomBindTransform to it
-			FVector Pos = UsdToUnreal::ConvertVector(UsdMesh.GetPrim().GetStage(), Point);
+			FVector Pos = UsdToUnreal::ConvertVector(StageUpAxis, Point);
 			Pos = GeomTransform.TransformPosition(Pos);
 
 			SkelMeshImportData.Points[PointIndex + NumExistingPoints] = Pos;
diff --git a/Engine/Plugins/Importers/USDImporter/Source/USDUtilities/Public/USDConversionUtils.h b/Engine/Plugins/Importers/USDImporter/Source/USDUtilities/Public/USDConversionUtils.h
--- a/Engine/Plugins/Importers/USDImporter/Source/USDUtilities/Public/USDConversionUtils.h
+++ b/Engine/Plugins/Importers/USDImporter/Source/USDUtilities/Public/USDConversionUtils.h
@@ -77,6 +77,14 @@ namespace UsdUtils
 	USDUTILITIES_API UClass* GetComponentTypeForPrim( const pxr::UsdPrim& Prim );
 
 	USDUTILITIES_API TUsdStore< pxr::TfToken > GetUVSetName( int32 UVChannelIndex );
+
+	USDUTILITIES_API bool IsAnimated( const pxr::UsdPrim& Prim );
+
+	/**
+	 * Returns all prims of type SchemaType (or a descendant type) in the subtree of prims rooted at StartPrim.
+	 * Stops going down the subtrees when it hits a schema type to exclude.
+	 */
+	USDUTILITIES_API TArray< TUsdStore< pxr::UsdPrim > > GetAllPrimsOfType( const pxr::UsdPrim& StartPrim, const pxr::TfType& SchemaType, const TArray< TUsdStore< pxr::TfType > >& ExcludeSchemaTypes = {} );
 }
 
 #endif // #if USE_USD_SDK
diff --git a/Engine/Plugins/Importers/USDImporter/Source/USDUtilities/Public/USDTypesConversion.h b/Engine/Plugins/Importers/USDImporter/Source/USDUtilities/Public/USDTypesConversion.h
--- a/Engine/Plugins/Importers/USDImporter/Source/USDUtilities/Public/USDTypesConversion.h
+++ b/Engine/Plugins/Importers/USDImporter/Source/USDUtilities/Public/USDTypesConversion.h
@@ -68,12 +68,11 @@ namespace UsdToUnreal
 		return FVector( InValue[0], InValue[1], InValue[2] );
 	}
 
-	static FVector ConvertVector( const pxr::UsdStageRefPtr& Stage, const pxr::GfVec3f& InValue )
+	static FVector ConvertVector( const pxr::TfToken& UpAxis, const pxr::GfVec3f& InValue )
 	{
 		FVector Value = ConvertVector( InValue );
 
-		pxr::TfToken UpAxisValue = UsdUtils::GetUsdStageAxis( Stage );
-		const bool bIsZUp = ( UpAxisValue == pxr::UsdGeomTokens->z );
+		const bool bIsZUp = ( UpAxis == pxr::UsdGeomTokens->z );
 
 		if ( bIsZUp )
 		{
@@ -149,13 +148,6 @@ namespace UsdToUnreal
 
 		return Transform;
 	}
-
-	static FTransform ConvertMatrix( const pxr::UsdStageRefPtr& Stage, const pxr::GfMatrix4d& InMatrix )
-	{
-		pxr::TfToken UpAxis = UsdUtils::GetUsdStageAxis(Stage);
-
-		return ConvertMatrix( UpAxis, InMatrix );
-	}
 }
 
 namespace UnrealToUsd
@@ -185,12 +177,11 @@ namespace UnrealToUsd
 		return pxr::GfVec3f( InValue[0], InValue[1], InValue[2] );
 	}
 
-	static pxr::GfVec3f ConvertVector( const pxr::UsdStageRefPtr& Stage, const FVector& InValue )
+	static pxr::GfVec3f ConvertVector( const pxr::TfToken& UpAxis, const FVector& InValue )
 	{
 		pxr::GfVec3f Value = ConvertVector( InValue );
 
-		pxr::TfToken UpAxisValue = UsdUtils::GetUsdStageAxis( Stage );
-		const bool bIsZUp = ( UpAxisValue == pxr::UsdGeomTokens->z );
+		const bool bIsZUp = ( UpAxis == pxr::UsdGeomTokens->z );
 
 		if ( bIsZUp )
 		{
diff --git a/Engine/Plugins/Importers/USDImporter/Source/UnrealUSDWrapper/Private/USDMemory.cpp b/Engine/Plugins/Importers/USDImporter/Source/UnrealUSDWrapper/Private/USDMemory.cpp
--- a/Engine/Plugins/Importers/USDImporter/Source/UnrealUSDWrapper/Private/USDMemory.cpp
+++ b/Engine/Plugins/Importers/USDImporter/Source/UnrealUSDWrapper/Private/USDMemory.cpp
@@ -5,6 +5,7 @@
 #include "Containers/Array.h"
 #include "HAL/PlatformTLS.h"
 #include "HAL/TlsAutoCleanup.h"
+#include "Misc/ScopeLock.h"
 
 class UNREALUSDWRAPPER_API FTlsSlot final
 {
@@ -91,6 +92,7 @@ private:
 
 TOptional< FTlsSlot > FUsdMemoryManager::ActiveAllocatorsStackTLS {};
 TSet< void* > FUsdMemoryManager::SystemAllocedPtrs {};
+FCriticalSection FUsdMemoryManager::CriticalSection{};
 
 void FUsdMemoryManager::Initialize()
 {
@@ -139,7 +141,11 @@ void* FUsdMemoryManager::Malloc( SIZE_T Count )
 	if ( FUsdMemoryManager::IsUsingSystemMalloc() )
 	{
 		Result = FMemory::SystemMalloc( Count );
-		SystemAllocedPtrs.Add( Result );
+
+		{
+			FScopeLock Lock( &CriticalSection );
+			SystemAllocedPtrs.Add( Result );
+		}
 	}
 	else
 	{
@@ -155,7 +161,11 @@ void FUsdMemoryManager::Free( void* Original )
 	// This can happen for inlined USD functions that call delete.
 	if ( FUsdMemoryManager::IsUsingSystemMalloc() || SystemAllocedPtrs.Contains( Original ) )
 	{
-		SystemAllocedPtrs.Remove( Original );
+		{
+			FScopeLock Lock( &CriticalSection );
+			SystemAllocedPtrs.Remove( Original );
+		}
+
 		FMemory::SystemFree( Original );
 	}
 	else
diff --git a/Engine/Plugins/Importers/USDImporter/Source/UnrealUSDWrapper/Public/USDMemory.h b/Engine/Plugins/Importers/USDImporter/Source/UnrealUSDWrapper/Public/USDMemory.h
--- a/Engine/Plugins/Importers/USDImporter/Source/UnrealUSDWrapper/Public/USDMemory.h
+++ b/Engine/Plugins/Importers/USDImporter/Source/UnrealUSDWrapper/Public/USDMemory.h
@@ -77,6 +77,8 @@ private:
 	static TOptional< FTlsSlot > ActiveAllocatorsStackTLS;
 
 	static TSet< void* > SystemAllocedPtrs;
+
+	static FCriticalSection CriticalSection;
 };
 
 /**
