Fixed level sequence getting unloaded.

diff --git a/Engine/Plugins/Importers/USDImporter/Source/USDStage/Private/USDStageActor.cpp b/Engine/Plugins/Importers/USDImporter/Source/USDStage/Private/USDStageActor.cpp
--- a/Engine/Plugins/Importers/USDImporter/Source/USDStage/Private/USDStageActor.cpp
+++ b/Engine/Plugins/Importers/USDImporter/Source/USDStage/Private/USDStageActor.cpp
@@ -831,7 +831,7 @@ void AUsdStageActor::OpenUsdStage()
 
 void AUsdStageActor::InitLevelSequence(float FramesPerSecond)
 {
-	if (LevelSequence.IsValid())
+	if (LevelSequence)
 	{
 		return;
 	}
@@ -865,7 +865,7 @@ void AUsdStageActor::InitLevelSequence(float FramesPerSecond)
 
 void AUsdStageActor::SetupLevelSequence()
 {
-	if (!LevelSequence.IsValid())
+	if (!LevelSequence)
 	{
 		return;
 	}
diff --git a/Engine/Plugins/Importers/USDImporter/Source/USDStage/Public/USDStageActor.h b/Engine/Plugins/Importers/USDImporter/Source/USDStage/Public/USDStageActor.h
--- a/Engine/Plugins/Importers/USDImporter/Source/USDStage/Public/USDStageActor.h
+++ b/Engine/Plugins/Importers/USDImporter/Source/USDStage/Public/USDStageActor.h
@@ -74,7 +74,7 @@ private:
 	float TimeCodesPerSecond;
 
 	UPROPERTY(VisibleAnywhere, Category = "USD", Transient)
-	TWeakObjectPtr< ULevelSequence > LevelSequence;
+	ULevelSequence* LevelSequence;
 
 public:
 	DECLARE_EVENT_OneParam( AUsdStageActor, FOnActorLoaded, AUsdStageActor* );
