USD - Added new material to prevent having to compile instances because of varying blend modes.

diff --git a/Engine/Plugins/Importers/USDImporter/Source/USDStage/Private/USDStageActor.cpp b/Engine/Plugins/Importers/USDImporter/Source/USDStage/Private/USDStageActor.cpp
--- a/Engine/Plugins/Importers/USDImporter/Source/USDStage/Private/USDStageActor.cpp
+++ b/Engine/Plugins/Importers/USDImporter/Source/USDStage/Private/USDStageActor.cpp
@@ -78,21 +78,21 @@ AUsdStageActor::FOnActorLoaded AUsdStageActor::OnActorLoaded{};
 
 namespace UsdStageActorImpl
 {
-FMeshDescription LoadMeshDescription( const pxr::UsdGeomMesh& UsdMesh, const pxr::UsdTimeCode TimeCode )
-{
-	if ( !UsdMesh )
+	FMeshDescription LoadMeshDescription( const pxr::UsdGeomMesh& UsdMesh, const pxr::UsdTimeCode TimeCode )
 	{
-		return {};
-	}
+		if ( !UsdMesh )
+		{
+			return {};
+		}
 
-	FMeshDescription MeshDescription;
-	FStaticMeshAttributes StaticMeshAttributes( MeshDescription );
-	StaticMeshAttributes.Register();
+		FMeshDescription MeshDescription;
+		FStaticMeshAttributes StaticMeshAttributes( MeshDescription );
+		StaticMeshAttributes.Register();
 
-	UsdToUnreal::ConvertGeomMesh( UsdMesh, MeshDescription, TimeCode );
+		UsdToUnreal::ConvertGeomMesh( UsdMesh, MeshDescription, TimeCode );
 
-	return MeshDescription;
-}
+		return MeshDescription;
+	}
 
 	void ProcessMaterials( const pxr::UsdPrim& UsdPrim, UStaticMesh& StaticMesh, TMap< FString, UMaterial* >& MaterialsCache, bool bHasPrimDisplayColor, float Time )
 	{
@@ -113,7 +113,7 @@ FMeshDescription LoadMeshDescription( const pxr::UsdGeomMesh& UsdMesh, const pxr
 			int32 MaterialIndex = INDEX_NONE;
 			int32 MeshMaterialIndex = 0;
 
-				for ( FStaticMaterial& StaticMaterial : StaticMesh.StaticMaterials )
+			for ( FStaticMaterial& StaticMaterial : StaticMesh.StaticMaterials )
 			{
 				if ( StaticMaterial.MaterialSlotName.IsEqual( ImportedMaterialSlotName ) )
 				{
@@ -756,8 +756,6 @@ void AUsdStageActor::SetTime(float InTime)
 
 void AUsdStageActor::PostEditChangeProperty(FPropertyChangedEvent& PropertyChangedEvent)
 {
-	Super::PostEditChangeProperty( PropertyChangedEvent );
-
 	UProperty* PropertyThatChanged = PropertyChangedEvent.MemberProperty;
 	const FName PropertyName = PropertyThatChanged ? PropertyThatChanged->GetFName() : NAME_None;
 	
@@ -774,6 +772,10 @@ void AUsdStageActor::PostEditChangeProperty(FPropertyChangedEvent& PropertyChang
 	{
 		Refresh();
 	}
+	else
+	{
+		Super::PostEditChangeProperty( PropertyChangedEvent );
+	}
 }
 
 void AUsdStageActor::Clear()
diff --git a/Engine/Plugins/Importers/USDImporter/Source/USDUtilities/Private/USDGeomMeshConversion.cpp b/Engine/Plugins/Importers/USDImporter/Source/USDUtilities/Private/USDGeomMeshConversion.cpp
--- a/Engine/Plugins/Importers/USDImporter/Source/USDUtilities/Private/USDGeomMeshConversion.cpp
+++ b/Engine/Plugins/Importers/USDImporter/Source/USDUtilities/Private/USDGeomMeshConversion.cpp
@@ -594,22 +594,25 @@ bool UsdToUnreal::ConvertDisplayColor( const pxr::UsdGeomMesh& UsdMesh, UMateria
 		{
 			FScopedUnrealAllocs UnrealAllocs;
 
-			FSoftObjectPath DisplayColorMaterialPath( TEXT("Material'/USDImporter/Materials/DisplayColorAndOpacity.DisplayColorAndOpacity'") );
-			UMaterialInterface* DisplayColorAndOpacityMaterial = Cast< UMaterialInterface >( DisplayColorMaterialPath.TryLoad() );
-
-			UMaterialEditingLibrary::SetMaterialInstanceParent( &MaterialInstance, DisplayColorAndOpacityMaterial );
+			FSoftObjectPath MaterialPath;
 
 			if ( bHasTransparency )
 			{
-				MaterialInstance.BasePropertyOverrides.bOverride_BlendMode = true;
-				MaterialInstance.BasePropertyOverrides.BlendMode = EBlendMode::BLEND_Translucent;
+				MaterialPath = FSoftObjectPath( TEXT("Material'/USDImporter/Materials/DisplayColorAndOpacity.DisplayColorAndOpacity'") );
+			}
+			else
+			{
+				MaterialPath = FSoftObjectPath( TEXT("Material'/USDImporter/Materials/DisplayColor.DisplayColor'") );
 			}
 
-			if ( UsdMesh.GetDoubleSidedAttr().IsDefined() )
+			UMaterialInterface* DisplayColorAndOpacityMaterial = Cast< UMaterialInterface >( MaterialPath.TryLoad() );
+			UMaterialEditingLibrary::SetMaterialInstanceParent( &MaterialInstance, DisplayColorAndOpacityMaterial );
+
+			/*if ( UsdMesh.GetDoubleSidedAttr().IsDefined() )
 			{
 				MaterialInstance.BasePropertyOverrides.bOverride_TwoSided = true;
 				MaterialInstance.BasePropertyOverrides.TwoSided = UsdUtils::GetUsdValue< bool >( UsdMesh.GetDoubleSidedAttr(), TimeCode );
-			}
+			}*/
 		}
 
 		return true;
