USD Stage Editor: Fixed stack overflow and a layer as itself as a sublayer.

#rnx

diff --git a/Engine/Plugins/Importers/USDImporter/Source/USDStageEditor/Private/SUSDLayersTreeView.cpp b/Engine/Plugins/Importers/USDImporter/Source/USDStageEditor/Private/SUSDLayersTreeView.cpp
--- a/Engine/Plugins/Importers/USDImporter/Source/USDStageEditor/Private/SUSDLayersTreeView.cpp
+++ b/Engine/Plugins/Importers/USDImporter/Source/USDStageEditor/Private/SUSDLayersTreeView.cpp
@@ -55,7 +55,7 @@ public:
 
 	bool IsValid() const
 	{
-		return (bool)UsdStage.Get();
+		return (bool)UsdStage.Get() && ( !ParentItem || ParentItem->LayerIdentifier.Get() != LayerIdentifier.Get() );
 	}
 
 	TArray< FUsdLayersTreeItemRef > GetChildren()
