USD Stage Editor: Fixed layers refresh issue

#rnx

diff --git a/Engine/Plugins/Importers/USDImporter/Source/USDStageEditor/Private/SUSDLayersTreeView.cpp b/Engine/Plugins/Importers/USDImporter/Source/USDStageEditor/Private/SUSDLayersTreeView.cpp
--- a/Engine/Plugins/Importers/USDImporter/Source/USDStageEditor/Private/SUSDLayersTreeView.cpp
+++ b/Engine/Plugins/Importers/USDImporter/Source/USDStageEditor/Private/SUSDLayersTreeView.cpp
@@ -85,9 +85,15 @@ public:
 						bNeedsRefresh = true;
 						break;
 					}
+
+					++SubLayerIndex;
 				}
 
-				++SubLayerIndex;
+				if ( !bNeedsRefresh && SubLayerIndex < Children.Num() )
+				{
+					Children.Reset();
+					bNeedsRefresh = true;
+				}
 			}
 		}
 
