USD Stage Editor

- Fixed initial payload set not modifying the actor.
- Fixed text color so that it's not orange on orange.
- Fixed tree view refresh issues.

diff --git a/Engine/Plugins/Importers/USDImporter/Source/USDStage/Private/USDListener.cpp b/Engine/Plugins/Importers/USDImporter/Source/USDStage/Private/USDListener.cpp
--- a/Engine/Plugins/Importers/USDImporter/Source/USDStage/Private/USDListener.cpp
+++ b/Engine/Plugins/Importers/USDImporter/Source/USDStage/Private/USDListener.cpp
@@ -48,28 +48,30 @@ void FUsdListener::HandleUsdNotice( const pxr::UsdNotice::ObjectsChanged& Notice
 		return;
 	}
 
+	auto BroadcastPrimUpdate = [ this ]( const pxr::SdfPath& PrimPath, bool bResync )
+	{
+		FScopedUnrealAllocs UnrealAllocs;
+
+		//FPlatformMisc::LowLevelOutputDebugStringf( TEXT("Prim %s changed (%s)\n"), *UsdToUnreal::ConvertPath( *It ), bResync ? TEXT("resync") : TEXT("update") );
+		OnPrimChanged.Broadcast( UsdToUnreal::ConvertPath( PrimPath ), bResync );
+	};
+
+	FScopedUsdAllocs UsdAllocs;
+
 	PathRange PathsToUpdate = Notice.GetResyncedPaths();
 
 	for ( PathRange::const_iterator It = PathsToUpdate.begin(); It != PathsToUpdate.end(); ++It )
 	{
-		//FPlatformMisc::LowLevelOutputDebugStringf( TEXT("Prim %s changed (resync)\n"), *UsdToUnreal::ConvertPath( *It ) );
-
-		FScopedUnrealAllocs UnrealAllocs;
-
 		constexpr bool bResync = true;
-		OnPrimChanged.Broadcast( UsdToUnreal::ConvertPath( *It ), bResync );
+		BroadcastPrimUpdate( It->GetAbsoluteRootOrPrimPath(), bResync );
 	}
 
 	PathsToUpdate = Notice.GetChangedInfoOnlyPaths();
 
 	for ( PathRange::const_iterator It = PathsToUpdate.begin(); It != PathsToUpdate.end(); ++It )
 	{
-		//FPlatformMisc::LowLevelOutputDebugStringf( TEXT("Prim %s changed (update)\n"), *UsdToUnreal::ConvertPath( *It ) );
-
-		FScopedUnrealAllocs UnrealAllocs;
-
 		constexpr bool bResync = false;
-		OnPrimChanged.Broadcast( UsdToUnreal::ConvertPath( *It ), bResync );
+		BroadcastPrimUpdate( It->GetAbsoluteRootOrPrimPath(), bResync );
 	}
 }
 
diff --git a/Engine/Plugins/Importers/USDImporter/Source/USDStage/Private/USDStageActor.cpp b/Engine/Plugins/Importers/USDImporter/Source/USDStage/Private/USDStageActor.cpp
--- a/Engine/Plugins/Importers/USDImporter/Source/USDStage/Private/USDStageActor.cpp
+++ b/Engine/Plugins/Importers/USDImporter/Source/USDStage/Private/USDStageActor.cpp
@@ -68,6 +68,8 @@
 
 static const EObjectFlags DefaultObjFlag = EObjectFlags::RF_Transactional | EObjectFlags::RF_Transient;
 
+AUsdStageActor::FOnActorLoaded AUsdStageActor::OnActorLoaded{};
+
 #if USE_USD_SDK
 FMeshDescription LoadMeshDescription( const pxr::UsdGeomMesh& UsdMesh, const pxr::UsdTimeCode TimeCode )
 {
@@ -227,6 +229,11 @@ AUsdStageActor::AUsdStageActor()
 	FCoreUObjectDelegates::OnObjectPropertyChanged.AddUObject( this, &AUsdStageActor::OnPrimObjectPropertyChanged );
 }
 
+AUsdStageActor::~AUsdStageActor()
+{
+	UnrealUSDWrapper::GetUsdStageCache().Erase( UsdStageStore.Get() );
+}
+
 #if USE_USD_SDK
 
 bool AUsdStageActor::LoadStaticMesh( const pxr::UsdGeomMesh& UsdMesh, UStaticMeshComponent& MeshComponent )
@@ -394,6 +401,8 @@ FUsdPrimTwin* AUsdStageActor::SpawnPrim( const pxr::SdfPath& UsdPrimPath )
 		return nullptr;
 	}
 
+	FScopedUnrealAllocs UnrealAllocs;
+
 	const bool bNeedsActor = ( ParentUsdPrimTwin->SceneComponent == nullptr || Prim.IsModel() || UsdUtils::HasCompositionArcs( Prim ) || Prim.IsGroup() || Prim.IsA< pxr::UsdGeomScope >() || Prim.IsA< pxr::UsdGeomCamera >() );
 
 	if ( bNeedsActor && !UsdPrimTwin->SpawnedActor.IsValid() )
@@ -805,12 +814,16 @@ void AUsdStageActor::SetTime(float InTime)
 
 void AUsdStageActor::PostEditChangeProperty(FPropertyChangedEvent& PropertyChangedEvent)
 {
+	Super::PostEditChangeProperty( PropertyChangedEvent );
+
 	UProperty* PropertyThatChanged = PropertyChangedEvent.MemberProperty;
 	const FName PropertyName = PropertyThatChanged ? PropertyThatChanged->GetFName() : NAME_None;
 	
 	if ( PropertyName == GET_MEMBER_NAME_CHECKED( AUsdStageActor, RootLayer ) )
 	{
 #if USE_USD_SDK
+		UnrealUSDWrapper::GetUsdStageCache().Erase( UsdStageStore.Get() );
+
 		UsdStageStore = TUsdStore< pxr::UsdStageRefPtr >();
 		LoadUsdStage();
 #endif // #if USE_USD_SDK
@@ -1022,6 +1035,13 @@ void AUsdStageActor::PostRegisterAllComponents()
 #endif // #if USE_USD_SDK
 }
 
+void AUsdStageActor::PostLoad()
+{
+	Super::PostLoad();
+
+	OnActorLoaded.Broadcast( this );
+}
+
 void AUsdStageActor::OnUsdPrimTwinDestroyed( const FUsdPrimTwin& UsdPrimTwin )
 {
 	TArray<FDelegateHandle> DelegateHandles;
@@ -1042,7 +1062,7 @@ void AUsdStageActor::OnUsdPrimTwinDestroyed( const FUsdPrimTwin& UsdPrimTwin )
 void AUsdStageActor::OnPrimObjectPropertyChanged( UObject* ObjectBeingModified, FPropertyChangedEvent& PropertyChangedEvent )
 {
 #if USE_USD_SDK
-	if ( !ObjectsToWatch.Contains( ObjectBeingModified ) )
+	if ( ObjectBeingModified == this || !ObjectsToWatch.Contains( ObjectBeingModified ) )
 	{
 		return;
 	}
diff --git a/Engine/Plugins/Importers/USDImporter/Source/USDStage/Public/USDStageActor.h b/Engine/Plugins/Importers/USDImporter/Source/USDStage/Public/USDStageActor.h
--- a/Engine/Plugins/Importers/USDImporter/Source/USDStage/Public/USDStageActor.h
+++ b/Engine/Plugins/Importers/USDImporter/Source/USDStage/Public/USDStageActor.h
@@ -73,6 +73,9 @@ private:
 	TWeakObjectPtr< ULevelSequence > LevelSequence;
 
 public:
+	DECLARE_EVENT_OneParam( AUsdStageActor, FOnActorLoaded, AUsdStageActor* );
+	USDSTAGE_API static FOnActorLoaded OnActorLoaded;
+
 	DECLARE_EVENT( AUsdStageActor, FOnStageChanged );
 	FOnStageChanged OnStageChanged;
 
@@ -84,12 +87,14 @@ public:
 
 public:
 	AUsdStageActor();
+	virtual ~AUsdStageActor();
 
 	void Refresh() const;
 
 public:
 	virtual void PostEditChangeProperty( FPropertyChangedEvent& PropertyChangedEvent ) override;
 	virtual void PostRegisterAllComponents() override;
+	virtual void PostLoad() override;
 
 private:
 	void Clear();
diff --git a/Engine/Plugins/Importers/USDImporter/Source/USDStageEditor/Private/SUSDStage.cpp b/Engine/Plugins/Importers/USDImporter/Source/USDStageEditor/Private/SUSDStage.cpp
--- a/Engine/Plugins/Importers/USDImporter/Source/USDStageEditor/Private/SUSDStage.cpp
+++ b/Engine/Plugins/Importers/USDImporter/Source/USDStageEditor/Private/SUSDStage.cpp
@@ -19,6 +19,7 @@
 #include "Engine/World.h"
 #include "Framework/MultiBox/MultiBoxBuilder.h"
 #include "Modules/ModuleManager.h"
+#include "ScopedTransaction.h"
 #include "UObject/GCObjectScopeGuard.h"
 #include "UObject/StrongObjectPtr.h"
 #include "Widgets/Layout/SSplitter.h"
@@ -50,6 +51,9 @@ namespace SUSDStageConstants
 
 void SUsdStage::Construct( const FArguments& InArgs )
 {
+	OnStageActorPropertyChangedHandle = FCoreUObjectDelegates::OnObjectPropertyChanged.AddSP( SharedThis( this ), &SUsdStage::OnStageActorPropertyChanged );
+	OnActorLoadedHandle = AUsdStageActor::OnActorLoaded.AddSP( SharedThis( this ), &SUsdStage::OnStageActorLoaded );
+
 	IUsdStageModule& UsdStageModule = FModuleManager::Get().LoadModuleChecked< IUsdStageModule >( "UsdStage" );
 	UsdStageActor = &UsdStageModule.GetUsdStageActor( GWorld );
 
@@ -58,6 +62,11 @@ void SUsdStage::Construct( const FArguments& InArgs )
 	if ( UsdStageActor.IsValid() )
 	{
 		UsdStage = UsdStageActor->GetUsdStage();
+
+		if ( !UsdStage.Get() )
+		{
+			UsdStageActor.Reset();
+		}
 	}
 
 	ChildSlot
@@ -150,37 +159,45 @@ void SUsdStage::SetupStageActorDelegates()
 				{
 					this->UsdStageTreeView->RefreshPrim( PrimPath, bResync );
 				}
+
+				if ( this->UsdPrimInfoWidget )
+				{
+					this->UsdPrimInfoWidget->SetPrimPath( UsdStageActor->GetUsdStage(), *PrimPath );
+				}
 			}
 		);
 
-		OnStageChangedHandle = UsdStageActor->OnStageChanged.AddLambda(
-			[ this ]()
-			{
-				if ( UsdStageActor.IsValid() )
+		if ( !OnStageChangedHandle.IsValid() )
+		{
+			OnStageChangedHandle = UsdStageActor->OnStageChanged.AddLambda(
+				[ this ]()
 				{
-					if ( this->UsdStageInfoWidget )
-					{
-						this->UsdStageInfoWidget->RefreshStageInfos( UsdStageActor.Get() );
-					}
-
-					if ( this->UsdStageTreeView )
-					{
-						this->UsdStageTreeView->Refresh( UsdStageActor.Get() );
-						this->UsdStageTreeView->RequestTreeRefresh();
-					}
-
-					if ( this->UsdPrimInfoWidget )
-					{
-						this->UsdPrimInfoWidget->SetPrimPath( UsdStageActor->GetUsdStage(), TEXT("/") );
-					}
-
-					if ( this->UsdLayersTreeView )
+					if ( UsdStageActor.IsValid() )
 					{
-						this->UsdLayersTreeView->Refresh( UsdStageActor.Get(), true );
+						if ( this->UsdStageInfoWidget )
+						{
+							this->UsdStageInfoWidget->RefreshStageInfos( UsdStageActor.Get() );
+						}
+
+						if ( this->UsdStageTreeView )
+						{
+							this->UsdStageTreeView->Refresh( UsdStageActor.Get() );
+							this->UsdStageTreeView->RequestTreeRefresh();
+						}
+
+						if ( this->UsdPrimInfoWidget )
+						{
+							this->UsdPrimInfoWidget->SetPrimPath( UsdStageActor->GetUsdStage(), TEXT("/") );
+						}
+
+						if ( this->UsdLayersTreeView )
+						{
+							this->UsdLayersTreeView->Refresh( UsdStageActor.Get(), true );
+						}
 					}
 				}
-			}
-		);
+			);
+		}
 
 		OnStageEditTargetChangedHandle = UsdStageActor->GetUsdListener().OnStageEditTargetChanged.AddLambda(
 			[ this ]()
@@ -196,6 +213,9 @@ void SUsdStage::SetupStageActorDelegates()
 
 SUsdStage::~SUsdStage()
 {
+	FCoreUObjectDelegates::OnObjectPropertyChanged.Remove( OnStageActorPropertyChangedHandle );
+	AUsdStageActor::OnActorLoaded.Remove( OnActorLoadedHandle );
+
 	if ( UsdStageActor.IsValid() )
 	{
 		UsdStageActor->OnStageChanged.Remove( OnStageChangedHandle );
@@ -308,25 +328,27 @@ void SUsdStage::OnNew()
 
 	if ( UsdFilePath )
 	{
-		FScopedUsdAllocs UsdAllocs;
+		{
+			FScopedUsdAllocs UsdAllocs;
 
-		pxr::UsdStageCacheContext UsdStageCacheContext( UnrealUSDWrapper::GetUsdStageCache() );
-		pxr::UsdStageRefPtr UsdStage = pxr::UsdStage::CreateNew( UnrealToUsd::ConvertString( *UsdFilePath.GetValue() ).Get() );
+			pxr::UsdStageCacheContext UsdStageCacheContext( UnrealUSDWrapper::GetUsdStageCache() );
+			pxr::UsdStageRefPtr UsdStage = pxr::UsdStage::CreateNew( UnrealToUsd::ConvertString( *UsdFilePath.GetValue() ).Get() );
 
-		if ( !UsdStage )
-		{
-			return;
-		}
+			if ( !UsdStage )
+			{
+				return;
+			}
 
-		// Create default prim
-		pxr::UsdGeomXform RootPrim = pxr::UsdGeomXform::Define( UsdStage, UnrealToUsd::ConvertPath( TEXT("/Root") ).Get() );
-		pxr::UsdModelAPI( RootPrim ).SetKind( pxr::TfToken("component") );
+			// Create default prim
+			pxr::UsdGeomXform RootPrim = pxr::UsdGeomXform::Define( UsdStage, UnrealToUsd::ConvertPath( TEXT("/Root") ).Get() );
+			pxr::UsdModelAPI( RootPrim ).SetKind( pxr::TfToken("component") );
 		
-		// Set default prim
-		UsdStage->SetDefaultPrim( RootPrim.GetPrim() );
+			// Set default prim
+			UsdStage->SetDefaultPrim( RootPrim.GetPrim() );
 
-		// Set up axis
-		UsdUtils::SetUsdStageAxis( UsdStage, pxr::UsdGeomTokens->z );
+			// Set up axis
+			UsdUtils::SetUsdStageAxis( UsdStage, pxr::UsdGeomTokens->z );
+		}
 
 		OpenStage( *UsdFilePath.GetValue() );
 	}
@@ -464,8 +486,10 @@ void SUsdStage::OnPrimSelected( FString PrimPath )
 
 void SUsdStage::OnInitialLoadSetChanged( EUsdInitialLoadSet InitialLoadSet )
 {
-	if ( UsdStageActor.IsValid() )
+	if ( UsdStageActor.IsValid() && UsdStageActor->InitialLoadSet != InitialLoadSet )
 	{
+		const FScopedTransaction Transaction( LOCTEXT("EditInitialLoadSetTransaction", "Edit Initial Load Set") );
+		UsdStageActor->Modify();
 		UsdStageActor->InitialLoadSet = InitialLoadSet;
 	}
 }
@@ -487,6 +511,28 @@ void SUsdStage::OpenStage( const TCHAR* FilePath )
 	UsdStageActor->PostEditChangeProperty( RootLayerPropertyChangedEvent );
 }
 
+void SUsdStage::OnStageActorLoaded( AUsdStageActor* InUsdStageActor )
+{
+	if ( UsdStageActor == InUsdStageActor )
+	{
+		return;
+	}
+
+	UsdStageActor = InUsdStageActor;
+	SetupStageActorDelegates();
+}
+
+void SUsdStage::OnStageActorPropertyChanged( UObject* ObjectBeingModified, FPropertyChangedEvent& PropertyChangedEvent )
+{
+	if ( ObjectBeingModified == UsdStageActor )
+	{
+		if ( this->UsdStageInfoWidget )
+		{
+			this->UsdStageInfoWidget->RefreshStageInfos( UsdStageActor.Get() );
+		}
+	}
+}
+
 #endif // #if USE_USD_SDK
 
 #undef LOCTEXT_NAMESPACE
diff --git a/Engine/Plugins/Importers/USDImporter/Source/USDStageEditor/Private/SUSDStage.h b/Engine/Plugins/Importers/USDImporter/Source/USDStageEditor/Private/SUSDStage.h
--- a/Engine/Plugins/Importers/USDImporter/Source/USDStageEditor/Private/SUSDStage.h
+++ b/Engine/Plugins/Importers/USDImporter/Source/USDStageEditor/Private/SUSDStage.h
@@ -47,6 +47,9 @@ protected:
 
 	void OpenStage( const TCHAR* FilePath );
 
+	void OnStageActorLoaded( AUsdStageActor* InUsdStageActor );
+	void OnStageActorPropertyChanged( UObject* ObjectBeingModified, FPropertyChangedEvent& PropertyChangedEvent );
+
 protected:
 	TSharedPtr< class SUsdStageInfo > UsdStageInfoWidget;
 	TSharedPtr< class SUsdStageTreeView > UsdStageTreeView;
@@ -55,6 +58,8 @@ protected:
 
 	TWeakObjectPtr< AUsdStageActor > UsdStageActor;
 
+	FDelegateHandle OnActorLoadedHandle;
+	FDelegateHandle OnStageActorPropertyChangedHandle;
 	FDelegateHandle OnStageChangedHandle;
 	FDelegateHandle OnStageEditTargetChangedHandle;
 	FDelegateHandle OnPrimChangedHandle;
diff --git a/Engine/Plugins/Importers/USDImporter/Source/USDStageEditor/Private/SUSDStageInfo.cpp b/Engine/Plugins/Importers/USDImporter/Source/USDStageEditor/Private/SUSDStageInfo.cpp
--- a/Engine/Plugins/Importers/USDImporter/Source/USDStageEditor/Private/SUSDStageInfo.cpp
+++ b/Engine/Plugins/Importers/USDImporter/Source/USDStageEditor/Private/SUSDStageInfo.cpp
@@ -63,7 +63,7 @@ void SUsdStageInfo::Construct( const FArguments& InArgs, AUsdStageActor* UsdStag
 			.HAlign( HAlign_Left )
 			.Padding( 2.f, 2.f )
 			[
-				SNew( STextComboBox )
+				SAssignNew( InitialLoadSetWidget, STextComboBox )
 				.OptionsSource( &InitialLoadSetStrings )
 				.InitiallySelectedItem( InitialLoadSetStrings[ (int32)StageInfos.InitialLoadSet ] )
 				.OnSelectionChanged( this, &SUsdStageInfo::OnInitialLoadSetSelectionChanged )
@@ -81,6 +81,13 @@ void SUsdStageInfo::RefreshStageInfos( AUsdStageActor* UsdStageActor )
 		return;
 	}
 
+	StageInfos.InitialLoadSet = UsdStageActor->InitialLoadSet;
+
+	if ( InitialLoadSetWidget && InitialLoadSetStrings.IsValidIndex( (int32)StageInfos.InitialLoadSet ) )
+	{
+		InitialLoadSetWidget->SetSelectedItem( InitialLoadSetStrings[ (int32)StageInfos.InitialLoadSet ] );
+	}
+
 	if ( const pxr::UsdStageRefPtr& UsdStage = UsdStageActor->GetUsdStage() )
 	{
 		TUsdStore< std::string > UsdDisplayName = UsdStage->GetRootLayer()->GetDisplayName();
diff --git a/Engine/Plugins/Importers/USDImporter/Source/USDStageEditor/Private/SUSDStageInfo.h b/Engine/Plugins/Importers/USDImporter/Source/USDStageEditor/Private/SUSDStageInfo.h
--- a/Engine/Plugins/Importers/USDImporter/Source/USDStageEditor/Private/SUSDStageInfo.h
+++ b/Engine/Plugins/Importers/USDImporter/Source/USDStageEditor/Private/SUSDStageInfo.h
@@ -39,6 +39,7 @@ private:
 	FText GetRootLayerDisplayName() const { return StageInfos.RootLayerDisplayName; }
 	FString GetInitialLoadSet() const { return *InitialLoadSetStrings[ (int32)StageInfos.InitialLoadSet ]; }
 
+	TSharedPtr< STextComboBox > InitialLoadSetWidget;
 };
 
 #endif // #if USE_USD_SDK
diff --git a/Engine/Plugins/Importers/USDImporter/Source/USDStageEditor/Private/SUSDStageTreeView.cpp b/Engine/Plugins/Importers/USDImporter/Source/USDStageEditor/Private/SUSDStageTreeView.cpp
--- a/Engine/Plugins/Importers/USDImporter/Source/USDStageEditor/Private/SUSDStageTreeView.cpp
+++ b/Engine/Plugins/Importers/USDImporter/Source/USDStageEditor/Private/SUSDStageTreeView.cpp
@@ -34,6 +34,8 @@
 
 #include "USDIncludesEnd.h"
 
+#include <iterator>
+
 enum class EPayloadsTrigger
 {
 	Load,
@@ -67,7 +69,7 @@ public:
 	{
 		UsdPrim = InUsdPrim;
 
-		RefreshData();
+		RefreshData( false );
 		FillChildren();
 	}
 
@@ -78,27 +80,52 @@ public:
 	{
 	}
 
-	TArray< FUsdStageTreeItemRef > GetChildren()
+	TArray< FUsdStageTreeItemRef >& UpdateChildren()
 	{
 		if ( !UsdPrim.Get() )
 		{
-			return {};
+			return Children;
 		}
 
 		bool bNeedsRefresh = false;
 
-		int32 ChildIndex = 0;
 		pxr::UsdPrimSiblingRange PrimChildren = UsdPrim.Get().GetFilteredChildren( pxr::UsdTraverseInstanceProxies( pxr::UsdPrimAllPrimsPredicate ) );
-		for ( pxr::UsdPrim Child : PrimChildren )
+
+		const int32 NumUsdChildren = (TArray< FUsdStageTreeItemRef >::SizeType )std::distance( PrimChildren.begin(), PrimChildren.end() );
+		const int32 NumUnrealChildren = [&]()
+			{
+				int32 ValidPrims = 0;
+				for ( const FUsdStageTreeItemRef& Child : Children )
+				{
+					if ( !Child->RowData->Name.IsEmpty() )
+					{
+						++ValidPrims;
+					}
+				}
+
+				return ValidPrims;
+			}();
+
+		if ( NumUsdChildren != NumUnrealChildren )
+		{
+			Children.Reset();
+			bNeedsRefresh = true;
+		}
+		else
 		{
-			if ( !Children.IsValidIndex( ChildIndex ) || Children[ ChildIndex ]->UsdPrim.Get().GetPrimPath() != Child.GetPrimPath() )
+			int32 ChildIndex = 0;
+
+			for ( const pxr::UsdPrim& Child : PrimChildren )
 			{
-				Children.Reset();
-				bNeedsRefresh = true;
-				break;
-			}
+				if ( !Children.IsValidIndex( ChildIndex ) || Children[ ChildIndex ]->UsdPrim.Get().GetPrimPath() != Child.GetPrimPath() )
+				{
+					Children.Reset();
+					bNeedsRefresh = true;
+					break;
+				}
 				
-			++ChildIndex;
+				++ChildIndex;
+			}
 		}
 
 		if ( bNeedsRefresh )
@@ -118,7 +145,7 @@ public:
 		}
 	}
 
-	void RefreshData()
+	void RefreshData( bool bRefreshChildren )
 	{
 		FString PrimName;
 
@@ -139,6 +166,14 @@ public:
 		{
 			RowData->bIsVisible = ( UsdGeomImageable.ComputeVisibility() != pxr::UsdGeomTokens->invisible );
 		}
+
+		if ( bRefreshChildren )
+		{
+			for ( FUsdStageTreeItemRef& Child : UpdateChildren() )
+			{
+				Child->RefreshData( bRefreshChildren );
+			}
+		}
 	}
 
 	/** Delegate for hooking up an inline editable text block to be notified that a rename is requested. */
@@ -163,6 +198,8 @@ public:
 	DECLARE_DELEGATE_TwoParams( FOnPrimNameCommitted, const FUsdStageTreeItemRef&, const FText& );
 	FOnPrimNameCommitted OnPrimNameCommitted;
 
+	TWeakPtr< SUsdStageTreeView > OwnerTree;
+
 	virtual TSharedRef< SWidget > GenerateWidget( const TSharedPtr< IUsdTreeViewItem > InTreeItem ) override
 	{
 		if ( !InTreeItem )
@@ -205,7 +242,15 @@ protected:
 
 		if ( TreeItem && TreeItem->RowData->HasCompositionArcs() )
 		{
-			TextColor = FLinearColor( FColor::Orange );
+			const TSharedPtr< SUsdStageTreeView > OwnerTreePtr = OwnerTree.Pin();
+			if ( OwnerTreePtr && OwnerTreePtr->IsItemSelected( TreeItem.ToSharedRef() ) )
+			{
+				TextColor = FSlateColor::UseForeground();
+			}
+			else
+			{
+				TextColor = FLinearColor( FColor::Orange );
+			}
 		}
 		
 		return TextColor;
@@ -377,7 +422,7 @@ TSharedRef< ITableRow > SUsdStageTreeView::OnGenerateRow( FUsdStageTreeItemRef I
 
 void SUsdStageTreeView::OnGetChildren( FUsdStageTreeItemRef InParent, TArray< FUsdStageTreeItemRef >& OutChildren ) const
 {
-	for ( FUsdStageTreeItemRef& Child : InParent->GetChildren() )
+	for ( const FUsdStageTreeItemRef& Child : InParent->UpdateChildren() )
 	{
 		OutChildren.Add( Child );
 	}
@@ -441,12 +486,8 @@ void SUsdStageTreeView::RefreshPrim( const FString& PrimPath, bool bResync )
 	{
 		if ( FUsdStageTreeItemPtr FoundItem = FindTreeItemFromPrimPath( RootItem ) )
 		{
-			FoundItem->RefreshData();
+			FoundItem->RefreshData( true );
 
-			if ( bResync )
-			{
-				FoundItem->Children.Reset();
-			}
 			break;
 		}
 		else
@@ -477,6 +518,7 @@ void SUsdStageTreeView::SetupColumns()
 	
 	{
 		TSharedRef< FUsdStageNameColumn > PrimNameColumn = MakeShared< FUsdStageNameColumn >();
+		PrimNameColumn->OwnerTree = SharedThis( this );
 		PrimNameColumn->bIsMainColumn = true;
 		PrimNameColumn->OnPrimNameCommitted.BindRaw( this, &SUsdStageTreeView::OnPrimNameCommitted );
 
@@ -659,8 +701,10 @@ void SUsdStageTreeView::OnAddReference()
 
 		pxr::UsdReferences References = SelectedItem->UsdPrim.Get().GetReferences();
 
-		FPaths::MakePathRelativeTo( PickedFile.GetValue(), *UsdStageActor->RootLayer.FilePath );
-		References.AddReference( UnrealToUsd::ConvertString( *PickedFile.GetValue() ).Get() );
+		FString AbsoluteFilePath = FPaths::ConvertRelativePathToFull( PickedFile.GetValue() );
+
+		FPaths::MakePathRelativeTo( AbsoluteFilePath, *UsdStageActor->RootLayer.FilePath );
+		References.AddReference( UnrealToUsd::ConvertString( *AbsoluteFilePath ).Get() );
 	}
 }
 
@@ -751,7 +795,7 @@ void SUsdStageTreeView::OnToggleAllPayloads( EPayloadsTrigger PayloadsTrigger )
 				}
 				else
 				{
-					for ( FUsdStageTreeItemRef Child : InSelectedItem->GetChildren() )
+					for ( FUsdStageTreeItemRef Child : InSelectedItem->UpdateChildren() )
 					{
 						RecursiveTogglePayloads( Child );
 					}
diff --git a/Engine/Plugins/Importers/USDImporter/Source/USDStageEditor/Private/Widgets/SUSDTreeView.h b/Engine/Plugins/Importers/USDImporter/Source/USDStageEditor/Private/Widgets/SUSDTreeView.h
--- a/Engine/Plugins/Importers/USDImporter/Source/USDStageEditor/Private/Widgets/SUSDTreeView.h
+++ b/Engine/Plugins/Importers/USDImporter/Source/USDStageEditor/Private/Widgets/SUSDTreeView.h
@@ -24,7 +24,7 @@ public:
 	virtual TSharedRef< SWidget > GenerateWidget( const TSharedPtr< IUsdTreeViewItem > TreeItem ) = 0;
 
 public:
-	bool bIsMainColumn;
+	bool bIsMainColumn = false;
 };
 
 struct FSharedUsdTreeData
