USD Stage Editor

- Fixed initial payload set not modifying the actor.
- Fixed text color so that it's not orange on orange.
- Fixed tree view refresh issues.

diff --git a/Engine/Plugins/Importers/USDImporter/Source/USDStage/Private/USDStageActor.cpp b/Engine/Plugins/Importers/USDImporter/Source/USDStage/Private/USDStageActor.cpp
--- a/Engine/Plugins/Importers/USDImporter/Source/USDStage/Private/USDStageActor.cpp
+++ b/Engine/Plugins/Importers/USDImporter/Source/USDStage/Private/USDStageActor.cpp
@@ -10,7 +10,6 @@
 #include "USDTypesConversion.h"
 #include "UnrealUSDWrapper.h"
 
-#include "Async/ParallelFor.h"
 #include "CineCameraActor.h"
 #include "CineCameraComponent.h"
 #include "Components/PoseableMeshComponent.h"
@@ -27,8 +26,6 @@
 #include "MaterialEditingLibrary.h"
 #include "Materials/Material.h"
 #include "Materials/MaterialInterface.h"
-#include "Materials/MaterialInstanceConstant.h"
-#include "MeshDescription.h"
 #include "MeshDescriptionOperations.h"
 #include "Misc/Paths.h"
 #include "Misc/ScopedSlowTask.h"
@@ -39,6 +36,7 @@
 #include "Channels/MovieSceneChannelProxy.h"
 #include "LevelSequence.h"
 #include "LevelSequenceActor.h"
+#include "MeshDescription.h"
 #include "Misc/FrameNumber.h"
 #include "MovieScene.h"
 #include "Sections/MovieSceneFloatSection.h"
@@ -68,69 +66,104 @@
 
 #define LOCTEXT_NAMESPACE "USDStageActor"
 
-DEFINE_LOG_CATEGORY( LogUsdStage );
-
 static const EObjectFlags DefaultObjFlag = EObjectFlags::RF_Transactional | EObjectFlags::RF_Transient;
 
 AUsdStageActor::FOnActorLoaded AUsdStageActor::OnActorLoaded{};
 
 #if USE_USD_SDK
-
-namespace UsdStageActorImpl
+FMeshDescription LoadMeshDescription( const pxr::UsdGeomMesh& UsdMesh, const pxr::UsdTimeCode TimeCode )
 {
-	FMeshDescription LoadMeshDescription( const pxr::UsdGeomMesh& UsdMesh, const pxr::UsdTimeCode TimeCode )
+	if ( !UsdMesh )
 	{
-		if ( !UsdMesh )
-		{
-			return {};
-		}
+		return {};
+	}
 
-		FMeshDescription MeshDescription;
-		FStaticMeshAttributes StaticMeshAttributes( MeshDescription );
-		StaticMeshAttributes.Register();
+	FMeshDescription MeshDescription;
+	FStaticMeshAttributes StaticMeshAttributes( MeshDescription );
+	StaticMeshAttributes.Register();
 
-		UsdToUnreal::ConvertGeomMesh( UsdMesh, MeshDescription, TimeCode );
+	UsdToUnreal::ConvertGeomMesh( UsdMesh, MeshDescription, TimeCode );
 
-		return MeshDescription;
-	}
+	return MeshDescription;
+}
 
-	void ProcessMaterials( const pxr::UsdPrim& UsdPrim, UStaticMesh& StaticMesh, TMap< FString, UMaterial* >& MaterialsCache, bool bHasPrimDisplayColor, float Time )
+void ProcessMaterials( const pxr::UsdStageRefPtr& Stage, UStaticMesh* StaticMesh, const FMeshDescription& MeshDescription, TMap< FString, UMaterial* >& MaterialsCache, bool bHasPrimDisplayColor )
+{
+	FStaticMeshConstAttributes StaticMeshAttributes( MeshDescription );
+
+	for ( const FPolygonGroupID PolygonGroupID : MeshDescription.PolygonGroups().GetElementIDs() )
 	{
-		const FMeshDescription* MeshDescription = StaticMesh.GetMeshDescription( 0 );
+		const FName& ImportedMaterialSlotName = StaticMeshAttributes.GetPolygonGroupMaterialSlotNames()[ PolygonGroupID ];
+		const FName MaterialSlotName = ImportedMaterialSlotName;
+
+		int32 MaterialIndex = INDEX_NONE;
+		int32 MeshMaterialIndex = 0;
 
-		if ( !MeshDescription )
+		for ( FStaticMaterial& StaticMaterial : StaticMesh->StaticMaterials )
 		{
-			return;
-		}
+			if ( StaticMaterial.MaterialSlotName.IsEqual( ImportedMaterialSlotName ) )
+			{
+				MaterialIndex = MeshMaterialIndex;
+				break;
+			}
 
-		FStaticMeshConstAttributes StaticMeshAttributes( *MeshDescription );
+			++MeshMaterialIndex;
+		}
 
-		for ( const FPolygonGroupID PolygonGroupID : MeshDescription->PolygonGroups().GetElementIDs() )
+		if ( MaterialIndex == INDEX_NONE )
 		{
-			const FName& ImportedMaterialSlotName = StaticMeshAttributes.GetPolygonGroupMaterialSlotNames()[ PolygonGroupID ];
-			const FName MaterialSlotName = ImportedMaterialSlotName;
+			MaterialIndex = PolygonGroupID.GetValue();
+		}
 
-			int32 MaterialIndex = INDEX_NONE;
-			int32 MeshMaterialIndex = 0;
+		UMaterialInterface* Material = nullptr;
+		pxr::UsdPrim MaterialPrim = Stage->GetPrimAtPath( UnrealToUsd::ConvertPath( *ImportedMaterialSlotName.ToString() ).Get() );
 
-			for ( FStaticMaterial& StaticMaterial : StaticMesh.StaticMaterials )
+		if ( MaterialPrim )
+		{
+			UMaterial*& CachedMaterial = MaterialsCache.FindOrAdd( UsdToUnreal::ConvertPath( MaterialPrim.GetPrimPath() ) );
+
+			if ( !CachedMaterial )
 			{
-				if ( StaticMaterial.MaterialSlotName.IsEqual( ImportedMaterialSlotName ) )
+				CachedMaterial = NewObject< UMaterial >();
+
+				if ( UsdToUnreal::ConvertMaterial( pxr::UsdShadeMaterial( MaterialPrim ), *CachedMaterial ) )
 				{
-					MaterialIndex = MeshMaterialIndex;
-					break;
+					//UMaterialEditingLibrary::RecompileMaterial( CachedMaterial ); // Too slow
+					CachedMaterial->PostEditChange();
+				}
+				else
+				{
+					CachedMaterial = nullptr;
 				}
-
-				++MeshMaterialIndex;
 			}
 
-			if ( MaterialIndex == INDEX_NONE )
-			{
-				MaterialIndex = PolygonGroupID.GetValue();
-			}
+			Material = CachedMaterial;
+		}
+
+		if ( Material == nullptr && bHasPrimDisplayColor )
+		{
+			FSoftObjectPath VertexColorMaterialPath( TEXT("Material'/Engine/EngineDebugMaterials/VertexColorMaterial.VertexColorMaterial'") );
+			Material = Cast< UMaterialInterface >( VertexColorMaterialPath.TryLoad() );
+		}
+
+		FStaticMaterial StaticMaterial( Material, MaterialSlotName, ImportedMaterialSlotName );
+		StaticMesh->StaticMaterials.Add( StaticMaterial );
+
+		FMeshSectionInfo MeshSectionInfo;
+		MeshSectionInfo.MaterialIndex = MaterialIndex;
+
+		StaticMesh->GetSectionInfoMap().Set( 0, PolygonGroupID.GetValue(), MeshSectionInfo );
+	}
+}
 
+void ProcessMaterials( const pxr::UsdStageRefPtr& Stage, FSkeletalMeshImportData& SkelMeshImportData, TMap< FString, UMaterial* >& MaterialsCache, bool bHasPrimDisplayColor )
+{
+	for (SkeletalMeshImportData::FMaterial& ImportedMaterial : SkelMeshImportData.Materials)
+	{
+		if (!ImportedMaterial.Material.IsValid())
+		{
 			UMaterialInterface* Material = nullptr;
-			TUsdStore< pxr::UsdPrim > MaterialPrim = UsdPrim.GetStage()->GetPrimAtPath( UnrealToUsd::ConvertPath( *ImportedMaterialSlotName.ToString() ).Get() );
+			TUsdStore< pxr::UsdPrim > MaterialPrim = Stage->GetPrimAtPath( UnrealToUsd::ConvertPath( *ImportedMaterial.MaterialImportName ).Get() );
 
 			if ( MaterialPrim.Get() )
 			{
@@ -142,7 +175,10 @@ namespace UsdStageActorImpl
 
 					if ( UsdToUnreal::ConvertMaterial( pxr::UsdShadeMaterial( MaterialPrim.Get() ), *CachedMaterial ) )
 					{
-						//UMaterialEditingLibrary::RecompileMaterial( CachedMaterial ); // Too slow
+						CachedMaterial->bUsedWithSkeletalMesh = true;
+						bool bNeedsRecompile = false;
+						CachedMaterial->GetMaterial()->SetMaterialUsage(bNeedsRecompile, MATUSAGE_SkeletalMesh);
+
 						CachedMaterial->PostEditChange();
 					}
 					else
@@ -150,78 +186,15 @@ namespace UsdStageActorImpl
 						CachedMaterial = nullptr;
 					}
 				}
-
 				Material = CachedMaterial;
 			}
 
 			if ( Material == nullptr && bHasPrimDisplayColor )
 			{
-				UMaterialInstanceConstant* MaterialInstance = NewObject< UMaterialInstanceConstant >();
-				if ( UsdToUnreal::ConvertDisplayColor( pxr::UsdGeomMesh( UsdPrim ), *MaterialInstance, pxr::UsdTimeCode( Time ) ) )
-				{
-					Material = MaterialInstance;
-				}
-			}
-
-			FStaticMaterial StaticMaterial( Material, MaterialSlotName, ImportedMaterialSlotName );
-				StaticMesh.StaticMaterials.Add( StaticMaterial );
-
-			FMeshSectionInfo MeshSectionInfo;
-			MeshSectionInfo.MaterialIndex = MaterialIndex;
-
-			StaticMesh.GetSectionInfoMap().Set( 0, PolygonGroupID.GetValue(), MeshSectionInfo );
-		}
-	}
-
-	void ProcessMaterials( const pxr::UsdPrim& UsdPrim, FSkeletalMeshImportData& SkelMeshImportData, TMap< FString, UMaterial* >& MaterialsCache, bool bHasPrimDisplayColor, float Time )
-	{
-		for (SkeletalMeshImportData::FMaterial& ImportedMaterial : SkelMeshImportData.Materials)
-		{
-			if (!ImportedMaterial.Material.IsValid())
-			{
-				UMaterialInterface* Material = nullptr;
-				TUsdStore< pxr::UsdPrim > MaterialPrim = UsdPrim.GetStage()->GetPrimAtPath( UnrealToUsd::ConvertPath( *ImportedMaterial.MaterialImportName ).Get() );
-
-				if ( MaterialPrim.Get() )
-				{
-					UMaterial*& CachedMaterial = MaterialsCache.FindOrAdd( UsdToUnreal::ConvertPath( MaterialPrim.Get().GetPrimPath() ) );
-
-					if ( !CachedMaterial )
-					{
-						CachedMaterial = NewObject< UMaterial >();
-
-						if ( UsdToUnreal::ConvertMaterial( pxr::UsdShadeMaterial( MaterialPrim.Get() ), *CachedMaterial ) )
-						{
-							CachedMaterial->bUsedWithSkeletalMesh = true;
-							bool bNeedsRecompile = false;
-							CachedMaterial->GetMaterial()->SetMaterialUsage(bNeedsRecompile, MATUSAGE_SkeletalMesh);
-
-							CachedMaterial->PostEditChange();
-						}
-						else
-						{
-							CachedMaterial = nullptr;
-						}
-					}
-					Material = CachedMaterial;
-				}
-
-				if ( Material == nullptr && bHasPrimDisplayColor )
-				{
-					UMaterialInstanceConstant* MaterialInstance = NewObject< UMaterialInstanceConstant >();
-					if ( UsdToUnreal::ConvertDisplayColor( pxr::UsdGeomMesh( UsdPrim ), *MaterialInstance, pxr::UsdTimeCode( Time ) ) )
-					{
-						if ( MaterialInstance && MaterialInstance->HasOverridenBaseProperties() )
-						{
-							MaterialInstance->ForceRecompileForRendering();
-						}
-						
-						Material = MaterialInstance;
-					}
-				}
-
-				ImportedMaterial.Material = Material;
+				FSoftObjectPath VertexColorMaterialPath( TEXT("Material'/Engine/EngineDebugMaterials/VertexColorMaterial.VertexColorMaterial'") );
+				Material = Cast< UMaterialInterface >( VertexColorMaterialPath.TryLoad() );
 			}
+			ImportedMaterial.Material = Material;
 		}
 	}
 }
@@ -243,10 +216,7 @@ AUsdStageActor::AUsdStageActor()
 		[ this ]( const FString& PrimPath, bool bResync )
 			{
 				TUsdStore< pxr::SdfPath > UsdPrimPath = UnrealToUsd::ConvertPath( *PrimPath );
-
-				TMap< FString, UStaticMesh* > PrimPathsToStaticMeshes = this->LoadStaticMeshes( this->GetUsdStage()->GetPrimAtPath( UsdPrimPath.Get() ) );
-
-				this->UpdatePrim( UsdPrimPath.Get(), bResync, PrimPathsToStaticMeshes );
+				this->UpdatePrim( UsdPrimPath.Get(), bResync );
 
 				this->OnPrimChanged.Broadcast( PrimPath, bResync );
 			}
@@ -261,18 +231,87 @@ AUsdStageActor::AUsdStageActor()
 
 AUsdStageActor::~AUsdStageActor()
 {
-#if USE_USD_SDK
 	UnrealUSDWrapper::GetUsdStageCache().Erase( UsdStageStore.Get() );
-#endif // #if USE_USD_SDK
 }
 
 #if USE_USD_SDK
 
-bool AUsdStageActor::SetStaticMesh( const pxr::UsdGeomMesh& UsdMesh, UStaticMeshComponent& MeshComponent, const TMap< FString, UStaticMesh* >& PrimPathsToStaticMeshes )
+bool AUsdStageActor::LoadStaticMesh( const pxr::UsdGeomMesh& UsdMesh, UStaticMeshComponent& MeshComponent )
 {
+	UStaticMesh* StaticMesh = nullptr;
+
+	pxr::UsdGeomXformable Xformable( UsdMesh.GetPrim() );
+
 	FString MeshPrimPath = UsdToUnreal::ConvertPath( UsdMesh.GetPrim().GetPrimPath() );
 
-	UStaticMesh* StaticMesh = PrimPathsToStaticMeshes.FindRef( MeshPrimPath );
+	FMeshDescription MeshDescription = LoadMeshDescription( UsdMesh, pxr::UsdTimeCode( Time ) );
+
+	FSHAHash MeshHash = FMeshDescriptionOperations::ComputeSHAHash( MeshDescription );
+
+	StaticMesh = MeshCache.FindRef( MeshHash.ToString() );
+
+	if ( !StaticMesh && !MeshDescription.IsEmpty() )
+	{
+		StaticMesh = NewObject< UStaticMesh >( GetTransientPackage(), NAME_None, DefaultObjFlag | EObjectFlags::RF_Public );
+
+		FStaticMeshSourceModel& SourceModel = StaticMesh->AddSourceModel();
+		SourceModel.BuildSettings.bGenerateLightmapUVs = false;
+		SourceModel.BuildSettings.bRecomputeNormals = false;
+		SourceModel.BuildSettings.bRecomputeTangents = false;
+		SourceModel.BuildSettings.bBuildAdjacencyBuffer = false;
+		SourceModel.BuildSettings.bBuildReversedIndexBuffer = false;
+
+		FMeshDescription* StaticMeshDescription = StaticMesh->CreateMeshDescription(0);
+		check( StaticMeshDescription );
+		*StaticMeshDescription = MoveTemp( MeshDescription );
+		StaticMesh->CommitMeshDescription(0);
+
+		const bool bHasPrimDisplayColor = UsdMesh.GetDisplayColorPrimvar().IsDefined();
+		ProcessMaterials( GetUsdStage(), StaticMesh, *StaticMeshDescription, MaterialsCache, bHasPrimDisplayColor );
+
+		// Create render data
+		if ( !StaticMesh->RenderData )
+		{
+			StaticMesh->RenderData.Reset(new(FMemory::Malloc(sizeof(FStaticMeshRenderData)))FStaticMeshRenderData());
+		}
+
+		ITargetPlatformManagerModule& TargetPlatformManager = GetTargetPlatformManagerRef();
+		ITargetPlatform* RunningPlatform = TargetPlatformManager.GetRunningTargetPlatform();
+		check(RunningPlatform);
+		const FStaticMeshLODSettings& LODSettings = RunningPlatform->GetStaticMeshLODSettings();
+
+		const FStaticMeshLODGroup& LODGroup = LODSettings.GetLODGroup(StaticMesh->LODGroup);
+
+		IMeshBuilderModule& MeshBuilderModule = FModuleManager::LoadModuleChecked< IMeshBuilderModule >( TEXT("MeshBuilder") );
+		if ( !MeshBuilderModule.BuildMesh( *StaticMesh->RenderData, StaticMesh, LODGroup ) )
+		{
+			UE_LOG(LogStaticMesh, Error, TEXT("Failed to build static mesh. See previous line(s) for details."));
+			return false;
+		}
+
+		for ( FStaticMeshLODResources& LODResources : StaticMesh->RenderData->LODResources )
+		{
+			LODResources.bHasColorVertexData = true;
+		}
+
+		// Bounds
+		StaticMesh->CalculateExtendedBounds();
+
+		// Recreate resources
+		StaticMesh->InitResources();
+
+		// Other setup
+		//StaticMesh->CreateBodySetup();
+
+		// Collision
+		//StaticMesh->BodySetup->CollisionTraceFlag = ECollisionTraceFlag::CTF_UseComplexAsSimple;
+
+		MeshCache.Add( MeshHash.ToString() ) = StaticMesh;
+	}
+	/*else
+	{
+		FPlatformMisc::LowLevelOutputDebugStringf( TEXT("Mesh found in cache %s\n"), *StaticMesh->GetName() );
+	}*/
 
 	if ( StaticMesh != MeshComponent.GetStaticMesh() )
 	{
@@ -286,7 +325,7 @@ bool AUsdStageActor::SetStaticMesh( const pxr::UsdGeomMesh& UsdMesh, UStaticMesh
 		MeshComponent.RegisterComponent();
 	}
 
-	return ( StaticMesh != nullptr );
+	return true;
 }
 
 FUsdPrimTwin* AUsdStageActor::SpawnPrim( const pxr::SdfPath& UsdPrimPath )
@@ -319,11 +358,41 @@ FUsdPrimTwin* AUsdStageActor::SpawnPrim( const pxr::SdfPath& UsdPrimPath )
 				this->OnUsdPrimTwinDestroyed( UsdPrimTwin );
 			} );
 
-		if ( !UsdPrimTwin->AnimationHandle.IsValid() && UsdUtils::IsAnimated( Prim ) )
+		if ( !UsdPrimTwin->AnimationHandle.IsValid() )
+		{
+			bool bHasXformbaleTimeSamples = false;
+			{
+				pxr::UsdGeomXformable Xformable( Prim );
+
+				if ( Xformable )
+				{
+					std::vector< double > TimeSamples;
+					Xformable.GetTimeSamples( &TimeSamples );
+
+					bHasXformbaleTimeSamples = TimeSamples.size() > 0;
+				}
+			}
+
+			bool bHasAttributesTimeSamples = false;
+			{
+				std::vector< pxr::UsdAttribute > Attributes = Prim.GetAttributes();
+
+				for ( pxr::UsdAttribute& Attribute : Attributes )
+				{
+					bHasAttributesTimeSamples = Attribute.ValueMightBeTimeVarying();
+					if ( bHasAttributesTimeSamples )
+					{
+						break;
+					}
+				}
+			}
+
+			if ( ( bHasXformbaleTimeSamples || bHasAttributesTimeSamples ) )
 			{
 				PrimsToAnimate.Add( PrimPath );
 			}
 		}
+	}
 
 	UClass* ComponentClass = UsdUtils::GetComponentTypeForPrim( Prim );
 
@@ -334,7 +403,7 @@ FUsdPrimTwin* AUsdStageActor::SpawnPrim( const pxr::SdfPath& UsdPrimPath )
 
 	FScopedUnrealAllocs UnrealAllocs;
 
-	const bool bNeedsActor = ( ParentUsdPrimTwin->SceneComponent == nullptr || Prim.IsModel() || UsdUtils::HasCompositionArcs( Prim ) || Prim.IsGroup() || Prim.IsA< pxr::UsdGeomCamera >() );
+	const bool bNeedsActor = ( ParentUsdPrimTwin->SceneComponent == nullptr || Prim.IsModel() || UsdUtils::HasCompositionArcs( Prim ) || Prim.IsGroup() || Prim.IsA< pxr::UsdGeomScope >() || Prim.IsA< pxr::UsdGeomCamera >() );
 
 	if ( bNeedsActor && !UsdPrimTwin->SpawnedActor.IsValid() )
 	{
@@ -433,7 +502,7 @@ FUsdPrimTwin* AUsdStageActor::SpawnPrim( const pxr::SdfPath& UsdPrimPath )
 	return UsdPrimTwin;
 }
 
-FUsdPrimTwin* AUsdStageActor::LoadPrim( const pxr::SdfPath& Path, const TMap< FString, UStaticMesh* >& PrimPathsToStaticMeshes )
+FUsdPrimTwin* AUsdStageActor::LoadPrim( const pxr::SdfPath& Path )
 {
 	TRACE_CPUPROFILER_EVENT_SCOPE( AUsdStageActor::LoadPrim );
 
@@ -468,7 +537,7 @@ FUsdPrimTwin* AUsdStageActor::LoadPrim( const pxr::SdfPath& Path, const TMap< FS
 
 	if ( UStaticMeshComponent* MeshComponent = Cast< UStaticMeshComponent >( UsdPrimTwin->SceneComponent.Get() ) )
 	{
-		SetStaticMesh( pxr::UsdGeomMesh( Prim ), *MeshComponent, PrimPathsToStaticMeshes );
+		LoadStaticMesh( pxr::UsdGeomMesh( Prim ), *MeshComponent );
 	}
 	else if ( UCineCameraComponent* CameraComponent = Cast< UCineCameraComponent >( UsdPrimTwin->SceneComponent.Get() ) )
 	{
@@ -491,16 +560,7 @@ FUsdPrimTwin* AUsdStageActor::LoadPrim( const pxr::SdfPath& Path, const TMap< FS
 		}
 
 		ComponentToAttach->AttachToComponent( AttachToComponent, FAttachmentTransformRules::KeepRelativeTransform );
-		
-		if ( PrimsToAnimate.Contains( UsdToUnreal::ConvertPath( Path.GetPrimPath() ) ) )
-		{
-			ComponentToAttach->SetMobility( EComponentMobility::Movable );
-		}
-		else
-		{
-			ComponentToAttach->SetMobility( EComponentMobility::Static );
-		}
-
+		UsdPrimTwin->SceneComponent->SetMobility( EComponentMobility::Movable );
 		UsdPrimTwin->SceneComponent->GetOwner()->AddInstanceComponent( UsdPrimTwin->SceneComponent.Get() );
 	}
 
@@ -542,7 +602,7 @@ FUsdPrimTwin* AUsdStageActor::LoadPrim( const pxr::SdfPath& Path, const TMap< FS
 	return UsdPrimTwin;
 }
 
-FUsdPrimTwin* AUsdStageActor::ExpandPrim( const pxr::UsdPrim& Prim, const TMap< FString, UStaticMesh* >& PrimPathsToStaticMeshes )
+FUsdPrimTwin* AUsdStageActor::ExpandPrim( const pxr::UsdPrim& Prim )
 {
 	if ( !Prim )
 	{
@@ -551,7 +611,7 @@ FUsdPrimTwin* AUsdStageActor::ExpandPrim( const pxr::UsdPrim& Prim, const TMap<
 
 	TRACE_CPUPROFILER_EVENT_SCOPE( AUsdStageActor::ExpandPrim );
 
-	FUsdPrimTwin* UsdPrimTwin = LoadPrim( Prim.GetPrimPath(), PrimPathsToStaticMeshes );
+	FUsdPrimTwin* UsdPrimTwin = LoadPrim( Prim.GetPrimPath() );
 
 	if ( Prim.IsA<pxr::UsdSkelRoot>() )
 	{
@@ -566,7 +626,7 @@ FUsdPrimTwin* AUsdStageActor::ExpandPrim( const pxr::UsdPrim& Prim, const TMap<
 		pxr::UsdPrimSiblingRange PrimChildren = Prim.GetFilteredChildren( pxr::UsdTraverseInstanceProxies() );
 		for ( TUsdStore< pxr::UsdPrim > ChildStore : PrimChildren )
 		{
-			ExpandPrim( ChildStore.Get(), PrimPathsToStaticMeshes );
+			ExpandPrim( ChildStore.Get() );
 		}
 	}
 
@@ -598,8 +658,6 @@ bool AUsdStageActor::ProcessSkeletonRoot( const pxr::UsdPrim& Prim, USkinnedMesh
 			return false;
 		}
 
-		const pxr::TfToken StageUpAxis = UsdUtils::GetUsdStageAxis( GetUsdStage() );
-
 		pxr::UsdGeomXformable Xformable(Prim);
 
 		std::vector<double> TimeSamples;
@@ -649,7 +707,7 @@ bool AUsdStageActor::ProcessSkeletonRoot( const pxr::UsdPrim& Prim, USkinnedMesh
 
 			if (TimeSamples.size() > 0)
 			{
-				auto UpdateBoneTransforms = [&, SkelQuery, StageUpAxis]()
+				auto UpdateBoneTransforms = [&, SkelQuery]()
 				{
 					FScopedUnrealAllocs UnrealAllocs;
 
@@ -667,7 +725,7 @@ bool AUsdStageActor::ProcessSkeletonRoot( const pxr::UsdPrim& Prim, USkinnedMesh
 						for (uint32 Index = 0; Index < UsdBoneTransforms.size(); ++Index)
 						{
 							const pxr::GfMatrix4d& UsdMatrix = UsdBoneTransforms[Index];
-							FTransform BoneTransform = UsdToUnreal::ConvertMatrix( StageUpAxis, UsdMatrix );
+							FTransform BoneTransform = UsdToUnreal::ConvertMatrix(GetUsdStage(), UsdMatrix);
 							BoneTransforms.Add(BoneTransform);
 						}
 					}
@@ -683,7 +741,7 @@ bool AUsdStageActor::ProcessSkeletonRoot( const pxr::UsdPrim& Prim, USkinnedMesh
 				PrimDelegates.Add(UsdToUnreal::ConvertPath(PrimPath), Handle);
 			}
 
-			UsdStageActorImpl::ProcessMaterials(Prim, SkelMeshImportData, MaterialsCache, bHasPrimDisplayColor, Time);
+			ProcessMaterials(GetUsdStage(), SkelMeshImportData, MaterialsCache, bHasPrimDisplayColor);
 
 			break;
 		}
@@ -704,7 +762,7 @@ bool AUsdStageActor::ProcessSkeletonRoot( const pxr::UsdPrim& Prim, USkinnedMesh
 	return SkeletalMesh != nullptr;
 }
 
-void AUsdStageActor::UpdatePrim( const pxr::SdfPath& InUsdPrimPath, bool bResync, const TMap< FString, UStaticMesh* >& PrimPathsToStaticMeshes )
+void AUsdStageActor::UpdatePrim( const pxr::SdfPath& InUsdPrimPath, bool bResync )
 {
 	TRACE_CPUPROFILER_EVENT_SCOPE( AUsdStageActor::UpdatePrim );
 
@@ -731,7 +789,7 @@ void AUsdStageActor::UpdatePrim( const pxr::SdfPath& InUsdPrimPath, bool bResync
 		}
 
 		TUsdStore< pxr::UsdPrim > PrimToExpand = GetUsdStage()->GetPrimAtPath( UsdPrimPath );
-		FUsdPrimTwin* UsdPrimTwin = ExpandPrim( PrimToExpand.Get(), PrimPathsToStaticMeshes );
+		FUsdPrimTwin* UsdPrimTwin = ExpandPrim( PrimToExpand.Get() );
 
 		GEditor->BroadcastLevelActorListChanged();
 		GEditor->RedrawLevelEditingViewports();
@@ -756,6 +814,8 @@ void AUsdStageActor::SetTime(float InTime)
 
 void AUsdStageActor::PostEditChangeProperty(FPropertyChangedEvent& PropertyChangedEvent)
 {
+	Super::PostEditChangeProperty( PropertyChangedEvent );
+
 	UProperty* PropertyThatChanged = PropertyChangedEvent.MemberProperty;
 	const FName PropertyName = PropertyThatChanged ? PropertyThatChanged->GetFName() : NAME_None;
 	
@@ -772,10 +832,6 @@ void AUsdStageActor::PostEditChangeProperty(FPropertyChangedEvent& PropertyChang
 	{
 		Refresh();
 	}
-	else
-	{
-		Super::PostEditChangeProperty( PropertyChangedEvent );
-	}
 }
 
 void AUsdStageActor::Clear()
@@ -833,7 +889,7 @@ void AUsdStageActor::OpenUsdStage()
 
 void AUsdStageActor::InitLevelSequence(float FramesPerSecond)
 {
-	if (LevelSequence)
+	if (LevelSequence.IsValid())
 	{
 		return;
 	}
@@ -867,7 +923,7 @@ void AUsdStageActor::InitLevelSequence(float FramesPerSecond)
 
 void AUsdStageActor::SetupLevelSequence()
 {
-	if (!LevelSequence)
+	if (!LevelSequence.IsValid())
 	{
 		return;
 	}
@@ -915,8 +971,6 @@ void AUsdStageActor::SetupLevelSequence()
 void AUsdStageActor::LoadUsdStage()
 {
 #if USE_USD_SDK
-	double StartTime = FPlatformTime::Cycles64();
-
 	FScopedSlowTask SlowTask( 1.f, LOCTEXT( "LoadingUDStage", "Loading USD Stage") );
 	SlowTask.MakeDialog();
 	SlowTask.EnterProgressFrame();
@@ -940,9 +994,7 @@ void AUsdStageActor::LoadUsdStage()
 
 	SetupLevelSequence();
 
-	TMap< FString, UStaticMesh* > PrimPathsToStaticMeshes = LoadStaticMeshes( UsdStage->GetPseudoRoot() );
-
-	UpdatePrim( UsdStage->GetPseudoRoot().GetPrimPath(), true, PrimPathsToStaticMeshes );
+	UpdatePrim( UsdStage->GetPseudoRoot().GetPrimPath(), true );
 
 	SetTime( StartTimeCode );
 
@@ -953,10 +1005,7 @@ void AUsdStageActor::LoadUsdStage()
 	{
 		for ( const FString& PrimToAnimate : PrimsToAnimate )
 		{
-			TUsdStore< pxr::SdfPath > PrimPath = UnrealToUsd::ConvertPath( *PrimToAnimate );
-
-			TMap< FString, UStaticMesh* > PrimPathsToStaticMeshes = LoadStaticMeshes( this->GetUsdStage()->GetPrimAtPath( PrimPath.Get() ) );
-			this->LoadPrim( PrimPath.Get(), PrimPathsToStaticMeshes );
+			this->LoadPrim( UnrealToUsd::ConvertPath( *PrimToAnimate ).Get() );
 		}
 
 		GEditor->BroadcastLevelActorListChanged();
@@ -964,14 +1013,6 @@ void AUsdStageActor::LoadUsdStage()
 	};
 
 	OnTimeChanged.AddLambda( AnimatePrims );
-
-	// Log time spent to load the stage
-	double ElapsedSeconds = FPlatformTime::ToSeconds64(FPlatformTime::Cycles64() - StartTime);
-
-	int ElapsedMin = int(ElapsedSeconds / 60.0);
-	ElapsedSeconds -= 60.0 * (double)ElapsedMin;
-
-	UE_LOG( LogUsdStage, Log, TEXT("%s %s in [%d min %.3f s]"), TEXT("Stage loaded"), *FPaths::GetBaseFilename( RootLayer.FilePath ), ElapsedMin, ElapsedSeconds );
 #endif // #if USE_USD_SDK
 }
 
@@ -989,8 +1030,7 @@ void AUsdStageActor::PostRegisterAllComponents()
 
 	if ( UsdStage )
 	{
-		TMap< FString, UStaticMesh* > PrimPathsToStaticMeshes = LoadStaticMeshes( UsdStage->GetPseudoRoot() );
-		UpdatePrim( UsdStage->GetPseudoRoot().GetPrimPath(), true, PrimPathsToStaticMeshes );
+		UpdatePrim( UsdStage->GetPseudoRoot().GetPrimPath(), true );
 	}
 #endif // #if USE_USD_SDK
 }
@@ -1055,198 +1095,4 @@ void AUsdStageActor::OnPrimObjectPropertyChanged( UObject* ObjectBeingModified,
 #endif // #if USE_USD_SDK
 }
 
-namespace UsdStageActorImpl
-{
-	void PreBuildStaticMesh( const pxr::UsdGeomMesh& UsdMesh, UStaticMesh& StaticMesh, TMap< FString, UMaterial* >& MaterialsCache, float Time )
-	{
-		const bool bHasPrimDisplayColor = UsdMesh.GetDisplayColorPrimvar().IsDefined();
-		ProcessMaterials( UsdMesh.GetPrim(), StaticMesh, MaterialsCache, bHasPrimDisplayColor, Time );
-
-		StaticMesh.RenderData = MakeUnique< FStaticMeshRenderData >();
-	}
-
-	bool BuildStaticMesh( UStaticMesh& StaticMesh, IMeshBuilderModule& MeshBuilderModule )
-	{
-		ITargetPlatformManagerModule& TargetPlatformManager = GetTargetPlatformManagerRef();
-		ITargetPlatform* RunningPlatform = TargetPlatformManager.GetRunningTargetPlatform();
-		check(RunningPlatform);
-
-		const FStaticMeshLODSettings& LODSettings = RunningPlatform->GetStaticMeshLODSettings();
-
-		const FStaticMeshLODGroup& LODGroup = LODSettings.GetLODGroup(StaticMesh.LODGroup);
-
-		if ( !MeshBuilderModule.BuildMesh( *StaticMesh.RenderData, &StaticMesh, LODGroup ) )
-		{
-			return false;
-		}
-
-		for ( FStaticMeshLODResources& LODResources : StaticMesh.RenderData->LODResources )
-		{
-			LODResources.bHasColorVertexData = true;
-		}
-
-		return true;
-	}
-
-	void PostBuildStaticMesh( UStaticMesh& StaticMesh )
-	{
-		StaticMesh.InitResources();
-
-		if ( const FMeshDescription* MeshDescription = StaticMesh.GetMeshDescription( 0 ) )
-		{
-			StaticMesh.RenderData->Bounds = MeshDescription->GetBounds();
-		}
-	
-		StaticMesh.CalculateExtendedBounds();
-	}
-}
-
-UStaticMesh* AUsdStageActor::CreateStaticMesh( const pxr::UsdGeomMesh& UsdMesh, FMeshDescription&& MeshDescription, bool& bOutIsNew )
-{
-	UStaticMesh* StaticMesh = nullptr;
-
-	FSHAHash MeshHash = FMeshDescriptionOperations::ComputeSHAHash( MeshDescription );
-
-	StaticMesh = MeshCache.FindRef( MeshHash.ToString() );
-
-	if ( !StaticMesh && !MeshDescription.IsEmpty() )
-	{
-		bOutIsNew = true;
-
-		StaticMesh = NewObject< UStaticMesh >( GetTransientPackage(), NAME_None, DefaultObjFlag | EObjectFlags::RF_Public );
-
-		FStaticMeshSourceModel& SourceModel = StaticMesh->AddSourceModel();
-		SourceModel.BuildSettings.bGenerateLightmapUVs = false;
-		SourceModel.BuildSettings.bRecomputeNormals = false;
-		SourceModel.BuildSettings.bRecomputeTangents = false;
-		SourceModel.BuildSettings.bBuildAdjacencyBuffer = false;
-		SourceModel.BuildSettings.bBuildReversedIndexBuffer = false;
-
-		FMeshDescription* StaticMeshDescription = StaticMesh->CreateMeshDescription(0);
-		check( StaticMeshDescription );
-		*StaticMeshDescription = MoveTemp( MeshDescription );
-
-		MeshCache.Add( MeshHash.ToString() ) = StaticMesh;
-	}
-	else
-	{
-		//FPlatformMisc::LowLevelOutputDebugStringf( TEXT("Mesh found in cache %s\n"), *StaticMesh->GetName() );
-		bOutIsNew = false;
-	}
-
-	return StaticMesh;
-}
-
-TMap< FString, UStaticMesh* > AUsdStageActor::LoadStaticMeshes( const pxr::UsdPrim& StartPrim )
-{
-	TMap< FString, UStaticMesh* > PrimPathsToStaticMeshes;
-
-	FScopedSlowTask LoadAllProgress( 6.f, LOCTEXT( "LoadingStaticMeshes", "Loading Static Meshes" ) );
-	LoadAllProgress.MakeDialog();
-
-	TArray< TUsdStore< pxr::UsdPrim > > AllGeomMeshes = UsdUtils::GetAllPrimsOfType( StartPrim, pxr::TfType::Find< pxr::UsdGeomMesh >(), { pxr::TfType::Find< pxr::UsdSkelRoot >() } );
-	TArray< FMeshDescription > MeshDescriptions;
-	MeshDescriptions.AddDefaulted( AllGeomMeshes.Num() );
-
-	// Create all mesh descriptions
-	{
-		LoadAllProgress.EnterProgressFrame();
-
-		ParallelFor( AllGeomMeshes.Num(), 
-			[ &AllGeomMeshes, &MeshDescriptions, TimeCode = pxr::UsdTimeCode( Time ) ]( int32 GeomMeshIndex )
-			{
-				MeshDescriptions[ GeomMeshIndex ] = UsdStageActorImpl::LoadMeshDescription( pxr::UsdGeomMesh( AllGeomMeshes[ GeomMeshIndex ].Get() ), TimeCode );
-			} );
-	}
-
-	TArray< TPair< int32, UStaticMesh* > > StaticMeshesToBuild;
-	StaticMeshesToBuild.Reserve( AllGeomMeshes.Num() );
-
-	// Create all static meshes
-	{
-		LoadAllProgress.EnterProgressFrame();
-
-		for ( int32 PrimIndex = 0; PrimIndex < AllGeomMeshes.Num(); ++PrimIndex )
-		{
-			bool bIsNew = true;
-			UStaticMesh* StaticMesh = CreateStaticMesh( pxr::UsdGeomMesh( AllGeomMeshes[ PrimIndex ].Get() ), MoveTemp( MeshDescriptions[ PrimIndex ] ), bIsNew );
-
-			PrimPathsToStaticMeshes.Add( UsdToUnreal::ConvertPath( AllGeomMeshes[ PrimIndex ].Get().GetPrimPath() ), StaticMesh );
-
-			if ( bIsNew )
-			{
-				StaticMeshesToBuild.Emplace( PrimIndex, StaticMesh );
-			}
-		}
-	}
-
-	// Commit all mesh descriptions
-	{
-		LoadAllProgress.EnterProgressFrame();
-
-		ParallelFor( StaticMeshesToBuild.Num(), 
-			[ &StaticMeshesToBuild ]( int32 MeshIndex )
-			{
-				UStaticMesh::FCommitMeshDescriptionParams Params;
-				Params.bMarkPackageDirty = false;
-				Params.bUseHashAsGuid = true;
-
-				TPair< int32, UStaticMesh* >& MeshToBuild = StaticMeshesToBuild[ MeshIndex ];
-
-				if ( UStaticMesh* StaticMesh = MeshToBuild.Value )
-				{
-					StaticMesh->CommitMeshDescription( 0, Params );
-				}
-			} );
-	}
-
-	// PreBuild all static meshes
-	{
-		LoadAllProgress.EnterProgressFrame();
-
-		for ( TPair< int32, UStaticMesh* >& MeshToBuild : StaticMeshesToBuild )
-		{
-			if ( UStaticMesh* StaticMesh = MeshToBuild.Value )
-			{
-				UsdStageActorImpl::PreBuildStaticMesh( pxr::UsdGeomMesh( AllGeomMeshes[ MeshToBuild.Key ].Get() ), *StaticMesh, MaterialsCache, Time );
-			}
-		}
-	}
-
-	// Build all static meshes
-	{
-		LoadAllProgress.EnterProgressFrame();
-
-		IMeshBuilderModule& MeshBuilderModule = FModuleManager::LoadModuleChecked< IMeshBuilderModule >( TEXT("MeshBuilder") );
-
-		ParallelFor( StaticMeshesToBuild.Num(), 
-			[ &StaticMeshesToBuild, &MeshBuilderModule ]( int32 MeshIndex )
-			{
-				if ( UStaticMesh* StaticMesh = StaticMeshesToBuild[ MeshIndex ].Value )
-				{
-					if ( !UsdStageActorImpl::BuildStaticMesh( *StaticMesh, MeshBuilderModule ) )
-					{
-						// Build failed, discard the mesh
-						StaticMeshesToBuild[ MeshIndex ].Value = nullptr;
-					}
-				}
-			} );
-	}
-
-	// PostBuild all static meshes
-	{
-		LoadAllProgress.EnterProgressFrame();
-
-		for ( TPair< int32, UStaticMesh* >& MeshToBuild : StaticMeshesToBuild )
-		{
-			if ( UStaticMesh* StaticMesh = MeshToBuild.Value )
-			{
-				UsdStageActorImpl::PostBuildStaticMesh( *StaticMesh );
-			}
-		}
-	}
-
-	return PrimPathsToStaticMeshes;
-}
-
 #undef LOCTEXT_NAMESPACE
diff --git a/Engine/Plugins/Importers/USDImporter/Source/USDStageEditor/Private/SUSDStageTreeView.cpp b/Engine/Plugins/Importers/USDImporter/Source/USDStageEditor/Private/SUSDStageTreeView.cpp
--- a/Engine/Plugins/Importers/USDImporter/Source/USDStageEditor/Private/SUSDStageTreeView.cpp
+++ b/Engine/Plugins/Importers/USDImporter/Source/USDStageEditor/Private/SUSDStageTreeView.cpp
@@ -151,11 +151,6 @@ public:
 
 		if ( !UsdPrim.Get() )
 		{
-			if ( ParentItem )
-			{
-				ParentItem->RefreshData( true );
-			}
-
 			RowData.Get() = FUsdStageTreeRowData();
 			return;
 		}
@@ -464,8 +459,8 @@ void SUsdStageTreeView::RefreshPrim( const FString& PrimPath, bool bResync )
 
 	TUsdStore< pxr::SdfPath > UsdPrimPath = UnrealToUsd::ConvertPath( *PrimPath );
 
-	TFunction< FUsdStageTreeItemPtr( const TUsdStore< pxr::SdfPath >&, const FUsdStageTreeItemRef& ) > FindTreeItemFromPrimPath;
-	FindTreeItemFromPrimPath = [ &FindTreeItemFromPrimPath ]( const TUsdStore< pxr::SdfPath >& UsdPrimPath, const FUsdStageTreeItemRef& ItemRef ) -> FUsdStageTreeItemPtr
+	TFunction< FUsdStageTreeItemPtr( FUsdStageTreeItemRef ) > FindTreeItemFromPrimPath;
+	FindTreeItemFromPrimPath = [ &UsdPrimPath, &FindTreeItemFromPrimPath ]( FUsdStageTreeItemRef ItemRef ) -> FUsdStageTreeItemPtr
 	{
 		if ( ItemRef->UsdPrim.Get().GetPrimPath() == UsdPrimPath.Get() )
 		{
@@ -475,7 +470,7 @@ void SUsdStageTreeView::RefreshPrim( const FString& PrimPath, bool bResync )
 		{
 			for ( FUsdStageTreeItemRef ChildItem : ItemRef->Children )
 			{
-				if ( FUsdStageTreeItemPtr ChildValue = FindTreeItemFromPrimPath( UsdPrimPath, ChildItem ) )
+				if ( FUsdStageTreeItemPtr ChildValue = FindTreeItemFromPrimPath( ChildItem ) )
 				{
 					return ChildValue;
 				}
@@ -489,35 +484,13 @@ void SUsdStageTreeView::RefreshPrim( const FString& PrimPath, bool bResync )
 
 	for ( FUsdStageTreeItemRef RootItem : this->RootItems )
 	{
-		TUsdStore< pxr::SdfPath > PrimPathToSearch = UsdPrimPath;
-
-		bool bFoundPrimToRefresh = false;
-
-		while ( !bFoundPrimToRefresh )
+		if ( FUsdStageTreeItemPtr FoundItem = FindTreeItemFromPrimPath( RootItem ) )
 		{
-			if ( FUsdStageTreeItemPtr FoundItem = FindTreeItemFromPrimPath( PrimPathToSearch, RootItem ) )
-			{
-				FoundItem->RefreshData( true );
-				bFoundPrimToRefresh = true;
-
-				break;
-			}
-			else
-			{
-				TUsdStore< pxr::SdfPath > ParentPrimPath = PrimPathToSearch.Get().GetParentPath();
+			FoundItem->RefreshData( true );
 
-				if ( ParentPrimPath.Get() == PrimPathToSearch.Get() )
-				{
-					break;
-				}
-				else
-				{
-					PrimPathToSearch = MoveTemp( ParentPrimPath );
-				}
-			}
+			break;
 		}
-
-		if ( !bFoundPrimToRefresh )
+		else
 		{
 			bNeedsFullRefresh = true; // Prim not found, refresh the whole tree
 		}
