Fixed build, missing forward declaration

#rnx

diff --git a/Engine/Plugins/Importers/USDImporter/Source/USDStageEditor/Private/SUSDStageInfo.h b/Engine/Plugins/Importers/USDImporter/Source/USDStageEditor/Private/SUSDStageInfo.h
--- a/Engine/Plugins/Importers/USDImporter/Source/USDStageEditor/Private/SUSDStageInfo.h
+++ b/Engine/Plugins/Importers/USDImporter/Source/USDStageEditor/Private/SUSDStageInfo.h
@@ -7,6 +7,7 @@
 
 class AUsdStageActor;
 enum class EUsdInitialLoadSet;
+class STextComboBox;
 
 #if USE_USD_SDK
 
