Build fix for Clang Win64

#rnx

diff --git a/Engine/Plugins/Importers/USDImporter/Source/USDStage/Private/USDStageActor.cpp b/Engine/Plugins/Importers/USDImporter/Source/USDStage/Private/USDStageActor.cpp
--- a/Engine/Plugins/Importers/USDImporter/Source/USDStage/Private/USDStageActor.cpp
+++ b/Engine/Plugins/Importers/USDImporter/Source/USDStage/Private/USDStageActor.cpp
@@ -231,7 +231,9 @@ AUsdStageActor::AUsdStageActor()
 
 AUsdStageActor::~AUsdStageActor()
 {
+#if USE_USD_SDK
 	UnrealUSDWrapper::GetUsdStageCache().Erase( UsdStageStore.Get() );
+#endif // #if USE_USD_SDK
 }
 
 #if USE_USD_SDK
