USD Stage Editor
- Fixed adding a prim collapsing the whole tree.
- Fixed removing a prim leaving an empty widget.

#rnx

diff --git a/Engine/Plugins/Importers/USDImporter/Source/USDStageEditor/Private/SUSDStageTreeView.cpp b/Engine/Plugins/Importers/USDImporter/Source/USDStageEditor/Private/SUSDStageTreeView.cpp
--- a/Engine/Plugins/Importers/USDImporter/Source/USDStageEditor/Private/SUSDStageTreeView.cpp
+++ b/Engine/Plugins/Importers/USDImporter/Source/USDStageEditor/Private/SUSDStageTreeView.cpp
@@ -151,6 +151,11 @@ public:
 
 		if ( !UsdPrim.Get() )
 		{
+			if ( ParentItem )
+			{
+				ParentItem->RefreshData( true );
+			}
+
 			RowData.Get() = FUsdStageTreeRowData();
 			return;
 		}
@@ -459,8 +464,8 @@ void SUsdStageTreeView::RefreshPrim( const FString& PrimPath, bool bResync )
 
 	TUsdStore< pxr::SdfPath > UsdPrimPath = UnrealToUsd::ConvertPath( *PrimPath );
 
-	TFunction< FUsdStageTreeItemPtr( FUsdStageTreeItemRef ) > FindTreeItemFromPrimPath;
-	FindTreeItemFromPrimPath = [ &UsdPrimPath, &FindTreeItemFromPrimPath ]( FUsdStageTreeItemRef ItemRef ) -> FUsdStageTreeItemPtr
+	TFunction< FUsdStageTreeItemPtr( const TUsdStore< pxr::SdfPath >&, const FUsdStageTreeItemRef& ) > FindTreeItemFromPrimPath;
+	FindTreeItemFromPrimPath = [ &FindTreeItemFromPrimPath ]( const TUsdStore< pxr::SdfPath >& UsdPrimPath, const FUsdStageTreeItemRef& ItemRef ) -> FUsdStageTreeItemPtr
 	{
 		if ( ItemRef->UsdPrim.Get().GetPrimPath() == UsdPrimPath.Get() )
 		{
@@ -470,7 +475,7 @@ void SUsdStageTreeView::RefreshPrim( const FString& PrimPath, bool bResync )
 		{
 			for ( FUsdStageTreeItemRef ChildItem : ItemRef->Children )
 			{
-				if ( FUsdStageTreeItemPtr ChildValue = FindTreeItemFromPrimPath( ChildItem ) )
+				if ( FUsdStageTreeItemPtr ChildValue = FindTreeItemFromPrimPath( UsdPrimPath, ChildItem ) )
 				{
 					return ChildValue;
 				}
@@ -484,13 +489,35 @@ void SUsdStageTreeView::RefreshPrim( const FString& PrimPath, bool bResync )
 
 	for ( FUsdStageTreeItemRef RootItem : this->RootItems )
 	{
-		if ( FUsdStageTreeItemPtr FoundItem = FindTreeItemFromPrimPath( RootItem ) )
+		TUsdStore< pxr::SdfPath > PrimPathToSearch = UsdPrimPath;
+
+		bool bFoundPrimToRefresh = false;
+
+		while ( !bFoundPrimToRefresh )
 		{
-			FoundItem->RefreshData( true );
+			if ( FUsdStageTreeItemPtr FoundItem = FindTreeItemFromPrimPath( PrimPathToSearch, RootItem ) )
+			{
+				FoundItem->RefreshData( true );
+				bFoundPrimToRefresh = true;
 
-			break;
+				break;
+			}
+			else
+			{
+				TUsdStore< pxr::SdfPath > ParentPrimPath = PrimPathToSearch.Get().GetParentPath();
+
+				if ( ParentPrimPath.Get() == PrimPathToSearch.Get() )
+				{
+					break;
+				}
+				else
+				{
+					PrimPathToSearch = MoveTemp( ParentPrimPath );
+				}
+			}
 		}
-		else
+
+		if ( !bFoundPrimToRefresh )
 		{
 			bNeedsFullRefresh = true; // Prim not found, refresh the whole tree
 		}
