Fixed USD initial load set not being set on stage actor when opening a stage for the first time.

#rnx

diff --git a/Engine/Plugins/Importers/USDImporter/Source/USDStageEditor/Private/SUSDStage.cpp b/Engine/Plugins/Importers/USDImporter/Source/USDStageEditor/Private/SUSDStage.cpp
--- a/Engine/Plugins/Importers/USDImporter/Source/USDStageEditor/Private/SUSDStage.cpp
+++ b/Engine/Plugins/Importers/USDImporter/Source/USDStageEditor/Private/SUSDStage.cpp
@@ -501,6 +501,11 @@ void SUsdStage::OpenStage( const TCHAR* FilePath )
 		IUsdStageModule& UsdStageModule = FModuleManager::Get().LoadModuleChecked< IUsdStageModule >( "UsdStage" );
 		UsdStageActor = &UsdStageModule.GetUsdStageActor( GWorld );
 
+		if ( UsdStageInfoWidget )
+		{
+			UsdStageActor->InitialLoadSet = UsdStageInfoWidget->GetInitialLoadSet();
+		}
+
 		SetupStageActorDelegates();
 	}
 
diff --git a/Engine/Plugins/Importers/USDImporter/Source/USDStageEditor/Private/SUSDStageInfo.h b/Engine/Plugins/Importers/USDImporter/Source/USDStageEditor/Private/SUSDStageInfo.h
--- a/Engine/Plugins/Importers/USDImporter/Source/USDStageEditor/Private/SUSDStageInfo.h
+++ b/Engine/Plugins/Importers/USDImporter/Source/USDStageEditor/Private/SUSDStageInfo.h
@@ -23,6 +23,8 @@ public:
 	void Construct( const FArguments& InArgs, AUsdStageActor* UsdStageActor );
 	void RefreshStageInfos( AUsdStageActor* UsdStageActor );
 
+	EUsdInitialLoadSet GetInitialLoadSet() const { return StageInfos.InitialLoadSet; }
+
 private:
 	void OnInitialLoadSetSelectionChanged( TSharedPtr<FString> NewValue, ESelectInfo::Type SelectInfo );
 
@@ -38,7 +40,6 @@ private:
 	FStageInfos StageInfos;
 
 	FText GetRootLayerDisplayName() const { return StageInfos.RootLayerDisplayName; }
-	FString GetInitialLoadSet() const { return *InitialLoadSetStrings[ (int32)StageInfos.InitialLoadSet ]; }
 
 	TSharedPtr< STextComboBox > InitialLoadSetWidget;
 };
