Copied cl 9961290

Fixed LevelSequence getting unloaded for UsdStageActor.

#rnx

diff --git a/Engine/Plugins/Importers/USDImporter/Source/USDStage/Private/USDStageActor.cpp b/Engine/Plugins/Importers/USDImporter/Source/USDStage/Private/USDStageActor.cpp
--- a/Engine/Plugins/Importers/USDImporter/Source/USDStage/Private/USDStageActor.cpp
+++ b/Engine/Plugins/Importers/USDImporter/Source/USDStage/Private/USDStageActor.cpp
@@ -816,8 +816,6 @@ void AUsdStageActor::SetTime(float InTime)
 
 void AUsdStageActor::PostEditChangeProperty(FPropertyChangedEvent& PropertyChangedEvent)
 {
-	Super::PostEditChangeProperty( PropertyChangedEvent );
-
 	UProperty* PropertyThatChanged = PropertyChangedEvent.MemberProperty;
 	const FName PropertyName = PropertyThatChanged ? PropertyThatChanged->GetFName() : NAME_None;
 	
@@ -834,6 +832,10 @@ void AUsdStageActor::PostEditChangeProperty(FPropertyChangedEvent& PropertyChang
 	{
 		Refresh();
 	}
+	else
+	{
+		Super::PostEditChangeProperty( PropertyChangedEvent );
+	}
 }
 
 void AUsdStageActor::Clear()
@@ -891,7 +893,7 @@ void AUsdStageActor::OpenUsdStage()
 
 void AUsdStageActor::InitLevelSequence(float FramesPerSecond)
 {
-	if (LevelSequence.IsValid())
+	if (LevelSequence)
 	{
 		return;
 	}
@@ -925,7 +927,7 @@ void AUsdStageActor::InitLevelSequence(float FramesPerSecond)
 
 void AUsdStageActor::SetupLevelSequence()
 {
-	if (!LevelSequence.IsValid())
+	if (!LevelSequence)
 	{
 		return;
 	}
diff --git a/Engine/Plugins/Importers/USDImporter/Source/USDStage/Public/USDStageActor.h b/Engine/Plugins/Importers/USDImporter/Source/USDStage/Public/USDStageActor.h
--- a/Engine/Plugins/Importers/USDImporter/Source/USDStage/Public/USDStageActor.h
+++ b/Engine/Plugins/Importers/USDImporter/Source/USDStage/Public/USDStageActor.h
@@ -70,7 +70,7 @@ private:
 	float TimeCodesPerSecond;
 
 	UPROPERTY(VisibleAnywhere, Category = "USD", Transient)
-	TWeakObjectPtr< ULevelSequence > LevelSequence;
+	ULevelSequence* LevelSequence;
 
 public:
 	DECLARE_EVENT_OneParam( AUsdStageActor, FOnActorLoaded, AUsdStageActor* );
