Copied cl 9961937

USD Stage Editor: Fixed stack overflow and a layer as itself as a sublayer.

#rnx

diff --git a/Engine/Plugins/Importers/USDImporter/Source/USDStageEditor/Private/SUSDLayersTreeView.cpp b/Engine/Plugins/Importers/USDImporter/Source/USDStageEditor/Private/SUSDLayersTreeView.cpp
--- a/Engine/Plugins/Importers/USDImporter/Source/USDStageEditor/Private/SUSDLayersTreeView.cpp
+++ b/Engine/Plugins/Importers/USDImporter/Source/USDStageEditor/Private/SUSDLayersTreeView.cpp
@@ -588,6 +588,7 @@ void SUsdLayersTreeView::OnRemoveSelectedLayers()
 					if ( ParentLayerHandle )
 					{
 						ParentLayerHandle->RemoveSubLayerPath( SubLayerIndex );
+						//SelectedLayer->ParentItem->Children.Remove( Child );
 						bLayerRemoved = true;
 					}
 					break;
