New material for USD vertex colors

diff --git a/Engine/Plugins/Importers/USDImporter/Source/USDStage/Private/USDStageActor.cpp b/Engine/Plugins/Importers/USDImporter/Source/USDStage/Private/USDStageActor.cpp
--- a/Engine/Plugins/Importers/USDImporter/Source/USDStage/Private/USDStageActor.cpp
+++ b/Engine/Plugins/Importers/USDImporter/Source/USDStage/Private/USDStageActor.cpp
@@ -142,7 +142,7 @@ void ProcessMaterials( const pxr::UsdStageRefPtr& Stage, UStaticMesh* StaticMesh
 
 		if ( Material == nullptr && bHasPrimDisplayColor )
 		{
-			FSoftObjectPath VertexColorMaterialPath( TEXT("Material'/Engine/EngineDebugMaterials/VertexColorMaterial.VertexColorMaterial'") );
+			FSoftObjectPath VertexColorMaterialPath( TEXT("Material'/USDImporter/Materials/DisplayColor.DisplayColor'") );
 			Material = Cast< UMaterialInterface >( VertexColorMaterialPath.TryLoad() );
 		}
 
@@ -191,7 +191,7 @@ void ProcessMaterials( const pxr::UsdStageRefPtr& Stage, FSkeletalMeshImportData
 
 			if ( Material == nullptr && bHasPrimDisplayColor )
 			{
-				FSoftObjectPath VertexColorMaterialPath( TEXT("Material'/Engine/EngineDebugMaterials/VertexColorMaterial.VertexColorMaterial'") );
+				FSoftObjectPath VertexColorMaterialPath( TEXT("Material'/USDImporter/Materials/DisplayColor.DisplayColor'") );
 				Material = Cast< UMaterialInterface >( VertexColorMaterialPath.TryLoad() );
 			}
 			ImportedMaterial.Material = Material;
