USD Stage - Fix bad merge introcuded with CL 10013052

diff --git a/Engine/Plugins/Importers/USDImporter/Source/USDStage/Private/USDStageActor.cpp b/Engine/Plugins/Importers/USDImporter/Source/USDStage/Private/USDStageActor.cpp
--- a/Engine/Plugins/Importers/USDImporter/Source/USDStage/Private/USDStageActor.cpp
+++ b/Engine/Plugins/Importers/USDImporter/Source/USDStage/Private/USDStageActor.cpp
@@ -10,6 +10,7 @@
 #include "USDTypesConversion.h"
 #include "UnrealUSDWrapper.h"
 
+#include "Async/ParallelFor.h"
 #include "CineCameraActor.h"
 #include "CineCameraComponent.h"
 #include "Components/PoseableMeshComponent.h"
@@ -26,6 +27,8 @@
 #include "MaterialEditingLibrary.h"
 #include "Materials/Material.h"
 #include "Materials/MaterialInterface.h"
+#include "Materials/MaterialInstanceConstant.h"
+#include "MeshDescription.h"
 #include "MeshDescriptionOperations.h"
 #include "Misc/Paths.h"
 #include "Misc/ScopedSlowTask.h"
@@ -36,7 +39,6 @@
 #include "Channels/MovieSceneChannelProxy.h"
 #include "LevelSequence.h"
 #include "LevelSequenceActor.h"
-#include "MeshDescription.h"
 #include "Misc/FrameNumber.h"
 #include "MovieScene.h"
 #include "Sections/MovieSceneFloatSection.h"
@@ -66,104 +68,69 @@
 
 #define LOCTEXT_NAMESPACE "USDStageActor"
 
+DEFINE_LOG_CATEGORY( LogUsdStage );
+
 static const EObjectFlags DefaultObjFlag = EObjectFlags::RF_Transactional | EObjectFlags::RF_Transient;
 
 AUsdStageActor::FOnActorLoaded AUsdStageActor::OnActorLoaded{};
 
 #if USE_USD_SDK
-FMeshDescription LoadMeshDescription( const pxr::UsdGeomMesh& UsdMesh, const pxr::UsdTimeCode TimeCode )
+
+namespace UsdStageActorImpl
 {
-	if ( !UsdMesh )
+	FMeshDescription LoadMeshDescription( const pxr::UsdGeomMesh& UsdMesh, const pxr::UsdTimeCode TimeCode )
 	{
-		return {};
-	}
-
-	FMeshDescription MeshDescription;
-	FStaticMeshAttributes StaticMeshAttributes( MeshDescription );
-	StaticMeshAttributes.Register();
+		if ( !UsdMesh )
+		{
+			return {};
+		}
 
-	UsdToUnreal::ConvertGeomMesh( UsdMesh, MeshDescription, TimeCode );
+		FMeshDescription MeshDescription;
+		FStaticMeshAttributes StaticMeshAttributes( MeshDescription );
+		StaticMeshAttributes.Register();
 
-	return MeshDescription;
-}
+		UsdToUnreal::ConvertGeomMesh( UsdMesh, MeshDescription, TimeCode );
 
-void ProcessMaterials( const pxr::UsdStageRefPtr& Stage, UStaticMesh* StaticMesh, const FMeshDescription& MeshDescription, TMap< FString, UMaterial* >& MaterialsCache, bool bHasPrimDisplayColor )
-{
-	FStaticMeshConstAttributes StaticMeshAttributes( MeshDescription );
+		return MeshDescription;
+	}
 
-	for ( const FPolygonGroupID PolygonGroupID : MeshDescription.PolygonGroups().GetElementIDs() )
+	void ProcessMaterials( const pxr::UsdPrim& UsdPrim, UStaticMesh& StaticMesh, TMap< FString, UMaterial* >& MaterialsCache, bool bHasPrimDisplayColor, float Time )
 	{
-		const FName& ImportedMaterialSlotName = StaticMeshAttributes.GetPolygonGroupMaterialSlotNames()[ PolygonGroupID ];
-		const FName MaterialSlotName = ImportedMaterialSlotName;
-
-		int32 MaterialIndex = INDEX_NONE;
-		int32 MeshMaterialIndex = 0;
+		const FMeshDescription* MeshDescription = StaticMesh.GetMeshDescription( 0 );
 
-		for ( FStaticMaterial& StaticMaterial : StaticMesh->StaticMaterials )
+		if ( !MeshDescription )
 		{
-			if ( StaticMaterial.MaterialSlotName.IsEqual( ImportedMaterialSlotName ) )
-			{
-				MaterialIndex = MeshMaterialIndex;
-				break;
-			}
-
-			++MeshMaterialIndex;
+			return;
 		}
 
-		if ( MaterialIndex == INDEX_NONE )
-		{
-			MaterialIndex = PolygonGroupID.GetValue();
-		}
+		FStaticMeshConstAttributes StaticMeshAttributes( *MeshDescription );
 
-		UMaterialInterface* Material = nullptr;
-		pxr::UsdPrim MaterialPrim = Stage->GetPrimAtPath( UnrealToUsd::ConvertPath( *ImportedMaterialSlotName.ToString() ).Get() );
-
-		if ( MaterialPrim )
+		for ( const FPolygonGroupID PolygonGroupID : MeshDescription->PolygonGroups().GetElementIDs() )
 		{
-			UMaterial*& CachedMaterial = MaterialsCache.FindOrAdd( UsdToUnreal::ConvertPath( MaterialPrim.GetPrimPath() ) );
+			const FName& ImportedMaterialSlotName = StaticMeshAttributes.GetPolygonGroupMaterialSlotNames()[ PolygonGroupID ];
+			const FName MaterialSlotName = ImportedMaterialSlotName;
 
-			if ( !CachedMaterial )
-			{
-				CachedMaterial = NewObject< UMaterial >();
+			int32 MaterialIndex = INDEX_NONE;
+			int32 MeshMaterialIndex = 0;
 
-				if ( UsdToUnreal::ConvertMaterial( pxr::UsdShadeMaterial( MaterialPrim ), *CachedMaterial ) )
-				{
-					//UMaterialEditingLibrary::RecompileMaterial( CachedMaterial ); // Too slow
-					CachedMaterial->PostEditChange();
-				}
-				else
+			for ( FStaticMaterial& StaticMaterial : StaticMesh.StaticMaterials )
+			{
+				if ( StaticMaterial.MaterialSlotName.IsEqual( ImportedMaterialSlotName ) )
 				{
-					CachedMaterial = nullptr;
+					MaterialIndex = MeshMaterialIndex;
+					break;
 				}
-			}
-
-			Material = CachedMaterial;
-		}
 
-		if ( Material == nullptr && bHasPrimDisplayColor )
-		{
-			FSoftObjectPath VertexColorMaterialPath( TEXT("Material'/Engine/EngineDebugMaterials/VertexColorMaterial.VertexColorMaterial'") );
-			Material = Cast< UMaterialInterface >( VertexColorMaterialPath.TryLoad() );
-		}
-
-		FStaticMaterial StaticMaterial( Material, MaterialSlotName, ImportedMaterialSlotName );
-		StaticMesh->StaticMaterials.Add( StaticMaterial );
-
-		FMeshSectionInfo MeshSectionInfo;
-		MeshSectionInfo.MaterialIndex = MaterialIndex;
+				++MeshMaterialIndex;
+			}
 
-		StaticMesh->GetSectionInfoMap().Set( 0, PolygonGroupID.GetValue(), MeshSectionInfo );
-	}
-}
+			if ( MaterialIndex == INDEX_NONE )
+			{
+				MaterialIndex = PolygonGroupID.GetValue();
+			}
 
-void ProcessMaterials( const pxr::UsdStageRefPtr& Stage, FSkeletalMeshImportData& SkelMeshImportData, TMap< FString, UMaterial* >& MaterialsCache, bool bHasPrimDisplayColor )
-{
-	for (SkeletalMeshImportData::FMaterial& ImportedMaterial : SkelMeshImportData.Materials)
-	{
-		if (!ImportedMaterial.Material.IsValid())
-		{
 			UMaterialInterface* Material = nullptr;
-			TUsdStore< pxr::UsdPrim > MaterialPrim = Stage->GetPrimAtPath( UnrealToUsd::ConvertPath( *ImportedMaterial.MaterialImportName ).Get() );
+			TUsdStore< pxr::UsdPrim > MaterialPrim = UsdPrim.GetStage()->GetPrimAtPath( UnrealToUsd::ConvertPath( *ImportedMaterialSlotName.ToString() ).Get() );
 
 			if ( MaterialPrim.Get() )
 			{
@@ -175,10 +142,7 @@ void ProcessMaterials( const pxr::UsdStageRefPtr& Stage, FSkeletalMeshImportData
 
 					if ( UsdToUnreal::ConvertMaterial( pxr::UsdShadeMaterial( MaterialPrim.Get() ), *CachedMaterial ) )
 					{
-						CachedMaterial->bUsedWithSkeletalMesh = true;
-						bool bNeedsRecompile = false;
-						CachedMaterial->GetMaterial()->SetMaterialUsage(bNeedsRecompile, MATUSAGE_SkeletalMesh);
-
+						//UMaterialEditingLibrary::RecompileMaterial( CachedMaterial ); // Too slow
 						CachedMaterial->PostEditChange();
 					}
 					else
@@ -186,15 +150,78 @@ void ProcessMaterials( const pxr::UsdStageRefPtr& Stage, FSkeletalMeshImportData
 						CachedMaterial = nullptr;
 					}
 				}
+
 				Material = CachedMaterial;
 			}
 
 			if ( Material == nullptr && bHasPrimDisplayColor )
 			{
-				FSoftObjectPath VertexColorMaterialPath( TEXT("Material'/Engine/EngineDebugMaterials/VertexColorMaterial.VertexColorMaterial'") );
-				Material = Cast< UMaterialInterface >( VertexColorMaterialPath.TryLoad() );
+				UMaterialInstanceConstant* MaterialInstance = NewObject< UMaterialInstanceConstant >();
+				if ( UsdToUnreal::ConvertDisplayColor( pxr::UsdGeomMesh( UsdPrim ), *MaterialInstance, pxr::UsdTimeCode( Time ) ) )
+				{
+					Material = MaterialInstance;
+				}
+			}
+
+			FStaticMaterial StaticMaterial( Material, MaterialSlotName, ImportedMaterialSlotName );
+				StaticMesh.StaticMaterials.Add( StaticMaterial );
+
+			FMeshSectionInfo MeshSectionInfo;
+			MeshSectionInfo.MaterialIndex = MaterialIndex;
+
+			StaticMesh.GetSectionInfoMap().Set( 0, PolygonGroupID.GetValue(), MeshSectionInfo );
+		}
+	}
+
+	void ProcessMaterials( const pxr::UsdPrim& UsdPrim, FSkeletalMeshImportData& SkelMeshImportData, TMap< FString, UMaterial* >& MaterialsCache, bool bHasPrimDisplayColor, float Time )
+	{
+		for (SkeletalMeshImportData::FMaterial& ImportedMaterial : SkelMeshImportData.Materials)
+		{
+			if (!ImportedMaterial.Material.IsValid())
+			{
+				UMaterialInterface* Material = nullptr;
+				TUsdStore< pxr::UsdPrim > MaterialPrim = UsdPrim.GetStage()->GetPrimAtPath( UnrealToUsd::ConvertPath( *ImportedMaterial.MaterialImportName ).Get() );
+
+				if ( MaterialPrim.Get() )
+				{
+					UMaterial*& CachedMaterial = MaterialsCache.FindOrAdd( UsdToUnreal::ConvertPath( MaterialPrim.Get().GetPrimPath() ) );
+
+					if ( !CachedMaterial )
+					{
+						CachedMaterial = NewObject< UMaterial >();
+
+						if ( UsdToUnreal::ConvertMaterial( pxr::UsdShadeMaterial( MaterialPrim.Get() ), *CachedMaterial ) )
+						{
+							CachedMaterial->bUsedWithSkeletalMesh = true;
+							bool bNeedsRecompile = false;
+							CachedMaterial->GetMaterial()->SetMaterialUsage(bNeedsRecompile, MATUSAGE_SkeletalMesh);
+
+							CachedMaterial->PostEditChange();
+						}
+						else
+						{
+							CachedMaterial = nullptr;
+						}
+					}
+					Material = CachedMaterial;
+				}
+
+				if ( Material == nullptr && bHasPrimDisplayColor )
+				{
+					UMaterialInstanceConstant* MaterialInstance = NewObject< UMaterialInstanceConstant >();
+					if ( UsdToUnreal::ConvertDisplayColor( pxr::UsdGeomMesh( UsdPrim ), *MaterialInstance, pxr::UsdTimeCode( Time ) ) )
+					{
+						if ( MaterialInstance && MaterialInstance->HasOverridenBaseProperties() )
+						{
+							MaterialInstance->ForceRecompileForRendering();
+						}
+						
+						Material = MaterialInstance;
+					}
+				}
+
+				ImportedMaterial.Material = Material;
 			}
-			ImportedMaterial.Material = Material;
 		}
 	}
 }
@@ -216,7 +243,10 @@ AUsdStageActor::AUsdStageActor()
 		[ this ]( const FString& PrimPath, bool bResync )
 			{
 				TUsdStore< pxr::SdfPath > UsdPrimPath = UnrealToUsd::ConvertPath( *PrimPath );
-				this->UpdatePrim( UsdPrimPath.Get(), bResync );
+
+				TMap< FString, UStaticMesh* > PrimPathsToStaticMeshes = this->LoadStaticMeshes( this->GetUsdStage()->GetPrimAtPath( UsdPrimPath.Get() ) );
+
+				this->UpdatePrim( UsdPrimPath.Get(), bResync, PrimPathsToStaticMeshes );
 
 				this->OnPrimChanged.Broadcast( PrimPath, bResync );
 			}
@@ -238,82 +268,11 @@ AUsdStageActor::~AUsdStageActor()
 
 #if USE_USD_SDK
 
-bool AUsdStageActor::LoadStaticMesh( const pxr::UsdGeomMesh& UsdMesh, UStaticMeshComponent& MeshComponent )
+bool AUsdStageActor::SetStaticMesh( const pxr::UsdGeomMesh& UsdMesh, UStaticMeshComponent& MeshComponent, const TMap< FString, UStaticMesh* >& PrimPathsToStaticMeshes )
 {
-	UStaticMesh* StaticMesh = nullptr;
-
-	pxr::UsdGeomXformable Xformable( UsdMesh.GetPrim() );
-
 	FString MeshPrimPath = UsdToUnreal::ConvertPath( UsdMesh.GetPrim().GetPrimPath() );
 
-	FMeshDescription MeshDescription = LoadMeshDescription( UsdMesh, pxr::UsdTimeCode( Time ) );
-
-	FSHAHash MeshHash = FMeshDescriptionOperations::ComputeSHAHash( MeshDescription );
-
-	StaticMesh = MeshCache.FindRef( MeshHash.ToString() );
-
-	if ( !StaticMesh && !MeshDescription.IsEmpty() )
-	{
-		StaticMesh = NewObject< UStaticMesh >( GetTransientPackage(), NAME_None, DefaultObjFlag | EObjectFlags::RF_Public );
-
-		FStaticMeshSourceModel& SourceModel = StaticMesh->AddSourceModel();
-		SourceModel.BuildSettings.bGenerateLightmapUVs = false;
-		SourceModel.BuildSettings.bRecomputeNormals = false;
-		SourceModel.BuildSettings.bRecomputeTangents = false;
-		SourceModel.BuildSettings.bBuildAdjacencyBuffer = false;
-		SourceModel.BuildSettings.bBuildReversedIndexBuffer = false;
-
-		FMeshDescription* StaticMeshDescription = StaticMesh->CreateMeshDescription(0);
-		check( StaticMeshDescription );
-		*StaticMeshDescription = MoveTemp( MeshDescription );
-		StaticMesh->CommitMeshDescription(0);
-
-		const bool bHasPrimDisplayColor = UsdMesh.GetDisplayColorPrimvar().IsDefined();
-		ProcessMaterials( GetUsdStage(), StaticMesh, *StaticMeshDescription, MaterialsCache, bHasPrimDisplayColor );
-
-		// Create render data
-		if ( !StaticMesh->RenderData )
-		{
-			StaticMesh->RenderData.Reset(new(FMemory::Malloc(sizeof(FStaticMeshRenderData)))FStaticMeshRenderData());
-		}
-
-		ITargetPlatformManagerModule& TargetPlatformManager = GetTargetPlatformManagerRef();
-		ITargetPlatform* RunningPlatform = TargetPlatformManager.GetRunningTargetPlatform();
-		check(RunningPlatform);
-		const FStaticMeshLODSettings& LODSettings = RunningPlatform->GetStaticMeshLODSettings();
-
-		const FStaticMeshLODGroup& LODGroup = LODSettings.GetLODGroup(StaticMesh->LODGroup);
-
-		IMeshBuilderModule& MeshBuilderModule = FModuleManager::LoadModuleChecked< IMeshBuilderModule >( TEXT("MeshBuilder") );
-		if ( !MeshBuilderModule.BuildMesh( *StaticMesh->RenderData, StaticMesh, LODGroup ) )
-		{
-			UE_LOG(LogStaticMesh, Error, TEXT("Failed to build static mesh. See previous line(s) for details."));
-			return false;
-		}
-
-		for ( FStaticMeshLODResources& LODResources : StaticMesh->RenderData->LODResources )
-		{
-			LODResources.bHasColorVertexData = true;
-		}
-
-		// Bounds
-		StaticMesh->CalculateExtendedBounds();
-
-		// Recreate resources
-		StaticMesh->InitResources();
-
-		// Other setup
-		//StaticMesh->CreateBodySetup();
-
-		// Collision
-		//StaticMesh->BodySetup->CollisionTraceFlag = ECollisionTraceFlag::CTF_UseComplexAsSimple;
-
-		MeshCache.Add( MeshHash.ToString() ) = StaticMesh;
-	}
-	/*else
-	{
-		FPlatformMisc::LowLevelOutputDebugStringf( TEXT("Mesh found in cache %s\n"), *StaticMesh->GetName() );
-	}*/
+	UStaticMesh* StaticMesh = PrimPathsToStaticMeshes.FindRef( MeshPrimPath );
 
 	if ( StaticMesh != MeshComponent.GetStaticMesh() )
 	{
@@ -327,7 +286,7 @@ bool AUsdStageActor::LoadStaticMesh( const pxr::UsdGeomMesh& UsdMesh, UStaticMes
 		MeshComponent.RegisterComponent();
 	}
 
-	return true;
+	return ( StaticMesh != nullptr );
 }
 
 FUsdPrimTwin* AUsdStageActor::SpawnPrim( const pxr::SdfPath& UsdPrimPath )
@@ -360,41 +319,11 @@ FUsdPrimTwin* AUsdStageActor::SpawnPrim( const pxr::SdfPath& UsdPrimPath )
 				this->OnUsdPrimTwinDestroyed( UsdPrimTwin );
 			} );
 
-		if ( !UsdPrimTwin->AnimationHandle.IsValid() )
-		{
-			bool bHasXformbaleTimeSamples = false;
-			{
-				pxr::UsdGeomXformable Xformable( Prim );
-
-				if ( Xformable )
-				{
-					std::vector< double > TimeSamples;
-					Xformable.GetTimeSamples( &TimeSamples );
-
-					bHasXformbaleTimeSamples = TimeSamples.size() > 0;
-				}
-			}
-
-			bool bHasAttributesTimeSamples = false;
-			{
-				std::vector< pxr::UsdAttribute > Attributes = Prim.GetAttributes();
-
-				for ( pxr::UsdAttribute& Attribute : Attributes )
-				{
-					bHasAttributesTimeSamples = Attribute.ValueMightBeTimeVarying();
-					if ( bHasAttributesTimeSamples )
-					{
-						break;
-					}
-				}
-			}
-
-			if ( ( bHasXformbaleTimeSamples || bHasAttributesTimeSamples ) )
+		if ( !UsdPrimTwin->AnimationHandle.IsValid() && UsdUtils::IsAnimated( Prim ) )
 			{
 				PrimsToAnimate.Add( PrimPath );
 			}
 		}
-	}
 
 	UClass* ComponentClass = UsdUtils::GetComponentTypeForPrim( Prim );
 
@@ -405,7 +334,7 @@ FUsdPrimTwin* AUsdStageActor::SpawnPrim( const pxr::SdfPath& UsdPrimPath )
 
 	FScopedUnrealAllocs UnrealAllocs;
 
-	const bool bNeedsActor = ( ParentUsdPrimTwin->SceneComponent == nullptr || Prim.IsModel() || UsdUtils::HasCompositionArcs( Prim ) || Prim.IsGroup() || Prim.IsA< pxr::UsdGeomScope >() || Prim.IsA< pxr::UsdGeomCamera >() );
+	const bool bNeedsActor = ( ParentUsdPrimTwin->SceneComponent == nullptr || Prim.IsModel() || UsdUtils::HasCompositionArcs( Prim ) || Prim.IsGroup() || Prim.IsA< pxr::UsdGeomCamera >() );
 
 	if ( bNeedsActor && !UsdPrimTwin->SpawnedActor.IsValid() )
 	{
@@ -504,7 +433,7 @@ FUsdPrimTwin* AUsdStageActor::SpawnPrim( const pxr::SdfPath& UsdPrimPath )
 	return UsdPrimTwin;
 }
 
-FUsdPrimTwin* AUsdStageActor::LoadPrim( const pxr::SdfPath& Path )
+FUsdPrimTwin* AUsdStageActor::LoadPrim( const pxr::SdfPath& Path, const TMap< FString, UStaticMesh* >& PrimPathsToStaticMeshes )
 {
 	TRACE_CPUPROFILER_EVENT_SCOPE( AUsdStageActor::LoadPrim );
 
@@ -539,7 +468,7 @@ FUsdPrimTwin* AUsdStageActor::LoadPrim( const pxr::SdfPath& Path )
 
 	if ( UStaticMeshComponent* MeshComponent = Cast< UStaticMeshComponent >( UsdPrimTwin->SceneComponent.Get() ) )
 	{
-		LoadStaticMesh( pxr::UsdGeomMesh( Prim ), *MeshComponent );
+		SetStaticMesh( pxr::UsdGeomMesh( Prim ), *MeshComponent, PrimPathsToStaticMeshes );
 	}
 	else if ( UCineCameraComponent* CameraComponent = Cast< UCineCameraComponent >( UsdPrimTwin->SceneComponent.Get() ) )
 	{
@@ -562,7 +491,16 @@ FUsdPrimTwin* AUsdStageActor::LoadPrim( const pxr::SdfPath& Path )
 		}
 
 		ComponentToAttach->AttachToComponent( AttachToComponent, FAttachmentTransformRules::KeepRelativeTransform );
-		UsdPrimTwin->SceneComponent->SetMobility( EComponentMobility::Movable );
+		
+		if ( PrimsToAnimate.Contains( UsdToUnreal::ConvertPath( Path.GetPrimPath() ) ) )
+		{
+			ComponentToAttach->SetMobility( EComponentMobility::Movable );
+		}
+		else
+		{
+			ComponentToAttach->SetMobility( EComponentMobility::Static );
+		}
+
 		UsdPrimTwin->SceneComponent->GetOwner()->AddInstanceComponent( UsdPrimTwin->SceneComponent.Get() );
 	}
 
@@ -604,7 +542,7 @@ FUsdPrimTwin* AUsdStageActor::LoadPrim( const pxr::SdfPath& Path )
 	return UsdPrimTwin;
 }
 
-FUsdPrimTwin* AUsdStageActor::ExpandPrim( const pxr::UsdPrim& Prim )
+FUsdPrimTwin* AUsdStageActor::ExpandPrim( const pxr::UsdPrim& Prim, const TMap< FString, UStaticMesh* >& PrimPathsToStaticMeshes )
 {
 	if ( !Prim )
 	{
@@ -613,7 +551,7 @@ FUsdPrimTwin* AUsdStageActor::ExpandPrim( const pxr::UsdPrim& Prim )
 
 	TRACE_CPUPROFILER_EVENT_SCOPE( AUsdStageActor::ExpandPrim );
 
-	FUsdPrimTwin* UsdPrimTwin = LoadPrim( Prim.GetPrimPath() );
+	FUsdPrimTwin* UsdPrimTwin = LoadPrim( Prim.GetPrimPath(), PrimPathsToStaticMeshes );
 
 	if ( Prim.IsA<pxr::UsdSkelRoot>() )
 	{
@@ -628,7 +566,7 @@ FUsdPrimTwin* AUsdStageActor::ExpandPrim( const pxr::UsdPrim& Prim )
 		pxr::UsdPrimSiblingRange PrimChildren = Prim.GetFilteredChildren( pxr::UsdTraverseInstanceProxies() );
 		for ( TUsdStore< pxr::UsdPrim > ChildStore : PrimChildren )
 		{
-			ExpandPrim( ChildStore.Get() );
+			ExpandPrim( ChildStore.Get(), PrimPathsToStaticMeshes );
 		}
 	}
 
@@ -660,6 +598,8 @@ bool AUsdStageActor::ProcessSkeletonRoot( const pxr::UsdPrim& Prim, USkinnedMesh
 			return false;
 		}
 
+		const pxr::TfToken StageUpAxis = UsdUtils::GetUsdStageAxis( GetUsdStage() );
+
 		pxr::UsdGeomXformable Xformable(Prim);
 
 		std::vector<double> TimeSamples;
@@ -709,7 +649,7 @@ bool AUsdStageActor::ProcessSkeletonRoot( const pxr::UsdPrim& Prim, USkinnedMesh
 
 			if (TimeSamples.size() > 0)
 			{
-				auto UpdateBoneTransforms = [&, SkelQuery]()
+				auto UpdateBoneTransforms = [&, SkelQuery, StageUpAxis]()
 				{
 					FScopedUnrealAllocs UnrealAllocs;
 
@@ -727,7 +667,7 @@ bool AUsdStageActor::ProcessSkeletonRoot( const pxr::UsdPrim& Prim, USkinnedMesh
 						for (uint32 Index = 0; Index < UsdBoneTransforms.size(); ++Index)
 						{
 							const pxr::GfMatrix4d& UsdMatrix = UsdBoneTransforms[Index];
-							FTransform BoneTransform = UsdToUnreal::ConvertMatrix(GetUsdStage(), UsdMatrix);
+							FTransform BoneTransform = UsdToUnreal::ConvertMatrix( StageUpAxis, UsdMatrix );
 							BoneTransforms.Add(BoneTransform);
 						}
 					}
@@ -743,7 +683,7 @@ bool AUsdStageActor::ProcessSkeletonRoot( const pxr::UsdPrim& Prim, USkinnedMesh
 				PrimDelegates.Add(UsdToUnreal::ConvertPath(PrimPath), Handle);
 			}
 
-			ProcessMaterials(GetUsdStage(), SkelMeshImportData, MaterialsCache, bHasPrimDisplayColor);
+			UsdStageActorImpl::ProcessMaterials(Prim, SkelMeshImportData, MaterialsCache, bHasPrimDisplayColor, Time);
 
 			break;
 		}
@@ -764,7 +704,7 @@ bool AUsdStageActor::ProcessSkeletonRoot( const pxr::UsdPrim& Prim, USkinnedMesh
 	return SkeletalMesh != nullptr;
 }
 
-void AUsdStageActor::UpdatePrim( const pxr::SdfPath& InUsdPrimPath, bool bResync )
+void AUsdStageActor::UpdatePrim( const pxr::SdfPath& InUsdPrimPath, bool bResync, const TMap< FString, UStaticMesh* >& PrimPathsToStaticMeshes )
 {
 	TRACE_CPUPROFILER_EVENT_SCOPE( AUsdStageActor::UpdatePrim );
 
@@ -791,7 +731,7 @@ void AUsdStageActor::UpdatePrim( const pxr::SdfPath& InUsdPrimPath, bool bResync
 		}
 
 		TUsdStore< pxr::UsdPrim > PrimToExpand = GetUsdStage()->GetPrimAtPath( UsdPrimPath );
-		FUsdPrimTwin* UsdPrimTwin = ExpandPrim( PrimToExpand.Get() );
+		FUsdPrimTwin* UsdPrimTwin = ExpandPrim( PrimToExpand.Get(), PrimPathsToStaticMeshes );
 
 		GEditor->BroadcastLevelActorListChanged();
 		GEditor->RedrawLevelEditingViewports();
@@ -975,6 +915,8 @@ void AUsdStageActor::SetupLevelSequence()
 void AUsdStageActor::LoadUsdStage()
 {
 #if USE_USD_SDK
+	double StartTime = FPlatformTime::Cycles64();
+
 	FScopedSlowTask SlowTask( 1.f, LOCTEXT( "LoadingUDStage", "Loading USD Stage") );
 	SlowTask.MakeDialog();
 	SlowTask.EnterProgressFrame();
@@ -998,7 +940,9 @@ void AUsdStageActor::LoadUsdStage()
 
 	SetupLevelSequence();
 
-	UpdatePrim( UsdStage->GetPseudoRoot().GetPrimPath(), true );
+	TMap< FString, UStaticMesh* > PrimPathsToStaticMeshes = LoadStaticMeshes( UsdStage->GetPseudoRoot() );
+
+	UpdatePrim( UsdStage->GetPseudoRoot().GetPrimPath(), true, PrimPathsToStaticMeshes );
 
 	SetTime( StartTimeCode );
 
@@ -1009,7 +953,10 @@ void AUsdStageActor::LoadUsdStage()
 	{
 		for ( const FString& PrimToAnimate : PrimsToAnimate )
 		{
-			this->LoadPrim( UnrealToUsd::ConvertPath( *PrimToAnimate ).Get() );
+			TUsdStore< pxr::SdfPath > PrimPath = UnrealToUsd::ConvertPath( *PrimToAnimate );
+
+			TMap< FString, UStaticMesh* > PrimPathsToStaticMeshes = LoadStaticMeshes( this->GetUsdStage()->GetPrimAtPath( PrimPath.Get() ) );
+			this->LoadPrim( PrimPath.Get(), PrimPathsToStaticMeshes );
 		}
 
 		GEditor->BroadcastLevelActorListChanged();
@@ -1017,6 +964,14 @@ void AUsdStageActor::LoadUsdStage()
 	};
 
 	OnTimeChanged.AddLambda( AnimatePrims );
+
+	// Log time spent to load the stage
+	double ElapsedSeconds = FPlatformTime::ToSeconds64(FPlatformTime::Cycles64() - StartTime);
+
+	int ElapsedMin = int(ElapsedSeconds / 60.0);
+	ElapsedSeconds -= 60.0 * (double)ElapsedMin;
+
+	UE_LOG( LogUsdStage, Log, TEXT("%s %s in [%d min %.3f s]"), TEXT("Stage loaded"), *FPaths::GetBaseFilename( RootLayer.FilePath ), ElapsedMin, ElapsedSeconds );
 #endif // #if USE_USD_SDK
 }
 
@@ -1034,7 +989,8 @@ void AUsdStageActor::PostRegisterAllComponents()
 
 	if ( UsdStage )
 	{
-		UpdatePrim( UsdStage->GetPseudoRoot().GetPrimPath(), true );
+		TMap< FString, UStaticMesh* > PrimPathsToStaticMeshes = LoadStaticMeshes( UsdStage->GetPseudoRoot() );
+		UpdatePrim( UsdStage->GetPseudoRoot().GetPrimPath(), true, PrimPathsToStaticMeshes );
 	}
 #endif // #if USE_USD_SDK
 }
@@ -1099,4 +1055,198 @@ void AUsdStageActor::OnPrimObjectPropertyChanged( UObject* ObjectBeingModified,
 #endif // #if USE_USD_SDK
 }
 
-#undef LOCTEXT_NAMESPACE
+namespace UsdStageActorImpl
+{
+	void PreBuildStaticMesh( const pxr::UsdGeomMesh& UsdMesh, UStaticMesh& StaticMesh, TMap< FString, UMaterial* >& MaterialsCache, float Time )
+	{
+		const bool bHasPrimDisplayColor = UsdMesh.GetDisplayColorPrimvar().IsDefined();
+		ProcessMaterials( UsdMesh.GetPrim(), StaticMesh, MaterialsCache, bHasPrimDisplayColor, Time );
+
+		StaticMesh.RenderData = MakeUnique< FStaticMeshRenderData >();
+	}
+
+	bool BuildStaticMesh( UStaticMesh& StaticMesh, IMeshBuilderModule& MeshBuilderModule )
+	{
+		ITargetPlatformManagerModule& TargetPlatformManager = GetTargetPlatformManagerRef();
+		ITargetPlatform* RunningPlatform = TargetPlatformManager.GetRunningTargetPlatform();
+		check(RunningPlatform);
+
+		const FStaticMeshLODSettings& LODSettings = RunningPlatform->GetStaticMeshLODSettings();
+
+		const FStaticMeshLODGroup& LODGroup = LODSettings.GetLODGroup(StaticMesh.LODGroup);
+
+		if ( !MeshBuilderModule.BuildMesh( *StaticMesh.RenderData, &StaticMesh, LODGroup ) )
+		{
+			return false;
+		}
+
+		for ( FStaticMeshLODResources& LODResources : StaticMesh.RenderData->LODResources )
+		{
+			LODResources.bHasColorVertexData = true;
+		}
+
+		return true;
+	}
+
+	void PostBuildStaticMesh( UStaticMesh& StaticMesh )
+	{
+		StaticMesh.InitResources();
+
+		if ( const FMeshDescription* MeshDescription = StaticMesh.GetMeshDescription( 0 ) )
+		{
+			StaticMesh.RenderData->Bounds = MeshDescription->GetBounds();
+		}
+	
+		StaticMesh.CalculateExtendedBounds();
+	}
+}
+
+UStaticMesh* AUsdStageActor::CreateStaticMesh( const pxr::UsdGeomMesh& UsdMesh, FMeshDescription&& MeshDescription, bool& bOutIsNew )
+{
+	UStaticMesh* StaticMesh = nullptr;
+
+	FSHAHash MeshHash = FMeshDescriptionOperations::ComputeSHAHash( MeshDescription );
+
+	StaticMesh = MeshCache.FindRef( MeshHash.ToString() );
+
+	if ( !StaticMesh && !MeshDescription.IsEmpty() )
+	{
+		bOutIsNew = true;
+
+		StaticMesh = NewObject< UStaticMesh >( GetTransientPackage(), NAME_None, DefaultObjFlag | EObjectFlags::RF_Public );
+
+		FStaticMeshSourceModel& SourceModel = StaticMesh->AddSourceModel();
+		SourceModel.BuildSettings.bGenerateLightmapUVs = false;
+		SourceModel.BuildSettings.bRecomputeNormals = false;
+		SourceModel.BuildSettings.bRecomputeTangents = false;
+		SourceModel.BuildSettings.bBuildAdjacencyBuffer = false;
+		SourceModel.BuildSettings.bBuildReversedIndexBuffer = false;
+
+		FMeshDescription* StaticMeshDescription = StaticMesh->CreateMeshDescription(0);
+		check( StaticMeshDescription );
+		*StaticMeshDescription = MoveTemp( MeshDescription );
+
+		MeshCache.Add( MeshHash.ToString() ) = StaticMesh;
+	}
+	else
+	{
+		//FPlatformMisc::LowLevelOutputDebugStringf( TEXT("Mesh found in cache %s\n"), *StaticMesh->GetName() );
+		bOutIsNew = false;
+	}
+
+	return StaticMesh;
+}
+
+TMap< FString, UStaticMesh* > AUsdStageActor::LoadStaticMeshes( const pxr::UsdPrim& StartPrim )
+{
+	TMap< FString, UStaticMesh* > PrimPathsToStaticMeshes;
+
+	FScopedSlowTask LoadAllProgress( 6.f, LOCTEXT( "LoadingStaticMeshes", "Loading Static Meshes" ) );
+	LoadAllProgress.MakeDialog();
+
+	TArray< TUsdStore< pxr::UsdPrim > > AllGeomMeshes = UsdUtils::GetAllPrimsOfType( StartPrim, pxr::TfType::Find< pxr::UsdGeomMesh >(), { pxr::TfType::Find< pxr::UsdSkelRoot >() } );
+	TArray< FMeshDescription > MeshDescriptions;
+	MeshDescriptions.AddDefaulted( AllGeomMeshes.Num() );
+
+	// Create all mesh descriptions
+	{
+		LoadAllProgress.EnterProgressFrame();
+
+		ParallelFor( AllGeomMeshes.Num(), 
+			[ &AllGeomMeshes, &MeshDescriptions, TimeCode = pxr::UsdTimeCode( Time ) ]( int32 GeomMeshIndex )
+			{
+				MeshDescriptions[ GeomMeshIndex ] = UsdStageActorImpl::LoadMeshDescription( pxr::UsdGeomMesh( AllGeomMeshes[ GeomMeshIndex ].Get() ), TimeCode );
+			} );
+	}
+
+	TArray< TPair< int32, UStaticMesh* > > StaticMeshesToBuild;
+	StaticMeshesToBuild.Reserve( AllGeomMeshes.Num() );
+
+	// Create all static meshes
+	{
+		LoadAllProgress.EnterProgressFrame();
+
+		for ( int32 PrimIndex = 0; PrimIndex < AllGeomMeshes.Num(); ++PrimIndex )
+		{
+			bool bIsNew = true;
+			UStaticMesh* StaticMesh = CreateStaticMesh( pxr::UsdGeomMesh( AllGeomMeshes[ PrimIndex ].Get() ), MoveTemp( MeshDescriptions[ PrimIndex ] ), bIsNew );
+
+			PrimPathsToStaticMeshes.Add( UsdToUnreal::ConvertPath( AllGeomMeshes[ PrimIndex ].Get().GetPrimPath() ), StaticMesh );
+
+			if ( bIsNew )
+			{
+				StaticMeshesToBuild.Emplace( PrimIndex, StaticMesh );
+			}
+		}
+	}
+
+	// Commit all mesh descriptions
+	{
+		LoadAllProgress.EnterProgressFrame();
+
+		ParallelFor( StaticMeshesToBuild.Num(), 
+			[ &StaticMeshesToBuild ]( int32 MeshIndex )
+			{
+				UStaticMesh::FCommitMeshDescriptionParams Params;
+				Params.bMarkPackageDirty = false;
+				Params.bUseHashAsGuid = true;
+
+				TPair< int32, UStaticMesh* >& MeshToBuild = StaticMeshesToBuild[ MeshIndex ];
+
+				if ( UStaticMesh* StaticMesh = MeshToBuild.Value )
+				{
+					StaticMesh->CommitMeshDescription( 0, Params );
+				}
+			} );
+	}
+
+	// PreBuild all static meshes
+	{
+		LoadAllProgress.EnterProgressFrame();
+
+		for ( TPair< int32, UStaticMesh* >& MeshToBuild : StaticMeshesToBuild )
+		{
+			if ( UStaticMesh* StaticMesh = MeshToBuild.Value )
+			{
+				UsdStageActorImpl::PreBuildStaticMesh( pxr::UsdGeomMesh( AllGeomMeshes[ MeshToBuild.Key ].Get() ), *StaticMesh, MaterialsCache, Time );
+			}
+		}
+	}
+
+	// Build all static meshes
+	{
+		LoadAllProgress.EnterProgressFrame();
+
+		IMeshBuilderModule& MeshBuilderModule = FModuleManager::LoadModuleChecked< IMeshBuilderModule >( TEXT("MeshBuilder") );
+
+		ParallelFor( StaticMeshesToBuild.Num(), 
+			[ &StaticMeshesToBuild, &MeshBuilderModule ]( int32 MeshIndex )
+			{
+				if ( UStaticMesh* StaticMesh = StaticMeshesToBuild[ MeshIndex ].Value )
+				{
+					if ( !UsdStageActorImpl::BuildStaticMesh( *StaticMesh, MeshBuilderModule ) )
+					{
+						// Build failed, discard the mesh
+						StaticMeshesToBuild[ MeshIndex ].Value = nullptr;
+					}
+				}
+			} );
+	}
+
+	// PostBuild all static meshes
+	{
+		LoadAllProgress.EnterProgressFrame();
+
+		for ( TPair< int32, UStaticMesh* >& MeshToBuild : StaticMeshesToBuild )
+		{
+			if ( UStaticMesh* StaticMesh = MeshToBuild.Value )
+			{
+				UsdStageActorImpl::PostBuildStaticMesh( *StaticMesh );
+			}
+		}
+	}
+
+	return PrimPathsToStaticMeshes;
+}
+
+#undef LOCTEXT_NAMESPACE
\ No newline at end of file
