USD: Fixed uniform vertex colors

diff --git a/Engine/Plugins/Importers/USDImporter/Source/USDUtilities/Private/USDGeomMeshConversion.cpp b/Engine/Plugins/Importers/USDImporter/Source/USDUtilities/Private/USDGeomMeshConversion.cpp
--- a/Engine/Plugins/Importers/USDImporter/Source/USDUtilities/Private/USDGeomMeshConversion.cpp
+++ b/Engine/Plugins/Importers/USDImporter/Source/USDUtilities/Private/USDGeomMeshConversion.cpp
@@ -29,6 +29,28 @@
 
 #include "USDIncludesEnd.h"
 
+namespace UsdToUnrealImpl
+{
+	int32 GetPrimValueIndex( const pxr::TfToken& InterpType, const int32 VertexIndex, const int32 VertexInstanceIndex, const int32 PolygonIndex )
+	{
+		if ( InterpType == pxr::UsdGeomTokens->vertex )
+		{
+			return VertexIndex;
+		}
+		else if ( InterpType == pxr::UsdGeomTokens->faceVarying )
+		{
+			return VertexInstanceIndex;
+		}
+		else if ( InterpType == pxr::UsdGeomTokens->uniform )
+		{
+			return PolygonIndex;
+		}
+		else /* if ( InterpType == pxr::UsdGeomTokens->constant ) */
+		{
+			return 0; // return index 0 for constant or any other unsupported cases
+		}
+	}
+}
 
 bool UsdToUnreal::ConvertGeomMesh( const pxr::UsdGeomMesh& UsdMesh, FMeshDescription& MeshDescription, const pxr::UsdTimeCode TimeCode )
 {
@@ -124,7 +146,7 @@ bool UsdToUnreal::ConvertGeomMesh( const pxr::UsdGeomMesh& UsdMesh, FMeshDescrip
 			TOptional< VtIntArray > UVIndices; // UVs might be indexed or they might be flat (one per vertex)
 			VtVec2fArray UVs;
 
-			EUsdInterpolationMethod InterpolationMethod = EUsdInterpolationMethod::FaceVarying;
+			pxr::TfToken InterpType = UsdGeomTokens->faceVarying;
 		};
 
 		TArray< FUVSet > UVSets;
@@ -138,23 +160,7 @@ bool UsdToUnreal::ConvertGeomMesh( const pxr::UsdGeomMesh& UsdMesh, FMeshDescrip
 			if ( PrimvarST )
 			{
 				FUVSet UVSet;
-
-				if ( PrimvarST.GetInterpolation() == UsdGeomTokens->vertex )
-				{
-					UVSet.InterpolationMethod = EUsdInterpolationMethod::Vertex;
-				}
-				else if ( PrimvarST.GetInterpolation() == UsdGeomTokens->faceVarying )
-				{
-					UVSet.InterpolationMethod = EUsdInterpolationMethod::FaceVarying;
-				}
-				else if (  PrimvarST.GetInterpolation() == UsdGeomTokens->uniform )
-				{
-					UVSet.InterpolationMethod = EUsdInterpolationMethod::Uniform;
-				}
-				else if ( PrimvarST.GetInterpolation() == UsdGeomTokens->constant )
-				{
-					UVSet.InterpolationMethod = EUsdInterpolationMethod::Constant;
-				}
+				UVSet.InterpType = PrimvarST.GetInterpolation();
 
 				if ( PrimvarST.IsIndexed() )
 				{
@@ -206,25 +212,50 @@ bool UsdToUnreal::ConvertGeomMesh( const pxr::UsdGeomMesh& UsdMesh, FMeshDescrip
 		MeshDescription.ReserveNewPolygons( FaceCounts.size() );
 		MeshDescription.ReserveNewEdges( FaceCounts.size() * 2 );
 
+		// Vertex color
+		TVertexInstanceAttributesRef< FVector4 > MeshDescriptionColors = StaticMeshAttributes.GetVertexInstanceColors();
+
+		UsdGeomPrimvar ColorPrimvar = UsdMesh.GetDisplayColorPrimvar();
+		pxr::VtArray< pxr::GfVec3f > UsdColors;
+
+		if ( ColorPrimvar )
+		{	
+			ColorPrimvar.ComputeFlattened( &UsdColors, TimeCode );
+		}
+
+		// Vertex opacity
+		UsdGeomPrimvar OpacityPrimvar = UsdMesh.GetDisplayOpacityPrimvar();
+		pxr::VtArray< float > UsdOpacities;
+
+		if ( OpacityPrimvar )
+		{
+			OpacityPrimvar.ComputeFlattened( &UsdOpacities );
+		}
+
 		for ( int32 PolygonIndex = 0; PolygonIndex < FaceCounts.size(); ++PolygonIndex )
 		{
 			int32 PolygonVertexCount = FaceCounts[PolygonIndex];
-			CornerInstanceIDs.Reset();
-			CornerInstanceIDs.AddUninitialized(PolygonVertexCount);
-			CornerVerticesIDs.Reset();
-			CornerVerticesIDs.AddUninitialized(PolygonVertexCount);
+			CornerInstanceIDs.Reset( PolygonVertexCount );
+			CornerVerticesIDs.Reset( PolygonVertexCount );
 
 			for (int32 CornerIndex = 0; CornerIndex < PolygonVertexCount; ++CornerIndex, ++CurrentVertexInstanceIndex)
 			{
 				int32 VertexInstanceIndex = VertexInstanceOffset + CurrentVertexInstanceIndex;
 				const FVertexInstanceID VertexInstanceID(VertexInstanceIndex);
-				CornerInstanceIDs[CornerIndex] = VertexInstanceID;
 				const int32 ControlPointIndex = FaceIndices[CurrentVertexInstanceIndex];
 				const FVertexID VertexID(VertexOffset + ControlPointIndex);
 				const FVector VertexPosition = MeshDescriptionVertexPositions[VertexID];
-				CornerVerticesIDs[CornerIndex] = VertexID;
+
+				// Make sure a face doesn't use the same vertex twice as MeshDescription doesn't like that
+				if ( CornerVerticesIDs.Contains( VertexID ) )
+				{
+					continue;
+				}
+
+				CornerVerticesIDs.Add( VertexID );
 
 				FVertexInstanceID AddedVertexInstanceId = MeshDescription.CreateVertexInstance(VertexID);
+				CornerInstanceIDs.Add( AddedVertexInstanceId );
 
 				if ( Normals.size() > 0 )
 				{
@@ -240,24 +271,7 @@ bool UsdToUnreal::ConvertGeomMesh( const pxr::UsdGeomMesh& UsdMesh, FMeshDescrip
 				int32 UVLayerIndex = 0;
 				for ( const FUVSet& UVSet : UVSets )
 				{
-					int32 ValueIndex = 0;
-
-					if ( UVSet.InterpolationMethod == EUsdInterpolationMethod::Vertex )
-					{
-						ValueIndex = VertexID.GetValue();
-					}
-					else if ( UVSet.InterpolationMethod == EUsdInterpolationMethod::FaceVarying )
-					{
-						ValueIndex = CurrentVertexInstanceIndex;
-					}
-					else if ( UVSet.InterpolationMethod == EUsdInterpolationMethod::Uniform )
-					{
-						ValueIndex = PolygonIndex;
-					}
-					else if ( UVSet.InterpolationMethod == EUsdInterpolationMethod::Constant )
-					{
-						ValueIndex = 0;
-					}
+					const int32 ValueIndex = UsdToUnrealImpl::GetPrimValueIndex( UVSet.InterpType, VertexID.GetValue(), CurrentVertexInstanceIndex, PolygonIndex );
 
 					GfVec2f UV( 0.f, 0.f );
 
@@ -279,6 +293,35 @@ bool UsdToUnreal::ConvertGeomMesh( const pxr::UsdGeomMesh& UsdMesh, FMeshDescrip
 
 					++UVLayerIndex;
 				}
+
+				// Vertex color
+				{
+					auto ConvertToLinear = []( const pxr::GfVec3f& UsdColor ) -> FLinearColor
+					{
+						return FLinearColor( FLinearColor( UsdToUnreal::ConvertColor( UsdColor ) ).ToFColor( false ) );
+					};
+
+					const int32 ValueIndex = UsdToUnrealImpl::GetPrimValueIndex( ColorPrimvar.GetInterpolation(), VertexID.GetValue(), CurrentVertexInstanceIndex, PolygonIndex );
+
+					GfVec3f UsdColor( 1.f, 1.f, 1.f );
+
+					if ( !UsdColors.empty() && ensure( UsdColors.size() > ValueIndex ) )
+					{
+						UsdColor = UsdColors[ ValueIndex ];
+					}
+
+					MeshDescriptionColors[ AddedVertexInstanceId ] = ConvertToLinear( UsdColor );
+				}
+
+				// Vertex opacity
+				{
+					const int32 ValueIndex = UsdToUnrealImpl::GetPrimValueIndex( OpacityPrimvar.GetInterpolation(), VertexID.GetValue(), CurrentVertexInstanceIndex, PolygonIndex );
+
+					if ( !UsdOpacities.empty() && ensure( UsdOpacities.size() > ValueIndex ) )
+					{
+						MeshDescriptionColors[ AddedVertexInstanceId ][3] = UsdOpacities[ ValueIndex ];
+					}
+				}
 			}
 
 			int32 MaterialIndex = 0;
@@ -327,81 +370,6 @@ bool UsdToUnreal::ConvertGeomMesh( const pxr::UsdGeomMesh& UsdMesh, FMeshDescrip
 				MeshDescription.ReversePolygonFacing( NewPolygonID );
 			}
 		}
-
-		// Vertex color
-		TVertexInstanceAttributesRef< FVector4 > MeshDescriptionColors = StaticMeshAttributes.GetVertexInstanceColors();
-
-		UsdGeomPrimvar ColorPrimvar = UsdMesh.GetDisplayColorPrimvar();
-		if ( ColorPrimvar )
-		{
-			pxr::VtArray<pxr::GfVec3f> UsdColors;
-			ColorPrimvar.ComputeFlattened(&UsdColors);
-
-			int32 NumColors = UsdColors.size();
-
-			auto ConvertToLinear = []( const pxr::GfVec3f& UsdColor ) -> FLinearColor
-			{
-				return FLinearColor( FLinearColor( UsdToUnreal::ConvertColor( UsdColor ) ).ToFColor( false ) );
-			};
-
-			pxr::TfToken UsdInterpType = ColorPrimvar.GetInterpolation();
-			if ( UsdInterpType == pxr::UsdGeomTokens->faceVarying && NumColors >= MeshDescriptionColors.GetNumElements() )
-			{
-				for(int Index = 0; Index < MeshDescriptionColors.GetNumElements(); ++Index)
-				{
-					MeshDescriptionColors[ FVertexInstanceID(Index) ] = ConvertToLinear( UsdColors[Index] );
-				}
-			}
-			else if ( UsdInterpType == pxr::UsdGeomTokens->vertex && NumColors >= MeshDescription.Vertices().Num() )
-			{
-				for( FVertexInstanceID VertexInstID : MeshDescription.VertexInstances().GetElementIDs() )
-				{
-					FVertexID VertexID = MeshDescription.GetVertexInstanceVertex( VertexInstID );
-					MeshDescriptionColors[ VertexInstID ] = ConvertToLinear( UsdColors[ VertexID.GetValue() ] );
-				}
-			}
-			else if ( UsdInterpType == pxr::UsdGeomTokens->constant && NumColors == 1 )
-			{
-				for( int Index = 0; Index < MeshDescriptionColors.GetNumElements(); ++Index )
-				{
-					MeshDescriptionColors[ FVertexInstanceID(Index) ] = ConvertToLinear( UsdColors[0] );
-				}
-			}
-		}
-
-		// Vertex opacity
-		UsdGeomPrimvar OpacityPrimvar = UsdMesh.GetDisplayOpacityPrimvar();
-		if ( OpacityPrimvar )
-		{
-			pxr::VtArray< float > UsdOpacities;
-			OpacityPrimvar.ComputeFlattened( &UsdOpacities );
-
-			const int32 NumOpacities = UsdOpacities.size();
-
-			pxr::TfToken UsdInterpType = OpacityPrimvar.GetInterpolation();
-			if ( UsdInterpType == pxr::UsdGeomTokens->faceVarying && NumOpacities >= MeshDescriptionColors.GetNumElements() )
-			{
-				for(int Index = 0; Index < MeshDescriptionColors.GetNumElements(); ++Index)
-				{
-					MeshDescriptionColors[ FVertexInstanceID(Index) ][3] = UsdOpacities[Index];
-				}
-			}
-			else if ( UsdInterpType == pxr::UsdGeomTokens->vertex && NumOpacities >= MeshDescription.Vertices().Num() )
-			{
-				for( FVertexInstanceID VertexInstID : MeshDescription.VertexInstances().GetElementIDs() )
-				{
-					FVertexID VertexID = MeshDescription.GetVertexInstanceVertex( VertexInstID );
-					MeshDescriptionColors[ VertexInstID ][3] = UsdOpacities[ VertexID.GetValue() ];
-				}
-			}
-			else if ( UsdInterpType == pxr::UsdGeomTokens->constant && NumOpacities == 1 )
-			{
-				for( int Index = 0; Index < MeshDescriptionColors.GetNumElements(); ++Index )
-				{
-					MeshDescriptionColors[ FVertexInstanceID(Index) ][3] = UsdOpacities[0];
-				}
-			}
-		}
 	}
 
 	return true;
