Fixed check when a face in USDGeomMesh references the same vertex multiple times. We consider this as invalid and will skip the repeating vertices.

#rnx

diff --git a/Engine/Plugins/Importers/USDImporter/Source/USDUtilities/Private/USDGeomMeshConversion.cpp b/Engine/Plugins/Importers/USDImporter/Source/USDUtilities/Private/USDGeomMeshConversion.cpp
--- a/Engine/Plugins/Importers/USDImporter/Source/USDUtilities/Private/USDGeomMeshConversion.cpp
+++ b/Engine/Plugins/Importers/USDImporter/Source/USDUtilities/Private/USDGeomMeshConversion.cpp
@@ -206,22 +206,27 @@ bool UsdToUnreal::ConvertGeomMesh( const pxr::UsdGeomMesh& UsdMesh, FMeshDescrip
 		for ( int32 PolygonIndex = 0; PolygonIndex < FaceCounts.size(); ++PolygonIndex )
 		{
 			int32 PolygonVertexCount = FaceCounts[PolygonIndex];
-			CornerInstanceIDs.Reset();
-			CornerInstanceIDs.AddUninitialized(PolygonVertexCount);
-			CornerVerticesIDs.Reset();
-			CornerVerticesIDs.AddUninitialized(PolygonVertexCount);
+			CornerInstanceIDs.Reset( PolygonVertexCount );
+			CornerVerticesIDs.Reset( PolygonVertexCount );
 
 			for (int32 CornerIndex = 0; CornerIndex < PolygonVertexCount; ++CornerIndex, ++CurrentVertexInstanceIndex)
 			{
 				int32 VertexInstanceIndex = VertexInstanceOffset + CurrentVertexInstanceIndex;
 				const FVertexInstanceID VertexInstanceID(VertexInstanceIndex);
-				CornerInstanceIDs[CornerIndex] = VertexInstanceID;
 				const int32 ControlPointIndex = FaceIndices[CurrentVertexInstanceIndex];
 				const FVertexID VertexID(VertexOffset + ControlPointIndex);
 				const FVector VertexPosition = MeshDescriptionVertexPositions[VertexID];
-				CornerVerticesIDs[CornerIndex] = VertexID;
+
+				// Make sure a face doesn't use the same vertex twice as MeshDescription doesn't like that
+				if ( CornerVerticesIDs.Contains( VertexID ) )
+				{
+					continue;
+				}
+
+				CornerVerticesIDs.Add( VertexID );
 
 				FVertexInstanceID AddedVertexInstanceId = MeshDescription.CreateVertexInstance(VertexID);
+				CornerInstanceIDs.Add( AddedVertexInstanceId );
 
 				if ( Normals.size() > 0 )
 				{
