Fixed crash when opening an invalid USD file.

#rnx

diff --git a/Engine/Plugins/Importers/USDImporter/Source/USDStage/Private/USDStageActor.cpp b/Engine/Plugins/Importers/USDImporter/Source/USDStage/Private/USDStageActor.cpp
--- a/Engine/Plugins/Importers/USDImporter/Source/USDStage/Private/USDStageActor.cpp
+++ b/Engine/Plugins/Importers/USDImporter/Source/USDStage/Private/USDStageActor.cpp
@@ -883,11 +883,14 @@ void AUsdStageActor::OpenUsdStage()
 		UsdUtils::SetUsdStageAxis( UsdStageStore.Get(), pxr::UsdGeomTokens->z );
 	}
 
-	UsdStageStore.Get()->SetEditTarget( UsdStageStore.Get()->GetRootLayer() );
+	if ( UsdStageStore.Get() )
+	{
+		UsdStageStore.Get()->SetEditTarget( UsdStageStore.Get()->GetRootLayer() );
 
-	UsdListener.Register( UsdStageStore.Get() );
+		UsdListener.Register( UsdStageStore.Get() );
 
-	OnStageChanged.Broadcast();
+		OnStageChanged.Broadcast();
+	}
 #endif // #if USE_USD_SDK
 }
 
