Fixed crash with temporary std::string being deleted with the wrong allocator

#rnx

diff --git a/Engine/Plugins/Importers/USDImporter/Source/USDUtilities/Public/USDTypesConversion.h b/Engine/Plugins/Importers/USDImporter/Source/USDUtilities/Public/USDTypesConversion.h
--- a/Engine/Plugins/Importers/USDImporter/Source/USDUtilities/Public/USDTypesConversion.h
+++ b/Engine/Plugins/Importers/USDImporter/Source/USDUtilities/Public/USDTypesConversion.h
@@ -28,6 +28,13 @@ namespace UsdToUnreal
 		return FString( ANSI_TO_TCHAR( InString.c_str() ) );
 	}
 
+	static FString ConvertString( std::string&& InString )
+	{
+		TUsdStore< std::string > UsdString( MoveTemp( InString ) ); // Store the temporary so that it gets destroyed with the USD allocator
+
+		return FString( ANSI_TO_TCHAR( UsdString.Get().c_str() ) );
+	}
+
 	static FString ConvertString( const char* InString )
 	{
 		return FString( ANSI_TO_TCHAR( InString ) );
@@ -48,6 +55,13 @@ namespace UsdToUnreal
 		return FName( InString.c_str() );
 	}
 
+	static FName ConvertName( std::string&& InString )
+	{
+		TUsdStore< std::string > UsdString( MoveTemp( InString ) ); // Store the temporary so that it gets destroyed with the USD allocator
+
+		return FName( UsdString.Get().c_str() );
+	}
+
 	static FLinearColor ConvertColor( const pxr::GfVec3f& InValue )
 	{
 		return FLinearColor( InValue[0], InValue[1], InValue[2] );
