USD : Made sublayer paths relative to their parent layer.

#rnx

diff --git a/Engine/Plugins/Importers/USDImporter/Source/USDStageEditor/Private/SUSDLayersTreeView.cpp b/Engine/Plugins/Importers/USDImporter/Source/USDStageEditor/Private/SUSDLayersTreeView.cpp
--- a/Engine/Plugins/Importers/USDImporter/Source/USDStageEditor/Private/SUSDLayersTreeView.cpp
+++ b/Engine/Plugins/Importers/USDImporter/Source/USDStageEditor/Private/SUSDLayersTreeView.cpp
@@ -481,9 +481,9 @@ void SUsdLayersTreeView::OnEditSelectedLayer()
 
 void SUsdLayersTreeView::OnAddSubLayer()
 {
-	TOptional< FString > LayerFile = UsdUtils::BrowseUsdFile( UsdUtils::EBrowseFileMode::Open, AsShared() );
+	TOptional< FString > SubLayerFile = UsdUtils::BrowseUsdFile( UsdUtils::EBrowseFileMode::Open, AsShared() );
 
-	if ( !LayerFile )
+	if ( !SubLayerFile )
 	{
 		return;
 	}
@@ -497,7 +497,13 @@ void SUsdLayersTreeView::OnAddSubLayer()
 		pxr::SdfLayerHandle LayerHandle = SelectedItem->GetLayerHandle().Get();
 		if ( LayerHandle )
 		{
-			LayerHandle->InsertSubLayerPath( UnrealToUsd::ConvertString( *LayerFile.GetValue() ).Get() );
+			std::string UsdLayerFilePath = LayerHandle->GetRealPath();
+			FString LayerFilePath = UsdToUnreal::ConvertString( UsdLayerFilePath );
+
+			FString SubLayerFilePath = FPaths::ConvertRelativePathToFull( SubLayerFile.GetValue() );
+			FPaths::MakePathRelativeTo( SubLayerFilePath, *LayerFilePath );
+
+			LayerHandle->InsertSubLayerPath( UnrealToUsd::ConvertString( *SubLayerFilePath ).Get() );
 		}
 		
 		break;
