USD Schemas translators
- Generic schema parsing system to allow multithreaded parsing and parsing of user defined schemas.
- Added support for USD Point Instancers.
- Skeletal Meshes now use the assets cache.

diff --git a/Engine/Plugins/Importers/USDImporter/Source/USDSchemas/Private/USDGeomCameraTranslator.cpp b/Engine/Plugins/Importers/USDImporter/Source/USDSchemas/Private/USDGeomCameraTranslator.cpp
new file mode 100644
--- /dev/null
+++ b/Engine/Plugins/Importers/USDImporter/Source/USDSchemas/Private/USDGeomCameraTranslator.cpp
@@ -0,0 +1,44 @@
+// Copyright 1998-2019 Epic Games, Inc. All Rights Reserved.
+
+#include "USDGeomCameraTranslator.h"
+
+#include "USDMemory.h"
+#include "USDPrimConversion.h"
+#include "USDSchemasModule.h"
+
+#include "CineCameraActor.h"
+#include "CineCameraComponent.h"
+#include "Modules/ModuleManager.h"
+
+#if USE_USD_SDK
+#include "USDIncludesStart.h"
+	#include "pxr/usd/usdGeom/camera.h"
+#include "USDIncludesEnd.h"
+
+
+void FUsdGeomCameraTranslator::RegisterTranslator()
+{
+	IUsdSchemasModule& UsdSchemasModule = FModuleManager::Get().LoadModuleChecked< IUsdSchemasModule >( TEXT("USDSchemas") );
+	UsdSchemasModule.GetTranslatorRegistry().Register< FUsdGeomCameraTranslator >( TEXT("UsdGeomCamera") );
+}
+
+USceneComponent* FUsdGeomCameraTranslator::CreateComponents()
+{
+	constexpr bool bNeedsActor = true;
+	USceneComponent* RootComponent = FUsdGeomXformableTranslator::CreateComponents( UCineCameraComponent::StaticClass(), bNeedsActor );
+
+	UpdateComponents( RootComponent );
+
+	return RootComponent;
+}
+
+void FUsdGeomCameraTranslator::UpdateComponents( USceneComponent* SceneComponent )
+{
+	if ( UCineCameraComponent* CameraComponent = Cast< UCineCameraComponent >( SceneComponent ) )
+	{
+		FScopedUsdAllocs UsdAllocs;
+		UsdToUnreal::ConvertGeomCamera( Schema.Get().GetPrim().GetStage(), pxr::UsdGeomCamera( Schema.Get() ), *CameraComponent, pxr::UsdTimeCode( Context->Time ) );
+	}
+}
+
+#endif // #if USE_USD_SDK
diff --git a/Engine/Plugins/Importers/USDImporter/Source/USDSchemas/Private/USDGeomMeshTranslator.cpp b/Engine/Plugins/Importers/USDImporter/Source/USDSchemas/Private/USDGeomMeshTranslator.cpp
new file mode 100644
--- /dev/null
+++ b/Engine/Plugins/Importers/USDImporter/Source/USDSchemas/Private/USDGeomMeshTranslator.cpp
@@ -0,0 +1,381 @@
+// Copyright 1998-2019 Epic Games, Inc. All Rights Reserved.
+
+#include "USDGeomMeshTranslator.h"
+
+#if USE_USD_SDK
+
+#include "USDConversionUtils.h"
+#include "USDGeomMeshConversion.h"
+#include "USDSchemasModule.h"
+#include "USDTypesConversion.h"
+
+#include "Async/Async.h"
+#include "Components/StaticMeshComponent.h"
+#include "Engine/CollisionProfile.h"
+#include "Engine/StaticMesh.h"
+#include "IMeshBuilderModule.h"
+#include "Interfaces/ITargetPlatform.h"
+#include "Interfaces/ITargetPlatformManagerModule.h"
+#include "Materials/Material.h"
+#include "Materials/MaterialInterface.h"
+#include "Materials/MaterialInstanceConstant.h"
+#include "MeshDescriptionOperations.h"
+#include "MeshDescriptionToDynamicMesh.h"
+#include "Misc/SecureHash.h"
+#include "PhysicsEngine/BodySetup.h"
+#include "StaticMeshAttributes.h"
+
+#include "USDIncludesStart.h"
+	#include "pxr/usd/usd/prim.h"
+	#include "pxr/usd/usd/stage.h"
+	#include "pxr/usd/usdGeom/mesh.h"
+	#include "pxr/usd/usdGeom/xformable.h"
+	#include "pxr/usd/usdShade/material.h"
+#include "USDIncludesEnd.h"
+
+
+namespace UsdGeomMeshTranslatorImpl
+{
+	bool IsGeometryAnimated( const pxr::UsdGeomMesh& GeomMesh, const pxr::UsdTimeCode TimeCode )
+	{
+		FScopedUsdAllocs UsdAllocs;
+
+		bool bHasAttributesTimeSamples = false;
+		{
+			constexpr bool bIncludeInherited = false;
+			pxr::TfTokenVector GeomMeshAttributeNames = pxr::UsdGeomMesh::GetSchemaAttributeNames( bIncludeInherited );
+			pxr::TfTokenVector GeomPointBasedAttributeNames = pxr::UsdGeomPointBased::GetSchemaAttributeNames( bIncludeInherited );
+
+			GeomMeshAttributeNames.reserve( GeomMeshAttributeNames.size() + GeomPointBasedAttributeNames.size() );
+			GeomMeshAttributeNames.insert( GeomMeshAttributeNames.end(), GeomPointBasedAttributeNames.begin(), GeomPointBasedAttributeNames.end() );
+
+			for ( const pxr::TfToken& AttributeName : GeomMeshAttributeNames )
+			{
+				const pxr::UsdAttribute& Attribute = GeomMesh.GetPrim().GetAttribute( AttributeName );
+
+				//bHasAttributesTimeSamples = Attribute.ValueMightBeTimeVarying();
+
+				double DesiredTime = TimeCode.GetValue();
+				double MinTime = 0.0;
+				double MaxTime = 0.0;
+				bool bHasTimeSamples = false;
+
+				Attribute.GetBracketingTimeSamples( DesiredTime, &MinTime, &MaxTime, &bHasTimeSamples );
+				if ( bHasTimeSamples && DesiredTime >= MinTime && DesiredTime <= MaxTime )
+				{
+					bHasAttributesTimeSamples = true;
+					break;
+				}
+			}
+		}
+
+		return bHasAttributesTimeSamples;
+	}
+
+	void ProcessMaterials( const pxr::UsdPrim& UsdPrim, UStaticMesh& StaticMesh, TMap< FString, UObject* >& MaterialsCache, bool bHasPrimDisplayColor, float Time )
+	{
+		const FMeshDescription* MeshDescription = StaticMesh.GetMeshDescription( 0 );
+
+		if ( !MeshDescription )
+		{
+			return;
+		}
+
+		FStaticMeshConstAttributes StaticMeshAttributes( *MeshDescription );
+
+		for ( const FPolygonGroupID PolygonGroupID : MeshDescription->PolygonGroups().GetElementIDs() )
+		{
+			const FName& ImportedMaterialSlotName = StaticMeshAttributes.GetPolygonGroupMaterialSlotNames()[ PolygonGroupID ];
+			const FName MaterialSlotName = ImportedMaterialSlotName;
+
+			int32 MaterialIndex = INDEX_NONE;
+			int32 MeshMaterialIndex = 0;
+
+			for ( FStaticMaterial& StaticMaterial : StaticMesh.StaticMaterials )
+			{
+				if ( StaticMaterial.MaterialSlotName.IsEqual( ImportedMaterialSlotName ) )
+				{
+					MaterialIndex = MeshMaterialIndex;
+					break;
+				}
+
+				++MeshMaterialIndex;
+			}
+
+			if ( MaterialIndex == INDEX_NONE )
+			{
+				MaterialIndex = PolygonGroupID.GetValue();
+			}
+
+			UMaterialInterface* Material = nullptr;
+			TUsdStore< pxr::UsdPrim > MaterialPrim = UsdPrim.GetStage()->GetPrimAtPath( UnrealToUsd::ConvertPath( *ImportedMaterialSlotName.ToString() ).Get() );
+
+			if ( MaterialPrim.Get() )
+			{
+				UObject*& CachedMaterial = MaterialsCache.FindOrAdd( UsdToUnreal::ConvertPath( MaterialPrim.Get().GetPrimPath() ) );
+
+				if ( !CachedMaterial )
+				{
+					UMaterial* NewMaterial = NewObject< UMaterial >();
+
+					if ( UsdToUnreal::ConvertMaterial( pxr::UsdShadeMaterial( MaterialPrim.Get() ), *NewMaterial ) )
+					{
+						//UMaterialEditingLibrary::RecompileMaterial( CachedMaterial ); // Too slow
+						NewMaterial->PostEditChange();
+					}
+					else
+					{
+						NewMaterial = nullptr;
+					}
+
+					Material = NewMaterial;
+				}
+				else
+				{
+					Material = Cast< UMaterialInterface >( CachedMaterial );
+				}
+			}
+
+			if ( Material == nullptr && bHasPrimDisplayColor )
+			{
+				UMaterialInstanceConstant* MaterialInstance = NewObject< UMaterialInstanceConstant >();
+				if ( UsdToUnreal::ConvertDisplayColor( pxr::UsdGeomMesh( UsdPrim ), *MaterialInstance, pxr::UsdTimeCode( Time ) ) )
+				{
+					Material = MaterialInstance;
+				}
+			}
+
+			FStaticMaterial StaticMaterial( Material, MaterialSlotName, ImportedMaterialSlotName );
+				StaticMesh.StaticMaterials.Add( StaticMaterial );
+
+			FMeshSectionInfo MeshSectionInfo;
+			MeshSectionInfo.MaterialIndex = MaterialIndex;
+
+			StaticMesh.GetSectionInfoMap().Set( 0, PolygonGroupID.GetValue(), MeshSectionInfo );
+		}
+	}
+
+	FMeshDescription LoadMeshDescription( const pxr::UsdGeomMesh& UsdMesh, const pxr::UsdTimeCode TimeCode )
+	{
+		if ( !UsdMesh )
+		{
+			return {};
+		}
+
+		FMeshDescription MeshDescription;
+		FStaticMeshAttributes StaticMeshAttributes( MeshDescription );
+		StaticMeshAttributes.Register();
+
+		UsdToUnreal::ConvertGeomMesh( UsdMesh, MeshDescription, TimeCode );
+
+		return MeshDescription;
+	}
+
+	UStaticMesh* CreateStaticMesh( const pxr::UsdGeomMesh& UsdMesh, FMeshDescription&& MeshDescription, FUsdSchemaTranslationContext& Context, bool& bOutIsNew )
+	{
+		UStaticMesh* StaticMesh = nullptr;
+
+		FSHAHash MeshHash = FMeshDescriptionOperations::ComputeSHAHash( MeshDescription );
+
+		StaticMesh = Cast< UStaticMesh >( Context.AssetsCache.FindRef( MeshHash.ToString() ) );
+
+		if ( !StaticMesh && !MeshDescription.IsEmpty() )
+		{
+			bOutIsNew = true;
+
+			StaticMesh = NewObject< UStaticMesh >( GetTransientPackage(), NAME_None, Context.ObjectFlags | EObjectFlags::RF_Public );
+
+			FStaticMeshSourceModel& SourceModel = StaticMesh->AddSourceModel();
+			SourceModel.BuildSettings.bGenerateLightmapUVs = false;
+			SourceModel.BuildSettings.bRecomputeNormals = false;
+			SourceModel.BuildSettings.bRecomputeTangents = false;
+			SourceModel.BuildSettings.bBuildAdjacencyBuffer = false;
+			SourceModel.BuildSettings.bBuildReversedIndexBuffer = false;
+
+			FMeshDescription* StaticMeshDescription = StaticMesh->CreateMeshDescription(0);
+			check( StaticMeshDescription );
+			*StaticMeshDescription = MoveTemp( MeshDescription );
+
+			Context.AssetsCache.Add( MeshHash.ToString() ) = StaticMesh;
+		}
+		else
+		{
+			//FPlatformMisc::LowLevelOutputDebugStringf( TEXT("Mesh found in cache %s\n"), *StaticMesh->GetName() );
+			bOutIsNew = false;
+		}
+
+		return StaticMesh;
+	}
+
+	void PreBuildStaticMesh( const pxr::UsdGeomMesh& UsdMesh, UStaticMesh& StaticMesh, TMap< FString, UObject* >& MaterialsCache, float Time )
+	{
+		const bool bHasPrimDisplayColor = UsdMesh.GetDisplayColorPrimvar().IsDefined();
+		ProcessMaterials( UsdMesh.GetPrim(), StaticMesh, MaterialsCache, bHasPrimDisplayColor, Time );
+
+		StaticMesh.RenderData = MakeUnique< FStaticMeshRenderData >();
+		//StaticMesh.BodySetup = NewObject< UBodySetup >( &StaticMesh );
+	}
+
+	bool BuildStaticMesh( UStaticMesh& StaticMesh, IMeshBuilderModule& MeshBuilderModule )
+	{
+		//StaticMesh.BodySetup->DefaultInstance.SetCollisionProfileName( UCollisionProfile::BlockAll_ProfileName );
+
+		ITargetPlatformManagerModule& TargetPlatformManager = GetTargetPlatformManagerRef();
+		ITargetPlatform* RunningPlatform = TargetPlatformManager.GetRunningTargetPlatform();
+		check(RunningPlatform);
+
+		const FStaticMeshLODSettings& LODSettings = RunningPlatform->GetStaticMeshLODSettings();
+
+		const FStaticMeshLODGroup& LODGroup = LODSettings.GetLODGroup(StaticMesh.LODGroup);
+
+		if ( !MeshBuilderModule.BuildMesh( *StaticMesh.RenderData, &StaticMesh, LODGroup ) )
+		{
+			return false;
+		}
+
+		for ( FStaticMeshLODResources& LODResources : StaticMesh.RenderData->LODResources )
+		{
+			LODResources.bHasColorVertexData = true;
+		}
+
+		return true;
+	}
+
+	void PostBuildStaticMesh( UStaticMesh& StaticMesh )
+	{
+		//StaticMesh.CreateBodySetup();
+		StaticMesh.InitResources();
+
+		if ( const FMeshDescription* MeshDescription = StaticMesh.GetMeshDescription( 0 ) )
+		{
+			StaticMesh.RenderData->Bounds = MeshDescription->GetBounds();
+		}
+	
+		StaticMesh.CalculateExtendedBounds();
+	}
+}
+
+void FGeomMeshCreateAssetsTaskChain::SetupTasks()
+{
+	FScopedUnrealAllocs UnrealAllocs;
+
+	// Create mesh description (Async)
+	constexpr bool bIsAsyncTask = true;
+	Do( bIsAsyncTask, 
+		[ this ]() -> bool
+		{
+			MeshDescription = UsdGeomMeshTranslatorImpl::LoadMeshDescription( GeomMesh.Get(), pxr::UsdTimeCode( Context->Time ) );
+
+			return !MeshDescription.IsEmpty();
+		} );
+		
+	if ( !UsdGeomMeshTranslatorImpl::IsGeometryAnimated( GeomMesh.Get(), pxr::UsdTimeCode( Context->Time ) ) )
+	{
+		IMeshBuilderModule* MeshBuilderModule = &FModuleManager::LoadModuleChecked< IMeshBuilderModule >( TEXT("MeshBuilder") );
+
+		// Create static mesh (Main thread)
+		Then( !bIsAsyncTask,
+			[ this ]()
+			{
+				bool bIsNew = true;
+				StaticMesh = UsdGeomMeshTranslatorImpl::CreateStaticMesh( GeomMesh.Get(), MoveTemp( MeshDescription ), *Context, bIsNew );
+
+				FScopeLock Lock( &Context->CriticalSection );
+				{
+					Context->PrimPathsToAssets.Add( UsdToUnreal::ConvertPath( GeomMesh.Get().GetPrim().GetPrimPath() ), StaticMesh );
+				}
+
+				if ( !bIsNew )
+				{
+					StaticMesh = nullptr;
+				}
+
+				return bIsNew;
+			} );
+
+		// Commit mesh description (Async)
+		Then( bIsAsyncTask,
+			[ this ]()
+			{
+				UStaticMesh::FCommitMeshDescriptionParams Params;
+				Params.bMarkPackageDirty = false;
+				Params.bUseHashAsGuid = true;
+
+				StaticMesh->CommitMeshDescription( 0, Params );
+
+				return true;
+			} );
+
+		// PreBuild static mesh (Main thread)
+		Then( !bIsAsyncTask,
+			[ this ]()
+			{
+				UsdGeomMeshTranslatorImpl::PreBuildStaticMesh( GeomMesh.Get(), *StaticMesh, Context->AssetsCache, Context->Time );
+
+				return true;
+			} );
+
+		// Build static mesh (Async)
+		Then( bIsAsyncTask,
+			[ this, MeshBuilderModule ]() mutable
+			{
+				if ( !UsdGeomMeshTranslatorImpl::BuildStaticMesh( *StaticMesh, *MeshBuilderModule ) )
+				{
+					// Build failed, discard the mesh
+					StaticMesh = nullptr;
+
+					return false;
+				}
+
+				return true;
+			} );
+
+		// PostBuild static mesh (Main thread)
+		Then( !bIsAsyncTask,
+			[ this ]()
+			{
+				UsdGeomMeshTranslatorImpl::PostBuildStaticMesh( *StaticMesh );
+
+				return true;
+			} );
+	}
+}
+
+void FUsdGeomMeshTranslator::RegisterTranslator()
+{
+	IUsdSchemasModule& UsdSchemasModule = FModuleManager::Get().LoadModuleChecked< IUsdSchemasModule >( TEXT("USDSchemas") );
+	UsdSchemasModule.GetTranslatorRegistry().Register< FUsdGeomMeshTranslator >( TEXT("UsdGeomMesh") );
+}
+
+void FUsdGeomMeshTranslator::CreateAssets()
+{
+	TSharedRef< FGeomMeshCreateAssetsTaskChain > AssetsTaskChain = MakeShared< FGeomMeshCreateAssetsTaskChain >( Context, pxr::UsdGeomMesh( Schema.Get() ) );
+
+	Context->TranslatorTasks.Add( MoveTemp( AssetsTaskChain ) );
+}
+
+USceneComponent* FUsdGeomMeshTranslator::CreateComponents()
+{
+	USceneComponent* RootComponent = FUsdGeomXformableTranslator::CreateComponents();
+
+	if ( UStaticMeshComponent* StaticMeshComponent = Cast< UStaticMeshComponent >( RootComponent ) )
+	{
+		UStaticMesh* PrimStaticMesh = Cast< UStaticMesh >( Context->PrimPathsToAssets.FindRef( UsdToUnreal::ConvertPath( Schema.Get().GetPath() ) ) );
+
+		if ( PrimStaticMesh != StaticMeshComponent->GetStaticMesh() )
+		{
+			if ( StaticMeshComponent->IsRegistered() )
+			{
+				StaticMeshComponent->UnregisterComponent();
+			}
+
+			StaticMeshComponent->SetStaticMesh( PrimStaticMesh );
+
+			StaticMeshComponent->RegisterComponent();
+		}
+	}
+
+	return RootComponent;
+}
+
+#endif // #if USE_USD_SDK
diff --git a/Engine/Plugins/Importers/USDImporter/Source/USDSchemas/Private/USDGeomPointInstancerTranslator.cpp b/Engine/Plugins/Importers/USDImporter/Source/USDSchemas/Private/USDGeomPointInstancerTranslator.cpp
new file mode 100644
--- /dev/null
+++ b/Engine/Plugins/Importers/USDImporter/Source/USDSchemas/Private/USDGeomPointInstancerTranslator.cpp
@@ -0,0 +1,161 @@
+// Copyright 1998-2019 Epic Games, Inc. All Rights Reserved.
+
+#include "USDGeomPointInstancerTranslator.h"
+
+#include "USDConversionUtils.h"
+#include "USDSchemasModule.h"
+#include "USDTypesConversion.h"
+
+#if USE_USD_SDK
+
+#include "Components/HierarchicalInstancedStaticMeshComponent.h"
+#include "Engine/StaticMesh.h"
+#include "Modules/ModuleManager.h"
+
+#include "USDIncludesStart.h"
+	#include "pxr/usd/usd/prim.h"
+	#include "pxr/usd/usd/stage.h"
+	#include "pxr/usd/usdGeom/camera.h"
+	#include "pxr/usd/usdGeom/mesh.h"
+	#include "pxr/usd/usdGeom/pointInstancer.h"
+	#include "pxr/usd/usdGeom/xformable.h"
+#include "USDIncludesEnd.h"
+
+namespace UsdGeomPointInstancerTranslatorImpl
+{
+	bool ConvertGeomPointInstancer( const pxr::UsdStageRefPtr& Stage, const pxr::UsdGeomPointInstancer& PointInstancer, const int32 ProtoIndex, UHierarchicalInstancedStaticMeshComponent& HismComponent, pxr::UsdTimeCode EvalTime )
+	{
+		FScopedUsdAllocs UsdAllocs;
+
+		pxr::VtArray< int > ProtoIndices = UsdUtils::GetUsdValue< pxr::VtArray< int > >( PointInstancer.GetProtoIndicesAttr(), EvalTime );
+
+		pxr::VtMatrix4dArray UsdInstanceTransforms;
+
+		if ( !PointInstancer.ComputeInstanceTransformsAtTime(
+				&UsdInstanceTransforms,
+				EvalTime,
+				EvalTime ) )
+		{
+			return false;
+		}
+
+		pxr::TfToken UpAxis = UsdUtils::GetUsdStageAxis( Stage );
+
+		int32 Index = 0;
+
+		FScopedUnrealAllocs UnrealAllocs;
+
+		for ( pxr::GfMatrix4d& UsdMatrix : UsdInstanceTransforms )
+		{
+			if ( ProtoIndices[ Index ] == ProtoIndex )
+			{
+				FTransform InstanceTransform = UsdToUnreal::ConvertMatrix( UpAxis, UsdMatrix );
+
+				HismComponent.AddInstance( InstanceTransform );
+			}
+
+			++Index;
+		}
+
+		HismComponent.BuildTreeIfOutdated( true, true );
+
+		return true;
+	}
+
+	bool SetStaticMesh( const pxr::UsdGeomMesh& UsdMesh, UStaticMeshComponent& MeshComponent, const TMap< FString, UObject* >& PrimPathsToAssets )
+	{
+		FScopedUnrealAllocs UnrealAllocs;
+
+		FString MeshPrimPath = UsdToUnreal::ConvertPath( UsdMesh.GetPrim().GetPrimPath() );
+
+		UStaticMesh* StaticMesh = Cast< UStaticMesh >( PrimPathsToAssets.FindRef( MeshPrimPath ) );
+
+		if ( StaticMesh != MeshComponent.GetStaticMesh() )
+		{
+			if ( MeshComponent.IsRegistered() )
+			{
+				MeshComponent.UnregisterComponent();
+			}
+
+			StaticMesh->CreateBodySetup(); // BodySetup is required for HISM component
+			MeshComponent.SetStaticMesh( StaticMesh );
+
+			MeshComponent.RegisterComponent();
+		}
+
+		return ( StaticMesh != nullptr );
+	}
+}
+
+void FUsdGeomPointInstancerTranslator::RegisterTranslator()
+{
+	IUsdSchemasModule& UsdSchemasModule = FModuleManager::Get().LoadModuleChecked< IUsdSchemasModule >( TEXT("USDSchemas") );
+	UsdSchemasModule.GetTranslatorRegistry().Register< FUsdGeomPointInstancerTranslator >( TEXT("UsdGeomPointInstancer") );
+}
+
+USceneComponent* FUsdGeomPointInstancerTranslator::CreateComponents()
+{
+	// Spawn Xform Actor / SceneComponent
+	USceneComponent* PointInstancerRootComponent = FUsdGeomXformableTranslator::CreateComponents();
+
+	if ( !PointInstancerRootComponent )
+	{
+		return nullptr;
+	}
+
+	FScopedUsdAllocs UsdAllocs;
+
+	pxr::UsdGeomPointInstancer PointInstancer( Schema.Get() );
+
+	if ( !PointInstancer )
+	{
+		return nullptr;
+	}
+
+	pxr::UsdPrim Prim = PointInstancer.GetPrim();
+
+	// For each prototype
+	const pxr::UsdRelationship& Prototypes = PointInstancer.GetPrototypesRel();
+	pxr::SdfPathVector PrototypesPaths;
+
+	if ( Prototypes.GetTargets( &PrototypesPaths ) )
+	{
+		TGuardValue< USceneComponent* > ParentComponentGuard( Context->ParentComponent, PointInstancerRootComponent );
+
+		int32 PrototypeIndex = 0;
+		for ( const pxr::SdfPath& PrototypePath : PrototypesPaths )
+		{
+			pxr::UsdPrim PrototypePrim = Prim.GetStage()->GetPrimAtPath( PrototypePath );
+
+			FUsdGeomXformableTranslator PrototypeXformTranslator( Context, pxr::UsdTyped( PrototypePrim ) );
+			USceneComponent* PrototypeXformComponent = PrototypeXformTranslator.CreateComponents();
+
+			TGuardValue< USceneComponent* > ParentComponentGuard2( Context->ParentComponent, PrototypeXformComponent );
+
+			// Find first UsdGeomMesh in prototype childs
+			TArray< TUsdStore< pxr::UsdPrim > > ChildGeomMeshPrims = UsdUtils::GetAllPrimsOfType( PrototypePrim, pxr::TfType::Find< pxr::UsdGeomMesh >() );
+
+			if ( ChildGeomMeshPrims.Num() > 0 )
+			{
+				pxr::UsdPrim PrototypeGeomMeshPrim = ChildGeomMeshPrims[ 0 ].Get();
+				pxr::UsdGeomMesh PrototypeGeomMesh( PrototypeGeomMeshPrim );
+
+				FUsdGeomXformableTranslator PrototypeGeomMeshTranslator( UHierarchicalInstancedStaticMeshComponent::StaticClass(), Context, PrototypeGeomMesh );
+
+				USceneComponent* UsdGeomPrimComponent = PrototypeGeomMeshTranslator.CreateComponents();
+
+				if ( UHierarchicalInstancedStaticMeshComponent* HismComponent = Cast< UHierarchicalInstancedStaticMeshComponent >( UsdGeomPrimComponent ) )
+				{
+					UsdGeomPointInstancerTranslatorImpl::SetStaticMesh( PrototypeGeomMesh, *HismComponent, Context->PrimPathsToAssets );
+					UsdGeomPointInstancerTranslatorImpl::ConvertGeomPointInstancer( Prim.GetStage(), PointInstancer, PrototypeIndex, *HismComponent, pxr::UsdTimeCode( Context->Time ) );
+				}
+			}
+
+			++PrototypeIndex;
+		}
+	}
+
+	return PointInstancerRootComponent;
+}
+
+#endif // #if USE_USD_SDK
diff --git a/Engine/Plugins/Importers/USDImporter/Source/USDSchemas/Private/USDGeomXformableTranslator.cpp b/Engine/Plugins/Importers/USDImporter/Source/USDSchemas/Private/USDGeomXformableTranslator.cpp
new file mode 100644
--- /dev/null
+++ b/Engine/Plugins/Importers/USDImporter/Source/USDSchemas/Private/USDGeomXformableTranslator.cpp
@@ -0,0 +1,145 @@
+// Copyright 1998-2019 Epic Games, Inc. All Rights Reserved.
+
+#include "USDGeomXformableTranslator.h"
+
+#include "USDConversionUtils.h"
+#include "USDPrimConversion.h"
+#include "USDSchemasModule.h"
+#include "USDTypesConversion.h"
+
+#include "Components/SceneComponent.h"
+#include "Engine/World.h"
+#include "Modules/ModuleManager.h"
+
+#if USE_USD_SDK
+
+#include "USDIncludesStart.h"
+	#include "pxr/usd/usd/prim.h"
+	#include "pxr/usd/usdGeom/xformable.h"
+#include "USDIncludesEnd.h"
+
+void FUsdGeomXformableTranslator::RegisterTranslator()
+{
+	IUsdSchemasModule& UsdSchemasModule = FModuleManager::Get().LoadModuleChecked< IUsdSchemasModule >( TEXT("USDSchemas") );
+	UsdSchemasModule.GetTranslatorRegistry().Register< FUsdGeomXformableTranslator >( TEXT("UsdGeomXformable") );
+}
+
+FUsdGeomXformableTranslator::FUsdGeomXformableTranslator( TSubclassOf< USceneComponent > InComponentTypeOverride, TSharedRef< FUsdSchemaTranslationContext > InContext, const pxr::UsdTyped& InSchema )
+	: FUsdSchemaTranslator( InContext, InSchema )
+	, ComponentTypeOverride( InComponentTypeOverride )
+{
+}
+
+USceneComponent* FUsdGeomXformableTranslator::CreateComponents()
+{
+	TSubclassOf< USceneComponent > ComponentType;
+
+	if ( ComponentTypeOverride.IsSet() )
+	{
+		ComponentType = ComponentTypeOverride.GetValue();
+	}
+	else
+	{
+		ComponentType = UsdUtils::GetComponentTypeForPrim( Schema.Get().GetPrim() );
+	}
+
+	return CreateComponents( ComponentType );
+}
+
+USceneComponent* FUsdGeomXformableTranslator::CreateComponents( TSubclassOf< USceneComponent > ComponentType )
+{
+	TUsdStore< pxr::UsdPrim > Prim = Schema.Get().GetPrim();
+	const bool bNeedsActor = ( Context->ParentComponent == nullptr || Prim.Get().IsModel() || UsdUtils::HasCompositionArcs( Prim.Get() ) || Prim.Get().IsGroup() );
+
+	return CreateComponents( ComponentType, bNeedsActor );
+}
+
+USceneComponent* FUsdGeomXformableTranslator::CreateComponents( TSubclassOf< USceneComponent > ComponentType, const bool bNeedsActor )
+{
+	if ( !Context->IsValid() )
+	{
+		return nullptr;
+	}
+
+	TUsdStore< pxr::UsdPrim > Prim = Schema.Get().GetPrim();
+	
+	if ( !Prim.Get() )
+	{
+		return nullptr;
+	}
+
+	FScopedUnrealAllocs UnrealAllocs;
+
+	USceneComponent* SceneComponent = nullptr;
+	UObject* ComponentOuter = nullptr;
+
+	if ( bNeedsActor )
+	{
+		// Spawn actor
+		FActorSpawnParameters SpawnParameters;
+		SpawnParameters.ObjectFlags = Context->ObjectFlags;
+		SpawnParameters.OverrideLevel =  Context->Level;
+
+		UClass* ActorClass = UsdUtils::GetActorTypeForPrim( Prim.Get() );
+		AActor* SpawnedActor = Context->Level->GetWorld()->SpawnActor( ActorClass, nullptr, SpawnParameters );
+
+		if ( SpawnedActor )
+		{
+			SpawnedActor->SetActorLabel( Prim.Get().GetName().GetText() );
+			SpawnedActor->Tags.AddUnique( TEXT("SequencerActor") );	// Hack to show transient actors in world outliner
+
+			SceneComponent = SpawnedActor->GetRootComponent();
+
+			ComponentOuter = SpawnedActor;
+		}
+	}
+	else
+	{
+		ComponentOuter = Context->ParentComponent;
+	}
+	
+	if ( !SceneComponent )
+	{
+		SceneComponent = NewObject< USceneComponent >( ComponentOuter, ComponentType, FName( Prim.Get().GetName().GetText() ), Context->ObjectFlags );
+
+		if ( AActor* Owner = SceneComponent->GetOwner() )
+		{
+			Owner->AddInstanceComponent( SceneComponent );
+		}
+	}
+
+	if ( SceneComponent )
+	{
+		UpdateComponents( SceneComponent );
+
+		if ( Context->ParentComponent && Context->ParentComponent->Mobility == EComponentMobility::Movable )
+		{
+			SceneComponent->SetMobility( EComponentMobility::Movable );
+		}
+		else
+		{
+			SceneComponent->SetMobility( UsdUtils::IsAnimated( Prim.Get() ) ? EComponentMobility::Movable : EComponentMobility::Static ); /*PrimsToAnimate.Contains( UsdToUnreal::ConvertPath( Path.GetPrimPath() ) */
+		}
+
+		// Attach to parent
+		SceneComponent->AttachToComponent( Context->ParentComponent, FAttachmentTransformRules::KeepRelativeTransform );
+
+		if ( !SceneComponent->GetOwner()->GetRootComponent() )
+		{
+			SceneComponent->GetOwner()->SetRootComponent( SceneComponent );
+		}
+	}
+
+	return SceneComponent;
+}
+
+void FUsdGeomXformableTranslator::UpdateComponents( USceneComponent* SceneComponent )
+{
+	if ( SceneComponent )
+	{
+		FScopedUsdAllocs UsdAllocs;
+		UsdToUnreal::ConvertXformable( Schema.Get().GetPrim().GetStage(), pxr::UsdGeomXformable( Schema.Get() ), *SceneComponent, pxr::UsdTimeCode( Context->Time ) );
+	}
+}
+
+#endif // #if USE_USD_SDK
diff --git a/Engine/Plugins/Importers/USDImporter/Source/USDSchemas/Private/USDSchemaTranslator.cpp b/Engine/Plugins/Importers/USDImporter/Source/USDSchemas/Private/USDSchemaTranslator.cpp
new file mode 100644
--- /dev/null
+++ b/Engine/Plugins/Importers/USDImporter/Source/USDSchemas/Private/USDSchemaTranslator.cpp
@@ -0,0 +1,170 @@
+// Copyright 1998-2019 Epic Games, Inc. All Rights Reserved.
+
+#include "USDSchemaTranslator.h"
+
+#include "USDTypesConversion.h"
+
+#include "Async/Async.h"
+
+#if USE_USD_SDK
+
+#include "USDIncludesStart.h"
+	#include "pxr/base/tf/token.h"
+#include "USDIncludesEnd.h"
+
+
+TSharedPtr< FUsdSchemaTranslator > FUsdSchemaTranslatorRegistry::CreateTranslatorForSchema( TSharedRef< FUsdSchemaTranslationContext > InTranslationContext, const pxr::UsdTyped& InSchema )
+{
+	TUsdStore< pxr::UsdPrim > Prim = InSchema.GetPrim();
+
+	for ( TPair< FString, TCreateTranslator >& RegisteredSchema : CreationMethods )
+	{
+		pxr::TfToken RegisteredSchemaToken( UnrealToUsd::ConvertString( *RegisteredSchema.Key ).Get() );
+		pxr::TfType RegisteredSchemaType = pxr::UsdSchemaRegistry::GetTypeFromName( RegisteredSchemaToken );
+
+		if ( Prim.Get().IsA( RegisteredSchemaType ) )
+		{
+			return RegisteredSchema.Value( InTranslationContext, InSchema );
+		}
+	}
+
+	return {};
+}
+
+void FUsdSchemaTranslatorRegistry::Register( const FString& SchemaName, TCreateTranslator CreateFunction )
+{
+	// Insert most specialized first
+	int32 InsertionIndex = 0;
+
+	pxr::TfToken SchemaToRegisterToken( UnrealToUsd::ConvertString( *SchemaName ).Get() );
+	pxr::TfType SchemaToRegisterType = pxr::UsdSchemaRegistry::GetTypeFromName( SchemaToRegisterToken );
+
+	for ( TPair< FString, TCreateTranslator >& RegisteredSchema : CreationMethods )
+	{
+		pxr::TfToken RegisteredSchemaToken( UnrealToUsd::ConvertString( *RegisteredSchema.Key ).Get() );
+		pxr::TfType RegisteredSchemaType = pxr::UsdSchemaRegistry::GetTypeFromName( RegisteredSchemaToken );
+
+		if ( SchemaToRegisterType.IsA( RegisteredSchemaType ) )
+		{
+			// We need to be registered before our ancestor types
+			break;
+		}
+		else
+		{
+			++InsertionIndex;
+		}
+	}
+	
+	CreationMethods.EmplaceAt( InsertionIndex, SchemaName, CreateFunction );
+}
+
+void FUsdSchemaTranslationContext::CompleteTasks()
+{
+	bool bFinished = false;
+	while ( !bFinished )
+	{
+		for ( TArray< TSharedPtr< FUsdSchemaTranslatorTaskChain > >::TIterator TaskChainIterator = TranslatorTasks.CreateIterator(); TaskChainIterator; ++TaskChainIterator )
+		{
+			if ( (*TaskChainIterator)->Execute() == ESchemaTranslationStatus::Done )
+			{
+				TaskChainIterator.RemoveCurrent();
+			}
+		}
+
+		bFinished = ( TranslatorTasks.Num() == 0 );
+	}
+}
+
+void FSchemaTranslatorTask::Start()
+{
+	if ( bAsync )
+	{
+		Result = Async( EAsyncExecution::LargeThreadPool,
+			[ this ]() -> bool
+			{
+				return DoWork();
+			} );
+	}
+	else
+	{
+		// Execute on this thread
+		if ( !DoWork() )
+		{
+			Continuation.Reset();
+		}
+	}
+}
+
+void FSchemaTranslatorTask::StartIfAsync()
+{
+	if ( bAsync )
+	{
+		Start();
+	}
+}
+
+bool FSchemaTranslatorTask::DoWork()
+{
+	ensure( bIsDone == false );
+	bool bContinue = Callable();
+	bIsDone = true;
+
+	return bContinue;
+}
+
+FUsdSchemaTranslatorTaskChain& FUsdSchemaTranslatorTaskChain::Do( bool bAsync, TFunction< bool() > Callable )
+{
+	CurrentTask = MakeShared< FSchemaTranslatorTask >( bAsync, Callable );
+
+	CurrentTask->StartIfAsync(); // Queue it right now if async
+
+	return *this;
+}
+
+FUsdSchemaTranslatorTaskChain& FUsdSchemaTranslatorTaskChain::Then( bool bAsync, TFunction< bool() > Callable )
+{
+	TSharedPtr< FSchemaTranslatorTask > LastTask = CurrentTask;
+
+	while ( LastTask->Continuation.IsValid() )
+	{
+		LastTask = LastTask->Continuation;
+	}
+
+	if ( LastTask )
+	{
+		LastTask->Continuation = MakeShared< FSchemaTranslatorTask >( bAsync, Callable );
+	}
+
+	return *this;
+}
+
+ESchemaTranslationStatus FUsdSchemaTranslatorTaskChain::Execute()
+{
+	if ( !CurrentTask )
+	{
+		return ESchemaTranslationStatus::Done;
+	}
+
+	FSchemaTranslatorTask& TranslatorTask = *CurrentTask;
+
+	if ( !TranslatorTask.IsDone() )
+	{
+		if ( !TranslatorTask.IsStarted() )
+		{
+			TranslatorTask.Start();
+		}
+	}
+	else
+	{
+		CurrentTask = CurrentTask->Continuation;
+
+		if ( CurrentTask )
+		{
+			CurrentTask->StartIfAsync(); // Queue the next task asap if async
+		}
+	}
+
+	return CurrentTask ? ESchemaTranslationStatus::InProgress : ESchemaTranslationStatus::Done;
+}
+
+#endif //#if USE_USD_SDK
diff --git a/Engine/Plugins/Importers/USDImporter/Source/USDSchemas/Private/USDSchemasModule.cpp b/Engine/Plugins/Importers/USDImporter/Source/USDSchemas/Private/USDSchemasModule.cpp
new file mode 100644
--- /dev/null
+++ b/Engine/Plugins/Importers/USDImporter/Source/USDSchemas/Private/USDSchemasModule.cpp
@@ -0,0 +1,39 @@
+// Copyright 1998-2019 Epic Games, Inc. All Rights Reserved.
+
+#include "USDSchemasModule.h"
+
+#if USE_USD_SDK
+#include "USDGeomCameraTranslator.h"
+#include "USDGeomMeshTranslator.h"
+#include "USDGeomPointInstancerTranslator.h"
+#include "USDGeomXFormableTranslator.h"
+#include "USDMemory.h"
+#include "USDSchemaTranslator.h"
+#include "USDSkelRootTranslator.h"
+#endif // #if USE_USD_SDK
+
+class FUsdSchemasModule : public IUsdSchemasModule
+{
+public:
+	virtual void StartupModule() override
+	{
+#if USE_USD_SDK
+		// Register the default translators
+		FUsdGeomCameraTranslator::RegisterTranslator();
+		FUsdGeomMeshTranslator::RegisterTranslator();
+		FUsdGeomPointInstancerTranslator::RegisterTranslator();
+		FUsdSkelRootTranslator::RegisterTranslator();
+		FUsdGeomXformableTranslator::RegisterTranslator();
+#endif // #if USE_USD_SDK
+	}
+
+	virtual FUsdSchemaTranslatorRegistry& GetTranslatorRegistry() override
+	{
+		return UsdSchemaTranslatorRegistry;
+	}
+
+protected:
+	FUsdSchemaTranslatorRegistry UsdSchemaTranslatorRegistry;
+};
+
+IMPLEMENT_MODULE_USD( FUsdSchemasModule, USDSchemas );
diff --git a/Engine/Plugins/Importers/USDImporter/Source/USDSchemas/Private/USDSkelRootTranslator.cpp b/Engine/Plugins/Importers/USDImporter/Source/USDSchemas/Private/USDSkelRootTranslator.cpp
new file mode 100644
--- /dev/null
+++ b/Engine/Plugins/Importers/USDImporter/Source/USDSchemas/Private/USDSkelRootTranslator.cpp
@@ -0,0 +1,289 @@
+// Copyright 1998-2019 Epic Games, Inc. All Rights Reserved.
+
+#include "USDSkelRootTranslator.h"
+
+#if USE_USD_SDK
+
+#include "USDGeomMeshConversion.h"
+#include "USDMemory.h"
+#include "USDSchemasModule.h"
+#include "USDSkeletalDataConversion.h"
+#include "USDTypesConversion.h"
+
+#include "Components/PoseableMeshComponent.h"
+#include "Components/SkinnedMeshComponent.h"
+#include "Engine/SkeletalMesh.h"
+#include "Materials/Material.h"
+#include "Materials/MaterialInstanceConstant.h"
+#include "Modules/ModuleManager.h"
+#include "Rendering/SkeletalMeshLODImporterData.h"
+
+#include "USDIncludesStart.h"
+	#include "pxr/usd/usdGeom/mesh.h"
+	#include "pxr/usd/usdShade/material.h"
+	#include "pxr/usd/usdSkel/binding.h"
+	#include "pxr/usd/usdSkel/cache.h"
+	#include "pxr/usd/usdSkel/root.h"
+	#include "pxr/usd/usdSkel/skeletonQuery.h"
+#include "USDIncludesEnd.h"
+
+
+namespace UsdSkelRootTranslatorImpl
+{
+	void ProcessMaterials( const pxr::UsdPrim& UsdPrim, FSkeletalMeshImportData& SkelMeshImportData, TMap< FString, UObject* >& MaterialsCache, bool bHasPrimDisplayColor, float Time )
+	{
+		for ( SkeletalMeshImportData::FMaterial& ImportedMaterial : SkelMeshImportData.Materials )
+		{
+			if ( !ImportedMaterial.Material.IsValid() )
+			{
+				UMaterialInterface* Material = nullptr;
+				TUsdStore< pxr::UsdPrim > MaterialPrim = UsdPrim.GetStage()->GetPrimAtPath( UnrealToUsd::ConvertPath( *ImportedMaterial.MaterialImportName ).Get() );
+
+				if ( MaterialPrim.Get() )
+				{
+					UObject*& CachedMaterial = MaterialsCache.FindOrAdd( UsdToUnreal::ConvertPath( MaterialPrim.Get().GetPrimPath() ) );
+
+					if ( !CachedMaterial )
+					{
+						UMaterial* NewMaterial = NewObject< UMaterial >();
+
+						if ( UsdToUnreal::ConvertMaterial( pxr::UsdShadeMaterial( MaterialPrim.Get() ), *NewMaterial ) )
+						{
+							NewMaterial->bUsedWithSkeletalMesh = true;
+							bool bNeedsRecompile = false;
+							NewMaterial->GetMaterial()->SetMaterialUsage(bNeedsRecompile, MATUSAGE_SkeletalMesh);
+
+							NewMaterial->PostEditChange();
+						}
+						else
+						{
+							NewMaterial = nullptr;
+						}
+
+						Material = NewMaterial;
+						CachedMaterial = Material;
+					}
+					else
+					{
+						Material = Cast< UMaterialInterface >( CachedMaterial );
+					}
+				}
+
+				if ( Material == nullptr && bHasPrimDisplayColor )
+				{
+					UMaterialInstanceConstant* MaterialInstance = NewObject< UMaterialInstanceConstant >();
+					if ( UsdToUnreal::ConvertDisplayColor( pxr::UsdGeomMesh( UsdPrim ), *MaterialInstance, pxr::UsdTimeCode( Time ) ) )
+					{
+						if ( MaterialInstance && MaterialInstance->HasOverridenBaseProperties() )
+						{
+							MaterialInstance->ForceRecompileForRendering();
+						}
+						
+						Material = MaterialInstance;
+					}
+				}
+
+				ImportedMaterial.Material = Material;
+			}
+		}
+	}
+
+	FSHAHash ComputeSHAHash( const FSkeletalMeshImportData& SkeletalMeshImportData )
+	{
+		FSHA1 HashState;
+
+		// We're only hashing the mesh points for now. Might need to refine that.
+		HashState.Update( (uint8*)SkeletalMeshImportData.Points.GetData(), SkeletalMeshImportData.Points.Num() * SkeletalMeshImportData.Points.GetTypeSize() );
+
+		FSHAHash OutHash;
+
+		HashState.Final();
+		HashState.GetHash( &OutHash.Hash[0] );
+
+		return OutHash;
+	}
+}
+
+void FUsdSkelRootTranslator::RegisterTranslator()
+{
+	IUsdSchemasModule& UsdSchemasModule = FModuleManager::Get().LoadModuleChecked< IUsdSchemasModule >( TEXT("USDSchemas") );
+	UsdSchemasModule.GetTranslatorRegistry().Register< FUsdSkelRootTranslator >( TEXT("UsdSkelRoot") );
+}
+
+void FUsdSkelRootTranslator::CreateAssets()
+{
+	TRACE_CPUPROFILER_EVENT_SCOPE( FUsdSkelRootTranslator::CreateAssets );
+
+	FSkeletalMeshImportData SkelMeshImportData;
+	bool bIsSkeletalDataValid = true;
+
+	// Retrieve the USD skeletal data under the SkeletonRoot into the SkeletalMeshImportData
+	{
+		FScopedUsdAllocs UsdAllocs;
+
+		pxr::UsdPrim Prim = Schema.Get().GetPrim();
+
+		pxr::UsdSkelCache SkeletonCache;
+		pxr::UsdSkelRoot SkeletonRoot( Prim );
+		SkeletonCache.Populate( SkeletonRoot );
+
+		pxr::SdfPath PrimPath = Prim.GetPath();
+
+		std::vector< pxr::UsdSkelBinding > SkeletonBindings;
+		SkeletonCache.ComputeSkelBindings( SkeletonRoot, &SkeletonBindings );
+
+		if ( SkeletonBindings.size() == 0 )
+		{
+			return;
+		}
+
+		const pxr::TfToken StageUpAxis = UsdUtils::GetUsdStageAxis( Prim.GetStage() );
+
+		pxr::UsdGeomXformable Xformable( Prim );
+
+		std::vector< double > TimeSamples;
+
+		bool bHasPrimDisplayColor = false;
+
+		// Note that there could be multiple skeleton bindings under the SkeletonRoot
+		// For now, extract just the first one 
+		for ( const pxr::UsdSkelBinding& Binding : SkeletonBindings )
+		{
+			const pxr::UsdSkelSkeleton& Skeleton = Binding.GetSkeleton();
+			pxr::UsdSkelSkeletonQuery SkelQuery = SkeletonCache.GetSkelQuery( Skeleton );
+
+			const bool bSkeletonValid = UsdToUnreal::ConvertSkeleton( SkelQuery, SkelMeshImportData );
+			if ( !bSkeletonValid )
+			{
+				bIsSkeletalDataValid = false;
+				break;
+			}
+
+			for (const pxr::UsdSkelSkinningQuery& SkinningQuery : Binding.GetSkinningTargets())
+			{
+				// In USD, the skinning target need not be a mesh, but for Unreal we are only interested in skinning meshes
+				pxr::UsdGeomMesh SkinningMesh = pxr::UsdGeomMesh( SkinningQuery.GetPrim() );
+				if (SkinningMesh)
+				{
+					UsdToUnreal::ConvertSkinnedMesh( SkinningQuery, SkelMeshImportData );
+					bHasPrimDisplayColor = bHasPrimDisplayColor || SkinningMesh.GetDisplayColorPrimvar().IsDefined();
+				}
+			}
+
+			UsdSkelRootTranslatorImpl::ProcessMaterials( Prim, SkelMeshImportData, Context->AssetsCache, bHasPrimDisplayColor, Context->Time );
+
+			break;
+		}
+	}
+
+	if ( !bIsSkeletalDataValid )
+	{
+		return;
+	}
+
+	FSHAHash SkeletalMeshHash = UsdSkelRootTranslatorImpl::ComputeSHAHash( SkelMeshImportData );
+
+	USkeletalMesh* SkeletalMesh = Cast< USkeletalMesh >( Context->AssetsCache.FindRef( SkeletalMeshHash.ToString() ) );
+
+	if ( !SkeletalMesh )
+	{
+		SkeletalMesh = UsdToUnreal::GetSkeletalMeshFromImportData( SkelMeshImportData, Context->ObjectFlags );
+		Context->PrimPathsToAssets.Add( UsdToUnreal::ConvertPath( Schema.Get().GetPath() ), SkeletalMesh );
+		Context->AssetsCache.Add( SkeletalMeshHash.ToString(), SkeletalMesh );
+	}
+}
+
+USceneComponent* FUsdSkelRootTranslator::CreateComponents()
+{
+	USceneComponent* RootComponent = FUsdGeomXformableTranslator::CreateComponents();
+
+	if ( USkinnedMeshComponent* SkinnedMeshComponent = Cast< USkinnedMeshComponent >( RootComponent ) )
+	{
+		SkinnedMeshComponent->SetSkeletalMesh( Cast< USkeletalMesh >( Context->PrimPathsToAssets.FindRef( UsdToUnreal::ConvertPath( Schema.Get().GetPath() ) ) ) );
+		UpdateComponents( SkinnedMeshComponent );
+	}
+
+	return RootComponent;
+}
+
+void FUsdSkelRootTranslator::UpdateComponents( USceneComponent* SceneComponent )
+{
+	if ( !SceneComponent )
+	{
+		return;
+	}
+
+	Super::UpdateComponents( SceneComponent );
+
+	{
+		FScopedUsdAllocs UsdAllocs;
+
+		pxr::UsdPrim Prim = Schema.Get().GetPrim();
+
+		pxr::UsdSkelCache SkeletonCache;
+		pxr::UsdSkelRoot SkeletonRoot( Prim );
+		SkeletonCache.Populate( SkeletonRoot );
+
+		pxr::SdfPath PrimPath = Prim.GetPath();
+
+		std::vector< pxr::UsdSkelBinding > SkeletonBindings;
+		SkeletonCache.ComputeSkelBindings( SkeletonRoot, &SkeletonBindings );
+
+		if ( SkeletonBindings.size() == 0 )
+		{
+			return;
+		}
+
+		const pxr::TfToken StageUpAxis = UsdUtils::GetUsdStageAxis( Prim.GetStage() );
+
+		pxr::UsdGeomXformable Xformable( Prim );
+
+		std::vector< double > TimeSamples;
+
+		// Note that there could be multiple skeleton bindings under the SkeletonRoot
+		// For now, extract just the first one 
+		for ( const pxr::UsdSkelBinding& Binding : SkeletonBindings )
+		{
+			const pxr::UsdSkelSkeleton& Skeleton = Binding.GetSkeleton();
+			pxr::UsdSkelSkeletonQuery SkelQuery = SkeletonCache.GetSkelQuery( Skeleton );
+
+			const pxr::UsdSkelAnimQuery& AnimQuery = SkelQuery.GetAnimQuery();
+			const bool bRes = AnimQuery.GetJointTransformTimeSamples( &TimeSamples );
+
+			if ( TimeSamples.size() > 0 )
+			{
+				FScopedUnrealAllocs UnrealAllocs;
+
+				UPoseableMeshComponent* PoseableMeshComponent = Cast< UPoseableMeshComponent >( SceneComponent );
+				if ( !PoseableMeshComponent || !PoseableMeshComponent->SkeletalMesh || !SkelQuery )
+				{
+					return;
+				}
+
+				if ( PoseableMeshComponent->BoneSpaceTransforms.Num() != PoseableMeshComponent->SkeletalMesh->RefSkeleton.GetNum() )
+				{
+					PoseableMeshComponent->AllocateTransformData();
+				}
+
+				TArray< FTransform > BoneTransforms;
+				TUsdStore< pxr::VtArray< pxr::GfMatrix4d > > UsdBoneTransforms;
+
+				const bool bJointTransformsComputed = SkelQuery.ComputeJointLocalTransforms( &UsdBoneTransforms.Get(), pxr::UsdTimeCode( Context->Time ) );
+				if ( bJointTransformsComputed )
+				{
+					BoneTransforms.Reserve( decltype( BoneTransforms )::SizeType( UsdBoneTransforms.Get().size() ) );
+
+					for ( uint32 BoneIndex = 0; BoneIndex < UsdBoneTransforms.Get().size(); ++BoneIndex )
+					{
+						const pxr::GfMatrix4d& UsdMatrix = UsdBoneTransforms.Get()[ BoneIndex ];
+						PoseableMeshComponent->BoneSpaceTransforms[ BoneIndex ] = UsdToUnreal::ConvertMatrix( StageUpAxis, UsdMatrix );
+					}
+				}
+
+				PoseableMeshComponent->RefreshBoneTransforms();
+			}
+		}
+	}
+}
+
+#endif // #if USE_USD_SDK
diff --git a/Engine/Plugins/Importers/USDImporter/Source/USDSchemas/Public/USDGeomCameraTranslator.h b/Engine/Plugins/Importers/USDImporter/Source/USDSchemas/Public/USDGeomCameraTranslator.h
new file mode 100644
--- /dev/null
+++ b/Engine/Plugins/Importers/USDImporter/Source/USDSchemas/Public/USDGeomCameraTranslator.h
@@ -0,0 +1,29 @@
+// Copyright 1998-2019 Epic Games, Inc. All Rights Reserved.
+
+#pragma once
+
+#include "USDGeomXformableTranslator.h"
+
+#if USE_USD_SDK
+
+#include "USDIncludesStart.h"
+	#include "pxr/pxr.h"
+#include "USDIncludesEnd.h"
+
+PXR_NAMESPACE_OPEN_SCOPE
+	class UsdGeomCamera;
+PXR_NAMESPACE_CLOSE_SCOPE
+
+class USDSCHEMAS_API FUsdGeomCameraTranslator : public FUsdGeomXformableTranslator
+{
+public:
+	static void RegisterTranslator();
+
+public:
+	using FUsdGeomXformableTranslator::FUsdGeomXformableTranslator;
+
+	virtual USceneComponent* CreateComponents() override;
+	virtual void UpdateComponents( USceneComponent* SceneComponent ) override;
+};
+
+#endif // #if USE_USD_SDK
diff --git a/Engine/Plugins/Importers/USDImporter/Source/USDSchemas/Public/USDGeomMeshTranslator.h b/Engine/Plugins/Importers/USDImporter/Source/USDSchemas/Public/USDGeomMeshTranslator.h
new file mode 100644
--- /dev/null
+++ b/Engine/Plugins/Importers/USDImporter/Source/USDSchemas/Public/USDGeomMeshTranslator.h
@@ -0,0 +1,60 @@
+// Copyright 1998-2019 Epic Games, Inc. All Rights Reserved.
+
+#pragma once
+
+#include "USDGeomXformableTranslator.h"
+
+#include "MeshDescription.h"
+
+#if USE_USD_SDK
+
+#include "USDIncludesStart.h"
+	#include "pxr/usd/usdGeom/mesh.h"
+#include "USDIncludesEnd.h"
+
+class UStaticMesh;
+
+PXR_NAMESPACE_OPEN_SCOPE
+	class UsdGeomMesh;
+PXR_NAMESPACE_CLOSE_SCOPE
+
+class FGeomMeshCreateAssetsTaskChain : public FUsdSchemaTranslatorTaskChain
+{
+public:
+	FGeomMeshCreateAssetsTaskChain( TSharedRef< FUsdSchemaTranslationContext > InContext, const TUsdStore< pxr::UsdGeomMesh > InGeomMesh )
+		: Context( InContext )
+		, GeomMesh( InGeomMesh )
+	{
+		SetupTasks();
+	}
+
+public:
+	// Inputs
+	TSharedRef< FUsdSchemaTranslationContext > Context;
+	TUsdStore< pxr::UsdGeomMesh > GeomMesh;
+
+	// Outputs
+	FMeshDescription MeshDescription;
+	UStaticMesh* StaticMesh = nullptr;
+
+protected:
+	void SetupTasks();
+};
+
+class USDSCHEMAS_API FUsdGeomMeshTranslator : public FUsdGeomXformableTranslator
+{
+public:
+	static void RegisterTranslator();
+
+public:
+	using FUsdGeomXformableTranslator::FUsdGeomXformableTranslator;
+
+	FUsdGeomMeshTranslator( const FUsdGeomMeshTranslator& Other ) = delete;
+	FUsdGeomMeshTranslator& operator=( const FUsdGeomMeshTranslator& Other ) = delete;
+
+	virtual void CreateAssets() override;
+	virtual USceneComponent* CreateComponents() override;
+
+};
+
+#endif // #if USE_USD_SDK
diff --git a/Engine/Plugins/Importers/USDImporter/Source/USDSchemas/Public/USDGeomPointInstancerTranslator.h b/Engine/Plugins/Importers/USDImporter/Source/USDSchemas/Public/USDGeomPointInstancerTranslator.h
new file mode 100644
--- /dev/null
+++ b/Engine/Plugins/Importers/USDImporter/Source/USDSchemas/Public/USDGeomPointInstancerTranslator.h
@@ -0,0 +1,28 @@
+// Copyright 1998-2019 Epic Games, Inc. All Rights Reserved.
+
+#pragma once
+
+#include "USDGeomXformableTranslator.h"
+
+#if USE_USD_SDK
+
+#include "USDIncludesStart.h"
+	#include "pxr/pxr.h"
+	#include "pxr/usd/usdGeom/pointInstancer.h"
+#include "USDIncludesEnd.h"
+
+struct FUsdSchemaTranslationContext;
+
+class USDSCHEMAS_API FUsdGeomPointInstancerTranslator : public FUsdGeomXformableTranslator
+{
+public:
+	static void RegisterTranslator();
+
+public:
+	using FUsdGeomXformableTranslator::FUsdGeomXformableTranslator;
+	
+	virtual USceneComponent* CreateComponents() override;
+	virtual bool CollapsedHierarchy() const override { return true; }
+};
+
+#endif // #if USE_USD_SDK
diff --git a/Engine/Plugins/Importers/USDImporter/Source/USDSchemas/Public/USDGeomXformableTranslator.h b/Engine/Plugins/Importers/USDImporter/Source/USDSchemas/Public/USDGeomXformableTranslator.h
new file mode 100644
--- /dev/null
+++ b/Engine/Plugins/Importers/USDImporter/Source/USDSchemas/Public/USDGeomXformableTranslator.h
@@ -0,0 +1,40 @@
+// Copyright 1998-2019 Epic Games, Inc. All Rights Reserved.
+
+#pragma once
+
+#include "USDSchemaTranslator.h"
+
+#if USE_USD_SDK
+
+#include "USDIncludesStart.h"
+	#include "pxr/pxr.h"
+#include "USDIncludesEnd.h"
+
+class USceneComponent;
+
+PXR_NAMESPACE_OPEN_SCOPE
+	class UsdGeomXformable;
+PXR_NAMESPACE_CLOSE_SCOPE
+
+class USDSCHEMAS_API FUsdGeomXformableTranslator : public FUsdSchemaTranslator
+{
+public:
+	static void RegisterTranslator();
+
+public:
+	using FUsdSchemaTranslator::FUsdSchemaTranslator;
+
+	explicit FUsdGeomXformableTranslator( TSubclassOf< USceneComponent > InComponentTypeOverride, TSharedRef< FUsdSchemaTranslationContext > InContext, const pxr::UsdTyped& InSchema );
+
+	virtual USceneComponent* CreateComponents() override;
+	virtual void UpdateComponents( USceneComponent* SceneComponent );
+
+protected:
+	USceneComponent* CreateComponents( TSubclassOf< USceneComponent > ComponentType );
+	USceneComponent* CreateComponents( TSubclassOf< USceneComponent > ComponentType, const bool bNeedsActor );
+
+private:
+	TOptional< TSubclassOf< USceneComponent > > ComponentTypeOverride;
+};
+
+#endif // #if USE_USD_SDK
diff --git a/Engine/Plugins/Importers/USDImporter/Source/USDSchemas/Public/USDSchemaTranslator.h b/Engine/Plugins/Importers/USDImporter/Source/USDSchemas/Public/USDSchemaTranslator.h
new file mode 100644
--- /dev/null
+++ b/Engine/Plugins/Importers/USDImporter/Source/USDSchemas/Public/USDSchemaTranslator.h
@@ -0,0 +1,176 @@
+// Copyright 1998-2019 Epic Games, Inc. All Rights Reserved.
+
+#pragma once
+
+#include "CoreMinimal.h"
+
+#include "USDMemory.h"
+
+#include "Async/Future.h"
+#include "HAL/ThreadSafeBool.h"
+#include "Misc/Optional.h"
+#include "Templates/SubclassOf.h"
+#include "UObject/ObjectMacros.h"
+
+#if USE_USD_SDK
+
+#include "USDIncludesStart.h"
+	#include "pxr/pxr.h"
+	#include "pxr/usd/usd/typed.h"
+#include "USDIncludesEnd.h"
+
+struct FUsdSchemaTranslationContext;
+class FUsdSchemaTranslator;
+class FUsdSchemaTranslatorTaskChain;
+class ULevel;
+class USceneComponent;
+class UStaticMesh;
+
+PXR_NAMESPACE_OPEN_SCOPE
+	class TfToken;
+	class UsdGeomPointInstancer;
+	class UsdGeomXformable;
+PXR_NAMESPACE_CLOSE_SCOPE
+
+#endif //#if USE_USD_SDK
+
+class USDSCHEMAS_API FUsdSchemaTranslatorRegistry
+{
+#if USE_USD_SDK
+	using TCreateTranslator = TFunction< TSharedRef< FUsdSchemaTranslator >( TSharedRef< FUsdSchemaTranslationContext >, const pxr::UsdTyped& ) >;
+
+public:
+	TSharedPtr< FUsdSchemaTranslator > CreateTranslatorForSchema( TSharedRef< FUsdSchemaTranslationContext > InTranslationContext, const pxr::UsdTyped& InSchema );
+
+	template< typename TSchemaTranslator >
+	void Register( const FString& SchemaName )
+	{
+		auto CreateSchemaTranslator =
+		[]( TSharedRef< FUsdSchemaTranslationContext > InContext, const pxr::UsdTyped& InSchema ) -> TSharedRef< FUsdSchemaTranslator >
+		{
+			return MakeShared< TSchemaTranslator >( InContext, InSchema );
+		};
+
+		Register( SchemaName, CreateSchemaTranslator );
+	}
+
+protected:
+	void Register( const FString& SchemaName, TCreateTranslator CreateFunction );
+
+	TArray< TPair< FString, TCreateTranslator > > CreationMethods;
+#endif // #if USE_USD_SDK
+};
+
+#if USE_USD_SDK
+
+struct USDSCHEMAS_API FUsdSchemaTranslationContext : public TSharedFromThis< FUsdSchemaTranslationContext >
+{
+	explicit FUsdSchemaTranslationContext( TMap< FString, UObject* >& InPrimPathsToAssets, TMap< FString, UObject* >& InAssetsCache )
+		: PrimPathsToAssets( InPrimPathsToAssets )
+		, AssetsCache( InAssetsCache )
+	{
+	}
+
+	/** Level to spawn actors in */
+	ULevel* Level = nullptr;
+
+	/** Flags used when creating UObjects */
+	EObjectFlags ObjectFlags;
+
+	/** The parent component when translating children */
+	USceneComponent* ParentComponent = nullptr;
+
+	/** The time at which we are translating */
+	float Time = 0.f;
+
+	/** Map of translated UsdPrims to UAssets */
+	TMap< FString, UObject* >& PrimPathsToAssets;
+
+	TMap< FString, UObject* >& AssetsCache;
+
+	FCriticalSection CriticalSection;
+
+	bool IsValid() const
+	{
+		return Level != nullptr;
+	}
+
+	void CompleteTasks();
+
+	TArray< TSharedPtr< FUsdSchemaTranslatorTaskChain > > TranslatorTasks;
+};
+
+enum class ESchemaTranslationStatus
+{
+	InProgress,
+	Done
+};
+
+class FUsdSchemaTranslator
+{
+public:
+	explicit FUsdSchemaTranslator( TSharedRef< FUsdSchemaTranslationContext > InContext, const pxr::UsdTyped& InSchema )
+		: Context( InContext )
+		, Schema( InSchema )
+	{
+	}
+
+	virtual ~FUsdSchemaTranslator() = default;
+
+	virtual void CreateAssets() {}
+
+	virtual USceneComponent* CreateComponents() { return nullptr; }
+	virtual void UpdateComponents( USceneComponent* SceneComponent ) {}
+
+	virtual bool CollapsedHierarchy() const { return false; }
+
+protected:
+	TSharedRef< FUsdSchemaTranslationContext > Context;
+	TUsdStore< pxr::UsdTyped > Schema;
+};
+
+struct FSchemaTranslatorTask
+{
+	explicit FSchemaTranslatorTask( bool bInAsync, TFunction< bool() > InCallable )
+		: Callable( InCallable )
+		, bAsync( bInAsync )
+		, bIsDone( false )
+	{
+	}
+
+	TFunction< bool() > Callable;
+	TOptional< TFuture< bool > > Result;
+	TSharedPtr< FSchemaTranslatorTask > Continuation;
+
+	bool bAsync = false;
+	FThreadSafeBool bIsDone;
+
+	void Start();
+	void StartIfAsync();
+
+	bool IsStarted() const
+	{
+		return Result.IsSet() && Result->IsValid();
+	}
+
+	bool DoWork();
+
+	bool IsDone() const
+	{
+		return bIsDone;
+	}
+};
+
+class FUsdSchemaTranslatorTaskChain
+{
+public:
+	FUsdSchemaTranslatorTaskChain& Do( bool bAsync, TFunction< bool() > Callable );
+	FUsdSchemaTranslatorTaskChain& Then( bool bAsync, TFunction< bool() > Callable );
+
+	ESchemaTranslationStatus Execute();
+
+private:
+	TSharedPtr< FSchemaTranslatorTask > CurrentTask;
+};
+
+#endif //#if USE_USD_SDK
diff --git a/Engine/Plugins/Importers/USDImporter/Source/USDSchemas/Public/USDSchemasModule.h b/Engine/Plugins/Importers/USDImporter/Source/USDSchemas/Public/USDSchemasModule.h
new file mode 100644
--- /dev/null
+++ b/Engine/Plugins/Importers/USDImporter/Source/USDSchemas/Public/USDSchemasModule.h
@@ -0,0 +1,13 @@
+// Copyright 1998-2019 Epic Games, Inc. All Rights Reserved.
+
+#pragma once
+
+#include "Modules/ModuleInterface.h"
+
+class FUsdSchemaTranslatorRegistry;
+
+class IUsdSchemasModule : public IModuleInterface
+{
+public:
+	virtual FUsdSchemaTranslatorRegistry& GetTranslatorRegistry() = 0;
+};
diff --git a/Engine/Plugins/Importers/USDImporter/Source/USDSchemas/Public/USDSkelRootTranslator.h b/Engine/Plugins/Importers/USDImporter/Source/USDSchemas/Public/USDSkelRootTranslator.h
new file mode 100644
--- /dev/null
+++ b/Engine/Plugins/Importers/USDImporter/Source/USDSchemas/Public/USDSkelRootTranslator.h
@@ -0,0 +1,37 @@
+// Copyright 1998-2019 Epic Games, Inc. All Rights Reserved.
+
+#pragma once
+
+#include "USDGeomMeshTranslator.h"
+
+#if USE_USD_SDK
+
+#include "USDIncludesStart.h"
+	#include "pxr/pxr.h"
+#include "USDIncludesEnd.h"
+
+PXR_NAMESPACE_OPEN_SCOPE
+	class UsdSkelRoot;
+PXR_NAMESPACE_CLOSE_SCOPE
+
+struct FUsdSchemaTranslationContext;
+
+class USDSCHEMAS_API FUsdSkelRootTranslator : public FUsdGeomXformableTranslator
+{
+	using Super = FUsdGeomXformableTranslator;
+
+public:
+	static void RegisterTranslator();
+
+public:
+	using FUsdGeomXformableTranslator::FUsdGeomXformableTranslator;
+	//explicit FUsdSkelRootTranslator( FUsdSchemaTranslationContext& InContext, const pxr::UsdSkelRoot& SkelRoot );
+	
+	virtual void CreateAssets() override;
+	virtual USceneComponent* CreateComponents() override;
+	virtual void UpdateComponents( USceneComponent* SceneComponent ) override;
+
+	virtual bool CollapsedHierarchy() const override { return true; }
+};
+
+#endif //#if USE_USD_SDK
diff --git a/Engine/Plugins/Importers/USDImporter/Source/USDSchemas/USDSchemas.Build.cs b/Engine/Plugins/Importers/USDImporter/Source/USDSchemas/USDSchemas.Build.cs
new file mode 100644
--- /dev/null
+++ b/Engine/Plugins/Importers/USDImporter/Source/USDSchemas/USDSchemas.Build.cs
@@ -0,0 +1,43 @@
+// Copyright 1998-2019 Epic Games, Inc. All Rights Reserved.
+
+using System;
+using System.IO;
+
+namespace UnrealBuildTool.Rules
+{
+	public class USDSchemas : ModuleRules
+	{
+		public USDSchemas(ReadOnlyTargetRules Target) : base(Target)
+		{
+			bUseRTTI = true;
+
+			PublicDependencyModuleNames.AddRange(
+				new string[]
+				{
+					"ModelingComponents",
+				}
+				);
+
+			PrivateDependencyModuleNames.AddRange(
+				new string[]
+				{
+					"CinematicCamera",
+					"Core",
+					"CoreUObject",
+					"EditorStyle",
+					"Engine",
+					"MaterialEditor",
+					"MeshDescription",
+					"MeshDescriptionOperations",
+					"MeshUtilities",
+					"MessageLog",
+					"PropertyEditor",
+					"StaticMeshDescription",
+					"UnrealUSDWrapper",
+					"USDUtilities",
+					"UnrealEd",
+				}
+				);
+		}
+	}
+}
diff --git a/Engine/Plugins/Importers/USDImporter/Source/USDStage/Private/USDPrimTwin.cpp b/Engine/Plugins/Importers/USDImporter/Source/USDStage/Private/USDPrimTwin.cpp
--- a/Engine/Plugins/Importers/USDImporter/Source/USDStage/Private/USDPrimTwin.cpp
+++ b/Engine/Plugins/Importers/USDImporter/Source/USDStage/Private/USDPrimTwin.cpp
@@ -48,9 +48,19 @@ void FUsdPrimTwin::Clear()
 		OnDestroyed.Broadcast( *this );
 	}
 
-	if ( SpawnedActor.IsValid() && !SpawnedActor->IsA<AUsdStageActor>() && !SpawnedActor->IsActorBeingDestroyed() )
+	AActor* ActorToDestroy = SpawnedActor.Get();
+
+	if ( !ActorToDestroy )
+	{
+		if ( SceneComponent.IsValid() && SceneComponent->GetOwner()->GetRootComponent() == SceneComponent.Get() )
+		{
+			ActorToDestroy = SceneComponent->GetOwner();
+		}
+	}
+
+	if ( ActorToDestroy && !ActorToDestroy->IsA< AUsdStageActor >() && !ActorToDestroy->IsActorBeingDestroyed() && ActorToDestroy->GetWorld() )
 	{
-		SpawnedActor->GetWorld()->DestroyActor( SpawnedActor.Get() );
+		ActorToDestroy->GetWorld()->DestroyActor( ActorToDestroy );
 		SpawnedActor = nullptr;
 	}
 	else if ( SceneComponent.IsValid() && !SceneComponent->IsBeingDestroyed() )
diff --git a/Engine/Plugins/Importers/USDImporter/Source/USDStage/Private/USDStageActor.cpp b/Engine/Plugins/Importers/USDImporter/Source/USDStage/Private/USDStageActor.cpp
--- a/Engine/Plugins/Importers/USDImporter/Source/USDStage/Private/USDStageActor.cpp
+++ b/Engine/Plugins/Importers/USDImporter/Source/USDStage/Private/USDStageActor.cpp
@@ -3,16 +3,17 @@
 #include "USDStageActor.h"
 
 #include "USDConversionUtils.h"
-#include "USDGeomMeshConversion.h"
+#include "USDGeomMeshTranslator.h"
+#include "USDGeomXFormableTranslator.h"
 #include "USDListener.h"
 #include "USDPrimConversion.h"
-#include "USDSkeletalDataConversion.h"
+#include "USDSchemasModule.h"
+#include "USDSchemaTranslator.h"
+#include "USDSkelRootTranslator.h"
 #include "USDTypesConversion.h"
 #include "UnrealUSDWrapper.h"
 
 #include "Async/ParallelFor.h"
-#include "CineCameraActor.h"
-#include "CineCameraComponent.h"
 #include "Components/PoseableMeshComponent.h"
 #include "Components/StaticMeshComponent.h"
 #include "Editor.h"
@@ -21,9 +22,6 @@
 #include "Engine/StaticMesh.h"
 #include "Engine/World.h"
 #include "HAL/FileManager.h"
-#include "IMeshBuilderModule.h"
-#include "Interfaces/ITargetPlatform.h"
-#include "Interfaces/ITargetPlatformManagerModule.h"
 #include "MaterialEditingLibrary.h"
 #include "Materials/Material.h"
 #include "Materials/MaterialInterface.h"
@@ -40,29 +38,26 @@
 #include "LevelSequence.h"
 #include "LevelSequenceActor.h"
 #include "Misc/FrameNumber.h"
+#include "Modules/ModuleManager.h"
 #include "MovieScene.h"
 #include "Sections/MovieSceneFloatSection.h"
 #include "Tracks/MovieSceneFloatTrack.h"
 
 #if USE_USD_SDK
 #include "USDIncludesStart.h"
-
-#include "pxr/usd/sdf/layer.h"
-#include "pxr/usd/usd/modelAPI.h"
-#include "pxr/usd/usd/stageCache.h"
-#include "pxr/usd/usd/stageCacheContext.h"
-#include "pxr/usd/usdGeom/camera.h"
-#include "pxr/usd/usdGeom/mesh.h"
-#include "pxr/usd/usdGeom/metrics.h"
-#include "pxr/usd/usdGeom/scope.h"
-#include "pxr/usd/usdGeom/xform.h"
-#include "pxr/usd/usdGeom/xformCommonAPI.h"
-#include "pxr/usd/usdShade/material.h"
-#include "pxr/usd/usdSkel/binding.h"
-#include "pxr/usd/usdSkel/cache.h"
-#include "pxr/usd/usdSkel/root.h"
-#include "pxr/usd/usdSkel/skeletonQuery.h"
-
+	#include "pxr/usd/sdf/layer.h"
+	#include "pxr/usd/usd/modelAPI.h"
+	#include "pxr/usd/usd/stageCache.h"
+	#include "pxr/usd/usd/stageCacheContext.h"
+	#include "pxr/usd/usdGeom/camera.h"
+	#include "pxr/usd/usdGeom/mesh.h"
+	#include "pxr/usd/usdGeom/metrics.h"
+	#include "pxr/usd/usdGeom/pointInstancer.h"
+	#include "pxr/usd/usdGeom/scope.h"
+	#include "pxr/usd/usdGeom/xform.h"
+	#include "pxr/usd/usdGeom/xformCommonAPI.h"
+	#include "pxr/usd/usdShade/material.h"
+	#include "pxr/usd/usdSkel/root.h"
 #include "USDIncludesEnd.h"
 #endif // #if USE_USD_SDK
 
@@ -76,155 +71,34 @@ AUsdStageActor::FOnActorLoaded AUsdStageActor::OnActorLoaded{};
 
 #if USE_USD_SDK
 
-namespace UsdStageActorImpl
+struct UsdStageActorImpl
 {
-	FMeshDescription LoadMeshDescription( const pxr::UsdGeomMesh& UsdMesh, const pxr::UsdTimeCode TimeCode )
+	static TSharedRef< FUsdSchemaTranslationContext > CreateUsdSchemaTranslationContext( AUsdStageActor* StageActor, const FString& PrimPath )
 	{
-		if ( !UsdMesh )
-		{
-			return {};
-		}
-
-		FMeshDescription MeshDescription;
-		FStaticMeshAttributes StaticMeshAttributes( MeshDescription );
-		StaticMeshAttributes.Register();
-
-		UsdToUnreal::ConvertGeomMesh( UsdMesh, MeshDescription, TimeCode );
+		TSharedRef< FUsdSchemaTranslationContext > TranslationContext = MakeShared< FUsdSchemaTranslationContext >( StageActor->PrimPathsToAssets, StageActor->AssetsCache );
+		TranslationContext->Level = StageActor->GetLevel();
+		TranslationContext->ObjectFlags = DefaultObjFlag;
+		TranslationContext->Time = StageActor->GetTime();
 
-		return MeshDescription;
-	}
-
-	void ProcessMaterials( const pxr::UsdPrim& UsdPrim, UStaticMesh& StaticMesh, TMap< FString, UMaterial* >& MaterialsCache, bool bHasPrimDisplayColor, float Time )
-	{
-		const FMeshDescription* MeshDescription = StaticMesh.GetMeshDescription( 0 );
+		TUsdStore< pxr::SdfPath > UsdPrimPath = UnrealToUsd::ConvertPath( *PrimPath );
+		FUsdPrimTwin* ParentUsdPrimTwin = StageActor->RootUsdTwin.Find( UsdToUnreal::ConvertPath( UsdPrimPath.Get().GetParentPath() ) );
 
-		if ( !MeshDescription )
+		if ( !ParentUsdPrimTwin )
 		{
-			return;
+			ParentUsdPrimTwin = &StageActor->RootUsdTwin;
 		}
 
-		FStaticMeshConstAttributes StaticMeshAttributes( *MeshDescription );
+		TranslationContext->ParentComponent = ParentUsdPrimTwin ? ParentUsdPrimTwin->SceneComponent.Get() : nullptr;
 
-		for ( const FPolygonGroupID PolygonGroupID : MeshDescription->PolygonGroups().GetElementIDs() )
+		if ( !TranslationContext->ParentComponent )
 		{
-			const FName& ImportedMaterialSlotName = StaticMeshAttributes.GetPolygonGroupMaterialSlotNames()[ PolygonGroupID ];
-			const FName MaterialSlotName = ImportedMaterialSlotName;
-
-			int32 MaterialIndex = INDEX_NONE;
-			int32 MeshMaterialIndex = 0;
-
-			for ( FStaticMaterial& StaticMaterial : StaticMesh.StaticMaterials )
-			{
-				if ( StaticMaterial.MaterialSlotName.IsEqual( ImportedMaterialSlotName ) )
-				{
-					MaterialIndex = MeshMaterialIndex;
-					break;
-				}
-
-				++MeshMaterialIndex;
-			}
-
-			if ( MaterialIndex == INDEX_NONE )
-			{
-				MaterialIndex = PolygonGroupID.GetValue();
-			}
-
-			UMaterialInterface* Material = nullptr;
-			TUsdStore< pxr::UsdPrim > MaterialPrim = UsdPrim.GetStage()->GetPrimAtPath( UnrealToUsd::ConvertPath( *ImportedMaterialSlotName.ToString() ).Get() );
-
-			if ( MaterialPrim.Get() )
-			{
-				UMaterial*& CachedMaterial = MaterialsCache.FindOrAdd( UsdToUnreal::ConvertPath( MaterialPrim.Get().GetPrimPath() ) );
-
-				if ( !CachedMaterial )
-				{
-					CachedMaterial = NewObject< UMaterial >();
-
-					if ( UsdToUnreal::ConvertMaterial( pxr::UsdShadeMaterial( MaterialPrim.Get() ), *CachedMaterial ) )
-					{
-						//UMaterialEditingLibrary::RecompileMaterial( CachedMaterial ); // Too slow
-						CachedMaterial->PostEditChange();
-					}
-					else
-					{
-						CachedMaterial = nullptr;
-					}
-				}
-
-				Material = CachedMaterial;
-			}
-
-			if ( Material == nullptr && bHasPrimDisplayColor )
-			{
-				UMaterialInstanceConstant* MaterialInstance = NewObject< UMaterialInstanceConstant >();
-				if ( UsdToUnreal::ConvertDisplayColor( pxr::UsdGeomMesh( UsdPrim ), *MaterialInstance, pxr::UsdTimeCode( Time ) ) )
-				{
-					Material = MaterialInstance;
-				}
-			}
-
-			FStaticMaterial StaticMaterial( Material, MaterialSlotName, ImportedMaterialSlotName );
-				StaticMesh.StaticMaterials.Add( StaticMaterial );
-
-			FMeshSectionInfo MeshSectionInfo;
-			MeshSectionInfo.MaterialIndex = MaterialIndex;
-
-			StaticMesh.GetSectionInfoMap().Set( 0, PolygonGroupID.GetValue(), MeshSectionInfo );
+			TranslationContext->ParentComponent = StageActor->RootComponent;
 		}
-	}
-
-	void ProcessMaterials( const pxr::UsdPrim& UsdPrim, FSkeletalMeshImportData& SkelMeshImportData, TMap< FString, UMaterial* >& MaterialsCache, bool bHasPrimDisplayColor, float Time )
-	{
-		for (SkeletalMeshImportData::FMaterial& ImportedMaterial : SkelMeshImportData.Materials)
-		{
-			if (!ImportedMaterial.Material.IsValid())
-			{
-				UMaterialInterface* Material = nullptr;
-				TUsdStore< pxr::UsdPrim > MaterialPrim = UsdPrim.GetStage()->GetPrimAtPath( UnrealToUsd::ConvertPath( *ImportedMaterial.MaterialImportName ).Get() );
 
-				if ( MaterialPrim.Get() )
-				{
-					UMaterial*& CachedMaterial = MaterialsCache.FindOrAdd( UsdToUnreal::ConvertPath( MaterialPrim.Get().GetPrimPath() ) );
-
-					if ( !CachedMaterial )
-					{
-						CachedMaterial = NewObject< UMaterial >();
-
-						if ( UsdToUnreal::ConvertMaterial( pxr::UsdShadeMaterial( MaterialPrim.Get() ), *CachedMaterial ) )
-						{
-							CachedMaterial->bUsedWithSkeletalMesh = true;
-							bool bNeedsRecompile = false;
-							CachedMaterial->GetMaterial()->SetMaterialUsage(bNeedsRecompile, MATUSAGE_SkeletalMesh);
-
-							CachedMaterial->PostEditChange();
-						}
-						else
-						{
-							CachedMaterial = nullptr;
-						}
-					}
-					Material = CachedMaterial;
-				}
-
-				if ( Material == nullptr && bHasPrimDisplayColor )
-				{
-					UMaterialInstanceConstant* MaterialInstance = NewObject< UMaterialInstanceConstant >();
-					if ( UsdToUnreal::ConvertDisplayColor( pxr::UsdGeomMesh( UsdPrim ), *MaterialInstance, pxr::UsdTimeCode( Time ) ) )
-					{
-						if ( MaterialInstance && MaterialInstance->HasOverridenBaseProperties() )
-						{
-							MaterialInstance->ForceRecompileForRendering();
-						}
-						
-						Material = MaterialInstance;
-					}
-				}
-
-				ImportedMaterial.Material = Material;
-			}
-		}
+		return TranslationContext;
 	}
-}
+};
+
 #endif // #if USE_USD_SDK
 
 AUsdStageActor::AUsdStageActor()
@@ -235,20 +109,29 @@ AUsdStageActor::AUsdStageActor()
 	, TimeCodesPerSecond( 24.f )
 {
 	SceneComponent = CreateDefaultSubobject< USceneComponent >( TEXT("SceneComponent0") );
-	
+	SceneComponent->Mobility = EComponentMobility::Static;
+
 	RootComponent = SceneComponent;
 
 #if USE_USD_SDK
+	RootUsdTwin.PrimPath = TEXT("/");
+
 	UsdListener.OnPrimChanged.AddLambda(
 		[ this ]( const FString& PrimPath, bool bResync )
 			{
 				TUsdStore< pxr::SdfPath > UsdPrimPath = UnrealToUsd::ConvertPath( *PrimPath );
 
-				TMap< FString, UStaticMesh* > PrimPathsToStaticMeshes = this->LoadStaticMeshes( this->GetUsdStage()->GetPrimAtPath( UsdPrimPath.Get() ) );
+				TSharedRef< FUsdSchemaTranslationContext > TranslationContext = UsdStageActorImpl::CreateUsdSchemaTranslationContext( this, PrimPath );
 
-				this->UpdatePrim( UsdPrimPath.Get(), bResync, PrimPathsToStaticMeshes );
+				this->LoadAssets( *TranslationContext, this->GetUsdStage()->GetPrimAtPath( UsdPrimPath.Get() ) );
 
-				this->OnPrimChanged.Broadcast( PrimPath, bResync );
+				this->UpdatePrim( UsdPrimPath.Get(), bResync, *TranslationContext );
+				TranslationContext->CompleteTasks();
+
+				if ( this->HasAutorithyOverStage() )
+				{
+					this->OnPrimChanged.Broadcast( PrimPath, bResync );
+				}
 			}
 		);
 #endif // #if USE_USD_SDK
@@ -256,40 +139,25 @@ AUsdStageActor::AUsdStageActor()
 	InitLevelSequence(TimeCodesPerSecond);
 	SetupLevelSequence();
 
-	FCoreUObjectDelegates::OnObjectPropertyChanged.AddUObject( this, &AUsdStageActor::OnPrimObjectPropertyChanged );
+	if ( HasAutorithyOverStage() )
+	{
+		FCoreUObjectDelegates::OnObjectPropertyChanged.AddUObject( this, &AUsdStageActor::OnPrimObjectPropertyChanged );
+	}
 }
 
 AUsdStageActor::~AUsdStageActor()
 {
 #if USE_USD_SDK
-	UnrealUSDWrapper::GetUsdStageCache().Erase( UsdStageStore.Get() );
+	if ( HasAutorithyOverStage() )
+	{
+		UnrealUSDWrapper::GetUsdStageCache().Erase( UsdStageStore.Get() );
+	}
 #endif // #if USE_USD_SDK
 }
 
 #if USE_USD_SDK
 
-bool AUsdStageActor::SetStaticMesh( const pxr::UsdGeomMesh& UsdMesh, UStaticMeshComponent& MeshComponent, const TMap< FString, UStaticMesh* >& PrimPathsToStaticMeshes )
-{
-	FString MeshPrimPath = UsdToUnreal::ConvertPath( UsdMesh.GetPrim().GetPrimPath() );
-
-	UStaticMesh* StaticMesh = PrimPathsToStaticMeshes.FindRef( MeshPrimPath );
-
-	if ( StaticMesh != MeshComponent.GetStaticMesh() )
-	{
-		if ( MeshComponent.IsRegistered() )
-		{
-			MeshComponent.UnregisterComponent();
-		}
-
-		MeshComponent.SetStaticMesh( StaticMesh );
-
-		MeshComponent.RegisterComponent();
-	}
-
-	return ( StaticMesh != nullptr );
-}
-
-FUsdPrimTwin* AUsdStageActor::SpawnPrim( const pxr::SdfPath& UsdPrimPath )
+FUsdPrimTwin* AUsdStageActor::GetOrCreatePrimTwin( const pxr::SdfPath& UsdPrimPath )
 {
 	const FString PrimPath = UsdToUnreal::ConvertPath( UsdPrimPath );
 	const FString ParentPrimPath = UsdToUnreal::ConvertPath( UsdPrimPath.GetParentPath() );
@@ -320,391 +188,75 @@ FUsdPrimTwin* AUsdStageActor::SpawnPrim( const pxr::SdfPath& UsdPrimPath )
 			} );
 
 		if ( !UsdPrimTwin->AnimationHandle.IsValid() && UsdUtils::IsAnimated( Prim ) )
-			{
-				PrimsToAnimate.Add( PrimPath );
-			}
-		}
-
-	UClass* ComponentClass = UsdUtils::GetComponentTypeForPrim( Prim );
-
-	if ( !ComponentClass )
-	{
-		return nullptr;
-	}
-
-	FScopedUnrealAllocs UnrealAllocs;
-
-	const bool bNeedsActor = ( ParentUsdPrimTwin->SceneComponent == nullptr || Prim.IsModel() || UsdUtils::HasCompositionArcs( Prim ) || Prim.IsGroup() || Prim.IsA< pxr::UsdGeomCamera >() );
-
-	if ( bNeedsActor && !UsdPrimTwin->SpawnedActor.IsValid() )
-	{
-		// Try to find an existing actor for the new prim
-		AActor* ParentActor = ParentUsdPrimTwin->SpawnedActor.Get();
-
-		if ( !ParentActor )
 		{
-			ParentActor = this;
-		}
-
-		if ( ParentActor )
-		{
-			TArray< AActor* > AttachedActors;
-			ParentActor->GetAttachedActors( AttachedActors );
-
-			FString PrimName = UsdToUnreal::ConvertString( Prim.GetName().GetText() );
-
-			for ( AActor* AttachedActor : AttachedActors )
-			{
-				if ( AttachedActor->GetName() == PrimName )
-				{
-					// Assuming that this actor isn't used by another prim because USD names are unique
-					UsdPrimTwin->SpawnedActor = AttachedActor;
-					break;
-				}
-			}
-		}
-
-		if ( !UsdPrimTwin->SpawnedActor.IsValid() )
-		{
-			// Spawn actor
-			FActorSpawnParameters SpawnParameters;
-			SpawnParameters.ObjectFlags = DefaultObjFlag;
-			SpawnParameters.OverrideLevel = GetLevel();
-
-			UClass* ActorClass = UsdUtils::GetActorTypeForPrim( Prim );
-			AActor* SpawnedActor = GetWorld()->SpawnActor( ActorClass, nullptr, SpawnParameters );
-			UsdPrimTwin->SpawnedActor = SpawnedActor;
-		}
-
-		UsdPrimTwin->SpawnedActor->SetActorLabel( Prim.GetName().GetText() );
-		UsdPrimTwin->SpawnedActor->Tags.AddUnique( TEXT("SequencerActor") );	// Hack to show transient actors in world outliner
-
-		if ( UActorComponent* ExistingComponent = UsdPrimTwin->SpawnedActor->GetComponentByClass( ComponentClass ) )
-		{
-			UsdPrimTwin->SceneComponent = Cast< USceneComponent >( ExistingComponent );
-		}
-		else
-		{
-			UsdPrimTwin->SceneComponent = NewObject< USceneComponent >( UsdPrimTwin->SpawnedActor.Get(), ComponentClass, FName( Prim.GetName().GetText() ), DefaultObjFlag );
-			UsdPrimTwin->SpawnedActor->SetRootComponent( UsdPrimTwin->SceneComponent.Get() );
-		}
-	}
-	else if ( !bNeedsActor && UsdPrimTwin->SpawnedActor.IsValid() )
-	{
-		GetWorld()->DestroyActor( UsdPrimTwin->SpawnedActor.Get() );
-		UsdPrimTwin->SpawnedActor = nullptr;
-	}
-
-	if ( !UsdPrimTwin->SceneComponent.IsValid() || !UsdPrimTwin->SceneComponent->IsA( ComponentClass ) )
-	{
-		if ( UsdPrimTwin->SceneComponent.IsValid() )
-		{
-			UsdPrimTwin->SceneComponent->DestroyComponent();
-			UsdPrimTwin->SceneComponent = nullptr;
-		}
-
-		UObject* ComponentOwner = UsdPrimTwin->SpawnedActor.Get();
-
-		if ( !ComponentOwner )
-		{
-			if ( ParentUsdPrimTwin->SceneComponent.IsValid() )
-			{
-				ComponentOwner = ParentUsdPrimTwin->SceneComponent.Get();
-			}
-			else
-			{
-				ComponentOwner = ParentUsdPrimTwin->SpawnedActor.Get();
-			}
-		}
-
-		UsdPrimTwin->SceneComponent = NewObject< USceneComponent >( ComponentOwner, ComponentClass, FName( Prim.GetName().GetText() ), DefaultObjFlag );
-
-		if ( UsdPrimTwin->SpawnedActor.IsValid() )
-		{
-			UsdPrimTwin->SpawnedActor->AddInstanceComponent( UsdPrimTwin->SceneComponent.Get() );
-
-			if ( !UsdPrimTwin->SpawnedActor->GetRootComponent() )
-			{
-				UsdPrimTwin->SpawnedActor->SetRootComponent( UsdPrimTwin->SceneComponent.Get() );
-			}
+			PrimsToAnimate.Add( PrimPath );
 		}
 	}
 
 	return UsdPrimTwin;
 }
 
-FUsdPrimTwin* AUsdStageActor::LoadPrim( const pxr::SdfPath& Path, const TMap< FString, UStaticMesh* >& PrimPathsToStaticMeshes )
+FUsdPrimTwin* AUsdStageActor::ExpandPrim( const pxr::UsdPrim& Prim, FUsdSchemaTranslationContext& TranslationContext )
 {
-	TRACE_CPUPROFILER_EVENT_SCOPE( AUsdStageActor::LoadPrim );
-
-	FScopedUsdAllocs UsdAllocs;
-
-	if ( !GetUsdStage() )
-	{
-		return nullptr;
-	}
-
-	TUsdStore< pxr::UsdPrim > PrimStore = GetUsdStage()->GetPrimAtPath( Path.GetPrimPath() );
-	if ( !PrimStore.Get() )
+	if ( !Prim )
 	{
 		return nullptr;
 	}
 
-	pxr::UsdPrim& Prim = PrimStore.Get();
-
-	FUsdPrimTwin* ParentUsdPrimTwin = RootUsdTwin.Find( UsdToUnreal::ConvertPath( Path.GetParentPath() ) );
-
-	if ( !ParentUsdPrimTwin )
-	{
-		ParentUsdPrimTwin = &RootUsdTwin;
-	}
+	TRACE_CPUPROFILER_EVENT_SCOPE( AUsdStageActor::ExpandPrim );
 
-	FUsdPrimTwin* UsdPrimTwin = SpawnPrim( Path );
+	FUsdPrimTwin* UsdPrimTwin = GetOrCreatePrimTwin( Prim.GetPrimPath() );
 
-	if ( !UsdPrimTwin || !UsdPrimTwin->SceneComponent.IsValid() )
-	{
-		return nullptr;
-	}
+	bool bExpandChilren = true;
 
-	if ( UStaticMeshComponent* MeshComponent = Cast< UStaticMeshComponent >( UsdPrimTwin->SceneComponent.Get() ) )
-	{
-		SetStaticMesh( pxr::UsdGeomMesh( Prim ), *MeshComponent, PrimPathsToStaticMeshes );
-	}
-	else if ( UCineCameraComponent* CameraComponent = Cast< UCineCameraComponent >( UsdPrimTwin->SceneComponent.Get() ) )
-	{
-		UsdToUnreal::ConvertGeomCamera( GetUsdStage(), pxr::UsdGeomCamera( Prim ), *CameraComponent, pxr::UsdTimeCode( Time ) );
-	}
+	IUsdSchemasModule& UsdSchemasModule = FModuleManager::Get().LoadModuleChecked< IUsdSchemasModule >( TEXT("USDSchemas") );
 
-	if ( Path.IsPrimPath() )
+	if ( TSharedPtr< FUsdSchemaTranslator > SchemaTranslator = UsdSchemasModule.GetTranslatorRegistry().CreateTranslatorForSchema( TranslationContext.AsShared(), pxr::UsdTyped( Prim ) ) )
 	{
-		USceneComponent* AttachToComponent = ParentUsdPrimTwin->SceneComponent.Get();
-
-		if ( !AttachToComponent )
+		if ( !UsdPrimTwin->SceneComponent.IsValid() )
 		{
-			AttachToComponent = SceneComponent;
-		}
+			UsdPrimTwin->SceneComponent = SchemaTranslator->CreateComponents();
 
-		USceneComponent* ComponentToAttach = UsdPrimTwin->SceneComponent.Get();
-		if ( ComponentToAttach->GetOwner()->GetRootComponent() != ComponentToAttach && ComponentToAttach->GetOwner() != AttachToComponent->GetOwner() )
-		{
-			ComponentToAttach = ComponentToAttach->GetOwner()->GetRootComponent();
-		}
-
-		ComponentToAttach->AttachToComponent( AttachToComponent, FAttachmentTransformRules::KeepRelativeTransform );
-		
-		if ( PrimsToAnimate.Contains( UsdToUnreal::ConvertPath( Path.GetPrimPath() ) ) )
-		{
-			ComponentToAttach->SetMobility( EComponentMobility::Movable );
+			if ( UsdPrimTwin->SceneComponent.IsValid() )
+			{
+				ObjectsToWatch.Add( UsdPrimTwin->SceneComponent.Get(), UsdPrimTwin->PrimPath );
+			}
 		}
 		else
 		{
-			ComponentToAttach->SetMobility( EComponentMobility::Static );
+			SchemaTranslator->UpdateComponents( UsdPrimTwin->SceneComponent.Get() );
 		}
 
-		UsdPrimTwin->SceneComponent->GetOwner()->AddInstanceComponent( UsdPrimTwin->SceneComponent.Get() );
-	}
-
-	if ( !Path.IsPropertyPath() || pxr::UsdGeomXformOp::IsXformOp( Path.GetNameToken() ) )
-	{
-		pxr::UsdGeomXformable Xformable( Prim );
-		UsdToUnreal::ConvertXformable( Prim.GetStage(), Xformable, *UsdPrimTwin->SceneComponent );
+		bExpandChilren = !SchemaTranslator->CollapsedHierarchy();
 	}
 
-	if ( pxr::UsdGeomImageable GeomImageable = pxr::UsdGeomImageable( Prim ) )
+	if ( bExpandChilren )
 	{
-		bool bIsHidden = ( GeomImageable.ComputeVisibility() == pxr::UsdGeomTokens->invisible );
+		USceneComponent* ContextParentComponent = TranslationContext.ParentComponent;
 
-		if ( UsdPrimTwin->SpawnedActor.IsValid() )
+		if ( UsdPrimTwin && UsdPrimTwin->SceneComponent.IsValid() )
 		{
-			UsdPrimTwin->SpawnedActor->SetIsTemporarilyHiddenInEditor( bIsHidden );
-			UsdPrimTwin->SpawnedActor->SetActorHiddenInGame( bIsHidden );
+			ContextParentComponent = UsdPrimTwin->SceneComponent.Get();
 		}
-		else if ( UsdPrimTwin->SceneComponent.IsValid() )
-		{
-			UsdPrimTwin->SceneComponent->SetVisibility( !bIsHidden );
-		}
-	}
-
-	if ( UsdPrimTwin->SpawnedActor.IsValid() )
-	{
-		ObjectsToWatch.Add( UsdPrimTwin->SpawnedActor.Get(), UsdPrimTwin->PrimPath );
-	}
-	else if ( UsdPrimTwin->SceneComponent.IsValid() )
-	{
-		ObjectsToWatch.Add( UsdPrimTwin->SceneComponent.Get(), UsdPrimTwin->PrimPath );
-	}
 
-	if ( !UsdPrimTwin->SceneComponent->IsRegistered() )
-	{
-		UsdPrimTwin->SceneComponent->RegisterComponent();
-	}
+		TGuardValue< USceneComponent* > ParentComponentGuard( TranslationContext.ParentComponent, ContextParentComponent );
 
-	return UsdPrimTwin;
-}
-
-FUsdPrimTwin* AUsdStageActor::ExpandPrim( const pxr::UsdPrim& Prim, const TMap< FString, UStaticMesh* >& PrimPathsToStaticMeshes )
-{
-	if ( !Prim )
-	{
-		return nullptr;
-	}
-
-	TRACE_CPUPROFILER_EVENT_SCOPE( AUsdStageActor::ExpandPrim );
-
-	FUsdPrimTwin* UsdPrimTwin = LoadPrim( Prim.GetPrimPath(), PrimPathsToStaticMeshes );
-
-	if ( Prim.IsA<pxr::UsdSkelRoot>() )
-	{
-		USkinnedMeshComponent* SkinnedMeshComponent = Cast< USkinnedMeshComponent >( UsdPrimTwin->SceneComponent.Get() );
-		if ( SkinnedMeshComponent )
-		{
-			ProcessSkeletonRoot( Prim, *SkinnedMeshComponent );
-		}
-	}
-	else if ( Prim )
-	{
 		pxr::UsdPrimSiblingRange PrimChildren = Prim.GetFilteredChildren( pxr::UsdTraverseInstanceProxies() );
 		for ( TUsdStore< pxr::UsdPrim > ChildStore : PrimChildren )
 		{
-			ExpandPrim( ChildStore.Get(), PrimPathsToStaticMeshes );
-		}
-	}
-
-	return UsdPrimTwin;
-}
-
-bool AUsdStageActor::ProcessSkeletonRoot( const pxr::UsdPrim& Prim, USkinnedMeshComponent& SkinnedMeshComponent )
-{
-	TRACE_CPUPROFILER_EVENT_SCOPE( AUsdStageActor::ProcessSkeletonRoot );
-
-	FSkeletalMeshImportData SkelMeshImportData;
-	bool bIsSkeletalDataValid = true;
-
-	// Retrieve the USD skeletal data under the SkeletonRoot into the SkeletalMeshImportData
-	{
-		FScopedUsdAllocs UsdAllocs;
-
-		pxr::UsdSkelCache SkeletonCache;
-		pxr::UsdSkelRoot SkeletonRoot(Prim);
-		SkeletonCache.Populate(SkeletonRoot);
-
-		pxr::SdfPath PrimPath = Prim.GetPath();
-
-		std::vector<pxr::UsdSkelBinding> SkeletonBindings;
-		SkeletonCache.ComputeSkelBindings(SkeletonRoot, &SkeletonBindings);
-
-		if (SkeletonBindings.size() == 0)
-		{
-			return false;
-		}
-
-		const pxr::TfToken StageUpAxis = UsdUtils::GetUsdStageAxis( GetUsdStage() );
-
-		pxr::UsdGeomXformable Xformable(Prim);
-
-		std::vector<double> TimeSamples;
-		bool bHasTimeSamples = Xformable.GetTimeSamples(&TimeSamples);
-
-		if (TimeSamples.size() > 0)
-		{
-			auto UpdateSkelRootTransform = [&, Prim]()
-			{
-				pxr::UsdGeomXformable Xformable(Prim);
-				UsdToUnreal::ConvertXformable(Prim.GetStage(), Xformable, SkinnedMeshComponent, pxr::UsdTimeCode(Time));
-			};
-
-			FDelegateHandle Handle = OnTimeChanged.AddLambda(UpdateSkelRootTransform);
-			PrimDelegates.Add(UsdToUnreal::ConvertPath(PrimPath), Handle);
-		}
-
-		bool bHasPrimDisplayColor = false;
-
-		// Note that there could be multiple skeleton bindings under the SkeletonRoot
-		// For now, extract just the first one 
-		for (const pxr::UsdSkelBinding& Binding : SkeletonBindings)
-		{
-			const pxr::UsdSkelSkeleton& Skeleton = Binding.GetSkeleton();
-			pxr::UsdSkelSkeletonQuery SkelQuery = SkeletonCache.GetSkelQuery(Skeleton);
-
-			bool bSkeletonValid = UsdToUnreal::ConvertSkeleton(SkelQuery, SkelMeshImportData);
-			if (!bSkeletonValid)
-			{
-				bIsSkeletalDataValid = false;
-				break;
-			}
-
-			for (const pxr::UsdSkelSkinningQuery& SkinningQuery : Binding.GetSkinningTargets())
-			{
-				// In USD, the skinning target need not be a mesh, but for Unreal we are only interested in skinning meshes
-				pxr::UsdGeomMesh SkinningMesh = pxr::UsdGeomMesh(SkinningQuery.GetPrim());
-				if (SkinningMesh)
-				{
-					UsdToUnreal::ConvertSkinnedMesh(SkinningQuery, SkelMeshImportData);
-					bHasPrimDisplayColor = bHasPrimDisplayColor || SkinningMesh.GetDisplayColorPrimvar().IsDefined();
-				}
-			}
-
-			const pxr::UsdSkelAnimQuery& AnimQuery = SkelQuery.GetAnimQuery();
-			bool bRes = AnimQuery.GetJointTransformTimeSamples(&TimeSamples);
-
-			if (TimeSamples.size() > 0)
-			{
-				auto UpdateBoneTransforms = [&, SkelQuery, StageUpAxis]()
-				{
-					FScopedUnrealAllocs UnrealAllocs;
-
-					UPoseableMeshComponent* PoseableMeshComponent = Cast<UPoseableMeshComponent>(&SkinnedMeshComponent);
-					if (!PoseableMeshComponent || !SkelQuery)
-					{
-						return;
-					}
-
-					TArray<FTransform> BoneTransforms;
-					pxr::VtArray<pxr::GfMatrix4d> UsdBoneTransforms;
-					bool bJointTransformsComputed = SkelQuery.ComputeJointLocalTransforms(&UsdBoneTransforms, pxr::UsdTimeCode(Time));
-					if (bJointTransformsComputed)
-					{
-						for (uint32 Index = 0; Index < UsdBoneTransforms.size(); ++Index)
-						{
-							const pxr::GfMatrix4d& UsdMatrix = UsdBoneTransforms[Index];
-							FTransform BoneTransform = UsdToUnreal::ConvertMatrix( StageUpAxis, UsdMatrix );
-							BoneTransforms.Add(BoneTransform);
-						}
-					}
-
-					for (int32 Index = 0; Index < BoneTransforms.Num(); ++Index)
-					{
-						PoseableMeshComponent->BoneSpaceTransforms[Index] = BoneTransforms[Index];
-					}
-					PoseableMeshComponent->RefreshBoneTransforms();
-				};
-
-				FDelegateHandle Handle = OnTimeChanged.AddLambda(UpdateBoneTransforms);
-				PrimDelegates.Add(UsdToUnreal::ConvertPath(PrimPath), Handle);
-			}
-
-			UsdStageActorImpl::ProcessMaterials(Prim, SkelMeshImportData, MaterialsCache, bHasPrimDisplayColor, Time);
-
-			break;
+			ExpandPrim( ChildStore.Get(), TranslationContext );
 		}
 	}
 
-	if (!bIsSkeletalDataValid)
+	if ( UsdPrimTwin && UsdPrimTwin->SceneComponent.IsValid() && !UsdPrimTwin->SceneComponent->IsRegistered() )
 	{
-		return false;
-	}
-
-	USkeletalMesh* SkeletalMesh = UsdToUnreal::GetSkeletalMeshFromImportData(SkelMeshImportData, DefaultObjFlag);
-
-	if (SkeletalMesh)
-	{
-		SkinnedMeshComponent.SetSkeletalMesh(SkeletalMesh);
+		UsdPrimTwin->SceneComponent->RegisterComponent();
 	}
 
-	return SkeletalMesh != nullptr;
+	return UsdPrimTwin;
 }
 
-void AUsdStageActor::UpdatePrim( const pxr::SdfPath& InUsdPrimPath, bool bResync, const TMap< FString, UStaticMesh* >& PrimPathsToStaticMeshes )
+void AUsdStageActor::UpdatePrim( const pxr::SdfPath& InUsdPrimPath, bool bResync, FUsdSchemaTranslationContext& TranslationContext )
 {
 	TRACE_CPUPROFILER_EVENT_SCOPE( AUsdStageActor::UpdatePrim );
 
@@ -731,7 +283,7 @@ void AUsdStageActor::UpdatePrim( const pxr::SdfPath& InUsdPrimPath, bool bResync
 		}
 
 		TUsdStore< pxr::UsdPrim > PrimToExpand = GetUsdStage()->GetPrimAtPath( UsdPrimPath );
-		FUsdPrimTwin* UsdPrimTwin = ExpandPrim( PrimToExpand.Get(), PrimPathsToStaticMeshes );
+		FUsdPrimTwin* UsdPrimTwin = ExpandPrim( PrimToExpand.Get(), TranslationContext );
 
 		GEditor->BroadcastLevelActorListChanged();
 		GEditor->RedrawLevelEditingViewports();
@@ -759,20 +311,30 @@ void AUsdStageActor::PostEditChangeProperty(FPropertyChangedEvent& PropertyChang
 	UProperty* PropertyThatChanged = PropertyChangedEvent.MemberProperty;
 	const FName PropertyName = PropertyThatChanged ? PropertyThatChanged->GetFName() : NAME_None;
 	
+#if USE_USD_SDK
 	if ( PropertyName == GET_MEMBER_NAME_CHECKED( AUsdStageActor, RootLayer ) )
 	{
-#if USE_USD_SDK
 		UnrealUSDWrapper::GetUsdStageCache().Erase( UsdStageStore.Get() );
 
 		UsdStageStore = TUsdStore< pxr::UsdStageRefPtr >();
 		LoadUsdStage();
-#endif // #if USE_USD_SDK
 	}
 	else if ( PropertyName == GET_MEMBER_NAME_CHECKED( AUsdStageActor, Time ) )
 	{
 		Refresh();
 	}
+	else if ( PropertyName == GET_MEMBER_NAME_CHECKED( AUsdStageActor, StartTimeCode ) || PropertyName == GET_MEMBER_NAME_CHECKED( AUsdStageActor, EndTimeCode ) )
+	{
+		if ( UsdStageStore.Get() )
+		{
+			UsdStageStore.Get()->SetStartTimeCode( StartTimeCode );
+			UsdStageStore.Get()->SetEndTimeCode( EndTimeCode );
+
+			SetupLevelSequence();
+		}
+	}
 	else
+#endif // #if USE_USD_SDK
 	{
 		Super::PostEditChangeProperty( PropertyChangedEvent );
 	}
@@ -780,8 +342,8 @@ void AUsdStageActor::PostEditChangeProperty(FPropertyChangedEvent& PropertyChang
 
 void AUsdStageActor::Clear()
 {
-	MeshCache.Reset();
-	MaterialsCache.Reset();
+	AssetsCache.Reset();
+	PrimPathsToAssets.Reset();
 	ObjectsToWatch.Reset();
 }
 
@@ -836,7 +398,7 @@ void AUsdStageActor::OpenUsdStage()
 
 void AUsdStageActor::InitLevelSequence(float FramesPerSecond)
 {
-	if (LevelSequence)
+	if ( LevelSequence || !HasAutorithyOverStage() )
 	{
 		return;
 	}
@@ -943,30 +505,20 @@ void AUsdStageActor::LoadUsdStage()
 
 	SetupLevelSequence();
 
-	TMap< FString, UStaticMesh* > PrimPathsToStaticMeshes = LoadStaticMeshes( UsdStage->GetPseudoRoot() );
+	TSharedRef< FUsdSchemaTranslationContext > TranslationContext = UsdStageActorImpl::CreateUsdSchemaTranslationContext( this, RootUsdTwin.PrimPath );
+
+	LoadAssets( *TranslationContext, UsdStage->GetPseudoRoot() );
+
+	UpdatePrim( UsdStage->GetPseudoRoot().GetPrimPath(), true, *TranslationContext );
 
-	UpdatePrim( UsdStage->GetPseudoRoot().GetPrimPath(), true, PrimPathsToStaticMeshes );
+	TranslationContext->CompleteTasks();
 
 	SetTime( StartTimeCode );
 
 	GEditor->BroadcastLevelActorListChanged();
 	GEditor->RedrawLevelEditingViewports();
 
-	auto AnimatePrims = [ this ]()
-	{
-		for ( const FString& PrimToAnimate : PrimsToAnimate )
-		{
-			TUsdStore< pxr::SdfPath > PrimPath = UnrealToUsd::ConvertPath( *PrimToAnimate );
-
-			TMap< FString, UStaticMesh* > PrimPathsToStaticMeshes = LoadStaticMeshes( this->GetUsdStage()->GetPrimAtPath( PrimPath.Get() ) );
-			this->LoadPrim( PrimPath.Get(), PrimPathsToStaticMeshes );
-		}
-
-		GEditor->BroadcastLevelActorListChanged();
-		GEditor->RedrawLevelEditingViewports();
-	};
-
-	OnTimeChanged.AddLambda( AnimatePrims );
+	OnTimeChanged.AddUObject( this, &AUsdStageActor::AnimatePrims );
 
 	// Log time spent to load the stage
 	double ElapsedSeconds = FPlatformTime::ToSeconds64(FPlatformTime::Cycles64() - StartTime);
@@ -980,7 +532,10 @@ void AUsdStageActor::LoadUsdStage()
 
 void AUsdStageActor::Refresh() const
 {
-	OnTimeChanged.Broadcast();
+	if ( HasAutorithyOverStage() )
+	{
+		OnTimeChanged.Broadcast();
+	}
 }
 
 void AUsdStageActor::PostRegisterAllComponents()
@@ -988,12 +543,21 @@ void AUsdStageActor::PostRegisterAllComponents()
 	Super::PostRegisterAllComponents();
 
 #if USE_USD_SDK
-	const pxr::UsdStageRefPtr& UsdStage = GetUsdStage();
-
-	if ( UsdStage )
+	if ( HasAutorithyOverStage() )
+	{
+		LoadUsdStage();
+	}
+	else
 	{
-		TMap< FString, UStaticMesh* > PrimPathsToStaticMeshes = LoadStaticMeshes( UsdStage->GetPseudoRoot() );
-		UpdatePrim( UsdStage->GetPseudoRoot().GetPrimPath(), true, PrimPathsToStaticMeshes );
+		const pxr::UsdStageRefPtr& UsdStage = GetUsdStage();
+
+		if ( UsdStage )
+		{
+			TSharedRef< FUsdSchemaTranslationContext > TranslationContext = UsdStageActorImpl::CreateUsdSchemaTranslationContext( this, RootUsdTwin.PrimPath );
+
+			LoadAssets( *TranslationContext, UsdStage->GetPseudoRoot() );
+			UpdatePrim( UsdStage->GetPseudoRoot().GetPrimPath(), true, *TranslationContext );
+		}
 	}
 #endif // #if USE_USD_SDK
 }
@@ -1002,7 +566,10 @@ void AUsdStageActor::PostLoad()
 {
 	Super::PostLoad();
 
-	OnActorLoaded.Broadcast( this );
+	if ( HasAutorithyOverStage() )
+	{
+		OnActorLoaded.Broadcast( this );
+	}
 }
 
 void AUsdStageActor::OnUsdPrimTwinDestroyed( const FUsdPrimTwin& UsdPrimTwin )
@@ -1059,198 +626,46 @@ void AUsdStageActor::OnPrimObjectPropertyChanged( UObject* ObjectBeingModified,
 }
 
 #if USE_USD_SDK
-namespace UsdStageActorImpl
+void AUsdStageActor::LoadAssets( FUsdSchemaTranslationContext& TranslationContext, const pxr::UsdPrim& StartPrim )
 {
-	void PreBuildStaticMesh( const pxr::UsdGeomMesh& UsdMesh, UStaticMesh& StaticMesh, TMap< FString, UMaterial* >& MaterialsCache, float Time )
-	{
-		const bool bHasPrimDisplayColor = UsdMesh.GetDisplayColorPrimvar().IsDefined();
-		ProcessMaterials( UsdMesh.GetPrim(), StaticMesh, MaterialsCache, bHasPrimDisplayColor, Time );
-
-		StaticMesh.RenderData = MakeUnique< FStaticMeshRenderData >();
-	}
-
-	bool BuildStaticMesh( UStaticMesh& StaticMesh, IMeshBuilderModule& MeshBuilderModule )
-	{
-		ITargetPlatformManagerModule& TargetPlatformManager = GetTargetPlatformManagerRef();
-		ITargetPlatform* RunningPlatform = TargetPlatformManager.GetRunningTargetPlatform();
-		check(RunningPlatform);
+	TArray< TUsdStore< pxr::UsdPrim > > AllPrimAssets = UsdUtils::GetAllPrimsOfType( StartPrim, pxr::TfType::Find< pxr::UsdGeomMesh >(), { pxr::TfType::Find< pxr::UsdSkelRoot >() } );
+	AllPrimAssets.Append( UsdUtils::GetAllPrimsOfType( StartPrim, pxr::TfType::Find< pxr::UsdSkelRoot >() ) );
 
-		const FStaticMeshLODSettings& LODSettings = RunningPlatform->GetStaticMeshLODSettings();
-
-		const FStaticMeshLODGroup& LODGroup = LODSettings.GetLODGroup(StaticMesh.LODGroup);
-
-		if ( !MeshBuilderModule.BuildMesh( *StaticMesh.RenderData, &StaticMesh, LODGroup ) )
-		{
-			return false;
-		}
-
-		for ( FStaticMeshLODResources& LODResources : StaticMesh.RenderData->LODResources )
-		{
-			LODResources.bHasColorVertexData = true;
-		}
-
-		return true;
-	}
-
-	void PostBuildStaticMesh( UStaticMesh& StaticMesh )
+	IUsdSchemasModule& UsdSchemasModule = FModuleManager::Get().LoadModuleChecked< IUsdSchemasModule >( TEXT("USDSchemas") );
+	
+	for ( TUsdStore< pxr::UsdPrim >& UsdPrim : AllPrimAssets )
 	{
-		StaticMesh.InitResources();
-
-		if ( const FMeshDescription* MeshDescription = StaticMesh.GetMeshDescription( 0 ) )
+		if ( TSharedPtr< FUsdSchemaTranslator > SchemaTranslator = UsdSchemasModule.GetTranslatorRegistry().CreateTranslatorForSchema( TranslationContext.AsShared(), pxr::UsdTyped( UsdPrim.Get() ) ) )
 		{
-			StaticMesh.RenderData->Bounds = MeshDescription->GetBounds();
+			SchemaTranslator->CreateAssets();
 		}
-	
-		StaticMesh.CalculateExtendedBounds();
-	}
-}
-
-UStaticMesh* AUsdStageActor::CreateStaticMesh( const pxr::UsdGeomMesh& UsdMesh, FMeshDescription&& MeshDescription, bool& bOutIsNew )
-{
-	UStaticMesh* StaticMesh = nullptr;
-
-	FSHAHash MeshHash = FMeshDescriptionOperations::ComputeSHAHash( MeshDescription );
-
-	StaticMesh = MeshCache.FindRef( MeshHash.ToString() );
-
-	if ( !StaticMesh && !MeshDescription.IsEmpty() )
-	{
-		bOutIsNew = true;
-
-		StaticMesh = NewObject< UStaticMesh >( GetTransientPackage(), NAME_None, DefaultObjFlag | EObjectFlags::RF_Public );
-
-		FStaticMeshSourceModel& SourceModel = StaticMesh->AddSourceModel();
-		SourceModel.BuildSettings.bGenerateLightmapUVs = false;
-		SourceModel.BuildSettings.bRecomputeNormals = false;
-		SourceModel.BuildSettings.bRecomputeTangents = false;
-		SourceModel.BuildSettings.bBuildAdjacencyBuffer = false;
-		SourceModel.BuildSettings.bBuildReversedIndexBuffer = false;
-
-		FMeshDescription* StaticMeshDescription = StaticMesh->CreateMeshDescription(0);
-		check( StaticMeshDescription );
-		*StaticMeshDescription = MoveTemp( MeshDescription );
-
-		MeshCache.Add( MeshHash.ToString() ) = StaticMesh;
-	}
-	else
-	{
-		//FPlatformMisc::LowLevelOutputDebugStringf( TEXT("Mesh found in cache %s\n"), *StaticMesh->GetName() );
-		bOutIsNew = false;
 	}
 
-	return StaticMesh;
+	TranslationContext.CompleteTasks(); // Finish the assets tasks before moving on
 }
 
-TMap< FString, UStaticMesh* > AUsdStageActor::LoadStaticMeshes( const pxr::UsdPrim& StartPrim )
+void AUsdStageActor::AnimatePrims()
 {
-	TMap< FString, UStaticMesh* > PrimPathsToStaticMeshes;
-
-	FScopedSlowTask LoadAllProgress( 6.f, LOCTEXT( "LoadingStaticMeshes", "Loading Static Meshes" ) );
-	LoadAllProgress.MakeDialog();
-
-	TArray< TUsdStore< pxr::UsdPrim > > AllGeomMeshes = UsdUtils::GetAllPrimsOfType( StartPrim, pxr::TfType::Find< pxr::UsdGeomMesh >(), { pxr::TfType::Find< pxr::UsdSkelRoot >() } );
-	TArray< FMeshDescription > MeshDescriptions;
-	MeshDescriptions.AddDefaulted( AllGeomMeshes.Num() );
-
-	// Create all mesh descriptions
-	{
-		LoadAllProgress.EnterProgressFrame();
-
-		ParallelFor( AllGeomMeshes.Num(), 
-			[ &AllGeomMeshes, &MeshDescriptions, TimeCode = pxr::UsdTimeCode( Time ) ]( int32 GeomMeshIndex )
-			{
-				MeshDescriptions[ GeomMeshIndex ] = UsdStageActorImpl::LoadMeshDescription( pxr::UsdGeomMesh( AllGeomMeshes[ GeomMeshIndex ].Get() ), TimeCode );
-			} );
-	}
-
-	TArray< TPair< int32, UStaticMesh* > > StaticMeshesToBuild;
-	StaticMeshesToBuild.Reserve( AllGeomMeshes.Num() );
-
-	// Create all static meshes
-	{
-		LoadAllProgress.EnterProgressFrame();
-
-		for ( int32 PrimIndex = 0; PrimIndex < AllGeomMeshes.Num(); ++PrimIndex )
-		{
-			bool bIsNew = true;
-			UStaticMesh* StaticMesh = CreateStaticMesh( pxr::UsdGeomMesh( AllGeomMeshes[ PrimIndex ].Get() ), MoveTemp( MeshDescriptions[ PrimIndex ] ), bIsNew );
-
-			PrimPathsToStaticMeshes.Add( UsdToUnreal::ConvertPath( AllGeomMeshes[ PrimIndex ].Get().GetPrimPath() ), StaticMesh );
-
-			if ( bIsNew )
-			{
-				StaticMeshesToBuild.Emplace( PrimIndex, StaticMesh );
-			}
-		}
-	}
-
-	// Commit all mesh descriptions
-	{
-		LoadAllProgress.EnterProgressFrame();
-
-		ParallelFor( StaticMeshesToBuild.Num(), 
-			[ &StaticMeshesToBuild ]( int32 MeshIndex )
-			{
-				UStaticMesh::FCommitMeshDescriptionParams Params;
-				Params.bMarkPackageDirty = false;
-				Params.bUseHashAsGuid = true;
-
-				TPair< int32, UStaticMesh* >& MeshToBuild = StaticMeshesToBuild[ MeshIndex ];
-
-				if ( UStaticMesh* StaticMesh = MeshToBuild.Value )
-				{
-					StaticMesh->CommitMeshDescription( 0, Params );
-				}
-			} );
-	}
-
-	// PreBuild all static meshes
-	{
-		LoadAllProgress.EnterProgressFrame();
-
-		for ( TPair< int32, UStaticMesh* >& MeshToBuild : StaticMeshesToBuild )
-		{
-			if ( UStaticMesh* StaticMesh = MeshToBuild.Value )
-			{
-				UsdStageActorImpl::PreBuildStaticMesh( pxr::UsdGeomMesh( AllGeomMeshes[ MeshToBuild.Key ].Get() ), *StaticMesh, MaterialsCache, Time );
-			}
-		}
-	}
-
-	// Build all static meshes
-	{
-		LoadAllProgress.EnterProgressFrame();
-
-		IMeshBuilderModule& MeshBuilderModule = FModuleManager::LoadModuleChecked< IMeshBuilderModule >( TEXT("MeshBuilder") );
+	TSharedRef< FUsdSchemaTranslationContext > TranslationContext = UsdStageActorImpl::CreateUsdSchemaTranslationContext( this, RootUsdTwin.PrimPath );
 
-		ParallelFor( StaticMeshesToBuild.Num(), 
-			[ &StaticMeshesToBuild, &MeshBuilderModule ]( int32 MeshIndex )
-			{
-				if ( UStaticMesh* StaticMesh = StaticMeshesToBuild[ MeshIndex ].Value )
-				{
-					if ( !UsdStageActorImpl::BuildStaticMesh( *StaticMesh, MeshBuilderModule ) )
-					{
-						// Build failed, discard the mesh
-						StaticMeshesToBuild[ MeshIndex ].Value = nullptr;
-					}
-				}
-			} );
-	}
-
-	// PostBuild all static meshes
+	for ( const FString& PrimToAnimate : PrimsToAnimate )
 	{
-		LoadAllProgress.EnterProgressFrame();
+		TUsdStore< pxr::SdfPath > PrimPath = UnrealToUsd::ConvertPath( *PrimToAnimate );
 
-		for ( TPair< int32, UStaticMesh* >& MeshToBuild : StaticMeshesToBuild )
+		IUsdSchemasModule& SchemasModule = FModuleManager::Get().LoadModuleChecked< IUsdSchemasModule >( "USDSchemas" );
+		if ( TSharedPtr< FUsdSchemaTranslator > SchemaTranslator = SchemasModule.GetTranslatorRegistry().CreateTranslatorForSchema( TranslationContext, pxr::UsdTyped( this->GetUsdStage()->GetPrimAtPath( PrimPath.Get() ) ) ) )
 		{
-			if ( UStaticMesh* StaticMesh = MeshToBuild.Value )
+			if ( FUsdPrimTwin* UsdPrimTwin = RootUsdTwin.Find( PrimToAnimate ) )
 			{
-				UsdStageActorImpl::PostBuildStaticMesh( *StaticMesh );
+				SchemaTranslator->UpdateComponents( UsdPrimTwin->SceneComponent.Get() );
 			}
 		}
 	}
 
-	return PrimPathsToStaticMeshes;
+	TranslationContext->CompleteTasks();
+		
+	GEditor->BroadcastLevelActorListChanged();
+	GEditor->RedrawLevelEditingViewports();
 }
 #endif // #if USE_USD_SDK
 
diff --git a/Engine/Plugins/Importers/USDImporter/Source/USDStage/Public/USDStageActor.h b/Engine/Plugins/Importers/USDImporter/Source/USDStage/Public/USDStageActor.h
--- a/Engine/Plugins/Importers/USDImporter/Source/USDStage/Public/USDStageActor.h
+++ b/Engine/Plugins/Importers/USDImporter/Source/USDStage/Public/USDStageActor.h
@@ -27,6 +27,7 @@
 
 class ALevelSequenceActor;
 struct FMeshDescription;
+struct FUsdSchemaTranslationContext;
 class IMeshBuilderModule;
 class ULevelSequence;
 class UMaterial;
@@ -46,6 +47,8 @@ class AUsdStageActor : public AActor
 {
 	GENERATED_BODY()
 
+	friend struct UsdStageActorImpl;
+
 public:
 	UPROPERTY(EditAnywhere, Category = "USD", meta = (FilePathFilter = "usd files (*.usd; *.usda; *.usdc)|*.usd; *.usda; *.usdc"))
 	FFilePath RootLayer;
@@ -53,6 +56,12 @@ public:
 	UPROPERTY(EditAnywhere, Category = "USD")
 	EUsdInitialLoadSet InitialLoadSet;
 
+	UFUNCTION(BlueprintCallable, Category = "USD", meta = (CallInEditor = "true"))
+	float GetTime() const { return Time; }
+
+	UFUNCTION(BlueprintCallable, Category = "USD", meta = (CallInEditor = "true"))
+	void SetTime(float InTime);
+
 private:
 	UPROPERTY(Category = UsdStageActor, VisibleAnywhere, BlueprintReadOnly, meta = (ExposeFunctionCategories = "Mesh,Rendering,Physics,Components|StaticMesh", AllowPrivateAccess = "true"))
 	class USceneComponent* SceneComponent;
@@ -61,13 +70,10 @@ private:
 	UPROPERTY(EditAnywhere, Category = "USD")
 	float Time;
 
-	UFUNCTION(BlueprintCallable, Category = "USD", meta = (CallInEditor = "true"))
-	void SetTime(float InTime);
-
-	UPROPERTY(VisibleAnywhere, Category = "USD")
+	UPROPERTY(EditAnywhere, Category = "USD")
 	float StartTimeCode;
 
-	UPROPERTY(VisibleAnywhere, Category = "USD")
+	UPROPERTY(EditAnywhere, Category = "USD")
 	float EndTimeCode;
 
 	UPROPERTY(VisibleAnywhere, Category = "USD")
@@ -111,6 +117,7 @@ private:
 	void OnUsdPrimTwinDestroyed( const FUsdPrimTwin& UsdPrimTwin );
 
 	void OnPrimObjectPropertyChanged( UObject* ObjectBeingModified, FPropertyChangedEvent& PropertyChangedEvent );
+	bool HasAutorithyOverStage() const { return !GetWorld() || !GetWorld()->IsGameWorld(); }
 
 private:
 	FUsdPrimTwin RootUsdTwin;
@@ -125,10 +132,10 @@ private:
 
 private:
 	UPROPERTY( NonPIEDuplicateTransient )
-	TMap< FString, UStaticMesh* > MeshCache;
+	TMap< FString, UObject* > AssetsCache;
 
 	UPROPERTY( NonPIEDuplicateTransient )
-	TMap< FString, UMaterial* > MaterialsCache;
+	TMap< FString, UObject* > PrimPathsToAssets;
 
 #if USE_USD_SDK
 public:
@@ -137,16 +144,13 @@ public:
 	FUsdListener& GetUsdListener() { return UsdListener; }
 	const FUsdListener& GetUsdListener() const { return UsdListener; }
 
-	FUsdPrimTwin* SpawnPrim( const pxr::SdfPath& UsdPrimPath );
-	FUsdPrimTwin* LoadPrim( const pxr::SdfPath& Path, const TMap< FString, UStaticMesh* >& PrimPathsToStaticMeshes );
-	FUsdPrimTwin* ExpandPrim( const pxr::UsdPrim& Prim, const TMap< FString, UStaticMesh* >& PrimPathsToStaticMeshes );
-	void UpdatePrim( const pxr::SdfPath& UsdPrimPath, bool bResync, const TMap< FString, UStaticMesh* >& PrimPathsToStaticMeshes );
-	bool SetStaticMesh( const pxr::UsdGeomMesh& UsdMesh, UStaticMeshComponent& MeshComponent, const TMap< FString, UStaticMesh* >& PrimPathsToStaticMeshes );
-	bool ProcessSkeletonRoot( const pxr::UsdPrim& Prim, USkinnedMeshComponent& SkinnedMeshComponent );
+	FUsdPrimTwin* GetOrCreatePrimTwin( const pxr::SdfPath& UsdPrimPath );
+	FUsdPrimTwin* ExpandPrim( const pxr::UsdPrim& Prim, FUsdSchemaTranslationContext& TranslationContext );
+	void UpdatePrim( const pxr::SdfPath& UsdPrimPath, bool bResync, FUsdSchemaTranslationContext& TranslationContext );
 
 protected:
-	TMap< FString, UStaticMesh* > LoadStaticMeshes( const pxr::UsdPrim& StartPrim );
-	UStaticMesh* CreateStaticMesh( const pxr::UsdGeomMesh& UsdMesh, FMeshDescription&& MeshDescription, bool& bOutIsNew );
+	void LoadAssets( FUsdSchemaTranslationContext& TranslationContext, const pxr::UsdPrim& StartPrim );
+	void AnimatePrims();
 
 private:
 	TUsdStore< pxr::UsdStageRefPtr > UsdStageStore;
diff --git a/Engine/Plugins/Importers/USDImporter/Source/USDStage/USDStage.Build.cs b/Engine/Plugins/Importers/USDImporter/Source/USDStage/USDStage.Build.cs
--- a/Engine/Plugins/Importers/USDImporter/Source/USDStage/USDStage.Build.cs
+++ b/Engine/Plugins/Importers/USDImporter/Source/USDStage/USDStage.Build.cs
@@ -31,7 +31,6 @@ namespace UnrealBuildTool.Rules
 					"MovieScene",
 					"MovieSceneTracks",
 					"PropertyEditor",
-					"PropertyEditor",
 					"PythonScriptPlugin",
 					"RHI",
 					"RenderCore",
@@ -39,6 +38,7 @@ namespace UnrealBuildTool.Rules
 					"SlateCore",
 					"StaticMeshDescription",
 					"USDImporter",
+					"USDSchemas",
 					"USDUtilities",
 					"UnrealEd",
 				}
diff --git a/Engine/Plugins/Importers/USDImporter/Source/USDStageEditor/Private/SUSDStage.cpp b/Engine/Plugins/Importers/USDImporter/Source/USDStageEditor/Private/SUSDStage.cpp
--- a/Engine/Plugins/Importers/USDImporter/Source/USDStageEditor/Private/SUSDStage.cpp
+++ b/Engine/Plugins/Importers/USDImporter/Source/USDStageEditor/Private/SUSDStage.cpp
@@ -160,7 +160,7 @@ void SUsdStage::SetupStageActorDelegates()
 					this->UsdStageTreeView->RefreshPrim( PrimPath, bResync );
 				}
 
-				if ( this->UsdPrimInfoWidget )
+				if ( this->UsdPrimInfoWidget && UsdStageActor.IsValid() )
 				{
 					this->UsdPrimInfoWidget->SetPrimPath( UsdStageActor->GetUsdStage(), *PrimPath );
 				}
@@ -202,7 +202,7 @@ void SUsdStage::SetupStageActorDelegates()
 		OnStageEditTargetChangedHandle = UsdStageActor->GetUsdListener().OnStageEditTargetChanged.AddLambda(
 			[ this ]()
 			{
-				if ( this->UsdLayersTreeView )
+				if ( this->UsdLayersTreeView && UsdStageActor.IsValid() )
 				{
 					this->UsdLayersTreeView->Refresh( UsdStageActor.Get(), false );
 				}
diff --git a/Engine/Plugins/Importers/USDImporter/Source/USDStageEditor/Private/SUSDStageTreeView.cpp b/Engine/Plugins/Importers/USDImporter/Source/USDStageEditor/Private/SUSDStageTreeView.cpp
--- a/Engine/Plugins/Importers/USDImporter/Source/USDStageEditor/Private/SUSDStageTreeView.cpp
+++ b/Engine/Plugins/Importers/USDImporter/Source/USDStageEditor/Private/SUSDStageTreeView.cpp
@@ -345,6 +345,8 @@ public:
 			{
 				UsdGeomImageable.MakeVisible();
 			}
+
+			TreeItem->RefreshData( false );
 		}
 
 		return FReply::Handled();
diff --git a/Engine/Plugins/Importers/USDImporter/Source/USDUtilities/Private/USDPrimConversion.cpp b/Engine/Plugins/Importers/USDImporter/Source/USDUtilities/Private/USDPrimConversion.cpp
--- a/Engine/Plugins/Importers/USDImporter/Source/USDUtilities/Private/USDPrimConversion.cpp
+++ b/Engine/Plugins/Importers/USDImporter/Source/USDUtilities/Private/USDPrimConversion.cpp
@@ -47,6 +47,9 @@ bool UsdToUnreal::ConvertXformable( const pxr::UsdStageRefPtr& Stage, const pxr:
 	FTransform Transform = UsdToUnreal::ConvertMatrix( UpAxis, UsdMatrix ) * FTransform( AdditionalRotation );
 	SceneComponent.SetRelativeTransform( Transform );
 
+	const bool bIsHidden = ( Xformable.ComputeVisibility( EvalTime ) == pxr::UsdGeomTokens->invisible );
+	SceneComponent.SetVisibility( !bIsHidden );
+
 	return true;
 }
 
diff --git a/Engine/Plugins/Importers/USDImporter/USDImporter.uplugin b/Engine/Plugins/Importers/USDImporter/USDImporter.uplugin
--- a/Engine/Plugins/Importers/USDImporter/USDImporter.uplugin
+++ b/Engine/Plugins/Importers/USDImporter/USDImporter.uplugin
@@ -39,6 +39,14 @@
 				"Win64"
 			]
 		},
+		{
+			"Name": "USDSchemas",
+			"Type": "Editor",
+			"LoadingPhase": "Default",
+			"WhitelistPlatforms": [
+				"Win64"
+			]
+		},
 		{
 			"Name": "USDStage",
 			"Type": "Editor",
