Build fix: Removed dependency on ModelingComponents

#rnx

diff --git a/Engine/Plugins/Importers/USDImporter/Source/USDSchemas/Private/USDGeomMeshTranslator.cpp b/Engine/Plugins/Importers/USDImporter/Source/USDSchemas/Private/USDGeomMeshTranslator.cpp
--- a/Engine/Plugins/Importers/USDImporter/Source/USDSchemas/Private/USDGeomMeshTranslator.cpp
+++ b/Engine/Plugins/Importers/USDImporter/Source/USDSchemas/Private/USDGeomMeshTranslator.cpp
@@ -20,7 +20,6 @@
 #include "Materials/MaterialInterface.h"
 #include "Materials/MaterialInstanceConstant.h"
 #include "MeshDescriptionOperations.h"
-#include "MeshDescriptionToDynamicMesh.h"
 #include "Misc/SecureHash.h"
 #include "PhysicsEngine/BodySetup.h"
 #include "StaticMeshAttributes.h"
diff --git a/Engine/Plugins/Importers/USDImporter/Source/USDSchemas/USDSchemas.Build.cs b/Engine/Plugins/Importers/USDImporter/Source/USDSchemas/USDSchemas.Build.cs
--- a/Engine/Plugins/Importers/USDImporter/Source/USDSchemas/USDSchemas.Build.cs
+++ b/Engine/Plugins/Importers/USDImporter/Source/USDSchemas/USDSchemas.Build.cs
@@ -11,13 +11,6 @@ namespace UnrealBuildTool.Rules
 		{
 			bUseRTTI = true;
 
-			PublicDependencyModuleNames.AddRange(
-				new string[]
-				{
-					"ModelingComponents",
-				}
-				);
-
 			PrivateDependencyModuleNames.AddRange(
 				new string[]
 				{
