Build fix

#rnx

diff --git a/Engine/Plugins/Importers/USDImporter/Source/USDStage/Private/USDStageActor.cpp b/Engine/Plugins/Importers/USDImporter/Source/USDStage/Private/USDStageActor.cpp
--- a/Engine/Plugins/Importers/USDImporter/Source/USDStage/Private/USDStageActor.cpp
+++ b/Engine/Plugins/Importers/USDImporter/Source/USDStage/Private/USDStageActor.cpp
@@ -625,6 +625,11 @@ void AUsdStageActor::OnPrimObjectPropertyChanged( UObject* ObjectBeingModified,
 #endif // #if USE_USD_SDK
 }
 
+bool AUsdStageActor::HasAutorithyOverStage() const
+{
+	return !GetWorld() || !GetWorld()->IsGameWorld();
+}
+
 #if USE_USD_SDK
 void AUsdStageActor::LoadAssets( FUsdSchemaTranslationContext& TranslationContext, const pxr::UsdPrim& StartPrim )
 {
diff --git a/Engine/Plugins/Importers/USDImporter/Source/USDStage/Public/USDStageActor.h b/Engine/Plugins/Importers/USDImporter/Source/USDStage/Public/USDStageActor.h
--- a/Engine/Plugins/Importers/USDImporter/Source/USDStage/Public/USDStageActor.h
+++ b/Engine/Plugins/Importers/USDImporter/Source/USDStage/Public/USDStageActor.h
@@ -117,7 +117,7 @@ private:
 	void OnUsdPrimTwinDestroyed( const FUsdPrimTwin& UsdPrimTwin );
 
 	void OnPrimObjectPropertyChanged( UObject* ObjectBeingModified, FPropertyChangedEvent& PropertyChangedEvent );
-	bool HasAutorithyOverStage() const { return !GetWorld() || !GetWorld()->IsGameWorld(); }
+	bool HasAutorithyOverStage() const;
 
 private:
 	FUsdPrimTwin RootUsdTwin;
