Fixed USD import into level not going through instances

diff --git a/Engine/Plugins/Importers/USDImporter/Source/USDImporter/Private/USDPrimResolver.cpp b/Engine/Plugins/Importers/USDImporter/Source/USDImporter/Private/USDPrimResolver.cpp
--- a/Engine/Plugins/Importers/USDImporter/Source/USDImporter/Private/USDPrimResolver.cpp
+++ b/Engine/Plugins/Importers/USDImporter/Source/USDImporter/Private/USDPrimResolver.cpp
@@ -73,7 +73,7 @@ void UUSDPrimResolver::FindMeshAssetsToImport(FUsdImportContext& ImportContext,
 		}
 		else if(bRecursive)
 		{
-			for ( pxr::UsdPrim Child : StartPrim.Get().GetChildren() )
+			for ( pxr::UsdPrim Child : StartPrim.Get().GetFilteredChildren( pxr::UsdTraverseInstanceProxies() ) )
 			{
 				FindMeshAssetsToImport(ImportContext, Child, StartPrim, OutAssetsToImport);
 			}
@@ -89,7 +89,7 @@ void UUSDPrimResolver::FindActorsToSpawn(FUSDSceneImportContext& ImportContext,
 	}
 	else
 	{
-		for (pxr::UsdPrim Child : ImportContext.RootPrim.Get().GetChildren())
+		for (pxr::UsdPrim Child : ImportContext.RootPrim.Get().GetFilteredChildren( pxr::UsdTraverseInstanceProxies() ))
 		{
 			FindActorsToSpawn_Recursive(ImportContext, Child, pxr::UsdPrim(), OutActorSpawnDatas);
 		}
@@ -402,7 +402,7 @@ void UUSDPrimResolver::FindMeshChildren(FUsdImportContext& ImportContext, const
 			OutMeshChildren.Add(ParentPrim);
 		}
 
-		for ( pxr::UsdPrim Child : ParentPrim.Get().GetChildren() )
+		for ( pxr::UsdPrim Child : ParentPrim.Get().GetFilteredChildren( pxr::UsdTraverseInstanceProxies() ) )
 		{
 			if (!IUsdPrim::IsProxyOrGuide(Child) && !IUsdPrim::IsKindChildOf(Child, USDKindTypes::Component))
 			{
@@ -452,7 +452,7 @@ void UUSDPrimResolver::FindActorsToSpawn_Recursive(FUSDSceneImportContext& Impor
 
 	if (!ImportContext.bFindUnrealAssetReferences || AssetPath.IsEmpty())
 	{
-		for (pxr::UsdPrim Child : Prim.Get().GetChildren())
+		for (pxr::UsdPrim Child : Prim.Get().GetFilteredChildren( pxr::UsdTraverseInstanceProxies() ))
 		{
 			FindActorsToSpawn_Recursive(ImportContext, Child, Prim, *SpawnDataArray);
 		}
diff --git a/Engine/Plugins/Importers/USDImporter/Source/USDImporter/Private/USDPrimResolverKind.cpp b/Engine/Plugins/Importers/USDImporter/Source/USDImporter/Private/USDPrimResolverKind.cpp
--- a/Engine/Plugins/Importers/USDImporter/Source/USDImporter/Private/USDPrimResolverKind.cpp
+++ b/Engine/Plugins/Importers/USDImporter/Source/USDImporter/Private/USDPrimResolverKind.cpp
@@ -102,7 +102,7 @@ void UUSDPrimResolverKind::FindActorsToSpawn_Recursive(FUSDSceneImportContext& I
 		UE_LOG(LogUSDImport, Log, TEXT("Adding %s Group Actor to spawn"), *PrimName.ToString());
 	}
 
-	for (pxr::UsdPrim Child : Prim.Get().GetChildren())
+	for (pxr::UsdPrim Child : Prim.Get().GetFilteredChildren( pxr::UsdTraverseInstanceProxies() ))
 	{
 		FindActorsToSpawn_Recursive(ImportContext, Child, GroupParent, OutSpawnDatas);
 	}
