USD Schemas
- Allowed to override schema translation.
- Added support for custom schemas.

diff --git a/Engine/Plugins/Importers/USDImporter/Source/USDSchemas/Private/USDGeomCameraTranslator.cpp b/Engine/Plugins/Importers/USDImporter/Source/USDSchemas/Private/USDGeomCameraTranslator.cpp
--- a/Engine/Plugins/Importers/USDImporter/Source/USDSchemas/Private/USDGeomCameraTranslator.cpp
+++ b/Engine/Plugins/Importers/USDImporter/Source/USDSchemas/Private/USDGeomCameraTranslator.cpp
@@ -4,11 +4,9 @@
 
 #include "USDMemory.h"
 #include "USDPrimConversion.h"
-#include "USDSchemasModule.h"
 
 #include "CineCameraActor.h"
 #include "CineCameraComponent.h"
-#include "Modules/ModuleManager.h"
 
 #if USE_USD_SDK
 #include "USDIncludesStart.h"
@@ -16,12 +14,6 @@
 #include "USDIncludesEnd.h"
 
 
-void FUsdGeomCameraTranslator::RegisterTranslator()
-{
-	IUsdSchemasModule& UsdSchemasModule = FModuleManager::Get().LoadModuleChecked< IUsdSchemasModule >( TEXT("USDSchemas") );
-	UsdSchemasModule.GetTranslatorRegistry().Register< FUsdGeomCameraTranslator >( TEXT("UsdGeomCamera") );
-}
-
 USceneComponent* FUsdGeomCameraTranslator::CreateComponents()
 {
 	constexpr bool bNeedsActor = true;
diff --git a/Engine/Plugins/Importers/USDImporter/Source/USDSchemas/Private/USDGeomMeshTranslator.cpp b/Engine/Plugins/Importers/USDImporter/Source/USDSchemas/Private/USDGeomMeshTranslator.cpp
--- a/Engine/Plugins/Importers/USDImporter/Source/USDSchemas/Private/USDGeomMeshTranslator.cpp
+++ b/Engine/Plugins/Importers/USDImporter/Source/USDSchemas/Private/USDGeomMeshTranslator.cpp
@@ -6,7 +6,6 @@
 
 #include "USDConversionUtils.h"
 #include "USDGeomMeshConversion.h"
-#include "USDSchemasModule.h"
 #include "USDTypesConversion.h"
 
 #include "Async/Async.h"
@@ -71,7 +70,7 @@ namespace UsdGeomMeshTranslatorImpl
 		return bHasAttributesTimeSamples;
 	}
 
-	void ProcessMaterials( const pxr::UsdPrim& UsdPrim, UStaticMesh& StaticMesh, TMap< FString, UObject* >& MaterialsCache, bool bHasPrimDisplayColor, float Time )
+	void ProcessMaterials( const pxr::UsdPrim& UsdPrim, UStaticMesh& StaticMesh, TMap< FString, UObject* >& PrimPathsToAssets, bool bHasPrimDisplayColor, float Time )
 	{
 		const FMeshDescription* MeshDescription = StaticMesh.GetMeshDescription( 0 );
 
@@ -111,28 +110,7 @@ namespace UsdGeomMeshTranslatorImpl
 
 			if ( MaterialPrim.Get() )
 			{
-				UObject*& CachedMaterial = MaterialsCache.FindOrAdd( UsdToUnreal::ConvertPath( MaterialPrim.Get().GetPrimPath() ) );
-
-				if ( !CachedMaterial )
-				{
-					UMaterial* NewMaterial = NewObject< UMaterial >();
-
-					if ( UsdToUnreal::ConvertMaterial( pxr::UsdShadeMaterial( MaterialPrim.Get() ), *NewMaterial ) )
-					{
-						//UMaterialEditingLibrary::RecompileMaterial( CachedMaterial ); // Too slow
-						NewMaterial->PostEditChange();
-					}
-					else
-					{
-						NewMaterial = nullptr;
-					}
-
-					Material = NewMaterial;
-				}
-				else
-				{
-					Material = Cast< UMaterialInterface >( CachedMaterial );
-				}
+				Material = Cast< UMaterialInterface >( PrimPathsToAssets.FindRef( UsdToUnreal::ConvertPath( MaterialPrim.Get().GetPrimPath() ) ) );
 			}
 
 			if ( Material == nullptr && bHasPrimDisplayColor )
@@ -206,10 +184,10 @@ namespace UsdGeomMeshTranslatorImpl
 		return StaticMesh;
 	}
 
-	void PreBuildStaticMesh( const pxr::UsdGeomMesh& UsdMesh, UStaticMesh& StaticMesh, TMap< FString, UObject* >& MaterialsCache, float Time )
+	void PreBuildStaticMesh( const pxr::UsdGeomMesh& UsdMesh, UStaticMesh& StaticMesh, TMap< FString, UObject* >& PrimPathsToAssets, float Time )
 	{
 		const bool bHasPrimDisplayColor = UsdMesh.GetDisplayColorPrimvar().IsDefined();
-		ProcessMaterials( UsdMesh.GetPrim(), StaticMesh, MaterialsCache, bHasPrimDisplayColor, Time );
+		ProcessMaterials( UsdMesh.GetPrim(), StaticMesh, PrimPathsToAssets, bHasPrimDisplayColor, Time );
 
 		StaticMesh.RenderData = MakeUnique< FStaticMeshRenderData >();
 		//StaticMesh.BodySetup = NewObject< UBodySetup >( &StaticMesh );
@@ -309,7 +287,7 @@ void FGeomMeshCreateAssetsTaskChain::SetupTasks()
 		Then( !bIsAsyncTask,
 			[ this ]()
 			{
-				UsdGeomMeshTranslatorImpl::PreBuildStaticMesh( GeomMesh.Get(), *StaticMesh, Context->AssetsCache, Context->Time );
+				UsdGeomMeshTranslatorImpl::PreBuildStaticMesh( GeomMesh.Get(), *StaticMesh, Context->PrimPathsToAssets, Context->Time );
 
 				return true;
 			} );
@@ -340,12 +318,6 @@ void FGeomMeshCreateAssetsTaskChain::SetupTasks()
 	}
 }
 
-void FUsdGeomMeshTranslator::RegisterTranslator()
-{
-	IUsdSchemasModule& UsdSchemasModule = FModuleManager::Get().LoadModuleChecked< IUsdSchemasModule >( TEXT("USDSchemas") );
-	UsdSchemasModule.GetTranslatorRegistry().Register< FUsdGeomMeshTranslator >( TEXT("UsdGeomMesh") );
-}
-
 void FUsdGeomMeshTranslator::CreateAssets()
 {
 	TSharedRef< FGeomMeshCreateAssetsTaskChain > AssetsTaskChain = MakeShared< FGeomMeshCreateAssetsTaskChain >( Context, pxr::UsdGeomMesh( Schema.Get() ) );
diff --git a/Engine/Plugins/Importers/USDImporter/Source/USDSchemas/Private/USDGeomPointInstancerTranslator.cpp b/Engine/Plugins/Importers/USDImporter/Source/USDSchemas/Private/USDGeomPointInstancerTranslator.cpp
--- a/Engine/Plugins/Importers/USDImporter/Source/USDSchemas/Private/USDGeomPointInstancerTranslator.cpp
+++ b/Engine/Plugins/Importers/USDImporter/Source/USDSchemas/Private/USDGeomPointInstancerTranslator.cpp
@@ -3,14 +3,12 @@
 #include "USDGeomPointInstancerTranslator.h"
 
 #include "USDConversionUtils.h"
-#include "USDSchemasModule.h"
 #include "USDTypesConversion.h"
 
 #if USE_USD_SDK
 
 #include "Components/HierarchicalInstancedStaticMeshComponent.h"
 #include "Engine/StaticMesh.h"
-#include "Modules/ModuleManager.h"
 
 #include "USDIncludesStart.h"
 	#include "pxr/usd/usd/prim.h"
@@ -87,12 +85,6 @@ namespace UsdGeomPointInstancerTranslatorImpl
 	}
 }
 
-void FUsdGeomPointInstancerTranslator::RegisterTranslator()
-{
-	IUsdSchemasModule& UsdSchemasModule = FModuleManager::Get().LoadModuleChecked< IUsdSchemasModule >( TEXT("USDSchemas") );
-	UsdSchemasModule.GetTranslatorRegistry().Register< FUsdGeomPointInstancerTranslator >( TEXT("UsdGeomPointInstancer") );
-}
-
 USceneComponent* FUsdGeomPointInstancerTranslator::CreateComponents()
 {
 	// Spawn Xform Actor / SceneComponent
diff --git a/Engine/Plugins/Importers/USDImporter/Source/USDSchemas/Private/USDGeomXformableTranslator.cpp b/Engine/Plugins/Importers/USDImporter/Source/USDSchemas/Private/USDGeomXformableTranslator.cpp
--- a/Engine/Plugins/Importers/USDImporter/Source/USDSchemas/Private/USDGeomXformableTranslator.cpp
+++ b/Engine/Plugins/Importers/USDImporter/Source/USDSchemas/Private/USDGeomXformableTranslator.cpp
@@ -4,12 +4,10 @@
 
 #include "USDConversionUtils.h"
 #include "USDPrimConversion.h"
-#include "USDSchemasModule.h"
 #include "USDTypesConversion.h"
 
 #include "Components/SceneComponent.h"
 #include "Engine/World.h"
-#include "Modules/ModuleManager.h"
 
 #if USE_USD_SDK
 
@@ -18,11 +16,6 @@
 	#include "pxr/usd/usdGeom/xformable.h"
 #include "USDIncludesEnd.h"
 
-void FUsdGeomXformableTranslator::RegisterTranslator()
-{
-	IUsdSchemasModule& UsdSchemasModule = FModuleManager::Get().LoadModuleChecked< IUsdSchemasModule >( TEXT("USDSchemas") );
-	UsdSchemasModule.GetTranslatorRegistry().Register< FUsdGeomXformableTranslator >( TEXT("UsdGeomXformable") );
-}
 
 FUsdGeomXformableTranslator::FUsdGeomXformableTranslator( TSubclassOf< USceneComponent > InComponentTypeOverride, TSharedRef< FUsdSchemaTranslationContext > InContext, const pxr::UsdTyped& InSchema )
 	: FUsdSchemaTranslator( InContext, InSchema )
diff --git a/Engine/Plugins/Importers/USDImporter/Source/USDSchemas/Private/USDSchemaTranslator.cpp b/Engine/Plugins/Importers/USDImporter/Source/USDSchemas/Private/USDSchemaTranslator.cpp
--- a/Engine/Plugins/Importers/USDImporter/Source/USDSchemas/Private/USDSchemaTranslator.cpp
+++ b/Engine/Plugins/Importers/USDImporter/Source/USDSchemas/Private/USDSchemaTranslator.cpp
@@ -4,6 +4,7 @@
 
 #include "USDTypesConversion.h"
 
+#include "Algo/Find.h"
 #include "Async/Async.h"
 
 #if USE_USD_SDK
@@ -13,49 +14,102 @@
 #include "USDIncludesEnd.h"
 
 
+int32 FRegisteredSchemaTranslatorHandle::CurrentSchemaTranslatorId = 0;
+
 TSharedPtr< FUsdSchemaTranslator > FUsdSchemaTranslatorRegistry::CreateTranslatorForSchema( TSharedRef< FUsdSchemaTranslationContext > InTranslationContext, const pxr::UsdTyped& InSchema )
 {
 	TUsdStore< pxr::UsdPrim > Prim = InSchema.GetPrim();
 
-	for ( TPair< FString, TCreateTranslator >& RegisteredSchema : CreationMethods )
+	for ( TPair< FString, FSchemaTranslatorsStack >& RegisteredSchemasStack : RegisteredSchemaTranslators )
 	{
-		pxr::TfToken RegisteredSchemaToken( UnrealToUsd::ConvertString( *RegisteredSchema.Key ).Get() );
+		pxr::TfToken RegisteredSchemaToken( UnrealToUsd::ConvertString( *RegisteredSchemasStack.Key ).Get() );
 		pxr::TfType RegisteredSchemaType = pxr::UsdSchemaRegistry::GetTypeFromName( RegisteredSchemaToken );
 
-		if ( Prim.Get().IsA( RegisteredSchemaType ) )
+		if ( Prim.Get().IsA( RegisteredSchemaType ) && RegisteredSchemasStack.Value.Num() > 0 )
 		{
-			return RegisteredSchema.Value( InTranslationContext, InSchema );
+			return RegisteredSchemasStack.Value.Top().CreateFunction( InTranslationContext, InSchema );
 		}
 	}
 
 	return {};
 }
 
-void FUsdSchemaTranslatorRegistry::Register( const FString& SchemaName, TCreateTranslator CreateFunction )
+FRegisteredSchemaTranslatorHandle FUsdSchemaTranslatorRegistry::Register( const FString& SchemaName, FCreateTranslator CreateFunction )
 {
-	// Insert most specialized first
-	int32 InsertionIndex = 0;
+	FSchemaTranslatorsStack* SchemaTranslatorsStack = FindSchemaTranslatorStack( SchemaName );
+
+	if ( !SchemaTranslatorsStack )
+	{
+		// Insert most specialized first
+		int32 SchemaRegistryIndex = 0;
+
+		pxr::TfToken SchemaToRegisterToken( UnrealToUsd::ConvertString( *SchemaName ).Get() );
+		pxr::TfType SchemaToRegisterType = pxr::UsdSchemaRegistry::GetTypeFromName( SchemaToRegisterToken );
+
+		for ( TPair< FString, FSchemaTranslatorsStack >& RegisteredSchemasStack : RegisteredSchemaTranslators )
+		{
+			pxr::TfToken RegisteredSchemaToken( UnrealToUsd::ConvertString( *RegisteredSchemasStack.Key ).Get() );
+			pxr::TfType RegisteredSchemaType = pxr::UsdSchemaRegistry::GetTypeFromName( RegisteredSchemaToken );
+
+			if ( SchemaToRegisterType.IsA( RegisteredSchemaType ) )
+			{
+				// We need to be registered before our ancestor types
+				break;
+			}
+			else
+			{
+				++SchemaRegistryIndex;
+			}
+		}
+
+		SchemaTranslatorsStack = &RegisteredSchemaTranslators.EmplaceAt_GetRef( SchemaRegistryIndex, SchemaName, FSchemaTranslatorsStack() ).Value;
+	}
+
+	FRegisteredSchemaTranslator RegisteredSchemaTranslator;
+	RegisteredSchemaTranslator.Handle = FRegisteredSchemaTranslatorHandle( SchemaName );
+	RegisteredSchemaTranslator.CreateFunction = CreateFunction;
+
+	SchemaTranslatorsStack->Push( RegisteredSchemaTranslator );
 
-	pxr::TfToken SchemaToRegisterToken( UnrealToUsd::ConvertString( *SchemaName ).Get() );
-	pxr::TfType SchemaToRegisterType = pxr::UsdSchemaRegistry::GetTypeFromName( SchemaToRegisterToken );
+	return RegisteredSchemaTranslator.Handle;
+}
+
+void FUsdSchemaTranslatorRegistry::Unregister( const FRegisteredSchemaTranslatorHandle& TranslatorHandle )
+{
+	FSchemaTranslatorsStack* SchemaTranslatorsStack = FindSchemaTranslatorStack( TranslatorHandle.GetSchemaName() );
 
-	for ( TPair< FString, TCreateTranslator >& RegisteredSchema : CreationMethods )
+	if ( !SchemaTranslatorsStack )
 	{
-		pxr::TfToken RegisteredSchemaToken( UnrealToUsd::ConvertString( *RegisteredSchema.Key ).Get() );
-		pxr::TfType RegisteredSchemaType = pxr::UsdSchemaRegistry::GetTypeFromName( RegisteredSchemaToken );
+		return;
+	}
 
-		if ( SchemaToRegisterType.IsA( RegisteredSchemaType ) )
+	for ( FSchemaTranslatorsStack::TIterator RegisteredSchemaTranslatorIt = SchemaTranslatorsStack->CreateIterator();
+		RegisteredSchemaTranslatorIt; ++RegisteredSchemaTranslatorIt )
+	{
+		if ( RegisteredSchemaTranslatorIt->Handle.GetId() == TranslatorHandle.GetId() )
 		{
-			// We need to be registered before our ancestor types
+			RegisteredSchemaTranslatorIt.RemoveCurrent();
 			break;
 		}
-		else
+	}
+}
+
+FUsdSchemaTranslatorRegistry::FSchemaTranslatorsStack* FUsdSchemaTranslatorRegistry::FindSchemaTranslatorStack( const FString& SchemaName )
+{
+	TPair< FString, FSchemaTranslatorsStack >* Result = Algo::FindByPredicate( RegisteredSchemaTranslators,
+		[ &SchemaName ]( const TPair< FString, FSchemaTranslatorsStack >& Element ) -> bool
 		{
-			++InsertionIndex;
-		}
+			return Element.Key == SchemaName;
+		} );
+
+	if ( Result )
+	{
+		return &Result->Value;
+	}
+	else
+	{
+		return nullptr;
 	}
-	
-	CreationMethods.EmplaceAt( InsertionIndex, SchemaName, CreateFunction );
 }
 
 void FUsdSchemaTranslationContext::CompleteTasks()
diff --git a/Engine/Plugins/Importers/USDImporter/Source/USDSchemas/Private/USDSchemasModule.cpp b/Engine/Plugins/Importers/USDImporter/Source/USDSchemas/Private/USDSchemasModule.cpp
--- a/Engine/Plugins/Importers/USDImporter/Source/USDSchemas/Private/USDSchemasModule.cpp
+++ b/Engine/Plugins/Importers/USDImporter/Source/USDSchemas/Private/USDSchemasModule.cpp
@@ -9,6 +9,7 @@
 #include "USDGeomXFormableTranslator.h"
 #include "USDMemory.h"
 #include "USDSchemaTranslator.h"
+#include "USDShadeMaterialTranslator.h"
 #include "USDSkelRootTranslator.h"
 #endif // #if USE_USD_SDK
 
@@ -19,11 +20,24 @@ public:
 	{
 #if USE_USD_SDK
 		// Register the default translators
-		FUsdGeomCameraTranslator::RegisterTranslator();
-		FUsdGeomMeshTranslator::RegisterTranslator();
-		FUsdGeomPointInstancerTranslator::RegisterTranslator();
-		FUsdSkelRootTranslator::RegisterTranslator();
-		FUsdGeomXformableTranslator::RegisterTranslator();
+		UsdGeomCameraTranslatorHandle = GetTranslatorRegistry().Register< FUsdGeomCameraTranslator >( TEXT("UsdGeomCamera") );
+		UsdGeomMeshTranslatorHandle = GetTranslatorRegistry().Register< FUsdGeomMeshTranslator >( TEXT("UsdGeomMesh") );
+		UsdGeomPointInstancerTranslatorHandle = GetTranslatorRegistry().Register< FUsdGeomPointInstancerTranslator >( TEXT("UsdGeomPointInstancer") );
+		UsdSkelRootTranslatorHandle = GetTranslatorRegistry().Register< FUsdSkelRootTranslator >( TEXT("UsdSkelRoot") );
+		UsdGeomXformableTranslatorHandle = GetTranslatorRegistry().Register< FUsdGeomXformableTranslator >( TEXT("UsdGeomXformable") );
+		UsdShadeMaterialTranslatorHandle = GetTranslatorRegistry().Register< FUsdShadeMaterialTranslator >( TEXT("UsdShadeMaterial") );
+#endif // #if USE_USD_SDK
+	}
+
+	virtual void ShutdownModule() override
+	{
+#if USE_USD_SDK
+		GetTranslatorRegistry().Unregister( UsdGeomCameraTranslatorHandle );
+		GetTranslatorRegistry().Unregister( UsdGeomMeshTranslatorHandle );
+		GetTranslatorRegistry().Unregister( UsdGeomPointInstancerTranslatorHandle );
+		GetTranslatorRegistry().Unregister( UsdSkelRootTranslatorHandle );
+		GetTranslatorRegistry().Unregister( UsdGeomXformableTranslatorHandle );
+		GetTranslatorRegistry().Unregister( UsdShadeMaterialTranslatorHandle );
 #endif // #if USE_USD_SDK
 	}
 
@@ -34,6 +48,13 @@ public:
 
 protected:
 	FUsdSchemaTranslatorRegistry UsdSchemaTranslatorRegistry;
+
+	FRegisteredSchemaTranslatorHandle UsdGeomCameraTranslatorHandle;
+	FRegisteredSchemaTranslatorHandle UsdGeomMeshTranslatorHandle;
+	FRegisteredSchemaTranslatorHandle UsdGeomPointInstancerTranslatorHandle;
+	FRegisteredSchemaTranslatorHandle UsdSkelRootTranslatorHandle;
+	FRegisteredSchemaTranslatorHandle UsdGeomXformableTranslatorHandle;
+	FRegisteredSchemaTranslatorHandle UsdShadeMaterialTranslatorHandle;
 };
 
 IMPLEMENT_MODULE_USD( FUsdSchemasModule, USDSchemas );
diff --git a/Engine/Plugins/Importers/USDImporter/Source/USDSchemas/Private/USDShadeMaterialTranslator.cpp b/Engine/Plugins/Importers/USDImporter/Source/USDSchemas/Private/USDShadeMaterialTranslator.cpp
new file mode 100644
--- /dev/null
+++ b/Engine/Plugins/Importers/USDImporter/Source/USDSchemas/Private/USDShadeMaterialTranslator.cpp
@@ -0,0 +1,47 @@
+// Copyright 1998-2019 Epic Games, Inc. All Rights Reserved.
+
+#include "USDShadeMaterialTranslator.h"
+
+#include "USDGeomMeshConversion.h"
+#include "USDTypesConversion.h"
+
+#include "Materials/Material.h"
+
+#include "USDIncludesStart.h"
+	#include "pxr/usd/usdShade/material.h"
+#include "USDIncludesEnd.h"
+
+
+void FUsdShadeMaterialTranslator::CreateAssets()
+{
+	pxr::UsdShadeMaterial ShadeMaterial( Schema.Get() );
+
+	if ( !ShadeMaterial )
+	{
+		return;
+	}
+
+	UObject*& CachedMaterial = Context->AssetsCache.FindOrAdd( UsdToUnreal::ConvertPath( ShadeMaterial.GetPath() ) );
+
+	if ( !CachedMaterial )
+	{
+		UMaterial* NewMaterial = NewObject< UMaterial >();
+
+		if ( UsdToUnreal::ConvertMaterial( ShadeMaterial, *NewMaterial ) )
+		{
+			//UMaterialEditingLibrary::RecompileMaterial( CachedMaterial ); // Too slow
+			NewMaterial->PostEditChange();
+		}
+		else
+		{
+			NewMaterial = nullptr;
+		}
+
+		CachedMaterial = NewMaterial;
+	}
+
+	FScopeLock Lock( &Context->CriticalSection );
+	{
+		Context->PrimPathsToAssets.Add( UsdToUnreal::ConvertPath( Schema.Get().GetPath() ), CachedMaterial );
+	}
+}
diff --git a/Engine/Plugins/Importers/USDImporter/Source/USDSchemas/Private/USDSkelRootTranslator.cpp b/Engine/Plugins/Importers/USDImporter/Source/USDSchemas/Private/USDSkelRootTranslator.cpp
--- a/Engine/Plugins/Importers/USDImporter/Source/USDSchemas/Private/USDSkelRootTranslator.cpp
+++ b/Engine/Plugins/Importers/USDImporter/Source/USDSchemas/Private/USDSkelRootTranslator.cpp
@@ -6,7 +6,6 @@
 
 #include "USDGeomMeshConversion.h"
 #include "USDMemory.h"
-#include "USDSchemasModule.h"
 #include "USDSkeletalDataConversion.h"
 #include "USDTypesConversion.h"
 
@@ -15,7 +14,6 @@
 #include "Engine/SkeletalMesh.h"
 #include "Materials/Material.h"
 #include "Materials/MaterialInstanceConstant.h"
-#include "Modules/ModuleManager.h"
 #include "Rendering/SkeletalMeshLODImporterData.h"
 
 #include "USDIncludesStart.h"
@@ -30,7 +28,7 @@
 
 namespace UsdSkelRootTranslatorImpl
 {
-	void ProcessMaterials( const pxr::UsdPrim& UsdPrim, FSkeletalMeshImportData& SkelMeshImportData, TMap< FString, UObject* >& MaterialsCache, bool bHasPrimDisplayColor, float Time )
+	void ProcessMaterials( const pxr::UsdPrim& UsdPrim, FSkeletalMeshImportData& SkelMeshImportData, TMap< FString, UObject* >& PrimPathsToAssets, bool bHasPrimDisplayColor, float Time )
 	{
 		for ( SkeletalMeshImportData::FMaterial& ImportedMaterial : SkelMeshImportData.Materials )
 		{
@@ -41,31 +39,12 @@ namespace UsdSkelRootTranslatorImpl
 
 				if ( MaterialPrim.Get() )
 				{
-					UObject*& CachedMaterial = MaterialsCache.FindOrAdd( UsdToUnreal::ConvertPath( MaterialPrim.Get().GetPrimPath() ) );
+					Material = Cast< UMaterialInterface >( PrimPathsToAssets.FindRef( UsdToUnreal::ConvertPath( MaterialPrim.Get().GetPrimPath() ) ) );
 
-					if ( !CachedMaterial )
+					if ( Material )
 					{
-						UMaterial* NewMaterial = NewObject< UMaterial >();
-
-						if ( UsdToUnreal::ConvertMaterial( pxr::UsdShadeMaterial( MaterialPrim.Get() ), *NewMaterial ) )
-						{
-							NewMaterial->bUsedWithSkeletalMesh = true;
-							bool bNeedsRecompile = false;
-							NewMaterial->GetMaterial()->SetMaterialUsage(bNeedsRecompile, MATUSAGE_SkeletalMesh);
-
-							NewMaterial->PostEditChange();
-						}
-						else
-						{
-							NewMaterial = nullptr;
-						}
-
-						Material = NewMaterial;
-						CachedMaterial = Material;
-					}
-					else
-					{
-						Material = Cast< UMaterialInterface >( CachedMaterial );
+						bool bNeedsRecompile = false;
+						Material->GetMaterial()->SetMaterialUsage( bNeedsRecompile, MATUSAGE_SkeletalMesh );
 					}
 				}
 
@@ -104,12 +83,6 @@ namespace UsdSkelRootTranslatorImpl
 	}
 }
 
-void FUsdSkelRootTranslator::RegisterTranslator()
-{
-	IUsdSchemasModule& UsdSchemasModule = FModuleManager::Get().LoadModuleChecked< IUsdSchemasModule >( TEXT("USDSchemas") );
-	UsdSchemasModule.GetTranslatorRegistry().Register< FUsdSkelRootTranslator >( TEXT("UsdSkelRoot") );
-}
-
 void FUsdSkelRootTranslator::CreateAssets()
 {
 	TRACE_CPUPROFILER_EVENT_SCOPE( FUsdSkelRootTranslator::CreateAssets );
@@ -170,7 +143,7 @@ void FUsdSkelRootTranslator::CreateAssets()
 				}
 			}
 
-			UsdSkelRootTranslatorImpl::ProcessMaterials( Prim, SkelMeshImportData, Context->AssetsCache, bHasPrimDisplayColor, Context->Time );
+			UsdSkelRootTranslatorImpl::ProcessMaterials( Prim, SkelMeshImportData, Context->PrimPathsToAssets, bHasPrimDisplayColor, Context->Time );
 
 			break;
 		}
diff --git a/Engine/Plugins/Importers/USDImporter/Source/USDSchemas/Public/USDGeomCameraTranslator.h b/Engine/Plugins/Importers/USDImporter/Source/USDSchemas/Public/USDGeomCameraTranslator.h
--- a/Engine/Plugins/Importers/USDImporter/Source/USDSchemas/Public/USDGeomCameraTranslator.h
+++ b/Engine/Plugins/Importers/USDImporter/Source/USDSchemas/Public/USDGeomCameraTranslator.h
@@ -16,9 +16,6 @@ PXR_NAMESPACE_CLOSE_SCOPE
 
 class USDSCHEMAS_API FUsdGeomCameraTranslator : public FUsdGeomXformableTranslator
 {
-public:
-	static void RegisterTranslator();
-
 public:
 	using FUsdGeomXformableTranslator::FUsdGeomXformableTranslator;
 
diff --git a/Engine/Plugins/Importers/USDImporter/Source/USDSchemas/Public/USDGeomMeshTranslator.h b/Engine/Plugins/Importers/USDImporter/Source/USDSchemas/Public/USDGeomMeshTranslator.h
--- a/Engine/Plugins/Importers/USDImporter/Source/USDSchemas/Public/USDGeomMeshTranslator.h
+++ b/Engine/Plugins/Importers/USDImporter/Source/USDSchemas/Public/USDGeomMeshTranslator.h
@@ -43,9 +43,6 @@ protected:
 
 class USDSCHEMAS_API FUsdGeomMeshTranslator : public FUsdGeomXformableTranslator
 {
-public:
-	static void RegisterTranslator();
-
 public:
 	using FUsdGeomXformableTranslator::FUsdGeomXformableTranslator;
 
diff --git a/Engine/Plugins/Importers/USDImporter/Source/USDSchemas/Public/USDGeomPointInstancerTranslator.h b/Engine/Plugins/Importers/USDImporter/Source/USDSchemas/Public/USDGeomPointInstancerTranslator.h
--- a/Engine/Plugins/Importers/USDImporter/Source/USDSchemas/Public/USDGeomPointInstancerTranslator.h
+++ b/Engine/Plugins/Importers/USDImporter/Source/USDSchemas/Public/USDGeomPointInstancerTranslator.h
@@ -15,9 +15,6 @@ struct FUsdSchemaTranslationContext;
 
 class USDSCHEMAS_API FUsdGeomPointInstancerTranslator : public FUsdGeomXformableTranslator
 {
-public:
-	static void RegisterTranslator();
-
 public:
 	using FUsdGeomXformableTranslator::FUsdGeomXformableTranslator;
 	
diff --git a/Engine/Plugins/Importers/USDImporter/Source/USDSchemas/Public/USDGeomXformableTranslator.h b/Engine/Plugins/Importers/USDImporter/Source/USDSchemas/Public/USDGeomXformableTranslator.h
--- a/Engine/Plugins/Importers/USDImporter/Source/USDSchemas/Public/USDGeomXformableTranslator.h
+++ b/Engine/Plugins/Importers/USDImporter/Source/USDSchemas/Public/USDGeomXformableTranslator.h
@@ -18,9 +18,6 @@ PXR_NAMESPACE_CLOSE_SCOPE
 
 class USDSCHEMAS_API FUsdGeomXformableTranslator : public FUsdSchemaTranslator
 {
-public:
-	static void RegisterTranslator();
-
 public:
 	using FUsdSchemaTranslator::FUsdSchemaTranslator;
 
diff --git a/Engine/Plugins/Importers/USDImporter/Source/USDSchemas/Public/USDSchemaTranslator.h b/Engine/Plugins/Importers/USDImporter/Source/USDSchemas/Public/USDSchemaTranslator.h
--- a/Engine/Plugins/Importers/USDImporter/Source/USDSchemas/Public/USDSchemaTranslator.h
+++ b/Engine/Plugins/Importers/USDImporter/Source/USDSchemas/Public/USDSchemaTranslator.h
@@ -34,16 +34,53 @@ PXR_NAMESPACE_CLOSE_SCOPE
 
 #endif //#if USE_USD_SDK
 
+class USDSCHEMAS_API FRegisteredSchemaTranslatorHandle
+{
+public:
+	FRegisteredSchemaTranslatorHandle()
+		: Id( CurrentSchemaTranslatorId++ )
+	{
+	}
+
+	explicit FRegisteredSchemaTranslatorHandle( const FString& InSchemaName )
+		: FRegisteredSchemaTranslatorHandle()
+	{
+		SchemaName = InSchemaName;
+	}
+
+	int32 GetId() const { return Id; }
+	void SetId( int32 InId ) { Id = InId; }
+
+	const FString& GetSchemaName() const { return SchemaName; }
+	void SetSchemaName( const FString& InSchemaName ) { SchemaName = InSchemaName; }
+
+private:
+	static int32 CurrentSchemaTranslatorId;
+
+	FString SchemaName;
+	int32 Id;
+};
+
+class FRegisteredSchemaTranslator
+{
+	using FCreateTranslator = TFunction< TSharedRef< FUsdSchemaTranslator >( TSharedRef< FUsdSchemaTranslationContext >, const pxr::UsdTyped& ) >;
+
+public:
+	FRegisteredSchemaTranslatorHandle Handle;
+	FCreateTranslator CreateFunction;
+};
+
 class USDSCHEMAS_API FUsdSchemaTranslatorRegistry
 {
 #if USE_USD_SDK
-	using TCreateTranslator = TFunction< TSharedRef< FUsdSchemaTranslator >( TSharedRef< FUsdSchemaTranslationContext >, const pxr::UsdTyped& ) >;
+	using FCreateTranslator = TFunction< TSharedRef< FUsdSchemaTranslator >( TSharedRef< FUsdSchemaTranslationContext >, const pxr::UsdTyped& ) >;
+	using FSchemaTranslatorsStack = TArray< FRegisteredSchemaTranslator, TInlineAllocator< 1 > >;
 
 public:
 	TSharedPtr< FUsdSchemaTranslator > CreateTranslatorForSchema( TSharedRef< FUsdSchemaTranslationContext > InTranslationContext, const pxr::UsdTyped& InSchema );
 
 	template< typename TSchemaTranslator >
-	void Register( const FString& SchemaName )
+	FRegisteredSchemaTranslatorHandle Register( const FString& SchemaName )
 	{
 		auto CreateSchemaTranslator =
 		[]( TSharedRef< FUsdSchemaTranslationContext > InContext, const pxr::UsdTyped& InSchema ) -> TSharedRef< FUsdSchemaTranslator >
@@ -51,13 +88,17 @@ public:
 			return MakeShared< TSchemaTranslator >( InContext, InSchema );
 		};
 
-		Register( SchemaName, CreateSchemaTranslator );
+		return Register( SchemaName, CreateSchemaTranslator );
 	}
 
+	void Unregister( const FRegisteredSchemaTranslatorHandle& TranslatorHandle );
+
 protected:
-	void Register( const FString& SchemaName, TCreateTranslator CreateFunction );
+	FRegisteredSchemaTranslatorHandle Register( const FString& SchemaName, FCreateTranslator CreateFunction );
+
+	FSchemaTranslatorsStack* FindSchemaTranslatorStack( const FString& SchemaName );
 
-	TArray< TPair< FString, TCreateTranslator > > CreationMethods;
+	TArray< TPair< FString, FSchemaTranslatorsStack > > RegisteredSchemaTranslators;
 #endif // #if USE_USD_SDK
 };
 
diff --git a/Engine/Plugins/Importers/USDImporter/Source/USDSchemas/Public/USDShadeMaterialTranslator.h b/Engine/Plugins/Importers/USDImporter/Source/USDSchemas/Public/USDShadeMaterialTranslator.h
new file mode 100644
--- /dev/null
+++ b/Engine/Plugins/Importers/USDImporter/Source/USDSchemas/Public/USDShadeMaterialTranslator.h
@@ -0,0 +1,22 @@
+// Copyright 1998-2019 Epic Games, Inc. All Rights Reserved.
+
+#pragma once
+
+#include "USDSchemaTranslator.h"
+
+#if USE_USD_SDK
+
+#include "USDIncludesStart.h"
+	#include "pxr/pxr.h"
+#include "USDIncludesEnd.h"
+
+class USDSCHEMAS_API FUsdShadeMaterialTranslator : public FUsdSchemaTranslator
+{
+public:
+	using FUsdSchemaTranslator::FUsdSchemaTranslator;
+
+	virtual void CreateAssets() override;
+
+};
+
+#endif // #if USE_USD_SDK
diff --git a/Engine/Plugins/Importers/USDImporter/Source/USDSchemas/Public/USDSkelRootTranslator.h b/Engine/Plugins/Importers/USDImporter/Source/USDSchemas/Public/USDSkelRootTranslator.h
--- a/Engine/Plugins/Importers/USDImporter/Source/USDSchemas/Public/USDSkelRootTranslator.h
+++ b/Engine/Plugins/Importers/USDImporter/Source/USDSchemas/Public/USDSkelRootTranslator.h
@@ -20,12 +20,8 @@ class USDSCHEMAS_API FUsdSkelRootTranslator : public FUsdGeomXformableTranslator
 {
 	using Super = FUsdGeomXformableTranslator;
 
-public:
-	static void RegisterTranslator();
-
 public:
 	using FUsdGeomXformableTranslator::FUsdGeomXformableTranslator;
-	//explicit FUsdSkelRootTranslator( FUsdSchemaTranslationContext& InContext, const pxr::UsdSkelRoot& SkelRoot );
 	
 	virtual void CreateAssets() override;
 	virtual USceneComponent* CreateComponents() override;
diff --git a/Engine/Plugins/Importers/USDImporter/Source/USDStage/Private/USDStageActor.cpp b/Engine/Plugins/Importers/USDImporter/Source/USDStage/Private/USDStageActor.cpp
--- a/Engine/Plugins/Importers/USDImporter/Source/USDStage/Private/USDStageActor.cpp
+++ b/Engine/Plugins/Importers/USDImporter/Source/USDStage/Private/USDStageActor.cpp
@@ -633,20 +633,31 @@ bool AUsdStageActor::HasAutorithyOverStage() const
 #if USE_USD_SDK
 void AUsdStageActor::LoadAssets( FUsdSchemaTranslationContext& TranslationContext, const pxr::UsdPrim& StartPrim )
 {
-	TArray< TUsdStore< pxr::UsdPrim > > AllPrimAssets = UsdUtils::GetAllPrimsOfType( StartPrim, pxr::TfType::Find< pxr::UsdGeomMesh >(), { pxr::TfType::Find< pxr::UsdSkelRoot >() } );
-	AllPrimAssets.Append( UsdUtils::GetAllPrimsOfType( StartPrim, pxr::TfType::Find< pxr::UsdSkelRoot >() ) );
-
-	IUsdSchemasModule& UsdSchemasModule = FModuleManager::Get().LoadModuleChecked< IUsdSchemasModule >( TEXT("USDSchemas") );
-	
-	for ( TUsdStore< pxr::UsdPrim >& UsdPrim : AllPrimAssets )
+	auto CreateAssetsForPrims = [ &TranslationContext ]( const TArray< TUsdStore< pxr::UsdPrim > >& AllPrimAssets )
 	{
-		if ( TSharedPtr< FUsdSchemaTranslator > SchemaTranslator = UsdSchemasModule.GetTranslatorRegistry().CreateTranslatorForSchema( TranslationContext.AsShared(), pxr::UsdTyped( UsdPrim.Get() ) ) )
+		IUsdSchemasModule& UsdSchemasModule = FModuleManager::Get().LoadModuleChecked< IUsdSchemasModule >( TEXT("USDSchemas") );
+
+		for ( const TUsdStore< pxr::UsdPrim >& UsdPrim : AllPrimAssets )
 		{
-			SchemaTranslator->CreateAssets();
+			if ( TSharedPtr< FUsdSchemaTranslator > SchemaTranslator = UsdSchemasModule.GetTranslatorRegistry().CreateTranslatorForSchema( TranslationContext.AsShared(), pxr::UsdTyped( UsdPrim.Get() ) ) )
+			{
+				SchemaTranslator->CreateAssets();
+			}
 		}
-	}
 
-	TranslationContext.CompleteTasks(); // Finish the assets tasks before moving on
+		TranslationContext.CompleteTasks(); // Finish the assets tasks before moving on
+	};
+
+	// Load materials first since meshes are referencing them
+	TArray< TUsdStore< pxr::UsdPrim > > AllPrimAssets = UsdUtils::GetAllPrimsOfType( StartPrim, pxr::TfType::Find< pxr::UsdShadeMaterial >() );
+	
+	CreateAssetsForPrims( AllPrimAssets );
+
+	// Load meshes
+	AllPrimAssets = UsdUtils::GetAllPrimsOfType( StartPrim, pxr::TfType::Find< pxr::UsdGeomMesh >(), { pxr::TfType::Find< pxr::UsdSkelRoot >() } );
+	AllPrimAssets.Append( UsdUtils::GetAllPrimsOfType( StartPrim, pxr::TfType::Find< pxr::UsdSkelRoot >() ) );
+	
+	CreateAssetsForPrims( AllPrimAssets );
 }
 
 void AUsdStageActor::AnimatePrims()
