Build fix for Clang Win64

#rnx

diff --git a/Engine/Plugins/Importers/USDImporter/Source/USDSchemas/Public/USDSchemaTranslator.h b/Engine/Plugins/Importers/USDImporter/Source/USDSchemas/Public/USDSchemaTranslator.h
--- a/Engine/Plugins/Importers/USDImporter/Source/USDSchemas/Public/USDSchemaTranslator.h
+++ b/Engine/Plugins/Importers/USDImporter/Source/USDSchemas/Public/USDSchemaTranslator.h
@@ -19,6 +19,7 @@
 	#include "pxr/usd/usd/typed.h"
 #include "USDIncludesEnd.h"
 
+class FRegisteredSchemaTranslator;
 struct FUsdSchemaTranslationContext;
 class FUsdSchemaTranslator;
 class FUsdSchemaTranslatorTaskChain;
@@ -61,15 +62,6 @@ private:
 	int32 Id;
 };
 
-class FRegisteredSchemaTranslator
-{
-	using FCreateTranslator = TFunction< TSharedRef< FUsdSchemaTranslator >( TSharedRef< FUsdSchemaTranslationContext >, const pxr::UsdTyped& ) >;
-
-public:
-	FRegisteredSchemaTranslatorHandle Handle;
-	FCreateTranslator CreateFunction;
-};
-
 class USDSCHEMAS_API FUsdSchemaTranslatorRegistry
 {
 #if USE_USD_SDK
@@ -104,6 +96,15 @@ protected:
 
 #if USE_USD_SDK
 
+class FRegisteredSchemaTranslator
+{
+	using FCreateTranslator = TFunction< TSharedRef< FUsdSchemaTranslator >( TSharedRef< FUsdSchemaTranslationContext >, const pxr::UsdTyped& ) >;
+
+public:
+	FRegisteredSchemaTranslatorHandle Handle;
+	FCreateTranslator CreateFunction;
+};
+
 struct USDSCHEMAS_API FUsdSchemaTranslationContext : public TSharedFromThis< FUsdSchemaTranslationContext >
 {
 	explicit FUsdSchemaTranslationContext( TMap< FString, UObject* >& InPrimPathsToAssets, TMap< FString, UObject* >& InAssetsCache )
