Only include USD file if SDK is available #jira UE-85522 #rb JeanMichel.Dignard

diff --git a/Engine/Plugins/Importers/USDImporter/Source/UnrealUSDBinding/Private/UnrealUSDBinding.cpp b/Engine/Plugins/Importers/USDImporter/Source/UnrealUSDBinding/Private/UnrealUSDBinding.cpp
--- a/Engine/Plugins/Importers/USDImporter/Source/UnrealUSDBinding/Private/UnrealUSDBinding.cpp
+++ b/Engine/Plugins/Importers/USDImporter/Source/UnrealUSDBinding/Private/UnrealUSDBinding.cpp
@@ -1,5 +1,7 @@
 // Copyright 1998-2019 Epic Games, Inc. All Rights Reserved.
 
+#if USE_USD_SDK
+
 #include "UnrealUSDWrapper.h"
 
 #include "USDIncludesStart.h"
@@ -15,4 +17,6 @@ BOOST_PYTHON_MODULE(UnrealUSDBinding)
 	// Note: Only use this if you can guarantee the lifetime of the C++ object (in this case its static)
 	// See https://www.boost.org/doc/libs/1_40_0/libs/python/doc/v2/reference_existing_object.html
 	def("GetStageCache", UnrealUSDWrapper::GetUsdStageCache, return_value_policy<reference_existing_object>());
-}
\ No newline at end of file
+}
+
+#endif // #if USE_USD_SDK
\ No newline at end of file
