build fix for ClangEditor Win64

#rnx

diff --git a/Engine/Plugins/Importers/USDImporter/Source/USDSchemas/Private/USDSchemasModule.cpp b/Engine/Plugins/Importers/USDImporter/Source/USDSchemas/Private/USDSchemasModule.cpp
--- a/Engine/Plugins/Importers/USDImporter/Source/USDSchemas/Private/USDSchemasModule.cpp
+++ b/Engine/Plugins/Importers/USDImporter/Source/USDSchemas/Private/USDSchemasModule.cpp
@@ -2,13 +2,14 @@
 
 #include "USDSchemasModule.h"
 
+#include "USDSchemaTranslator.h"
+
 #if USE_USD_SDK
 #include "USDGeomCameraTranslator.h"
 #include "USDGeomMeshTranslator.h"
 #include "USDGeomPointInstancerTranslator.h"
-#include "USDGeomXFormableTranslator.h"
+#include "USDGeomXformableTranslator.h"
 #include "USDMemory.h"
-#include "USDSchemaTranslator.h"
 #include "USDShadeMaterialTranslator.h"
 #include "USDSkelRootTranslator.h"
 #endif // #if USE_USD_SDK
diff --git a/Engine/Plugins/Importers/USDImporter/Source/USDSchemas/Private/USDShadeMaterialTranslator.cpp b/Engine/Plugins/Importers/USDImporter/Source/USDSchemas/Private/USDShadeMaterialTranslator.cpp
--- a/Engine/Plugins/Importers/USDImporter/Source/USDSchemas/Private/USDShadeMaterialTranslator.cpp
+++ b/Engine/Plugins/Importers/USDImporter/Source/USDSchemas/Private/USDShadeMaterialTranslator.cpp
@@ -7,6 +7,8 @@
 
 #include "Materials/Material.h"
 
+#if USE_USD_SDK
+
 #include "USDIncludesStart.h"
 	#include "pxr/usd/usdShade/material.h"
 #include "USDIncludesEnd.h"
@@ -45,3 +47,5 @@ void FUsdShadeMaterialTranslator::CreateAssets()
 		Context->PrimPathsToAssets.Add( UsdToUnreal::ConvertPath( Schema.Get().GetPath() ), CachedMaterial );
 	}
 }
+
+#endif // #if USE_USD_SDK
diff --git a/Engine/Plugins/Importers/USDImporter/Source/USDStage/Private/USDStageActor.cpp b/Engine/Plugins/Importers/USDImporter/Source/USDStage/Private/USDStageActor.cpp
--- a/Engine/Plugins/Importers/USDImporter/Source/USDStage/Private/USDStageActor.cpp
+++ b/Engine/Plugins/Importers/USDImporter/Source/USDStage/Private/USDStageActor.cpp
@@ -4,7 +4,7 @@
 
 #include "USDConversionUtils.h"
 #include "USDGeomMeshTranslator.h"
-#include "USDGeomXFormableTranslator.h"
+#include "USDGeomXformableTranslator.h"
 #include "USDListener.h"
 #include "USDPrimConversion.h"
 #include "USDSchemasModule.h"
