Fix for USD Prim Info refresh
- Don't change the prim info prim on every prim change, only when it's the selected one.

diff --git a/Engine/Plugins/Importers/USDImporter/Source/USDStageEditor/Private/SUSDStage.cpp b/Engine/Plugins/Importers/USDImporter/Source/USDStageEditor/Private/SUSDStage.cpp
--- a/Engine/Plugins/Importers/USDImporter/Source/USDStageEditor/Private/SUSDStage.cpp
+++ b/Engine/Plugins/Importers/USDImporter/Source/USDStageEditor/Private/SUSDStage.cpp
@@ -160,7 +160,7 @@ void SUsdStage::SetupStageActorDelegates()
 					this->UsdStageTreeView->RefreshPrim( PrimPath, bResync );
 				}
 
-				if ( this->UsdPrimInfoWidget && UsdStageActor.IsValid() )
+				if ( this->UsdPrimInfoWidget && UsdStageActor.IsValid() && SelectedPrimPath.Equals( PrimPath, ESearchCase::IgnoreCase ) )
 				{
 					this->UsdPrimInfoWidget->SetPrimPath( UsdStageActor->GetUsdStage(), *PrimPath );
 				}
@@ -480,6 +480,7 @@ void SUsdStage::OnPrimSelected( FString PrimPath )
 			UsdStage = UsdStageActor->GetUsdStage();
 		}
 
+		SelectedPrimPath = PrimPath;
 		UsdPrimInfoWidget->SetPrimPath( UsdStage, *PrimPath );
 	}
 }
diff --git a/Engine/Plugins/Importers/USDImporter/Source/USDStageEditor/Private/SUSDStage.h b/Engine/Plugins/Importers/USDImporter/Source/USDStageEditor/Private/SUSDStage.h
--- a/Engine/Plugins/Importers/USDImporter/Source/USDStageEditor/Private/SUSDStage.h
+++ b/Engine/Plugins/Importers/USDImporter/Source/USDStageEditor/Private/SUSDStage.h
@@ -63,6 +63,8 @@ protected:
 	FDelegateHandle OnStageChangedHandle;
 	FDelegateHandle OnStageEditTargetChangedHandle;
 	FDelegateHandle OnPrimChangedHandle;
+
+	FString SelectedPrimPath;
 };
 
 #endif // #if USE_USD_SDK
diff --git a/Engine/Plugins/Importers/USDImporter/Source/USDStageEditor/Private/SUSDStageTreeView.cpp b/Engine/Plugins/Importers/USDImporter/Source/USDStageEditor/Private/SUSDStageTreeView.cpp
--- a/Engine/Plugins/Importers/USDImporter/Source/USDStageEditor/Private/SUSDStageTreeView.cpp
+++ b/Engine/Plugins/Importers/USDImporter/Source/USDStageEditor/Private/SUSDStageTreeView.cpp
@@ -409,12 +409,14 @@ void SUsdStageTreeView::Construct( const FArguments& InArgs, AUsdStageActor* InU
 
 	OnSelectionChanged = FOnSelectionChanged::CreateLambda( [ this ]( FUsdStageTreeItemPtr UsdStageTreeItem, ESelectInfo::Type SelectionType )
 	{
-		if ( !UsdStageTreeItem )
+		FString SelectedPrimPath;
+
+		if ( UsdStageTreeItem )
 		{
-			return;
+			SelectedPrimPath = UsdToUnreal::ConvertPath( UsdStageTreeItem->UsdPrim.Get().GetPrimPath() );
 		}
 
-		this->OnPrimSelected.ExecuteIfBound( UsdToUnreal::ConvertPath( UsdStageTreeItem->UsdPrim.Get().GetPrimPath() ) );
+		this->OnPrimSelected.ExecuteIfBound( SelectedPrimPath );
 	} );
 
 	OnPrimSelected = InArgs._OnPrimSelected;
