Fixed issue where a stage actor referencing an inexisting file would create a new stage from it. This used to be how we would create a new stage but users should now go through USD Stage Editor -> File -> New.

#rnx

diff --git a/Engine/Plugins/Importers/USDImporter/Source/USDStage/Private/USDStageActor.cpp b/Engine/Plugins/Importers/USDImporter/Source/USDStage/Private/USDStageActor.cpp
--- a/Engine/Plugins/Importers/USDImporter/Source/USDStage/Private/USDStageActor.cpp
+++ b/Engine/Plugins/Importers/USDImporter/Source/USDStage/Private/USDStageActor.cpp
@@ -367,23 +367,6 @@ void AUsdStageActor::OpenUsdStage()
 		pxr::UsdStageCacheContext UsdStageCacheContext( UnrealUSDWrapper::GetUsdStageCache() );
 		UsdStageStore = pxr::UsdStage::Open( UnrealToUsd::ConvertString( *RootLayer.FilePath ).Get(), pxr::UsdStage::InitialLoadSet( InitialLoadSet ) );
 	}
-	else
-	{
-		FScopedUsdAllocs UsdAllocs;
-
-		pxr::UsdStageCacheContext UsdStageCacheContext( UnrealUSDWrapper::GetUsdStageCache() );
-		UsdStageStore = pxr::UsdStage::CreateNew( UnrealToUsd::ConvertString( *RootLayer.FilePath ).Get() );
-
-		// Create default prim
-		pxr::UsdGeomXform RootPrim = pxr::UsdGeomXform::Define( UsdStageStore.Get(), UnrealToUsd::ConvertPath( TEXT("/Root") ).Get() );
-		pxr::UsdModelAPI( RootPrim ).SetKind( pxr::TfToken("component") );
-		
-		// Set default prim
-		UsdStageStore.Get()->SetDefaultPrim( RootPrim.GetPrim() );
-
-		// Set up axis
-		UsdUtils::SetUsdStageAxis( UsdStageStore.Get(), pxr::UsdGeomTokens->z );
-	}
 
 	if ( UsdStageStore.Get() )
 	{
