Fixed USD Sublayers relative paths

#rnx

diff --git a/Engine/Plugins/Importers/USDImporter/Source/USDStageEditor/Private/SUSDLayersTreeView.cpp b/Engine/Plugins/Importers/USDImporter/Source/USDStageEditor/Private/SUSDLayersTreeView.cpp
--- a/Engine/Plugins/Importers/USDImporter/Source/USDStageEditor/Private/SUSDLayersTreeView.cpp
+++ b/Engine/Plugins/Importers/USDImporter/Source/USDStageEditor/Private/SUSDLayersTreeView.cpp
@@ -32,6 +32,29 @@
 
 #define LOCTEXT_NAMESPACE "SUSDLayersTreeView"
 
+namespace UsdLayersTreeViewImpl
+{
+	bool InsertSubLayer( const TUsdStore< pxr::SdfLayerHandle >& LayerHandle, const FString& SubLayerFile )
+	{
+		if ( LayerHandle.Get() )
+		{
+			FScopedUsdAllocs UsdAllocs;
+
+			std::string UsdLayerFilePath = LayerHandle.Get()->GetRealPath();
+			FString LayerFilePath = UsdToUnreal::ConvertString( UsdLayerFilePath );
+
+			FString SubLayerFilePath = FPaths::ConvertRelativePathToFull( SubLayerFile );
+			FPaths::MakePathRelativeTo( SubLayerFilePath, *LayerFilePath );
+
+			LayerHandle.Get()->InsertSubLayerPath( UnrealToUsd::ConvertString( *SubLayerFilePath ).Get() );
+
+			return true;
+		}
+
+		return false;
+	}
+}
+
 struct FUsdLayerData : public TSharedFromThis< FUsdLayerData >
 {
 	FText GetDisplayName() const { return DisplayName; }
@@ -492,19 +515,7 @@ void SUsdLayersTreeView::OnAddSubLayer()
 
 	for ( FUsdLayersTreeItemRef SelectedItem : MySelectedItems )
 	{
-		FScopedUsdAllocs UsdAllocs;
-
-		pxr::SdfLayerHandle LayerHandle = SelectedItem->GetLayerHandle().Get();
-		if ( LayerHandle )
-		{
-			std::string UsdLayerFilePath = LayerHandle->GetRealPath();
-			FString LayerFilePath = UsdToUnreal::ConvertString( UsdLayerFilePath );
-
-			FString SubLayerFilePath = FPaths::ConvertRelativePathToFull( SubLayerFile.GetValue() );
-			FPaths::MakePathRelativeTo( SubLayerFilePath, *LayerFilePath );
-
-			LayerHandle->InsertSubLayerPath( UnrealToUsd::ConvertString( *SubLayerFilePath ).Get() );
-		}
+		UsdLayersTreeViewImpl::InsertSubLayer( SelectedItem->GetLayerHandle(), SubLayerFile.GetValue() );
 		
 		break;
 	}
@@ -514,9 +525,9 @@ void SUsdLayersTreeView::OnAddSubLayer()
 
 void SUsdLayersTreeView::OnNewSubLayer()
 {
-	TOptional< FString > LayerFile = UsdUtils::BrowseUsdFile( UsdUtils::EBrowseFileMode::Save, AsShared() );
+	TOptional< FString > SubLayerFile = UsdUtils::BrowseUsdFile( UsdUtils::EBrowseFileMode::Save, AsShared() );
 
-	if ( !LayerFile )
+	if ( !SubLayerFile )
 	{
 		return;
 	}
@@ -527,14 +538,9 @@ void SUsdLayersTreeView::OnNewSubLayer()
 		FScopedUsdAllocs UsdAllocs;
 		for ( FUsdLayersTreeItemRef SelectedItem : MySelectedItems )
 		{
-			pxr::SdfLayerRefPtr SubLayerHandle = UsdUtils::CreateNewLayer( SelectedItem->UsdStage, *LayerFile.GetValue() ).Get();
+			pxr::SdfLayerRefPtr SubLayerHandle = UsdUtils::CreateNewLayer( SelectedItem->UsdStage, *SubLayerFile.GetValue() ).Get();
 
-			pxr::SdfLayerHandle LayerHandle = SelectedItem->GetLayerHandle().Get();
-
-			if ( LayerHandle && SubLayerHandle )
-			{
-				LayerHandle->InsertSubLayerPath( SubLayerHandle->GetRealPath() );
-			}
+			UsdLayersTreeViewImpl::InsertSubLayer( SelectedItem->GetLayerHandle(), SubLayerFile.GetValue() );
 			
 			break;
 		}
