Fixed warnings following merge from Main arising from MeshDescription deprecations.

diff --git a/Engine/Plugins/Importers/USDImporter/Source/USDImporter/Private/StaticMeshImporter.cpp b/Engine/Plugins/Importers/USDImporter/Source/USDImporter/Private/StaticMeshImporter.cpp
--- a/Engine/Plugins/Importers/USDImporter/Source/USDImporter/Private/StaticMeshImporter.cpp
+++ b/Engine/Plugins/Importers/USDImporter/Source/USDImporter/Private/StaticMeshImporter.cpp
@@ -377,7 +377,7 @@ bool FUSDStaticMeshImportState::AddPolygons( FMeshDescriptionWrapper& DestMeshWr
 		{
 			for(auto VertexInstID : MeshDescription->VertexInstances().GetElementIDs())
 			{
-				FVertexID VertexID = MeshDescription->GetVertexInstance(VertexInstID).VertexID;
+				FVertexID VertexID = MeshDescription->GetVertexInstanceVertex(VertexInstID);
 				Colors[VertexInstID] = ConvertToLinear(USDColors[VertexID.GetValue()]);
 			}
 		}
