Fixed USD importer material search. The name to search for needed some adjustements.

diff --git a/Engine/Plugins/Importers/USDImporter/Source/USDImporter/Private/StaticMeshImporter.cpp b/Engine/Plugins/Importers/USDImporter/Source/USDImporter/Private/StaticMeshImporter.cpp
--- a/Engine/Plugins/Importers/USDImporter/Source/USDImporter/Private/StaticMeshImporter.cpp
+++ b/Engine/Plugins/Importers/USDImporter/Source/USDImporter/Private/StaticMeshImporter.cpp
@@ -110,13 +110,28 @@ void FUSDStaticMeshImportState::ProcessMaterials(int32 LODIndex)
 			MaterialIndex = PolygonGroupID.GetValue();
 		}
 
+		if ( !Materials.IsValidIndex( MaterialIndex ) )
+		{
+			Materials.AddDefaulted( Materials.Num() - MaterialIndex + 1 );
+			Materials[ MaterialIndex ].Name = ImportedMaterialSlotNameString;
+		}
+
 		UMaterialInterface* Material = nullptr;
 		if (Materials.IsValidIndex(MaterialIndex))
 		{
 			Material = Materials[MaterialIndex].UnrealMaterial;
 			if (Material == nullptr)
 			{
-				const FString& MaterialFullName = Materials[MaterialIndex].Name;
+				FString MaterialFullName = Materials[MaterialIndex].Name;
+
+				// Only keep material name without prim path when searching for a UMaterial
+				FString MaterialPath;
+				FString MaterialName;
+				if ( MaterialFullName.Split( FString( TEXT("/") ), &MaterialPath, &MaterialName, ESearchCase::IgnoreCase, ESearchDir::FromEnd ) )
+				{
+					MaterialFullName = MaterialName;
+				}
+
 				FString MaterialBasePackageName = BasePackageName;
 				MaterialBasePackageName += TEXT("/");
 				MaterialBasePackageName += MaterialFullName;
