Copied cl 10925312 USD: Fixes related to sublayers

Fix edit target being cleared when reloading (fixing crash);
Work around layer registry not caching relative paths correctly;

diff --git a/Engine/Plugins/Importers/USDImporter/Source/USDStageEditor/Private/SUSDStage.cpp b/Engine/Plugins/Importers/USDImporter/Source/USDStageEditor/Private/SUSDStage.cpp
--- a/Engine/Plugins/Importers/USDImporter/Source/USDStageEditor/Private/SUSDStage.cpp
+++ b/Engine/Plugins/Importers/USDImporter/Source/USDStageEditor/Private/SUSDStage.cpp
@@ -35,6 +35,7 @@
 #include "pxr/pxr.h"
 #include "pxr/usd/sdf/copyUtils.h"
 #include "pxr/usd/sdf/schema.h"
+#include "pxr/usd/usd/editTarget.h"
 #include "pxr/usd/usd/modelAPI.h"
 #include "pxr/usd/usd/stageCache.h"
 #include "pxr/usd/usd/stageCacheContext.h"
@@ -87,7 +88,7 @@ void SUsdStage::Construct( const FArguments& InArgs )
 			+SVerticalBox::Slot()
 			.AutoHeight()
 			.Padding( SUSDStageConstants::SectionPadding )
-			[				
+			[
 				SNew(SBorder)
 				.BorderImage( FEditorStyle::GetBrush(TEXT("ToolPanel.GroupBorder")) )
 				[
@@ -95,14 +96,14 @@ void SUsdStage::Construct( const FArguments& InArgs )
 					.OnInitialLoadSetChanged( this, &SUsdStage::OnInitialLoadSetChanged )
 				]
 			]
-			
+
 			+SVerticalBox::Slot()
 			.FillHeight(1.f)
 			.Padding( SUSDStageConstants::SectionPadding )
 			[
 				SNew( SSplitter )
 				.Orientation( Orient_Vertical )
-			
+
 				+SSplitter::Slot()
 				.Value( 0.7f )
 				[
@@ -134,7 +135,7 @@ void SUsdStage::Construct( const FArguments& InArgs )
 				// Layers tree
 				+SSplitter::Slot()
 				.Value( 0.3f )
-				[				
+				[
 					SNew(SBorder)
 					.BorderImage( FEditorStyle::GetBrush(TEXT("ToolPanel.GroupBorder")) )
 					[
@@ -230,13 +231,13 @@ TSharedRef< SWidget > SUsdStage::MakeMainMenu()
 	FMenuBarBuilder MenuBuilder( nullptr );
 	{
 		// File
-		MenuBuilder.AddPullDownMenu( 
+		MenuBuilder.AddPullDownMenu(
 			LOCTEXT( "FileMenu", "File" ),
 			LOCTEXT( "FileMenu_ToolTip", "Opens the file menu" ),
 			FNewMenuDelegate::CreateSP( this, &SUsdStage::FillFileMenu ) );
 
 		// Actions
-		MenuBuilder.AddPullDownMenu( 
+		MenuBuilder.AddPullDownMenu(
 			LOCTEXT( "ActionsMenu", "Actions" ),
 			LOCTEXT( "ActionsMenu_ToolTip", "Opens the actions menu" ),
 			FNewMenuDelegate::CreateSP( this, &SUsdStage::FillActionsMenu ) );
@@ -342,7 +343,7 @@ void SUsdStage::OnNew()
 			// Create default prim
 			pxr::UsdGeomXform RootPrim = pxr::UsdGeomXform::Define( UsdStage, UnrealToUsd::ConvertPath( TEXT("/Root") ).Get() );
 			pxr::UsdModelAPI( RootPrim ).SetKind( pxr::TfToken("component") );
-		
+
 			// Set default prim
 			UsdStage->SetDefaultPrim( RootPrim.GetPrim() );
 
@@ -389,6 +390,14 @@ void SUsdStage::OnReloadStage()
 		{
 			UsdStage->Reload();
 
+			// If we were editing an unsaved layer, when we reload the edit target will be cleared.
+			// We need to make sure we're always editing something or else UsdEditContext might trigger some errors
+			const pxr::UsdEditTarget& EditTarget = UsdStage->GetEditTarget();
+			if (!EditTarget.IsValid() || EditTarget.IsNull())
+			{
+				UsdStage->SetEditTarget(UsdStage->GetRootLayer());
+			}
+
 			if ( UsdLayersTreeView )
 			{
 				UsdLayersTreeView->Refresh( UsdStageActor.Get(), true );
@@ -438,10 +447,10 @@ void SUsdStage::OnImport()
 
 		FString PathPart;
 		FString ExtensionPart;
-		
+
 		FPaths::Split( UsdToUnreal::ConvertString( UsdStage->GetRootLayer()->GetDisplayName() ), PathPart, StageName, ExtensionPart );
 	}
-	
+
 	// Import directly from stage
 	bool bCanceled = false;
 	{
@@ -453,7 +462,7 @@ void SUsdStage::OnImport()
 		ImportContextContainer->ImportContext.bApplyWorldTransformToGeometry = false;
 
 		UUSDImporter* UsdImporter = IUSDImporterModule::Get().GetImporter();
-		
+
 		bCanceled = !UsdImporter->ShowImportOptions( ImportContextContainer->ImportContext );
 
 		if ( !bCanceled )
diff --git a/Engine/Plugins/Importers/USDImporter/Source/USDUtilities/Private/USDLayerUtils.cpp b/Engine/Plugins/Importers/USDImporter/Source/USDUtilities/Private/USDLayerUtils.cpp
--- a/Engine/Plugins/Importers/USDImporter/Source/USDUtilities/Private/USDLayerUtils.cpp
+++ b/Engine/Plugins/Importers/USDImporter/Source/USDUtilities/Private/USDLayerUtils.cpp
@@ -44,14 +44,14 @@ TOptional< FString > UsdUtils::BrowseUsdFile( EBrowseFileMode Mode, TSharedRef<
 
 			if ( DesktopPlatform->OpenFileDialog( ParentWindowHandle, LOCTEXT( "ChooseFile", "Choose file").ToString(), TEXT(""), TEXT(""), TEXT("usd files (*.usd; *.usda; *.usdc)|*.usd; *.usda; *.usdc"), EFileDialogFlags::None, OutFiles ) )
 			{
-				return OutFiles[0];
+				return FPaths::ConvertRelativePathToFull(OutFiles[0]);
 			}
 			break;
 
 		case EBrowseFileMode::Save :
 			if ( DesktopPlatform->SaveFileDialog( ParentWindowHandle, LOCTEXT( "ChooseFile", "Choose file").ToString(), TEXT(""), TEXT(""), TEXT("usd files (*.usd; *.usda; *.usdc)|*.usd; *.usda; *.usdc"), EFileDialogFlags::None, OutFiles ) )
 			{
-				return OutFiles[0];
+				return FPaths::ConvertRelativePathToFull(OutFiles[0]);
 			}
 			break;
 
