Prevent USDImporter being built by default if the editor is being compiled in C++17 mode.

#rnx

diff --git a/Engine/Plugins/Importers/USDImporter/Source/USDStage/Private/USDStageActor.cpp b/Engine/Plugins/Importers/USDImporter/Source/USDStage/Private/USDStageActor.cpp
--- a/Engine/Plugins/Importers/USDImporter/Source/USDStage/Private/USDStageActor.cpp
+++ b/Engine/Plugins/Importers/USDImporter/Source/USDStage/Private/USDStageActor.cpp
@@ -68,7 +68,7 @@
 
 static const EObjectFlags DefaultObjFlag = EObjectFlags::RF_Transactional | EObjectFlags::RF_Transient;
 
-AUsdStageActor::FOnActorLoaded AUsdStageActor::OnActorLoaded{};
+AUsdStageActor::FOnActorLoaded AUsdStageActor::OnActorLoaded;
 
 #if USE_USD_SDK
 FMeshDescription LoadMeshDescription( const pxr::UsdGeomMesh& UsdMesh, const pxr::UsdTimeCode TimeCode )
diff --git a/Engine/Plugins/Importers/USDImporter/Source/UnrealUSDWrapper/UnrealUSDWrapper.Build.cs b/Engine/Plugins/Importers/USDImporter/Source/UnrealUSDWrapper/UnrealUSDWrapper.Build.cs
--- a/Engine/Plugins/Importers/USDImporter/Source/UnrealUSDWrapper/UnrealUSDWrapper.Build.cs
+++ b/Engine/Plugins/Importers/USDImporter/Source/UnrealUSDWrapper/UnrealUSDWrapper.Build.cs
@@ -20,6 +20,7 @@ namespace UnrealBuildTool.Rules
 			var PythonSourceTPSDir = Path.Combine(EngineDir, "Source", "ThirdParty", "Python");
 
 			if (Target.WindowsPlatform.Compiler != WindowsCompiler.Clang && Target.WindowsPlatform.StaticAnalyzer == WindowsStaticAnalyzer.None &&
+				Target.CppStandard < CppStandardVersion.Cpp17 && // Not currently compatible with C++17 due to old version of Boost
 				Target.LinkType != TargetLinkType.Monolithic) // If you want to use USD in a monolithic target, you'll have to use the ANSI allocator and remove this condition
 			{
 				PublicDefinitions.Add("USE_USD_SDK=1");
@@ -70,6 +71,6 @@ namespace UnrealBuildTool.Rules
 			{
 				PublicDefinitions.Add("USE_USD_SDK=0");
 			}
-        }
+		}
 	}
 }
