Fixed position modifications not being written back to USD. We are storing components but were receiving an actor and the comparison was failing.

#rnx

diff --git a/Engine/Plugins/Importers/USDImporter/Source/USDStage/Private/USDStageActor.cpp b/Engine/Plugins/Importers/USDImporter/Source/USDStage/Private/USDStageActor.cpp
--- a/Engine/Plugins/Importers/USDImporter/Source/USDStage/Private/USDStageActor.cpp
+++ b/Engine/Plugins/Importers/USDImporter/Source/USDStage/Private/USDStageActor.cpp
@@ -575,12 +575,33 @@ void AUsdStageActor::OnUsdPrimTwinDestroyed( const FUsdPrimTwin& UsdPrimTwin )
 void AUsdStageActor::OnPrimObjectPropertyChanged( UObject* ObjectBeingModified, FPropertyChangedEvent& PropertyChangedEvent )
 {
 #if USE_USD_SDK
-	if ( ObjectBeingModified == this || !ObjectsToWatch.Contains( ObjectBeingModified ) )
+	if ( ObjectBeingModified == this )
 	{
 		return;
 	}
 
-	FString PrimPath = ObjectsToWatch[ ObjectBeingModified ];
+	UObject* PrimObject = ObjectBeingModified;
+
+	if ( !ObjectsToWatch.Contains( ObjectBeingModified ) )
+	{
+		if ( AActor* ActorBeingModified = Cast< AActor >( ObjectBeingModified ) )
+		{
+			if ( !ObjectsToWatch.Contains( ActorBeingModified->GetRootComponent() ) )
+			{
+				return;
+			}
+			else
+			{
+				PrimObject = ActorBeingModified->GetRootComponent();
+			}
+		}
+		else
+		{
+			return;
+		}
+	}
+	
+	FString PrimPath = ObjectsToWatch[ PrimObject ];
 
 	if ( FUsdPrimTwin* UsdPrimTwin = RootUsdTwin.Find( PrimPath ) )
 	{
@@ -610,7 +631,7 @@ void AUsdStageActor::OnPrimObjectPropertyChanged( UObject* ObjectBeingModified,
 
 bool AUsdStageActor::HasAutorithyOverStage() const
 {
-	return !GetWorld() || !GetWorld()->IsGameWorld();
+	return !IsTemplate() && ( !GetWorld() || !GetWorld()->IsGameWorld() );
 }
 
 #if USE_USD_SDK
