Fixed crash when edit target doesn't have a layer

diff --git a/Engine/Plugins/Importers/USDImporter/Source/USDStageEditor/Private/SUSDLayersTreeView.cpp b/Engine/Plugins/Importers/USDImporter/Source/USDStageEditor/Private/SUSDLayersTreeView.cpp
--- a/Engine/Plugins/Importers/USDImporter/Source/USDStageEditor/Private/SUSDLayersTreeView.cpp
+++ b/Engine/Plugins/Importers/USDImporter/Source/USDStageEditor/Private/SUSDLayersTreeView.cpp
@@ -167,7 +167,9 @@ public:
 
 		LayerData->DisplayName = FText::FromString( UsdToUnreal::ConvertString( pxr::SdfLayer::GetDisplayNameFromIdentifier( LayerIdentifier.Get() ) ) );
 		LayerData->bIsMuted = UsdStage.Get()->IsLayerMuted( LayerIdentifier.Get() );
-		LayerData->bIsEditTarget = ( UsdStage.Get()->GetEditTarget().GetLayer()->GetIdentifier() == LayerIdentifier.Get() );
+
+		const pxr::SdfLayerHandle& EditTargetLayer = UsdStage.Get()->GetEditTarget().GetLayer();
+		LayerData->bIsEditTarget = ( EditTargetLayer ? EditTargetLayer->GetIdentifier() == LayerIdentifier.Get() : false );
 
 		for ( FUsdLayersTreeItemRef Child : Children )
 		{
