Fix for some USD modules not being available on Linux

#rnx

diff --git a/Engine/Plugins/Importers/USDImporter/Source/UnrealUSDBinding/UnrealUSDBinding.Build.cs b/Engine/Plugins/Importers/USDImporter/Source/UnrealUSDBinding/UnrealUSDBinding.Build.cs
--- a/Engine/Plugins/Importers/USDImporter/Source/UnrealUSDBinding/UnrealUSDBinding.Build.cs
+++ b/Engine/Plugins/Importers/USDImporter/Source/UnrealUSDBinding/UnrealUSDBinding.Build.cs
@@ -15,13 +15,9 @@ namespace UnrealBuildTool.Rules
 				new string[]
 				{
 					"Core",
-				}
+                    "UnrealUSDWrapper"
+                }
 			);
-
-			if (Target.Platform == UnrealTargetPlatform.Win64)
-			{
-				PrivateDependencyModuleNames.Add("UnrealUSDWrapper");
-			}
 		}
 	}
 }
diff --git a/Engine/Plugins/Importers/USDImporter/USDImporter.uplugin b/Engine/Plugins/Importers/USDImporter/USDImporter.uplugin
--- a/Engine/Plugins/Importers/USDImporter/USDImporter.uplugin
+++ b/Engine/Plugins/Importers/USDImporter/USDImporter.uplugin
@@ -38,7 +38,8 @@
 			"Type": "Editor",
 			"LoadingPhase": "Default",
 			"WhitelistPlatforms": [
-				"Win64"
+				"Win64",
+				"Linux"
 			]
 		},
 		{
@@ -55,7 +56,8 @@
 			"Type": "Editor",
 			"LoadingPhase": "Default",
 			"WhitelistPlatforms": [
-				"Win64"
+				"Win64",
+				"Linux"
 			]
 		},
 		{
