Linux no-RTTI fix

#rnx

diff --git a/Engine/Plugins/Importers/USDImporter/Source/USDExporter/USDExporter.Build.cs b/Engine/Plugins/Importers/USDImporter/Source/USDExporter/USDExporter.Build.cs
--- a/Engine/Plugins/Importers/USDImporter/Source/USDExporter/USDExporter.Build.cs
+++ b/Engine/Plugins/Importers/USDImporter/Source/USDExporter/USDExporter.Build.cs
@@ -9,7 +9,11 @@ namespace UnrealBuildTool.Rules
 	{
 		public USDExporter(ReadOnlyTargetRules Target) : base(Target)
         {
-			bUseRTTI = true;
+			// We require the whole editor to be RTTI enabled on Linux for now
+			if (Target.Platform != UnrealTargetPlatform.Linux)
+			{
+				bUseRTTI = true;
+			}
 
 			PublicDependencyModuleNames.AddRange(
 				new string[]
diff --git a/Engine/Plugins/Importers/USDImporter/Source/USDImporter/USDImporter.Build.cs b/Engine/Plugins/Importers/USDImporter/Source/USDImporter/USDImporter.Build.cs
--- a/Engine/Plugins/Importers/USDImporter/Source/USDImporter/USDImporter.Build.cs
+++ b/Engine/Plugins/Importers/USDImporter/Source/USDImporter/USDImporter.Build.cs
@@ -9,7 +9,11 @@ namespace UnrealBuildTool.Rules
 	{
 		public USDImporter(ReadOnlyTargetRules Target) : base(Target)
         {
-			bUseRTTI = true;
+			// We require the whole editor to be RTTI enabled on Linux for now
+			if (Target.Platform != UnrealTargetPlatform.Linux)
+			{
+				bUseRTTI = true;
+			}
 
 			PrivateDependencyModuleNames.AddRange(
 				new string[]
@@ -34,32 +38,10 @@ namespace UnrealBuildTool.Rules
                     "RenderCore",
                     "RHI",
 					"StaticMeshDescription",
+					"UnrealUSDWrapper",
 					"USDUtilities",
                 }
 				);
-
-			if (Target.Platform == UnrealTargetPlatform.Win64)
-			{
-				PrivateDependencyModuleNames.Add("UnrealUSDWrapper");
-
-				foreach (string FilePath in Directory.EnumerateFiles(Path.Combine(ModuleDirectory, "../../Binaries/Win64/"), "*.dll", SearchOption.AllDirectories))
-                {
-                    RuntimeDependencies.Add(FilePath);
-                }
-            }
-            else if (Target.Platform == UnrealTargetPlatform.Linux && Target.Architecture.StartsWith("x86_64"))
-			{
-				PrivateDependencyModuleNames.Add("UnrealUSDWrapper");
-
-				// link directly to runtime libs on Linux, as this also puts them into rpath
-				string RuntimeLibraryPath = Path.Combine(ModuleDirectory, "../../Binaries", Target.Platform.ToString(), Target.Architecture.ToString());
-				PrivateRuntimeLibraryPaths.Add(RuntimeLibraryPath);
-
-				foreach (string FilePath in Directory.EnumerateFiles(RuntimeLibraryPath, "*.so*", SearchOption.AllDirectories))
-                {
-                    RuntimeDependencies.Add(FilePath);
-                }
-            }
 		}
 	}
 }
diff --git a/Engine/Plugins/Importers/USDImporter/Source/USDStage/USDStage.Build.cs b/Engine/Plugins/Importers/USDImporter/Source/USDStage/USDStage.Build.cs
--- a/Engine/Plugins/Importers/USDImporter/Source/USDStage/USDStage.Build.cs
+++ b/Engine/Plugins/Importers/USDImporter/Source/USDStage/USDStage.Build.cs
@@ -9,7 +9,11 @@ namespace UnrealBuildTool.Rules
 	{
 		public USDStage(ReadOnlyTargetRules Target) : base(Target)
 		{
-			bUseRTTI = true;
+			// We require the whole editor to be RTTI enabled on Linux for now
+			if (Target.Platform != UnrealTargetPlatform.Linux)
+			{
+				bUseRTTI = true;
+			}
 
 			PrivateDependencyModuleNames.AddRange(
 				new string[]
