USD: Extract error handling to an utility #jira UE-85771, UE-82974 #rb JeanMichel.Dignard

diff --git a/Engine/Plugins/Importers/USDImporter/Source/USDImporter/Private/USDImporter.cpp b/Engine/Plugins/Importers/USDImporter/Source/USDImporter/Private/USDImporter.cpp
--- a/Engine/Plugins/Importers/USDImporter/Source/USDImporter/Private/USDImporter.cpp
+++ b/Engine/Plugins/Importers/USDImporter/Source/USDImporter/Private/USDImporter.cpp
@@ -25,10 +25,11 @@
 #include "Widgets/Layout/SUniformGridPanel.h"
 
 #include "PropertySetter.h"
+#include "USDErrorUtils.h"
 #include "USDImportOptions.h"
 #include "USDImporterProjectSettings.h"
-#include "USDPrimResolverKind.h"
 #include "USDMemory.h"
+#include "USDPrimResolverKind.h"
 #include "USDSceneImportFactory.h"
 #include "USDTypesConversion.h"
 
@@ -300,7 +301,7 @@ TArray<UObject*> UUSDImporter::ImportMeshes(FUsdImportContext& ImportContext, co
 			}
 
 			NewPackageName = UPackageTools::SanitizePackageName(FinalPackagePathName);
-		
+
 			// Once we've already imported it we dont need to import it again
 			if(!ImportContext.PathToImportAssetMap.Contains(NewPackageName))
 			{
@@ -487,13 +488,16 @@ TUsdStore< pxr::UsdStageRefPtr > UUSDImporter::ReadUsdFile(FUsdImportContext& Im
 	FilePath = FPaths::GetPath(FilePath) + TEXT("/");
 	FString CleanFilename = FPaths::GetCleanFilename(Filename);
 
+	UsdUtils::StartMonitoringErrors();
+
 	TUsdStore< pxr::UsdStageRefPtr > Stage = UnrealUSDWrapper::OpenUsdStage(TCHAR_TO_ANSI(*FilePath), TCHAR_TO_ANSI(*CleanFilename));
 
-	const char* Errors = UnrealUSDWrapper::GetErrors();
-	if (Errors)
+	TArray<FString> ErrorStrings = UsdUtils::GetErrorsAndStopMonitoring();
+	FString Error = FString::Join(ErrorStrings, TEXT("\n"));
+
+	if (!Error.IsEmpty())
 	{
-		FString ErrorStr = UsdToUnreal::ConvertString(Errors);
-		ImportContext.AddErrorMessage(EMessageSeverity::Error, FText::Format(LOCTEXT("CouldNotImportUSDFile", "Could not import USD file {0}\n {1}"), FText::FromString(CleanFilename), FText::FromString(ErrorStr)));
+		ImportContext.AddErrorMessage(EMessageSeverity::Error, FText::Format(LOCTEXT("CouldNotImportUSDFile", "Could not import USD file {0}\n {1}"), FText::FromString(CleanFilename), FText::FromString(Error)));
 	}
 	return Stage;
 }
@@ -505,7 +509,7 @@ void UUSDImporter::ImportUsdStage( FUSDSceneImportContext& ImportContext )
 	if ( ImportContext.Stage.Get() && ImportOptions )
 	{
 		UUSDPrimResolver* PrimResolver = ImportContext.PrimResolver;
-		
+
 		EExistingActorPolicy ExistingActorPolicy = ImportOptions->ExistingActorPolicy;
 
 		TArray<FActorSpawnData> SpawnDatas;
diff --git a/Engine/Plugins/Importers/USDImporter/Source/USDStageEditor/Private/SUSDStage.cpp b/Engine/Plugins/Importers/USDImporter/Source/USDStageEditor/Private/SUSDStage.cpp
--- a/Engine/Plugins/Importers/USDImporter/Source/USDStageEditor/Private/SUSDStage.cpp
+++ b/Engine/Plugins/Importers/USDImporter/Source/USDStageEditor/Private/SUSDStage.cpp
@@ -7,12 +7,13 @@
 #include "SUSDPrimInfo.h"
 #include "SUSDStageInfo.h"
 #include "SUSDStageTreeView.h"
-#include "UnrealUSDWrapper.h"
+#include "USDErrorUtils.h"
 #include "USDImporter.h"
 #include "USDLayerUtils.h"
 #include "USDStageActor.h"
 #include "USDStageModule.h"
 #include "USDTypesConversion.h"
+#include "UnrealUSDWrapper.h"
 
 #include "Dialogs/DlgPickPath.h"
 #include "EditorStyleSet.h"
@@ -87,7 +88,7 @@ void SUsdStage::Construct( const FArguments& InArgs )
 			+SVerticalBox::Slot()
 			.AutoHeight()
 			.Padding( SUSDStageConstants::SectionPadding )
-			[				
+			[
 				SNew(SBorder)
 				.BorderImage( FEditorStyle::GetBrush(TEXT("ToolPanel.GroupBorder")) )
 				[
@@ -95,14 +96,14 @@ void SUsdStage::Construct( const FArguments& InArgs )
 					.OnInitialLoadSetChanged( this, &SUsdStage::OnInitialLoadSetChanged )
 				]
 			]
-			
+
 			+SVerticalBox::Slot()
 			.FillHeight(1.f)
 			.Padding( SUSDStageConstants::SectionPadding )
 			[
 				SNew( SSplitter )
 				.Orientation( Orient_Vertical )
-			
+
 				+SSplitter::Slot()
 				.Value( 0.7f )
 				[
@@ -134,7 +135,7 @@ void SUsdStage::Construct( const FArguments& InArgs )
 				// Layers tree
 				+SSplitter::Slot()
 				.Value( 0.3f )
-				[				
+				[
 					SNew(SBorder)
 					.BorderImage( FEditorStyle::GetBrush(TEXT("ToolPanel.GroupBorder")) )
 					[
@@ -230,13 +231,13 @@ TSharedRef< SWidget > SUsdStage::MakeMainMenu()
 	FMenuBarBuilder MenuBuilder( nullptr );
 	{
 		// File
-		MenuBuilder.AddPullDownMenu( 
+		MenuBuilder.AddPullDownMenu(
 			LOCTEXT( "FileMenu", "File" ),
 			LOCTEXT( "FileMenu_ToolTip", "Opens the file menu" ),
 			FNewMenuDelegate::CreateSP( this, &SUsdStage::FillFileMenu ) );
 
 		// Actions
-		MenuBuilder.AddPullDownMenu( 
+		MenuBuilder.AddPullDownMenu(
 			LOCTEXT( "ActionsMenu", "Actions" ),
 			LOCTEXT( "ActionsMenu_ToolTip", "Opens the actions menu" ),
 			FNewMenuDelegate::CreateSP( this, &SUsdStage::FillActionsMenu ) );
@@ -342,7 +343,7 @@ void SUsdStage::OnNew()
 			// Create default prim
 			pxr::UsdGeomXform RootPrim = pxr::UsdGeomXform::Define( UsdStage, UnrealToUsd::ConvertPath( TEXT("/Root") ).Get() );
 			pxr::UsdModelAPI( RootPrim ).SetKind( pxr::TfToken("component") );
-		
+
 			// Set default prim
 			UsdStage->SetDefaultPrim( RootPrim.GetPrim() );
 
@@ -360,7 +361,11 @@ void SUsdStage::OnOpen()
 
 	if ( UsdFilePath )
 	{
+		UsdUtils::StartMonitoringErrors();
+
 		OpenStage( *UsdFilePath.GetValue() );
+
+		UsdUtils::ShowErrorsAndStopMonitoring(FText::Format(LOCTEXT("USDOpenError", "Encountered some errors opening USD file at path '{0}!\nCheck the Output Log for details."), FText::FromString(UsdFilePath.GetValue())));
 	}
 }
 
@@ -374,7 +379,11 @@ void SUsdStage::OnSave()
 		{
 			FScopedUsdAllocs UsdAllocs;
 
+			UsdUtils::StartMonitoringErrors();
+
 			UsdStage->Save();
+
+			UsdUtils::ShowErrorsAndStopMonitoring(LOCTEXT("USDSaveError", "Failed to save current USD Stage!\nCheck the Output Log for details."));
 		}
 	}
 }
@@ -438,10 +447,10 @@ void SUsdStage::OnImport()
 
 		FString PathPart;
 		FString ExtensionPart;
-		
+
 		FPaths::Split( UsdToUnreal::ConvertString( UsdStage->GetRootLayer()->GetDisplayName() ), PathPart, StageName, ExtensionPart );
 	}
-	
+
 	// Import directly from stage
 	bool bCanceled = false;
 	{
@@ -453,7 +462,7 @@ void SUsdStage::OnImport()
 		ImportContextContainer->ImportContext.bApplyWorldTransformToGeometry = false;
 
 		UUSDImporter* UsdImporter = IUSDImporterModule::Get().GetImporter();
-		
+
 		bCanceled = !UsdImporter->ShowImportOptions( ImportContextContainer->ImportContext );
 
 		if ( !bCanceled )
diff --git a/Engine/Plugins/Importers/USDImporter/Source/USDUtilities/Private/USDErrorUtils.cpp b/Engine/Plugins/Importers/USDImporter/Source/USDUtilities/Private/USDErrorUtils.cpp
new file mode 100644
--- /dev/null
+++ b/Engine/Plugins/Importers/USDImporter/Source/USDUtilities/Private/USDErrorUtils.cpp
@@ -0,0 +1,119 @@
+// Copyright 1998-2019 Epic Games, Inc. All Rights Reserved.
+
+#include "USDErrorUtils.h"
+
+#if USE_USD_SDK
+
+#include "USDMemory.h"
+#include "USDTypesConversion.h"
+
+#include "Framework/Notifications/NotificationManager.h"
+#include "Templates/SharedPointer.h"
+#include "Widgets/Notifications/SNotificationList.h"
+
+#include "USDIncludesStart.h"
+#include "pxr/pxr.h"
+#include "pxr/base/tf/errorMark.h"
+#include "USDIncludesEnd.h"
+
+#define LOCTEXT_NAMESPACE "USDErrorUtils"
+
+DEFINE_LOG_CATEGORY(LogUsdErrorUtility);
+
+namespace UsdUtils
+{
+	// We need an extra level of indirection because TfErrorMark is noncopyable
+	using MarkRef = TUsdStore<TSharedRef<pxr::TfErrorMark>>;
+
+	static TArray<MarkRef> ErrorMarkStack;
+
+	void StartMonitoringErrors()
+	{
+		FScopedUsdAllocs Allocs;
+
+		TSharedRef<pxr::TfErrorMark> Mark = MakeShared<pxr::TfErrorMark>();
+		Mark->SetMark();
+
+		ErrorMarkStack.Emplace(Mark);
+	}
+
+	TArray<FString> GetErrorsAndStopMonitoring()
+	{
+		if (ErrorMarkStack.Num() == 0)
+		{
+			return {};
+		}
+
+		MarkRef Store = ErrorMarkStack.Pop();
+		pxr::TfErrorMark& Mark = Store.Get().Get();
+
+		if (Mark.IsClean())
+		{
+			return {};
+		}
+
+		TArray<FString> Errors;
+
+		for (pxr::TfErrorMark::Iterator ErrorIter = Mark.GetBegin();
+			 ErrorIter != Mark.GetEnd();
+			 ++ErrorIter)
+		{
+			std::string ErrorStr = ErrorIter->GetErrorCodeAsString();
+			ErrorStr += " ";
+			ErrorStr += ErrorIter->GetCommentary();
+
+			// Add unique here as for some errors (e.g. parsing errors) USD can emit the exact same
+			// error message 5+ times in a row
+			Errors.AddUnique(UsdToUnreal::ConvertString(ErrorStr));
+		}
+
+		Mark.Clear();
+
+		return Errors;
+	}
+
+	bool ShowErrorsAndStopMonitoring(const FText& ToastMessage)
+	{
+		TArray<FString> Errors = GetErrorsAndStopMonitoring();
+		bool bHadErrors = Errors.Num() > 0;
+
+		if (bHadErrors)
+		{
+			FNotificationInfo ErrorToast(!ToastMessage.IsEmpty() ? ToastMessage : LOCTEXT("USDErrorsToast", "Encountered USD errors!\nCheck the Output Log for details."));
+
+			ErrorToast.ExpireDuration = 5.0f;
+			ErrorToast.bFireAndForget = true;
+			ErrorToast.Image = FCoreStyle::Get().GetBrush(TEXT("MessageLog.Error"));
+			FSlateNotificationManager::Get().AddNotification(ErrorToast);
+		}
+
+		for (const FString& Error : Errors)
+		{
+			UE_LOG(LogUsdErrorUtility, Error, TEXT("%s"), *Error);
+		}
+
+		return Errors.Num() > 0;
+	}
+}; // namespace UsdUtils
+
+#undef LOCTEXT_NAMESPACE
+
+#else // #if USE_USD_SDK
+
+namespace UsdUtils
+{
+	void StartMonitoringErrors()
+	{
+	}
+	TArray<FString> GetErrorsAndStopMonitoring()
+	{
+		return {TEXT("USD SDK is not available!")};
+	}
+	bool ShowErrorsAndStopMonitoring(const FText& ToastMessage)
+	{
+		return false;
+	}
+}
+
+#endif // #if USE_USD_SDK
+
diff --git a/Engine/Plugins/Importers/USDImporter/Source/USDUtilities/Public/USDErrorUtils.h b/Engine/Plugins/Importers/USDImporter/Source/USDUtilities/Public/USDErrorUtils.h
new file mode 100644
--- /dev/null
+++ b/Engine/Plugins/Importers/USDImporter/Source/USDUtilities/Public/USDErrorUtils.h
@@ -0,0 +1,29 @@
+// Copyright 1998-2019 Epic Games, Inc. All Rights Reserved.
+
+#pragma once
+
+#include "CoreMinimal.h"
+
+DECLARE_LOG_CATEGORY_EXTERN(LogUsdErrorUtility, Log, All);
+
+namespace UsdUtils
+{
+    /**
+     * Pushes an USD error monitoring object into the stack and catches any emitted errors
+     */
+	USDUTILITIES_API void StartMonitoringErrors();
+
+    /**
+     * Returns all errors that were captured since StartMonitoringErrors(), clears and pops an
+     * error monitoring object from the stack
+     */
+	USDUTILITIES_API TArray<FString> GetErrorsAndStopMonitoring();
+
+    /**
+     * Displays the error messages for each captured error since StartMonitoringErrors(),
+	 * clears and pops an error monitoring object from the stack.
+	 * If ToastMessage is empty, a default message will be displayed.
+     * Returns true if there were any errors.
+     */
+	USDUTILITIES_API bool ShowErrorsAndStopMonitoring(const FText& ToastMessage = FText());
+}
diff --git a/Engine/Plugins/Importers/USDImporter/Source/UnrealUSDWrapper/Private/UnrealUSDWrapper.cpp b/Engine/Plugins/Importers/USDImporter/Source/UnrealUSDWrapper/Private/UnrealUSDWrapper.cpp
--- a/Engine/Plugins/Importers/USDImporter/Source/UnrealUSDWrapper/Private/UnrealUSDWrapper.cpp
+++ b/Engine/Plugins/Importers/USDImporter/Source/UnrealUSDWrapper/Private/UnrealUSDWrapper.cpp
@@ -469,7 +469,7 @@ int IUsdPrim::GetNumLODs(const UsdPrim& Prim)
 {
 	FScopedUsdAllocs UsdAllocs;
 
-	// 0 indicates no variant or no lods in variant. 
+	// 0 indicates no variant or no lods in variant.
 	int NumLODs = 0;
 	if (Prim.HasVariantSets())
 	{
@@ -487,7 +487,7 @@ int IUsdPrim::GetNumLODs(const UsdPrim& Prim)
 bool IUsdPrim::IsKindChildOf(const UsdPrim& Prim, const std::string& InBaseKind)
 {
 	TfToken BaseKind(InBaseKind);
-	
+
 	KindRegistry& Registry = KindRegistry::GetInstance();
 
 	TfToken PrimKind( GetKind(Prim) );
@@ -511,7 +511,7 @@ TfToken IUsdPrim::GetKind(const pxr::UsdPrim& Prim)
 		static TfToken KindMetaDataToken("kind");
 		Prim.GetMetadata(KindMetaDataToken, &KindType);
 	}
-	
+
 	return KindType;
 }
 
@@ -632,9 +632,9 @@ namespace Internal
 	{
 		TArray< FString > MaterialNames;
 
-		// load each material at the material path; 
+		// load each material at the material path;
 		UsdPrim MaterialPrim = Stage->Load(Path);
-		
+
 		if(MaterialPrim)
 		{
 			// Default to using the prim path name as the path for this material in Unreal
@@ -690,7 +690,7 @@ TTuple< TArray< FString >, TArray< int32 > > IUsdPrim::GetGeometryMaterials(doub
 		return MakeTuple( MaterialNames, FaceMaterialIndices );
 	}
 
-	// If the gprim does not have a material faceSet which represents per-face 
+	// If the gprim does not have a material faceSet which represents per-face
 	// shader assignments, assign the shading engine to the entire gprim.
 	std::vector<UsdGeomSubset> FaceSubsets = UsdShadeMaterialBindingAPI(Prim).GetMaterialBindSubsets();
 	std::vector<UsdGeomFaceSetAPI> FaceSets = UsdGeomFaceSetAPI::GetFaceSets(Prim);
@@ -737,7 +737,7 @@ TTuple< TArray< FString >, TArray< int32 > > IUsdPrim::GetGeometryMaterials(doub
 			UsdShadeMaterialBindingAPI SubsetBindingAPI(Subset.GetPrim());
 			UsdShadeMaterial BoundMaterial = SubsetBindingAPI.ComputeBoundMaterial();
 
-			// Only transfer the first timeSample or default Indices, if 
+			// Only transfer the first timeSample or default Indices, if
 			// there are no time-samples.
 			VtIntArray Indices;
 			Subset.GetIndicesAttr().Get(&Indices, UsdTimeCode::EarliestTime());
@@ -842,7 +842,7 @@ TTuple< TArray< FString >, TArray< int32 > > IUsdPrim::GetGeometryMaterials(doub
 	}
 	else
 	{
-		// No face sets, find a relationship that defines the material 
+		// No face sets, find a relationship that defines the material
 		UsdRelationship Relationship = Prim.GetRelationship(UnrealIdentifiers::MaterialRelationship);
 		if (Relationship)
 		{
@@ -887,7 +887,7 @@ bool IUsdPrim::SetActiveLODIndex(const pxr::UsdPrim& Prim, int LODIndex)
 			}
 		}
 	}
-		
+
 	return false;
 }
 
@@ -904,7 +904,7 @@ EUsdGeomOrientation IUsdPrim::GetGeometryOrientation(const pxr::UsdGeomMesh& Mes
 	{
 		UsdAttribute Orientation = Mesh.GetOrientationAttr();
 		if(Orientation)
-		{ 
+		{
 			static TfToken RightHanded("rightHanded");
 			static TfToken LeftHanded("leftHanded");
 
@@ -919,7 +919,6 @@ EUsdGeomOrientation IUsdPrim::GetGeometryOrientation(const pxr::UsdGeomMesh& Mes
 }
 
 bool UnrealUSDWrapper::bInitialized = false;
-std::string UnrealUSDWrapper::Errors;
 
 
 void UnrealUSDWrapper::Initialize( const TArray< FString >& InPluginDirectories )
@@ -951,8 +950,6 @@ void UnrealUSDWrapper::Initialize( const TArray< FString >& InPluginDirectories
 
 TUsdStore< pxr::UsdStageRefPtr > UnrealUSDWrapper::OpenUsdStage(const char* Path, const char* Filename)
 {
-	Errors.clear();
-
 	if (!bInitialized)
 	{
 		return {};
@@ -960,29 +957,15 @@ TUsdStore< pxr::UsdStageRefPtr > UnrealUSDWrapper::OpenUsdStage(const char* Path
 
 	bool bImportedSuccessfully = false;
 
-	TfErrorMark ErrorMark;
-
 	string PathAndFilename = string(Path) + string(Filename);
 
 	bool bIsSupported = UsdStage::IsSupportedFile(PathAndFilename);
-	
+
 	FScopedUsdAllocs UsdAllocs;
 
 	pxr::UsdStageCacheContext UsdStageCacheContext( GetUsdStageCache() );
 	UsdStageRefPtr Stage = UsdStage::Open(PathAndFilename);
 
-	if (!ErrorMark.IsClean())
-	{
-		TfErrorMark::Iterator i;
-		for (i = ErrorMark.GetBegin(); i != ErrorMark.GetEnd(); ++i)
-		{
-			Errors += i->GetErrorCodeAsString();
-			Errors += " ";
-			Errors += i->GetCommentary();
-			Errors += "\n";
-		}
-	}
-
 	return Stage;
 }
 
@@ -991,11 +974,6 @@ double UnrealUSDWrapper::GetDefaultTimeCode()
 	return UsdTimeCode::Default().GetValue();
 }
 
-const char* UnrealUSDWrapper::GetErrors()
-{
-	return Errors.length() > 0 ? Errors.c_str() : nullptr;
-}
-
 pxr::UsdStageCache& UnrealUSDWrapper::GetUsdStageCache()
 {
 	static TUsdStore< pxr::UsdStageCache > UsdStageCache;
diff --git a/Engine/Plugins/Importers/USDImporter/Source/UnrealUSDWrapper/Public/UnrealUSDWrapper.h b/Engine/Plugins/Importers/USDImporter/Source/UnrealUSDWrapper/Public/UnrealUSDWrapper.h
--- a/Engine/Plugins/Importers/USDImporter/Source/UnrealUSDWrapper/Public/UnrealUSDWrapper.h
+++ b/Engine/Plugins/Importers/USDImporter/Source/UnrealUSDWrapper/Public/UnrealUSDWrapper.h
@@ -152,14 +152,12 @@ class UnrealUSDWrapper
 public:
 #if USE_USD_SDK
 	UNREALUSDWRAPPER_API static void Initialize( const TArray< FString > & InPluginDirectories );
-	UNREALUSDWRAPPER_API static const char* GetErrors();
 	UNREALUSDWRAPPER_API static double GetDefaultTimeCode();
 
 	UNREALUSDWRAPPER_API static TUsdStore< pxr::TfRefPtr< pxr::UsdStage > > OpenUsdStage(const char* Path, const char* Filename);
 	UNREALUSDWRAPPER_API static pxr::UsdStageCache& GetUsdStageCache();
 #endif  // #if USE_USD_SDK
 private:
-	static std::string Errors;
 	static bool bInitialized;
 };
 
@@ -184,7 +182,7 @@ public:
 
 	UNREALUSDWRAPPER_API static bool IsUnsigned(const pxr::UsdAttribute& Attribute);
 #endif // #if USE_USD_SDK
-	
+
 };
 
 
