USD: Fix sublayers not showing up when reloading #rb JeanMichel.Dignard

diff --git a/Engine/Plugins/Importers/USDImporter/Source/USDStageEditor/Private/SUSDLayersTreeView.cpp b/Engine/Plugins/Importers/USDImporter/Source/USDStageEditor/Private/SUSDLayersTreeView.cpp
--- a/Engine/Plugins/Importers/USDImporter/Source/USDStageEditor/Private/SUSDLayersTreeView.cpp
+++ b/Engine/Plugins/Importers/USDImporter/Source/USDStageEditor/Private/SUSDLayersTreeView.cpp
@@ -34,7 +34,7 @@
 
 namespace UsdLayersTreeViewImpl
 {
-	bool InsertSubLayer( const TUsdStore< pxr::SdfLayerHandle >& LayerHandle, const FString& SubLayerFile )
+	bool InsertSubLayer( const TUsdStore< pxr::SdfLayerRefPtr >& LayerHandle, const FString& SubLayerFile )
 	{
 		if ( LayerHandle.Get() )
 		{
@@ -93,7 +93,7 @@ public:
 		{
 			FScopedUsdAllocs UsdAllocs;
 
-			pxr::SdfLayerHandle UsdLayer = GetLayerHandle().Get();
+			pxr::SdfLayerRefPtr UsdLayer = GetLayerHandle().Get();
 
 			if ( UsdLayer )
 			{
@@ -140,7 +140,7 @@ public:
 		{
 			FScopedUsdAllocs UsdAllocs;
 
-			pxr::SdfLayerHandle UsdLayer = GetLayerHandle().Get();
+			pxr::SdfLayerRefPtr UsdLayer = GetLayerHandle().Get();
 
 			if ( UsdLayer )
 			{
@@ -177,9 +177,9 @@ public:
 		}
 	}
 
-	TUsdStore< pxr::SdfLayerHandle > GetLayerHandle() const
+	TUsdStore< pxr::SdfLayerRefPtr > GetLayerHandle() const
 	{
-		return MakeUsdStore< pxr::SdfLayerHandle >( pxr::SdfLayer::FindOrOpen( LayerIdentifier.Get() ) );
+		return MakeUsdStore< pxr::SdfLayerRefPtr >( pxr::SdfLayer::FindOrOpen( LayerIdentifier.Get() ) );
 	}
 
 public:
@@ -256,7 +256,7 @@ public:
 				SNew( SImage )
 				.Image( this, &FUsdLayerMutedColumn::GetBrush, TreeItem )
 			];
-		
+
 		return Item;
 	}
 
@@ -314,7 +314,7 @@ public:
 			SNew(SImage)
 				.Image( this, &FUsdLayerEditColumn::GetCheckedImage, TreeItem )
 				.ColorAndOpacity( FEditorStyle::Get().GetWidgetStyle< FCheckBoxStyle >( "Checkbox" ).ForegroundColor );
-		
+
 		return Item;
 	}
 };
@@ -492,7 +492,7 @@ void SUsdLayersTreeView::OnEditSelectedLayer()
 
 	for ( FUsdLayersTreeItemRef SelectedItem : MySelectedItems )
 	{
-		TUsdStore< pxr::SdfLayerHandle > LayerHandle = SelectedItem->GetLayerHandle();
+		TUsdStore< pxr::SdfLayerRefPtr > LayerHandle = SelectedItem->GetLayerHandle();
 		if ( !LayerHandle.Get() || !CanEditLayer( SelectedItem ) )
 		{
 			continue;
@@ -518,7 +518,7 @@ void SUsdLayersTreeView::OnAddSubLayer()
 	for ( FUsdLayersTreeItemRef SelectedItem : MySelectedItems )
 	{
 		UsdLayersTreeViewImpl::InsertSubLayer( SelectedItem->GetLayerHandle(), SubLayerFile.GetValue() );
-		
+
 		break;
 	}
 
@@ -543,7 +543,7 @@ void SUsdLayersTreeView::OnNewSubLayer()
 			pxr::SdfLayerRefPtr SubLayerHandle = UsdUtils::CreateNewLayer( SelectedItem->UsdStage, *SubLayerFile.GetValue() ).Get();
 
 			UsdLayersTreeViewImpl::InsertSubLayer( SelectedItem->GetLayerHandle(), SubLayerFile.GetValue() );
-			
+
 			break;
 		}
 	}
