USD: Bug fixes related to sublayers #jira UE-85816 #rb JeanMichel.Dignard

Fix edit target being cleared when reloading;
Work around layer registry not caching relative paths correctly;
Add error handling when switching edit context;
Move InsertSubLayer to utils because we need to call it after creating the layer but before we edit it;

diff --git a/Engine/Plugins/Importers/USDImporter/Source/USDStageEditor/Private/SUSDLayersTreeView.cpp b/Engine/Plugins/Importers/USDImporter/Source/USDStageEditor/Private/SUSDLayersTreeView.cpp
--- a/Engine/Plugins/Importers/USDImporter/Source/USDStageEditor/Private/SUSDLayersTreeView.cpp
+++ b/Engine/Plugins/Importers/USDImporter/Source/USDStageEditor/Private/SUSDLayersTreeView.cpp
@@ -32,29 +32,6 @@
 
 #define LOCTEXT_NAMESPACE "SUSDLayersTreeView"
 
-namespace UsdLayersTreeViewImpl
-{
-	bool InsertSubLayer( const TUsdStore< pxr::SdfLayerRefPtr >& LayerHandle, const FString& SubLayerFile )
-	{
-		if ( LayerHandle.Get() )
-		{
-			FScopedUsdAllocs UsdAllocs;
-
-			std::string UsdLayerFilePath = LayerHandle.Get()->GetRealPath();
-			FString LayerFilePath = UsdToUnreal::ConvertString( UsdLayerFilePath );
-
-			FString SubLayerFilePath = FPaths::ConvertRelativePathToFull( SubLayerFile );
-			FPaths::MakePathRelativeTo( SubLayerFilePath, *LayerFilePath );
-
-			LayerHandle.Get()->InsertSubLayerPath( UnrealToUsd::ConvertString( *SubLayerFilePath ).Get() );
-
-			return true;
-		}
-
-		return false;
-	}
-}
-
 struct FUsdLayerData : public TSharedFromThis< FUsdLayerData >
 {
 	FText GetDisplayName() const { return DisplayName; }
@@ -517,7 +494,7 @@ void SUsdLayersTreeView::OnAddSubLayer()
 
 	for ( FUsdLayersTreeItemRef SelectedItem : MySelectedItems )
 	{
-		UsdLayersTreeViewImpl::InsertSubLayer( SelectedItem->GetLayerHandle(), SubLayerFile.GetValue() );
+		UsdUtils::InsertSubLayer( SelectedItem->GetLayerHandle(), *SubLayerFile.GetValue() );
 
 		break;
 	}
@@ -540,9 +517,7 @@ void SUsdLayersTreeView::OnNewSubLayer()
 		FScopedUsdAllocs UsdAllocs;
 		for ( FUsdLayersTreeItemRef SelectedItem : MySelectedItems )
 		{
-			pxr::SdfLayerRefPtr SubLayerHandle = UsdUtils::CreateNewLayer( SelectedItem->UsdStage, *SubLayerFile.GetValue() ).Get();
-
-			UsdLayersTreeViewImpl::InsertSubLayer( SelectedItem->GetLayerHandle(), SubLayerFile.GetValue() );
+			UsdUtils::CreateNewLayer( SelectedItem->UsdStage, SelectedItem->GetLayerHandle(), *SubLayerFile.GetValue() );
 
 			break;
 		}
diff --git a/Engine/Plugins/Importers/USDImporter/Source/USDStageEditor/Private/SUSDStage.cpp b/Engine/Plugins/Importers/USDImporter/Source/USDStageEditor/Private/SUSDStage.cpp
--- a/Engine/Plugins/Importers/USDImporter/Source/USDStageEditor/Private/SUSDStage.cpp
+++ b/Engine/Plugins/Importers/USDImporter/Source/USDStageEditor/Private/SUSDStage.cpp
@@ -398,6 +398,14 @@ void SUsdStage::OnReloadStage()
 		{
 			UsdStage->Reload();
 
+			// If we were editing an unsaved layer, when we reload the edit target will be cleared.
+			// We need to make sure we're always editing something or else UsdEditContext might trigger some errors
+			const pxr::UsdEditTarget& EditTarget = UsdStage->GetEditTarget();
+			if (!EditTarget.IsValid() || EditTarget.IsNull())
+			{
+				UsdStage->SetEditTarget(UsdStage->GetRootLayer());
+			}
+
 			if ( UsdLayersTreeView )
 			{
 				UsdLayersTreeView->Refresh( UsdStageActor.Get(), true );
diff --git a/Engine/Plugins/Importers/USDImporter/Source/USDUtilities/Private/USDLayerUtils.cpp b/Engine/Plugins/Importers/USDImporter/Source/USDUtilities/Private/USDLayerUtils.cpp
--- a/Engine/Plugins/Importers/USDImporter/Source/USDUtilities/Private/USDLayerUtils.cpp
+++ b/Engine/Plugins/Importers/USDImporter/Source/USDUtilities/Private/USDLayerUtils.cpp
@@ -3,6 +3,7 @@
 #include "USDLayerUtils.h"
 
 #include "USDTypesConversion.h"
+#include "USDErrorUtils.h"
 
 #include "DesktopPlatformModule.h"
 #include "Framework/Application/SlateApplication.h"
@@ -21,6 +22,26 @@
 
 #define LOCTEXT_NAMESPACE "USDLayerUtils"
 
+bool UsdUtils::InsertSubLayer( const TUsdStore< pxr::SdfLayerRefPtr >& ParentLayer, const TCHAR* SubLayerFile )
+{
+	if ( !ParentLayer.Get() )
+	{
+		return false;
+	}
+
+	FScopedUsdAllocs UsdAllocs;
+
+	std::string UsdLayerFilePath = ParentLayer.Get()->GetRealPath();
+	FString LayerFilePath = UsdToUnreal::ConvertString( UsdLayerFilePath );
+
+	FString SubLayerFilePath = FPaths::ConvertRelativePathToFull( SubLayerFile );
+	FPaths::MakePathRelativeTo( SubLayerFilePath, *LayerFilePath );
+
+	ParentLayer.Get()->InsertSubLayerPath( UnrealToUsd::ConvertString( *SubLayerFilePath ).Get() );
+
+	return true;
+}
+
 TOptional< FString > UsdUtils::BrowseUsdFile( EBrowseFileMode Mode, TSharedRef< const SWidget > OriginatingWidget )
 {
 	IDesktopPlatform* DesktopPlatform = FDesktopPlatformModule::Get();
@@ -44,14 +65,14 @@ TOptional< FString > UsdUtils::BrowseUsdFile( EBrowseFileMode Mode, TSharedRef<
 
 			if ( DesktopPlatform->OpenFileDialog( ParentWindowHandle, LOCTEXT( "ChooseFile", "Choose file").ToString(), TEXT(""), TEXT(""), TEXT("usd files (*.usd; *.usda; *.usdc)|*.usd; *.usda; *.usdc"), EFileDialogFlags::None, OutFiles ) )
 			{
-				return OutFiles[0];
+				return FPaths::ConvertRelativePathToFull(OutFiles[0]);
 			}
 			break;
 
 		case EBrowseFileMode::Save :
 			if ( DesktopPlatform->SaveFileDialog( ParentWindowHandle, LOCTEXT( "ChooseFile", "Choose file").ToString(), TEXT(""), TEXT(""), TEXT("usd files (*.usd; *.usda; *.usdc)|*.usd; *.usda; *.usdc"), EFileDialogFlags::None, OutFiles ) )
 			{
-				return OutFiles[0];
+				return FPaths::ConvertRelativePathToFull(OutFiles[0]);
 			}
 			break;
 
@@ -62,7 +83,7 @@ TOptional< FString > UsdUtils::BrowseUsdFile( EBrowseFileMode Mode, TSharedRef<
 	return {};
 }
 
-TUsdStore< pxr::SdfLayerRefPtr > UsdUtils::CreateNewLayer( TUsdStore< pxr::UsdStageRefPtr > UsdStage, const TCHAR* LayerFilePath )
+TUsdStore< pxr::SdfLayerRefPtr > UsdUtils::CreateNewLayer( TUsdStore< pxr::UsdStageRefPtr > UsdStage, const TUsdStore<pxr::SdfLayerRefPtr>& ParentLayer, const TCHAR* LayerFilePath )
 {
 	FScopedUsdAllocs UsdAllocs;
 
@@ -75,7 +96,17 @@ TUsdStore< pxr::SdfLayerRefPtr > UsdUtils::CreateNewLayer( TUsdStore< pxr::UsdSt
 		return {};
 	}
 
+	// New layer needs to be created and in the stage layer stack before we can edit it
+	UsdUtils::InsertSubLayer(ParentLayer, LayerFilePath);
+
+	UsdUtils::StartMonitoringErrors();
 	pxr::UsdEditContext UsdEditContext( UsdStage.Get(), LayerRef );
+	bool bHadErrors = UsdUtils::ShowErrorsAndStopMonitoring();
+
+	if (bHadErrors)
+	{
+		return {};
+	}
 
 	// Create default prim
 	FString PrimPath = TEXT("/") + FPaths::GetBaseFilename( UsdToUnreal::ConvertString( LayerRef->GetDisplayName() ) );
diff --git a/Engine/Plugins/Importers/USDImporter/Source/USDUtilities/Public/USDLayerUtils.h b/Engine/Plugins/Importers/USDImporter/Source/USDUtilities/Public/USDLayerUtils.h
--- a/Engine/Plugins/Importers/USDImporter/Source/USDUtilities/Public/USDLayerUtils.h
+++ b/Engine/Plugins/Importers/USDImporter/Source/USDUtilities/Public/USDLayerUtils.h
@@ -24,11 +24,14 @@ namespace UsdUtils
 		Save
 	};
 
+	/** Inserts the SubLayerFile path into ParentLayer as a sublayer */
+	USDUTILITIES_API bool InsertSubLayer( const TUsdStore< pxr::SdfLayerRefPtr >& ParentLayer, const TCHAR* SubLayerFile );
+
 	/** Opens a file dialog to open or save a USD file */
 	USDUTILITIES_API TOptional< FString > BrowseUsdFile( EBrowseFileMode Mode, TSharedRef< const SWidget > OriginatingWidget );
 
 	/** Creates a new layer with a default prim */
-	USDUTILITIES_API TUsdStore< pxr::SdfLayerRefPtr > CreateNewLayer( TUsdStore< pxr::UsdStageRefPtr > UsdStage, const TCHAR* LayerFilePath );
+	USDUTILITIES_API TUsdStore< pxr::SdfLayerRefPtr > CreateNewLayer( TUsdStore< pxr::UsdStageRefPtr > UsdStage, const TUsdStore<pxr::SdfLayerRefPtr>& ParentLayer, const TCHAR* LayerFilePath );
 }
 
 #endif // #if USE_USD_SDK
