USD look dev
- When assigning an Unreal material to a mesh from a USD Stage, write a custom attribute on the USD prim with the material path so that assignations are persisted if we reload the file.

diff --git a/Engine/Plugins/Importers/USDImporter/Source/USDImporter/Private/USDPrimResolverKind.cpp b/Engine/Plugins/Importers/USDImporter/Source/USDImporter/Private/USDPrimResolverKind.cpp
--- a/Engine/Plugins/Importers/USDImporter/Source/USDImporter/Private/USDPrimResolverKind.cpp
+++ b/Engine/Plugins/Importers/USDImporter/Source/USDImporter/Private/USDPrimResolverKind.cpp
@@ -30,8 +30,6 @@ void UUSDPrimResolverKind::FindActorsToSpawn_Recursive(FUSDSceneImportContext& I
 	// Parent/child hierarchy will be ignored unless the kind is a group.  Keep track of the current parent and use it 
 	TUsdStore< pxr::UsdPrim > GroupParent = ParentPrim;
 
-	FPlatformMisc::LowLevelOutputDebugStringf( TEXT("Prim Kind: %hs\n"), IUsdPrim::GetKind( Prim.Get() ).GetString().c_str() );
-
 	if (IUsdPrim::IsKindChildOf( *Prim, USDKindTypes::Component ))
 	{
 		bool bHasUnrealAssetPath = IUsdPrim::GetUnrealAssetPath( *Prim ).size() > 0;
diff --git a/Engine/Plugins/Importers/USDImporter/Source/USDSchemas/Private/USDGeomMeshTranslator.cpp b/Engine/Plugins/Importers/USDImporter/Source/USDSchemas/Private/USDGeomMeshTranslator.cpp
--- a/Engine/Plugins/Importers/USDImporter/Source/USDSchemas/Private/USDGeomMeshTranslator.cpp
+++ b/Engine/Plugins/Importers/USDImporter/Source/USDSchemas/Private/USDGeomMeshTranslator.cpp
@@ -4,6 +4,7 @@
 
 #if USE_USD_SDK
 
+#include "UnrealUSDWrapper.h"
 #include "USDConversionUtils.h"
 #include "USDGeomMeshConversion.h"
 #include "USDTypesConversion.h"
@@ -22,6 +23,7 @@
 #include "Misc/SecureHash.h"
 #include "PhysicsEngine/BodySetup.h"
 #include "StaticMeshAttributes.h"
+#include "UObject/SoftObjectPath.h"
 
 #include "USDIncludesStart.h"
 	#include "pxr/usd/usd/prim.h"
@@ -81,6 +83,22 @@ namespace UsdGeomMeshTranslatorImpl
 
 		FStaticMeshConstAttributes StaticMeshAttributes( *MeshDescription );
 
+		TArray< FString > UEMaterialsAttribute;
+		{
+			FScopedUsdAllocs UsdAllocs;
+
+			if ( pxr::UsdAttribute MaterialsAttribute = UsdPrim.GetAttribute( UnrealIdentifiers::MaterialAssignments ) )
+			{
+				pxr::VtStringArray UEMaterials;
+				MaterialsAttribute.Get( &UEMaterials, pxr::UsdTimeCode::EarliestTime() );
+
+				for ( const std::string& UEMaterial : UEMaterials )
+				{
+					UEMaterialsAttribute.Emplace( ANSI_TO_TCHAR( UEMaterial.c_str() ) );
+				}
+			}
+		}
+
 		for ( const FPolygonGroupID PolygonGroupID : MeshDescription->PolygonGroups().GetElementIDs() )
 		{
 			const FName& ImportedMaterialSlotName = StaticMeshAttributes.GetPolygonGroupMaterialSlotNames()[ PolygonGroupID ];
@@ -106,11 +124,19 @@ namespace UsdGeomMeshTranslatorImpl
 			}
 
 			UMaterialInterface* Material = nullptr;
-			TUsdStore< pxr::UsdPrim > MaterialPrim = UsdPrim.GetStage()->GetPrimAtPath( UnrealToUsd::ConvertPath( *ImportedMaterialSlotName.ToString() ).Get() );
 
-			if ( MaterialPrim.Get() )
+			if ( UEMaterialsAttribute.Num() > MeshMaterialIndex )
+			{
+				Material = Cast< UMaterialInterface >( FSoftObjectPath( UEMaterialsAttribute[ MeshMaterialIndex ] ).TryLoad() );
+			}
+			else
 			{
-				Material = Cast< UMaterialInterface >( PrimPathsToAssets.FindRef( UsdToUnreal::ConvertPath( MaterialPrim.Get().GetPrimPath() ) ) );
+				TUsdStore< pxr::UsdPrim > MaterialPrim = UsdPrim.GetStage()->GetPrimAtPath( UnrealToUsd::ConvertPath( *ImportedMaterialSlotName.ToString() ).Get() );
+
+				if ( MaterialPrim.Get() )
+				{
+					Material = Cast< UMaterialInterface >( PrimPathsToAssets.FindRef( UsdToUnreal::ConvertPath( MaterialPrim.Get().GetPrimPath() ) ) );
+				}
 			}
 
 			if ( Material == nullptr && bHasPrimDisplayColor )
diff --git a/Engine/Plugins/Importers/USDImporter/Source/USDStage/Private/USDStageActor.cpp b/Engine/Plugins/Importers/USDImporter/Source/USDStage/Private/USDStageActor.cpp
--- a/Engine/Plugins/Importers/USDImporter/Source/USDStage/Private/USDStageActor.cpp
+++ b/Engine/Plugins/Importers/USDImporter/Source/USDStage/Private/USDStageActor.cpp
@@ -622,7 +622,15 @@ void AUsdStageActor::OnPrimObjectPropertyChanged( UObject* ObjectBeingModified,
 				FScopedBlockNotices BlockNotices( UsdListener );
 
 				TUsdStore< pxr::UsdPrim > UsdPrim = UsdStage->GetPrimAtPath( UnrealToUsd::ConvertPath( *PrimPath ).Get() );
-				UnrealToUsd::ConvertSceneComponent( UsdStage, PrimSceneComponent, UsdPrim.Get() );
+
+				if ( UMeshComponent* MeshComponent = Cast< UMeshComponent >( PrimSceneComponent ) )
+				{
+					UnrealToUsd::ConvertMeshComponent( UsdStage, MeshComponent, UsdPrim.Get() );
+				}
+				else
+				{
+					UnrealToUsd::ConvertSceneComponent( UsdStage, PrimSceneComponent, UsdPrim.Get() );
+				}
 			}
 		}
 	}
diff --git a/Engine/Plugins/Importers/USDImporter/Source/USDUtilities/Private/USDPrimConversion.cpp b/Engine/Plugins/Importers/USDImporter/Source/USDUtilities/Private/USDPrimConversion.cpp
--- a/Engine/Plugins/Importers/USDImporter/Source/USDUtilities/Private/USDPrimConversion.cpp
+++ b/Engine/Plugins/Importers/USDImporter/Source/USDUtilities/Private/USDPrimConversion.cpp
@@ -2,10 +2,12 @@
 
 #include "USDPrimConversion.h"
 
+#include "UnrealUSDWrapper.h"
 #include "USDConversionUtils.h"
 #include "USDTypesConversion.h"
 
 #include "CineCameraComponent.h"
+#include "Components/MeshComponent.h"
 #include "Components/SceneComponent.h"
 
 #if USE_USD_SDK
@@ -93,7 +95,7 @@ bool UnrealToUsd::ConvertSceneComponent( const pxr::UsdStageRefPtr& Stage, const
 
 		if ( GfIsClose( UsdMatrix, UsdTransform, THRESH_VECTORS_ARE_NEAR ) )
 		{
-			return false;
+			return true;
 		}
 
 		bResetXFormStack = false;
@@ -124,4 +126,43 @@ bool UnrealToUsd::ConvertSceneComponent( const pxr::UsdStageRefPtr& Stage, const
 
 	return true;
 }
+
+bool UnrealToUsd::ConvertMeshComponent( const pxr::UsdStageRefPtr& Stage, const UMeshComponent* MeshComponent, pxr::UsdPrim& UsdPrim )
+{
+	if ( !UsdPrim || !MeshComponent )
+	{
+		return false;
+	}
+
+	if ( !ConvertSceneComponent( Stage, MeshComponent, UsdPrim ) )
+	{
+		return false;
+	}
+
+	if ( MeshComponent->GetNumMaterials() > 0 || UsdPrim.HasAttribute( UnrealIdentifiers::MaterialAssignments ) )
+	{
+		if ( pxr::UsdAttribute UEMaterialsAttribute = UsdPrim.CreateAttribute( UnrealIdentifiers::MaterialAssignments, pxr::SdfValueTypeNames->StringArray ) )
+		{
+			FScopedUsdAllocs UsdAllocs;
+
+			pxr::VtArray< std::string > UEMaterials = UsdUtils::GetUsdValue< pxr::VtArray< std::string > >( UEMaterialsAttribute );
+			UEMaterials.clear();
+
+			for ( int32 MaterialIndex = 0; MaterialIndex < MeshComponent->GetNumMaterials(); ++MaterialIndex )
+			{
+				if ( UMaterialInterface* AssignedMaterial = MeshComponent->GetMaterial( MaterialIndex ) )
+				{
+					FString AssignedMaterialPathName = AssignedMaterial->GetPathName();
+
+					UEMaterials.push_back( UnrealToUsd::ConvertString( *AssignedMaterialPathName ).Get() );
+				}
+			}
+
+			UEMaterialsAttribute.Set( UEMaterials );
+		}
+	}
+
+	return true;
+}
+
 #endif // #if USE_USD_SDK
diff --git a/Engine/Plugins/Importers/USDImporter/Source/USDUtilities/Public/USDPrimConversion.h b/Engine/Plugins/Importers/USDImporter/Source/USDUtilities/Public/USDPrimConversion.h
--- a/Engine/Plugins/Importers/USDImporter/Source/USDUtilities/Public/USDPrimConversion.h
+++ b/Engine/Plugins/Importers/USDImporter/Source/USDUtilities/Public/USDPrimConversion.h
@@ -21,6 +21,7 @@ PXR_NAMESPACE_OPEN_SCOPE
 PXR_NAMESPACE_CLOSE_SCOPE
 
 class UCineCameraComponent;
+class UMeshComponent;
 class USceneComponent;
 
 namespace UsdToUnreal
@@ -32,6 +33,7 @@ namespace UsdToUnreal
 namespace UnrealToUsd
 {
 	USDUTILITIES_API bool ConvertSceneComponent( const pxr::UsdStageRefPtr& Stage, const USceneComponent* SceneComponent, pxr::UsdPrim& UsdPrim );
+	USDUTILITIES_API bool ConvertMeshComponent( const pxr::UsdStageRefPtr& Stage, const UMeshComponent* MeshComponent, pxr::UsdPrim& UsdPrim );
 }
 
 #endif // #if USE_USD_SDK
diff --git a/Engine/Plugins/Importers/USDImporter/Source/USDUtilities/Public/USDTypesConversion.h b/Engine/Plugins/Importers/USDImporter/Source/USDUtilities/Public/USDTypesConversion.h
--- a/Engine/Plugins/Importers/USDImporter/Source/USDUtilities/Public/USDTypesConversion.h
+++ b/Engine/Plugins/Importers/USDImporter/Source/USDUtilities/Public/USDTypesConversion.h
@@ -181,6 +181,11 @@ namespace UnrealToUsd
 		return MakeUsdStore< std::string >( TCHAR_TO_ANSI( *InName.ToString() ) );
 	}
 
+	static TUsdStore< pxr::TfToken > ConvertToken( const TCHAR* InString )
+	{
+		return MakeUsdStore< pxr::TfToken >( TCHAR_TO_ANSI( InString ) );
+	}
+
 	static pxr::GfVec2f ConvertVector( const FVector2D& InValue )
 	{
 		return pxr::GfVec2f( InValue[0], InValue[1] );
diff --git a/Engine/Plugins/Importers/USDImporter/Source/UnrealUSDWrapper/Private/UnrealUSDWrapper.cpp b/Engine/Plugins/Importers/USDImporter/Source/UnrealUSDWrapper/Private/UnrealUSDWrapper.cpp
--- a/Engine/Plugins/Importers/USDImporter/Source/UnrealUSDWrapper/Private/UnrealUSDWrapper.cpp
+++ b/Engine/Plugins/Importers/USDImporter/Source/UnrealUSDWrapper/Private/UnrealUSDWrapper.cpp
@@ -59,9 +59,9 @@ static TUsdStore< UsdGeomXformCache > XFormCache;
 namespace UnrealIdentifiers
 {
 	/**
-	* Identifies the LOD variant set on a primitive which means this primitive has child prims that LOD meshes
-	* named LOD0, LOD1, LOD2, etc
-	*/
+	 * Identifies the LOD variant set on a primitive which means this primitive has child prims that LOD meshes
+	 * named LOD0, LOD1, LOD2, etc
+	 */
 	static const TfToken LOD("LOD");
 
 	static const TfToken AssetPath("unrealAssetPath");
@@ -75,8 +75,9 @@ namespace UnrealIdentifiers
 	static const TfToken GuidePurpose("guide");
 
 	static const TfToken MaterialRelationship("material:binding");
-}
 
+	const TfToken MaterialAssignments = TfToken("unrealMaterials");
+}
 
 void Log(const char* Format, ...)
 {
diff --git a/Engine/Plugins/Importers/USDImporter/Source/UnrealUSDWrapper/Public/UnrealUSDWrapper.h b/Engine/Plugins/Importers/USDImporter/Source/UnrealUSDWrapper/Public/UnrealUSDWrapper.h
--- a/Engine/Plugins/Importers/USDImporter/Source/UnrealUSDWrapper/Public/UnrealUSDWrapper.h
+++ b/Engine/Plugins/Importers/USDImporter/Source/UnrealUSDWrapper/Public/UnrealUSDWrapper.h
@@ -15,8 +15,9 @@
 #if USE_USD_SDK
 
 #include "USDIncludesStart.h"
-#include "pxr/pxr.h"
-#include "pxr/usd/usd/stageCache.h"
+	#include "pxr/pxr.h"
+	#include "pxr/base/tf/token.h"
+	#include "pxr/usd/usd/stageCache.h"
 #include "USDIncludesEnd.h"
 
 PXR_NAMESPACE_OPEN_SCOPE
@@ -221,4 +222,10 @@ public:
 #endif // #if USE_USD_SDK
 };
 
-
+namespace UnrealIdentifiers
+{
+#if USE_USD_SDK
+	/* Attribute name when assigning Unreal materials to UsdGeomMeshes */
+	extern UNREALUSDWRAPPER_API const pxr::TfToken MaterialAssignments;
+#endif // #if USE_USD_SDK
+}
diff --git a/Engine/Source/Editor/UnrealEd/Private/Kismet2/ComponentEditorUtils.cpp b/Engine/Source/Editor/UnrealEd/Private/Kismet2/ComponentEditorUtils.cpp
--- a/Engine/Source/Editor/UnrealEd/Private/Kismet2/ComponentEditorUtils.cpp
+++ b/Engine/Source/Editor/UnrealEd/Private/Kismet2/ComponentEditorUtils.cpp
@@ -845,6 +845,7 @@ bool FComponentEditorUtils::AttemptApplyMaterialToComponent(USceneComponent* Sce
 		}
 
 		SceneComponent->MarkRenderStateDirty();
+		SceneComponent->PostEditChange();
 		GEditor->OnSceneMaterialsModified();
 	}
 
