USD: Add error handling to Stage reloading #jira UE-86274 #rb JeanMichel.Dignard

diff --git a/Engine/Plugins/Importers/USDImporter/Source/USDStageEditor/Private/SUSDStage.cpp b/Engine/Plugins/Importers/USDImporter/Source/USDStageEditor/Private/SUSDStage.cpp
--- a/Engine/Plugins/Importers/USDImporter/Source/USDStageEditor/Private/SUSDStage.cpp
+++ b/Engine/Plugins/Importers/USDImporter/Source/USDStageEditor/Private/SUSDStage.cpp
@@ -396,7 +396,12 @@ void SUsdStage::OnReloadStage()
 
 		if ( UsdStage )
 		{
+			UsdUtils::StartMonitoringErrors();
 			UsdStage->Reload();
+			if (UsdUtils::ShowErrorsAndStopMonitoring())
+			{
+				return;
+			}
 
 			// If we were editing an unsaved layer, when we reload the edit target will be cleared.
 			// We need to make sure we're always editing something or else UsdEditContext might trigger some errors
