Fixed exported static meshes as USD being unable to be referenced. Wrapped the UsdGeomMesh in a generic prim.

diff --git a/Engine/Plugins/Importers/USDImporter/Source/USDExporter/Private/StaticMeshExporterUSD.cpp b/Engine/Plugins/Importers/USDImporter/Source/USDExporter/Private/StaticMeshExporterUSD.cpp
--- a/Engine/Plugins/Importers/USDImporter/Source/USDExporter/Private/StaticMeshExporterUSD.cpp
+++ b/Engine/Plugins/Importers/USDImporter/Source/USDExporter/Private/StaticMeshExporterUSD.cpp
@@ -51,15 +51,25 @@ bool UStaticMeshExporterUsd::ExportBinary( UObject* Object, const TCHAR* Type, F
 		// Set up axis
 		UsdUtils::SetUsdStageAxis( UsdStage, pxr::UsdGeomTokens->z );
 
-		TUsdStore< pxr::UsdGeomMesh > UsdGeomMeshStore = pxr::UsdGeomMesh::Define( UsdStage, UnrealToUsd::ConvertPath( *( TEXT("/") + StaticMesh->GetName() ) ).Get() );
-		pxr::UsdGeomMesh& UsdGeomMesh = UsdGeomMeshStore.Get();
+		FString RootPrimPath = ( TEXT("/") + StaticMesh->GetName() );
+
+		pxr::UsdPrim RootPrim = UsdStage->DefinePrim( UnrealToUsd::ConvertPath( *RootPrimPath ).Get() );
+
+		if ( !RootPrim )
+		{
+			return false;
+		}
+
+		FString MeshPrimPath = FPaths::Combine( RootPrimPath, StaticMesh->GetName() );
+
+		pxr::UsdGeomMesh UsdGeomMesh = pxr::UsdGeomMesh::Define( UsdStage, UnrealToUsd::ConvertPath( *MeshPrimPath ).Get() );
 
 		if ( UsdGeomMesh )
 		{
 			// Set default prim
-			UsdStage->SetDefaultPrim( UsdGeomMesh.GetPrim() );
+			UsdStage->SetDefaultPrim( RootPrim );
 
-			UnrealToUsd::ConvertStaticMesh( StaticMesh, UsdGeomMeshStore.Get() );
+			UnrealToUsd::ConvertStaticMesh( StaticMesh, UsdGeomMesh );
 		}
 
 		UsdStage->GetRootLayer()->Save();
