USD: Control USD sublayers with the Sequencer

diff --git a/Engine/Plugins/Importers/USDImporter/Source/USDStage/Private/USDLevelSequenceHelper.cpp b/Engine/Plugins/Importers/USDImporter/Source/USDStage/Private/USDLevelSequenceHelper.cpp
new file mode 100644
--- /dev/null
+++ b/Engine/Plugins/Importers/USDImporter/Source/USDStage/Private/USDLevelSequenceHelper.cpp
@@ -0,0 +1,763 @@
+// Copyright Epic Games, Inc. All Rights Reserved.
+
+#include "USDLevelSequenceHelper.h"
+
+#include "USDErrorUtils.h"
+#include "USDListener.h"
+#include "USDMemory.h"
+#include "USDStageActor.h"
+#include "USDTypesConversion.h"
+
+#include "Channels/MovieSceneChannelProxy.h"
+#include "CoreMinimal.h"
+#include "Editor.h"
+#include "LevelSequence.h"
+#include "MovieScene.h"
+#include "ObjectTools.h"
+#include "Sections/MovieSceneFloatSection.h"
+#include "Sections/MovieSceneSubSection.h"
+#include "Subsystems/AssetEditorSubsystem.h"
+#include "Tracks/MovieSceneFloatTrack.h"
+#include "Tracks/MovieSceneSubTrack.h"
+#include "UObject/UObjectGlobals.h"
+
+#if USE_USD_SDK
+#include "USDIncludesStart.h"
+
+#include "pxr/pxr.h"
+#include "pxr/usd/sdf/layer.h"
+#include "pxr/usd/sdf/layerUtils.h"
+#include "pxr/usd/usd/stage.h"
+#include "boost/functional/hash.hpp"
+
+#include "USDIncludesEnd.h"
+
+#define DEFAULT_FRAMERATE 24.0
+#define TIME_TRACK_NAME TEXT("Time")
+
+class FUsdLevelSequenceHelperImpl
+{
+public:
+	struct FLayerTimeInfo
+	{
+		FString Identifier;
+		FString FilePath;
+
+		TArray<FString> SubLayerIdentifiers;
+		TArray<double> OffsetTimeCodes;
+		TArray<double> Scales;
+
+		TOptional<double> StartTimeCode;
+		TOptional<double> EndTimeCode;
+		TOptional<double> TimeCodesPerSecond;
+
+		bool IsAnimated() const
+		{
+			return StartTimeCode.IsSet() && EndTimeCode.IsSet() && TimeCodesPerSecond.IsSet();
+		}
+	};
+
+	FUsdLevelSequenceHelperImpl(TWeakObjectPtr<AUsdStageActor> InStageActor);
+	~FUsdLevelSequenceHelperImpl();
+
+	/** Returns the FLayerTimeInfo corresponding to the root layer */
+	FLayerTimeInfo* GetRootLayerInfo();
+
+	/** Creates all the ULevelSequence assets for all FLayerTimeInfos we have */
+	void CreateLevelSequences();
+
+	/**
+	 * Creates the subsequence tracks on the ULevelSequence corresponding to Info, as well
+	 * as recursively for all sublayers of Info
+	 */
+	void CreateSubsequenceTrack(const FUsdLevelSequenceHelperImpl::FLayerTimeInfo* Info) const;
+
+	/** Creates a time track on the ULevelSequence corresponding to Info */
+	void CreateTimeTrack(const FUsdLevelSequenceHelperImpl::FLayerTimeInfo* Info) const;
+
+	/** Construct and cache a FLayerTimeInfo object for each SdfLayer within UsdStage */
+	void RebuildLayerTimeInfoCache(const pxr::UsdStageRefPtr& UsdStage);
+
+	/**
+	 * Initializes the StageActor Start and End timecodes (as well as our root layer info)
+	 * with the "top-level" timecodes from UsdStage
+	 */
+	void InitializeActorTimeCodes(const pxr::UsdStageRefPtr& UsdStage);
+
+private:
+	/** Create a FLayerTimeInfo object with data from LayerHandle, and calls itself for all of its sublayers */
+	FLayerTimeInfo* RecursivelyCreateLayerTimeInfo(const TUsdStore<pxr::SdfLayerHandle>& LayerHandle, TSet<FString>& ExploredFiles);
+
+	/**
+	 * Propagate timecodes for leaf FLayerTimeInfo all the way up to Info. We use this to calculate enveloping
+	 * time codes that we can use for parent layers, in case they don't have time code information but have animated sublayers
+	 */
+	void RecursivelyPropagateTimeCodes(FUsdLevelSequenceHelperImpl::FLayerTimeInfo* Info);
+
+	/** Updates a FLayerTimeInfo object with new offset/scale values when Section has been moved by the user */
+	void UpdateInfoFromSection(FUsdLevelSequenceHelperImpl::FLayerTimeInfo* Info, const UMovieSceneSubSection* Section, const UMovieSceneSequence* Sequence);
+
+	/** Capture whenever a UMovieSceneSubSection has transacted and call UpdateInfoFromSection */
+	void OnObjectTransacted(UObject* Object, const class FTransactionObjectEvent& Event);
+
+private:
+	FString RootLayerInfoIdentifier;
+	TMap<FString, FUsdLevelSequenceHelperImpl::FLayerTimeInfo> LayerTimeInfosByIdentifier;
+	TMap<FName, FString> LayerInfoIdentifierByLevelSequenceName;
+	TWeakObjectPtr<AUsdStageActor> StageActor;
+	FGuid StageActorBinding;
+	static EObjectFlags DefaultObjFlags;
+	FDelegateHandle OnObjectTransactedHandle;
+};
+
+EObjectFlags FUsdLevelSequenceHelperImpl::DefaultObjFlags = EObjectFlags::RF_Transactional | EObjectFlags::RF_Transient | EObjectFlags::RF_Public;
+
+FUsdLevelSequenceHelperImpl::FUsdLevelSequenceHelperImpl(TWeakObjectPtr<AUsdStageActor> InStageActor)
+	: StageActor(InStageActor)
+	, StageActorBinding()
+{
+	OnObjectTransactedHandle = FCoreUObjectDelegates::OnObjectTransacted.AddRaw(this, &FUsdLevelSequenceHelperImpl::OnObjectTransacted);
+}
+
+FUsdLevelSequenceHelperImpl::~FUsdLevelSequenceHelperImpl()
+{
+	FCoreUObjectDelegates::OnObjectTransacted.Remove(OnObjectTransactedHandle);
+	OnObjectTransactedHandle.Reset();
+}
+
+FUsdLevelSequenceHelperImpl::FLayerTimeInfo* FUsdLevelSequenceHelperImpl::GetRootLayerInfo()
+{
+	return LayerTimeInfosByIdentifier.Find(RootLayerInfoIdentifier);
+}
+
+void FUsdLevelSequenceHelperImpl::CreateLevelSequences()
+{
+	AUsdStageActor* ValidStageActor = StageActor.Get();
+	if (!ValidStageActor)
+	{
+		return;
+	}
+
+	double TimeCodesPerSecond = ValidStageActor->TimeCodesPerSecond;
+
+	UE_LOG(LogUsdStage, Verbose, TEXT("CreateLevelSequences: Initializing level sequence with '%f' fps"), TimeCodesPerSecond);
+	if ( ValidStageActor->LevelSequence )
+	{
+		return;
+	}
+
+	const FLayerTimeInfo* RootLayerInfo = GetRootLayerInfo();
+	if (!RootLayerInfo)
+	{
+		return;
+	}
+	FString Name = ObjectTools::SanitizeObjectName(FPaths::GetCleanFilename(RootLayerInfo->FilePath));
+
+	// Create level sequence for root layer (visible on StageActor's details panel)
+	ValidStageActor->LevelSequence = NewObject<ULevelSequence>(GetTransientPackage(), *Name, FUsdLevelSequenceHelperImpl::DefaultObjFlags);
+	ValidStageActor->LevelSequence->Initialize();
+	ValidStageActor->LevelSequence->MovieScene->SetDisplayRate(FFrameRate(TimeCodesPerSecond, 1));
+	LayerInfoIdentifierByLevelSequenceName.Add(ValidStageActor->LevelSequence->GetFName(), RootLayerInfoIdentifier);
+
+	// Bind stage actor
+	UMovieScene* MovieScene = ValidStageActor->LevelSequence->GetMovieScene();
+	StageActorBinding = MovieScene->AddPossessable(ValidStageActor->GetActorLabel(), ValidStageActor->GetClass());
+	ValidStageActor->LevelSequence->BindPossessableObject(StageActorBinding, *ValidStageActor, ValidStageActor->GetWorld());
+
+	// Create level sequences for all sub layers (accessible via the main level sequence but otherwise hidden)
+	for (TPair<FString, FUsdLevelSequenceHelperImpl::FLayerTimeInfo>& InfoPair : LayerTimeInfosByIdentifier)
+	{
+		FUsdLevelSequenceHelperImpl::FLayerTimeInfo& Info = InfoPair.Value;
+		const FString& Identifier = InfoPair.Key;
+		if (Identifier == RootLayerInfoIdentifier)
+		{
+			continue;
+		}
+
+		ULevelSequence*& Sequence = ValidStageActor->SubLayerLevelSequencesByIdentifier.FindOrAdd(Identifier);
+		if (!Sequence)
+		{
+			FString SubName = ObjectTools::SanitizeObjectName(FPaths::GetCleanFilename(Info.FilePath));
+
+			Sequence = NewObject<ULevelSequence>(GetTransientPackage(), *SubName, FUsdLevelSequenceHelperImpl::DefaultObjFlags);
+			Sequence->Initialize();
+			LayerInfoIdentifierByLevelSequenceName.Add(Sequence->GetFName(), Identifier);
+
+			UMovieScene* SubMovieScene = Sequence->MovieScene;
+			if (!SubMovieScene)
+			{
+				continue;
+			}
+			SubMovieScene->SetDisplayRate(FFrameRate(Info.TimeCodesPerSecond.GetValue(), 1));
+
+			UE_LOG(LogUsdStage, Verbose, TEXT("CreateLevelSequences: Created new LevelSequence '%s' for layer '%s'"), *Sequence->GetName(), *Identifier);
+		}
+	}
+}
+
+void FUsdLevelSequenceHelperImpl::CreateSubsequenceTrack(const FUsdLevelSequenceHelperImpl::FLayerTimeInfo* Info) const
+{
+	if (!Info)
+	{
+		return;
+	}
+
+	AUsdStageActor* ValidStageActor = StageActor.Get();
+	if (!ValidStageActor)
+	{
+		return;
+	}
+
+	UE_LOG(LogUsdStage, Verbose, TEXT("CreateSubsequenceSections: Initializing for layer '%s'"), *Info->Identifier);
+
+	// Recurse first as we'll need the sublayer LevelSequences to be completed to insert them as subsections
+	bool bFoundAnimatedSubLayer = false;
+	for (const FString& SubLayerIdentifier : Info->SubLayerIdentifiers)
+	{
+		if (const FUsdLevelSequenceHelperImpl::FLayerTimeInfo* SubLayerInfo = LayerTimeInfosByIdentifier.Find(SubLayerIdentifier))
+		{
+			if (SubLayerInfo->IsAnimated())
+			{
+				CreateSubsequenceTrack(SubLayerInfo);
+				bFoundAnimatedSubLayer = true;
+			}
+		}
+	}
+
+	ULevelSequence* Sequence = nullptr;
+	if (Info->Identifier == RootLayerInfoIdentifier)
+	{
+		Sequence = ValidStageActor->LevelSequence;
+	}
+	else if (ULevelSequence** FoundSequence = ValidStageActor->SubLayerLevelSequencesByIdentifier.Find(Info->Identifier))
+	{
+		Sequence = *FoundSequence;
+	}
+
+	if (Sequence != nullptr && Info->IsAnimated() && bFoundAnimatedSubLayer)
+	{
+		UE_LOG(LogUsdStage, Verbose, TEXT("CreateSubsequenceSections: Layer '%s' is animated"), *Info->Identifier);
+
+		UMovieScene* MovieScene = Sequence->GetMovieScene();
+		if (!MovieScene)
+		{
+			return;
+		}
+
+		FFrameRate FrameRate = MovieScene->GetTickResolution();
+
+		UMovieSceneSubTrack* SubTrack = MovieScene->FindMasterTrack<UMovieSceneSubTrack>();
+		if (!SubTrack)
+		{
+			SubTrack = MovieScene->AddMasterTrack<UMovieSceneSubTrack>();
+		}
+		SubTrack->RemoveAllAnimationData();
+
+		for (int32 SubLayerIndex = 0; SubLayerIndex < Info->SubLayerIdentifiers.Num(); ++SubLayerIndex)
+		{
+			const FUsdLevelSequenceHelperImpl::FLayerTimeInfo* SubInfo = LayerTimeInfosByIdentifier.Find(Info->SubLayerIdentifiers[SubLayerIndex]);
+			if (!SubInfo || !SubInfo->IsAnimated())
+			{
+				continue;
+			}
+
+			ULevelSequence** SubSequence = ValidStageActor->SubLayerLevelSequencesByIdentifier.Find(SubInfo->Identifier);
+			if (!SubSequence || !*SubSequence)
+			{
+				UE_LOG(LogUsdStage, Warning, TEXT("CreateSubsequenceSections: Invalid LevelSequence for sublayer '%s'"), *SubInfo->Identifier);
+				continue;
+			}
+
+			double SubOffsetTimeCode = Info->OffsetTimeCodes[SubLayerIndex];
+			double SubScale = Info->Scales[SubLayerIndex];
+			double SubTimeCodesPerSecond = SubInfo->TimeCodesPerSecond.GetValue();
+			double SubStartSeconds  = (SubOffsetTimeCode + SubScale * SubInfo->StartTimeCode.GetValue()) / SubTimeCodesPerSecond;
+			double SubEndSeconds 	= (SubOffsetTimeCode + SubScale * SubInfo->EndTimeCode.GetValue())   / SubTimeCodesPerSecond;
+
+			FFrameNumber StartFrame = FrameRate.AsFrameNumber(SubStartSeconds);
+			int32 Duration = FrameRate.AsFrameNumber(SubEndSeconds).Value - StartFrame.Value;
+			UMovieSceneSubSection* NewSection = SubTrack->AddSequenceOnRow(*SubSequence, StartFrame, Duration, INDEX_NONE);
+
+			UE_LOG(LogUsdStage, Verbose, TEXT("CreateSubsequenceSections: Layer '%s' received subsection for layer '%s' starting at frame '%d' and with duration '%d'"),
+				*Info->Identifier,
+				*SubInfo->Identifier,
+				StartFrame.Value,
+				Duration);
+		}
+
+		// Setup view/work ranges
+		double TimeCodesPerSecond = Info->TimeCodesPerSecond.GetValue();
+		double StartTimeSeconds   = Info->StartTimeCode.GetValue() / TimeCodesPerSecond;
+		double EndTimeSeconds     = Info->EndTimeCode.GetValue() / TimeCodesPerSecond;
+
+		FFrameNumber StartFrame = FrameRate.AsFrameNumber(StartTimeSeconds);
+		FFrameNumber EndFrame = FrameRate.AsFrameNumber(EndTimeSeconds);
+		TRange<FFrameNumber> TimeRange = TRange<FFrameNumber>::Inclusive(StartFrame, EndFrame);
+
+		MovieScene->SetPlaybackRange(TimeRange);
+		MovieScene->SetViewRange(StartTimeSeconds - 1.0f, 1.0f + EndTimeSeconds);
+		MovieScene->SetWorkingRange(StartTimeSeconds - 1.0f, 1.0f + EndTimeSeconds);
+
+		UE_LOG(LogUsdStage, Verbose, TEXT("CreateSubsequenceSections: Layer '%s' received view range [%f, %f]"),
+			*Info->Identifier,
+			StartTimeSeconds - 1.0f,
+			EndTimeSeconds + 1.0f);
+	}
+}
+
+void FUsdLevelSequenceHelperImpl::CreateTimeTrack(const FUsdLevelSequenceHelperImpl::FLayerTimeInfo* Info) const
+{
+	if (!Info || !Info->IsAnimated())
+	{
+		return;
+	}
+
+	AUsdStageActor* ValidStageActor = StageActor.Get();
+	if (!ValidStageActor)
+	{
+		return;
+	}
+
+	ULevelSequence* Sequence = nullptr;
+	if (Info->Identifier == RootLayerInfoIdentifier)
+	{
+		Sequence = ValidStageActor->LevelSequence;
+	}
+	else if (ULevelSequence** FoundSequence = ValidStageActor->SubLayerLevelSequencesByIdentifier.Find(Info->Identifier))
+	{
+		Sequence = *FoundSequence;
+	}
+	if (!Sequence)
+	{
+		return;
+	}
+
+	UMovieScene* MovieScene = Sequence->GetMovieScene();
+	if(!MovieScene)
+	{
+		return;
+	}
+
+	UMovieSceneFloatTrack* TimeTrack = MovieScene->FindTrack<UMovieSceneFloatTrack>(StageActorBinding, FName(TIME_TRACK_NAME));
+	if (TimeTrack)
+	{
+		TimeTrack->RemoveAllAnimationData();
+	}
+	else
+	{
+		TimeTrack = MovieScene->AddTrack<UMovieSceneFloatTrack>(StageActorBinding);
+		TimeTrack->SetPropertyNameAndPath(FName(TIME_TRACK_NAME), "Time");
+
+		MovieScene->SetEvaluationType(EMovieSceneEvaluationType::FrameLocked);
+	}
+
+	double StartTimeCode = Info->StartTimeCode.GetValue();
+	double EndTimeCode = Info->EndTimeCode.GetValue();
+	double TimeCodesPerSecond = Info->TimeCodesPerSecond.Get(DEFAULT_FRAMERATE);
+
+	FFrameRate DestFrameRate = MovieScene->GetTickResolution();
+	FFrameNumber StartFrame  = DestFrameRate.AsFrameNumber(StartTimeCode / TimeCodesPerSecond);
+	FFrameNumber EndFrame    = DestFrameRate.AsFrameNumber(EndTimeCode / TimeCodesPerSecond);
+
+	bool bSectionAdded = false;
+	UMovieSceneFloatSection* TimeSection = Cast<UMovieSceneFloatSection>(TimeTrack->FindOrAddSection(0, bSectionAdded));
+	TimeSection->EvalOptions.CompletionMode = EMovieSceneCompletionMode::KeepState;
+	TimeSection->SetRange(TRange<FFrameNumber>::All());
+
+	TArray<FFrameNumber> FrameNumbers;
+	FrameNumbers.Add(StartFrame);
+	FrameNumbers.Add(EndFrame);
+
+	TArray<FMovieSceneFloatValue> FrameValues;
+	FrameValues.Add_GetRef(FMovieSceneFloatValue(StartTimeCode)).InterpMode = ERichCurveInterpMode::RCIM_Linear;
+	FrameValues.Add_GetRef(FMovieSceneFloatValue(EndTimeCode)).InterpMode = ERichCurveInterpMode::RCIM_Linear;
+
+	FMovieSceneFloatChannel* TimeChannel = TimeSection->GetChannelProxy().GetChannel<FMovieSceneFloatChannel>(0);
+	TimeChannel->Set(FrameNumbers, FrameValues);
+
+	TRange<FFrameNumber> TimeRange = TRange<FFrameNumber>::Inclusive(StartFrame, EndFrame);
+
+	MovieScene->SetPlaybackRange( TimeRange );
+	MovieScene->SetViewRange( StartTimeCode / TimeCodesPerSecond -1.0f, 1.0f + EndTimeCode / TimeCodesPerSecond );
+	MovieScene->SetWorkingRange( StartTimeCode / TimeCodesPerSecond -1.0f, 1.0f + EndTimeCode / TimeCodesPerSecond );
+}
+
+void FUsdLevelSequenceHelperImpl::RebuildLayerTimeInfoCache(const pxr::UsdStageRefPtr& UsdStage)
+{
+	if (!UsdStage)
+	{
+		return;
+	}
+
+	UE_LOG(LogUsdStage, Verbose, TEXT("RebuildLayerTimeInfoCache"));
+
+	LayerTimeInfosByIdentifier.Reset();
+	LayerInfoIdentifierByLevelSequenceName.Reset();
+
+	TUsdStore<pxr::SdfLayerHandle> RootLayerHandle = MakeUsdStore<pxr::SdfLayerHandle>(UsdStage->GetRootLayer());
+
+	// Used to catch N-item infinite loops (e.g. A/B/C/A/B/C/etc)
+	TSet<FString> ExploredPaths;
+	FLayerTimeInfo* RootLayerInfo = RecursivelyCreateLayerTimeInfo(RootLayerHandle.Get(), ExploredPaths);
+	RootLayerInfoIdentifier = RootLayerInfo->Identifier;
+
+	RecursivelyPropagateTimeCodes(RootLayerInfo);
+}
+
+void FUsdLevelSequenceHelperImpl::InitializeActorTimeCodes(const pxr::UsdStageRefPtr& UsdStage)
+{
+	AUsdStageActor* ValidStageActor = StageActor.Get();
+	FUsdLevelSequenceHelperImpl::FLayerTimeInfo* RootInfo = GetRootLayerInfo();
+	if (!ValidStageActor || !RootInfo)
+	{
+		return;
+	}
+
+	{
+		FScopedUsdAllocs Allocs;
+		pxr::SdfLayerHandle RootLayer = UsdStage->GetRootLayer();
+		pxr::SdfLayerHandle SessionLayer = UsdStage->GetSessionLayer();
+
+		// If the stage has a valid start time code, use that
+		if ((RootLayer && RootLayer->HasStartTimeCode()) ||
+			(SessionLayer && SessionLayer->HasStartTimeCode()))
+		{
+			RootInfo->StartTimeCode = UsdStage->GetStartTimeCode();
+		}
+		// Or else, instead of wrapping the first animated layer, start at timecode zero
+		else
+		{
+			RootInfo->StartTimeCode = 0;
+		}
+
+		// If the stage has a valid end time code, use that. Or else we'll
+		// keep the one we calculated (range of all sublayers)
+		if ((RootLayer && RootLayer->HasEndTimeCode()) ||
+			(SessionLayer && SessionLayer->HasEndTimeCode()))
+		{
+			RootInfo->EndTimeCode = UsdStage->GetEndTimeCode();
+		}
+	}
+
+	ValidStageActor->StartTimeCode = RootInfo->StartTimeCode.Get(0);
+	ValidStageActor->EndTimeCode = RootInfo->EndTimeCode.Get(0);
+	ValidStageActor->TimeCodesPerSecond = RootInfo->TimeCodesPerSecond.Get(DEFAULT_FRAMERATE);
+	UE_LOG(LogUsdStage, Verbose, TEXT("Initialized AUsdStageActor's Stage. StartTimeCode: %f, EndTimeCode: %f, TimeCodesPerSeconds: %f"), ValidStageActor->StartTimeCode, ValidStageActor->EndTimeCode, ValidStageActor->TimeCodesPerSecond);
+}
+
+FUsdLevelSequenceHelperImpl::FLayerTimeInfo* FUsdLevelSequenceHelperImpl::RecursivelyCreateLayerTimeInfo(const TUsdStore<pxr::SdfLayerHandle>& LayerHandleWrapper, TSet<FString>& ExploredFiles)
+{
+	const pxr::SdfLayerHandle& LayerHandle = LayerHandleWrapper.Get();
+
+	// We want to keep track of all SdfLayers we find, but each layer should only appear once in a path to the root, or else we infinitely loop
+	FString Identifier = UsdToUnreal::ConvertString(LayerHandle->GetIdentifier());
+	if (ExploredFiles.Contains(Identifier))
+	{
+		UE_LOG(LogUsdStage, Warning, TEXT("Detected infinite SubLayer recursion when visiting layer with identifier '%s'!"), *Identifier);
+		return nullptr;
+	}
+
+	FString FilePath = UsdToUnreal::ConvertString(LayerHandle->GetRealPath());
+	ExploredFiles.Add(FilePath);
+
+	FUsdLevelSequenceHelperImpl::FLayerTimeInfo LayerTimeInfo;
+	LayerTimeInfo.Identifier = Identifier;
+	LayerTimeInfo.FilePath = FilePath;
+	LayerTimeInfo.StartTimeCode      = LayerHandle->HasStartTimeCode() ? LayerHandle->GetStartTimeCode() : TOptional<double>();
+	LayerTimeInfo.EndTimeCode        = LayerHandle->HasEndTimeCode() ? LayerHandle->GetEndTimeCode() : TOptional<double>();
+	LayerTimeInfo.TimeCodesPerSecond = LayerHandle->HasTimeCodesPerSecond() ? LayerHandle->GetTimeCodesPerSecond() : TOptional<double>();
+
+	UE_LOG(LogUsdStage, Verbose, TEXT("Creating layer time info for layer '%s'. Original timecodes: '%s', '%s' and '%s'"),
+		*LayerTimeInfo.Identifier,
+		LayerTimeInfo.StartTimeCode.IsSet() ? *LexToString(LayerTimeInfo.StartTimeCode.GetValue()) : TEXT("null"),
+		LayerTimeInfo.EndTimeCode.IsSet() ? *LexToString(LayerTimeInfo.EndTimeCode.GetValue()) : TEXT("null"),
+		LayerTimeInfo.TimeCodesPerSecond.IsSet() ? *LexToString(LayerTimeInfo.TimeCodesPerSecond.GetValue()) : TEXT("null"));
+
+	TUsdStore<std::vector<std::string>> SubLayerPaths;
+	TUsdStore<std::vector<pxr::SdfLayerOffset>> SubLayerOffsetScales;
+	{
+		FScopedUsdAllocs USDAllocs;
+		SubLayerPaths = LayerHandle->GetSubLayerPaths();
+		SubLayerOffsetScales = LayerHandle->GetSubLayerOffsets();
+	}
+
+	for (int64 LayerIndex = 0; LayerIndex < static_cast<int64>(LayerHandle->GetNumSubLayerPaths()); ++LayerIndex)
+	{
+		TUsdStore<pxr::SdfLayerHandle> SubLayer;
+		{
+			FScopedUsdAllocs USDAllocs;
+			std::string RelativePath = pxr::SdfComputeAssetPathRelativeToLayer(LayerHandle, SubLayerPaths.Get()[LayerIndex]);
+			SubLayer = LayerHandle->Find(RelativePath);
+		}
+
+		if (!SubLayer.Get())
+		{
+			continue;
+		}
+
+		const pxr::SdfLayerOffset& SubLayerOffsetScale = SubLayerOffsetScales.Get()[LayerIndex];
+		double SubOffset = 0.0;
+		double SubScale = 1.0;
+		if (SubLayerOffsetScale.IsValid())
+		{
+			SubOffset = SubLayerOffsetScale.GetOffset();
+			SubScale = SubLayerOffsetScale.GetScale();
+		}
+
+		if (FUsdLevelSequenceHelperImpl::FLayerTimeInfo* SubLayerInfo = RecursivelyCreateLayerTimeInfo(SubLayer, ExploredFiles))
+		{
+			UE_LOG(LogUsdStage, Verbose, TEXT("Adding sublayer '%s' to layer '%s'. Offset: '%f', Scale: '%f'"),
+				*SubLayerInfo->Identifier,
+				*LayerTimeInfo.Identifier,
+				SubOffset,
+				SubScale);
+
+			LayerTimeInfo.SubLayerIdentifiers.Add(SubLayerInfo->Identifier);
+			LayerTimeInfo.OffsetTimeCodes.Add(SubOffset);
+			LayerTimeInfo.Scales.Add(SubScale);
+		}
+	}
+
+	ExploredFiles.Remove(FilePath);
+	return &(LayerTimeInfosByIdentifier.Add(Identifier, LayerTimeInfo));
+}
+
+void FUsdLevelSequenceHelperImpl::RecursivelyPropagateTimeCodes(FUsdLevelSequenceHelperImpl::FLayerTimeInfo* Info)
+{
+	if (!Info)
+	{
+		return;
+	}
+
+	// Recurse first as the timecodes must propagate upwards to the root
+	for (const FString& SubLayerIdentifier : Info->SubLayerIdentifiers)
+	{
+		if (FUsdLevelSequenceHelperImpl::FLayerTimeInfo* SubLayerInfo = LayerTimeInfosByIdentifier.Find(SubLayerIdentifier))
+		{
+			RecursivelyPropagateTimeCodes(SubLayerInfo);
+		}
+	}
+
+	if (!Info->TimeCodesPerSecond.IsSet())
+	{
+		Info->TimeCodesPerSecond = DEFAULT_FRAMERATE;
+	}
+
+	// Don't propagate child timecodes if the layer has any manual timecodes set.
+	// That should work as kind of an user override
+	bool bHaveStartCode	= Info->StartTimeCode.IsSet();
+	bool bHaveEndCode	= Info->EndTimeCode.IsSet();
+	if (bHaveStartCode && bHaveEndCode)
+	{
+		UE_LOG(LogUsdStage, Verbose, TEXT("RecursivelyPropagateTimeCodes: Aborting for layer '%s' as it already has '%s', '%s' and '%s' set"),
+			*Info->Identifier,
+			Info->StartTimeCode.IsSet() ? *LexToString(Info->StartTimeCode.GetValue()) : TEXT("null"),
+			Info->EndTimeCode.IsSet() ? *LexToString(Info->EndTimeCode.GetValue()) : TEXT("null"),
+			Info->TimeCodesPerSecond.IsSet() ? *LexToString(Info->TimeCodesPerSecond.GetValue()) : TEXT("null"));
+		return;
+	}
+
+	// Calculate bounds for our level sequence section. We do this in seconds as
+	// our subs may each have their own different TimeCodesPerSecond
+	TOptional<double> MinStartSeconds;
+	TOptional<double> MaxEndSeconds;
+	for (int32 SubLayerIndex = 0; SubLayerIndex < Info->SubLayerIdentifiers.Num(); ++SubLayerIndex)
+	{
+		FUsdLevelSequenceHelperImpl::FLayerTimeInfo* SubLayerInfo = LayerTimeInfosByIdentifier.Find(Info->SubLayerIdentifiers[SubLayerIndex]);
+		if (!SubLayerInfo || !SubLayerInfo->IsAnimated())
+		{
+			continue;
+		}
+
+		double OffsetForSub = Info->OffsetTimeCodes[SubLayerIndex];
+		double ScaleForSub  = Info->Scales[SubLayerIndex];
+
+		double SubCodesPerSecond = SubLayerInfo->TimeCodesPerSecond.GetValue();
+		double SubStartSeconds	 = (OffsetForSub + ScaleForSub * SubLayerInfo->StartTimeCode.GetValue()) / SubCodesPerSecond;
+		double SubEndSeconds	 = (OffsetForSub + ScaleForSub * SubLayerInfo->EndTimeCode.GetValue()) / SubCodesPerSecond;
+
+		MinStartSeconds = FMath::Min(SubStartSeconds, MinStartSeconds.Get(DBL_MAX));
+		MaxEndSeconds	= FMath::Max(SubEndSeconds,   MaxEndSeconds.Get(-DBL_MAX));
+	}
+
+	double TimeCodesPerSecond = Info->TimeCodesPerSecond.GetValue();
+	if (!bHaveStartCode && MinStartSeconds.IsSet())
+	{
+		Info->StartTimeCode = MinStartSeconds.GetValue() * TimeCodesPerSecond;
+	}
+	if (!bHaveEndCode && MaxEndSeconds.IsSet())
+	{
+		Info->EndTimeCode = MaxEndSeconds.GetValue() * TimeCodesPerSecond;
+	}
+
+	UE_LOG(LogUsdStage, Verbose, TEXT("RecursivelyPropagateTimeCodes: Finished for layer '%s': '%s', '%s' and '%s' set"),
+		*Info->Identifier,
+		Info->StartTimeCode.IsSet() ? *LexToString(Info->StartTimeCode.GetValue()) : TEXT("null"),
+		Info->EndTimeCode.IsSet() ? *LexToString(Info->EndTimeCode.GetValue()) : TEXT("null"),
+		Info->TimeCodesPerSecond.IsSet() ? *LexToString(Info->TimeCodesPerSecond.GetValue()) : TEXT("null"));
+}
+
+void FUsdLevelSequenceHelperImpl::UpdateInfoFromSection(FUsdLevelSequenceHelperImpl::FLayerTimeInfo* Info, const UMovieSceneSubSection* Section, const UMovieSceneSequence* Sequence)
+{
+	if (!Info || !Section || !Sequence)
+	{
+		return;
+	}
+
+	TUsdStore<pxr::SdfLayerHandle> Layer;
+	{
+		FScopedUsdAllocs Allocs;
+		Layer = MakeUsdStore<pxr::SdfLayerHandle>(pxr::SdfLayer::Find(UnrealToUsd::ConvertString(*Info->Identifier).Get()));
+		if (!Layer.Get())
+		{
+			UE_LOG(LogUsdStage, Warning, TEXT("Failed to update sublayer '%s'"), *Info->Identifier);
+			return;
+		}
+	}
+
+	UE_LOG(LogUsdStage, Verbose, TEXT("Updating LevelSequence '%s' for sublayer '%s'"), *Sequence->GetName(), *Info->Identifier);
+
+	UMovieScene* MovieScene = Sequence->GetMovieScene();
+	UMovieSceneSequence* SubSequence = Section->GetSequence();
+	if (!MovieScene || !SubSequence)
+	{
+		return;
+	}
+
+	FString* InfoIdentifier = LayerInfoIdentifierByLevelSequenceName.Find(SubSequence->GetFName());
+	if (!InfoIdentifier)
+	{
+		return;
+	}
+
+	FUsdLevelSequenceHelperImpl::FLayerTimeInfo* SubLayerInfo = LayerTimeInfosByIdentifier.Find(*InfoIdentifier);
+	if (!SubLayerInfo)
+	{
+		return;
+	}
+
+	int32 SubLayerIndex = Info->SubLayerIdentifiers.IndexOfByKey(SubLayerInfo->Identifier);
+
+	double OffsetForSub = Info->OffsetTimeCodes[SubLayerIndex];
+	double ScaleForSub  = Info->Scales[SubLayerIndex];
+	double SubCodesPerSecond = SubLayerInfo->TimeCodesPerSecond.GetValue();
+	double SubStartTimeCode  = SubLayerInfo->StartTimeCode.GetValue();
+	double SubEndTimeCode    = SubLayerInfo->EndTimeCode.GetValue();
+	double SubStartSeconds = (OffsetForSub + ScaleForSub * SubStartTimeCode) / SubCodesPerSecond;
+	double SubEndSeconds   = (OffsetForSub + ScaleForSub * SubEndTimeCode)   / SubCodesPerSecond;
+
+	FFrameRate FrameRate = MovieScene->GetTickResolution();
+	FFrameNumber OriginalStartFrame = FrameRate.AsFrameNumber(SubStartSeconds);
+	FFrameNumber OriginalEndFrame   = FrameRate.AsFrameNumber(SubEndSeconds);
+	FFrameNumber ModifiedStartFrame = Section->GetInclusiveStartFrame();
+	FFrameNumber ModifiedEndFrame   = Section->GetExclusiveEndFrame();
+
+	// We haven't modified this subsequence, skip it
+	if (ModifiedStartFrame == OriginalStartFrame && ModifiedEndFrame == OriginalEndFrame)
+	{
+		return;
+	}
+
+	// This will obviously be quantized to frame intervals for now
+	double ModifiedStartTimeCode = FrameRate.AsSeconds(ModifiedStartFrame) * SubCodesPerSecond;
+	double ModifiedEndTimeCode   = FrameRate.AsSeconds(ModifiedEndFrame) * SubCodesPerSecond;
+	double ModifiedScale  = (ModifiedEndTimeCode - ModifiedStartTimeCode) / (SubEndTimeCode - SubStartTimeCode);
+	double ModifiedOffset = ModifiedStartTimeCode - SubStartTimeCode * ModifiedScale;
+
+	// Prevent twins from being rebuilt when we update the layer offsets
+	FScopedBlockNotices BlockNotices( StageActor.Get()->GetUsdListener() );
+
+	// Apply timecode offset and scale to the sublayer
+	Info->OffsetTimeCodes[SubLayerIndex] = ModifiedOffset;
+	Info->Scales[SubLayerIndex] = ModifiedScale;
+	pxr::SdfLayerOffset NewLayerOffset{ModifiedOffset, ModifiedScale};
+	Layer.Get()->SetSubLayerOffset(NewLayerOffset, SubLayerIndex);
+
+	UE_LOG(LogUsdStage, Verbose, TEXT("\tNew OffsetScale: %f, %f"), NewLayerOffset.GetOffset(), NewLayerOffset.GetScale());
+}
+
+void FUsdLevelSequenceHelperImpl::OnObjectTransacted(UObject* Object, const class FTransactionObjectEvent& Event)
+{
+	UMovieSceneSubSection* Section = Cast<UMovieSceneSubSection>(Object);
+	if (!Section)
+	{
+		return;
+	}
+
+	UMovieSceneSequence* ParentSequence = Section->GetTypedOuter<UMovieSceneSequence>();
+	if (!ParentSequence)
+	{
+		return;
+	}
+
+	const FString* InfoIdentifier = LayerInfoIdentifierByLevelSequenceName.Find(ParentSequence->GetFName());
+	if (!InfoIdentifier)
+	{
+		return;
+	}
+
+	FUsdLevelSequenceHelperImpl::FLayerTimeInfo* Info = LayerTimeInfosByIdentifier.Find(*InfoIdentifier);
+	UpdateInfoFromSection(Info, Section, ParentSequence);
+}
+#else
+class FUsdLevelSequenceHelperImpl
+{
+public:
+	FUsdLevelSequenceHelperImpl(TWeakObjectPtr<AUsdStageActor> InStageActor) {}
+	~FUsdLevelSequenceHelperImpl(){}
+};
+#endif // USE_USD_SDK
+
+
+FUsdLevelSequenceHelper::FUsdLevelSequenceHelper(TWeakObjectPtr<AUsdStageActor> InStageActor)
+{
+	if (AUsdStageActor* ValidStageActor = InStageActor.Get())
+	{
+		UsdSequencerImpl = MakeUnique<FUsdLevelSequenceHelperImpl>(InStageActor);
+	}
+}
+
+// These are required in order to be a member of AUsdStageActor,
+// and have a TUniquePtr to a forward declared Impl
+FUsdLevelSequenceHelper::FUsdLevelSequenceHelper()
+{
+}
+FUsdLevelSequenceHelper::~FUsdLevelSequenceHelper()
+{
+}
+
+#if USE_USD_SDK
+void FUsdLevelSequenceHelper::InitLevelSequence(const pxr::UsdStageRefPtr& UsdStage)
+{
+	if (UsdSequencerImpl.IsValid() && UsdStage)
+	{
+		UE_LOG(LogUsdStage, Verbose, TEXT("InitLevelSequence"));
+
+		UsdSequencerImpl->RebuildLayerTimeInfoCache(UsdStage);
+
+		UsdSequencerImpl->InitializeActorTimeCodes(UsdStage);
+
+		FUsdLevelSequenceHelperImpl::FLayerTimeInfo* RootInfo = UsdSequencerImpl->GetRootLayerInfo();
+		UsdSequencerImpl->CreateLevelSequences();
+		UsdSequencerImpl->CreateTimeTrack(RootInfo);
+		UsdSequencerImpl->CreateSubsequenceTrack(RootInfo);
+	}
+}
+
+void FUsdLevelSequenceHelper::UpdateLevelSequence(const pxr::UsdStageRefPtr& UsdStage)
+{
+	if (UsdSequencerImpl.IsValid() && UsdStage)
+	{
+		UE_LOG(LogUsdStage, Verbose, TEXT("UpdateLevelSequence"));
+
+		FUsdLevelSequenceHelperImpl::FLayerTimeInfo* RootInfo = UsdSequencerImpl->GetRootLayerInfo();
+		RootInfo->StartTimeCode = UsdStage->GetStartTimeCode();
+		RootInfo->EndTimeCode = UsdStage->GetEndTimeCode();
+
+		UsdSequencerImpl->CreateTimeTrack(RootInfo);
+	}
+}
+
+#undef DEFAULT_FRAMERATE
+#undef TIME_TRACK_NAME
+#endif // USE_USD_SDK
\ No newline at end of file
diff --git a/Engine/Plugins/Importers/USDImporter/Source/USDStage/Private/USDListener.cpp b/Engine/Plugins/Importers/USDImporter/Source/USDStage/Private/USDListener.cpp
--- a/Engine/Plugins/Importers/USDImporter/Source/USDStage/Private/USDListener.cpp
+++ b/Engine/Plugins/Importers/USDImporter/Source/USDStage/Private/USDListener.cpp
@@ -29,6 +29,13 @@ void FUsdListener::Register( const pxr::UsdStageRefPtr& Stage )
 	}
 
 	RegisteredStageEditTargetChangedKey = pxr::TfNotice::Register( pxr::TfWeakPtr< FUsdListener >( this ), &FUsdListener::HandleStageEditTargetChangedNotice, Stage );
+
+	if (RegisteredLayersChangedKey.IsValid())
+	{
+		pxr::TfNotice::Revoke(RegisteredLayersChangedKey);
+	}
+
+	RegisteredLayersChangedKey = pxr::TfNotice::Register(pxr::TfWeakPtr< FUsdListener >( this ), &FUsdListener::HandleLayersChangedNotice );
 }
 
 FUsdListener::~FUsdListener()
@@ -36,6 +43,7 @@ FUsdListener::~FUsdListener()
 	FScopedUsdAllocs UsdAllocs;
 	pxr::TfNotice::Revoke( RegisteredObjectsChangedKey );
 	pxr::TfNotice::Revoke( RegisteredStageEditTargetChangedKey );
+	pxr::TfNotice::Revoke( RegisteredLayersChangedKey );
 }
 
 void FUsdListener::HandleUsdNotice( const pxr::UsdNotice::ObjectsChanged& Notice, const pxr::UsdStageWeakPtr& Sender )
@@ -81,4 +89,14 @@ void FUsdListener::HandleStageEditTargetChangedNotice( const pxr::UsdNotice::Sta
 	OnStageEditTargetChanged.Broadcast();
 }
 
+void FUsdListener::HandleLayersChangedNotice(const pxr::SdfNotice::LayersDidChange& Notice)
+{
+	if ( !OnLayersChanged.IsBound() || IsBlocked.GetValue() > 0 )
+	{
+		return;
+	}
+
+	OnLayersChanged.Broadcast(Notice.GetChangeListMap());
+}
+
 #endif // #if USE_USD_SDK
diff --git a/Engine/Plugins/Importers/USDImporter/Source/USDStage/Private/USDStageActor.cpp b/Engine/Plugins/Importers/USDImporter/Source/USDStage/Private/USDStageActor.cpp
--- a/Engine/Plugins/Importers/USDImporter/Source/USDStage/Private/USDStageActor.cpp
+++ b/Engine/Plugins/Importers/USDImporter/Source/USDStage/Private/USDStageActor.cpp
@@ -22,10 +22,11 @@
 #include "Engine/StaticMesh.h"
 #include "Engine/World.h"
 #include "HAL/FileManager.h"
+#include "LevelSequence.h"
 #include "MaterialEditingLibrary.h"
 #include "Materials/Material.h"
-#include "Materials/MaterialInterface.h"
 #include "Materials/MaterialInstanceConstant.h"
+#include "Materials/MaterialInterface.h"
 #include "MeshDescription.h"
 #include "MeshDescriptionOperations.h"
 #include "Misc/Paths.h"
@@ -33,15 +34,9 @@
 #include "PhysicsEngine/BodySetup.h"
 #include "Rendering/SkeletalMeshLODImporterData.h"
 #include "StaticMeshAttributes.h"
+#include "Subsystems/AssetEditorSubsystem.h"
 
-#include "Channels/MovieSceneChannelProxy.h"
-#include "LevelSequence.h"
-#include "LevelSequenceActor.h"
-#include "Misc/FrameNumber.h"
 #include "Modules/ModuleManager.h"
-#include "MovieScene.h"
-#include "Sections/MovieSceneFloatSection.h"
-#include "Tracks/MovieSceneFloatTrack.h"
 
 #if USE_USD_SDK
 #include "USDIncludesStart.h"
@@ -71,7 +66,7 @@ AUsdStageActor::FOnActorLoaded AUsdStageActor::OnActorLoaded{};
 
 #if USE_USD_SDK
 
-struct UsdStageActorImpl
+struct FUsdStageActorImpl
 {
 	static TSharedRef< FUsdSchemaTranslationContext > CreateUsdSchemaTranslationContext( AUsdStageActor* StageActor, const FString& PrimPath )
 	{
@@ -107,6 +102,7 @@ AUsdStageActor::AUsdStageActor()
 	, StartTimeCode( 0.f )
 	, EndTimeCode( 100.f )
 	, TimeCodesPerSecond( 24.f )
+	, LevelSequenceHelper(this)
 {
 	SceneComponent = CreateDefaultSubobject< USceneComponent >( TEXT("SceneComponent0") );
 	SceneComponent->Mobility = EComponentMobility::Static;
@@ -121,7 +117,7 @@ AUsdStageActor::AUsdStageActor()
 			{
 				TUsdStore< pxr::SdfPath > UsdPrimPath = UnrealToUsd::ConvertPath( *PrimPath );
 
-				TSharedRef< FUsdSchemaTranslationContext > TranslationContext = UsdStageActorImpl::CreateUsdSchemaTranslationContext( this, PrimPath );
+				TSharedRef< FUsdSchemaTranslationContext > TranslationContext = FUsdStageActorImpl::CreateUsdSchemaTranslationContext( this, PrimPath );
 
 				this->LoadAssets( *TranslationContext, this->GetUsdStage()->GetPrimAtPath( UsdPrimPath.Get() ) );
 
@@ -134,10 +130,29 @@ AUsdStageActor::AUsdStageActor()
 				}
 			}
 		);
-#endif // #if USE_USD_SDK
 
-	InitLevelSequence(TimeCodesPerSecond);
-	SetupLevelSequence();
+	UsdListener.OnLayersChanged.AddLambda(
+		[&](const pxr::SdfLayerChangeListMap& ChangeMap)
+		{
+			// Check to see if any layer reloaded. If so, rebuild all of our animations as a single layer changing
+			// might propagate timecodes through all level sequences
+			for (const std::pair<pxr::SdfLayerHandle, pxr::SdfChangeList>& ChangeMapItem : ChangeMap)
+			{
+				const pxr::SdfChangeList::EntryList& ChangeList = ChangeMapItem.second.GetEntryList();
+				for (const std::pair<pxr::SdfPath, pxr::SdfChangeList::Entry>& Change : ChangeList)
+				{
+					const pxr::SdfChangeList::Entry::_Flags& Flags = Change.second.flags;
+					if (Flags.didReloadContent)
+					{
+						UE_LOG(LogUsdStage, Verbose, TEXT("Reloading animations because layer '%s' reloaded"), *UsdToUnreal::ConvertPath(Change.first));
+						ReloadAnimations();
+						return;
+					}
+				}
+			}
+		}
+	);
+#endif // #if USE_USD_SDK
 
 	if ( HasAutorithyOverStage() )
 	{
@@ -310,7 +325,7 @@ void AUsdStageActor::PostEditChangeProperty(FPropertyChangedEvent& PropertyChang
 {
 	FProperty* PropertyThatChanged = PropertyChangedEvent.MemberProperty;
 	const FName PropertyName = PropertyThatChanged ? PropertyThatChanged->GetFName() : NAME_None;
-	
+
 #if USE_USD_SDK
 	if ( PropertyName == GET_MEMBER_NAME_CHECKED( AUsdStageActor, RootLayer ) )
 	{
@@ -325,12 +340,12 @@ void AUsdStageActor::PostEditChangeProperty(FPropertyChangedEvent& PropertyChang
 	}
 	else if ( PropertyName == GET_MEMBER_NAME_CHECKED( AUsdStageActor, StartTimeCode ) || PropertyName == GET_MEMBER_NAME_CHECKED( AUsdStageActor, EndTimeCode ) )
 	{
-		if ( UsdStageStore.Get() )
+		if ( const pxr::UsdStageRefPtr& UsdStage = GetUsdStage() )
 		{
-			UsdStageStore.Get()->SetStartTimeCode( StartTimeCode );
-			UsdStageStore.Get()->SetEndTimeCode( EndTimeCode );
+			UsdStage->SetStartTimeCode( StartTimeCode );
+			UsdStage->SetEndTimeCode( EndTimeCode );
 
-			SetupLevelSequence();
+			LevelSequenceHelper.UpdateLevelSequence( UsdStage );
 		}
 	}
 	else
@@ -379,87 +394,6 @@ void AUsdStageActor::OpenUsdStage()
 #endif // #if USE_USD_SDK
 }
 
-void AUsdStageActor::InitLevelSequence(float FramesPerSecond)
-{
-	if ( LevelSequence || !HasAutorithyOverStage() )
-	{
-		return;
-	}
-
-	// Create LevelSequence
-	LevelSequence = NewObject<ULevelSequence>(GetTransientPackage(), NAME_None, DefaultObjFlag | EObjectFlags::RF_Public);
-	LevelSequence->Initialize();
-
-	UMovieScene* MovieScene = LevelSequence->MovieScene;
-
-	FFrameRate FrameRate(FramesPerSecond, 1);
-	MovieScene->SetDisplayRate(FrameRate);
-
-	// Populate the Sequence
-	FGuid ObjectBinding = MovieScene->AddPossessable(GetActorLabel(), this->GetClass());
-	LevelSequence->BindPossessableObject(ObjectBinding, *this, GetWorld());
-
-	// Add the Time track
-	UMovieSceneFloatTrack* TimeTrack = MovieScene->AddTrack<UMovieSceneFloatTrack>(ObjectBinding);
-	TimeTrack->SetPropertyNameAndPath(FName(TEXT("Time")), "Time");
-
-	FFrameRate DestFrameRate = LevelSequence->MovieScene->GetTickResolution();
-	FFrameNumber StartFrame = DestFrameRate.AsFrameNumber(StartTimeCode / TimeCodesPerSecond);
-
-	bool bSectionAdded = false;
-	UMovieSceneFloatSection* TimeSection = Cast<UMovieSceneFloatSection>(TimeTrack->FindOrAddSection(StartFrame, bSectionAdded));
-
-	TimeSection->EvalOptions.CompletionMode = EMovieSceneCompletionMode::KeepState;
-	MovieScene->SetEvaluationType(EMovieSceneEvaluationType::FrameLocked);
-}
-
-void AUsdStageActor::SetupLevelSequence()
-{
-	if (!LevelSequence)
-	{
-		return;
-	}
-
-	UMovieScene* MovieScene = LevelSequence->MovieScene;
-	if (!MovieScene)
-	{
-		return;
-	}
-
-	FFrameRate DestFrameRate = MovieScene->GetTickResolution();
-
-	TArray<UMovieSceneSection*> Sections = MovieScene->GetAllSections();
-	if (Sections.Num() > 0)
-	{
-		UMovieSceneSection* TimeSection = Sections[0];
-		FMovieSceneFloatChannel* TimeChannel = TimeSection->GetChannelProxy().GetChannel<FMovieSceneFloatChannel>(0);
-
-		FFrameNumber StartFrame = DestFrameRate.AsFrameNumber(StartTimeCode / TimeCodesPerSecond);
-		FFrameNumber EndFrame = DestFrameRate.AsFrameNumber(EndTimeCode / TimeCodesPerSecond);
-
-		TArray<FFrameNumber> FrameNumbers;
-		FrameNumbers.Add(StartFrame);
-		FrameNumbers.Add(EndFrame);
-
-		TArray<FMovieSceneFloatValue> FrameValues;
-		FrameValues.Add_GetRef(FMovieSceneFloatValue(StartTimeCode)).InterpMode = ERichCurveInterpMode::RCIM_Linear;
-		FrameValues.Add_GetRef(FMovieSceneFloatValue(EndTimeCode)).InterpMode = ERichCurveInterpMode::RCIM_Linear;
-
-		TimeChannel->Set(FrameNumbers, FrameValues);
-
-		TimeSection->SetRange(TRange<FFrameNumber>::All());
-
-		float StartTime = StartTimeCode / TimeCodesPerSecond;
-		float EndTime = EndTimeCode / TimeCodesPerSecond;
-
-		TRange<FFrameNumber> TimeRange = TRange<FFrameNumber>::Inclusive(StartFrame, EndFrame);
-
-		LevelSequence->MovieScene->SetPlaybackRange(TimeRange);
-		LevelSequence->MovieScene->SetViewRange((StartTimeCode - 5.f) / TimeCodesPerSecond, (EndTimeCode + 5.f) / TimeCodesPerSecond);
-		LevelSequence->MovieScene->SetWorkingRange((StartTimeCode - 5.f) / TimeCodesPerSecond, (EndTimeCode + 5.f) / TimeCodesPerSecond);
-	}
-}
-
 void AUsdStageActor::LoadUsdStage()
 {
 #if USE_USD_SDK
@@ -482,13 +416,9 @@ void AUsdStageActor::LoadUsdStage()
 		return;
 	}
 
-	StartTimeCode = UsdStage->GetStartTimeCode();
-	EndTimeCode = UsdStage->GetEndTimeCode();
-	TimeCodesPerSecond = UsdStage->GetTimeCodesPerSecond();
-
-	SetupLevelSequence();
+	ReloadAnimations();
 
-	TSharedRef< FUsdSchemaTranslationContext > TranslationContext = UsdStageActorImpl::CreateUsdSchemaTranslationContext( this, RootUsdTwin.PrimPath );
+	TSharedRef< FUsdSchemaTranslationContext > TranslationContext = FUsdStageActorImpl::CreateUsdSchemaTranslationContext( this, RootUsdTwin.PrimPath );
 
 	LoadAssets( *TranslationContext, UsdStage->GetPseudoRoot() );
 
@@ -521,6 +451,25 @@ void AUsdStageActor::Refresh() const
 	}
 }
 
+void AUsdStageActor::ReloadAnimations()
+{
+#if USE_USD_SDK
+	const pxr::UsdStageRefPtr& UsdStage = GetUsdStage();
+	if ( !UsdStage )
+	{
+		return;
+	}
+
+	if ( HasAutorithyOverStage() )
+	{
+		LevelSequenceHelper.InitLevelSequence(UsdStage);
+
+		// The sequencer won't update on its own, so let's at least force it closed
+		GEditor->GetEditorSubsystem<UAssetEditorSubsystem>()->CloseAllEditorsForAsset(LevelSequence);
+	}
+#endif // #if USE_USD_SDK
+}
+
 void AUsdStageActor::PostRegisterAllComponents()
 {
 	Super::PostRegisterAllComponents();
@@ -536,7 +485,7 @@ void AUsdStageActor::PostRegisterAllComponents()
 
 		if ( UsdStage )
 		{
-			TSharedRef< FUsdSchemaTranslationContext > TranslationContext = UsdStageActorImpl::CreateUsdSchemaTranslationContext( this, RootUsdTwin.PrimPath );
+			TSharedRef< FUsdSchemaTranslationContext > TranslationContext = FUsdStageActorImpl::CreateUsdSchemaTranslationContext( this, RootUsdTwin.PrimPath );
 
 			LoadAssets( *TranslationContext, UsdStage->GetPseudoRoot() );
 			UpdatePrim( UsdStage->GetPseudoRoot().GetPrimPath(), true, *TranslationContext );
@@ -600,7 +549,7 @@ void AUsdStageActor::OnPrimObjectPropertyChanged( UObject* ObjectBeingModified,
 			return;
 		}
 	}
-	
+
 	FString PrimPath = ObjectsToWatch[ PrimObject ];
 
 	if ( FUsdPrimTwin* UsdPrimTwin = RootUsdTwin.Find( PrimPath ) )
@@ -662,19 +611,19 @@ void AUsdStageActor::LoadAssets( FUsdSchemaTranslationContext& TranslationContex
 
 	// Load materials first since meshes are referencing them
 	TArray< TUsdStore< pxr::UsdPrim > > AllPrimAssets = UsdUtils::GetAllPrimsOfType( StartPrim, pxr::TfType::Find< pxr::UsdShadeMaterial >() );
-	
+
 	CreateAssetsForPrims( AllPrimAssets );
 
 	// Load meshes
 	AllPrimAssets = UsdUtils::GetAllPrimsOfType( StartPrim, pxr::TfType::Find< pxr::UsdGeomMesh >(), { pxr::TfType::Find< pxr::UsdSkelRoot >() } );
 	AllPrimAssets.Append( UsdUtils::GetAllPrimsOfType( StartPrim, pxr::TfType::Find< pxr::UsdSkelRoot >() ) );
-	
+
 	CreateAssetsForPrims( AllPrimAssets );
 }
 
 void AUsdStageActor::AnimatePrims()
 {
-	TSharedRef< FUsdSchemaTranslationContext > TranslationContext = UsdStageActorImpl::CreateUsdSchemaTranslationContext( this, RootUsdTwin.PrimPath );
+	TSharedRef< FUsdSchemaTranslationContext > TranslationContext = FUsdStageActorImpl::CreateUsdSchemaTranslationContext( this, RootUsdTwin.PrimPath );
 
 	for ( const FString& PrimToAnimate : PrimsToAnimate )
 	{
@@ -691,7 +640,7 @@ void AUsdStageActor::AnimatePrims()
 	}
 
 	TranslationContext->CompleteTasks();
-		
+
 	GEditor->BroadcastLevelActorListChanged();
 	GEditor->RedrawLevelEditingViewports();
 }
diff --git a/Engine/Plugins/Importers/USDImporter/Source/USDStage/Public/USDLevelSequenceHelper.h b/Engine/Plugins/Importers/USDImporter/Source/USDStage/Public/USDLevelSequenceHelper.h
new file mode 100644
--- /dev/null
+++ b/Engine/Plugins/Importers/USDImporter/Source/USDStage/Public/USDLevelSequenceHelper.h
@@ -0,0 +1,34 @@
+// Copyright Epic Games, Inc. All Rights Reserved.
+
+#pragma once
+
+#include "CoreMinimal.h"
+
+#if USE_USD_SDK
+#include "USDIncludesStart.h"
+
+#include "pxr/pxr.h"
+#include "pxr/usd/sdf/changeList.h"
+#include "pxr/usd/usd/stage.h"
+
+#include "USDIncludesEnd.h"
+#endif // USE_USD_SDK
+
+class AUsdStageActor;
+class FUsdLevelSequenceHelperImpl;
+
+class FUsdLevelSequenceHelper
+{
+public:
+	FUsdLevelSequenceHelper();
+	explicit FUsdLevelSequenceHelper(TWeakObjectPtr<AUsdStageActor> InStageActor);
+	virtual ~FUsdLevelSequenceHelper();
+
+#if USE_USD_SDK
+	void InitLevelSequence(const pxr::UsdStageRefPtr& UsdStage);
+	void UpdateLevelSequence(const pxr::UsdStageRefPtr& UsdStage);
+#endif // USE_USD_SDK
+
+private:
+	TUniquePtr<FUsdLevelSequenceHelperImpl> UsdSequencerImpl;
+};
\ No newline at end of file
diff --git a/Engine/Plugins/Importers/USDImporter/Source/USDStage/Public/USDListener.h b/Engine/Plugins/Importers/USDImporter/Source/USDStage/Public/USDListener.h
--- a/Engine/Plugins/Importers/USDImporter/Source/USDStage/Public/USDListener.h
+++ b/Engine/Plugins/Importers/USDImporter/Source/USDStage/Public/USDListener.h
@@ -8,6 +8,7 @@
 
 #include "pxr/base/tf/weakBase.h"
 #include "pxr/usd/sdf/path.h"
+#include "pxr/usd/sdf/notice.h"
 #include "pxr/usd/usd/notice.h"
 #include "pxr/usd/usd/stage.h"
 
@@ -34,15 +35,20 @@ public:
 	DECLARE_EVENT_TwoParams( FUsdListener, FOnPrimChanged, const FString&, bool );
 	FOnPrimChanged OnPrimChanged;
 
+	DECLARE_EVENT_OneParam( FUsdListener, FOnLayersChanged, const pxr::SdfLayerChangeListMap& );
+	FOnLayersChanged OnLayersChanged;
+
 	FThreadSafeCounter IsBlocked;
 
 protected:
 	void HandleUsdNotice( const pxr::UsdNotice::ObjectsChanged& Notice, const pxr::UsdStageWeakPtr& Sender );
 	void HandleStageEditTargetChangedNotice( const pxr::UsdNotice::StageEditTargetChanged& Notice, const pxr::UsdStageWeakPtr& Sender );
+	void HandleLayersChangedNotice ( const pxr::SdfNotice::LayersDidChange& Notice );
 
 private:
 	pxr::TfNotice::Key RegisteredObjectsChangedKey;
 	pxr::TfNotice::Key RegisteredStageEditTargetChangedKey;
+	pxr::TfNotice::Key RegisteredLayersChangedKey;
 };
 
 class FScopedBlockNotices final
diff --git a/Engine/Plugins/Importers/USDImporter/Source/USDStage/Public/USDStageActor.h b/Engine/Plugins/Importers/USDImporter/Source/USDStage/Public/USDStageActor.h
--- a/Engine/Plugins/Importers/USDImporter/Source/USDStage/Public/USDStageActor.h
+++ b/Engine/Plugins/Importers/USDImporter/Source/USDStage/Public/USDStageActor.h
@@ -6,6 +6,7 @@
 #include "Misc/SecureHash.h"
 #include "Templates/Function.h"
 
+#include "USDLevelSequenceHelper.h"
 #include "USDListener.h"
 #include "USDMemory.h"
 #include "USDPrimTwin.h"
@@ -47,7 +48,8 @@ class AUsdStageActor : public AActor
 {
 	GENERATED_BODY()
 
-	friend struct UsdStageActorImpl;
+	friend struct FUsdStageActorImpl;
+	friend class FUsdLevelSequenceHelperImpl;
 
 public:
 	UPROPERTY(EditAnywhere, Category = "USD", meta = (FilePathFilter = "usd files (*.usd; *.usda; *.usdc)|*.usd; *.usda; *.usdc"))
@@ -82,6 +84,9 @@ private:
 	UPROPERTY(VisibleAnywhere, Category = "USD", Transient)
 	ULevelSequence* LevelSequence;
 
+	UPROPERTY(Transient)
+	TMap<FString, ULevelSequence*> SubLayerLevelSequencesByIdentifier;
+
 public:
 	DECLARE_EVENT_OneParam( AUsdStageActor, FOnActorLoaded, AUsdStageActor* );
 	USDSTAGE_API static FOnActorLoaded OnActorLoaded;
@@ -100,6 +105,7 @@ public:
 	virtual ~AUsdStageActor();
 
 	void Refresh() const;
+	void ReloadAnimations();
 
 public:
 	virtual void PostEditChangeProperty( FPropertyChangedEvent& PropertyChangedEvent ) override;
@@ -111,9 +117,6 @@ private:
 	void OpenUsdStage();
 	void LoadUsdStage();
 
-	void InitLevelSequence( float FramesPerSecond );
-	void SetupLevelSequence();
-
 	void OnUsdPrimTwinDestroyed( const FUsdPrimTwin& UsdPrimTwin );
 
 	void OnPrimObjectPropertyChanged( UObject* ObjectBeingModified, FPropertyChangedEvent& PropertyChangedEvent );
@@ -155,7 +158,7 @@ protected:
 private:
 	TUsdStore< pxr::UsdStageRefPtr > UsdStageStore;
 	FUsdListener UsdListener;
-
 #endif // #if USE_USD_SDK
 
+	FUsdLevelSequenceHelper LevelSequenceHelper;
 };
