USD: Support Purposes

diff --git a/Engine/Plugins/Importers/USDImporter/Source/USDImporter/Private/USDImportOptions.cpp b/Engine/Plugins/Importers/USDImporter/Source/USDImporter/Private/USDImportOptions.cpp
--- a/Engine/Plugins/Importers/USDImporter/Source/USDImporter/Private/USDImportOptions.cpp
+++ b/Engine/Plugins/Importers/USDImporter/Source/USDImporter/Private/USDImportOptions.cpp
@@ -24,6 +24,7 @@ void UUSDImportOptions::PostEditChangeProperty(struct FPropertyChangedEvent& Pro
 UUSDSceneImportOptions::UUSDSceneImportOptions(const FObjectInitializer& ObjectInitializer)
 	: Super(ObjectInitializer)
 {
+	PurposesToImport = (int32) (EUsdPurpose::Default | EUsdPurpose::Proxy);
 	bFlattenHierarchy = true;
 	bImportMeshes = true;
 	PathForAssets.Path = TEXT("/Game");
@@ -44,7 +45,7 @@ bool UUSDSceneImportOptions::CanEditChange(const FProperty* InProperty) const
 	FName PropertyName = InProperty ? InProperty->GetFName() : NAME_None;
 
 	if (GET_MEMBER_NAME_CHECKED(UUSDImportOptions, MeshImportType) == PropertyName ||
-		GET_MEMBER_NAME_CHECKED(UUSDImportOptions, bApplyWorldTransformToGeometry) == PropertyName || 
+		GET_MEMBER_NAME_CHECKED(UUSDImportOptions, bApplyWorldTransformToGeometry) == PropertyName ||
 		GET_MEMBER_NAME_CHECKED(UUSDImportOptions, bGenerateUniquePathPerUSDPrim) == PropertyName)
 	{
 		bCanEdit &= bImportMeshes;
@@ -75,7 +76,7 @@ bool UUSDBatchImportOptions::CanEditChange(const FProperty* InProperty) const
 	FName PropertyName = InProperty ? InProperty->GetFName() : NAME_None;
 
 	if (GET_MEMBER_NAME_CHECKED(UUSDImportOptions, MeshImportType) == PropertyName ||
-		GET_MEMBER_NAME_CHECKED(UUSDImportOptions, bApplyWorldTransformToGeometry) == PropertyName || 
+		GET_MEMBER_NAME_CHECKED(UUSDImportOptions, bApplyWorldTransformToGeometry) == PropertyName ||
 		GET_MEMBER_NAME_CHECKED(UUSDImportOptions, bGenerateUniquePathPerUSDPrim) == PropertyName)
 	{
 		bCanEdit &= bImportMeshes;
diff --git a/Engine/Plugins/Importers/USDImporter/Source/USDImporter/Private/USDImportOptions.h b/Engine/Plugins/Importers/USDImporter/Source/USDImporter/Private/USDImportOptions.h
deleted file mode 100644
diff --git a/Engine/Plugins/Importers/USDImporter/Source/USDImporter/Private/USDPrimResolver.cpp b/Engine/Plugins/Importers/USDImporter/Source/USDImporter/Private/USDPrimResolver.cpp
--- a/Engine/Plugins/Importers/USDImporter/Source/USDImporter/Private/USDPrimResolver.cpp
+++ b/Engine/Plugins/Importers/USDImporter/Source/USDImporter/Private/USDPrimResolver.cpp
@@ -6,6 +6,7 @@
 #include "USDImportOptions.h"
 #include "USDImporter.h"
 #include "USDTypesConversion.h"
+#include "UnrealUSDWrapper.h"
 
 #include "ActorFactories/ActorFactory.h"
 #include "AssetData.h"
@@ -42,41 +43,45 @@ void UUSDPrimResolver::FindMeshAssetsToImport(FUsdImportContext& ImportContext,
 	bool bHasUnrealAssetPath = IUsdPrim::GetUnrealAssetPath( StartPrim.Get() ).size() > 0;
 	bool bHasUnrealActorClass = IUsdPrim::GetUnrealActorClass( StartPrim.Get() ).size() > 0;
 
-	if ( !IUsdPrim::IsProxyOrGuide(StartPrim.Get()) )
+	EUsdPurpose EnabledPurposes = EUsdPurpose::Default | EUsdPurpose::Render;
+	if (UUSDSceneImportOptions* SceneImportOptions = Cast<UUSDSceneImportOptions>(ImportContext.ImportOptions))
 	{
-		if (IUsdPrim::HasGeometryDataOrLODVariants(StartPrim.Get()))
-		{
-			FUsdAssetPrimToImport NewTopLevelPrim;
+		EnabledPurposes = (EUsdPurpose)SceneImportOptions->PurposesToImport;
+	}
 
-			FString FinalPrimName;
-			// if the prim has a path use that as the final name
-			if (bHasUnrealAssetPath)
-			{
-				FinalPrimName = UsdToUnreal::ConvertString( IUsdPrim::GetUnrealAssetPath( StartPrim.Get() ).c_str() );
-			}
-			else
-			{
-				FinalPrimName = PrimName;
-			}
+	if (EnumHasAnyFlags(EnabledPurposes, IUsdPrim::GetPurpose(StartPrim.Get())) &&
+		IUsdPrim::HasGeometryDataOrLODVariants(StartPrim.Get()))
+	{
+		FUsdAssetPrimToImport NewTopLevelPrim;
 
-			NewTopLevelPrim.Prim = StartPrim;
-			NewTopLevelPrim.AssetPath = FinalPrimName;
+		FString FinalPrimName;
+		// if the prim has a path use that as the final name
+		if (bHasUnrealAssetPath)
+		{
+			FinalPrimName = UsdToUnreal::ConvertString( IUsdPrim::GetUnrealAssetPath( StartPrim.Get() ).c_str() );
+		}
+		else
+		{
+			FinalPrimName = PrimName;
+		}
 
-			FindMeshChildren(ImportContext, StartPrim, true, NewTopLevelPrim.MeshPrims);
+		NewTopLevelPrim.Prim = StartPrim;
+		NewTopLevelPrim.AssetPath = FinalPrimName;
 
-			for ( const TUsdStore< pxr::UsdPrim >& MeshPrim : NewTopLevelPrim.MeshPrims)
-			{
-				NewTopLevelPrim.NumLODs = FMath::Max(NewTopLevelPrim.NumLODs, IUsdPrim::GetNumLODs( *MeshPrim ));
-			}
+		FindMeshChildren(ImportContext, StartPrim, true, NewTopLevelPrim.MeshPrims);
 
-			OutAssetsToImport.Add(NewTopLevelPrim);
+		for ( const TUsdStore< pxr::UsdPrim >& MeshPrim : NewTopLevelPrim.MeshPrims)
+		{
+			NewTopLevelPrim.NumLODs = FMath::Max(NewTopLevelPrim.NumLODs, IUsdPrim::GetNumLODs( *MeshPrim ));
 		}
-		else if(bRecursive)
+
+		OutAssetsToImport.Add(NewTopLevelPrim);
+	}
+	else if(bRecursive)
+	{
+		for ( pxr::UsdPrim Child : StartPrim.Get().GetFilteredChildren( pxr::UsdTraverseInstanceProxies() ) )
 		{
-			for ( pxr::UsdPrim Child : StartPrim.Get().GetFilteredChildren( pxr::UsdTraverseInstanceProxies() ) )
-			{
-				FindMeshAssetsToImport(ImportContext, Child, StartPrim, OutAssetsToImport);
-			}
+			FindMeshAssetsToImport(ImportContext, Child, StartPrim, OutAssetsToImport);
 		}
 	}
 }
@@ -217,6 +222,7 @@ AActor* UUSDPrimResolver::SpawnActor(FUSDSceneImportContext& ImportContext, cons
 				{
 					SpawnedActor = ImportContext.World->SpawnActor( AActor::StaticClass() );
 					USceneComponent* SceneComponent = NewObject< USceneComponent >( SpawnedActor, SpawnData.ActorName );
+					SceneComponent->ComponentTags.Add(IUsdPrim::GetPurposeName(IUsdPrim::GetPurpose(SpawnData.ActorPrim.Get())));
 					SpawnedActor->AddInstanceComponent( SceneComponent );
 					SpawnedActor->SetRootComponent( SceneComponent );
 				}
@@ -256,6 +262,7 @@ AActor* UUSDPrimResolver::SpawnActor(FUSDSceneImportContext& ImportContext, cons
 
 					UStaticMeshComponent* StaticMeshComponent = NewObject< UStaticMeshComponent >( SpawnedActor, ComponentName );
 					StaticMeshComponent->SetStaticMesh( ImportedStaticMesh );
+					StaticMeshComponent->ComponentTags.Add(IUsdPrim::GetPurposeName(IUsdPrim::GetPurpose(UsdAssetPrimToImport.Prim.Get())));
 
 					// Don't add the prim transform if its the same prim used for the actor as it's already accounted for in the ActorTransform
 					if ( UsdAssetPrimToImport.Prim.Get() != SpawnData.ActorPrim.Get() )
@@ -284,7 +291,7 @@ AActor* UUSDPrimResolver::SpawnActor(FUSDSceneImportContext& ImportContext, cons
 			{
 				SpawnedActor = ActorFactory->CreateActor(ActorAsset, ImportContext.World->GetCurrentLevel(), FTransform::Identity, RF_Transactional, SpawnData.ActorName);
 
-				// For empty group actors set their initial mobility to static 
+				// For empty group actors set their initial mobility to static
 				if ( ActorFactory == ImportContext.EmptyActorFactory )
 				{
 					SpawnedActor->GetRootComponent()->SetMobility(EComponentMobility::Static);
@@ -390,21 +397,28 @@ void UUSDPrimResolver::FindMeshChildren(FUsdImportContext& ImportContext, const
 
 	const bool bIncludeLODs = bOnlyLODRoots;
 
-	if(bOnlyLODRoots && IUsdPrim::GetNumLODs( ParentPrim.Get() ) > 0)
+	EUsdPurpose EnabledPurposes = EUsdPurpose::Default | EUsdPurpose::Render;
+	if (UUSDSceneImportOptions* SceneImportOptions = Cast<UUSDSceneImportOptions>(ImportContext.ImportOptions))
+	{
+		EnabledPurposes = (EUsdPurpose)SceneImportOptions->PurposesToImport;
+	}
+	bool bValidPurpose = EnumHasAnyFlags(EnabledPurposes, IUsdPrim::GetPurpose(ParentPrim.Get()));
+
+	if(bOnlyLODRoots && IUsdPrim::GetNumLODs( ParentPrim.Get() ) > 0 && bValidPurpose)
 	{
 		// We're only looking for lod roots and this prim has LODs so add the prim and dont recurse into children
 		OutMeshChildren.Add(ParentPrim);
 	}
 	else
 	{
-		if (IUsdPrim::HasGeometryData(ParentPrim.Get()))
+		if (IUsdPrim::HasGeometryData(ParentPrim.Get()) && bValidPurpose)
 		{
 			OutMeshChildren.Add(ParentPrim);
 		}
 
 		for ( pxr::UsdPrim Child : ParentPrim.Get().GetFilteredChildren( pxr::UsdTraverseInstanceProxies() ) )
 		{
-			if (!IUsdPrim::IsProxyOrGuide(Child) && !IUsdPrim::IsKindChildOf(Child, USDKindTypes::Component))
+			if (!IUsdPrim::IsKindChildOf(Child, USDKindTypes::Component))
 			{
 				FindMeshChildren(ImportContext, Child, bOnlyLODRoots, OutMeshChildren);
 			}
diff --git a/Engine/Plugins/Importers/USDImporter/Source/USDImporter/Public/USDImportOptions.h b/Engine/Plugins/Importers/USDImporter/Source/USDImporter/Public/USDImportOptions.h
new file mode 100644
--- /dev/null
+++ b/Engine/Plugins/Importers/USDImporter/Source/USDImporter/Public/USDImportOptions.h
@@ -0,0 +1,167 @@
+// Copyright Epic Games, Inc. All Rights Reserved.
+
+#pragma once
+
+#include "UnrealUSDWrapper.h"
+
+#include "CoreMinimal.h"
+#include "Engine/EngineTypes.h"
+#include "Factories/MaterialImportHelpers.h"
+#include "UObject/ObjectMacros.h"
+
+#include "USDImportOptions.generated.h"
+
+UENUM(BlueprintType)
+enum class EExistingActorPolicy : uint8
+{
+	/** Replaces existing actors with new ones */
+	Replace,
+	/** Update transforms on existing actors but do not replace actor the actor class or any other data */
+	UpdateTransform,
+	/** Ignore any existing actor with the same name */
+	Ignore,
+
+};
+
+UENUM(BlueprintType)
+enum class EExistingAssetPolicy : uint8
+{
+	/** Reimports existing assets */
+	Reimport,
+
+	/** Ignores existing assets and doesnt reimport them */
+	Ignore,
+};
+
+UENUM(BlueprintType)
+enum class EUsdMeshImportType : uint8
+{
+	StaticMesh,
+};
+
+UCLASS(config = EditorPerProjectUserSettings)
+class USDIMPORTER_API UUSDImportOptions : public UObject
+{
+	GENERATED_UCLASS_BODY()
+public:
+	/** Defines what should happen with existing actors */
+	UPROPERTY(BlueprintReadWrite, config, EditAnywhere, Category=Mesh)
+	EUsdMeshImportType MeshImportType;
+
+	/**
+	 * If checked, To enforce unique asset paths, all assets will be created in directories that match with their prim path
+	 * e.g a USD path /root/myassets/myprim_mesh will generate the path in the game directory "/Game/myassets/" with a mesh asset called "myprim_mesh" within that path.
+	 */
+	UPROPERTY(BlueprintReadWrite, config, EditAnywhere, Category=Mesh)
+	bool bGenerateUniquePathPerUSDPrim;
+
+	UPROPERTY(BlueprintReadWrite, config, EditAnywhere, Category=Mesh)
+	bool bApplyWorldTransformToGeometry;
+
+	UPROPERTY(BlueprintReadWrite, EditAnywhere, config, Category="Mesh|Materials")
+	EMaterialSearchLocation MaterialSearchLocation;
+
+	UPROPERTY(BlueprintReadWrite, config, EditAnywhere, Category = Mesh)
+	float Scale;
+
+public:
+	virtual void PostEditChangeProperty(struct FPropertyChangedEvent& PropertyChangedEvent) override;
+};
+
+UCLASS(config = EditorPerProjectUserSettings)
+class USDIMPORTER_API UUSDSceneImportOptions : public UUSDImportOptions
+{
+	GENERATED_UCLASS_BODY()
+public:
+#if WITH_EDITOR
+	virtual void PostEditChangeProperty(struct FPropertyChangedEvent& PropertyChangedEvent) override;
+	virtual bool CanEditChange(const FProperty* InProperty) const override;
+#endif
+public:
+	/* Only import prims with these specific purposes from the USD file */
+	UPROPERTY(EditAnywhere, Category = General, meta = (Bitmask, BitmaskEnum=EUsdPurpose))
+	int32 PurposesToImport;
+
+	/** If checked, all actors generated will have a world space transform and will not have any attachment hierarchy */
+	UPROPERTY(BlueprintReadWrite, config, EditAnywhere, Category=General)
+	bool bFlattenHierarchy;
+
+	/** Defines what should happen with existing actors */
+	UPROPERTY(BlueprintReadWrite, config, EditAnywhere, Category=General)
+	EExistingActorPolicy ExistingActorPolicy;
+
+	/** Whether or not to import custom properties and set their unreal equivalent on spawned actors */
+	UPROPERTY(BlueprintReadWrite, config, EditAnywhere, Category = General)
+	bool bImportProperties;
+
+	/** Whether or not to import mesh geometry or to just spawn actors using existing meshes */
+	UPROPERTY(BlueprintReadWrite, config, EditAnywhere, Category=Mesh)
+	bool bImportMeshes;
+
+	/** The path where new assets are imported */
+	UPROPERTY(BlueprintReadWrite, config, EditAnywhere, Category=Mesh, meta=(ContentDir, EditCondition = bImportMeshes))
+	FDirectoryPath PathForAssets;
+
+	/** What should happen with existing assets */
+	UPROPERTY(BlueprintReadWrite, config, EditAnywhere, Category=Mesh, meta = (EditCondition=bImportMeshes))
+	EExistingAssetPolicy ExistingAssetPolicy;
+
+	/**
+	 * This setting determines what to do if more than one USD prim is found with the same name.  If this setting is true a unique name will be generated and a unique asset will be imported
+	 * If this is false, the first asset found is generated. Assets will be reused when spawning actors into the world.
+	 */
+	UPROPERTY(BlueprintReadWrite, config, EditAnywhere, Category=Mesh, meta=(EditCondition=bImportMeshes))
+	bool bGenerateUniqueMeshes;
+};
+
+UCLASS()
+class USDIMPORTER_API UUSDBatchImportOptionsSubTask : public UObject
+{
+	GENERATED_UCLASS_BODY()
+public:
+
+	/** Path in the USD stage to import from */
+	UPROPERTY(BlueprintReadWrite, Category = Mesh)
+	FString SourcePath;
+
+	/** Path to import asset as */
+	UPROPERTY(BlueprintReadWrite, Category = Mesh)
+	FString DestPath;
+
+	UPROPERTY(BlueprintReadWrite, Category = Mesh)
+	FString ErrorMessage;
+};
+
+UCLASS(config = EditorPerProjectUserSettings)
+class USDIMPORTER_API UUSDBatchImportOptions : public UUSDImportOptions
+{
+	GENERATED_UCLASS_BODY()
+public:
+#if WITH_EDITOR
+	virtual void PostEditChangeProperty(struct FPropertyChangedEvent& PropertyChangedEvent) override;
+	virtual bool CanEditChange(const FProperty* InProperty) const override;
+#endif
+public:
+
+	/** Whether or not to import mesh geometry or to just spawn actors using existing meshes */
+	UPROPERTY(BlueprintReadWrite, config, EditAnywhere, Category=Mesh)
+	bool bImportMeshes;
+
+	/** The path where new assets are imported */
+	UPROPERTY(BlueprintReadWrite, config, EditAnywhere, Category=Mesh, meta=(ContentDir, EditCondition = bImportMeshes))
+	FDirectoryPath PathForAssets;
+
+	/** What should happen with existing assets */
+	UPROPERTY(BlueprintReadWrite, config, EditAnywhere, Category=Mesh, meta = (EditCondition=bImportMeshes))
+	EExistingAssetPolicy ExistingAssetPolicy;
+
+	/**
+	 * This setting determines what to do if more than one USD prim is found with the same name.  If this setting is true a unique name will be generated and a unique asset will be imported
+	 * If this is false, the first asset found is generated. Assets will be reused when spawning actors into the world.
+	 */
+	UPROPERTY(BlueprintReadWrite, config, EditAnywhere, Category=Mesh, meta=(EditCondition=bImportMeshes))
+	bool bGenerateUniqueMeshes;
+
+	UPROPERTY(BlueprintReadWrite, Category = Mesh)
+	TArray<UUSDBatchImportOptionsSubTask*> SubTasks;
+};
diff --git a/Engine/Plugins/Importers/USDImporter/Source/USDSchemas/Private/USDGeomMeshTranslator.cpp b/Engine/Plugins/Importers/USDImporter/Source/USDSchemas/Private/USDGeomMeshTranslator.cpp
--- a/Engine/Plugins/Importers/USDImporter/Source/USDSchemas/Private/USDGeomMeshTranslator.cpp
+++ b/Engine/Plugins/Importers/USDImporter/Source/USDSchemas/Private/USDGeomMeshTranslator.cpp
@@ -53,15 +53,12 @@ namespace UsdGeomMeshTranslatorImpl
 			{
 				const pxr::UsdAttribute& Attribute = GeomMesh.GetPrim().GetAttribute( AttributeName );
 
-				//bHasAttributesTimeSamples = Attribute.ValueMightBeTimeVarying();
-
 				double DesiredTime = TimeCode.GetValue();
 				double MinTime = 0.0;
 				double MaxTime = 0.0;
 				bool bHasTimeSamples = false;
 
-				Attribute.GetBracketingTimeSamples( DesiredTime, &MinTime, &MaxTime, &bHasTimeSamples );
-				if ( bHasTimeSamples && DesiredTime >= MinTime && DesiredTime <= MaxTime )
+				if ( bHasTimeSamples && DesiredTime >= MinTime && DesiredTime <= MaxTime && Attribute.ValueMightBeTimeVarying())
 				{
 					bHasAttributesTimeSamples = true;
 					break;
@@ -253,7 +250,7 @@ namespace UsdGeomMeshTranslatorImpl
 		{
 			StaticMesh.RenderData->Bounds = MeshDescription->GetBounds();
 		}
-	
+
 		StaticMesh.CalculateExtendedBounds();
 	}
 }
@@ -262,16 +259,22 @@ void FGeomMeshCreateAssetsTaskChain::SetupTasks()
 {
 	FScopedUnrealAllocs UnrealAllocs;
 
+	// Ignore meshes from disabled purposes
+	if (!EnumHasAnyFlags(Context->PurposesToLoad, IUsdPrim::GetPurpose(GeomMesh.Get().GetPrim())))
+	{
+		return;
+	}
+
 	// Create mesh description (Async)
 	constexpr bool bIsAsyncTask = true;
-	Do( bIsAsyncTask, 
+	Do( bIsAsyncTask,
 		[ this ]() -> bool
 		{
 			MeshDescription = UsdGeomMeshTranslatorImpl::LoadMeshDescription( GeomMesh.Get(), pxr::UsdTimeCode( Context->Time ) );
 
 			return !MeshDescription.IsEmpty();
 		} );
-		
+
 	if ( !UsdGeomMeshTranslatorImpl::IsGeometryAnimated( GeomMesh.Get(), pxr::UsdTimeCode( Context->Time ) ) )
 	{
 		IMeshBuilderModule* MeshBuilderModule = &FModuleManager::LoadModuleChecked< IMeshBuilderModule >( TEXT("MeshBuilder") );
diff --git a/Engine/Plugins/Importers/USDImporter/Source/USDSchemas/Private/USDGeomXformableTranslator.cpp b/Engine/Plugins/Importers/USDImporter/Source/USDSchemas/Private/USDGeomXformableTranslator.cpp
--- a/Engine/Plugins/Importers/USDImporter/Source/USDSchemas/Private/USDGeomXformableTranslator.cpp
+++ b/Engine/Plugins/Importers/USDImporter/Source/USDSchemas/Private/USDGeomXformableTranslator.cpp
@@ -55,7 +55,7 @@ USceneComponent* FUsdGeomXformableTranslator::CreateComponents( TSubclassOf< USc
 	}
 
 	TUsdStore< pxr::UsdPrim > Prim = Schema.Get().GetPrim();
-	
+
 	if ( !Prim.Get() )
 	{
 		return nullptr;
@@ -90,7 +90,7 @@ USceneComponent* FUsdGeomXformableTranslator::CreateComponents( TSubclassOf< USc
 	{
 		ComponentOuter = Context->ParentComponent;
 	}
-	
+
 	if ( !SceneComponent )
 	{
 		SceneComponent = NewObject< USceneComponent >( ComponentOuter, ComponentType, FName( Prim.Get().GetName().GetText() ), Context->ObjectFlags );
@@ -121,6 +121,9 @@ USceneComponent* FUsdGeomXformableTranslator::CreateComponents( TSubclassOf< USc
 		{
 			SceneComponent->GetOwner()->SetRootComponent( SceneComponent );
 		}
+
+		// Set purpose as a tag so its visibility can be toggled later
+		SceneComponent->ComponentTags.Add(IUsdPrim::GetPurposeName(IUsdPrim::GetPurpose(Prim.Get())));
 	}
 
 	return SceneComponent;
diff --git a/Engine/Plugins/Importers/USDImporter/Source/USDSchemas/Public/USDSchemaTranslator.h b/Engine/Plugins/Importers/USDImporter/Source/USDSchemas/Public/USDSchemaTranslator.h
--- a/Engine/Plugins/Importers/USDImporter/Source/USDSchemas/Public/USDSchemaTranslator.h
+++ b/Engine/Plugins/Importers/USDImporter/Source/USDSchemas/Public/USDSchemaTranslator.h
@@ -5,6 +5,7 @@
 #include "CoreMinimal.h"
 
 #include "USDMemory.h"
+#include "UnrealUSDWrapper.h"
 
 #include "Async/Future.h"
 #include "HAL/ThreadSafeBool.h"
@@ -125,6 +126,9 @@ struct USDSCHEMAS_API FUsdSchemaTranslationContext : public TSharedFromThis< FUs
 	/** The time at which we are translating */
 	float Time = 0.f;
 
+	/** We're only allowed to load prims with purposes that match these flags */
+	EUsdPurpose PurposesToLoad;
+
 	/** Map of translated UsdPrims to UAssets */
 	TMap< FString, UObject* >& PrimPathsToAssets;
 
diff --git a/Engine/Plugins/Importers/USDImporter/Source/USDStage/Private/USDStageActor.cpp b/Engine/Plugins/Importers/USDImporter/Source/USDStage/Private/USDStageActor.cpp
--- a/Engine/Plugins/Importers/USDImporter/Source/USDStage/Private/USDStageActor.cpp
+++ b/Engine/Plugins/Importers/USDImporter/Source/USDStage/Private/USDStageActor.cpp
@@ -74,6 +74,7 @@ struct FUsdStageActorImpl
 		TranslationContext->Level = StageActor->GetLevel();
 		TranslationContext->ObjectFlags = DefaultObjFlag;
 		TranslationContext->Time = StageActor->GetTime();
+		TranslationContext->PurposesToLoad = (EUsdPurpose) StageActor->PurposesToLoad;
 
 		TUsdStore< pxr::SdfPath > UsdPrimPath = UnrealToUsd::ConvertPath( *PrimPath );
 		FUsdPrimTwin* ParentUsdPrimTwin = StageActor->RootUsdTwin.Find( UsdToUnreal::ConvertPath( UsdPrimPath.Get().GetParentPath() ) );
@@ -92,12 +93,42 @@ struct FUsdStageActorImpl
 
 		return TranslationContext;
 	}
+
+	/**
+	 * Will recursively scan down the attach children of Parent and toggle each component's
+	 * Visibility based on whether they have the correct tag
+	 */
+	static void RecursivelyToggleVisibility(USceneComponent* Parent, const TSet<FName>& VisiblePurposeNames)
+	{
+		if (!Parent)
+		{
+			return;
+		}
+
+		bool bVisible = false;
+		for (const FName& Tag : Parent->ComponentTags)
+		{
+			if (VisiblePurposeNames.Contains(Tag))
+			{
+				bVisible = true;
+				break;
+			}
+		}
+		Parent->SetVisibility(bVisible, false);
+
+		for (USceneComponent* Child : Parent->GetAttachChildren())
+		{
+			RecursivelyToggleVisibility(Child, VisiblePurposeNames);
+		}
+	}
 };
 
 #endif // #if USE_USD_SDK
 
 AUsdStageActor::AUsdStageActor()
 	: InitialLoadSet( EUsdInitialLoadSet::LoadAll )
+	, PurposesToLoad((int32)(EUsdPurpose::Default | EUsdPurpose::Proxy))
+	, PurposeVisibility((int32)(EUsdPurpose::Default | EUsdPurpose::Proxy | EUsdPurpose::Render | EUsdPurpose::Guide))
 	, Time( 0.0f )
 	, StartTimeCode( 0.f )
 	, EndTimeCode( 100.f )
@@ -124,6 +155,8 @@ AUsdStageActor::AUsdStageActor()
 				this->UpdatePrim( UsdPrimPath.Get(), bResync, *TranslationContext );
 				TranslationContext->CompleteTasks();
 
+				this->RefreshVisibilityBasedOnPurpose();
+
 				if ( this->HasAutorithyOverStage() )
 				{
 					this->OnPrimChanged.Broadcast( PrimPath, bResync );
@@ -348,6 +381,10 @@ void AUsdStageActor::PostEditChangeProperty(FPropertyChangedEvent& PropertyChang
 			LevelSequenceHelper.UpdateLevelSequence( UsdStage );
 		}
 	}
+	else if ( PropertyName == GET_MEMBER_NAME_CHECKED(AUsdStageActor, PurposeVisibility) )
+	{
+		RefreshVisibilityBasedOnPurpose();
+	}
 	else
 #endif // #if USE_USD_SDK
 	{
@@ -394,6 +431,35 @@ void AUsdStageActor::OpenUsdStage()
 #endif // #if USE_USD_SDK
 }
 
+void AUsdStageActor::RefreshVisibilityBasedOnPurpose()
+{
+#if USE_USD_SDK
+	USceneComponent* Root = GetRootComponent();
+
+	TSet<FName> VisiblePurposeNames;
+	if (EnumHasAnyFlags((EUsdPurpose)PurposeVisibility, EUsdPurpose::Default))
+	{
+		VisiblePurposeNames.Add(IUsdPrim::GetPurposeName(EUsdPurpose::Default));
+	}
+	if (EnumHasAnyFlags((EUsdPurpose)PurposeVisibility, EUsdPurpose::Proxy))
+	{
+		VisiblePurposeNames.Add(IUsdPrim::GetPurposeName(EUsdPurpose::Proxy));
+	}
+	if (EnumHasAnyFlags((EUsdPurpose)PurposeVisibility, EUsdPurpose::Render))
+	{
+		VisiblePurposeNames.Add(IUsdPrim::GetPurposeName(EUsdPurpose::Render));
+	}
+	if (EnumHasAnyFlags((EUsdPurpose)PurposeVisibility, EUsdPurpose::Guide))
+	{
+		VisiblePurposeNames.Add(IUsdPrim::GetPurposeName(EUsdPurpose::Guide));
+	}
+
+	FUsdStageActorImpl::RecursivelyToggleVisibility(Root, VisiblePurposeNames);
+
+	GEditor->RedrawLevelEditingViewports();
+#endif // #if USE_USD_SDK
+}
+
 void AUsdStageActor::LoadUsdStage()
 {
 #if USE_USD_SDK
@@ -429,7 +495,8 @@ void AUsdStageActor::LoadUsdStage()
 	SetTime( StartTimeCode );
 
 	GEditor->BroadcastLevelActorListChanged();
-	GEditor->RedrawLevelEditingViewports();
+
+	this->RefreshVisibilityBasedOnPurpose();
 
 	OnTimeChanged.AddUObject( this, &AUsdStageActor::AnimatePrims );
 
diff --git a/Engine/Plugins/Importers/USDImporter/Source/USDStage/Public/USDStageActor.h b/Engine/Plugins/Importers/USDImporter/Source/USDStage/Public/USDStageActor.h
--- a/Engine/Plugins/Importers/USDImporter/Source/USDStage/Public/USDStageActor.h
+++ b/Engine/Plugins/Importers/USDImporter/Source/USDStage/Public/USDStageActor.h
@@ -27,12 +27,13 @@
 #include "USDStageActor.generated.h"
 
 class ALevelSequenceActor;
-struct FMeshDescription;
-struct FUsdSchemaTranslationContext;
 class IMeshBuilderModule;
 class ULevelSequence;
 class UMaterial;
 class UUsdAsset;
+enum class EUsdPurpose : int32;
+struct FMeshDescription;
+struct FUsdSchemaTranslationContext;
 
 DECLARE_LOG_CATEGORY_EXTERN( LogUsdStage, Log, All );
 
@@ -58,6 +59,14 @@ public:
 	UPROPERTY(EditAnywhere, Category = "USD")
 	EUsdInitialLoadSet InitialLoadSet;
 
+	/* Only load prims with these specific purposes from the USD file */
+	UPROPERTY(EditAnywhere, Category = "USD", meta = (Bitmask, BitmaskEnum=EUsdPurpose))
+	int32 PurposesToLoad;
+
+	/* Quickly toggle visibility of prims with specific purposes in the level based on component tags */
+	UPROPERTY(EditAnywhere, Category = "USD", meta = (Bitmask, BitmaskEnum=EUsdPurpose))
+	int32 PurposeVisibility;
+
 	UFUNCTION(BlueprintCallable, Category = "USD", meta = (CallInEditor = "true"))
 	float GetTime() const { return Time; }
 
@@ -117,6 +126,8 @@ private:
 	void OpenUsdStage();
 	void LoadUsdStage();
 
+	void RefreshVisibilityBasedOnPurpose();
+
 	void OnUsdPrimTwinDestroyed( const FUsdPrimTwin& UsdPrimTwin );
 
 	void OnPrimObjectPropertyChanged( UObject* ObjectBeingModified, FPropertyChangedEvent& PropertyChangedEvent );
diff --git a/Engine/Plugins/Importers/USDImporter/Source/USDStageEditor/Private/SUSDStage.cpp b/Engine/Plugins/Importers/USDImporter/Source/USDStageEditor/Private/SUSDStage.cpp
--- a/Engine/Plugins/Importers/USDImporter/Source/USDStageEditor/Private/SUSDStage.cpp
+++ b/Engine/Plugins/Importers/USDImporter/Source/USDStageEditor/Private/SUSDStage.cpp
@@ -8,6 +8,7 @@
 #include "SUSDStageInfo.h"
 #include "SUSDStageTreeView.h"
 #include "USDErrorUtils.h"
+#include "USDImportOptions.h"
 #include "USDImporter.h"
 #include "USDLayerUtils.h"
 #include "USDStageActor.h"
@@ -422,12 +423,13 @@ void SUsdStage::OnReloadStage()
 
 void SUsdStage::OnImport()
 {
-	if ( !UsdStageActor.IsValid() )
+	AUsdStageActor* StageActor = UsdStageActor.Get();
+	if ( !StageActor )
 	{
 		return;
 	}
 
-	const pxr::UsdStageRefPtr& UsdStage = UsdStageActor->GetUsdStage();
+	const pxr::UsdStageRefPtr& UsdStage = StageActor->GetUsdStage();
 
 	if ( !UsdStage )
 	{
@@ -475,6 +477,11 @@ void SUsdStage::OnImport()
 		ImportContextContainer->ImportContext.bIsAutomated = false;
 		ImportContextContainer->ImportContext.bApplyWorldTransformToGeometry = false;
 
+		if (UUSDSceneImportOptions* SceneOptions = Cast<UUSDSceneImportOptions>(ImportContextContainer->ImportContext.ImportOptions))
+		{
+			SceneOptions->PurposesToImport = StageActor->PurposesToLoad;
+		}
+
 		UUSDImporter* UsdImporter = IUSDImporterModule::Get().GetImporter();
 
 		bCanceled = !UsdImporter->ShowImportOptions( ImportContextContainer->ImportContext );
diff --git a/Engine/Plugins/Importers/USDImporter/Source/UnrealUSDWrapper/Private/UnrealUSDWrapper.cpp b/Engine/Plugins/Importers/USDImporter/Source/UnrealUSDWrapper/Private/UnrealUSDWrapper.cpp
--- a/Engine/Plugins/Importers/USDImporter/Source/UnrealUSDWrapper/Private/UnrealUSDWrapper.cpp
+++ b/Engine/Plugins/Importers/USDImporter/Source/UnrealUSDWrapper/Private/UnrealUSDWrapper.cpp
@@ -34,6 +34,7 @@
 #include "pxr/usd/usdGeom/mesh.h"
 #include "pxr/usd/usdGeom/metrics.h"
 #include "pxr/usd/usdGeom/modelAPI.h"
+#include "pxr/usd/usdGeom/tokens.h"
 #include "pxr/usd/usdGeom/xformCommonAPI.h"
 #include "pxr/usd/usdShade/materialBindingAPI.h"
 #include "pxr/usd/usdUtils/stageCache.h"
@@ -71,10 +72,6 @@ namespace UnrealIdentifiers
 
 	static const TfToken PropertyPath("unrealPropertyPath");
 
-	static const TfToken ProxyPurpose("proxy");
-
-	static const TfToken GuidePurpose("guide");
-
 	static const TfToken MaterialRelationship("material:binding");
 
 	const TfToken MaterialAssignments = TfToken("unrealMaterials");
@@ -441,20 +438,51 @@ int FUsdAttribute::GetArraySize( const pxr::UsdAttribute& Attribute )
 	return Value.IsArrayValued() ? (int)Value.GetArraySize() : -1;
 
 }
-bool IUsdPrim::IsProxyOrGuide( const UsdPrim& Prim )
+EUsdPurpose IUsdPrim::GetPurpose( const UsdPrim& Prim )
 {
 	UsdGeomImageable Geom(Prim);
 	if (Geom)
 	{
-		UsdAttribute PurposeAttr = Geom.GetPurposeAttr();
+		// Use compute purpose because it depends on the hierarchy:
+		// "If the purpose of </RootPrim> is set to "render", then the effective purpose
+		// of </RootPrim/ChildPrim> will be "render" even if that prim has a different
+		// authored value for purpose."
+		const TfToken Purpose = Geom.ComputePurpose();
+		if (Purpose == pxr::UsdGeomTokens->proxy)
+		{
+			return EUsdPurpose::Proxy;
+		}
+		else if (Purpose == pxr::UsdGeomTokens->render)
+		{
+			return EUsdPurpose::Render;
+		}
+		else if (Purpose == pxr::UsdGeomTokens->guide)
+		{
+			return EUsdPurpose::Guide;
+		}
+	}
 
-		TfToken Purpose;
-		PurposeAttr.Get(&Purpose);
+	return EUsdPurpose::Default;
+}
 
-		return Purpose == UnrealIdentifiers::ProxyPurpose || Purpose == UnrealIdentifiers::GuidePurpose;
+UNREALUSDWRAPPER_API FName IUsdPrim::GetPurposeName(EUsdPurpose Purpose)
+{
+	switch (Purpose)
+	{
+	case EUsdPurpose::Proxy:
+		return TEXT("ProxyPurpose");
+		break;
+	case EUsdPurpose::Render:
+		return TEXT("RenderPurpose");
+		break;
+	case EUsdPurpose::Guide:
+		return TEXT("GuidePurpose");
+		break;
+	default:
+		break;
 	}
 
-	return false;
+	return TEXT("DefaultPurpose");
 }
 
 bool IUsdPrim::HasGeometryData(const UsdPrim& Prim)
diff --git a/Engine/Plugins/Importers/USDImporter/Source/UnrealUSDWrapper/Public/UnrealUSDWrapper.h b/Engine/Plugins/Importers/USDImporter/Source/UnrealUSDWrapper/Public/UnrealUSDWrapper.h
--- a/Engine/Plugins/Importers/USDImporter/Source/UnrealUSDWrapper/Public/UnrealUSDWrapper.h
+++ b/Engine/Plugins/Importers/USDImporter/Source/UnrealUSDWrapper/Public/UnrealUSDWrapper.h
@@ -6,6 +6,7 @@
 
 #include "Modules/ModuleInterface.h"
 #include "Templates/Tuple.h"
+#include "UObject/ObjectMacros.h"
 #include "USDMemory.h"
 
 #include <string>
@@ -72,6 +73,16 @@ enum class EUsdUpAxis
 	ZAxis,
 };
 
+UENUM(meta = (Bitflags, UseEnumValuesAsMaskValuesInEditor = "true"))
+enum class EUsdPurpose : int32
+{
+	None = 0 UMETA(Hidden),
+	Default = 1,
+	Proxy = 2,
+	Render = 4,
+	Guide = 8
+};
+ENUM_CLASS_FLAGS(EUsdPurpose);
 
 
 struct FUsdVector2Data
@@ -192,7 +203,9 @@ class IUsdPrim
 {
 public:
 #if USE_USD_SDK
-	static UNREALUSDWRAPPER_API bool IsProxyOrGuide(const pxr::UsdPrim& Prim);
+	static UNREALUSDWRAPPER_API EUsdPurpose GetPurpose(const pxr::UsdPrim& Prim);
+	static UNREALUSDWRAPPER_API FName GetPurposeName(EUsdPurpose Purpose);
+
 	static UNREALUSDWRAPPER_API bool HasGeometryData(const pxr::UsdPrim& Prim);
 	static UNREALUSDWRAPPER_API bool HasGeometryDataOrLODVariants(const pxr::UsdPrim& Prim);
 	static UNREALUSDWRAPPER_API int GetNumLODs(const pxr::UsdPrim& Prim);
diff --git a/Engine/Plugins/Importers/USDImporter/Source/UnrealUSDWrapper/UnrealUSDWrapper.Build.cs b/Engine/Plugins/Importers/USDImporter/Source/UnrealUSDWrapper/UnrealUSDWrapper.Build.cs
--- a/Engine/Plugins/Importers/USDImporter/Source/UnrealUSDWrapper/UnrealUSDWrapper.Build.cs
+++ b/Engine/Plugins/Importers/USDImporter/Source/UnrealUSDWrapper/UnrealUSDWrapper.Build.cs
@@ -14,6 +14,7 @@ namespace UnrealBuildTool.Rules
 			PublicDependencyModuleNames.AddRange(
 				new string[] {
 				"Core",
+				"CoreUObject",
 				"Python"
 				});
 
