Fix for compile error.

diff --git a/Engine/Plugins/Importers/USDImporter/Source/UnrealUSDWrapper/UnrealUSDWrapper.Build.cs b/Engine/Plugins/Importers/USDImporter/Source/UnrealUSDWrapper/UnrealUSDWrapper.Build.cs
--- a/Engine/Plugins/Importers/USDImporter/Source/UnrealUSDWrapper/UnrealUSDWrapper.Build.cs
+++ b/Engine/Plugins/Importers/USDImporter/Source/UnrealUSDWrapper/UnrealUSDWrapper.Build.cs
@@ -23,7 +23,7 @@ namespace UnrealBuildTool.Rules
 
 			if (Target.WindowsPlatform.Compiler != WindowsCompiler.Clang && Target.WindowsPlatform.StaticAnalyzer == WindowsStaticAnalyzer.None &&
 				Target.CppStandard < CppStandardVersion.Cpp17 && // Not currently compatible with C++17 due to old version of Boost
-				Target.LinkType != TargetLinkType.Monolithicc && // If you want to use USD in a monolithic target, you'll have to use the ANSI allocator and remove this condition
+				Target.LinkType != TargetLinkType.Monolithic && // If you want to use USD in a monolithic target, you'll have to use the ANSI allocator and remove this condition
 				(Target.Platform != UnrealTargetPlatform.Linux || Target.bForceEnableRTTI)) // USD on Linux needs RTTI enabled for the whole editor
 			{
 				PublicDefinitions.Add("USE_USD_SDK=1");
