Copied latest USD Importer from dev-virtualproduction [at] cl 11445959

diff --git a/Engine/Plugins/Importers/USDImporter/Source/USDImporter/Private/USDAssetImportFactory.cpp b/Engine/Plugins/Importers/USDImporter/Source/USDImporter/Private/USDAssetImportFactory.cpp
--- a/Engine/Plugins/Importers/USDImporter/Source/USDImporter/Private/USDAssetImportFactory.cpp
+++ b/Engine/Plugins/Importers/USDImporter/Source/USDImporter/Private/USDAssetImportFactory.cpp
@@ -108,6 +108,9 @@ UObject* UUSDAssetImportFactory::FactoryCreateFile(UClass* InClass, UObject* InP
 			}
 		}
 
+		// Reset this cache or else reimport will not work properly
+		ImportContext.PathToImportAssetMap.Empty();
+
 		ImportContext.DisplayErrorMessages(IsAutomatedImport());
 	}
 	else
diff --git a/Engine/Plugins/Importers/USDImporter/Source/USDImporter/Public/USDImportOptions.h b/Engine/Plugins/Importers/USDImporter/Source/USDImporter/Public/USDImportOptions.h
--- a/Engine/Plugins/Importers/USDImporter/Source/USDImporter/Public/USDImportOptions.h
+++ b/Engine/Plugins/Importers/USDImporter/Source/USDImporter/Public/USDImportOptions.h
@@ -79,7 +79,7 @@ public:
 #endif
 public:
 	/* Only import prims with these specific purposes from the USD file */
-	UPROPERTY(EditAnywhere, Category = General, meta = (Bitmask, BitmaskEnum=EUsdPurpose))
+	UPROPERTY(BlueprintReadWrite, config, EditAnywhere, Category=General, meta = (Bitmask, BitmaskEnum=EUsdPurpose))
 	int32 PurposesToImport;
 
 	/** If checked, all actors generated will have a world space transform and will not have any attachment hierarchy */
diff --git a/Engine/Plugins/Importers/USDImporter/Source/USDSchemas/Private/USDGeomCameraTranslator.cpp b/Engine/Plugins/Importers/USDImporter/Source/USDSchemas/Private/USDGeomCameraTranslator.cpp
--- a/Engine/Plugins/Importers/USDImporter/Source/USDSchemas/Private/USDGeomCameraTranslator.cpp
+++ b/Engine/Plugins/Importers/USDImporter/Source/USDSchemas/Private/USDGeomCameraTranslator.cpp
@@ -26,10 +26,21 @@ USceneComponent* FUsdGeomCameraTranslator::CreateComponents()
 
 void FUsdGeomCameraTranslator::UpdateComponents( USceneComponent* SceneComponent )
 {
-	if ( UCineCameraComponent* CameraComponent = Cast< UCineCameraComponent >( SceneComponent ) )
+	if ( !SceneComponent )
 	{
-		FScopedUsdAllocs UsdAllocs;
-		UsdToUnreal::ConvertGeomCamera( Schema.Get().GetPrim().GetStage(), pxr::UsdGeomCamera( Schema.Get() ), *CameraComponent, pxr::UsdTimeCode( Context->Time ) );
+		return;
+	}
+
+	FUsdGeomXformableTranslator::UpdateComponents( SceneComponent );
+
+	// The CineCamera component is not actually the root component
+	for ( USceneComponent* ChildComponent : SceneComponent->GetAttachChildren() )
+	{
+		if ( UCineCameraComponent* CameraComponent = Cast< UCineCameraComponent >( ChildComponent ) )
+		{
+			FScopedUsdAllocs UsdAllocs;
+			UsdToUnreal::ConvertGeomCamera( Schema.Get().GetPrim().GetStage(), pxr::UsdGeomCamera( Schema.Get() ), *CameraComponent, pxr::UsdTimeCode( Context->Time ) );
+		}
 	}
 }
 
diff --git a/Engine/Plugins/Importers/USDImporter/Source/USDSchemas/Private/USDGeomMeshTranslator.cpp b/Engine/Plugins/Importers/USDImporter/Source/USDSchemas/Private/USDGeomMeshTranslator.cpp
--- a/Engine/Plugins/Importers/USDImporter/Source/USDSchemas/Private/USDGeomMeshTranslator.cpp
+++ b/Engine/Plugins/Importers/USDImporter/Source/USDSchemas/Private/USDGeomMeshTranslator.cpp
@@ -17,13 +17,14 @@
 #include "Interfaces/ITargetPlatform.h"
 #include "Interfaces/ITargetPlatformManagerModule.h"
 #include "Materials/Material.h"
-#include "Materials/MaterialInterface.h"
 #include "Materials/MaterialInstanceConstant.h"
+#include "Materials/MaterialInterface.h"
 #include "Misc/SecureHash.h"
 #include "Modules/ModuleManager.h"
 #include "PhysicsEngine/BodySetup.h"
 #include "StaticMeshAttributes.h"
 #include "StaticMeshOperations.h"
+#include "StaticMeshResources.h"
 #include "UObject/SoftObjectPath.h"
 
 #include "USDIncludesStart.h"
@@ -174,7 +175,7 @@ namespace UsdGeomMeshTranslatorImpl
 					Material = MaterialInstance;
 				}
 			}
-			
+
 			FStaticMaterial StaticMaterial( Material, MaterialSlotName );
 
 			if ( !StaticMesh.StaticMaterials.IsValidIndex( MaterialIndex ) )
@@ -302,6 +303,7 @@ namespace UsdGeomMeshTranslatorImpl
 	void PostBuildStaticMesh( UStaticMesh& StaticMesh )
 	{
 		TRACE_CPUPROFILER_EVENT_SCOPE( UsdGeomMeshTranslatorImpl::PostBuildStaticMesh );
+
 		StaticMesh.InitResources();
 
 		if ( const FMeshDescription* MeshDescription = StaticMesh.GetMeshDescription( 0 ) )
@@ -375,6 +377,8 @@ void FBuildStaticMeshTaskChain::SetupTasks()
 		Then( !bIsAsyncTask,
 			[ this ]()
 			{
+				RecreateRenderStateContextPtr = MakeShared<FStaticMeshComponentRecreateRenderStateContext>( StaticMesh, true, true );
+
 				UsdGeomMeshTranslatorImpl::PreBuildStaticMesh( Schema.Get().GetPrim(), *StaticMesh, Context->PrimPathsToAssets, Context->Time );
 
 				return true;
@@ -401,6 +405,8 @@ void FBuildStaticMeshTaskChain::SetupTasks()
 			{
 				UsdGeomMeshTranslatorImpl::PostBuildStaticMesh( *StaticMesh );
 
+				RecreateRenderStateContextPtr.Reset();
+
 				return true;
 			} );
 	}
@@ -412,7 +418,7 @@ void FGeomMeshCreateAssetsTaskChain::SetupTasks()
 
 	// Create mesh description (Async)
 	constexpr bool bIsAsyncTask = true;
-	Do( bIsAsyncTask, 
+	Do( bIsAsyncTask,
 		[ this ]() -> bool
 		{
 			MeshDescription = UsdGeomMeshTranslatorImpl::LoadMeshDescription( pxr::UsdGeomMesh( Schema.Get() ), pxr::UsdTimeCode( Context->Time ) );
diff --git a/Engine/Plugins/Importers/USDImporter/Source/USDSchemas/Private/USDGeomPointInstancerTranslator.cpp b/Engine/Plugins/Importers/USDImporter/Source/USDSchemas/Private/USDGeomPointInstancerTranslator.cpp
--- a/Engine/Plugins/Importers/USDImporter/Source/USDSchemas/Private/USDGeomPointInstancerTranslator.cpp
+++ b/Engine/Plugins/Importers/USDImporter/Source/USDSchemas/Private/USDGeomPointInstancerTranslator.cpp
@@ -85,14 +85,11 @@ namespace UsdGeomPointInstancerTranslatorImpl
 	}
 }
 
-USceneComponent* FUsdGeomPointInstancerTranslator::CreateComponents()
+void FUsdGeomPointInstancerTranslator::UpdateComponents(USceneComponent* PointInstancerRootComponent)
 {
-	// Spawn Xform Actor / SceneComponent
-	USceneComponent* PointInstancerRootComponent = FUsdGeomXformableTranslator::CreateComponents();
-
 	if ( !PointInstancerRootComponent )
 	{
-		return nullptr;
+		return;
 	}
 
 	FScopedUsdAllocs UsdAllocs;
@@ -101,7 +98,7 @@ USceneComponent* FUsdGeomPointInstancerTranslator::CreateComponents()
 
 	if ( !PointInstancer )
 	{
-		return nullptr;
+		return;
 	}
 
 	pxr::UsdPrim Prim = PointInstancer.GetPrim();
@@ -114,10 +111,13 @@ USceneComponent* FUsdGeomPointInstancerTranslator::CreateComponents()
 	{
 		TGuardValue< USceneComponent* > ParentComponentGuard( Context->ParentComponent, PointInstancerRootComponent );
 
-		int32 PrototypeIndex = 0;
-		for ( const pxr::SdfPath& PrototypePath : PrototypesPaths )
+		for ( int32 PrototypeIndex = 0; PrototypeIndex < PrototypesPaths.size(); ++PrototypeIndex)
 		{
-			pxr::UsdPrim PrototypePrim = Prim.GetStage()->GetPrimAtPath( PrototypePath );
+			pxr::UsdPrim PrototypePrim = Prim.GetStage()->GetPrimAtPath( PrototypesPaths[PrototypeIndex] );
+			if (!Prim)
+			{
+				continue;
+			}
 
 			FUsdGeomXformableTranslator PrototypeXformTranslator( Context, pxr::UsdTyped( PrototypePrim ) );
 			USceneComponent* PrototypeXformComponent = PrototypeXformTranslator.CreateComponents();
@@ -142,12 +142,8 @@ USceneComponent* FUsdGeomPointInstancerTranslator::CreateComponents()
 					UsdGeomPointInstancerTranslatorImpl::ConvertGeomPointInstancer( Prim.GetStage(), PointInstancer, PrototypeIndex, *HismComponent, pxr::UsdTimeCode( Context->Time ) );
 				}
 			}
-
-			++PrototypeIndex;
 		}
 	}
-
-	return PointInstancerRootComponent;
 }
 
 #endif // #if USE_USD_SDK
diff --git a/Engine/Plugins/Importers/USDImporter/Source/USDSchemas/Private/USDGeomXformableTranslator.cpp b/Engine/Plugins/Importers/USDImporter/Source/USDSchemas/Private/USDGeomXformableTranslator.cpp
--- a/Engine/Plugins/Importers/USDImporter/Source/USDSchemas/Private/USDGeomXformableTranslator.cpp
+++ b/Engine/Plugins/Importers/USDImporter/Source/USDSchemas/Private/USDGeomXformableTranslator.cpp
@@ -95,7 +95,7 @@ void FUsdGeomXformableCreateAssetsTaskChain::SetupTasks()
 
 	// Create mesh description (Async)
 	constexpr bool bIsAsyncTask = true;
-	Do( bIsAsyncTask, 
+	Do( bIsAsyncTask,
 		[ this ]() -> bool
 		{
 			MeshDescription = UsdGeomXformableTranslatorImpl::LoadMeshDescription( pxr::UsdGeomXformable( Schema.Get() ), Context->PurposesToLoad, pxr::UsdTimeCode( Context->Time ) );
@@ -149,23 +149,6 @@ USceneComponent* FUsdGeomXformableTranslator::CreateComponents()
 
 	USceneComponent* RootComponent = CreateComponents( ComponentType );
 
-	if ( UStaticMeshComponent* StaticMeshComponent = Cast< UStaticMeshComponent >( RootComponent ) )
-	{
-		UStaticMesh* PrimStaticMesh = Cast< UStaticMesh >( Context->PrimPathsToAssets.FindRef( UsdToUnreal::ConvertPath( Schema.Get().GetPath() ) ) );
-
-		if ( PrimStaticMesh != StaticMeshComponent->GetStaticMesh() )
-		{
-			if ( StaticMeshComponent->IsRegistered() )
-			{
-				StaticMeshComponent->UnregisterComponent();
-			}
-
-			StaticMeshComponent->SetStaticMesh( PrimStaticMesh );
-
-			StaticMeshComponent->RegisterComponent();
-		}
-	}
-
 	return RootComponent;
 }
 
@@ -174,7 +157,7 @@ USceneComponent* FUsdGeomXformableTranslator::CreateComponents( TSubclassOf< USc
 	TUsdStore< pxr::UsdPrim > Prim = Schema.Get().GetPrim();
 
 	const bool bNeedsActor =
-		( 
+		(
 			Context->ParentComponent == nullptr ||
 			Prim.Get().IsPseudoRoot() ||
 			Prim.Get().IsModel() ||
@@ -272,6 +255,28 @@ void FUsdGeomXformableTranslator::UpdateComponents( USceneComponent* SceneCompon
 	if ( SceneComponent )
 	{
 		UsdToUnreal::ConvertXformable( Schema.Get().GetPrim().GetStage(), pxr::UsdGeomXformable( Schema.Get() ), *SceneComponent, pxr::UsdTimeCode( Context->Time ) );
+
+		// If the user modified a mesh parameter (e.g. vertex color), the hash will be different and it will become a separate asset
+		// so we must check for this and assign the new StaticMesh
+		if ( UStaticMeshComponent* StaticMeshComponent = Cast< UStaticMeshComponent >( SceneComponent ) )
+		{
+			UStaticMesh* PrimStaticMesh = Cast< UStaticMesh >( Context->PrimPathsToAssets.FindRef( UsdToUnreal::ConvertPath( Schema.Get().GetPath() ) ) );
+
+			if ( PrimStaticMesh != StaticMeshComponent->GetStaticMesh() )
+			{
+				// Need to make sure the mesh's resources are initialized here as it may have just been built in another thread
+				PrimStaticMesh->InitResources();
+
+				if ( StaticMeshComponent->IsRegistered() )
+				{
+					StaticMeshComponent->UnregisterComponent();
+				}
+
+				StaticMeshComponent->SetStaticMesh( PrimStaticMesh );
+
+				StaticMeshComponent->RegisterComponent();
+			}
+		}
 	}
 }
 
@@ -299,7 +304,7 @@ bool FUsdGeomXformableTranslator::CollapsesChildren( ECollapsingType CollapsingT
 
 			TArray< TUsdStore< pxr::UsdPrim > > ChildXformPrims = UsdUtils::GetAllPrimsOfType( Schema.Get().GetPrim(), pxr::TfType::Find< pxr::UsdGeomXformable >() );
 
-			for ( const TUsdStore< pxr::UsdPrim >& ChildXformPrim : ChildXformPrims ) 
+			for ( const TUsdStore< pxr::UsdPrim >& ChildXformPrim : ChildXformPrims )
 			{
 				if ( TSharedPtr< FUsdSchemaTranslator > SchemaTranslator = UsdSchemasModule.GetTranslatorRegistry().CreateTranslatorForSchema( Context, pxr::UsdTyped( ChildXformPrim.Get() ) ) )
 				{
diff --git a/Engine/Plugins/Importers/USDImporter/Source/USDSchemas/Private/USDSchemaTranslator.cpp b/Engine/Plugins/Importers/USDImporter/Source/USDSchemas/Private/USDSchemaTranslator.cpp
--- a/Engine/Plugins/Importers/USDImporter/Source/USDSchemas/Private/USDSchemaTranslator.cpp
+++ b/Engine/Plugins/Importers/USDImporter/Source/USDSchemas/Private/USDSchemaTranslator.cpp
@@ -23,6 +23,10 @@ int32 FRegisteredSchemaTranslatorHandle::CurrentSchemaTranslatorId = 0;
 TSharedPtr< FUsdSchemaTranslator > FUsdSchemaTranslatorRegistry::CreateTranslatorForSchema( TSharedRef< FUsdSchemaTranslationContext > InTranslationContext, const pxr::UsdTyped& InSchema )
 {
 	TUsdStore< pxr::UsdPrim > Prim = InSchema.GetPrim();
+	if ( !Prim.Get() )
+	{
+		return {};
+	}
 
 	for ( TPair< FString, FSchemaTranslatorsStack >& RegisteredSchemasStack : RegisteredSchemaTranslators )
 	{
diff --git a/Engine/Plugins/Importers/USDImporter/Source/USDSchemas/Private/USDShadeMaterialTranslator.cpp b/Engine/Plugins/Importers/USDImporter/Source/USDSchemas/Private/USDShadeMaterialTranslator.cpp
--- a/Engine/Plugins/Importers/USDImporter/Source/USDSchemas/Private/USDShadeMaterialTranslator.cpp
+++ b/Engine/Plugins/Importers/USDImporter/Source/USDSchemas/Private/USDShadeMaterialTranslator.cpp
@@ -29,7 +29,7 @@ void FUsdShadeMaterialTranslator::CreateAssets()
 	{
 		UMaterial* NewMaterial = NewObject< UMaterial >();
 
-		if ( UsdToUnreal::ConvertMaterial( ShadeMaterial, *NewMaterial ) )
+		if ( UsdToUnreal::ConvertMaterial( ShadeMaterial, *NewMaterial, Context->AssetsCache ) )
 		{
 			//UMaterialEditingLibrary::RecompileMaterial( CachedMaterial ); // Too slow
 			NewMaterial->PostEditChange();
diff --git a/Engine/Plugins/Importers/USDImporter/Source/USDSchemas/Public/USDGeomMeshTranslator.h b/Engine/Plugins/Importers/USDImporter/Source/USDSchemas/Public/USDGeomMeshTranslator.h
--- a/Engine/Plugins/Importers/USDImporter/Source/USDSchemas/Public/USDGeomMeshTranslator.h
+++ b/Engine/Plugins/Importers/USDImporter/Source/USDSchemas/Public/USDGeomMeshTranslator.h
@@ -13,6 +13,7 @@
 #include "USDIncludesEnd.h"
 
 class UStaticMesh;
+class FStaticMeshComponentRecreateRenderStateContext;
 
 PXR_NAMESPACE_OPEN_SCOPE
 	class UsdGeomMesh;
@@ -33,6 +34,9 @@ protected:
 	// Outputs
 	UStaticMesh* StaticMesh = nullptr;
 
+	// Required to prevent StaticMesh from being used for drawing while it is being rebuilt
+	TSharedPtr<FStaticMeshComponentRecreateRenderStateContext> RecreateRenderStateContextPtr;
+
 protected:
 	FBuildStaticMeshTaskChain( const TSharedRef< FUsdSchemaTranslationContext >& InContext, const TUsdStore< pxr::UsdTyped >& InSchema )
 		: Schema( InSchema )
diff --git a/Engine/Plugins/Importers/USDImporter/Source/USDSchemas/Public/USDGeomPointInstancerTranslator.h b/Engine/Plugins/Importers/USDImporter/Source/USDSchemas/Public/USDGeomPointInstancerTranslator.h
--- a/Engine/Plugins/Importers/USDImporter/Source/USDSchemas/Public/USDGeomPointInstancerTranslator.h
+++ b/Engine/Plugins/Importers/USDImporter/Source/USDSchemas/Public/USDGeomPointInstancerTranslator.h
@@ -17,8 +17,8 @@ class USDSCHEMAS_API FUsdGeomPointInstancerTranslator : public FUsdGeomXformable
 {
 public:
 	using FUsdGeomXformableTranslator::FUsdGeomXformableTranslator;
-	
-	virtual USceneComponent* CreateComponents() override;
+
+	virtual void UpdateComponents( USceneComponent* SceneComponent ) override;
 	virtual bool CollapsesChildren( ECollapsingType CollapsingType ) const override { return CollapsingType == FUsdSchemaTranslator::ECollapsingType::Components; }
 	virtual bool CanBeCollapsed( ECollapsingType CollapsingType ) const override { return false; }
 };
diff --git a/Engine/Plugins/Importers/USDImporter/Source/USDStage/Private/USDLevelSequenceHelper.cpp b/Engine/Plugins/Importers/USDImporter/Source/USDStage/Private/USDLevelSequenceHelper.cpp
--- a/Engine/Plugins/Importers/USDImporter/Source/USDStage/Private/USDLevelSequenceHelper.cpp
+++ b/Engine/Plugins/Importers/USDImporter/Source/USDStage/Private/USDLevelSequenceHelper.cpp
@@ -346,6 +346,11 @@ void FUsdLevelSequenceHelperImpl::CreateTimeTrack(const FUsdLevelSequenceHelperI
 	else
 	{
 		TimeTrack = MovieScene->AddTrack<UMovieSceneFloatTrack>(StageActorBinding);
+		if (!TimeTrack)
+		{
+			return;
+		}
+
 		TimeTrack->SetPropertyNameAndPath(FName(TIME_TRACK_NAME), "Time");
 
 		MovieScene->SetEvaluationType(EMovieSceneEvaluationType::FrameLocked);
diff --git a/Engine/Plugins/Importers/USDImporter/Source/USDStage/Private/USDStageActor.cpp b/Engine/Plugins/Importers/USDImporter/Source/USDStage/Private/USDStageActor.cpp
--- a/Engine/Plugins/Importers/USDImporter/Source/USDStage/Private/USDStageActor.cpp
+++ b/Engine/Plugins/Importers/USDImporter/Source/USDStage/Private/USDStageActor.cpp
@@ -157,7 +157,7 @@ AUsdStageActor::AUsdStageActor()
 
 					if ( TSharedPtr< FUsdSchemaTranslator > SchemaTranslator = UsdSchemasModule.GetTranslatorRegistry().CreateTranslatorForSchema( TranslationContext, pxr::UsdTyped( UsdPrim.Get() ) ) )
 					{
-						while ( SchemaTranslator->IsCollapsed( FUsdSchemaTranslator::ECollapsingType::Assets ) )
+						while ( SchemaTranslator->IsCollapsed( CollapsingType ) )
 						{
 							UsdPrimPath = UsdPrimPath.Get().GetParentPath();
 							UsdPrim = this->GetUsdStage()->GetPrimAtPath( UsdPrimPath.Get() );
@@ -186,7 +186,7 @@ AUsdStageActor::AUsdStageActor()
 					this->UpdatePrim( ComponentsPrimPath.Get(), bResync, *TranslationContext );
 					TranslationContext->CompleteTasks();
 				}
-				
+
 				this->RefreshVisibilityBasedOnPurpose();
 
 				if ( this->HasAutorithyOverStage() )
@@ -235,6 +235,36 @@ AUsdStageActor::~AUsdStageActor()
 #endif // #if USE_USD_SDK
 }
 
+USDSTAGE_API void AUsdStageActor::Reset()
+{
+	Super::Reset();
+
+	Modify();
+
+	Clear();
+	AssetsCache.Reset();
+
+	GEditor->GetEditorSubsystem<UAssetEditorSubsystem>()->CloseAllEditorsForAsset(LevelSequence);
+	LevelSequence = nullptr;
+	StartTimeCode = 0.0f;
+	EndTimeCode = 100.0f;
+	TimeCodesPerSecond = 24.0f;
+	Time = 0;
+
+	RootUsdTwin.Clear();
+	RootUsdTwin.PrimPath = TEXT("/");
+	GEditor->BroadcastLevelActorListChanged();
+
+	RootLayer.FilePath.Empty();
+
+#if USE_USD_SDK
+	UnrealUSDWrapper::GetUsdStageCache().Erase( UsdStageStore.Get() );
+	UsdStageStore = TUsdStore< pxr::UsdStageRefPtr >();
+#endif // #if USE_USD_SDK
+
+	OnStageChanged.Broadcast();
+}
+
 #if USE_USD_SDK
 
 FUsdPrimTwin* AUsdStageActor::GetOrCreatePrimTwin( const pxr::SdfPath& UsdPrimPath )
@@ -395,8 +425,9 @@ void AUsdStageActor::PostEditChangeProperty(FPropertyChangedEvent& PropertyChang
 	if ( PropertyName == GET_MEMBER_NAME_CHECKED( AUsdStageActor, RootLayer ) )
 	{
 		UnrealUSDWrapper::GetUsdStageCache().Erase( UsdStageStore.Get() );
-
 		UsdStageStore = TUsdStore< pxr::UsdStageRefPtr >();
+
+		AssetsCache.Reset(); // We've changed USD file, clear the cache
 		LoadUsdStage();
 	}
 	else if ( PropertyName == GET_MEMBER_NAME_CHECKED( AUsdStageActor, Time ) )
@@ -426,9 +457,9 @@ void AUsdStageActor::PostEditChangeProperty(FPropertyChangedEvent& PropertyChang
 
 void AUsdStageActor::Clear()
 {
-	AssetsCache.Reset();
 	PrimPathsToAssets.Reset();
 	ObjectsToWatch.Reset();
+	SubLayerLevelSequencesByIdentifier.Reset();
 }
 
 void AUsdStageActor::OpenUsdStage()
@@ -565,37 +596,46 @@ void AUsdStageActor::ReloadAnimations()
 #endif // #if USE_USD_SDK
 }
 
-void AUsdStageActor::PostRegisterAllComponents()
+void AUsdStageActor::PostTransacted(const FTransactionObjectEvent& TransactionEvent)
 {
-	Super::PostRegisterAllComponents();
-
-#if USE_USD_SDK
-	if ( HasAutorithyOverStage() )
-	{
-		LoadUsdStage();
-	}
-	else
+	if (TransactionEvent.HasPendingKillChange())
 	{
-		const pxr::UsdStageRefPtr& UsdStage = GetUsdStage();
-
-		if ( UsdStage )
+		// Fires when being deleted in editor, redo delete
+		if (this->IsPendingKill())
 		{
-			TSharedRef< FUsdSchemaTranslationContext > TranslationContext = FUsdStageActorImpl::CreateUsdSchemaTranslationContext( this, RootUsdTwin.PrimPath );
+			OnActorDestroyed.Broadcast();
+		}
+		// This fires when being spawned, undo delete, redo spawn
+		else
+		{
+			OnActorLoaded.Broadcast( this );
+
+			if (HasAutorithyOverStage())
+			{
+				LoadUsdStage();
+			}
 
-			LoadAssets( *TranslationContext, UsdStage->GetPseudoRoot() );
-			UpdatePrim( UsdStage->GetPseudoRoot().GetPrimPath(), true, *TranslationContext );
+			OnStageChanged.Broadcast();
 		}
 	}
+	// Fires on Close and undo Close
+	else if (TransactionEvent.GetEventType() == ETransactionObjectEventType::UndoRedo)
+	{
+		// UsdStageStore can't be a UPROPERTY, so we have to make sure that it
+		// is kept in sync with the state of RootLayer, because LoadUsdStage will
+		// do the job of clearing our components and etc if the path is empty
+#if USE_USD_SDK
+		if (RootLayer.FilePath.IsEmpty())
+		{
+			UnrealUSDWrapper::GetUsdStageCache().Erase( UsdStageStore.Get() );
+			UsdStageStore = TUsdStore< pxr::UsdStageRefPtr >();
+		}
 #endif // #if USE_USD_SDK
-}
 
-void AUsdStageActor::PostLoad()
-{
-	Super::PostLoad();
-
-	if ( HasAutorithyOverStage() )
-	{
-		OnActorLoaded.Broadcast( this );
+		if ( HasAutorithyOverStage() )
+		{
+			LoadUsdStage();
+		}
 	}
 }
 
diff --git a/Engine/Plugins/Importers/USDImporter/Source/USDStage/Public/USDStageActor.h b/Engine/Plugins/Importers/USDImporter/Source/USDStage/Public/USDStageActor.h
--- a/Engine/Plugins/Importers/USDImporter/Source/USDStage/Public/USDStageActor.h
+++ b/Engine/Plugins/Importers/USDImporter/Source/USDStage/Public/USDStageActor.h
@@ -100,8 +100,9 @@ public:
 	DECLARE_EVENT_OneParam( AUsdStageActor, FOnActorLoaded, AUsdStageActor* );
 	USDSTAGE_API static FOnActorLoaded OnActorLoaded;
 
-	DECLARE_EVENT( AUsdStageActor, FOnStageChanged );
-	FOnStageChanged OnStageChanged;
+	DECLARE_EVENT( AUsdStageActor, FOnStageActorEvent );
+	FOnStageActorEvent OnStageChanged;
+	FOnStageActorEvent OnActorDestroyed;
 
 	DECLARE_EVENT_TwoParams( AUsdStageActor, FOnPrimChanged, const FString&, bool );
 	FOnPrimChanged OnPrimChanged;
@@ -113,13 +114,13 @@ public:
 	AUsdStageActor();
 	virtual ~AUsdStageActor();
 
+	USDSTAGE_API void Reset() override;
 	void Refresh() const;
 	void ReloadAnimations();
 
 public:
 	virtual void PostEditChangeProperty( FPropertyChangedEvent& PropertyChangedEvent ) override;
-	virtual void PostRegisterAllComponents() override;
-	virtual void PostLoad() override;
+	virtual void PostTransacted(const FTransactionObjectEvent& TransactionEvent) override;
 
 private:
 	void Clear();
diff --git a/Engine/Plugins/Importers/USDImporter/Source/USDStageEditor/Private/SUSDStage.cpp b/Engine/Plugins/Importers/USDImporter/Source/USDStageEditor/Private/SUSDStage.cpp
--- a/Engine/Plugins/Importers/USDImporter/Source/USDStageEditor/Private/SUSDStage.cpp
+++ b/Engine/Plugins/Importers/USDImporter/Source/USDStageEditor/Private/SUSDStage.cpp
@@ -20,6 +20,7 @@
 #include "EditorStyleSet.h"
 #include "Engine/World.h"
 #include "Framework/MultiBox/MultiBoxBuilder.h"
+#include "LevelEditor.h"
 #include "Modules/ModuleManager.h"
 #include "ScopedTransaction.h"
 #include "UObject/GCObjectScopeGuard.h"
@@ -60,6 +61,9 @@ void SUsdStage::Construct( const FArguments& InArgs )
 	IUsdStageModule& UsdStageModule = FModuleManager::Get().LoadModuleChecked< IUsdStageModule >( "UsdStage" );
 	UsdStageActor = &UsdStageModule.GetUsdStageActor( GWorld );
 
+	FLevelEditorModule& LevelEditorModule = FModuleManager::LoadModuleChecked<FLevelEditorModule>("LevelEditor");
+	OnMapChangedHandle = LevelEditorModule.OnMapChanged().AddSP(this, &SUsdStage::OnMapChanged);
+
 	TUsdStore< pxr::UsdStageRefPtr > UsdStage;
 
 	if ( UsdStageActor.IsValid() )
@@ -149,6 +153,8 @@ void SUsdStage::SetupStageActorDelegates()
 {
 	if ( UsdStageActor.IsValid() )
 	{
+		ClearStageActorDelegates();
+
 		OnPrimChangedHandle = UsdStageActor->OnPrimChanged.AddLambda(
 			[ this ]( const FString& PrimPath, bool bResync )
 			{
@@ -164,37 +170,28 @@ void SUsdStage::SetupStageActorDelegates()
 			}
 		);
 
-		if ( !OnStageChangedHandle.IsValid() )
-		{
-			OnStageChangedHandle = UsdStageActor->OnStageChanged.AddLambda(
-				[ this ]()
+		OnStageChangedHandle = UsdStageActor->OnStageChanged.AddLambda(
+			[ this ]()
+			{
+				if ( UsdStageActor.IsValid() )
 				{
-					if ( UsdStageActor.IsValid() )
+					if ( this->UsdPrimInfoWidget )
 					{
-						if ( this->UsdStageInfoWidget )
-						{
-							this->UsdStageInfoWidget->RefreshStageInfos( UsdStageActor.Get() );
-						}
-
-						if ( this->UsdStageTreeView )
-						{
-							this->UsdStageTreeView->Refresh( UsdStageActor.Get() );
-							this->UsdStageTreeView->RequestTreeRefresh();
-						}
-
-						if ( this->UsdPrimInfoWidget )
-						{
-							this->UsdPrimInfoWidget->SetPrimPath( UsdStageActor->GetUsdStage(), TEXT("/") );
-						}
-
-						if ( this->UsdLayersTreeView )
-						{
-							this->UsdLayersTreeView->Refresh( UsdStageActor.Get(), true );
-						}
+						this->UsdPrimInfoWidget->SetPrimPath( UsdStageActor->GetUsdStage(), TEXT("/") );
 					}
 				}
-			);
-		}
+
+				this->Refresh();
+			}
+		);
+
+		OnActorDestroyedHandle = UsdStageActor->OnActorDestroyed.AddLambda(
+			[ this ]()
+			{
+				ClearStageActorDelegates();
+				this->CloseStage();
+			}
+		);
 
 		OnStageEditTargetChangedHandle = UsdStageActor->GetUsdListener().OnStageEditTargetChanged.AddLambda(
 			[ this ]()
@@ -208,20 +205,29 @@ void SUsdStage::SetupStageActorDelegates()
 	}
 }
 
-SUsdStage::~SUsdStage()
+void SUsdStage::ClearStageActorDelegates()
 {
-	FCoreUObjectDelegates::OnObjectPropertyChanged.Remove( OnStageActorPropertyChangedHandle );
-	AUsdStageActor::OnActorLoaded.Remove( OnActorLoadedHandle );
-
 	if ( UsdStageActor.IsValid() )
 	{
 		UsdStageActor->OnStageChanged.Remove( OnStageChangedHandle );
 		UsdStageActor->OnPrimChanged.Remove( OnPrimChangedHandle );
+		UsdStageActor->OnActorDestroyed.Remove ( OnActorDestroyedHandle );
 
 		UsdStageActor->GetUsdListener().OnStageEditTargetChanged.Remove( OnStageEditTargetChangedHandle );
 	}
 }
 
+SUsdStage::~SUsdStage()
+{
+	FCoreUObjectDelegates::OnObjectPropertyChanged.Remove( OnStageActorPropertyChangedHandle );
+	AUsdStageActor::OnActorLoaded.Remove( OnActorLoadedHandle );
+
+	FLevelEditorModule& LevelEditorModule = FModuleManager::LoadModuleChecked<FLevelEditorModule>("LevelEditor");
+	LevelEditorModule.OnMapChanged().Remove(OnMapChangedHandle);
+
+	ClearStageActorDelegates();
+}
+
 TSharedRef< SWidget > SUsdStage::MakeMainMenu()
 {
 	FMenuBarBuilder MenuBuilder( nullptr );
@@ -302,6 +308,18 @@ void SUsdStage::FillFileMenu( FMenuBuilder& MenuBuilder )
 			NAME_None,
 			EUserInterfaceActionType::Button
 		);
+
+		MenuBuilder.AddMenuEntry(
+			LOCTEXT("Close", "Close"),
+			LOCTEXT("Close_ToolTip", "Closes the opened stage"),
+			FSlateIcon(),
+			FUIAction(
+				FExecuteAction::CreateSP( this, &SUsdStage::OnClose ),
+				FCanExecuteAction()
+			),
+			NAME_None,
+			EUserInterfaceActionType::Button
+		);
 	}
 	MenuBuilder.EndSection();
 }
@@ -566,6 +584,11 @@ void SUsdStage::OnOpen()
 
 	if ( UsdFilePath )
 	{
+		FScopedTransaction Transaction(FText::Format(
+			LOCTEXT("OpenStageTransaction", "Open USD stage '{0}'"),
+			FText::FromString(UsdFilePath.GetValue())
+		));
+
 		UsdUtils::StartMonitoringErrors();
 
 		OpenStage( *UsdFilePath.GetValue() );
@@ -602,7 +625,12 @@ void SUsdStage::OnReloadStage()
 		if ( UsdStage )
 		{
 			UsdUtils::StartMonitoringErrors();
-			UsdStage->Reload();
+			{
+				FScopedUsdAllocs Allocs;
+				const std::vector<pxr::SdfLayerHandle>& HandleVec = UsdStage->GetUsedLayers();
+				bool bForce = true;
+				pxr::SdfLayer::ReloadLayers({HandleVec.begin(), HandleVec.end()}, bForce);
+			}
 			if (UsdUtils::ShowErrorsAndStopMonitoring())
 			{
 				return;
@@ -624,6 +652,13 @@ void SUsdStage::OnReloadStage()
 	}
 }
 
+void SUsdStage::OnClose()
+{
+	FScopedTransaction Transaction(LOCTEXT("CloseTransaction", "Close USD stage"));
+
+	CloseStage();
+}
+
 void SUsdStage::OnImport()
 {
 	AUsdStageActor* StageActor = UsdStageActor.Get();
@@ -729,12 +764,45 @@ void SUsdStage::OpenStage( const TCHAR* FilePath )
 	}
 
 	check( UsdStageActor.IsValid() );
-	UsdStageActor->RootLayer.FilePath = FilePath;
+	UsdStageActor->Modify();
 
+	UsdStageActor->RootLayer.FilePath = FilePath;
 	FPropertyChangedEvent RootLayerPropertyChangedEvent( FindFieldChecked< FProperty >( UsdStageActor->GetClass(), FName("RootLayer") ) );
 	UsdStageActor->PostEditChangeProperty( RootLayerPropertyChangedEvent );
 }
 
+void SUsdStage::CloseStage()
+{
+	if ( AUsdStageActor* StageActor = UsdStageActor.Get() )
+	{
+		StageActor->Reset();
+	}
+
+	Refresh();
+}
+
+void SUsdStage::Refresh()
+{
+	// May be nullptr, but that is ok. Its how the widgets are reset
+	AUsdStageActor* StageActor = UsdStageActor.Get();
+
+	if (UsdLayersTreeView)
+	{
+		UsdLayersTreeView->Refresh( StageActor, true );
+	}
+
+	if (UsdStageInfoWidget)
+	{
+		UsdStageInfoWidget->RefreshStageInfos( StageActor );
+	}
+
+	if (UsdStageTreeView)
+	{
+		UsdStageTreeView->Refresh( StageActor );
+		UsdStageTreeView->RequestTreeRefresh();
+	}
+}
+
 void SUsdStage::OnStageActorLoaded( AUsdStageActor* InUsdStageActor )
 {
 	if ( UsdStageActor == InUsdStageActor )
@@ -742,8 +810,13 @@ void SUsdStage::OnStageActorLoaded( AUsdStageActor* InUsdStageActor )
 		return;
 	}
 
+	ClearStageActorDelegates();
 	UsdStageActor = InUsdStageActor;
 	SetupStageActorDelegates();
+
+	// Refresh here because we may be receiving an actor that has a stage already loaded,
+	// like during undo/redo
+	Refresh();
 }
 
 void SUsdStage::OnStageActorPropertyChanged( UObject* ObjectBeingModified, FPropertyChangedEvent& PropertyChangedEvent )
@@ -757,6 +830,11 @@ void SUsdStage::OnStageActorPropertyChanged( UObject* ObjectBeingModified, FProp
 	}
 }
 
+void SUsdStage::OnMapChanged(UWorld* World, EMapChangeType ChangeType)
+{
+	CloseStage();
+}
+
 #endif // #if USE_USD_SDK
 
 #undef LOCTEXT_NAMESPACE
diff --git a/Engine/Plugins/Importers/USDImporter/Source/USDStageEditor/Private/SUSDStage.h b/Engine/Plugins/Importers/USDImporter/Source/USDStageEditor/Private/SUSDStage.h
--- a/Engine/Plugins/Importers/USDImporter/Source/USDStageEditor/Private/SUSDStage.h
+++ b/Engine/Plugins/Importers/USDImporter/Source/USDStageEditor/Private/SUSDStage.h
@@ -12,9 +12,10 @@
 #include "Widgets/SCompoundWidget.h"
 
 class AUsdStageActor;
-enum class EUsdInitialLoadSet;
 class FLevelCollectionModel;
 class FMenuBuilder;
+enum class EMapChangeType : uint8;
+enum class EUsdInitialLoadSet;
 struct FSlateBrush;
 
 #if USE_USD_SDK
@@ -30,6 +31,7 @@ class SUsdStage : public SCompoundWidget
 
 protected:
 	void SetupStageActorDelegates();
+	void ClearStageActorDelegates();
 
 	TSharedRef< SWidget > MakeMainMenu();
 	void FillFileMenu( FMenuBuilder& MenuBuilder );
@@ -43,16 +45,22 @@ protected:
 	void OnOpen();
 	void OnSave();
 	void OnReloadStage();
+	void OnClose();
 
 	void OnImport();
 
 	void OnPrimSelected( FString PrimPath );
 
 	void OpenStage( const TCHAR* FilePath );
+	void CloseStage();
+
+	void Refresh();
 
 	void OnStageActorLoaded( AUsdStageActor* InUsdStageActor );
 	void OnStageActorPropertyChanged( UObject* ObjectBeingModified, FPropertyChangedEvent& PropertyChangedEvent );
 
+	void OnMapChanged( UWorld* World, EMapChangeType ChangeType );
+
 protected:
 	TSharedPtr< class SUsdStageInfo > UsdStageInfoWidget;
 	TSharedPtr< class SUsdStageTreeView > UsdStageTreeView;
@@ -62,10 +70,12 @@ protected:
 	TWeakObjectPtr< AUsdStageActor > UsdStageActor;
 
 	FDelegateHandle OnActorLoadedHandle;
+	FDelegateHandle OnActorDestroyedHandle;
 	FDelegateHandle OnStageActorPropertyChangedHandle;
 	FDelegateHandle OnStageChangedHandle;
 	FDelegateHandle OnStageEditTargetChangedHandle;
 	FDelegateHandle OnPrimChangedHandle;
+	FDelegateHandle OnMapChangedHandle;
 
 	FString SelectedPrimPath;
 };
diff --git a/Engine/Plugins/Importers/USDImporter/Source/USDStageEditor/Private/SUSDStageTreeView.cpp b/Engine/Plugins/Importers/USDImporter/Source/USDStageEditor/Private/SUSDStageTreeView.cpp
--- a/Engine/Plugins/Importers/USDImporter/Source/USDStageEditor/Private/SUSDStageTreeView.cpp
+++ b/Engine/Plugins/Importers/USDImporter/Source/USDStageEditor/Private/SUSDStageTreeView.cpp
@@ -123,7 +123,7 @@ public:
 					bNeedsRefresh = true;
 					break;
 				}
-				
+
 				++ChildIndex;
 			}
 		}
@@ -147,19 +147,11 @@ public:
 
 	void RefreshData( bool bRefreshChildren )
 	{
-		FString PrimName;
-
 		if ( !UsdPrim.Get() )
 		{
-			if ( ParentItem )
-			{
-				ParentItem->RefreshData( true );
-			}
-
-			RowData.Get() = FUsdStageTreeRowData();
 			return;
 		}
-		
+
 		RowData->Name = FText::FromString( UsdToUnreal::ConvertString( UsdPrim.Get().GetName() ) );
 		RowData->bHasCompositionArcs = UsdUtils::HasCompositionArcs( UsdPrim.Get() );
 
@@ -193,7 +185,7 @@ public:
 	FUsdStageTreeItem* ParentItem;
 
 	TArray< FUsdStageTreeItemRef > Children;
-	
+
 	TSharedRef< FUsdStageTreeRowData > RowData; // Data model
 };
 
@@ -257,7 +249,7 @@ protected:
 				TextColor = FLinearColor( FColor::Orange );
 			}
 		}
-		
+
 		return TextColor;
 	}
 };
@@ -308,14 +300,14 @@ public:
 		}
 
 		TSharedRef< SWidget > Item = SNullWidget::NullWidget;
-		
+
 		if ( bHasPayload )
 		{
 			Item = SNew( SCheckBox )
 					.IsChecked( this, &FUsdStagePayloadColumn::IsChecked, StaticCastSharedPtr< FUsdStageTreeItem >( InTreeItem ) )
 					.OnCheckStateChanged( this, &FUsdStagePayloadColumn::OnCheckedPayload, StaticCastSharedPtr< FUsdStageTreeItem >( InTreeItem ) );
 		}
-		
+
 		return Item;
 	}
 };
@@ -379,7 +371,7 @@ public:
 											SNew( SImage )
 											.Image( this, &FUsdStageVisibilityColumn::GetBrush, TreeItem )
 										];
-		
+
 		return Item;
 	}
 };
@@ -396,7 +388,7 @@ public:
 				.TextStyle( FEditorStyle::Get(), "LargeText" )
 				.Text( TreeItem->RowData, &FUsdStageTreeRowData::GetType )
 				.Font( FEditorStyle::GetFontStyle( "ContentBrowser.SourceTreeItemFont" ) );
-		
+
 		return Item;
 	}
 };
@@ -488,46 +480,52 @@ void SUsdStageTreeView::RefreshPrim( const FString& PrimPath, bool bResync )
 
 		return {};
 	};
-	
-	bool bNeedsFullRefresh = false;
 
+	// Search for the corresponding tree item to update
+	FUsdStageTreeItemPtr FoundItem = nullptr;
 	for ( FUsdStageTreeItemRef RootItem : this->RootItems )
 	{
 		TUsdStore< pxr::SdfPath > PrimPathToSearch = UsdPrimPath;
 
-		bool bFoundPrimToRefresh = false;
+		FoundItem = FindTreeItemFromPrimPath( PrimPathToSearch, RootItem );
 
-		while ( !bFoundPrimToRefresh )
+		while ( !FoundItem.IsValid() )
 		{
-			if ( FUsdStageTreeItemPtr FoundItem = FindTreeItemFromPrimPath( PrimPathToSearch, RootItem ) )
+			TUsdStore< pxr::SdfPath > ParentPrimPath = PrimPathToSearch.Get().GetParentPath();
+			if ( ParentPrimPath.Get() == PrimPathToSearch.Get() )
 			{
-				FoundItem->RefreshData( true );
-				bFoundPrimToRefresh = true;
-
 				break;
 			}
-			else
-			{
-				TUsdStore< pxr::SdfPath > ParentPrimPath = PrimPathToSearch.Get().GetParentPath();
+			PrimPathToSearch = MoveTemp( ParentPrimPath );
 
-				if ( ParentPrimPath.Get() == PrimPathToSearch.Get() )
-				{
-					break;
-				}
-				else
-				{
-					PrimPathToSearch = MoveTemp( ParentPrimPath );
-				}
-			}
+			FoundItem = FindTreeItemFromPrimPath( PrimPathToSearch, RootItem );
 		}
 
-		if ( !bFoundPrimToRefresh )
+		if (FoundItem.IsValid())
 		{
-			bNeedsFullRefresh = true; // Prim not found, refresh the whole tree
+			break;
 		}
 	}
 
-	if ( bNeedsFullRefresh )
+	if (FoundItem.IsValid())
+	{
+		FoundItem->RefreshData( true );
+
+		// Item doesn't match any prim, needs to be removed
+		if (!FoundItem->UsdPrim.Get())
+		{
+			if (FoundItem->ParentItem)
+			{
+				FoundItem->ParentItem->RefreshData( true );
+			}
+			else
+			{
+				RootItems.Remove(FoundItem.ToSharedRef());
+			}
+		}
+	}
+	// We couldn't find the target prim, do a full refresh instead
+	else
 	{
 		Refresh( UsdStageActor.Get() );
 	}
@@ -546,7 +544,7 @@ void SUsdStageTreeView::SetupColumns()
 	VisibilityColumnArguments.FixedWidth( 24.f );
 
 	AddColumn( TEXT( "Visibility" ), FText(), MakeShared< FUsdStageVisibilityColumn >(), VisibilityColumnArguments );
-	
+
 	{
 		TSharedRef< FUsdStageNameColumn > PrimNameColumn = MakeShared< FUsdStageNameColumn >();
 		PrimNameColumn->OwnerTree = SharedThis( this );
@@ -558,7 +556,7 @@ void SUsdStageTreeView::SetupColumns()
 
 		AddColumn( TEXT("Prim"), LOCTEXT( "Prim", "Prim" ), PrimNameColumn, PrimNameColumnArguments );
 	}
-	
+
 	SHeaderRow::FColumn::FArguments TypeColumnArguments;
 	TypeColumnArguments.FillWidth( 15.f );
 	AddColumn( TEXT("Type"), LOCTEXT( "Type", "Type" ), MakeShared< FUsdStagePrimTypeColumn >(), TypeColumnArguments );
@@ -581,7 +579,7 @@ TSharedPtr< SWidget > SUsdStageTreeView::ConstructPrimContextMenu()
 			FSlateIcon(),
 			FUIAction(
 				FExecuteAction::CreateSP( this, &SUsdStageTreeView::OnAddPrim ),
-				FCanExecuteAction::CreateSP( this, &SUsdStageTreeView::CanExecutePrimAction )
+				FCanExecuteAction::CreateSP( this, &SUsdStageTreeView::CanAddPrim )
 			),
 			NAME_None,
 			EUserInterfaceActionType::Button
@@ -683,10 +681,23 @@ void SUsdStageTreeView::OnAddPrim()
 
 	TArray< FUsdStageTreeItemRef > MySelectedItems = GetSelectedItems();
 
-	for ( FUsdStageTreeItemRef SelectedItem : MySelectedItems )
+	// Add a new child prim
+	if (MySelectedItems.Num() > 0)
+	{
+		for ( FUsdStageTreeItemRef SelectedItem : MySelectedItems )
+		{
+			FUsdStageTreeItemRef TreeItem = MakeShared< FUsdStageTreeItem >( &SelectedItem.Get(), SelectedItem->UsdStage );
+			SelectedItem->Children.Add( TreeItem );
+
+			PendingRenameItem = TreeItem;
+			ScrollItemIntoView( TreeItem );
+		}
+	}
+	// Add a new top-level prim (direct child of the pseudo-root prim)
+	else
 	{
-		FUsdStageTreeItemRef TreeItem = MakeShared< FUsdStageTreeItem >( &SelectedItem.Get(), SelectedItem->UsdStage );
-		SelectedItem->Children.Add( TreeItem );
+		FUsdStageTreeItemRef TreeItem = MakeShared< FUsdStageTreeItem >( nullptr, MakeUsdStore<pxr::UsdStageRefPtr>(UsdStageActor->GetUsdStage()) );
+		RootItems.Add( TreeItem );
 
 		PendingRenameItem = TreeItem;
 		ScrollItemIntoView( TreeItem );
@@ -757,6 +768,23 @@ void SUsdStageTreeView::OnClearReferences()
 	}
 }
 
+bool SUsdStageTreeView::CanAddPrim() const
+{
+	if ( !UsdStageActor.IsValid() )
+	{
+		return false;
+	}
+
+	TUsdStore< pxr::UsdStageRefPtr > UsdStage =  UsdStageActor->GetUsdStage();
+
+	if ( !UsdStage.Get() )
+	{
+		return false;
+	}
+
+	return true;
+}
+
 bool SUsdStageTreeView::CanExecutePrimAction() const
 {
 	if ( !UsdStageActor.IsValid() )
@@ -870,7 +898,7 @@ void SUsdStageTreeView::OnPrimNameCommitted( const FUsdStageTreeItemRef& TreeIte
 		FScopedUsdAllocs UsdAllocs;
 
 		pxr::SdfPath ParentPrimPath;
-		
+
 		if ( TreeItem->ParentItem )
 		{
 			ParentPrimPath = TreeItem->ParentItem->UsdPrim.Get().GetPrimPath();
@@ -881,15 +909,21 @@ void SUsdStageTreeView::OnPrimNameCommitted( const FUsdStageTreeItemRef& TreeIte
 		}
 
 		pxr::SdfPath NewPrimPath = ParentPrimPath.AppendChild( pxr::TfToken( UnrealToUsd::ConvertString( *InPrimName.ToString() ).Get() ) );
-		
+
 		TreeItem->UsdPrim = pxr::UsdGeomXform::Define( TreeItem->UsdStage.Get(), NewPrimPath ).GetPrim();
 
+		// Renamed a child item
 		if ( TreeItem->ParentItem )
 		{
 			TreeItem->ParentItem->Children.Remove( TreeItem );
 
 			RefreshPrim( UsdToUnreal::ConvertPath( TreeItem->ParentItem->UsdPrim.Get().GetPrimPath() ), true );
 		}
+		// Renamed a root item
+		else
+		{
+			RefreshPrim( UsdToUnreal::ConvertPath( TreeItem->UsdPrim.Get().GetPrimPath() ), true );
+		}
 	}
 }
 
diff --git a/Engine/Plugins/Importers/USDImporter/Source/USDStageEditor/Private/SUSDStageTreeView.h b/Engine/Plugins/Importers/USDImporter/Source/USDStageEditor/Private/SUSDStageTreeView.h
--- a/Engine/Plugins/Importers/USDImporter/Source/USDStageEditor/Private/SUSDStageTreeView.h
+++ b/Engine/Plugins/Importers/USDImporter/Source/USDStageEditor/Private/SUSDStageTreeView.h
@@ -46,7 +46,7 @@ public:
 private:
 	virtual TSharedRef< ITableRow > OnGenerateRow( FUsdStageTreeItemRef InDisplayNode, const TSharedRef< STableViewBase >& OwnerTable ) override;
 	virtual void OnGetChildren( FUsdStageTreeItemRef InParent, TArray< FUsdStageTreeItemRef >& OutChildren ) const override;
-	
+
 	void ScrollItemIntoView( FUsdStageTreeItemRef TreeItem );
 	virtual void OnTreeItemScrolledIntoView( FUsdStageTreeItemRef TreeItem, const TSharedPtr< ITableRow >& Widget ) override ;
 
@@ -62,6 +62,7 @@ private:
 	void OnAddReference();
 	void OnClearReferences();
 
+	bool CanAddPrim() const;
 	bool CanExecutePrimAction() const;
 
 	TOptional< FString > BrowseFile();
diff --git a/Engine/Plugins/Importers/USDImporter/Source/USDUtilities/Private/USDGeomMeshConversion.cpp b/Engine/Plugins/Importers/USDImporter/Source/USDUtilities/Private/USDGeomMeshConversion.cpp
--- a/Engine/Plugins/Importers/USDImporter/Source/USDUtilities/Private/USDGeomMeshConversion.cpp
+++ b/Engine/Plugins/Importers/USDImporter/Source/USDUtilities/Private/USDGeomMeshConversion.cpp
@@ -23,6 +23,10 @@
 
 #include "USDIncludesStart.h"
 
+#include "pxr/usd/ar/resolver.h"
+#include "pxr/usd/ar/resolverScopedCache.h"
+#include "pxr/usd/sdf/layer.h"
+#include "pxr/usd/sdf/layerUtils.h"
 #include "pxr/usd/usdGeom/mesh.h"
 #include "pxr/usd/usdGeom/subset.h"
 #include "pxr/usd/usdShade/material.h"
@@ -384,7 +388,44 @@ bool UsdToUnreal::ConvertGeomMesh( const pxr::UsdGeomMesh& UsdMesh, FMeshDescrip
 
 namespace UsdGeomMeshConversionImpl
 {
-	UMaterialExpression* ParseInputTexture( pxr::UsdShadeInput& ShadeInput, UMaterial& Material )
+	// Finds the strongest layer contributing to an attribute
+	TUsdStore< pxr::SdfLayerHandle > FindLayerHandle( const pxr::UsdAttribute& Attr, const pxr::UsdTimeCode& Time )
+	{
+		FScopedUsdAllocs UsdAllocs;
+
+		for ( const pxr::SdfPropertySpecHandle& PropertySpec : Attr.GetPropertyStack( Time ) )
+		{
+			if ( PropertySpec->HasDefaultValue() || PropertySpec->GetLayer()->GetNumTimeSamplesForPath( PropertySpec->GetPath() ) > 0 )
+			{
+				return PropertySpec->GetLayer();
+			}
+		}
+
+		return {};
+	}
+
+	// Given an AssetPath, resolve it to an actual file path
+	FString ResolveAssetPath( const TUsdStore< pxr::SdfLayerHandle >& LayerHandle, const FString& AssetPath )
+	{
+		FString AssetPathToResolve = AssetPath;
+		AssetPathToResolve.ReplaceInline( TEXT("<UDIM>"), TEXT("1001") ); // If it's a UDIM path, we replace the UDIM tag with the start tile
+
+		pxr::ArResolverScopedCache ResolverCache;
+		pxr::ArResolver& Resolver = pxr::ArGetResolver();
+
+		const FString RelativePathToResolve = 
+			LayerHandle.Get()
+			? UsdToUnreal::ConvertString( pxr::SdfComputeAssetPathRelativeToLayer( LayerHandle.Get(), UnrealToUsd::ConvertString( *AssetPathToResolve ).Get() ) )
+			: AssetPathToResolve;
+
+		TUsdStore< std::string > ResolvedPathUsd = Resolver.Resolve( UnrealToUsd::ConvertString( *RelativePathToResolve ).Get().c_str() );
+
+		FString ResolvedAssetPath = UsdToUnreal::ConvertString( ResolvedPathUsd.Get() );
+
+		return ResolvedAssetPath;
+	}
+
+	UMaterialExpression* ParseInputTexture( pxr::UsdShadeInput& ShadeInput, UMaterial& Material, TMap< FString, UObject* >& TexturesCache )
 	{
 		UMaterialExpression* Result = nullptr;
 
@@ -402,21 +443,50 @@ namespace UsdGeomMeshConversionImpl
 
 				if ( FileInput )
 				{
-					pxr::SdfAssetPath TexturePath;
-					FileInput.Get< pxr::SdfAssetPath >( &TexturePath );
+					pxr::SdfAssetPath TextureAssetPath;
+					FileInput.Get< pxr::SdfAssetPath >( &TextureAssetPath );
 
 					FScopedUnrealAllocs UnrealAllocs;
 
-					FString TextureFileName = UsdToUnreal::ConvertString( TexturePath.GetResolvedPath() );
-					TextureFileName.RemoveFromStart( TEXT("\\\\?\\") );
+					FString TexturePath = UsdToUnreal::ConvertString( TextureAssetPath.GetAssetPath() );
+					const bool bIsSupportedUdimTexture = TexturePath.Contains( TEXT("<UDIM>") );
+					FPaths::NormalizeFilename( TexturePath );
+
+					FString ResolvedTexturePath = UsdToUnreal::ConvertString( TextureAssetPath.GetResolvedPath() );
+					ResolvedTexturePath.RemoveFromStart( TEXT("\\\\?\\") );
+					FPaths::NormalizeFilename( ResolvedTexturePath );
+
+					UTexture2D* Texture2D = Cast< UTexture2D >( TexturesCache.FindRef( TexturePath ) );
+
+					if ( !Texture2D && !TexturePath.IsEmpty() )
+					{
+						bool bOutCancelled = false;
+						UTextureFactory* TextureFactory = NewObject< UTextureFactory >();
+						TextureFactory->SuppressImportOverwriteDialog();
+						TextureFactory->bUseHashAsGuid = true;
+
+						if ( bIsSupportedUdimTexture )
+						{
+							FString BaseFileName = FPaths::GetBaseFilename( TexturePath );
+
+							FString BaseFileNameBeforeUdim;
+							FString BaseFileNameAfterUdim;
+							BaseFileName.Split( TEXT("<UDIM>"), &BaseFileNameBeforeUdim, &BaseFileNameAfterUdim );
 
-					FPaths::NormalizeFilename( TextureFileName );
+							FString UdimRegexPattern = FString::Printf( TEXT(R"((%s)(\d{4})(%s))"), *BaseFileNameBeforeUdim, *BaseFileNameAfterUdim );
+							TextureFactory->UdimRegexPattern = MoveTemp( UdimRegexPattern );
+						}
+
+						if ( ResolvedTexturePath.IsEmpty() )
+						{
+							TUsdStore< pxr::SdfLayerHandle > LayerHandle = FindLayerHandle( FileInput.GetAttr(), pxr::UsdTimeCode::EarliestTime() );
+							ResolvedTexturePath = ResolveAssetPath( LayerHandle, TexturePath );
+						}
 
-					bool bOutCancelled = false;
-					UTextureFactory* TextureFactory = NewObject< UTextureFactory >();
-					TextureFactory->SuppressImportOverwriteDialog();
+						Texture2D = Cast< UTexture2D >( TextureFactory->ImportObject( UTexture2D::StaticClass(), GetTransientPackage(), NAME_None, RF_Transient, ResolvedTexturePath, TEXT(""), bOutCancelled ) );
 
-					UTexture2D* Texture2D = Cast< UTexture2D >( TextureFactory->FactoryCreateFile( UTexture2D::StaticClass(), GetTransientPackage(), NAME_None, RF_Transient, TextureFileName, TEXT(""), GWarn, bOutCancelled ) );
+						TexturesCache.Add( TexturePath, Texture2D );
+					}
 
 					if ( Texture2D )
 					{
@@ -435,6 +505,12 @@ namespace UsdGeomMeshConversionImpl
 }
 
 bool UsdToUnreal::ConvertMaterial( const pxr::UsdShadeMaterial& UsdShadeMaterial, UMaterial& Material )
+{
+	TMap< FString, UObject* > TexturesCache;
+	return ConvertMaterial( UsdShadeMaterial, Material, TexturesCache );
+}
+
+bool UsdToUnreal::ConvertMaterial( const pxr::UsdShadeMaterial& UsdShadeMaterial, UMaterial& Material, TMap< FString, UObject* >& TexturesCache )
 {
 	pxr::UsdShadeShader SurfaceShader = UsdShadeMaterial.ComputeSurfaceSource();
 	
@@ -445,7 +521,7 @@ bool UsdToUnreal::ConvertMaterial( const pxr::UsdShadeMaterial& UsdShadeMaterial
 
 	bool bHasMaterialInfo = false;
 
-	auto ParseFloatInput = [ &Material, &SurfaceShader ]( const ANSICHAR* InputName, float DefaultValue ) -> UMaterialExpression*
+	auto ParseFloatInput = [ &Material, &SurfaceShader, &TexturesCache ]( const ANSICHAR* InputName, float DefaultValue ) -> UMaterialExpression*
 	{
 		pxr::UsdShadeInput Input = SurfaceShader.GetInput( pxr::TfToken(InputName) );
 		UMaterialExpression* InputExpression = nullptr;
@@ -467,7 +543,7 @@ bool UsdToUnreal::ConvertMaterial( const pxr::UsdShadeMaterial& UsdShadeMaterial
 			}
 			else
 			{
-				InputExpression = UsdGeomMeshConversionImpl::ParseInputTexture( Input, Material );
+				InputExpression = UsdGeomMeshConversionImpl::ParseInputTexture( Input, Material, TexturesCache );
 			}
 		}
 
@@ -496,7 +572,7 @@ bool UsdToUnreal::ConvertMaterial( const pxr::UsdShadeMaterial& UsdShadeMaterial
 			}
 			else
 			{
-				BaseColorExpression = UsdGeomMeshConversionImpl::ParseInputTexture( DiffuseInput, Material );
+				BaseColorExpression = UsdGeomMeshConversionImpl::ParseInputTexture( DiffuseInput, Material, TexturesCache );
 			}
 
 			if ( BaseColorExpression )
diff --git a/Engine/Plugins/Importers/USDImporter/Source/USDUtilities/Private/USDPrimConversion.cpp b/Engine/Plugins/Importers/USDImporter/Source/USDUtilities/Private/USDPrimConversion.cpp
--- a/Engine/Plugins/Importers/USDImporter/Source/USDUtilities/Private/USDPrimConversion.cpp
+++ b/Engine/Plugins/Importers/USDImporter/Source/USDUtilities/Private/USDPrimConversion.cpp
@@ -42,13 +42,11 @@ bool UsdToUnreal::ConvertXformable( const pxr::UsdStageRefPtr& Stage, const pxr:
 
 	if ( Xformable.GetPrim().IsA< pxr::UsdGeomCamera >() )
 	{
-		AdditionalRotation = FRotator( -90.f, 0.f, 0.f );
-		AdditionalRotation = AdditionalRotation + FRotator( 0.f, -90.f, 0.f );
-
+		AdditionalRotation = FRotator( 0.0f, -90.f, 0.0f );
 		UpAxis = pxr::UsdGeomTokens->y; // Cameras are always Y up in USD
 	}
 
-	OutTransform = UsdToUnreal::ConvertMatrix( UpAxis, UsdMatrix ) * FTransform( AdditionalRotation );
+	OutTransform = FTransform( AdditionalRotation ) * UsdToUnreal::ConvertMatrix( UpAxis, UsdMatrix );
 
 	return true;
 }
@@ -122,7 +120,7 @@ bool UnrealToUsd::ConvertSceneComponent( const pxr::UsdStageRefPtr& Stage, const
 
 		bResetXFormStack = false;
 		bool bFoundTransformOp = false;
-	
+
 		std::vector< pxr::UsdGeomXformOp > XFormOps = XForm.GetOrderedXformOps( &bResetXFormStack );
 		for ( const pxr::UsdGeomXformOp& XFormOp : XFormOps )
 		{
diff --git a/Engine/Plugins/Importers/USDImporter/Source/USDUtilities/Public/USDGeomMeshConversion.h b/Engine/Plugins/Importers/USDImporter/Source/USDUtilities/Public/USDGeomMeshConversion.h
--- a/Engine/Plugins/Importers/USDImporter/Source/USDUtilities/Public/USDGeomMeshConversion.h
+++ b/Engine/Plugins/Importers/USDImporter/Source/USDUtilities/Public/USDGeomMeshConversion.h
@@ -27,7 +27,8 @@ namespace UsdToUnreal
 {
 	USDUTILITIES_API bool ConvertGeomMesh( const pxr::UsdGeomMesh& UsdMesh, FMeshDescription& MeshDescription, const pxr::UsdTimeCode TimeCode = pxr::UsdTimeCode::EarliestTime() );
 	USDUTILITIES_API bool ConvertGeomMesh( const pxr::UsdGeomMesh& UsdMesh, FMeshDescription& MeshDescription, const FTransform& AdditionalTransform, const pxr::UsdTimeCode TimeCode = pxr::UsdTimeCode::EarliestTime() );
-	USDUTILITIES_API bool ConvertMaterial( const pxr::UsdShadeMaterial& UsdMesh, UMaterial& Material );
+	USDUTILITIES_API bool ConvertMaterial( const pxr::UsdShadeMaterial& UsdShadeMaterial, UMaterial& Material );
+	USDUTILITIES_API bool ConvertMaterial( const pxr::UsdShadeMaterial& UsdShadeMaterial, UMaterial& Material, TMap< FString, UObject* >& TexturesCache );
 
 	/** Reads the first display color and opacity value and assigns it as the base color and opacity of the material */
 	USDUTILITIES_API bool ConvertDisplayColor( const pxr::UsdGeomMesh& UsdMesh, UMaterialInstanceConstant& MaterialInstance, const pxr::UsdTimeCode TimeCode = pxr::UsdTimeCode::EarliestTime() );
diff --git a/Engine/Source/Editor/UnrealEd/Classes/Factories/TextureFactory.h b/Engine/Source/Editor/UnrealEd/Classes/Factories/TextureFactory.h
--- a/Engine/Source/Editor/UnrealEd/Classes/Factories/TextureFactory.h
+++ b/Engine/Source/Editor/UnrealEd/Classes/Factories/TextureFactory.h
@@ -125,6 +125,14 @@ class UNREALED_API UTextureFactory : public UFactory, public IImportSettingsPars
 	UPROPERTY(Transient)
 	uint32 bUseHashAsGuid:1;
 
+	/**
+	 * The pattern to use to match UDIM files to indices. Defaults to match a filename that ends with either .1001 or _1001
+	 * This 1st and 3rd (optional) capture groups are used as the texture name. The 2nd capture group is considered to be the UDIM index.
+	 * ie: (Capture Group 1)(\d{4})( Capture Group 3)
+	 */
+	UPROPERTY(Transient)
+	FString UdimRegexPattern;
+
 public:
 	UTextureFactory(const FObjectInitializer& ObjectInitializer = FObjectInitializer::Get());
 
diff --git a/Engine/Source/Editor/UnrealEd/Private/Factories/EditorFactories.cpp b/Engine/Source/Editor/UnrealEd/Private/Factories/EditorFactories.cpp
--- a/Engine/Source/Editor/UnrealEd/Private/Factories/EditorFactories.cpp
+++ b/Engine/Source/Editor/UnrealEd/Private/Factories/EditorFactories.cpp
@@ -23,6 +23,7 @@
 #include "Fonts/CompositeFont.h"
 #include "Misc/Attribute.h"
 #include "Input/Reply.h"
+#include "Internationalization/Regex.h"
 #include "Widgets/DeclarativeSyntaxSupport.h"
 #include "Widgets/SWidget.h"
 #include "Widgets/SBoxPanel.h"
@@ -3037,6 +3038,8 @@ UTextureFactory::UTextureFactory(const FObjectInitializer& ObjectInitializer)
 
 	bCreateNew = false;
 	bEditorImport = true;
+
+	UdimRegexPattern = TEXT(R"((.+?)[._](\d{4})$)");
 }
 
 bool UTextureFactory::FactoryCanImport(const FString& Filename)
@@ -4124,39 +4127,37 @@ bool UTextureFactory::DoesSupportClass(UClass* Class)
 	return Class == UTexture2D::StaticClass() || Class == UTextureCube::StaticClass();
 }
 
-static int32 ParseUDIMName(const FString& Name, FString& OutRootName)
+static int32 ParseUDIMName(const FString& Name, const FString& UdimRegexPattern, FString& OutRootName)
 {
-	int32 SeparatorIndex = INDEX_NONE;
-	if (!Name.FindLastChar('.', SeparatorIndex))
+	FRegexPattern RegexPattern( UdimRegexPattern );
+	FRegexMatcher RegexMatcher( RegexPattern, Name );
+
+	int32 UdimValue = INDEX_NONE;
+
+	if ( RegexMatcher.FindNext() )
 	{
-		// '.' is the standard UDIM separator, but we'll accept '_' as well
-		if (!Name.FindLastChar('_', SeparatorIndex))
+		const int32 StartOfCaptureGroup1 = RegexMatcher.GetCaptureGroupBeginning(1);
+		const int32 EndOfCaptureGroup1 = RegexMatcher.GetCaptureGroupEnding(1);
+		const int32 StartOfCaptureGroup2 = RegexMatcher.GetCaptureGroupBeginning(2);
+		const int32 EndOfCaptureGroup2 = RegexMatcher.GetCaptureGroupEnding(2);
+		const int32 StartOfCaptureGroup3 = RegexMatcher.GetCaptureGroupBeginning(3);
+		const int32 EndOfCaptureGroup3 = RegexMatcher.GetCaptureGroupEnding(3);
+
+		if ( StartOfCaptureGroup1 != INDEX_NONE && StartOfCaptureGroup2 != INDEX_NONE &&
+			 EndOfCaptureGroup1 != INDEX_NONE && EndOfCaptureGroup2 != INDEX_NONE )
 		{
-			return INDEX_NONE;
-		}
-	}
-	if (SeparatorIndex + 5 != Name.Len())
-	{
-		return INDEX_NONE;
-	}
-	const TCHAR Digit0 = Name[SeparatorIndex + 4];
-	const TCHAR Digit1 = Name[SeparatorIndex + 3];
-	const TCHAR Digit2 = Name[SeparatorIndex + 2];
-	const TCHAR Digit3 = Name[SeparatorIndex + 1];
-	if (Digit0 < '0' || Digit0 > '9') return INDEX_NONE;
-	if (Digit1 < '0' || Digit1 > '9') return INDEX_NONE;
-	if (Digit2 < '0' || Digit2 > '9') return INDEX_NONE;
-	if (Digit3 < '0' || Digit3 > '9') return INDEX_NONE;
+			LexFromString( UdimValue, *Name.Mid( StartOfCaptureGroup2, EndOfCaptureGroup2 - StartOfCaptureGroup2 ) );
 
-	const int32 Value = (int32)(Digit0 - '0') + (int32)(Digit1 - '0') * 10 + (int32)(Digit2 - '0') * 100 + (int32)(Digit3 - '0') * 1000;
-	if (Value < 1001)
-	{
-		// UDIM starts with 1001 as the origin
-		return INDEX_NONE;
-	}
+			OutRootName = Name.Mid( StartOfCaptureGroup1, EndOfCaptureGroup1 - StartOfCaptureGroup1 );
 
-	OutRootName = Name.Left(SeparatorIndex);
-	return Value;
+			if ( StartOfCaptureGroup3 != INDEX_NONE && EndOfCaptureGroup3 != INDEX_NONE )
+			{
+				OutRootName.Append( Name.Mid( StartOfCaptureGroup3, EndOfCaptureGroup3 - StartOfCaptureGroup3 ) );
+			}
+		}
+	}
+	
+	return UdimValue;
 }
 
 UObject* UTextureFactory::FactoryCreateBinary
@@ -4182,7 +4183,7 @@ UObject* UTextureFactory::FactoryCreateBinary
 	{
 		const FString FilenameNoExtension = FPaths::GetBaseFilename(CurrentFilename);
 		FString BaseUDIMName;
-		const int32 BaseUDIMIndex = ParseUDIMName(FilenameNoExtension, BaseUDIMName);
+		const int32 BaseUDIMIndex = ParseUDIMName(FilenameNoExtension, UdimRegexPattern, BaseUDIMName);
 		if (BaseUDIMIndex != INDEX_NONE)
 		{
 			UDIMIndexToFile.Add(BaseUDIMIndex, CurrentFilename);
@@ -4199,7 +4200,7 @@ UObject* UTextureFactory::FactoryCreateBinary
 				if (!CurrentFilename.EndsWith(UDIMFile) && FactoryCanImport(UDIMFile))
 				{
 					FString UDIMName;
-					const int32 UDIMIndex = ParseUDIMName(FPaths::GetBaseFilename(UDIMFile), UDIMName);
+					const int32 UDIMIndex = ParseUDIMName(FPaths::GetBaseFilename(UDIMFile), UdimRegexPattern, UDIMName);
 					if (!UDIMIndexToFile.Contains(UDIMIndex) && UDIMName == BaseUDIMName)
 					{
 						UDIMIndexToFile.Add(UDIMIndex, Path / UDIMFile);
@@ -4219,7 +4220,7 @@ UObject* UTextureFactory::FactoryCreateBinary
 				InParent->GetName(PackageName);
 
 				FString PackageUDIMName;
-				const int32 PackageUDIMIndex = ParseUDIMName(PackageName, PackageUDIMName);
+				const int32 PackageUDIMIndex = ParseUDIMName(PackageName, UdimRegexPattern, PackageUDIMName);
 				if (PackageUDIMIndex == -1)
 				{
 					// If we're re-importing UDIM texture, the package will already be correctly named after the UDIM base name
