USD: Base-level error handling; Unify logging

diff --git a/Engine/Plugins/Importers/USDImporter/Source/USDImporter/Private/USDImporter.cpp b/Engine/Plugins/Importers/USDImporter/Source/USDImporter/Private/USDImporter.cpp
--- a/Engine/Plugins/Importers/USDImporter/Source/USDImporter/Private/USDImporter.cpp
+++ b/Engine/Plugins/Importers/USDImporter/Source/USDImporter/Private/USDImporter.cpp
@@ -28,6 +28,7 @@
 #include "USDErrorUtils.h"
 #include "USDImportOptions.h"
 #include "USDImporterProjectSettings.h"
+#include "USDLog.h"
 #include "USDMemory.h"
 #include "USDPrimResolverKind.h"
 #include "USDSceneImportFactory.h"
@@ -46,8 +47,6 @@
 
 #define LOCTEXT_NAMESPACE "USDImportPlugin"
 
-DEFINE_LOG_CATEGORY(LogUSDImport);
-
 
 class SUSDOptionsWindow : public SCompoundWidget
 {
@@ -595,7 +594,7 @@ void FUSDSceneImportContext::Init(UObject* InParent, const FString& InName, cons
 void FUsdImportContext::AddErrorMessage(EMessageSeverity::Type MessageSeverity, FText ErrorMessage)
 {
 	TokenizedErrorMessages.Add(FTokenizedMessage::Create(MessageSeverity, ErrorMessage));
-	UE_LOG(LogUSDImport, Error, TEXT("%s"), *ErrorMessage.ToString());
+	UE_LOG(LogUsd, Error, TEXT("%s"), *ErrorMessage.ToString());
 }
 
 void FUsdImportContext::DisplayErrorMessages(bool bAutomated)
@@ -619,7 +618,7 @@ void FUsdImportContext::DisplayErrorMessages(bool bAutomated)
 	{
 		for (const TSharedRef<FTokenizedMessage>& Message : TokenizedErrorMessages)
 		{
-			UE_LOG(LogUSDImport, Error, TEXT("%s"), *Message->ToText().ToString());
+			UE_LOG(LogUsd, Error, TEXT("%s"), *Message->ToText().ToString());
 		}
 	}
 }
diff --git a/Engine/Plugins/Importers/USDImporter/Source/USDImporter/Private/USDPrimResolver.cpp b/Engine/Plugins/Importers/USDImporter/Source/USDImporter/Private/USDPrimResolver.cpp
--- a/Engine/Plugins/Importers/USDImporter/Source/USDImporter/Private/USDPrimResolver.cpp
+++ b/Engine/Plugins/Importers/USDImporter/Source/USDImporter/Private/USDPrimResolver.cpp
@@ -5,6 +5,7 @@
 #include "USDConversionUtils.h"
 #include "USDImportOptions.h"
 #include "USDImporter.h"
+#include "USDLog.h"
 #include "USDTypesConversion.h"
 #include "UnrealUSDWrapper.h"
 
@@ -185,7 +186,7 @@ AActor* UUSDPrimResolver::SpawnActor(FUSDSceneImportContext& ImportContext, cons
 						FText::FromString(SpawnData.AssetPath),
 						FText::FromString(UsdToUnreal::ConvertString(SpawnData.ActorPrim.Get().GetPath().GetString()))));
 
-				UE_LOG(LogUSDImport, Error, TEXT("Could not find Unreal Asset '%s' for USD prim '%s'"), *SpawnData.AssetPath, *SpawnData.ActorName.ToString());
+				UE_LOG(LogUsd, Error, TEXT("Could not find Unreal Asset '%s' for USD prim '%s'"), *SpawnData.AssetPath, *SpawnData.ActorName.ToString());
 			}
 		}
 
diff --git a/Engine/Plugins/Importers/USDImporter/Source/USDImporter/Private/USDPrimResolverKind.cpp b/Engine/Plugins/Importers/USDImporter/Source/USDImporter/Private/USDPrimResolverKind.cpp
--- a/Engine/Plugins/Importers/USDImporter/Source/USDImporter/Private/USDPrimResolverKind.cpp
+++ b/Engine/Plugins/Importers/USDImporter/Source/USDImporter/Private/USDPrimResolverKind.cpp
@@ -5,6 +5,7 @@
 #include "USDConversionUtils.h"
 #include "USDImportOptions.h"
 #include "USDImporter.h"
+#include "USDLog.h"
 #include "USDSceneImportFactory.h"
 #include "USDTypesConversion.h"
 #include "UnrealUSDWrapper.h"
@@ -27,7 +28,7 @@ void UUSDPrimResolverKind::FindActorsToSpawn_Recursive(FUSDSceneImportContext& I
 
 	UUSDSceneImportOptions* ImportOptions = Cast<UUSDSceneImportOptions>(ImportContext.ImportOptions);
 
-	// Parent/child hierarchy will be ignored unless the kind is a group.  Keep track of the current parent and use it 
+	// Parent/child hierarchy will be ignored unless the kind is a group.  Keep track of the current parent and use it
 	TUsdStore< pxr::UsdPrim > GroupParent = ParentPrim;
 
 	if (IUsdPrim::IsKindChildOf( *Prim, USDKindTypes::Component ))
@@ -41,7 +42,7 @@ void UUSDPrimResolverKind::FindActorsToSpawn_Recursive(FUSDSceneImportContext& I
 		{
 			SpawnData.ActorClassName = UsdToUnreal::ConvertString( IUsdPrim::GetUnrealActorClass( *Prim ) );
 
-			UE_LOG(LogUSDImport, Log, TEXT("Adding %s Actor with custom actor class to spawn"), *PrimName.ToString());
+			UE_LOG(LogUsd, Log, TEXT("Adding %s Actor with custom actor class to spawn"), *PrimName.ToString());
 		}
 		else if(bHasUnrealAssetPath)
 		{
@@ -54,13 +55,13 @@ void UUSDPrimResolverKind::FindActorsToSpawn_Recursive(FUSDSceneImportContext& I
 				FindMeshAssetsToImport(ImportContext, *Prim, GroupParent, SpawnData.AssetsToImport);
 			}
 
-			UE_LOG(LogUSDImport, Log, TEXT("Adding %s Actor with custom asset path to spawn"), *PrimName.ToString());
+			UE_LOG(LogUsd, Log, TEXT("Adding %s Actor with custom asset path to spawn"), *PrimName.ToString());
 		}
-		else 
+		else
 		{
 			FindMeshAssetsToImport(ImportContext, *Prim, GroupParent, SpawnData.AssetsToImport);
 
-			UE_LOG(LogUSDImport, Log, TEXT("Adding %s Actor with %d meshes"), *PrimName.ToString(), SpawnData.AssetsToImport.Num());
+			UE_LOG(LogUsd, Log, TEXT("Adding %s Actor with %d meshes"), *PrimName.ToString(), SpawnData.AssetsToImport.Num());
 		}
 
 		SpawnData.WorldTransform = UsdToUnreal::ConvertMatrix( UsdUtils::GetUsdStageAxis( *ImportContext.Stage ), IUsdPrim::GetLocalTransform( *Prim ) );
@@ -75,7 +76,7 @@ void UUSDPrimResolverKind::FindActorsToSpawn_Recursive(FUSDSceneImportContext& I
 
 		OutSpawnDatas.Add(SpawnData);
 
-	
+
 	}
 	else if ( IUsdPrim::IsKindChildOf( *Prim, USDKindTypes::Group ))
 	{
@@ -97,7 +98,7 @@ void UUSDPrimResolverKind::FindActorsToSpawn_Recursive(FUSDSceneImportContext& I
 			ImportContext.ActorsToDestroy.Add(SpawnData.ActorName);
 		}
 
-		UE_LOG(LogUSDImport, Log, TEXT("Adding %s Group Actor to spawn"), *PrimName.ToString());
+		UE_LOG(LogUsd, Log, TEXT("Adding %s Group Actor to spawn"), *PrimName.ToString());
 	}
 
 	for (pxr::UsdPrim Child : Prim.Get().GetFilteredChildren( pxr::UsdTraverseInstanceProxies() ))
diff --git a/Engine/Plugins/Importers/USDImporter/Source/USDImporter/Private/USDSceneImportFactory.h b/Engine/Plugins/Importers/USDImporter/Source/USDImporter/Private/USDSceneImportFactory.h
--- a/Engine/Plugins/Importers/USDImporter/Source/USDImporter/Private/USDSceneImportFactory.h
+++ b/Engine/Plugins/Importers/USDImporter/Source/USDImporter/Private/USDSceneImportFactory.h
@@ -1,6 +1,6 @@
 // Copyright Epic Games, Inc. All Rights Reserved.
 
-#pragma once 
+#pragma once
 
 #include "Factories/SceneImportFactory.h"
 #include "Editor/EditorEngine.h"
@@ -9,8 +9,6 @@
 
 #include "USDSceneImportFactory.generated.h"
 
-DECLARE_LOG_CATEGORY_EXTERN(LogUSDSceneImport, Log, All);
-
 class UUSDSceneImportOptions;
 class UWorld;
 
diff --git a/Engine/Plugins/Importers/USDImporter/Source/USDImporter/Public/USDImporter.h b/Engine/Plugins/Importers/USDImporter/Source/USDImporter/Public/USDImporter.h
--- a/Engine/Plugins/Importers/USDImporter/Source/USDImporter/Public/USDImporter.h
+++ b/Engine/Plugins/Importers/USDImporter/Source/USDImporter/Public/USDImporter.h
@@ -1,6 +1,6 @@
 // Copyright Epic Games, Inc. All Rights Reserved.
 
-#pragma once 
+#pragma once
 
 #include "Editor/EditorEngine.h"
 #include "Logging/TokenizedMessage.h"
@@ -9,13 +9,11 @@
 #include "UnrealUSDWrapper.h"
 #include "USDMemory.h"
 
-DECLARE_LOG_CATEGORY_EXTERN(LogUSDImport, Log, All);
-
 namespace USDKindTypes
 {
 	// Note: std::string for compatiblity with USD
 
-	const std::string Component("component"); 
+	const std::string Component("component");
 	const std::string Group("group");
 	const std::string SubComponent("subcomponent");
 }
@@ -35,7 +33,7 @@ USTRUCT()
 struct USDIMPORTER_API FUsdImportContext
 {
 	GENERATED_BODY()
-	
+
 	/** Mapping of path to imported assets  */
 	TMap<FString, UObject*> PathToImportAssetMap;
 
diff --git a/Engine/Plugins/Importers/USDImporter/Source/USDStage/Private/USDLevelSequenceHelper.cpp b/Engine/Plugins/Importers/USDImporter/Source/USDStage/Private/USDLevelSequenceHelper.cpp
--- a/Engine/Plugins/Importers/USDImporter/Source/USDStage/Private/USDLevelSequenceHelper.cpp
+++ b/Engine/Plugins/Importers/USDImporter/Source/USDStage/Private/USDLevelSequenceHelper.cpp
@@ -2,8 +2,8 @@
 
 #include "USDLevelSequenceHelper.h"
 
-#include "USDErrorUtils.h"
 #include "USDListener.h"
+#include "USDLog.h"
 #include "USDMemory.h"
 #include "USDStageActor.h"
 #include "USDTypesConversion.h"
@@ -140,7 +140,7 @@ void FUsdLevelSequenceHelperImpl::CreateLevelSequences()
 
 	double TimeCodesPerSecond = ValidStageActor->TimeCodesPerSecond;
 
-	UE_LOG(LogUsdStage, Verbose, TEXT("CreateLevelSequences: Initializing level sequence with '%f' fps"), TimeCodesPerSecond);
+	UE_LOG(LogUsd, Verbose, TEXT("CreateLevelSequences: Initializing level sequence with '%f' fps"), TimeCodesPerSecond);
 	if ( ValidStageActor->LevelSequence )
 	{
 		return;
@@ -190,7 +190,7 @@ void FUsdLevelSequenceHelperImpl::CreateLevelSequences()
 			}
 			SubMovieScene->SetDisplayRate(FFrameRate(Info.TimeCodesPerSecond.GetValue(), 1));
 
-			UE_LOG(LogUsdStage, Verbose, TEXT("CreateLevelSequences: Created new LevelSequence '%s' for layer '%s'"), *Sequence->GetName(), *Identifier);
+			UE_LOG(LogUsd, Verbose, TEXT("CreateLevelSequences: Created new LevelSequence '%s' for layer '%s'"), *Sequence->GetName(), *Identifier);
 		}
 	}
 }
@@ -208,7 +208,7 @@ void FUsdLevelSequenceHelperImpl::CreateSubsequenceTrack(const FUsdLevelSequence
 		return;
 	}
 
-	UE_LOG(LogUsdStage, Verbose, TEXT("CreateSubsequenceSections: Initializing for layer '%s'"), *Info->Identifier);
+	UE_LOG(LogUsd, Verbose, TEXT("CreateSubsequenceSections: Initializing for layer '%s'"), *Info->Identifier);
 
 	// Recurse first as we'll need the sublayer LevelSequences to be completed to insert them as subsections
 	bool bFoundAnimatedSubLayer = false;
@@ -236,7 +236,7 @@ void FUsdLevelSequenceHelperImpl::CreateSubsequenceTrack(const FUsdLevelSequence
 
 	if (Sequence != nullptr && Info->IsAnimated() && bFoundAnimatedSubLayer)
 	{
-		UE_LOG(LogUsdStage, Verbose, TEXT("CreateSubsequenceSections: Layer '%s' is animated"), *Info->Identifier);
+		UE_LOG(LogUsd, Verbose, TEXT("CreateSubsequenceSections: Layer '%s' is animated"), *Info->Identifier);
 
 		UMovieScene* MovieScene = Sequence->GetMovieScene();
 		if (!MovieScene)
@@ -264,7 +264,7 @@ void FUsdLevelSequenceHelperImpl::CreateSubsequenceTrack(const FUsdLevelSequence
 			ULevelSequence** SubSequence = ValidStageActor->SubLayerLevelSequencesByIdentifier.Find(SubInfo->Identifier);
 			if (!SubSequence || !*SubSequence)
 			{
-				UE_LOG(LogUsdStage, Warning, TEXT("CreateSubsequenceSections: Invalid LevelSequence for sublayer '%s'"), *SubInfo->Identifier);
+				UE_LOG(LogUsd, Warning, TEXT("CreateSubsequenceSections: Invalid LevelSequence for sublayer '%s'"), *SubInfo->Identifier);
 				continue;
 			}
 
@@ -278,7 +278,7 @@ void FUsdLevelSequenceHelperImpl::CreateSubsequenceTrack(const FUsdLevelSequence
 			int32 Duration = FrameRate.AsFrameNumber(SubEndSeconds).Value - StartFrame.Value;
 			UMovieSceneSubSection* NewSection = SubTrack->AddSequenceOnRow(*SubSequence, StartFrame, Duration, INDEX_NONE);
 
-			UE_LOG(LogUsdStage, Verbose, TEXT("CreateSubsequenceSections: Layer '%s' received subsection for layer '%s' starting at frame '%d' and with duration '%d'"),
+			UE_LOG(LogUsd, Verbose, TEXT("CreateSubsequenceSections: Layer '%s' received subsection for layer '%s' starting at frame '%d' and with duration '%d'"),
 				*Info->Identifier,
 				*SubInfo->Identifier,
 				StartFrame.Value,
@@ -298,7 +298,7 @@ void FUsdLevelSequenceHelperImpl::CreateSubsequenceTrack(const FUsdLevelSequence
 		MovieScene->SetViewRange(StartTimeSeconds - 1.0f, 1.0f + EndTimeSeconds);
 		MovieScene->SetWorkingRange(StartTimeSeconds - 1.0f, 1.0f + EndTimeSeconds);
 
-		UE_LOG(LogUsdStage, Verbose, TEXT("CreateSubsequenceSections: Layer '%s' received view range [%f, %f]"),
+		UE_LOG(LogUsd, Verbose, TEXT("CreateSubsequenceSections: Layer '%s' received view range [%f, %f]"),
 			*Info->Identifier,
 			StartTimeSeconds - 1.0f,
 			EndTimeSeconds + 1.0f);
@@ -394,7 +394,7 @@ void FUsdLevelSequenceHelperImpl::RebuildLayerTimeInfoCache(const pxr::UsdStageR
 		return;
 	}
 
-	UE_LOG(LogUsdStage, Verbose, TEXT("RebuildLayerTimeInfoCache"));
+	UE_LOG(LogUsd, Verbose, TEXT("RebuildLayerTimeInfoCache"));
 
 	LayerTimeInfosByIdentifier.Reset();
 	LayerInfoIdentifierByLevelSequenceName.Reset();
@@ -447,7 +447,7 @@ void FUsdLevelSequenceHelperImpl::InitializeActorTimeCodes(const pxr::UsdStageRe
 	ValidStageActor->StartTimeCode = RootInfo->StartTimeCode.Get(0);
 	ValidStageActor->EndTimeCode = RootInfo->EndTimeCode.Get(0);
 	ValidStageActor->TimeCodesPerSecond = RootInfo->TimeCodesPerSecond.Get(DEFAULT_FRAMERATE);
-	UE_LOG(LogUsdStage, Verbose, TEXT("Initialized AUsdStageActor's Stage. StartTimeCode: %f, EndTimeCode: %f, TimeCodesPerSeconds: %f"), ValidStageActor->StartTimeCode, ValidStageActor->EndTimeCode, ValidStageActor->TimeCodesPerSecond);
+	UE_LOG(LogUsd, Verbose, TEXT("Initialized AUsdStageActor's Stage. StartTimeCode: %f, EndTimeCode: %f, TimeCodesPerSeconds: %f"), ValidStageActor->StartTimeCode, ValidStageActor->EndTimeCode, ValidStageActor->TimeCodesPerSecond);
 }
 
 FUsdLevelSequenceHelperImpl::FLayerTimeInfo* FUsdLevelSequenceHelperImpl::RecursivelyCreateLayerTimeInfo(const TUsdStore<pxr::SdfLayerHandle>& LayerHandleWrapper, TSet<FString>& ExploredFiles)
@@ -458,7 +458,7 @@ FUsdLevelSequenceHelperImpl::FLayerTimeInfo* FUsdLevelSequenceHelperImpl::Recurs
 	FString Identifier = UsdToUnreal::ConvertString(LayerHandle->GetIdentifier());
 	if (ExploredFiles.Contains(Identifier))
 	{
-		UE_LOG(LogUsdStage, Warning, TEXT("Detected infinite SubLayer recursion when visiting layer with identifier '%s'!"), *Identifier);
+		UE_LOG(LogUsd, Warning, TEXT("Detected infinite SubLayer recursion when visiting layer with identifier '%s'!"), *Identifier);
 		return nullptr;
 	}
 
@@ -472,7 +472,7 @@ FUsdLevelSequenceHelperImpl::FLayerTimeInfo* FUsdLevelSequenceHelperImpl::Recurs
 	LayerTimeInfo.EndTimeCode        = LayerHandle->HasEndTimeCode() ? LayerHandle->GetEndTimeCode() : TOptional<double>();
 	LayerTimeInfo.TimeCodesPerSecond = LayerHandle->HasTimeCodesPerSecond() ? LayerHandle->GetTimeCodesPerSecond() : TOptional<double>();
 
-	UE_LOG(LogUsdStage, Verbose, TEXT("Creating layer time info for layer '%s'. Original timecodes: '%s', '%s' and '%s'"),
+	UE_LOG(LogUsd, Verbose, TEXT("Creating layer time info for layer '%s'. Original timecodes: '%s', '%s' and '%s'"),
 		*LayerTimeInfo.Identifier,
 		LayerTimeInfo.StartTimeCode.IsSet() ? *LexToString(LayerTimeInfo.StartTimeCode.GetValue()) : TEXT("null"),
 		LayerTimeInfo.EndTimeCode.IsSet() ? *LexToString(LayerTimeInfo.EndTimeCode.GetValue()) : TEXT("null"),
@@ -511,7 +511,7 @@ FUsdLevelSequenceHelperImpl::FLayerTimeInfo* FUsdLevelSequenceHelperImpl::Recurs
 
 		if (FUsdLevelSequenceHelperImpl::FLayerTimeInfo* SubLayerInfo = RecursivelyCreateLayerTimeInfo(SubLayer, ExploredFiles))
 		{
-			UE_LOG(LogUsdStage, Verbose, TEXT("Adding sublayer '%s' to layer '%s'. Offset: '%f', Scale: '%f'"),
+			UE_LOG(LogUsd, Verbose, TEXT("Adding sublayer '%s' to layer '%s'. Offset: '%f', Scale: '%f'"),
 				*SubLayerInfo->Identifier,
 				*LayerTimeInfo.Identifier,
 				SubOffset,
@@ -554,7 +554,7 @@ void FUsdLevelSequenceHelperImpl::RecursivelyPropagateTimeCodes(FUsdLevelSequenc
 	bool bHaveEndCode	= Info->EndTimeCode.IsSet();
 	if (bHaveStartCode && bHaveEndCode)
 	{
-		UE_LOG(LogUsdStage, Verbose, TEXT("RecursivelyPropagateTimeCodes: Aborting for layer '%s' as it already has '%s', '%s' and '%s' set"),
+		UE_LOG(LogUsd, Verbose, TEXT("RecursivelyPropagateTimeCodes: Aborting for layer '%s' as it already has '%s', '%s' and '%s' set"),
 			*Info->Identifier,
 			Info->StartTimeCode.IsSet() ? *LexToString(Info->StartTimeCode.GetValue()) : TEXT("null"),
 			Info->EndTimeCode.IsSet() ? *LexToString(Info->EndTimeCode.GetValue()) : TEXT("null"),
@@ -595,7 +595,7 @@ void FUsdLevelSequenceHelperImpl::RecursivelyPropagateTimeCodes(FUsdLevelSequenc
 		Info->EndTimeCode = MaxEndSeconds.GetValue() * TimeCodesPerSecond;
 	}
 
-	UE_LOG(LogUsdStage, Verbose, TEXT("RecursivelyPropagateTimeCodes: Finished for layer '%s': '%s', '%s' and '%s' set"),
+	UE_LOG(LogUsd, Verbose, TEXT("RecursivelyPropagateTimeCodes: Finished for layer '%s': '%s', '%s' and '%s' set"),
 		*Info->Identifier,
 		Info->StartTimeCode.IsSet() ? *LexToString(Info->StartTimeCode.GetValue()) : TEXT("null"),
 		Info->EndTimeCode.IsSet() ? *LexToString(Info->EndTimeCode.GetValue()) : TEXT("null"),
@@ -615,12 +615,12 @@ void FUsdLevelSequenceHelperImpl::UpdateInfoFromSection(FUsdLevelSequenceHelperI
 		Layer = MakeUsdStore<pxr::SdfLayerHandle>(pxr::SdfLayer::Find(UnrealToUsd::ConvertString(*Info->Identifier).Get()));
 		if (!Layer.Get())
 		{
-			UE_LOG(LogUsdStage, Warning, TEXT("Failed to update sublayer '%s'"), *Info->Identifier);
+			UE_LOG(LogUsd, Warning, TEXT("Failed to update sublayer '%s'"), *Info->Identifier);
 			return;
 		}
 	}
 
-	UE_LOG(LogUsdStage, Verbose, TEXT("Updating LevelSequence '%s' for sublayer '%s'"), *Sequence->GetName(), *Info->Identifier);
+	UE_LOG(LogUsd, Verbose, TEXT("Updating LevelSequence '%s' for sublayer '%s'"), *Sequence->GetName(), *Info->Identifier);
 
 	UMovieScene* MovieScene = Sequence->GetMovieScene();
 	UMovieSceneSequence* SubSequence = Section->GetSequence();
@@ -678,7 +678,7 @@ void FUsdLevelSequenceHelperImpl::UpdateInfoFromSection(FUsdLevelSequenceHelperI
 	pxr::SdfLayerOffset NewLayerOffset{ModifiedOffset, ModifiedScale};
 	Layer.Get()->SetSubLayerOffset(NewLayerOffset, SubLayerIndex);
 
-	UE_LOG(LogUsdStage, Verbose, TEXT("\tNew OffsetScale: %f, %f"), NewLayerOffset.GetOffset(), NewLayerOffset.GetScale());
+	UE_LOG(LogUsd, Verbose, TEXT("\tNew OffsetScale: %f, %f"), NewLayerOffset.GetOffset(), NewLayerOffset.GetScale());
 }
 
 void FUsdLevelSequenceHelperImpl::OnObjectTransacted(UObject* Object, const class FTransactionObjectEvent& Event)
@@ -736,7 +736,7 @@ void FUsdLevelSequenceHelper::InitLevelSequence(const pxr::UsdStageRefPtr& UsdSt
 {
 	if (UsdSequencerImpl.IsValid() && UsdStage)
 	{
-		UE_LOG(LogUsdStage, Verbose, TEXT("InitLevelSequence"));
+		UE_LOG(LogUsd, Verbose, TEXT("InitLevelSequence"));
 
 		UsdSequencerImpl->RebuildLayerTimeInfoCache(UsdStage);
 
@@ -753,7 +753,7 @@ void FUsdLevelSequenceHelper::UpdateLevelSequence(const pxr::UsdStageRefPtr& Usd
 {
 	if (UsdSequencerImpl.IsValid() && UsdStage)
 	{
-		UE_LOG(LogUsdStage, Verbose, TEXT("UpdateLevelSequence"));
+		UE_LOG(LogUsd, Verbose, TEXT("UpdateLevelSequence"));
 
 		FUsdLevelSequenceHelperImpl::FLayerTimeInfo* RootInfo = UsdSequencerImpl->GetRootLayerInfo();
 		RootInfo->StartTimeCode = UsdStage->GetStartTimeCode();
diff --git a/Engine/Plugins/Importers/USDImporter/Source/USDStage/Private/USDStageActor.cpp b/Engine/Plugins/Importers/USDImporter/Source/USDStage/Private/USDStageActor.cpp
--- a/Engine/Plugins/Importers/USDImporter/Source/USDStage/Private/USDStageActor.cpp
+++ b/Engine/Plugins/Importers/USDImporter/Source/USDStage/Private/USDStageActor.cpp
@@ -6,9 +6,10 @@
 #include "USDGeomMeshTranslator.h"
 #include "USDGeomXformableTranslator.h"
 #include "USDListener.h"
+#include "USDLog.h"
 #include "USDPrimConversion.h"
-#include "USDSchemasModule.h"
 #include "USDSchemaTranslator.h"
+#include "USDSchemasModule.h"
 #include "USDSkelRootTranslator.h"
 #include "USDTypesConversion.h"
 #include "UnrealUSDWrapper.h"
@@ -58,8 +59,6 @@
 
 #define LOCTEXT_NAMESPACE "USDStageActor"
 
-DEFINE_LOG_CATEGORY( LogUsdStage );
-
 static const EObjectFlags DefaultObjFlag = EObjectFlags::RF_Transactional | EObjectFlags::RF_Transient;
 
 AUsdStageActor::FOnActorLoaded AUsdStageActor::OnActorLoaded;
@@ -209,7 +208,7 @@ AUsdStageActor::AUsdStageActor()
 					const pxr::SdfChangeList::Entry::_Flags& Flags = Change.second.flags;
 					if (Flags.didReloadContent)
 					{
-						UE_LOG(LogUsdStage, Verbose, TEXT("Reloading animations because layer '%s' reloaded"), *UsdToUnreal::ConvertPath(Change.first));
+						UE_LOG(LogUsd, Verbose, TEXT("Reloading animations because layer '%s' reloaded"), *UsdToUnreal::ConvertPath(Change.first));
 						ReloadAnimations();
 						return;
 					}
@@ -565,7 +564,7 @@ void AUsdStageActor::LoadUsdStage()
 	int ElapsedMin = int(ElapsedSeconds / 60.0);
 	ElapsedSeconds -= 60.0 * (double)ElapsedMin;
 
-	UE_LOG( LogUsdStage, Log, TEXT("%s %s in [%d min %.3f s]"), TEXT("Stage loaded"), *FPaths::GetBaseFilename( RootLayer.FilePath ), ElapsedMin, ElapsedSeconds );
+	UE_LOG( LogUsd, Log, TEXT("%s %s in [%d min %.3f s]"), TEXT("Stage loaded"), *FPaths::GetBaseFilename( RootLayer.FilePath ), ElapsedMin, ElapsedSeconds );
 #endif // #if USE_USD_SDK
 }
 
diff --git a/Engine/Plugins/Importers/USDImporter/Source/USDStage/Public/USDStageActor.h b/Engine/Plugins/Importers/USDImporter/Source/USDStage/Public/USDStageActor.h
--- a/Engine/Plugins/Importers/USDImporter/Source/USDStage/Public/USDStageActor.h
+++ b/Engine/Plugins/Importers/USDImporter/Source/USDStage/Public/USDStageActor.h
@@ -35,8 +35,6 @@ enum class EUsdPurpose : int32;
 struct FMeshDescription;
 struct FUsdSchemaTranslationContext;
 
-DECLARE_LOG_CATEGORY_EXTERN( LogUsdStage, Log, All );
-
 UENUM()
 enum class EUsdInitialLoadSet
 {
diff --git a/Engine/Plugins/Importers/USDImporter/Source/USDUtilities/Private/USDErrorUtils.cpp b/Engine/Plugins/Importers/USDImporter/Source/USDUtilities/Private/USDErrorUtils.cpp
--- a/Engine/Plugins/Importers/USDImporter/Source/USDUtilities/Private/USDErrorUtils.cpp
+++ b/Engine/Plugins/Importers/USDImporter/Source/USDUtilities/Private/USDErrorUtils.cpp
@@ -4,6 +4,7 @@
 
 #if USE_USD_SDK
 
+#include "USDLog.h"
 #include "USDMemory.h"
 #include "USDTypesConversion.h"
 
@@ -18,8 +19,6 @@
 
 #define LOCTEXT_NAMESPACE "USDErrorUtils"
 
-DEFINE_LOG_CATEGORY(LogUsdErrorUtility);
-
 namespace UsdUtils
 {
 	// We need an extra level of indirection because TfErrorMark is noncopyable
@@ -59,7 +58,7 @@ namespace UsdUtils
 			 ++ErrorIter)
 		{
 			std::string ErrorStr = ErrorIter->GetErrorCodeAsString();
-			ErrorStr += " ";
+			ErrorStr += ": ";
 			ErrorStr += ErrorIter->GetCommentary();
 
 			// Add unique here as for some errors (e.g. parsing errors) USD can emit the exact same
@@ -89,7 +88,7 @@ namespace UsdUtils
 
 		for (const FString& Error : Errors)
 		{
-			UE_LOG(LogUsdErrorUtility, Error, TEXT("%s"), *Error);
+			UE_LOG(LogUsd, Error, TEXT("%s"), *Error);
 		}
 
 		return Errors.Num() > 0;
diff --git a/Engine/Plugins/Importers/USDImporter/Source/USDUtilities/Public/USDErrorUtils.h b/Engine/Plugins/Importers/USDImporter/Source/USDUtilities/Public/USDErrorUtils.h
--- a/Engine/Plugins/Importers/USDImporter/Source/USDUtilities/Public/USDErrorUtils.h
+++ b/Engine/Plugins/Importers/USDImporter/Source/USDUtilities/Public/USDErrorUtils.h
@@ -4,8 +4,6 @@
 
 #include "CoreMinimal.h"
 
-DECLARE_LOG_CATEGORY_EXTERN(LogUsdErrorUtility, Log, All);
-
 namespace UsdUtils
 {
     /**
diff --git a/Engine/Plugins/Importers/USDImporter/Source/UnrealUSDWrapper/Private/UnrealUSDWrapper.cpp b/Engine/Plugins/Importers/USDImporter/Source/UnrealUSDWrapper/Private/UnrealUSDWrapper.cpp
--- a/Engine/Plugins/Importers/USDImporter/Source/UnrealUSDWrapper/Private/UnrealUSDWrapper.cpp
+++ b/Engine/Plugins/Importers/USDImporter/Source/UnrealUSDWrapper/Private/UnrealUSDWrapper.cpp
@@ -2,6 +2,7 @@
 
 #include "UnrealUSDWrapper.h"
 
+#include "USDLog.h"
 #include "USDMemory.h"
 
 #include "Modules/ModuleManager.h"
@@ -13,6 +14,7 @@
 #include "pxr/base/gf/rotation.h"
 #include "pxr/base/plug/plugin.h"
 #include "pxr/base/plug/registry.h"
+#include "pxr/base/tf/diagnosticMgr.h"
 #include "pxr/base/tf/errorMark.h"
 #include "pxr/base/tf/getenv.h"
 #include "pxr/base/tf/setenv.h"
@@ -885,10 +887,60 @@ EUsdGeomOrientation IUsdPrim::GetGeometryOrientation(const pxr::UsdGeomMesh& Mes
 
 	return GeomOrientation;
 }
+#endif // USE_USD_SDK
 
-bool UnrealUSDWrapper::bInitialized = false;
 
 
+DEFINE_LOG_CATEGORY( LogUsd );
+
+#if USE_USD_SDK
+class FUsdDiagnosticDelegate : public pxr::TfDiagnosticMgr::Delegate
+{
+public:
+	virtual ~FUsdDiagnosticDelegate() override {};
+	virtual void IssueError(const pxr::TfError& Error) override
+	{
+		FScopedUsdAllocs Allocs;
+
+		std::string Msg = Error.GetErrorCodeAsString();
+		Msg += ": ";
+		Msg += Error.GetCommentary();
+
+		UE_LOG(LogUsd, Error, TEXT("%s"), ANSI_TO_TCHAR( Msg.c_str() ));
+	}
+	virtual void IssueFatalError(const pxr::TfCallContext& Context, const std::string& Msg) override
+	{
+		UE_LOG(LogUsd, Error, TEXT("%s"), ANSI_TO_TCHAR( Msg.c_str() ));
+	}
+	virtual void IssueStatus(const pxr::TfStatus& Status) override
+	{
+		FScopedUsdAllocs Allocs;
+
+		std::string Msg = Status.GetDiagnosticCodeAsString();
+		Msg += ": ";
+		Msg += Status.GetCommentary();
+
+		UE_LOG(LogUsd, Log, TEXT("%s"), ANSI_TO_TCHAR( Msg.c_str() ));
+	}
+	virtual void IssueWarning(const pxr::TfWarning& Warning) override
+	{
+		FScopedUsdAllocs Allocs;
+
+		std::string Msg = Warning.GetDiagnosticCodeAsString();
+		Msg += ": ";
+		Msg += Warning.GetCommentary();
+
+		UE_LOG(LogUsd, Warning, TEXT("%s"), ANSI_TO_TCHAR( Msg.c_str() ));
+	}
+};
+#else
+class FUsdDiagnosticDelegate { };
+#endif // USE_USD_SDK
+
+TUniquePtr<FUsdDiagnosticDelegate> UnrealUSDWrapper::Delegate = nullptr;
+#if USE_USD_SDK
+bool UnrealUSDWrapper::bInitialized = false;
+
 void UnrealUSDWrapper::Initialize( const TArray< FString >& InPluginDirectories )
 {
 	{
@@ -948,16 +1000,50 @@ pxr::UsdStageCache& UnrealUSDWrapper::GetUsdStageCache()
 }
 #endif // USE_USD_SDK
 
+void UnrealUSDWrapper::SetupDiagnosticDelegate()
+{
+#if USE_USD_SDK
+	if (Delegate.IsValid())
+	{
+		UnrealUSDWrapper::ClearDiagnosticDelegate();
+	}
+
+	Delegate = MakeUnique<FUsdDiagnosticDelegate>();
+
+	pxr::TfDiagnosticMgr& DiagMgr = pxr::TfDiagnosticMgr::GetInstance();
+	DiagMgr.AddDelegate(Delegate.Get());
+#endif // USE_USD_SDK
+}
+
+void UnrealUSDWrapper::ClearDiagnosticDelegate()
+{
+#if USE_USD_SDK
+	if (!Delegate.IsValid())
+	{
+		return;
+	}
+
+	pxr::TfDiagnosticMgr& DiagMgr = pxr::TfDiagnosticMgr::GetInstance();
+	DiagMgr.RemoveDelegate(Delegate.Get());
+
+	Delegate = nullptr;
+#endif // USE_USD_SDK
+}
+
+
+
 class FUnrealUSDWrapperModule : public IUnrealUSDWrapperModule
 {
 public:
 	virtual void StartupModule() override
 	{
 		FUsdMemoryManager::Initialize();
+		UnrealUSDWrapper::SetupDiagnosticDelegate();
 	}
 
 	virtual void ShutdownModule() override
 	{
+		UnrealUSDWrapper::ClearDiagnosticDelegate();
 		FUsdMemoryManager::Shutdown();
 	}
 
@@ -965,7 +1051,7 @@ public:
 	{
 #if USE_USD_SDK
 		UnrealUSDWrapper::Initialize( InPluginDirectories );
-#endif
+#endif // USE_USD_SDK
 	}
 };
 
diff --git a/Engine/Plugins/Importers/USDImporter/Source/UnrealUSDWrapper/Public/USDLog.h b/Engine/Plugins/Importers/USDImporter/Source/UnrealUSDWrapper/Public/USDLog.h
new file mode 100644
--- /dev/null
+++ b/Engine/Plugins/Importers/USDImporter/Source/UnrealUSDWrapper/Public/USDLog.h
@@ -0,0 +1,7 @@
+// Copyright Epic Games, Inc. All Rights Reserved.
+
+#pragma once
+
+#include "Logging/LogMacros.h"
+
+UNREALUSDWRAPPER_API DECLARE_LOG_CATEGORY_EXTERN( LogUsd, Log, All );
\ No newline at end of file
diff --git a/Engine/Plugins/Importers/USDImporter/Source/UnrealUSDWrapper/Public/UnrealUSDWrapper.h b/Engine/Plugins/Importers/USDImporter/Source/UnrealUSDWrapper/Public/UnrealUSDWrapper.h
--- a/Engine/Plugins/Importers/USDImporter/Source/UnrealUSDWrapper/Public/UnrealUSDWrapper.h
+++ b/Engine/Plugins/Importers/USDImporter/Source/UnrealUSDWrapper/Public/UnrealUSDWrapper.h
@@ -14,13 +14,14 @@
 #include <memory>
 
 #if USE_USD_SDK
-
 #include "USDIncludesStart.h"
 	#include "pxr/pxr.h"
 	#include "pxr/base/tf/token.h"
 	#include "pxr/usd/usd/stageCache.h"
 #include "USDIncludesEnd.h"
+#endif // #if USE_USD_SDK
 
+#if USE_USD_SDK
 PXR_NAMESPACE_OPEN_SCOPE
 	class GfMatrix4d;
 	class SdfPath;
@@ -32,10 +33,10 @@ PXR_NAMESPACE_OPEN_SCOPE
 
 	template< typename T > class TfRefPtr;
 PXR_NAMESPACE_CLOSE_SCOPE
-
 #endif // #if USE_USD_SDK
 
 class IUsdPrim;
+class FUsdDiagnosticDelegate;
 
 enum class EUsdInterpolationMethod
 {
@@ -168,8 +169,12 @@ public:
 	UNREALUSDWRAPPER_API static TUsdStore< pxr::TfRefPtr< pxr::UsdStage > > OpenUsdStage(const char* Path, const char* Filename);
 	UNREALUSDWRAPPER_API static pxr::UsdStageCache& GetUsdStageCache();
 #endif  // #if USE_USD_SDK
+
+	UNREALUSDWRAPPER_API static void SetupDiagnosticDelegate();
+	UNREALUSDWRAPPER_API static void ClearDiagnosticDelegate();
 private:
 	static bool bInitialized;
+	static TUniquePtr<FUsdDiagnosticDelegate> Delegate;
 };
 
 class FUsdAttribute
