USD prim details
- Display all prim attributes

diff --git a/Engine/Plugins/Importers/USDImporter/Source/USDStageEditor/Private/SUSDPrimInfo.cpp b/Engine/Plugins/Importers/USDImporter/Source/USDStageEditor/Private/SUSDPrimInfo.cpp
--- a/Engine/Plugins/Importers/USDImporter/Source/USDStageEditor/Private/SUSDPrimInfo.cpp
+++ b/Engine/Plugins/Importers/USDImporter/Source/USDStageEditor/Private/SUSDPrimInfo.cpp
@@ -51,36 +51,34 @@ void SUsdPrimInfo::Construct( const FArguments& InArgs, const TUsdStore< pxr::Us
 		SNew( SVerticalBox )
 
 		+SVerticalBox::Slot()
-		.AutoHeight()
+		.FillHeight( 1.f )
 		[
-			SNew( SVerticalBox )
-
-			+SVerticalBox::Slot()
-			.AutoHeight()
+			SNew( SBox )
+			.Content()
 			[
-				SNew( SBorder )
-				.BorderImage( &FCoreStyle::Get().GetWidgetStyle< FHeaderRowStyle >("TableView.Header").BackgroundBrush )
+				SNew( SVerticalBox )
+
+				+SVerticalBox::Slot()
+				.AutoHeight()
+				[
+					SNew( SBorder )
+					.BorderImage( &FCoreStyle::Get().GetWidgetStyle< FHeaderRowStyle >("TableView.Header").BackgroundBrush )
 				.Padding( UsdPrimInfoWidgetConstants::CategoryHeaderPadding )
 				[
 					SNew( STextBlock )
 					.Font( FEditorStyle::GetFontStyle( TEXT("DetailsView.CategoryFontStyle") ) )
-					.Text( LOCTEXT( "Details", "Details" ) )
+				.Text( LOCTEXT( "Details", "Details" ) )
+				]
 				]
-			]
-
-			+SVerticalBox::Slot()
-			.AutoHeight()
-			[
-				SAssignNew( PropertiesList, SUsdPrimPropertiesList, PrimPath )
-			]
 
-			+SVerticalBox::Slot()
-			[
-				SNew( SSpacer )
-				.Size( FVector2D( 0.f, 10.f ) )
+				+SVerticalBox::Slot()
+				.FillHeight( 1.f )
+				[
+					SAssignNew( PropertiesList, SUsdPrimPropertiesList, PrimPath )
+				]
 			]
 		]
-		
+
 		+SVerticalBox::Slot()
 		.AutoHeight()
 		[
diff --git a/Engine/Plugins/Importers/USDImporter/Source/USDStageEditor/Private/Widgets/SUSDPrimPropertiesList.cpp b/Engine/Plugins/Importers/USDImporter/Source/USDStageEditor/Private/Widgets/SUSDPrimPropertiesList.cpp
--- a/Engine/Plugins/Importers/USDImporter/Source/USDStageEditor/Private/Widgets/SUSDPrimPropertiesList.cpp
+++ b/Engine/Plugins/Importers/USDImporter/Source/USDStageEditor/Private/Widgets/SUSDPrimPropertiesList.cpp
@@ -15,6 +15,7 @@
 #if USE_USD_SDK
 #include "USDIncludesStart.h"
 
+#include "pxr/base/tf/stringUtils.h"
 #include "pxr/usd/usd/prim.h"
 
 #include "USDIncludesEnd.h"
@@ -97,12 +98,13 @@ void SUsdPrimPropertiesList::Construct( const FArguments& InArgs, const TCHAR* I
 	GeneratePropertiesList( InPrimPath );
 
 	SAssignNew( HeaderRowWidget, SHeaderRow )
-	.Visibility( EVisibility::Collapsed )
 
 	+SHeaderRow::Column( FName( TEXT("PropertyName") ) )
+	.DefaultLabel( LOCTEXT( "PropertyName", "Name" ) )
 	.FillWidth( 20.f )
 
 	+SHeaderRow::Column( FName( TEXT("PropertyValue") ) )
+	.DefaultLabel( LOCTEXT( "PropertyValue", "Value" ) )
 	.FillWidth( 80.f );
 
 	SListView::Construct
@@ -126,25 +128,30 @@ void SUsdPrimPropertiesList::GeneratePropertiesList( const TCHAR* InPrimPath )
 	FString PrimName;
 	FString PrimKind;
 
+	TUsdStore< pxr::UsdPrim > UsdPrim;
+	pxr::UsdTimeCode TimeCode = pxr::UsdTimeCode::EarliestTime();
+
 	{
 		IUsdStageModule& UsdStageModule = FModuleManager::Get().LoadModuleChecked< IUsdStageModule >( "UsdStage" );
 		AUsdStageActor* UsdStageActor = &UsdStageModule.GetUsdStageActor( GWorld );
 
 		if ( UsdStageActor )
 		{
+			TimeCode = pxr::UsdTimeCode( UsdStageActor->GetTime() );
+
 			FScopedUsdAllocs UsdAllocs;
 
 			pxr::UsdStageRefPtr UsdStage = UsdStageActor->GetUsdStage();
 
 			if ( UsdStage )
 			{
-				pxr::UsdPrim UsdPrim = UsdStage->GetPrimAtPath( UnrealToUsd::ConvertPath( InPrimPath ).Get() );
+				UsdPrim = UsdStage->GetPrimAtPath( UnrealToUsd::ConvertPath( InPrimPath ).Get() );
 
-				if ( UsdPrim )
+				if ( UsdPrim.Get() )
 				{
 					PrimPath = InPrimPath;
-					PrimName = UsdToUnreal::ConvertString( UsdPrim.GetName() );
-					PrimKind = UsdToUnreal::ConvertString( IUsdPrim::GetKind( UsdPrim ).GetString() );
+					PrimName = UsdToUnreal::ConvertString( UsdPrim.Get().GetName() );
+					PrimKind = UsdToUnreal::ConvertString( IUsdPrim::GetKind( UsdPrim.Get() ).GetString() );
 				}
 			}
 		}
@@ -173,6 +180,27 @@ void SUsdPrimPropertiesList::GeneratePropertiesList( const TCHAR* InPrimPath )
 
 		PrimProperties.Add( MakeSharedUnreal< FUsdPrimProperty >( MoveTemp( PrimKindProperty ) ) );
 	}
+
+	if ( UsdPrim.Get() )
+	{
+		FScopedUsdAllocs UsdAllocs;
+
+		std::vector< pxr::UsdAttribute > PrimAttributes = UsdPrim.Get().GetAttributes();
+
+		for ( const pxr::UsdAttribute& PrimAttribute : PrimAttributes )
+		{
+			FUsdPrimProperty PrimAttributeProperty;
+			PrimAttributeProperty.Label = UsdToUnreal::ConvertString( PrimAttribute.GetName().GetString() );
+
+			pxr::VtValue VtValue;
+			PrimAttribute.Get( &VtValue, TimeCode );
+
+			FString AttributeValue = UsdToUnreal::ConvertString( pxr::TfStringify( VtValue ).c_str() );
+
+			PrimAttributeProperty.Value = MoveTemp( AttributeValue );
+			PrimProperties.Add( MakeSharedUnreal< FUsdPrimProperty >( MoveTemp( PrimAttributeProperty ) ) );
+		}
+	}
 }
 
 void SUsdPrimPropertiesList::SetPrimPath( const TCHAR* InPrimPath )
