USD - Readded support for animated static meshes by making sure we create new meshes when the geometry has changed.

diff --git a/Engine/Plugins/Importers/USDImporter/Source/USDSchemas/Private/USDGeomMeshTranslator.cpp b/Engine/Plugins/Importers/USDImporter/Source/USDSchemas/Private/USDGeomMeshTranslator.cpp
--- a/Engine/Plugins/Importers/USDImporter/Source/USDSchemas/Private/USDGeomMeshTranslator.cpp
+++ b/Engine/Plugins/Importers/USDImporter/Source/USDSchemas/Private/USDGeomMeshTranslator.cpp
@@ -38,7 +38,7 @@
 
 namespace UsdGeomMeshTranslatorImpl
 {
-	bool IsGeometryAnimated( const pxr::UsdGeomMesh& GeomMesh, const pxr::UsdTimeCode TimeCode )
+	bool IsAnimated( const pxr::UsdPrim& Prim )
 	{
 		FScopedUsdAllocs UsdAllocs;
 
@@ -53,14 +53,9 @@ namespace UsdGeomMeshTranslatorImpl
 
 			for ( const pxr::TfToken& AttributeName : GeomMeshAttributeNames )
 			{
-				const pxr::UsdAttribute& Attribute = GeomMesh.GetPrim().GetAttribute( AttributeName );
+				const pxr::UsdAttribute& Attribute = Prim.GetAttribute( AttributeName );
 
-				double DesiredTime = TimeCode.GetValue();
-				double MinTime = 0.0;
-				double MaxTime = 0.0;
-				bool bHasTimeSamples = false;
-
-				if ( bHasTimeSamples && DesiredTime >= MinTime && DesiredTime <= MaxTime && Attribute.ValueMightBeTimeVarying())
+				if ( Attribute.ValueMightBeTimeVarying() )
 				{
 					bHasAttributesTimeSamples = true;
 					break;
@@ -331,7 +326,6 @@ void FBuildStaticMeshTaskChain::SetupTasks()
 		return;
 	}
 
-	//if ( !UsdGeomMeshTranslatorImpl::IsGeometryAnimated( GeomMesh.Get(), pxr::UsdTimeCode( Context->Time ) ) ) // TODO: This test is now invalid since we can have multiple meshes collapsed into one
 	{
 		constexpr bool bIsAsyncTask = true;
 
@@ -438,30 +432,20 @@ void FUsdGeomMeshTranslator::CreateAssets()
 	Context->TranslatorTasks.Add( MoveTemp( AssetsTaskChain ) );
 }
 
-USceneComponent* FUsdGeomMeshTranslator::CreateComponents()
+void FUsdGeomMeshTranslator::UpdateComponents( USceneComponent* SceneComponent )
 {
-	TRACE_CPUPROFILER_EVENT_SCOPE( FUsdGeomMeshTranslator::CreateComponents );
-
-	USceneComponent* RootComponent = FUsdGeomXformableTranslator::CreateComponents();
-
-	if ( UStaticMeshComponent* StaticMeshComponent = Cast< UStaticMeshComponent >( RootComponent ) )
+	if ( UsdGeomMeshTranslatorImpl::IsAnimated( Schema.Get().GetPrim() ) )
 	{
-		UStaticMesh* PrimStaticMesh = Cast< UStaticMesh >( Context->PrimPathsToAssets.FindRef( UsdToUnreal::ConvertPath( Schema.Get().GetPath() ) ) );
-
-		if ( PrimStaticMesh != StaticMeshComponent->GetStaticMesh() )
-		{
-			if ( StaticMeshComponent->IsRegistered() )
-			{
-				StaticMeshComponent->UnregisterComponent();
-			}
-
-			StaticMeshComponent->SetStaticMesh( PrimStaticMesh );
-
-			StaticMeshComponent->RegisterComponent();
-		}
+		// The assets might have changed since our attributes are animated
+		CreateAssets();
 	}
 
-	return RootComponent;
+	Super::UpdateComponents( SceneComponent );
+}
+
+bool FUsdGeomMeshTranslator::CanBeCollapsed( ECollapsingType CollapsingType ) const
+{
+	return !UsdGeomMeshTranslatorImpl::IsAnimated( Schema.Get().GetPrim() );
 }
 
 #endif // #if USE_USD_SDK
diff --git a/Engine/Plugins/Importers/USDImporter/Source/USDSchemas/Private/USDGeomXformableTranslator.cpp b/Engine/Plugins/Importers/USDImporter/Source/USDSchemas/Private/USDGeomXformableTranslator.cpp
--- a/Engine/Plugins/Importers/USDImporter/Source/USDSchemas/Private/USDGeomXformableTranslator.cpp
+++ b/Engine/Plugins/Importers/USDImporter/Source/USDSchemas/Private/USDGeomXformableTranslator.cpp
@@ -264,8 +264,11 @@ void FUsdGeomXformableTranslator::UpdateComponents( USceneComponent* SceneCompon
 
 			if ( PrimStaticMesh != StaticMeshComponent->GetStaticMesh() )
 			{
-				// Need to make sure the mesh's resources are initialized here as it may have just been built in another thread
-				PrimStaticMesh->InitResources();
+				if ( PrimStaticMesh )
+				{
+					// Need to make sure the mesh's resources are initialized here as it may have just been built in another thread
+					PrimStaticMesh->InitResources();
+				}
 
 				if ( StaticMeshComponent->IsRegistered() )
 				{
diff --git a/Engine/Plugins/Importers/USDImporter/Source/USDSchemas/Public/USDGeomMeshTranslator.h b/Engine/Plugins/Importers/USDImporter/Source/USDSchemas/Public/USDGeomMeshTranslator.h
--- a/Engine/Plugins/Importers/USDImporter/Source/USDSchemas/Public/USDGeomMeshTranslator.h
+++ b/Engine/Plugins/Importers/USDImporter/Source/USDSchemas/Public/USDGeomMeshTranslator.h
@@ -63,15 +63,17 @@ protected:
 class USDSCHEMAS_API FUsdGeomMeshTranslator : public FUsdGeomXformableTranslator
 {
 public:
+	using Super = FUsdGeomXformableTranslator;
+
 	using FUsdGeomXformableTranslator::FUsdGeomXformableTranslator;
 
 	FUsdGeomMeshTranslator( const FUsdGeomMeshTranslator& Other ) = delete;
 	FUsdGeomMeshTranslator& operator=( const FUsdGeomMeshTranslator& Other ) = delete;
 
 	virtual void CreateAssets() override;
-	virtual USceneComponent* CreateComponents() override;
+	virtual void UpdateComponents( USceneComponent* SceneComponent ) override;
 
-	virtual bool CanBeCollapsed( ECollapsingType CollapsingType ) const override { return true; }
+	virtual bool CanBeCollapsed( ECollapsingType CollapsingType ) const override;
 
 };
 
