USD: Widgets to edit prim attributes;

Allows editing purpose and Kind attributes;
Tweak PrimInfo UI to remove double-header;
Change default new prim to 'assembly' kind as it likely is going to be;
Fix purpose component tag not updating when purpose changes;
Fix components being added directly to Stage Actor;

diff --git a/Engine/Plugins/Importers/USDImporter/Source/USDSchemas/Private/USDGeomXformableTranslator.cpp b/Engine/Plugins/Importers/USDImporter/Source/USDSchemas/Private/USDGeomXformableTranslator.cpp
--- a/Engine/Plugins/Importers/USDImporter/Source/USDSchemas/Private/USDGeomXformableTranslator.cpp
+++ b/Engine/Plugins/Importers/USDImporter/Source/USDSchemas/Private/USDGeomXformableTranslator.cpp
@@ -156,8 +156,19 @@ USceneComponent* FUsdGeomXformableTranslator::CreateComponents( TSubclassOf< USc
 {
 	TUsdStore< pxr::UsdPrim > Prim = Schema.Get().GetPrim();
 
+	// We should only add components to our transient/instanced actors, never to the AUsdStageActor or any other permanent actor
+	bool bOwnerActorIsPersistent = false;
+	if (USceneComponent* RootComponent = Context->ParentComponent)
+	{
+		if (AActor* Actor = RootComponent->GetOwner())
+		{
+			bOwnerActorIsPersistent = !Actor->HasAnyFlags(RF_Transient);
+		}
+	}
+
 	const bool bNeedsActor =
 		(
+			bOwnerActorIsPersistent ||
 			Context->ParentComponent == nullptr ||
 			Prim.Get().IsPseudoRoot() ||
 			Prim.Get().IsModel() ||
@@ -242,9 +253,6 @@ USceneComponent* FUsdGeomXformableTranslator::CreateComponents( TSubclassOf< USc
 		{
 			SceneComponent->GetOwner()->SetRootComponent( SceneComponent );
 		}
-
-		// Set purpose as a tag so its visibility can be toggled later
-		SceneComponent->ComponentTags.Add(IUsdPrim::GetPurposeName(IUsdPrim::GetPurpose(Prim.Get())));
 	}
 
 	return SceneComponent;
@@ -280,6 +288,11 @@ void FUsdGeomXformableTranslator::UpdateComponents( USceneComponent* SceneCompon
 				StaticMeshComponent->RegisterComponent();
 			}
 		}
+
+		// Set purpose as a tag so its visibility can be toggled later
+		TUsdStore< pxr::UsdPrim > Prim = Schema.Get().GetPrim();
+		SceneComponent->ComponentTags.RemoveAll([&](const FName& Tag){ return Tag.ToString().EndsWith(TEXT("Purpose")); });
+		SceneComponent->ComponentTags.Add(IUsdPrim::GetPurposeName(IUsdPrim::GetPurpose(Prim.Get())));
 	}
 }
 
diff --git a/Engine/Plugins/Importers/USDImporter/Source/USDStageEditor/Private/SUSDPrimInfo.cpp b/Engine/Plugins/Importers/USDImporter/Source/USDStageEditor/Private/SUSDPrimInfo.cpp
--- a/Engine/Plugins/Importers/USDImporter/Source/USDStageEditor/Private/SUSDPrimInfo.cpp
+++ b/Engine/Plugins/Importers/USDImporter/Source/USDStageEditor/Private/SUSDPrimInfo.cpp
@@ -42,7 +42,7 @@ namespace UsdPrimInfoWidgetConstants
 }
 
 void SUsdPrimInfo::Construct( const FArguments& InArgs, const TUsdStore< pxr::UsdStageRefPtr >& UsdStage, const TCHAR* PrimPath )
-{	
+{
 	TSharedRef< SWidget > VariantSetsWidget = GenerateVariantSetsWidget( UsdStage, PrimPath );
 	TSharedRef< SWidget > ReferencesListWidget = GenerateVariantSetsWidget( UsdStage, PrimPath );
 
@@ -56,26 +56,7 @@ void SUsdPrimInfo::Construct( const FArguments& InArgs, const TUsdStore< pxr::Us
 			SNew( SBox )
 			.Content()
 			[
-				SNew( SVerticalBox )
-
-				+SVerticalBox::Slot()
-				.AutoHeight()
-				[
-					SNew( SBorder )
-					.BorderImage( &FCoreStyle::Get().GetWidgetStyle< FHeaderRowStyle >("TableView.Header").BackgroundBrush )
-				.Padding( UsdPrimInfoWidgetConstants::CategoryHeaderPadding )
-				[
-					SNew( STextBlock )
-					.Font( FEditorStyle::GetFontStyle( TEXT("DetailsView.CategoryFontStyle") ) )
-				.Text( LOCTEXT( "Details", "Details" ) )
-				]
-				]
-
-				+SVerticalBox::Slot()
-				.FillHeight( 1.f )
-				[
-					SAssignNew( PropertiesList, SUsdPrimPropertiesList, PrimPath )
-				]
+				SAssignNew( PropertiesList, SUsdPrimPropertiesList, PrimPath )
 			]
 		]
 
diff --git a/Engine/Plugins/Importers/USDImporter/Source/USDStageEditor/Private/SUSDStage.cpp b/Engine/Plugins/Importers/USDImporter/Source/USDStageEditor/Private/SUSDStage.cpp
--- a/Engine/Plugins/Importers/USDImporter/Source/USDStageEditor/Private/SUSDStage.cpp
+++ b/Engine/Plugins/Importers/USDImporter/Source/USDStageEditor/Private/SUSDStage.cpp
@@ -36,6 +36,7 @@
 #include "USDIncludesStart.h"
 
 #include "pxr/pxr.h"
+#include "pxr/usd/kind/registry.h"
 #include "pxr/usd/sdf/copyUtils.h"
 #include "pxr/usd/sdf/schema.h"
 #include "pxr/usd/usd/editTarget.h"
@@ -565,7 +566,7 @@ void SUsdStage::OnNew()
 
 			// Create default prim
 			pxr::UsdGeomXform RootPrim = pxr::UsdGeomXform::Define( UsdStage, UnrealToUsd::ConvertPath( TEXT("/Root") ).Get() );
-			pxr::UsdModelAPI( RootPrim ).SetKind( pxr::TfToken("component") );
+			pxr::UsdModelAPI( RootPrim ).SetKind( pxr::KindTokens->assembly );
 
 			// Set default prim
 			UsdStage->SetDefaultPrim( RootPrim.GetPrim() );
diff --git a/Engine/Plugins/Importers/USDImporter/Source/USDStageEditor/Private/Widgets/SUSDPrimPropertiesList.cpp b/Engine/Plugins/Importers/USDImporter/Source/USDStageEditor/Private/Widgets/SUSDPrimPropertiesList.cpp
--- a/Engine/Plugins/Importers/USDImporter/Source/USDStageEditor/Private/Widgets/SUSDPrimPropertiesList.cpp
+++ b/Engine/Plugins/Importers/USDImporter/Source/USDStageEditor/Private/Widgets/SUSDPrimPropertiesList.cpp
@@ -2,21 +2,25 @@
 
 #include "SUSDPrimPropertiesList.h"
 
-#include "UnrealUSDWrapper.h"
+#include "USDLog.h"
 #include "USDMemory.h"
 #include "USDStageActor.h"
 #include "USDStageModule.h"
 #include "USDTypesConversion.h"
+#include "UnrealUSDWrapper.h"
 
 #include "EditorStyleSet.h"
 #include "Engine/World.h"
 #include "Modules/ModuleManager.h"
+#include "ScopedTransaction.h"
 
 #if USE_USD_SDK
 #include "USDIncludesStart.h"
 
 #include "pxr/base/tf/stringUtils.h"
+#include "pxr/usd/kind/registry.h"
 #include "pxr/usd/usd/prim.h"
+#include "pxr/usd/usdGeom/tokens.h"
 
 #include "USDIncludesEnd.h"
 
@@ -27,10 +31,120 @@ namespace UsdPrimPropertiesListConstants
 {
 	const FMargin LeftRowPadding( 6.0f, 2.5f, 2.0f, 2.5f );
 	const FMargin RightRowPadding( 3.0f, 2.5f, 2.0f, 2.5f );
+	const FMargin ComboBoxItemPadding( 3.0f, 0.0f, 2.0f, 0.0f );
 
 	const TCHAR* NormalFont = TEXT("PropertyWindow.NormalFont");
 }
 
+namespace UsdPrimPropertiesListImpl
+{
+	static TMap<FString, TArray<TSharedPtr<FString>>> TokenDropdownOptions;
+
+	void ResetOptions(const FString& TokenName)
+	{
+		TokenDropdownOptions.Remove(TokenName);
+	}
+
+	TArray<TSharedPtr<FString>>* GetTokenDropdownOptions(const FString& TokenName)
+	{
+		TArray<TSharedPtr<FString>>* FoundOptions = TokenDropdownOptions.Find(TokenName);
+		if (FoundOptions)
+		{
+			return FoundOptions;
+		}
+
+		if (TokenName == TEXT("Kind"))
+		{
+			TArray<TSharedPtr<FString>> Options;
+			{
+				FScopedUsdAllocs Allocs;
+
+				std::vector<pxr::TfToken> Kinds = pxr::KindRegistry::GetAllKinds();
+				Options.Reserve(Kinds.size());
+
+				for (const pxr::TfToken& Kind : Kinds)
+				{
+					Options.Add(MakeShared<FString>(UsdToUnreal::ConvertToken(Kind)));
+				}
+
+				// They are supposed to be in an unspecified order, so let's make them consistent
+				Options.Sort([](const TSharedPtr<FString>& A, const TSharedPtr<FString>& B)
+				{
+					return A.IsValid() && B.IsValid() && (*A < *B);
+				});
+			}
+			return &UsdPrimPropertiesListImpl::TokenDropdownOptions.Add(TokenName, Options);
+		}
+		else if (TokenName == UsdToUnreal::ConvertToken(pxr::UsdGeomTokens->purpose))
+		{
+			TArray<TSharedPtr<FString>> Options =
+			{
+				MakeShared<FString>(UsdToUnreal::ConvertToken(pxr::UsdGeomTokens->default_)),
+				MakeShared<FString>(UsdToUnreal::ConvertToken(pxr::UsdGeomTokens->proxy)),
+				MakeShared<FString>(UsdToUnreal::ConvertToken(pxr::UsdGeomTokens->render)),
+				MakeShared<FString>(UsdToUnreal::ConvertToken(pxr::UsdGeomTokens->guide)),
+			};
+			return &UsdPrimPropertiesListImpl::TokenDropdownOptions.Add(TokenName, Options);
+		}
+
+		return FoundOptions;
+	}
+
+	void SetPrimAttribute(const FString& PrimPath, const FString& AttributeName, const FString& Value)
+	{
+		FScopedTransaction Transaction(FText::Format(
+			LOCTEXT("SetPrimAttribute", "Set value '{0}' for attribute '{1}' of prim '{2}'"),
+			FText::FromString(Value),
+			FText::FromString(AttributeName),
+			FText::FromString(PrimPath)
+		));
+
+		IUsdStageModule& UsdStageModule = FModuleManager::Get().LoadModuleChecked< IUsdStageModule >("UsdStage");
+		AUsdStageActor* UsdStageActor = &UsdStageModule.GetUsdStageActor(GWorld);
+
+		FScopedUsdAllocs UsdAllocs;
+
+		bool bSuccess = false;
+
+		pxr::TfToken AttributeNameToken = UnrealToUsd::ConvertToken(*AttributeName).Get();
+
+		if (pxr::UsdStageRefPtr UsdStage = UsdStageActor->GetUsdStage())
+		{
+			pxr::UsdPrim UsdPrim = UsdStage->GetPrimAtPath(UnrealToUsd::ConvertPath(*PrimPath).Get());
+			if (UsdPrim)
+			{
+				if (AttributeName == TEXT("Kind"))
+				{
+					bSuccess = IUsdPrim::SetKind(UsdPrim, UnrealToUsd::ConvertToken(*Value).Get());
+				}
+				else if (AttributeNameToken == pxr::UsdGeomTokens->purpose)
+				{
+					pxr::UsdAttribute Attribute = UsdPrim.GetAttribute(AttributeNameToken);
+					bSuccess = Attribute.Set(UnrealToUsd::ConvertToken(*Value).Get());
+				}
+			}
+		}
+
+		if (!bSuccess)
+		{
+			UE_LOG(LogUsd, Error, TEXT("Failed to set value '%s' for attribute '%s' of prim '%s'"), *Value, *AttributeName, *PrimPath);
+		}
+	}
+}
+
+enum class EPrimPropertyWidget : uint8
+{
+	Text,
+	Dropdown,
+};
+
+struct FUsdPrimProperty : public TSharedFromThis< FUsdPrimProperty >
+{
+	FString Label;
+	FString Value;
+	EPrimPropertyWidget WidgetType = EPrimPropertyWidget::Text;
+};
+
 class SUsdPrimPropertyRow : public SMultiColumnTableRow< TSharedPtr< FUsdPrimProperty > >
 {
 	SLATE_BEGIN_ARGS( SUsdPrimPropertyRow ) {}
@@ -45,8 +159,11 @@ public:
 protected:
 	FText GetLabel() const { return FText::FromString( UsdPrimProperty->Label ); }
 	FText GetValue() const { return FText::FromString( UsdPrimProperty->Value ); }
-	
+	FText GetValueOrNone() const { return FText::FromString( UsdPrimProperty->Value.IsEmpty()? TEXT("none") : UsdPrimProperty->Value ); }
+
 private:
+	TSharedRef< SWidget > GenerateTextWidget(const TAttribute<FText>& Attribute);
+
 	TSharedPtr< FUsdPrimProperty > UsdPrimProperty;
 };
 
@@ -58,22 +175,52 @@ void SUsdPrimPropertyRow::Construct( const FArguments& InArgs, const TSharedPtr<
 }
 
 TSharedRef< SWidget > SUsdPrimPropertyRow::GenerateWidgetForColumn( const FName& ColumnName )
-{	
+{
 	TSharedRef< SWidget > ColumnWidget = SNullWidget::NullWidget;
 
 	if ( ColumnName == TEXT("PropertyName") )
 	{
-		SAssignNew( ColumnWidget, STextBlock )
-		.Text( this, &SUsdPrimPropertyRow::GetLabel )
-		.Font( FEditorStyle::GetFontStyle( UsdPrimPropertiesListConstants::NormalFont ) )
-		.Margin( UsdPrimPropertiesListConstants::LeftRowPadding );
+		ColumnWidget = SUsdPrimPropertyRow::GenerateTextWidget({this, &SUsdPrimPropertyRow::GetLabel});
 	}
 	else
 	{
-		SAssignNew( ColumnWidget, STextBlock )
-		.Text( this, &SUsdPrimPropertyRow::GetValue )
-		.Font( FEditorStyle::GetFontStyle( UsdPrimPropertiesListConstants::NormalFont ) )
-		.Margin( UsdPrimPropertiesListConstants::RightRowPadding );
+		if (UsdPrimProperty->WidgetType == EPrimPropertyWidget::Text)
+		{
+			ColumnWidget = GenerateTextWidget({this, &SUsdPrimPropertyRow::GetValue});
+		}
+		else if (UsdPrimProperty->WidgetType == EPrimPropertyWidget::Dropdown)
+		{
+			TArray<TSharedPtr<FString>>* Options = UsdPrimPropertiesListImpl::GetTokenDropdownOptions(UsdPrimProperty->Label);
+
+			// Show a dropdown if we know the available options for that token
+			if (Options)
+			{
+				SAssignNew(ColumnWidget, SComboBox<TSharedPtr<FString>>)
+				.OptionsSource(Options)
+				.OnGenerateWidget_Lambda([&](TSharedPtr<FString> Option)
+				{
+					return SUsdPrimPropertyRow::GenerateTextWidget(FText::FromString(*Option));
+				})
+				.OnSelectionChanged_Lambda([&](TSharedPtr<FString> ChosenOption, ESelectInfo::Type SelectInfo)
+				{
+					TSharedPtr<ITypedTableView<TSharedPtr<FUsdPrimProperty>>> PinnedParent = OwnerTablePtr.Pin();
+					TSharedPtr<SUsdPrimPropertiesList> ParentList = StaticCastSharedPtr<SUsdPrimPropertiesList>(PinnedParent);
+
+					UsdPrimPropertiesListImpl::SetPrimAttribute(ParentList->GetPrimPath(), UsdPrimProperty->Label, *ChosenOption);
+				})
+				[
+					SNew(STextBlock)
+					.Text(this, &SUsdPrimPropertyRow::GetValueOrNone)
+					.Font(FEditorStyle::GetFontStyle(UsdPrimPropertiesListConstants::NormalFont))
+					.Margin(UsdPrimPropertiesListConstants::ComboBoxItemPadding)
+				];
+			}
+			// Fallback to just displaying a simple text box
+			else
+			{
+				ColumnWidget = SUsdPrimPropertyRow::GenerateTextWidget({this, &SUsdPrimPropertyRow::GetValue});
+			}
+		}
 	}
 
 	return SNew( SHorizontalBox )
@@ -91,17 +238,28 @@ void SUsdPrimPropertyRow::SetUsdPrimProperty( const TSharedPtr< FUsdPrimProperty
 	UsdPrimProperty = InUsdPrimProperty;
 }
 
+TSharedRef< SWidget > SUsdPrimPropertyRow::GenerateTextWidget(const TAttribute<FText>& Attribute)
+{
+	return SNew(STextBlock)
+	.Text(Attribute)
+	.Font(FEditorStyle::GetFontStyle(UsdPrimPropertiesListConstants::NormalFont))
+	.Margin(UsdPrimPropertiesListConstants::RightRowPadding);
+}
+
 void SUsdPrimPropertiesList::Construct( const FArguments& InArgs, const TCHAR* InPrimPath )
 {
 	PrimPath = InPrimPath;
 
 	GeneratePropertiesList( InPrimPath );
 
+	// Clear map as usd file may have additional Kinds now
+	UsdPrimPropertiesListImpl::ResetOptions(TEXT("Kind"));
+
 	SAssignNew( HeaderRowWidget, SHeaderRow )
 
 	+SHeaderRow::Column( FName( TEXT("PropertyName") ) )
-	.DefaultLabel( LOCTEXT( "PropertyName", "Name" ) )
-	.FillWidth( 20.f )
+	.DefaultLabel( LOCTEXT( "PropertyName", "Property Name" ) )
+	.FillWidth( 25.f )
 
 	+SHeaderRow::Column( FName( TEXT("PropertyValue") ) )
 	.DefaultLabel( LOCTEXT( "PropertyValue", "Value" ) )
@@ -161,7 +319,7 @@ void SUsdPrimPropertiesList::GeneratePropertiesList( const TCHAR* InPrimPath )
 		FUsdPrimProperty PrimNameProperty;
 		PrimNameProperty.Label = TEXT("Name");
 		PrimNameProperty.Value = PrimName;
-	
+
 		PrimProperties.Add( MakeSharedUnreal< FUsdPrimProperty >( MoveTemp( PrimNameProperty ) ) );
 	}
 
@@ -177,6 +335,7 @@ void SUsdPrimPropertiesList::GeneratePropertiesList( const TCHAR* InPrimPath )
 		FUsdPrimProperty PrimKindProperty;
 		PrimKindProperty.Label = TEXT("Kind");
 		PrimKindProperty.Value = PrimKind;
+		PrimKindProperty.WidgetType = EPrimPropertyWidget::Dropdown;
 
 		PrimProperties.Add( MakeSharedUnreal< FUsdPrimProperty >( MoveTemp( PrimKindProperty ) ) );
 	}
@@ -192,9 +351,11 @@ void SUsdPrimPropertiesList::GeneratePropertiesList( const TCHAR* InPrimPath )
 			FUsdPrimProperty PrimAttributeProperty;
 			PrimAttributeProperty.Label = UsdToUnreal::ConvertString( PrimAttribute.GetName().GetString() );
 
+			// Just the Purpose attribute for now
+			PrimAttributeProperty.WidgetType = PrimAttribute.GetName() == pxr::UsdGeomTokens->purpose ? EPrimPropertyWidget::Dropdown : EPrimPropertyWidget::Text;
+
 			pxr::VtValue VtValue;
 			PrimAttribute.Get( &VtValue, TimeCode );
-
 			FString AttributeValue = UsdToUnreal::ConvertString( pxr::TfStringify( VtValue ).c_str() );
 
 			PrimAttributeProperty.Value = MoveTemp( AttributeValue );
diff --git a/Engine/Plugins/Importers/USDImporter/Source/USDStageEditor/Private/Widgets/SUSDPrimPropertiesList.h b/Engine/Plugins/Importers/USDImporter/Source/USDStageEditor/Private/Widgets/SUSDPrimPropertiesList.h
--- a/Engine/Plugins/Importers/USDImporter/Source/USDStageEditor/Private/Widgets/SUSDPrimPropertiesList.h
+++ b/Engine/Plugins/Importers/USDImporter/Source/USDStageEditor/Private/Widgets/SUSDPrimPropertiesList.h
@@ -6,11 +6,7 @@
 
 #if USE_USD_SDK
 
-struct FUsdPrimProperty : public TSharedFromThis< FUsdPrimProperty >
-{
-	FString Label;
-	FString Value;
-};
+struct FUsdPrimProperty;
 
 class SUsdPrimPropertiesList : public SListView< TSharedPtr< FUsdPrimProperty > >
 {
@@ -20,6 +16,7 @@ class SUsdPrimPropertiesList : public SListView< TSharedPtr< FUsdPrimProperty >
 public:
 	void Construct( const FArguments& InArgs, const TCHAR* InPrimPath );
 	void SetPrimPath( const TCHAR* InPrimPath );
+	const FString& GetPrimPath() { return PrimPath; };
 
 protected:
 	TSharedRef< ITableRow > OnGenerateRow( TSharedPtr< FUsdPrimProperty > InDisplayNode, const TSharedRef< STableViewBase >& OwnerTable );
diff --git a/Engine/Plugins/Importers/USDImporter/Source/USDUtilities/Public/USDTypesConversion.h b/Engine/Plugins/Importers/USDImporter/Source/USDUtilities/Public/USDTypesConversion.h
--- a/Engine/Plugins/Importers/USDImporter/Source/USDUtilities/Public/USDTypesConversion.h
+++ b/Engine/Plugins/Importers/USDImporter/Source/USDUtilities/Public/USDTypesConversion.h
@@ -62,6 +62,11 @@ namespace UsdToUnreal
 		return FName( UsdString.Get().c_str() );
 	}
 
+	static FString ConvertToken(const pxr::TfToken& Token)
+	{
+		return UsdToUnreal::ConvertString(Token.GetString());
+	}
+
 	static FLinearColor ConvertColor( const pxr::GfVec3f& InValue )
 	{
 		return FLinearColor( InValue[0], InValue[1], InValue[2] );
diff --git a/Engine/Plugins/Importers/USDImporter/Source/UnrealUSDWrapper/Private/UnrealUSDWrapper.cpp b/Engine/Plugins/Importers/USDImporter/Source/UnrealUSDWrapper/Private/UnrealUSDWrapper.cpp
--- a/Engine/Plugins/Importers/USDImporter/Source/UnrealUSDWrapper/Private/UnrealUSDWrapper.cpp
+++ b/Engine/Plugins/Importers/USDImporter/Source/UnrealUSDWrapper/Private/UnrealUSDWrapper.cpp
@@ -439,7 +439,7 @@ int FUsdAttribute::GetArraySize( const pxr::UsdAttribute& Attribute )
 	return Value.IsArrayValued() ? (int)Value.GetArraySize() : -1;
 
 }
-EUsdPurpose IUsdPrim::GetPurpose( const UsdPrim& Prim )
+EUsdPurpose IUsdPrim::GetPurpose( const UsdPrim& Prim, bool bComputed )
 {
 	UsdGeomImageable Geom(Prim);
 	if (Geom)
@@ -448,7 +448,21 @@ EUsdPurpose IUsdPrim::GetPurpose( const UsdPrim& Prim )
 		// "If the purpose of </RootPrim> is set to "render", then the effective purpose
 		// of </RootPrim/ChildPrim> will be "render" even if that prim has a different
 		// authored value for purpose."
-		const TfToken Purpose = Geom.ComputePurpose();
+		TfToken Purpose;
+		if (bComputed)
+		{
+			Purpose = Geom.ComputePurpose();
+		}
+		else
+		{
+			pxr::UsdAttribute PurposeAttr = Prim.GetAttribute(pxr::UsdGeomTokens->purpose);
+
+			pxr::VtValue Value;
+			PurposeAttr.Get(&Value);
+
+			Purpose = Value.Get<pxr::TfToken>();
+		}
+
 		if (Purpose == pxr::UsdGeomTokens->proxy)
 		{
 			return EUsdPurpose::Proxy;
@@ -539,13 +553,30 @@ TfToken IUsdPrim::GetKind(const pxr::UsdPrim& Prim)
 	else
 	{
 		// Prim is not a model, read kind directly from metadata
-		static TfToken KindMetaDataToken("kind");
+		const TfToken KindMetaDataToken("kind");
 		Prim.GetMetadata(KindMetaDataToken, &KindType);
 	}
 
 	return KindType;
 }
 
+bool IUsdPrim::SetKind(const pxr::UsdPrim& Prim, const pxr::TfToken& Kind)
+{
+	UsdModelAPI Model(Prim);
+	if (Model)
+	{
+		if (!Model.SetKind(Kind))
+		{
+			const TfToken KindMetaDataToken("kind");
+			return Prim.SetMetadata(KindMetaDataToken, Kind);
+		}
+
+		return true;
+	}
+
+	return false;
+}
+
 pxr::GfMatrix4d IUsdPrim::GetLocalTransform(const pxr::UsdPrim& Prim)
 {
 	pxr::GfMatrix4d USDMatrix(1);
diff --git a/Engine/Plugins/Importers/USDImporter/Source/UnrealUSDWrapper/Public/UnrealUSDWrapper.h b/Engine/Plugins/Importers/USDImporter/Source/UnrealUSDWrapper/Public/UnrealUSDWrapper.h
--- a/Engine/Plugins/Importers/USDImporter/Source/UnrealUSDWrapper/Public/UnrealUSDWrapper.h
+++ b/Engine/Plugins/Importers/USDImporter/Source/UnrealUSDWrapper/Public/UnrealUSDWrapper.h
@@ -207,14 +207,16 @@ class IUsdPrim
 {
 public:
 #if USE_USD_SDK
-	static UNREALUSDWRAPPER_API EUsdPurpose GetPurpose(const pxr::UsdPrim& Prim);
+	static UNREALUSDWRAPPER_API EUsdPurpose GetPurpose(const pxr::UsdPrim& Prim, bool bComputed = true);
 	static UNREALUSDWRAPPER_API FName GetPurposeName(EUsdPurpose Purpose);
 
 	static UNREALUSDWRAPPER_API bool HasGeometryData(const pxr::UsdPrim& Prim);
 	static UNREALUSDWRAPPER_API bool HasGeometryDataOrLODVariants(const pxr::UsdPrim& Prim);
 	static UNREALUSDWRAPPER_API int GetNumLODs(const pxr::UsdPrim& Prim);
+
 	static UNREALUSDWRAPPER_API bool IsKindChildOf(const pxr::UsdPrim& Prim, const std::string& InBaseKind);
 	static UNREALUSDWRAPPER_API pxr::TfToken GetKind(const pxr::UsdPrim& Prim);
+	static UNREALUSDWRAPPER_API bool SetKind(const pxr::UsdPrim& Prim, const pxr::TfToken& Kind);
 
 	static UNREALUSDWRAPPER_API pxr::GfMatrix4d GetLocalTransform(const pxr::UsdPrim& Prim);
 	static UNREALUSDWRAPPER_API pxr::GfMatrix4d GetLocalToWorldTransform(const pxr::UsdPrim& Prim );
