USD: Fix animations not resetting properly when reloading the stage

diff --git a/Engine/Plugins/Importers/USDImporter/Source/USDStage/Private/USDStageActor.cpp b/Engine/Plugins/Importers/USDImporter/Source/USDStage/Private/USDStageActor.cpp
--- a/Engine/Plugins/Importers/USDImporter/Source/USDStage/Private/USDStageActor.cpp
+++ b/Engine/Plugins/Importers/USDImporter/Source/USDStage/Private/USDStageActor.cpp
@@ -458,7 +458,6 @@ void AUsdStageActor::Clear()
 {
 	PrimPathsToAssets.Reset();
 	ObjectsToWatch.Reset();
-	SubLayerLevelSequencesByIdentifier.Reset();
 }
 
 void AUsdStageActor::OpenUsdStage()
@@ -587,10 +586,13 @@ void AUsdStageActor::ReloadAnimations()
 
 	if ( HasAutorithyOverStage() )
 	{
-		LevelSequenceHelper.InitLevelSequence(UsdStage);
-
 		// The sequencer won't update on its own, so let's at least force it closed
 		GEditor->GetEditorSubsystem<UAssetEditorSubsystem>()->CloseAllEditorsForAsset(LevelSequence);
+
+		LevelSequence = nullptr;
+		SubLayerLevelSequencesByIdentifier.Reset();
+
+		LevelSequenceHelper.InitLevelSequence(UsdStage);
 	}
 #endif // #if USE_USD_SDK
 }
