USD: Fix meshes rebuilding when they don't need to.

diff --git a/Engine/Plugins/Importers/USDImporter/Source/USDSchemas/Private/USDGeomMeshTranslator.cpp b/Engine/Plugins/Importers/USDImporter/Source/USDSchemas/Private/USDGeomMeshTranslator.cpp
--- a/Engine/Plugins/Importers/USDImporter/Source/USDSchemas/Private/USDGeomMeshTranslator.cpp
+++ b/Engine/Plugins/Importers/USDImporter/Source/USDSchemas/Private/USDGeomMeshTranslator.cpp
@@ -11,6 +11,7 @@
 
 #include "Async/Async.h"
 #include "Components/StaticMeshComponent.h"
+#include "EditorFramework/AssetImportData.h"
 #include "Engine/CollisionProfile.h"
 #include "Engine/StaticMesh.h"
 #include "IMeshBuilderModule.h"
@@ -67,7 +68,7 @@ namespace UsdGeomMeshTranslatorImpl
 	}
 
 	/** Returns true if material infos have changed on the StaticMesh */
-	bool ProcessMaterials( const pxr::UsdPrim& UsdPrim, UStaticMesh& StaticMesh, TMap< FString, UObject* >& PrimPathsToAssets, float Time )
+	bool ProcessMaterials( const pxr::UsdPrim& UsdPrim, UStaticMesh& StaticMesh, const TMap< FString, UObject* >& PrimPathsToAssets, float Time )
 	{
 		const FMeshDescription* MeshDescription = StaticMesh.GetMeshDescription( 0 );
 
@@ -139,12 +140,14 @@ namespace UsdGeomMeshTranslatorImpl
 
 			UMaterialInterface* Material = nullptr;
 
+			// If there was a UE4 material stored in the the main prim for this polygon group, try fetching it
 			if ( MainPrimUEMaterialsAttribute.IsValidIndex( MaterialIndex ))
 			{
 				Material = Cast< UMaterialInterface >( FSoftObjectPath( MainPrimUEMaterialsAttribute[ MaterialIndex ] ).TryLoad() );
 			}
 			else
 			{
+				// It could be that our UE4 material is stored in the subprim instead, so try fetching that
 				TArray< FString > PolygonGroupPrimUEMaterialsAttribute = FetchUEMaterialsAttribute( PolygonGroupPrim.Get(), Time );
 
 				if ( PolygonGroupPrimUEMaterialsAttribute.IsValidIndex( PolygonGroupPrimMaterialIndex ) )
@@ -153,6 +156,7 @@ namespace UsdGeomMeshTranslatorImpl
 				}
 				else
 				{
+					// We don't have a UE4 material stored, but we may have imported this material already, so check PrimPathsToAssets
 					TUsdStore< pxr::UsdPrim > MaterialPrim = UsdPrim.GetStage()->GetPrimAtPath( UnrealToUsd::ConvertPath( *ImportedMaterialSlotName.ToString() ).Get() );
 
 					if ( MaterialPrim.Get() )
@@ -162,9 +166,37 @@ namespace UsdGeomMeshTranslatorImpl
 				}
 			}
 
+			// Need to create a new material or reuse what is already on the mesh
 			if ( Material == nullptr )
 			{
-				UMaterialInstanceConstant* MaterialInstance = NewObject< UMaterialInstanceConstant >();
+				UMaterialInstanceConstant* MaterialInstance =
+					[&]()
+					{
+						UMaterialInstanceConstant* ExistingMaterialInstance = Cast< UMaterialInstanceConstant >( StaticMesh.GetMaterial( MaterialIndex ) );
+						const FString& SourcePrimPath = ImportedMaterialSlotName.ToString();
+
+						// Assuming that we own the material instance and that we can change it as we wish, reuse it
+						if ( ExistingMaterialInstance && ExistingMaterialInstance->GetOuter() == GetTransientPackage() )
+						{
+							if (UAssetImportData* ImportData = Cast<UAssetImportData>(ExistingMaterialInstance->AssetImportData))
+							{
+								if (ImportData->GetFirstFilename() == SourcePrimPath)
+								{
+									return ExistingMaterialInstance;
+								}
+							}
+						}
+
+						// Create new material
+						UMaterialInstanceConstant* NewMaterialInstance = NewObject< UMaterialInstanceConstant >();
+
+						UAssetImportData* ImportData = NewObject< UAssetImportData >(NewMaterialInstance, TEXT("AssetImportData"));
+						ImportData->UpdateFilenameOnly(SourcePrimPath);
+						NewMaterialInstance->AssetImportData = ImportData;
+
+						return NewMaterialInstance;
+					}();
+
 				if ( UsdToUnreal::ConvertDisplayColor( pxr::UsdGeomMesh( PolygonGroupPrim.Get() ), *MaterialInstance, pxr::UsdTimeCode( Time ) ) )
 				{
 					Material = MaterialInstance;
diff --git a/Engine/Plugins/Importers/USDImporter/Source/USDUtilities/Private/USDLayerUtils.cpp b/Engine/Plugins/Importers/USDImporter/Source/USDUtilities/Private/USDLayerUtils.cpp
--- a/Engine/Plugins/Importers/USDImporter/Source/USDUtilities/Private/USDLayerUtils.cpp
+++ b/Engine/Plugins/Importers/USDImporter/Source/USDUtilities/Private/USDLayerUtils.cpp
@@ -97,16 +97,10 @@ TUsdStore< pxr::SdfLayerRefPtr > UsdUtils::CreateNewLayer( TUsdStore< pxr::UsdSt
 	}
 
 	// New layer needs to be created and in the stage layer stack before we can edit it
-	UsdUtils::InsertSubLayer(ParentLayer, LayerFilePath);
+	UsdUtils::InsertSubLayer( ParentLayer, LayerFilePath );
 
 	UsdUtils::StartMonitoringErrors();
 	pxr::UsdEditContext UsdEditContext( UsdStage.Get(), LayerRef );
-	bool bHadErrors = UsdUtils::ShowErrorsAndStopMonitoring();
-
-	if (bHadErrors)
-	{
-		return {};
-	}
 
 	// Create default prim
 	FString PrimPath = TEXT("/") + FPaths::GetBaseFilename( UsdToUnreal::ConvertString( LayerRef->GetDisplayName() ) );
@@ -120,8 +114,12 @@ TUsdStore< pxr::SdfLayerRefPtr > UsdUtils::CreateNewLayer( TUsdStore< pxr::UsdSt
 		LayerRef->SetDefaultPrim( DefaultPrim.GetPrim().GetName() );
 	}
 
-	// Set up axis
-	UsdUtils::SetUsdStageAxis( UsdStage.Get(), pxr::UsdGeomTokens->z );
+	bool bHadErrors = UsdUtils::ShowErrorsAndStopMonitoring();
+
+	if (bHadErrors)
+	{
+		return {};
+	}
 
 	return TUsdStore< pxr::SdfLayerRefPtr >( LayerRef );
 }
