USD: Fix subsequence tracks not being added when adding sublayers;

Fix animations not reloading when layers were added or removed;
Fix sublayers not being parsed the first time around as they are dormant when the events are fired;
Fix failing to detect layer circular reference;

diff --git a/Engine/Plugins/Importers/USDImporter/Source/USDStage/Private/USDLevelSequenceHelper.cpp b/Engine/Plugins/Importers/USDImporter/Source/USDStage/Private/USDLevelSequenceHelper.cpp
--- a/Engine/Plugins/Importers/USDImporter/Source/USDStage/Private/USDLevelSequenceHelper.cpp
+++ b/Engine/Plugins/Importers/USDImporter/Source/USDStage/Private/USDLevelSequenceHelper.cpp
@@ -86,7 +86,7 @@ public:
 
 private:
 	/** Create a FLayerTimeInfo object with data from LayerHandle, and calls itself for all of its sublayers */
-	FLayerTimeInfo* RecursivelyCreateLayerTimeInfo(const TUsdStore<pxr::SdfLayerHandle>& LayerHandle, TSet<FString>& ExploredFiles);
+	FLayerTimeInfo* RecursivelyCreateLayerTimeInfo(const TUsdStore<pxr::SdfLayerRefPtr>& LayerHandle, TSet<FString>& ExploredFiles);
 
 	/**
 	 * Propagate timecodes for leaf FLayerTimeInfo all the way up to Info. We use this to calculate enveloping
@@ -399,7 +399,7 @@ void FUsdLevelSequenceHelperImpl::RebuildLayerTimeInfoCache(const pxr::UsdStageR
 	LayerTimeInfosByIdentifier.Reset();
 	LayerInfoIdentifierByLevelSequenceName.Reset();
 
-	TUsdStore<pxr::SdfLayerHandle> RootLayerHandle = MakeUsdStore<pxr::SdfLayerHandle>(UsdStage->GetRootLayer());
+	TUsdStore<pxr::SdfLayerRefPtr> RootLayerHandle = MakeUsdStore<pxr::SdfLayerRefPtr>(pxr::SdfLayer::FindOrOpen(UsdStage->GetRootLayer()->GetRealPath()));
 
 	// Used to catch N-item infinite loops (e.g. A/B/C/A/B/C/etc)
 	TSet<FString> ExploredPaths;
@@ -450,27 +450,27 @@ void FUsdLevelSequenceHelperImpl::InitializeActorTimeCodes(const pxr::UsdStageRe
 	UE_LOG(LogUsd, Verbose, TEXT("Initialized AUsdStageActor's Stage. StartTimeCode: %f, EndTimeCode: %f, TimeCodesPerSeconds: %f"), ValidStageActor->StartTimeCode, ValidStageActor->EndTimeCode, ValidStageActor->TimeCodesPerSecond);
 }
 
-FUsdLevelSequenceHelperImpl::FLayerTimeInfo* FUsdLevelSequenceHelperImpl::RecursivelyCreateLayerTimeInfo(const TUsdStore<pxr::SdfLayerHandle>& LayerHandleWrapper, TSet<FString>& ExploredFiles)
+FUsdLevelSequenceHelperImpl::FLayerTimeInfo* FUsdLevelSequenceHelperImpl::RecursivelyCreateLayerTimeInfo(const TUsdStore<pxr::SdfLayerRefPtr>& LayerHandleWrapper, TSet<FString>& ExploredFiles)
 {
-	const pxr::SdfLayerHandle& LayerHandle = LayerHandleWrapper.Get();
+	const pxr::SdfLayerRefPtr& LayerRefPtr = LayerHandleWrapper.Get();
 
 	// We want to keep track of all SdfLayers we find, but each layer should only appear once in a path to the root, or else we infinitely loop
-	FString Identifier = UsdToUnreal::ConvertString(LayerHandle->GetIdentifier());
-	if (ExploredFiles.Contains(Identifier))
+	FString FilePath = UsdToUnreal::ConvertString(LayerRefPtr->GetRealPath());
+	if (ExploredFiles.Contains(FilePath))
 	{
-		UE_LOG(LogUsd, Warning, TEXT("Detected infinite SubLayer recursion when visiting layer with identifier '%s'!"), *Identifier);
+		UE_LOG(LogUsd, Warning, TEXT("Detected infinite SubLayer recursion when visiting layer at path '%s'!"), *FilePath);
 		return nullptr;
 	}
-
-	FString FilePath = UsdToUnreal::ConvertString(LayerHandle->GetRealPath());
 	ExploredFiles.Add(FilePath);
 
+	FString Identifier = UsdToUnreal::ConvertString(LayerRefPtr->GetIdentifier());
+
 	FUsdLevelSequenceHelperImpl::FLayerTimeInfo LayerTimeInfo;
-	LayerTimeInfo.Identifier = Identifier;
-	LayerTimeInfo.FilePath = FilePath;
-	LayerTimeInfo.StartTimeCode      = LayerHandle->HasStartTimeCode() ? LayerHandle->GetStartTimeCode() : TOptional<double>();
-	LayerTimeInfo.EndTimeCode        = LayerHandle->HasEndTimeCode() ? LayerHandle->GetEndTimeCode() : TOptional<double>();
-	LayerTimeInfo.TimeCodesPerSecond = LayerHandle->HasTimeCodesPerSecond() ? LayerHandle->GetTimeCodesPerSecond() : TOptional<double>();
+	LayerTimeInfo.Identifier		 = Identifier;
+	LayerTimeInfo.FilePath			 = FilePath;
+	LayerTimeInfo.StartTimeCode      = LayerRefPtr->HasStartTimeCode() ? LayerRefPtr->GetStartTimeCode() : TOptional<double>();
+	LayerTimeInfo.EndTimeCode        = LayerRefPtr->HasEndTimeCode() ? LayerRefPtr->GetEndTimeCode() : TOptional<double>();
+	LayerTimeInfo.TimeCodesPerSecond = LayerRefPtr->HasTimeCodesPerSecond() ? LayerRefPtr->GetTimeCodesPerSecond() : TOptional<double>();
 
 	UE_LOG(LogUsd, Verbose, TEXT("Creating layer time info for layer '%s'. Original timecodes: '%s', '%s' and '%s'"),
 		*LayerTimeInfo.Identifier,
@@ -482,17 +482,18 @@ FUsdLevelSequenceHelperImpl::FLayerTimeInfo* FUsdLevelSequenceHelperImpl::Recurs
 	TUsdStore<std::vector<pxr::SdfLayerOffset>> SubLayerOffsetScales;
 	{
 		FScopedUsdAllocs USDAllocs;
-		SubLayerPaths = LayerHandle->GetSubLayerPaths();
-		SubLayerOffsetScales = LayerHandle->GetSubLayerOffsets();
+		SubLayerPaths = LayerRefPtr->GetSubLayerPaths();
+		SubLayerOffsetScales = LayerRefPtr->GetSubLayerOffsets();
 	}
 
-	for (int64 LayerIndex = 0; LayerIndex < static_cast<int64>(LayerHandle->GetNumSubLayerPaths()); ++LayerIndex)
+	for (int64 LayerIndex = 0; LayerIndex < static_cast<int64>(LayerRefPtr->GetNumSubLayerPaths()); ++LayerIndex)
 	{
-		TUsdStore<pxr::SdfLayerHandle> SubLayer;
+		TUsdStore<pxr::SdfLayerRefPtr> SubLayer;
 		{
 			FScopedUsdAllocs USDAllocs;
-			std::string RelativePath = pxr::SdfComputeAssetPathRelativeToLayer(LayerHandle, SubLayerPaths.Get()[LayerIndex]);
-			SubLayer = LayerHandle->Find(RelativePath);
+
+			std::string RelativePath = pxr::SdfComputeAssetPathRelativeToLayer(LayerRefPtr, SubLayerPaths.Get()[LayerIndex]);
+			SubLayer = pxr::SdfLayer::FindOrOpen(RelativePath);
 		}
 
 		if (!SubLayer.Get())
@@ -609,10 +610,10 @@ void FUsdLevelSequenceHelperImpl::UpdateInfoFromSection(FUsdLevelSequenceHelperI
 		return;
 	}
 
-	TUsdStore<pxr::SdfLayerHandle> Layer;
+	TUsdStore<pxr::SdfLayerRefPtr> Layer;
 	{
 		FScopedUsdAllocs Allocs;
-		Layer = MakeUsdStore<pxr::SdfLayerHandle>(pxr::SdfLayer::Find(UnrealToUsd::ConvertString(*Info->Identifier).Get()));
+		Layer = MakeUsdStore<pxr::SdfLayerRefPtr>(pxr::SdfLayer::FindOrOpen(UnrealToUsd::ConvertString(*Info->Identifier).Get()));
 		if (!Layer.Get())
 		{
 			UE_LOG(LogUsd, Warning, TEXT("Failed to update sublayer '%s'"), *Info->Identifier);
diff --git a/Engine/Plugins/Importers/USDImporter/Source/USDStage/Private/USDStageActor.cpp b/Engine/Plugins/Importers/USDImporter/Source/USDStage/Private/USDStageActor.cpp
--- a/Engine/Plugins/Importers/USDImporter/Source/USDStage/Private/USDStageActor.cpp
+++ b/Engine/Plugins/Importers/USDImporter/Source/USDStage/Private/USDStageActor.cpp
@@ -205,6 +205,18 @@ AUsdStageActor::AUsdStageActor()
 				const pxr::SdfChangeList::EntryList& ChangeList = ChangeVecItem.second.GetEntryList();
 				for (const std::pair<pxr::SdfPath, pxr::SdfChangeList::Entry>& Change : ChangeList)
 				{
+					for (const pxr::SdfChangeList::Entry::SubLayerChange& SubLayerChange : Change.second.subLayerChanges)
+					{
+						const pxr::SdfChangeList::SubLayerChangeType ChangeType = SubLayerChange.second;
+						if (ChangeType == pxr::SdfChangeList::SubLayerChangeType::SubLayerAdded ||
+							ChangeType == pxr::SdfChangeList::SubLayerChangeType::SubLayerRemoved)
+						{
+							UE_LOG(LogUsd, Verbose, TEXT("Reloading animations because layer '%s' was added/removed"), *UsdToUnreal::ConvertString(SubLayerChange.first));
+							ReloadAnimations();
+							return;
+						}
+					}
+
 					const pxr::SdfChangeList::Entry::_Flags& Flags = Change.second.flags;
 					if (Flags.didReloadContent)
 					{
@@ -586,8 +598,11 @@ void AUsdStageActor::ReloadAnimations()
 
 	if ( HasAutorithyOverStage() )
 	{
-		// The sequencer won't update on its own, so let's at least force it closed
-		GEditor->GetEditorSubsystem<UAssetEditorSubsystem>()->CloseAllEditorsForAsset(LevelSequence);
+		if (LevelSequence)
+		{
+			// The sequencer won't update on its own, so let's at least force it closed
+			GEditor->GetEditorSubsystem<UAssetEditorSubsystem>()->CloseAllEditorsForAsset(LevelSequence);
+		}
 
 		LevelSequence = nullptr;
 		SubLayerLevelSequencesByIdentifier.Reset();
