Display the first 300 characters of the attribute values since super long STextBox take a long time to calculate their desired size.

diff --git a/Engine/Plugins/Importers/USDImporter/Source/USDStageEditor/Private/Widgets/SUSDPrimPropertiesList.cpp b/Engine/Plugins/Importers/USDImporter/Source/USDStageEditor/Private/Widgets/SUSDPrimPropertiesList.cpp
--- a/Engine/Plugins/Importers/USDImporter/Source/USDStageEditor/Private/Widgets/SUSDPrimPropertiesList.cpp
+++ b/Engine/Plugins/Importers/USDImporter/Source/USDStageEditor/Private/Widgets/SUSDPrimPropertiesList.cpp
@@ -240,10 +240,10 @@ void SUsdPrimPropertyRow::SetUsdPrimProperty( const TSharedPtr< FUsdPrimProperty
 
 TSharedRef< SWidget > SUsdPrimPropertyRow::GenerateTextWidget(const TAttribute<FText>& Attribute)
 {
-	return SNew(STextBlock)
-	.Text(Attribute)
-	.Font(FEditorStyle::GetFontStyle(UsdPrimPropertiesListConstants::NormalFont))
-	.Margin(UsdPrimPropertiesListConstants::RightRowPadding);
+	return SNew( STextBlock )
+	.Text( Attribute )
+	.Font( FEditorStyle::GetFontStyle( UsdPrimPropertiesListConstants::NormalFont ) )
+	.Margin( UsdPrimPropertiesListConstants::RightRowPadding );
 }
 
 void SUsdPrimPropertiesList::Construct( const FArguments& InArgs, const TCHAR* InPrimPath )
@@ -358,6 +358,14 @@ void SUsdPrimPropertiesList::GeneratePropertiesList( const TCHAR* InPrimPath )
 			PrimAttribute.Get( &VtValue, TimeCode );
 			FString AttributeValue = UsdToUnreal::ConvertString( pxr::TfStringify( VtValue ).c_str() );
 
+			// STextBlock can get very slow calculating its desired size for very long string so chop it if needed
+			const int32 MaxValueLength = 300;
+			if ( AttributeValue.Len() > MaxValueLength )
+			{
+				AttributeValue.LeftInline( MaxValueLength );
+				AttributeValue.Append( TEXT("...") );
+			}
+
 			PrimAttributeProperty.Value = MoveTemp( AttributeValue );
 			PrimProperties.Add( MakeSharedUnreal< FUsdPrimProperty >( MoveTemp( PrimAttributeProperty ) ) );
 		}
