USD: Fix warning when constructing an empty SdfPath

diff --git a/Engine/Plugins/Importers/USDImporter/Source/USDUtilities/Public/USDTypesConversion.h b/Engine/Plugins/Importers/USDImporter/Source/USDUtilities/Public/USDTypesConversion.h
--- a/Engine/Plugins/Importers/USDImporter/Source/USDUtilities/Public/USDTypesConversion.h
+++ b/Engine/Plugins/Importers/USDImporter/Source/USDUtilities/Public/USDTypesConversion.h
@@ -178,6 +178,12 @@ namespace UnrealToUsd
 
 	static TUsdStore< pxr::SdfPath > ConvertPath( const TCHAR* InString )
 	{
+		// SdfPath throws a warning if constructed with the empty string
+		if (!InString || *InString == '\0')
+		{
+			return MakeUsdStore< pxr::SdfPath >();
+		}
+
 		return MakeUsdStore< pxr::SdfPath >( TCHAR_TO_ANSI( InString ) );
 	}
 
