Fixed USD adding the unrealMaterials attribute even if no materials override exist on the MeshComponent.
Also clear the unrealMaterials attribute if the overrides were removed on the MeshComponent.

diff --git a/Engine/Plugins/Importers/USDImporter/Source/USDUtilities/Private/USDPrimConversion.cpp b/Engine/Plugins/Importers/USDImporter/Source/USDUtilities/Private/USDPrimConversion.cpp
--- a/Engine/Plugins/Importers/USDImporter/Source/USDUtilities/Private/USDPrimConversion.cpp
+++ b/Engine/Plugins/Importers/USDImporter/Source/USDUtilities/Private/USDPrimConversion.cpp
@@ -159,30 +159,41 @@ bool UnrealToUsd::ConvertMeshComponent( const pxr::UsdStageRefPtr& Stage, const
 		return false;
 	}
 
-	if ( MeshComponent->GetNumMaterials() > 0 || UsdPrim.HasAttribute( UnrealIdentifiers::MaterialAssignments ) )
+	const bool bHasMaterialAttribute = UsdPrim.HasAttribute( UnrealIdentifiers::MaterialAssignments );
+
+	if ( MeshComponent->GetNumMaterials() > 0 || bHasMaterialAttribute )
 	{
-		if ( pxr::UsdAttribute UEMaterialsAttribute = UsdPrim.CreateAttribute( UnrealIdentifiers::MaterialAssignments, pxr::SdfValueTypeNames->StringArray ) )
-		{
-			FScopedUsdAllocs UsdAllocs;
+		FScopedUsdAllocs UsdAllocs;
+
+		pxr::VtArray< std::string > UEMaterials;
 
-			pxr::VtArray< std::string > UEMaterials = UsdUtils::GetUsdValue< pxr::VtArray< std::string > >( UEMaterialsAttribute );
-			UEMaterials.clear();
+		bool bHasUEMaterialAssignements = false;
 
-			for ( int32 MaterialIndex = 0; MaterialIndex < MeshComponent->GetNumMaterials(); ++MaterialIndex )
+		for ( int32 MaterialIndex = 0; MaterialIndex < MeshComponent->GetNumMaterials(); ++MaterialIndex )
+		{
+			if ( UMaterialInterface* AssignedMaterial = MeshComponent->GetMaterial( MaterialIndex ) )
 			{
-				if ( UMaterialInterface* AssignedMaterial = MeshComponent->GetMaterial( MaterialIndex ) )
+				FString AssignedMaterialPathName;
+				if ( AssignedMaterial->GetOutermost() != GetTransientPackage() )
 				{
-					FString AssignedMaterialPathName;
-					if ( AssignedMaterial->GetOutermost() != GetTransientPackage() )
-					{
-						AssignedMaterialPathName = AssignedMaterial->GetPathName();
-					}
-
-					UEMaterials.push_back( UnrealToUsd::ConvertString( *AssignedMaterialPathName ).Get() );
+					AssignedMaterialPathName = AssignedMaterial->GetPathName();
+					bHasUEMaterialAssignements = true;
 				}
+
+				UEMaterials.push_back( UnrealToUsd::ConvertString( *AssignedMaterialPathName ).Get() );
 			}
+		}
 
-			UEMaterialsAttribute.Set( UEMaterials );
+		if ( bHasUEMaterialAssignements )
+		{
+			if ( pxr::UsdAttribute UEMaterialsAttribute = UsdPrim.CreateAttribute( UnrealIdentifiers::MaterialAssignments, pxr::SdfValueTypeNames->StringArray ) )
+			{
+				UEMaterialsAttribute.Set( UEMaterials );
+			}
+		}
+		else if ( bHasMaterialAttribute )
+		{
+			UsdPrim.GetAttribute( UnrealIdentifiers::MaterialAssignments ).Clear();
 		}
 	}
 
