USD: Prim naming bugs.

Fix escaping out of setting a Prim name leaving a dummy row on the tree view;
Fix refreshing (e.g. when adding prims) resetting all item collapse states;
Prevent user from creating a prim with a name that would lead to a USD error;

diff --git a/Engine/Plugins/Importers/USDImporter/Source/USDStageEditor/Private/SUSDStageTreeView.cpp b/Engine/Plugins/Importers/USDImporter/Source/USDStageEditor/Private/SUSDStageTreeView.cpp
--- a/Engine/Plugins/Importers/USDImporter/Source/USDStageEditor/Private/SUSDStageTreeView.cpp
+++ b/Engine/Plugins/Importers/USDImporter/Source/USDStageEditor/Private/SUSDStageTreeView.cpp
@@ -2,6 +2,7 @@
 
 #include "SUSDStageTreeView.h"
 
+#include "UnrealUSDWrapper.h"
 #include "USDLayerUtils.h"
 #include "USDStageActor.h"
 #include "USDStageModule.h"
@@ -195,6 +196,9 @@ public:
 	DECLARE_DELEGATE_TwoParams( FOnPrimNameCommitted, const FUsdStageTreeItemRef&, const FText& );
 	FOnPrimNameCommitted OnPrimNameCommitted;
 
+	DECLARE_DELEGATE_ThreeParams( FOnPrimNameUpdated, const FUsdStageTreeItemRef&, const FText&, FText& );
+	FOnPrimNameUpdated OnPrimNameUpdated;
+
 	TWeakPtr< SUsdStageTreeView > OwnerTree;
 
 	virtual TSharedRef< SWidget > GenerateWidget( const TSharedPtr< IUsdTreeViewItem > InTreeItem ) override
@@ -212,6 +216,7 @@ public:
 			.Font( FEditorStyle::GetFontStyle( "ContentBrowser.SourceTreeItemFont" ) )
 			.ColorAndOpacity( this, &FUsdStageNameColumn::GetTextColor, TreeItem )
 			.OnTextCommitted( this, &FUsdStageNameColumn::OnTextCommitted, TreeItem )
+			.OnVerifyTextChanged( this, &FUsdStageNameColumn::OnTextUpdated, TreeItem )
 			.IsReadOnly_Lambda( [ TreeItem ]()
 			{
 				return !TreeItem || TreeItem->UsdPrim.Get();
@@ -233,6 +238,18 @@ protected:
 		OnPrimNameCommitted.ExecuteIfBound( TreeItem.ToSharedRef(), InPrimName );
 	}
 
+	bool OnTextUpdated(const FText& InPrimName, FText& ErrorMessage, TSharedPtr< FUsdStageTreeItem > TreeItem)
+	{
+		if (!TreeItem)
+		{
+			return false;
+		}
+
+		OnPrimNameUpdated.ExecuteIfBound( TreeItem.ToSharedRef(), InPrimName, ErrorMessage );
+
+		return ErrorMessage.IsEmpty();
+	}
+
 	FSlateColor GetTextColor( TSharedPtr< FUsdStageTreeItem > TreeItem ) const
 	{
 		FSlateColor TextColor = FLinearColor::White;
@@ -411,6 +428,22 @@ void SUsdStageTreeView::Construct( const FArguments& InArgs, AUsdStageActor* InU
 		this->OnPrimSelected.ExecuteIfBound( SelectedPrimPath );
 	} );
 
+	OnExpansionChanged = FOnExpansionChanged::CreateLambda([this]( const FUsdStageTreeItemPtr& UsdStageTreeItem, bool bIsExpanded)
+	{
+		if (!UsdStageTreeItem)
+		{
+			return;
+		}
+
+		const pxr::UsdPrim& Prim = UsdStageTreeItem->UsdPrim.Get();
+		if (!Prim)
+		{
+			return;
+		}
+
+		TreeItemExpansionStates.Add(UsdToUnreal::ConvertPath(Prim.GetPrimPath()), bIsExpanded);
+	});
+
 	OnPrimSelected = InArgs._OnPrimSelected;
 
 	Refresh( InUsdStageActor );
@@ -433,6 +466,11 @@ void SUsdStageTreeView::Refresh( AUsdStageActor* InUsdStageActor )
 {
 	RootItems.Empty();
 
+	if (UsdStageActor.Get() != InUsdStageActor)
+	{
+		TreeItemExpansionStates.Reset();
+	}
+
 	UsdStageActor = InUsdStageActor;
 
 	if ( !UsdStageActor.IsValid() )
@@ -452,6 +490,8 @@ void SUsdStageTreeView::Refresh( AUsdStageActor* InUsdStageActor )
 			}
 		}
 	}
+
+	RestoreExpansionStates();
 }
 
 void SUsdStageTreeView::RefreshPrim( const FString& PrimPath, bool bResync )
@@ -550,6 +590,7 @@ void SUsdStageTreeView::SetupColumns()
 		PrimNameColumn->OwnerTree = SharedThis( this );
 		PrimNameColumn->bIsMainColumn = true;
 		PrimNameColumn->OnPrimNameCommitted.BindRaw( this, &SUsdStageTreeView::OnPrimNameCommitted );
+		PrimNameColumn->OnPrimNameUpdated.BindRaw( this, &SUsdStageTreeView::OnPrimNameUpdated );
 
 		SHeaderRow::FColumn::FArguments PrimNameColumnArguments;
 		PrimNameColumnArguments.FillWidth( 70.f );
@@ -817,6 +858,36 @@ bool SUsdStageTreeView::CanExecutePrimAction() const
 	return bHasPrimSpec;
 }
 
+void SUsdStageTreeView::RequestListRefresh()
+{
+	SUsdTreeView< FUsdStageTreeItemRef >::RequestListRefresh();
+	RestoreExpansionStates();
+}
+
+void SUsdStageTreeView::RestoreExpansionStates()
+{
+	std::function<void(const FUsdStageTreeItemRef&)> SetExpansionRecursive = [&](const FUsdStageTreeItemRef& Item)
+	{
+		if (const pxr::UsdPrim& Prim = Item->UsdPrim.Get())
+		{
+			if (bool* bFoundExpansionState = TreeItemExpansionStates.Find(UsdToUnreal::ConvertPath(Prim.GetPrimPath())))
+			{
+				SetItemExpansion(Item, *bFoundExpansionState);
+			}
+		}
+
+		for (const FUsdStageTreeItemRef& Child : Item->Children)
+		{
+			SetExpansionRecursive(Child);
+		}
+	};
+
+	for (const FUsdStageTreeItemRef& RootItem : RootItems)
+	{
+		SetExpansionRecursive(RootItem);
+	}
+}
+
 void SUsdStageTreeView::OnToggleAllPayloads( EPayloadsTrigger PayloadsTrigger )
 {
 	TArray< FUsdStageTreeItemRef > MySelectedItems = GetSelectedItems();
@@ -891,6 +962,20 @@ void SUsdStageTreeView::OnPrimNameCommitted( const FUsdStageTreeItemRef& TreeIte
 {
 	if ( InPrimName.IsEmptyOrWhitespace() )
 	{
+		// Escaped out of initially setting a prim name
+		if (!TreeItem->UsdPrim.Get())
+		{
+			if (FUsdStageTreeItem* Parent = TreeItem->ParentItem)
+			{
+				TreeItem->ParentItem->Children.Remove( TreeItem );
+			}
+			else
+			{
+				RootItems.Remove(TreeItem);
+			}
+
+			RequestTreeRefresh();
+		}
 		return;
 	}
 
@@ -927,6 +1012,44 @@ void SUsdStageTreeView::OnPrimNameCommitted( const FUsdStageTreeItemRef& TreeIte
 	}
 }
 
+void SUsdStageTreeView::OnPrimNameUpdated(const FUsdStageTreeItemRef& TreeItem, const FText& InPrimName, FText& ErrorMessage)
+{
+	FString NameStr = InPrimName.ToString();
+	IUsdPrim::IsValidPrimName(NameStr, ErrorMessage);
+	if (!ErrorMessage.IsEmpty())
+	{
+		return;
+	}
+
+	{
+		FScopedUsdAllocs Allocs;
+
+		const pxr::UsdStageRefPtr& Stage = TreeItem->UsdStage.Get();
+		if (!Stage)
+		{
+			return;
+		}
+
+		pxr::SdfPath ParentPrimPath;
+		if (TreeItem->ParentItem)
+		{
+			ParentPrimPath = TreeItem->ParentItem->UsdPrim.Get().GetPrimPath();
+		}
+		else
+		{
+			ParentPrimPath = pxr::SdfPath::AbsoluteRootPath();
+		}
+
+		pxr::SdfPath NewPrimPath = ParentPrimPath.AppendChild(pxr::TfToken(UnrealToUsd::ConvertString(*NameStr).Get()));
+		const pxr::UsdPrim& Prim = Stage->GetPrimAtPath(NewPrimPath);
+		if (Prim)
+		{
+			ErrorMessage = LOCTEXT("DuplicatePrimName", "A Prim with this name already exists!");
+			return;
+		}
+	}
+}
+
 #endif // #if USE_USD_SDK
 
 #undef LOCTEXT_NAMESPACE
diff --git a/Engine/Plugins/Importers/USDImporter/Source/USDStageEditor/Private/SUSDStageTreeView.h b/Engine/Plugins/Importers/USDImporter/Source/USDStageEditor/Private/SUSDStageTreeView.h
--- a/Engine/Plugins/Importers/USDImporter/Source/USDStageEditor/Private/SUSDStageTreeView.h
+++ b/Engine/Plugins/Importers/USDImporter/Source/USDStageEditor/Private/SUSDStageTreeView.h
@@ -51,6 +51,7 @@ private:
 	virtual void OnTreeItemScrolledIntoView( FUsdStageTreeItemRef TreeItem, const TSharedPtr< ITableRow >& Widget ) override ;
 
 	void OnPrimNameCommitted( const FUsdStageTreeItemRef& TreeItem, const FText& InPrimName );
+	void OnPrimNameUpdated( const FUsdStageTreeItemRef& TreeItem, const FText& InPrimName, FText& ErrorMessage );
 
 	virtual void SetupColumns() override;
 	TSharedPtr< SWidget > ConstructPrimContextMenu();
@@ -67,9 +68,16 @@ private:
 
 	TOptional< FString > BrowseFile();
 
+	/** Uses TreeItemExpansionStates to travel the tree and call SetItemExpansion */
+	void RestoreExpansionStates();
+	virtual void RequestListRefresh() override;
+
 	TWeakObjectPtr< AUsdStageActor > UsdStageActor;
 	TWeakPtr< FUsdStageTreeItem > PendingRenameItem;
 
+	// So that we can store these across refreshes
+	TMap< FString, bool > TreeItemExpansionStates;
+
 	FOnPrimSelected OnPrimSelected;
 };
 
diff --git a/Engine/Plugins/Importers/USDImporter/Source/UnrealUSDWrapper/Private/UnrealUSDWrapper.cpp b/Engine/Plugins/Importers/USDImporter/Source/UnrealUSDWrapper/Private/UnrealUSDWrapper.cpp
--- a/Engine/Plugins/Importers/USDImporter/Source/UnrealUSDWrapper/Private/UnrealUSDWrapper.cpp
+++ b/Engine/Plugins/Importers/USDImporter/Source/UnrealUSDWrapper/Private/UnrealUSDWrapper.cpp
@@ -5,8 +5,11 @@
 #include "USDLog.h"
 #include "USDMemory.h"
 
+#include "Internationalization/Regex.h"
 #include "Modules/ModuleManager.h"
 
+#define LOCTEXT_NAMESPACE "UnrealUSDWrapper"
+
 #if USE_USD_SDK
 
 #include "USDIncludesStart.h"
@@ -439,6 +442,33 @@ int FUsdAttribute::GetArraySize( const pxr::UsdAttribute& Attribute )
 	return Value.IsArrayValued() ? (int)Value.GetArraySize() : -1;
 
 }
+
+bool IUsdPrim::IsValidPrimName(const FString& Name, FText& OutReason)
+{
+	if (Name.IsEmpty())
+	{
+		OutReason = LOCTEXT("EmptyStringInvalid", "Empty string is not a valid name!");
+		return false;
+	}
+
+	const FString InvalidCharacters = TEXT("\\W");
+	FRegexPattern RegexPattern( InvalidCharacters );
+	FRegexMatcher RegexMatcher( RegexPattern, Name );
+	if (RegexMatcher.FindNext())
+	{
+		OutReason = LOCTEXT("InvalidCharacter", "Can only use letters, numbers and underscore!");
+		return false;
+	}
+
+	if (Name.Left(1).IsNumeric())
+	{
+		OutReason = LOCTEXT("InvalidFirstCharacter", "First character cannot be a number!");
+		return false;
+	}
+
+	return true;
+}
+
 EUsdPurpose IUsdPrim::GetPurpose( const UsdPrim& Prim, bool bComputed )
 {
 	UsdGeomImageable Geom(Prim);
@@ -1087,3 +1117,5 @@ public:
 };
 
 IMPLEMENT_MODULE_USD(FUnrealUSDWrapperModule, UnrealUSDWrapper);
+
+#undef LOCTEXT_NAMESPACE
\ No newline at end of file
diff --git a/Engine/Plugins/Importers/USDImporter/Source/UnrealUSDWrapper/Public/UnrealUSDWrapper.h b/Engine/Plugins/Importers/USDImporter/Source/UnrealUSDWrapper/Public/UnrealUSDWrapper.h
--- a/Engine/Plugins/Importers/USDImporter/Source/UnrealUSDWrapper/Public/UnrealUSDWrapper.h
+++ b/Engine/Plugins/Importers/USDImporter/Source/UnrealUSDWrapper/Public/UnrealUSDWrapper.h
@@ -207,6 +207,8 @@ class IUsdPrim
 {
 public:
 #if USE_USD_SDK
+	static UNREALUSDWRAPPER_API bool IsValidPrimName(const FString& Name, FText& OutReason);
+
 	static UNREALUSDWRAPPER_API EUsdPurpose GetPurpose(const pxr::UsdPrim& Prim, bool bComputed = true);
 	static UNREALUSDWRAPPER_API FName GetPurposeName(EUsdPurpose Purpose);
 
