USD added support for units (metersPerUnit metadata)

diff --git a/Engine/Plugins/Importers/USDImporter/Source/USDImporter/Private/StaticMeshImporter.cpp b/Engine/Plugins/Importers/USDImporter/Source/USDImporter/Private/StaticMeshImporter.cpp
--- a/Engine/Plugins/Importers/USDImporter/Source/USDImporter/Private/StaticMeshImporter.cpp
+++ b/Engine/Plugins/Importers/USDImporter/Source/USDImporter/Private/StaticMeshImporter.cpp
@@ -179,9 +179,9 @@ UStaticMesh* FUSDStaticMeshImporter::ImportStaticMesh(FUsdImportContext& ImportC
 {
 	const pxr::UsdPrim& Prim = *PrimToImport.Prim;
 
-	const pxr::TfToken StageUpAxis = UsdUtils::GetUsdStageAxis( ImportContext.Stage.Get() );
+	const UsdToUnreal::FUsdStageInfo StageInfo( ImportContext.Stage.Get() );
 
-	FTransform PrimToWorld = ImportContext.bApplyWorldTransformToGeometry ? UsdToUnreal::ConvertMatrix(StageUpAxis, IUsdPrim::GetLocalTransform( Prim )) : FTransform::Identity;
+	FTransform PrimToWorld = ImportContext.bApplyWorldTransformToGeometry ? UsdToUnreal::ConvertMatrix(StageInfo, IUsdPrim::GetLocalTransform( Prim )) : FTransform::Identity;
 
 	FTransform FinalTransform = PrimToWorld;
 	if (ImportContext.ImportOptions->Scale != 1.0f)
diff --git a/Engine/Plugins/Importers/USDImporter/Source/USDImporter/Private/USDPrimResolver.cpp b/Engine/Plugins/Importers/USDImporter/Source/USDImporter/Private/USDPrimResolver.cpp
--- a/Engine/Plugins/Importers/USDImporter/Source/USDImporter/Private/USDPrimResolver.cpp
+++ b/Engine/Plugins/Importers/USDImporter/Source/USDImporter/Private/USDPrimResolver.cpp
@@ -232,7 +232,7 @@ AActor* UUSDPrimResolver::SpawnActor(FUSDSceneImportContext& ImportContext, cons
 					SpawnedActor = ImportContext.World->SpawnActor( AActor::StaticClass() );
 				}
 
-				const pxr::TfToken StageUpAxis = UsdUtils::GetUsdStageAxis( *ImportContext.Stage );
+				const UsdToUnreal::FUsdStageInfo StageInfo( *ImportContext.Stage );
 
 				int32 AssetIndex = 0;
 				for ( UObject* ImportedAsset : ImportedAssets )
@@ -255,7 +255,7 @@ AActor* UUSDPrimResolver::SpawnActor(FUSDSceneImportContext& ImportContext, cons
 					for ( int32 ParentPrimIndex = ParentPrims.Num() - 1; ParentPrimIndex >= 0; --ParentPrimIndex )
 					{
 						ParentPrim = ParentPrims[ ParentPrimIndex ];
-						LocalTransform = UsdToUnreal::ConvertMatrix( StageUpAxis, IUsdPrim::GetLocalTransform( ParentPrim ) ) * LocalTransform;
+						LocalTransform = UsdToUnreal::ConvertMatrix( StageInfo, IUsdPrim::GetLocalTransform( ParentPrim ) ) * LocalTransform;
 					}
 
 					FName ComponentBaseName = *UsdToUnreal::ConvertString( UsdAssetPrimToImport.Prim.Get().GetName().GetString().c_str() );
@@ -268,7 +268,7 @@ AActor* UUSDPrimResolver::SpawnActor(FUSDSceneImportContext& ImportContext, cons
 					// Don't add the prim transform if its the same prim used for the actor as it's already accounted for in the ActorTransform
 					if ( UsdAssetPrimToImport.Prim.Get() != SpawnData.ActorPrim.Get() )
 					{
-						LocalTransform = UsdToUnreal::ConvertMatrix( StageUpAxis, IUsdPrim::GetLocalTransform( UsdAssetPrimToImport.Prim.Get() ) ) * LocalTransform;
+						LocalTransform = UsdToUnreal::ConvertMatrix( StageInfo, IUsdPrim::GetLocalTransform( UsdAssetPrimToImport.Prim.Get() ) ) * LocalTransform;
 					}
 
 					StaticMeshComponent->SetRelativeTransform( LocalTransform );
@@ -453,7 +453,7 @@ void UUSDPrimResolver::FindActorsToSpawn_Recursive(FUSDSceneImportContext& Impor
 
 		FName PrimName = UsdToUnreal::ConvertName(Prim.Get().GetName().GetString());
 		SpawnData.ActorName = PrimName;
-		SpawnData.WorldTransform = UsdToUnreal::ConvertMatrix( UsdUtils::GetUsdStageAxis( *ImportContext.Stage ), IUsdPrim::GetLocalTransform( *Prim ) );
+		SpawnData.WorldTransform = UsdToUnreal::ConvertMatrix( UsdToUnreal::FUsdStageInfo( *ImportContext.Stage ), IUsdPrim::GetLocalTransform( *Prim ) );
 		SpawnData.AttachParentPrim = ParentPrim;
 		SpawnData.ActorPrim = Prim;
 
diff --git a/Engine/Plugins/Importers/USDImporter/Source/USDImporter/Private/USDPrimResolverKind.cpp b/Engine/Plugins/Importers/USDImporter/Source/USDImporter/Private/USDPrimResolverKind.cpp
--- a/Engine/Plugins/Importers/USDImporter/Source/USDImporter/Private/USDPrimResolverKind.cpp
+++ b/Engine/Plugins/Importers/USDImporter/Source/USDImporter/Private/USDPrimResolverKind.cpp
@@ -64,7 +64,7 @@ void UUSDPrimResolverKind::FindActorsToSpawn_Recursive(FUSDSceneImportContext& I
 			UE_LOG(LogUsd, Log, TEXT("Adding %s Actor with %d meshes"), *PrimName.ToString(), SpawnData.AssetsToImport.Num());
 		}
 
-		SpawnData.WorldTransform = UsdToUnreal::ConvertMatrix( UsdUtils::GetUsdStageAxis( *ImportContext.Stage ), IUsdPrim::GetLocalTransform( *Prim ) );
+		SpawnData.WorldTransform = UsdToUnreal::ConvertMatrix( UsdToUnreal::FUsdStageInfo( *ImportContext.Stage ), IUsdPrim::GetLocalTransform( *Prim ) );
 		SpawnData.ActorPrim = Prim;
 		SpawnData.ActorName = PrimName;
 		SpawnData.AttachParentPrim = GroupParent;
@@ -86,7 +86,7 @@ void UUSDPrimResolverKind::FindActorsToSpawn_Recursive(FUSDSceneImportContext& I
 		SpawnData.ActorPrim = Prim;
 		SpawnData.ActorName = PrimName;
 
-		SpawnData.WorldTransform = UsdToUnreal::ConvertMatrix( UsdUtils::GetUsdStageAxis( *ImportContext.Stage ), IUsdPrim::GetLocalTransform( *Prim ) );
+		SpawnData.WorldTransform = UsdToUnreal::ConvertMatrix( UsdToUnreal::FUsdStageInfo( *ImportContext.Stage ), IUsdPrim::GetLocalTransform( *Prim ) );
 		SpawnData.AttachParentPrim = GroupParent;
 
 		// New parent of all children is this prim
diff --git a/Engine/Plugins/Importers/USDImporter/Source/USDSchemas/Private/USDGeomPointInstancerTranslator.cpp b/Engine/Plugins/Importers/USDImporter/Source/USDSchemas/Private/USDGeomPointInstancerTranslator.cpp
--- a/Engine/Plugins/Importers/USDImporter/Source/USDSchemas/Private/USDGeomPointInstancerTranslator.cpp
+++ b/Engine/Plugins/Importers/USDImporter/Source/USDSchemas/Private/USDGeomPointInstancerTranslator.cpp
@@ -37,7 +37,7 @@ namespace UsdGeomPointInstancerTranslatorImpl
 			return false;
 		}
 
-		pxr::TfToken UpAxis = UsdUtils::GetUsdStageAxis( Stage );
+		UsdToUnreal::FUsdStageInfo StageInfo( Stage );
 
 		int32 Index = 0;
 
@@ -47,7 +47,7 @@ namespace UsdGeomPointInstancerTranslatorImpl
 		{
 			if ( ProtoIndices[ Index ] == ProtoIndex )
 			{
-				FTransform InstanceTransform = UsdToUnreal::ConvertMatrix( UpAxis, UsdMatrix );
+				FTransform InstanceTransform = UsdToUnreal::ConvertMatrix( StageInfo, UsdMatrix );
 
 				HismComponent.AddInstance( InstanceTransform );
 			}
diff --git a/Engine/Plugins/Importers/USDImporter/Source/USDSchemas/Private/USDSkelRootTranslator.cpp b/Engine/Plugins/Importers/USDImporter/Source/USDSchemas/Private/USDSkelRootTranslator.cpp
--- a/Engine/Plugins/Importers/USDImporter/Source/USDSchemas/Private/USDSkelRootTranslator.cpp
+++ b/Engine/Plugins/Importers/USDImporter/Source/USDSchemas/Private/USDSkelRootTranslator.cpp
@@ -206,7 +206,7 @@ void FUsdSkelRootTranslator::UpdateComponents( USceneComponent* SceneComponent )
 			return;
 		}
 
-		const pxr::TfToken StageUpAxis = UsdUtils::GetUsdStageAxis( Prim.GetStage() );
+		const UsdToUnreal::FUsdStageInfo StageInfo( Prim.GetStage() );
 
 		pxr::UsdGeomXformable Xformable( Prim );
 
@@ -248,7 +248,7 @@ void FUsdSkelRootTranslator::UpdateComponents( USceneComponent* SceneComponent )
 					for ( uint32 BoneIndex = 0; BoneIndex < UsdBoneTransforms.Get().size(); ++BoneIndex )
 					{
 						const pxr::GfMatrix4d& UsdMatrix = UsdBoneTransforms.Get()[ BoneIndex ];
-						PoseableMeshComponent->BoneSpaceTransforms[ BoneIndex ] = UsdToUnreal::ConvertMatrix( StageUpAxis, UsdMatrix );
+						PoseableMeshComponent->BoneSpaceTransforms[ BoneIndex ] = UsdToUnreal::ConvertMatrix( StageInfo, UsdMatrix );
 					}
 				}
 
diff --git a/Engine/Plugins/Importers/USDImporter/Source/USDStageEditor/Private/SUSDStageInfo.cpp b/Engine/Plugins/Importers/USDImporter/Source/USDStageEditor/Private/SUSDStageInfo.cpp
--- a/Engine/Plugins/Importers/USDImporter/Source/USDStageEditor/Private/SUSDStageInfo.cpp
+++ b/Engine/Plugins/Importers/USDImporter/Source/USDStageEditor/Private/SUSDStageInfo.cpp
@@ -9,6 +9,7 @@
 #include "EditorDirectories.h"
 #include "EditorStyleSet.h"
 #include "Widgets/SBoxPanel.h"
+#include "Widgets/Input/SEditableTextBox.h"
 #include "Widgets/Input/STextComboBox.h"
 #include "Widgets/Text/STextBlock.h"
 
@@ -18,15 +19,16 @@
 
 #include "pxr/pxr.h"
 #include "pxr/usd/usd/stage.h"
+#include "pxr/usd/usdGeom/metrics.h"
 
 #include "USDIncludesEnd.h"
 
 
 #define LOCTEXT_NAMESPACE "UsdStageInfo"
 
-void SUsdStageInfo::Construct( const FArguments& InArgs, AUsdStageActor* UsdStageActor )
+void SUsdStageInfo::Construct( const FArguments& InArgs, AUsdStageActor* InUsdStageActor )
 {
-	RefreshStageInfos( UsdStageActor );
+	RefreshStageInfos( InUsdStageActor );
 
 	ChildSlot
 	[
@@ -43,14 +45,38 @@ void SUsdStageInfo::Construct( const FArguments& InArgs, AUsdStageActor* UsdStag
 			.Text( this, &SUsdStageInfo::GetRootLayerDisplayName )
 			.Font( FEditorStyle::GetFontStyle( "ContentBrowser.SourceTreeItemFont" ) )
 		]
+		+SVerticalBox::Slot()
+		.AutoHeight()
+		.HAlign(HAlign_Left)
+		[
+			SNew( SHorizontalBox )
+
+			+SHorizontalBox::Slot()
+			.AutoWidth()
+			.VAlign(VAlign_Bottom)
+			.Padding(2.f, 2.f)
+			[
+				SNew( STextBlock )
+				.Text( LOCTEXT( "MetersPerUnit", "Meters per unit" ) )
+			]
+			+SHorizontalBox::Slot()
+			.Padding(2.f, 2.f)
+			[
+				SNew( SEditableTextBox )
+				.HintText( LOCTEXT( "Unset", "Unset" ) )
+				.Text( this, &SUsdStageInfo::GetMetersPerUnit )
+				.OnTextCommitted( this, &SUsdStageInfo::OnMetersPerUnitCommitted )
+			]
+		]
 	];
 }
 
-void SUsdStageInfo::RefreshStageInfos( AUsdStageActor* UsdStageActor )
+void SUsdStageInfo::RefreshStageInfos( AUsdStageActor* InUsdStageActor )
 {
+	UsdStageActor = InUsdStageActor;
 	StageInfos.RootLayerDisplayName = LOCTEXT( "NoUsdStage", "No Stage Available" );
 
-	if ( !UsdStageActor )
+	if ( !InUsdStageActor )
 	{
 		return;
 	}
@@ -59,6 +85,41 @@ void SUsdStageInfo::RefreshStageInfos( AUsdStageActor* UsdStageActor )
 	{
 		TUsdStore< std::string > UsdDisplayName = UsdStage->GetRootLayer()->GetDisplayName();
 		StageInfos.RootLayerDisplayName = FText::FromString( UsdToUnreal::ConvertString( UsdDisplayName.Get() ) );
+
+		if ( pxr::UsdGeomStageHasAuthoredMetersPerUnit( UsdStage ) )
+		{
+			StageInfos.MetersPerUnit = UsdUtils::GetUsdStageMetersPerUnit( UsdStage );
+		}
+		else
+		{
+			StageInfos.MetersPerUnit.Reset();
+		}
+	}
+}
+
+FText SUsdStageInfo::GetMetersPerUnit() const
+{
+	if ( StageInfos.MetersPerUnit )
+	{
+		return FText::FromString( LexToSanitizedString( StageInfos.MetersPerUnit.GetValue() ) );
+	}
+	else
+	{
+		return FText();
+	}
+}
+
+void SUsdStageInfo::OnMetersPerUnitCommitted( const FText& InUnitsPerMeterText, ETextCommit::Type InCommitInfo )
+{
+	if ( UsdStageActor.IsValid() )
+	{
+		float MetersPerUnit = 0.01f;
+		LexFromString( MetersPerUnit, *InUnitsPerMeterText.ToString() );
+
+		MetersPerUnit = FMath::Clamp( MetersPerUnit, 0.001f, 1000.f );
+
+		UsdUtils::SetUsdStageMetersPerUnit( UsdStageActor->GetUsdStage(), MetersPerUnit );
+		RefreshStageInfos( UsdStageActor.Get() );
 	}
 }
 
diff --git a/Engine/Plugins/Importers/USDImporter/Source/USDStageEditor/Private/SUSDStageInfo.h b/Engine/Plugins/Importers/USDImporter/Source/USDStageEditor/Private/SUSDStageInfo.h
--- a/Engine/Plugins/Importers/USDImporter/Source/USDStageEditor/Private/SUSDStageInfo.h
+++ b/Engine/Plugins/Importers/USDImporter/Source/USDStageEditor/Private/SUSDStageInfo.h
@@ -2,6 +2,8 @@
 
 #pragma once
 
+#include "Misc/Optional.h"
+#include "Types/SlateEnums.h"
 #include "Widgets/DeclarativeSyntaxSupport.h"
 #include "Widgets/SCompoundWidget.h"
 
@@ -15,18 +17,23 @@ public:
 	SLATE_BEGIN_ARGS( SUsdStageInfo ) {}
 	SLATE_END_ARGS()
 
-	void Construct( const FArguments& InArgs, AUsdStageActor* UsdStageActor );
-	void RefreshStageInfos( AUsdStageActor* UsdStageActor );
+	void Construct( const FArguments& InArgs, AUsdStageActor* InUsdStageActor );
+	void RefreshStageInfos( AUsdStageActor* InUsdStageActor );
 
 private:
 	struct FStageInfos
 	{
 		FText RootLayerDisplayName;
+		TOptional< float > MetersPerUnit;
 	};
 
 	FStageInfos StageInfos;
+	TWeakObjectPtr< AUsdStageActor > UsdStageActor;
 
 	FText GetRootLayerDisplayName() const { return StageInfos.RootLayerDisplayName; }
+	FText GetMetersPerUnit() const;
+
+	void OnMetersPerUnitCommitted( const FText& InUnitsPerMeterText, ETextCommit::Type InCommitInfo );
 };
 
 #endif // #if USE_USD_SDK
diff --git a/Engine/Plugins/Importers/USDImporter/Source/USDUtilities/Private/USDConversionUtils.cpp b/Engine/Plugins/Importers/USDImporter/Source/USDUtilities/Private/USDConversionUtils.cpp
--- a/Engine/Plugins/Importers/USDImporter/Source/USDUtilities/Private/USDConversionUtils.cpp
+++ b/Engine/Plugins/Importers/USDImporter/Source/USDUtilities/Private/USDConversionUtils.cpp
@@ -15,15 +15,43 @@
 #include "USDIncludesStart.h"
 
 #include "pxr/usd/usd/attribute.h"
+#include "pxr/usd/usd/editContext.h"
 #include "pxr/usd/usd/primRange.h"
 #include "pxr/usd/usdGeom/camera.h"
 #include "pxr/usd/usdGeom/scope.h"
 #include "pxr/usd/usdGeom/xform.h"
 #include "pxr/usd/usdGeom/mesh.h"
+#include "pxr/usd/usdGeom/metrics.h"
 #include "pxr/usd/usdSkel/root.h"
 
 #include "USDIncludesEnd.h"
 
+pxr::TfToken UsdUtils::GetUsdStageAxis( const pxr::UsdStageRefPtr& Stage )
+{
+	return pxr::UsdGeomGetStageUpAxis( Stage );
+}
+
+void UsdUtils::SetUsdStageAxis( const pxr::UsdStageRefPtr& Stage, pxr::TfToken Axis )
+{
+	pxr::UsdGeomSetStageUpAxis( Stage, Axis );
+}
+
+float UsdUtils::GetUsdStageMetersPerUnit( const pxr::UsdStageRefPtr& Stage )
+{
+	return (float)pxr::UsdGeomGetStageMetersPerUnit( Stage );
+}
+
+void UsdUtils::SetUsdStageMetersPerUnit( const pxr::UsdStageRefPtr& Stage, float MetersPerUnit )
+{
+	if ( !Stage || !Stage->GetRootLayer() )
+	{
+		return;
+	}
+
+	pxr::UsdEditContext( Stage, Stage->GetRootLayer() );
+	pxr::UsdGeomSetStageMetersPerUnit( Stage, MetersPerUnit );
+}
+
 UClass* UsdUtils::GetActorTypeForPrim( const pxr::UsdPrim& Prim )
 {
 	if ( Prim.IsA< pxr::UsdGeomCamera >() )
diff --git a/Engine/Plugins/Importers/USDImporter/Source/USDUtilities/Private/USDGeomMeshConversion.cpp b/Engine/Plugins/Importers/USDImporter/Source/USDUtilities/Private/USDGeomMeshConversion.cpp
--- a/Engine/Plugins/Importers/USDImporter/Source/USDUtilities/Private/USDGeomMeshConversion.cpp
+++ b/Engine/Plugins/Importers/USDImporter/Source/USDUtilities/Private/USDGeomMeshConversion.cpp
@@ -70,7 +70,7 @@ bool UsdToUnreal::ConvertGeomMesh( const pxr::UsdGeomMesh& UsdMesh, FMeshDescrip
 	FScopedUsdAllocs UsdAllocs;
 
 	pxr::UsdStageRefPtr Stage = UsdMesh.GetPrim().GetStage();
-	const pxr::TfToken StageUpAxis = UsdUtils::GetUsdStageAxis( Stage );
+	const UsdToUnreal::FUsdStageInfo StageInfo( Stage );
 
 	const double TimeCodeValue = TimeCode.GetValue();
 
@@ -100,7 +100,7 @@ bool UsdToUnreal::ConvertGeomMesh( const pxr::UsdGeomMesh& UsdMesh, FMeshDescrip
 			{
 				const GfVec3f& Point = PointsArray[ LocalPointIndex ];
 
-				FVector Position = AdditionalTransform.TransformPosition( UsdToUnreal::ConvertVector( StageUpAxis, Point ) );
+				FVector Position = AdditionalTransform.TransformPosition( UsdToUnreal::ConvertVector( StageInfo, Point ) );
 
 				FVertexID AddedVertexId = MeshDescription.CreateVertex();
 				MeshDescriptionVertexPositions[ AddedVertexId ] = Position;
@@ -278,7 +278,7 @@ bool UsdToUnreal::ConvertGeomMesh( const pxr::UsdGeomMesh& UsdMesh, FMeshDescrip
 					if ( NormalIndex < Normals.size() )
 					{
 						const GfVec3f& Normal = Normals[NormalIndex];
-						FVector TransformedNormal = AdditionalTransform.TransformVector( UsdToUnreal::ConvertVector( StageUpAxis, Normal ) ).GetSafeNormal();
+						FVector TransformedNormal = AdditionalTransform.TransformVector( UsdToUnreal::ConvertVector( StageInfo, Normal ) ).GetSafeNormal();
 
 						MeshDescriptionNormals[AddedVertexInstanceId] = TransformedNormal.GetSafeNormal();
 					}
diff --git a/Engine/Plugins/Importers/USDImporter/Source/USDUtilities/Private/USDPrimConversion.cpp b/Engine/Plugins/Importers/USDImporter/Source/USDUtilities/Private/USDPrimConversion.cpp
--- a/Engine/Plugins/Importers/USDImporter/Source/USDUtilities/Private/USDPrimConversion.cpp
+++ b/Engine/Plugins/Importers/USDImporter/Source/USDUtilities/Private/USDPrimConversion.cpp
@@ -31,8 +31,6 @@ bool UsdToUnreal::ConvertXformable( const pxr::UsdStageRefPtr& Stage, const pxr:
 
 	FScopedUsdAllocs UsdAllocs;
 
-	pxr::TfToken UpAxis = UsdUtils::GetUsdStageAxis( Stage );
-
 	// Transform
 	pxr::GfMatrix4d UsdMatrix;
 	bool bResetXFormStack = false;
@@ -40,13 +38,15 @@ bool UsdToUnreal::ConvertXformable( const pxr::UsdStageRefPtr& Stage, const pxr:
 
 	FRotator AdditionalRotation( ForceInit );
 
+	UsdToUnreal::FUsdStageInfo StageInfo( Stage );
+
 	if ( Xformable.GetPrim().IsA< pxr::UsdGeomCamera >() )
 	{
 		AdditionalRotation = FRotator( 0.0f, -90.f, 0.0f );
-		UpAxis = pxr::UsdGeomTokens->y; // Cameras are always Y up in USD
+		StageInfo.UpAxis = pxr::UsdGeomTokens->y; // Cameras are always Y up in USD
 	}
 
-	OutTransform = FTransform( AdditionalRotation ) * UsdToUnreal::ConvertMatrix( UpAxis, UsdMatrix );
+	OutTransform = FTransform( AdditionalRotation ) * UsdToUnreal::ConvertMatrix( StageInfo, UsdMatrix );
 
 	return true;
 }
diff --git a/Engine/Plugins/Importers/USDImporter/Source/USDUtilities/Private/USDSkeletalDataConversion.cpp b/Engine/Plugins/Importers/USDImporter/Source/USDUtilities/Private/USDSkeletalDataConversion.cpp
--- a/Engine/Plugins/Importers/USDImporter/Source/USDUtilities/Private/USDSkeletalDataConversion.cpp
+++ b/Engine/Plugins/Importers/USDImporter/Source/USDUtilities/Private/USDSkeletalDataConversion.cpp
@@ -67,10 +67,12 @@ bool UsdToUnreal::ConvertSkeleton(const pxr::UsdSkelSkeletonQuery& SkeletonQuery
 	if (bJointTransformsComputed)
 	{
 		UsdStageWeakPtr Stage = SkeletonQuery.GetSkeleton().GetPrim().GetStage();
+		const UsdToUnreal::FUsdStageInfo StageInfo(Stage);
+
 		for (uint32 Index = 0; Index < UsdBoneTransforms.size(); ++Index)
 		{
 			const GfMatrix4d& UsdMatrix = UsdBoneTransforms[Index];
-			FTransform BoneTransform = UsdToUnreal::ConvertMatrix( UsdUtils::GetUsdStageAxis( Stage ), UsdMatrix );
+			FTransform BoneTransform = UsdToUnreal::ConvertMatrix( StageInfo, UsdMatrix );
 			BoneTransforms.Add(BoneTransform);
 		}
 	}
@@ -126,9 +128,9 @@ bool UsdToUnreal::ConvertSkinnedMesh(const pxr::UsdSkelSkinningQuery& SkinningQu
 		GeomBindingAttribute.Get(&GeomBindingTransform, UsdTimeCode::Default());
 	}
 
-	const pxr::TfToken StageUpAxis = UsdUtils::GetUsdStageAxis( SkinningPrim.GetStage() );
+	const UsdToUnreal::FUsdStageInfo StageInfo( SkinningPrim.GetStage() );
 
-	FTransform GeomTransform = UsdToUnreal::ConvertMatrix(StageUpAxis, GeomBindingTransform);
+	FTransform GeomTransform = UsdToUnreal::ConvertMatrix(StageInfo, GeomBindingTransform);
 
 	// Ref. FFbxImporter::FillSkelMeshImporterFromFbx
 	UsdGeomMesh UsdMesh = UsdGeomMesh(SkinningPrim);
@@ -151,7 +153,7 @@ bool UsdToUnreal::ConvertSkinnedMesh(const pxr::UsdSkelSkinningQuery& SkinningQu
 			const GfVec3f& Point = UsdPoints[PointIndex];
 
 			// Convert the USD vertex to Unreal and apply the GeomBindTransform to it
-			FVector Pos = UsdToUnreal::ConvertVector(StageUpAxis, Point);
+			FVector Pos = UsdToUnreal::ConvertVector(StageInfo, Point);
 			Pos = GeomTransform.TransformPosition(Pos);
 
 			SkelMeshImportData.Points[PointIndex + NumExistingPoints] = Pos;
diff --git a/Engine/Plugins/Importers/USDImporter/Source/USDUtilities/Public/USDConversionUtils.h b/Engine/Plugins/Importers/USDImporter/Source/USDUtilities/Public/USDConversionUtils.h
--- a/Engine/Plugins/Importers/USDImporter/Source/USDUtilities/Public/USDConversionUtils.h
+++ b/Engine/Plugins/Importers/USDImporter/Source/USDUtilities/Public/USDConversionUtils.h
@@ -45,22 +45,11 @@ namespace UsdUtils
 		return Value;
 	}
 
-	static pxr::TfToken GetUsdStageAxis(pxr::UsdStageRefPtr Stage)
-	{
-		if (Stage->HasAuthoredMetadata(pxr::UsdGeomTokens->upAxis))
-		{
-			pxr::TfToken Axis = pxr::UsdGeomTokens->z;
-			Stage->GetMetadata(pxr::UsdGeomTokens->upAxis, &Axis);
-			return Axis;
-		}
+	USDUTILITIES_API pxr::TfToken GetUsdStageAxis( const pxr::UsdStageRefPtr& Stage );
+	USDUTILITIES_API void SetUsdStageAxis( const pxr::UsdStageRefPtr& Stage, pxr::TfToken Axis );
 
-		return pxr::UsdGeomTokens->z;
-	}
-
-	static void SetUsdStageAxis( pxr::UsdStageRefPtr Stage, pxr::TfToken Axis )
-	{
-		Stage->SetMetadata( pxr::UsdGeomTokens->upAxis, Axis );
-	}
+	USDUTILITIES_API float GetUsdStageMetersPerUnit( const pxr::UsdStageRefPtr& Stage );
+	USDUTILITIES_API void SetUsdStageMetersPerUnit( const pxr::UsdStageRefPtr& Stage, float MetersPerUnit );
 
 	inline bool HasCompositionArcs( const pxr::UsdPrim& Prim )
 	{
diff --git a/Engine/Plugins/Importers/USDImporter/Source/USDUtilities/Public/USDTypesConversion.h b/Engine/Plugins/Importers/USDImporter/Source/USDUtilities/Public/USDTypesConversion.h
--- a/Engine/Plugins/Importers/USDImporter/Source/USDUtilities/Public/USDTypesConversion.h
+++ b/Engine/Plugins/Importers/USDImporter/Source/USDUtilities/Public/USDTypesConversion.h
@@ -23,6 +23,18 @@
 
 namespace UsdToUnreal
 {
+	struct FUsdStageInfo
+	{
+		pxr::TfToken UpAxis = pxr::UsdGeomTokens->z;
+		float MetersPerUnit = 0.01f;
+
+		explicit FUsdStageInfo( const pxr::UsdStageRefPtr& Stage )
+		{
+			UpAxis = UsdUtils::GetUsdStageAxis( Stage );
+			MetersPerUnit = UsdUtils::GetUsdStageMetersPerUnit( Stage );
+		}
+	};
+
 	static FString ConvertString( const std::string& InString )
 	{
 		return FString( ANSI_TO_TCHAR( InString.c_str() ) );
@@ -62,9 +74,9 @@ namespace UsdToUnreal
 		return FName( UsdString.Get().c_str() );
 	}
 
-	static FString ConvertToken(const pxr::TfToken& Token)
+	static FString ConvertToken( const pxr::TfToken& Token )
 	{
-		return UsdToUnreal::ConvertString(Token.GetString());
+		return UsdToUnreal::ConvertString( Token.GetString() );
 	}
 
 	static FLinearColor ConvertColor( const pxr::GfVec3f& InValue )
@@ -87,11 +99,17 @@ namespace UsdToUnreal
 		return FVector( InValue[0], InValue[1], InValue[2] );
 	}
 
-	static FVector ConvertVector( const pxr::TfToken& UpAxis, const pxr::GfVec3f& InValue )
+	static FVector ConvertVector( const FUsdStageInfo& StageInfo, const pxr::GfVec3f& InValue )
 	{
 		FVector Value = ConvertVector( InValue );
 
-		const bool bIsZUp = ( UpAxis == pxr::UsdGeomTokens->z );
+		const float UEMetersPerUnit = 0.01f;
+		if ( !FMath::IsNearlyEqual( StageInfo.MetersPerUnit, UEMetersPerUnit ) )
+		{
+			Value *= StageInfo.MetersPerUnit / UEMetersPerUnit;
+		}
+
+		const bool bIsZUp = ( StageInfo.UpAxis == pxr::UsdGeomTokens->z );
 
 		if ( bIsZUp )
 		{
@@ -158,12 +176,18 @@ namespace UsdToUnreal
 		return UnrealMatrix;
 	}
 
-	static FTransform ConvertMatrix( const pxr::TfToken& UpAxis, const pxr::GfMatrix4d& InMatrix )
+	static FTransform ConvertMatrix( const FUsdStageInfo& StageInfo, const pxr::GfMatrix4d& InMatrix )
 	{
 		FMatrix Matrix = ConvertMatrix( InMatrix );
 		FTransform Transform( Matrix );
 
-		Transform = ConvertTransform( UpAxis == pxr::UsdGeomTokens->z, Transform );
+		Transform = ConvertTransform( StageInfo.UpAxis == pxr::UsdGeomTokens->z, Transform );
+
+		const float UEMetersPerUnit = 0.01f;
+		if ( !FMath::IsNearlyEqual( StageInfo.MetersPerUnit, UEMetersPerUnit ) )
+		{
+			Transform.ScaleTranslation( StageInfo.MetersPerUnit / UEMetersPerUnit );
+		}
 
 		return Transform;
 	}
