USD: Fix incorrect camera transforms on Z-up stages

diff --git a/Engine/Plugins/Importers/USDImporter/Source/USDUtilities/Private/USDPrimConversion.cpp b/Engine/Plugins/Importers/USDImporter/Source/USDUtilities/Private/USDPrimConversion.cpp
--- a/Engine/Plugins/Importers/USDImporter/Source/USDUtilities/Private/USDPrimConversion.cpp
+++ b/Engine/Plugins/Importers/USDImporter/Source/USDUtilities/Private/USDPrimConversion.cpp
@@ -36,14 +36,22 @@ bool UsdToUnreal::ConvertXformable( const pxr::UsdStageRefPtr& Stage, const pxr:
 	bool bResetXFormStack = false;
 	Xformable.GetLocalTransformation( &UsdMatrix, &bResetXFormStack, EvalTime );
 
-	FRotator AdditionalRotation( ForceInit );
-
 	UsdToUnreal::FUsdStageInfo StageInfo( Stage );
 
+	// Extra rotation to match different camera facing direction convention
+	// Note: The camera space is always Y-up, yes, but this is not what this is: This is the camera's transform wrt the stage,
+	// which follows the stage up axis
+	FRotator AdditionalRotation( ForceInit );
 	if ( Xformable.GetPrim().IsA< pxr::UsdGeomCamera >() )
 	{
-		AdditionalRotation = FRotator( 0.0f, -90.f, 0.0f );
-		StageInfo.UpAxis = pxr::UsdGeomTokens->y; // Cameras are always Y up in USD
+		if (StageInfo.UpAxis == pxr::UsdGeomTokens->y)
+		{
+			AdditionalRotation = FRotator(0.0f, -90.f, 0.0f);
+		}
+		else
+		{
+			AdditionalRotation = FRotator(-90.0f, -90.f, 0.0f);
+		}
 	}
 
 	OutTransform = FTransform( AdditionalRotation ) * UsdToUnreal::ConvertMatrix( StageInfo, UsdMatrix );
