USD: Undo/Redo fixes

Fix Undo/Redo when switching RootLayer path on stage actor;
Add transaction when creating new USD files;
Fix AUsdStageActor not loading stage when loading map;
Fix SUSDStage window not updating when loading map;

diff --git a/Engine/Plugins/Importers/USDImporter/Source/USDStage/Private/USDPrimTwin.cpp b/Engine/Plugins/Importers/USDImporter/Source/USDStage/Private/USDPrimTwin.cpp
--- a/Engine/Plugins/Importers/USDImporter/Source/USDStage/Private/USDPrimTwin.cpp
+++ b/Engine/Plugins/Importers/USDImporter/Source/USDStage/Private/USDPrimTwin.cpp
@@ -7,7 +7,7 @@
 
 #include "Engine/World.h"
 
-FUsdPrimTwin& FUsdPrimTwin::AddChild( const FString& InPrimPath )
+UUsdPrimTwin& UUsdPrimTwin::AddChild( const FString& InPrimPath )
 {
 	FScopedUnrealAllocs UnrealAllocs; // Make sure the call to new is done with the UE allocator
 
@@ -15,19 +15,23 @@ FUsdPrimTwin& FUsdPrimTwin::AddChild( const FString& InPrimPath )
 	FString ChildPrimName;
 	InPrimPath.Split( TEXT("/"), &Dummy, &ChildPrimName, ESearchCase::IgnoreCase, ESearchDir::FromEnd );
 
-	TUniquePtr< FUsdPrimTwin >& ChildPrim = Children.Add( ChildPrimName );
+	Modify();
 
-	ChildPrim = MakeUnique< FUsdPrimTwin >();
+	UUsdPrimTwin*& ChildPrim = Children.Add( ChildPrimName );
+
+	ChildPrim = NewObject<UUsdPrimTwin>(this, NAME_None, RF_Transient | RF_Transactional);
 	ChildPrim->PrimPath = InPrimPath;
 
 	return *ChildPrim;
 }
 
-void FUsdPrimTwin::RemoveChild( const TCHAR* InPrimPath )
+void UUsdPrimTwin::RemoveChild( const TCHAR* InPrimPath )
 {
 	FScopedUnrealAllocs UnrealAllocs;
 
-	for ( TMap< FString, TUniquePtr< FUsdPrimTwin > >::TIterator ChildIt = Children.CreateIterator(); ChildIt; ++ChildIt )
+	Modify();
+
+	for ( TMap< FString, UUsdPrimTwin* >::TIterator ChildIt = Children.CreateIterator(); ChildIt; ++ChildIt )
 	{
 		if ( ChildIt->Value->PrimPath == InPrimPath )
 		{
@@ -37,10 +41,16 @@ void FUsdPrimTwin::RemoveChild( const TCHAR* InPrimPath )
 	}
 }
 
-void FUsdPrimTwin::Clear()
+void UUsdPrimTwin::Clear()
 {
 	FScopedUnrealAllocs UnrealAllocs;
 
+	Modify();
+
+	for (const TPair< FString, UUsdPrimTwin* >& Pair : Children)
+	{
+		Pair.Value->Clear();
+	}
 	Children.Empty();
 
 	if ( !PrimPath.IsEmpty() )
@@ -60,17 +70,19 @@ void FUsdPrimTwin::Clear()
 
 	if ( ActorToDestroy && !ActorToDestroy->IsA< AUsdStageActor >() && !ActorToDestroy->IsActorBeingDestroyed() && ActorToDestroy->GetWorld() )
 	{
+		ActorToDestroy->Modify();
 		ActorToDestroy->GetWorld()->DestroyActor( ActorToDestroy );
 		SpawnedActor = nullptr;
 	}
 	else if ( SceneComponent.IsValid() && !SceneComponent->IsBeingDestroyed() )
 	{
+		SceneComponent->Modify();
 		SceneComponent->DestroyComponent();
 		SceneComponent = nullptr;
 	}
 }
 
-FUsdPrimTwin* FUsdPrimTwin::Find( const FString& InPrimPath )
+UUsdPrimTwin* UUsdPrimTwin::Find( const FString& InPrimPath )
 {
 	if ( PrimPath == InPrimPath )
 	{
@@ -98,7 +110,7 @@ FUsdPrimTwin* FUsdPrimTwin::Find( const FString& InPrimPath )
 	{
 		if ( RestOfPrimPathToFind.IsEmpty() )
 		{
-			return Children[ ChildPrimName ].Get();
+			return Children[ ChildPrimName ];
 		}
 		else
 		{
diff --git a/Engine/Plugins/Importers/USDImporter/Source/USDStage/Private/USDStageActor.cpp b/Engine/Plugins/Importers/USDImporter/Source/USDStage/Private/USDStageActor.cpp
--- a/Engine/Plugins/Importers/USDImporter/Source/USDStage/Private/USDStageActor.cpp
+++ b/Engine/Plugins/Importers/USDImporter/Source/USDStage/Private/USDStageActor.cpp
@@ -36,6 +36,9 @@
 #include "StaticMeshAttributes.h"
 #include "StaticMeshOperations.h"
 #include "Subsystems/AssetEditorSubsystem.h"
+#if WITH_EDITOR
+#include "LevelEditor.h"
+#endif // WITH_EDITOR
 
 #include "Modules/ModuleManager.h"
 
@@ -76,11 +79,11 @@ struct FUsdStageActorImpl
 		TranslationContext->PurposesToLoad = (EUsdPurpose) StageActor->PurposesToLoad;
 
 		TUsdStore< pxr::SdfPath > UsdPrimPath = UnrealToUsd::ConvertPath( *PrimPath );
-		FUsdPrimTwin* ParentUsdPrimTwin = StageActor->RootUsdTwin.Find( UsdToUnreal::ConvertPath( UsdPrimPath.Get().GetParentPath() ) );
+		UUsdPrimTwin* ParentUsdPrimTwin = StageActor->RootUsdTwin->Find( UsdToUnreal::ConvertPath( UsdPrimPath.Get().GetParentPath() ) );
 
 		if ( !ParentUsdPrimTwin )
 		{
-			ParentUsdPrimTwin = &StageActor->RootUsdTwin;
+			ParentUsdPrimTwin = StageActor->RootUsdTwin;
 		}
 
 		TranslationContext->ParentComponent = ParentUsdPrimTwin ? ParentUsdPrimTwin->SceneComponent.Get() : nullptr;
@@ -139,8 +142,17 @@ AUsdStageActor::AUsdStageActor()
 
 	RootComponent = SceneComponent;
 
+	RootUsdTwin = NewObject<UUsdPrimTwin>(this, TEXT("RootUsdTwin"), RF_Transient | RF_Transactional);
+
+#if WITH_EDITOR
+	// We can't use PostLoad to trigger LoadUsdStage when first loading a saved level because LoadUsdStage may trigger
+	// Rename() calls, and that is not allowed.
+	FLevelEditorModule& LevelEditorModule = FModuleManager::LoadModuleChecked<FLevelEditorModule>("LevelEditor");
+	LevelEditorModule.OnMapChanged().AddUObject(this, &AUsdStageActor::OnMapChanged);
+#endif // WITH_EDITOR
+
 #if USE_USD_SDK
-	RootUsdTwin.PrimPath = TEXT("/");
+	RootUsdTwin->PrimPath = TEXT("/");
 
 	UsdListener.OnPrimChanged.AddLambda(
 		[ this ]( const FString& OriginalPrimPath, bool bResync )
@@ -244,6 +256,11 @@ AUsdStageActor::~AUsdStageActor()
 		UnrealUSDWrapper::GetUsdStageCache().Erase( UsdStageStore.Get() );
 	}
 #endif // #if USE_USD_SDK
+
+#if WITH_EDITOR
+	FLevelEditorModule& LevelEditorModule = FModuleManager::LoadModuleChecked<FLevelEditorModule>("LevelEditor");
+	LevelEditorModule.OnMapChanged().RemoveAll(this);
+#endif // WITH_EDITOR
 }
 
 USDSTAGE_API void AUsdStageActor::Reset()
@@ -262,8 +279,8 @@ USDSTAGE_API void AUsdStageActor::Reset()
 	TimeCodesPerSecond = 24.0f;
 	Time = 0;
 
-	RootUsdTwin.Clear();
-	RootUsdTwin.PrimPath = TEXT("/");
+	RootUsdTwin->Clear();
+	RootUsdTwin->PrimPath = TEXT("/");
 	GEditor->BroadcastLevelActorListChanged();
 
 	RootLayer.FilePath.Empty();
@@ -278,13 +295,13 @@ USDSTAGE_API void AUsdStageActor::Reset()
 
 #if USE_USD_SDK
 
-FUsdPrimTwin* AUsdStageActor::GetOrCreatePrimTwin( const pxr::SdfPath& UsdPrimPath )
+UUsdPrimTwin* AUsdStageActor::GetOrCreatePrimTwin( const pxr::SdfPath& UsdPrimPath )
 {
 	const FString PrimPath = UsdToUnreal::ConvertPath( UsdPrimPath );
 	const FString ParentPrimPath = UsdToUnreal::ConvertPath( UsdPrimPath.GetParentPath() );
 
-	FUsdPrimTwin* UsdPrimTwin = RootUsdTwin.Find( PrimPath );
-	FUsdPrimTwin* ParentUsdPrimTwin = RootUsdTwin.Find( ParentPrimPath );
+	UUsdPrimTwin* UsdPrimTwin = RootUsdTwin->Find( PrimPath );
+	UUsdPrimTwin* ParentUsdPrimTwin = RootUsdTwin->Find( ParentPrimPath );
 
 	const pxr::UsdPrim Prim = GetUsdStage()->GetPrimAtPath( UsdPrimPath );
 
@@ -295,7 +312,7 @@ FUsdPrimTwin* AUsdStageActor::GetOrCreatePrimTwin( const pxr::SdfPath& UsdPrimPa
 
 	if ( !ParentUsdPrimTwin )
 	{
-		ParentUsdPrimTwin = &RootUsdTwin;
+		ParentUsdPrimTwin = RootUsdTwin;
 	}
 
 	if ( !UsdPrimTwin )
@@ -303,12 +320,12 @@ FUsdPrimTwin* AUsdStageActor::GetOrCreatePrimTwin( const pxr::SdfPath& UsdPrimPa
 		UsdPrimTwin = &ParentUsdPrimTwin->AddChild( *PrimPath );
 
 		UsdPrimTwin->OnDestroyed.AddLambda(
-			[ this ]( const FUsdPrimTwin& UsdPrimTwin )
+			[ this ]( const UUsdPrimTwin& UsdPrimTwin )
 			{
 				this->OnUsdPrimTwinDestroyed( UsdPrimTwin );
 			} );
 
-		if ( !UsdPrimTwin->AnimationHandle.IsValid() && UsdUtils::IsAnimated( Prim ) )
+		if ( UsdUtils::IsAnimated( Prim ) )
 		{
 			PrimsToAnimate.Add( PrimPath );
 		}
@@ -317,7 +334,7 @@ FUsdPrimTwin* AUsdStageActor::GetOrCreatePrimTwin( const pxr::SdfPath& UsdPrimPa
 	return UsdPrimTwin;
 }
 
-FUsdPrimTwin* AUsdStageActor::ExpandPrim( const pxr::UsdPrim& Prim, FUsdSchemaTranslationContext& TranslationContext )
+UUsdPrimTwin* AUsdStageActor::ExpandPrim( const pxr::UsdPrim& Prim, FUsdSchemaTranslationContext& TranslationContext )
 {
 	if ( !Prim )
 	{
@@ -326,7 +343,7 @@ FUsdPrimTwin* AUsdStageActor::ExpandPrim( const pxr::UsdPrim& Prim, FUsdSchemaTr
 
 	TRACE_CPUPROFILER_EVENT_SCOPE( AUsdStageActor::ExpandPrim );
 
-	FUsdPrimTwin* UsdPrimTwin = GetOrCreatePrimTwin( Prim.GetPrimPath() );
+	UUsdPrimTwin* UsdPrimTwin = GetOrCreatePrimTwin( Prim.GetPrimPath() );
 
 	bool bExpandChilren = true;
 
@@ -397,14 +414,14 @@ void AUsdStageActor::UpdatePrim( const pxr::SdfPath& InUsdPrimPath, bool bResync
 		if ( bResync )
 		{
 			FString PrimPath = UsdToUnreal::ConvertPath( UsdPrimPath );
-			if ( FUsdPrimTwin* UsdPrimTwin = RootUsdTwin.Find( PrimPath ) )
+			if ( UUsdPrimTwin* UsdPrimTwin = RootUsdTwin->Find( PrimPath ) )
 			{
 				UsdPrimTwin->Clear();
 			}
 		}
 
 		TUsdStore< pxr::UsdPrim > PrimToExpand = GetUsdStage()->GetPrimAtPath( UsdPrimPath );
-		FUsdPrimTwin* UsdPrimTwin = ExpandPrim( PrimToExpand.Get(), TranslationContext );
+		UUsdPrimTwin* UsdPrimTwin = ExpandPrim( PrimToExpand.Get(), TranslationContext );
 
 		GEditor->BroadcastLevelActorListChanged();
 		GEditor->RedrawLevelEditingViewports();
@@ -455,15 +472,18 @@ void AUsdStageActor::PostEditChangeProperty(FPropertyChangedEvent& PropertyChang
 			LevelSequenceHelper.UpdateLevelSequence( UsdStage );
 		}
 	}
-	else if ( PropertyName == GET_MEMBER_NAME_CHECKED(AUsdStageActor, PurposeVisibility) )
+	else if (PropertyName == GET_MEMBER_NAME_CHECKED(AUsdStageActor, PurposeVisibility))
 	{
 		RefreshVisibilityBasedOnPurpose();
 	}
-	else
-#endif // #if USE_USD_SDK
+	else if (PropertyName == GET_MEMBER_NAME_CHECKED(AUsdStageActor, InitialLoadSet) ||
+			 PropertyName == GET_MEMBER_NAME_CHECKED(AUsdStageActor, PurposesToLoad))
 	{
-		Super::PostEditChangeProperty( PropertyChangedEvent );
+		LoadUsdStage();
 	}
+#endif // #if USE_USD_SDK
+
+	Super::PostEditChangeProperty(PropertyChangedEvent);
 }
 
 void AUsdStageActor::Clear()
@@ -529,6 +549,31 @@ void AUsdStageActor::RefreshVisibilityBasedOnPurpose()
 #endif // #if USE_USD_SDK
 }
 
+#if WITH_EDITOR
+void AUsdStageActor::OnMapChanged(UWorld* World, EMapChangeType ChangeType)
+{
+	// Our current level has changed
+	if (HasAutorithyOverStage() && World && World == GetWorld() && World->GetCurrentLevel() == GetTypedOuter<ULevel>())
+	{
+		if (ChangeType == EMapChangeType::LoadMap || ChangeType == EMapChangeType::NewMap)
+		{
+			// This is in charge of first loading our stage when we load a level that had a AUsdStageActor saved with
+			// a valid root layer filepath
+			LoadUsdStage();
+
+			// SUSDStage window needs to update
+			OnActorLoaded.Broadcast(this);
+		}
+		else if (ChangeType == EMapChangeType::TearDownWorld)
+		{
+			// This is in charge of clearing the SUSDStage window when switching away from our level
+			// SUSDStage window needs to update
+			OnActorDestroyed.Broadcast();
+		}
+	}
+}
+#endif // WITH_EDITOR
+
 void AUsdStageActor::LoadUsdStage()
 {
 #if USE_USD_SDK
@@ -540,8 +585,8 @@ void AUsdStageActor::LoadUsdStage()
 
 	Clear();
 
-	RootUsdTwin.Clear();
-	RootUsdTwin.PrimPath = TEXT("/");
+	RootUsdTwin->Clear();
+	RootUsdTwin->PrimPath = TEXT("/");
 
 	const pxr::UsdStageRefPtr& UsdStage = GetUsdStage();
 
@@ -553,7 +598,7 @@ void AUsdStageActor::LoadUsdStage()
 
 	ReloadAnimations();
 
-	TSharedRef< FUsdSchemaTranslationContext > TranslationContext = FUsdStageActorImpl::CreateUsdSchemaTranslationContext( this, RootUsdTwin.PrimPath );
+	TSharedRef< FUsdSchemaTranslationContext > TranslationContext = FUsdStageActorImpl::CreateUsdSchemaTranslationContext( this, RootUsdTwin->PrimPath );
 
 	LoadAssets( *TranslationContext, UsdStage->GetPseudoRoot() );
 
@@ -607,7 +652,6 @@ void AUsdStageActor::ReloadAnimations()
 		LevelSequence = nullptr;
 		SubLayerLevelSequencesByIdentifier.Reset();
 
-		LevelSequenceHelper.InitLevelSequence(UsdStage);
 	}
 #endif // #if USE_USD_SDK
 }
@@ -621,51 +665,39 @@ void AUsdStageActor::PostTransacted(const FTransactionObjectEvent& TransactionEv
 		{
 			OnActorDestroyed.Broadcast();
 		}
-		// This fires when being spawned, undo delete, redo spawn
+		// This fires when being spawned in an existing level, undo delete, redo spawn
 		else
 		{
 			OnActorLoaded.Broadcast( this );
-
-			if (HasAutorithyOverStage())
-			{
-				LoadUsdStage();
-			}
-
 			OnStageChanged.Broadcast();
 		}
 	}
-	// Fires on Close and undo Close
 	else if (TransactionEvent.GetEventType() == ETransactionObjectEventType::UndoRedo)
 	{
 		// UsdStageStore can't be a UPROPERTY, so we have to make sure that it
 		// is kept in sync with the state of RootLayer, because LoadUsdStage will
-		// do the job of clearing our components and etc if the path is empty
+		// do the job of clearing our instanced actors/components if the path is empty
 #if USE_USD_SDK
-		if (RootLayer.FilePath.IsEmpty())
+		const TArray<FName>& ChangedProperties = TransactionEvent.GetChangedProperties();
+		if (ChangedProperties.Contains(GET_MEMBER_NAME_CHECKED(AUsdStageActor, RootLayer)))
 		{
+			// Changed the path, so we need to reopen to the correct stage
 			UnrealUSDWrapper::GetUsdStageCache().Erase( UsdStageStore.Get() );
 			UsdStageStore = TUsdStore< pxr::UsdStageRefPtr >();
-		}
-#endif // #if USE_USD_SDK
+			OnStageChanged.Broadcast();
 
-		if ( HasAutorithyOverStage() )
+			ReloadAnimations();
+		}
+		else if (ChangedProperties.Contains(GET_MEMBER_NAME_CHECKED(AUsdStageActor, Time)))
 		{
-			LoadUsdStage();
+			Refresh();
 		}
+#endif // #if USE_USD_SDK
 	}
 }
 
-void AUsdStageActor::OnUsdPrimTwinDestroyed( const FUsdPrimTwin& UsdPrimTwin )
+void AUsdStageActor::OnUsdPrimTwinDestroyed( const UUsdPrimTwin& UsdPrimTwin )
 {
-	TArray<FDelegateHandle> DelegateHandles;
-	PrimDelegates.MultiFind( UsdPrimTwin.PrimPath, DelegateHandles );
-
-	for ( FDelegateHandle Handle : DelegateHandles )
-	{
-		OnTimeChanged.Remove( Handle );
-	}
-
-	PrimDelegates.Remove( UsdPrimTwin.PrimPath );
 	PrimsToAnimate.Remove( UsdPrimTwin.PrimPath );
 
 	UObject* WatchedObject = UsdPrimTwin.SpawnedActor.IsValid() ? (UObject*)UsdPrimTwin.SpawnedActor.Get() : (UObject*)UsdPrimTwin.SceneComponent.Get();
@@ -703,7 +735,7 @@ void AUsdStageActor::OnPrimObjectPropertyChanged( UObject* ObjectBeingModified,
 
 	FString PrimPath = ObjectsToWatch[ PrimObject ];
 
-	if ( FUsdPrimTwin* UsdPrimTwin = RootUsdTwin.Find( PrimPath ) )
+	if ( UUsdPrimTwin* UsdPrimTwin = RootUsdTwin->Find( PrimPath ) )
 	{
 		// Update prim from UE
 		USceneComponent* PrimSceneComponent = UsdPrimTwin->SceneComponent.Get();
@@ -797,7 +829,7 @@ void AUsdStageActor::LoadAssets( FUsdSchemaTranslationContext& TranslationContex
 
 void AUsdStageActor::AnimatePrims()
 {
-	TSharedRef< FUsdSchemaTranslationContext > TranslationContext = FUsdStageActorImpl::CreateUsdSchemaTranslationContext( this, RootUsdTwin.PrimPath );
+	TSharedRef< FUsdSchemaTranslationContext > TranslationContext = FUsdStageActorImpl::CreateUsdSchemaTranslationContext( this, RootUsdTwin->PrimPath );
 
 	for ( const FString& PrimToAnimate : PrimsToAnimate )
 	{
@@ -806,7 +838,7 @@ void AUsdStageActor::AnimatePrims()
 		IUsdSchemasModule& SchemasModule = FModuleManager::Get().LoadModuleChecked< IUsdSchemasModule >( "USDSchemas" );
 		if ( TSharedPtr< FUsdSchemaTranslator > SchemaTranslator = SchemasModule.GetTranslatorRegistry().CreateTranslatorForSchema( TranslationContext, pxr::UsdTyped( this->GetUsdStage()->GetPrimAtPath( PrimPath.Get() ) ) ) )
 		{
-			if ( FUsdPrimTwin* UsdPrimTwin = RootUsdTwin.Find( PrimToAnimate ) )
+			if ( UUsdPrimTwin* UsdPrimTwin = RootUsdTwin->Find( PrimToAnimate ) )
 			{
 				SchemaTranslator->UpdateComponents( UsdPrimTwin->SceneComponent.Get() );
 			}
diff --git a/Engine/Plugins/Importers/USDImporter/Source/USDStage/Public/USDPrimTwin.h b/Engine/Plugins/Importers/USDImporter/Source/USDStage/Public/USDPrimTwin.h
--- a/Engine/Plugins/Importers/USDImporter/Source/USDStage/Public/USDPrimTwin.h
+++ b/Engine/Plugins/Importers/USDImporter/Source/USDStage/Public/USDPrimTwin.h
@@ -12,28 +12,25 @@
 #include "UObject/WeakObjectPtr.h"
 #include "UObject/WeakObjectPtrTemplates.h"
 
+#include "USDPrimTwin.generated.h"
+
 /** The Unreal equivalent (twin) of a USD prim */
-class FUsdPrimTwin final
+UCLASS(Transient)
+class UUsdPrimTwin final : public UObject
 {
-public:
-	FUsdPrimTwin() = default;
+	GENERATED_BODY()
 
-	FUsdPrimTwin( FUsdPrimTwin&& Other ) = default;
-	FUsdPrimTwin& operator=( FUsdPrimTwin&& Other ) = default;
-	FUsdPrimTwin( const FUsdPrimTwin& Other ) = delete;
-	FUsdPrimTwin& operator=( const FUsdPrimTwin& Other ) = delete;
-	~FUsdPrimTwin() { Clear(); }
-
-	FUsdPrimTwin& AddChild( const FString& InPrimPath );
+public:
+	UUsdPrimTwin& AddChild( const FString& InPrimPath );
 	void RemoveChild( const TCHAR* InPrimPath );
 
 	void Clear();
 
-	void Iterate( TFunction< void( FUsdPrimTwin& ) > Func, bool bRecursive )
+	void Iterate( TFunction< void( UUsdPrimTwin& ) > Func, bool bRecursive )
 	{
-		for ( TMap< FString, TUniquePtr< FUsdPrimTwin > >::TIterator It = Children.CreateIterator(); It; ++It )
+		for ( TMap< FString, UUsdPrimTwin* >::TIterator It = Children.CreateIterator(); It; ++It )
 		{
-			FUsdPrimTwin* Child = It->Value.Get();
+			UUsdPrimTwin* Child = It->Value;
 
 			Func( *Child );
 
@@ -44,19 +41,22 @@ public:
 		}
 	}
 
-	FUsdPrimTwin* Find( const FString& InPrimPath );
+	UUsdPrimTwin* Find( const FString& InPrimPath );
 
-	DECLARE_EVENT_OneParam( FUsdPrimTwin, FOnUsdPrimTwinDestroyed, const FUsdPrimTwin& );
+	DECLARE_EVENT_OneParam( UUsdPrimTwin, FOnUsdPrimTwinDestroyed, const UUsdPrimTwin& );
 	FOnUsdPrimTwinDestroyed OnDestroyed;
 
 public:
+	UPROPERTY()
 	FString PrimPath;
+
+	UPROPERTY()
 	TWeakObjectPtr< class AActor > SpawnedActor;
-	TWeakObjectPtr< class USceneComponent > SceneComponent;
 
-	FDelegateHandle AnimationHandle;
+	UPROPERTY()
+	TWeakObjectPtr< class USceneComponent > SceneComponent;
 
 private:
-	TMap< FString, TUniquePtr< FUsdPrimTwin > > Children;
-
+	UPROPERTY()
+	TMap< FString, UUsdPrimTwin* > Children;
 };
diff --git a/Engine/Plugins/Importers/USDImporter/Source/USDStage/Public/USDStageActor.h b/Engine/Plugins/Importers/USDImporter/Source/USDStage/Public/USDStageActor.h
--- a/Engine/Plugins/Importers/USDImporter/Source/USDStage/Public/USDStageActor.h
+++ b/Engine/Plugins/Importers/USDImporter/Source/USDStage/Public/USDStageActor.h
@@ -31,6 +31,7 @@ class IMeshBuilderModule;
 class ULevelSequence;
 class UMaterial;
 class UUsdAsset;
+enum class EMapChangeType : uint8;
 enum class EUsdPurpose : int32;
 struct FMeshDescription;
 struct FUsdSchemaTranslationContext;
@@ -127,20 +128,23 @@ private:
 
 	void RefreshVisibilityBasedOnPurpose();
 
-	void OnUsdPrimTwinDestroyed( const FUsdPrimTwin& UsdPrimTwin );
+#if WITH_EDITOR
+	void OnMapChanged(UWorld* World, EMapChangeType ChangeType);
+#endif // WITH_EDITOR
+
+	void OnUsdPrimTwinDestroyed( const UUsdPrimTwin& UsdPrimTwin );
 
 	void OnPrimObjectPropertyChanged( UObject* ObjectBeingModified, FPropertyChangedEvent& PropertyChangedEvent );
 	bool HasAutorithyOverStage() const;
 
 private:
-	FUsdPrimTwin RootUsdTwin;
-
-	TWeakObjectPtr< ALevelSequenceActor > LevelSequenceActor;
-
-	TMultiMap< FString, FDelegateHandle > PrimDelegates;
+	UPROPERTY(Transient)
+	UUsdPrimTwin* RootUsdTwin;
 
+	UPROPERTY(Transient)
 	TSet< FString > PrimsToAnimate;
 
+	UPROPERTY(Transient)
 	TMap< UObject*, FString > ObjectsToWatch;
 
 private:
@@ -159,8 +163,8 @@ public:
 	FUsdListener& GetUsdListener() { return UsdListener; }
 	const FUsdListener& GetUsdListener() const { return UsdListener; }
 
-	FUsdPrimTwin* GetOrCreatePrimTwin( const pxr::SdfPath& UsdPrimPath );
-	FUsdPrimTwin* ExpandPrim( const pxr::UsdPrim& Prim, FUsdSchemaTranslationContext& TranslationContext );
+	UUsdPrimTwin* GetOrCreatePrimTwin( const pxr::SdfPath& UsdPrimPath );
+	UUsdPrimTwin* ExpandPrim( const pxr::UsdPrim& Prim, FUsdSchemaTranslationContext& TranslationContext );
 	void UpdatePrim( const pxr::SdfPath& UsdPrimPath, bool bResync, FUsdSchemaTranslationContext& TranslationContext );
 
 protected:
diff --git a/Engine/Plugins/Importers/USDImporter/Source/USDStageEditor/Private/SUSDStage.cpp b/Engine/Plugins/Importers/USDImporter/Source/USDStageEditor/Private/SUSDStage.cpp
--- a/Engine/Plugins/Importers/USDImporter/Source/USDStageEditor/Private/SUSDStage.cpp
+++ b/Engine/Plugins/Importers/USDImporter/Source/USDStageEditor/Private/SUSDStage.cpp
@@ -62,9 +62,6 @@ void SUsdStage::Construct( const FArguments& InArgs )
 	IUsdStageModule& UsdStageModule = FModuleManager::Get().LoadModuleChecked< IUsdStageModule >( "UsdStage" );
 	UsdStageActor = &UsdStageModule.GetUsdStageActor( GWorld );
 
-	FLevelEditorModule& LevelEditorModule = FModuleManager::LoadModuleChecked<FLevelEditorModule>("LevelEditor");
-	OnMapChangedHandle = LevelEditorModule.OnMapChanged().AddSP(this, &SUsdStage::OnMapChanged);
-
 	TUsdStore< pxr::UsdStageRefPtr > UsdStage;
 
 	if ( UsdStageActor.IsValid() )
@@ -223,9 +220,6 @@ SUsdStage::~SUsdStage()
 	FCoreUObjectDelegates::OnObjectPropertyChanged.Remove( OnStageActorPropertyChangedHandle );
 	AUsdStageActor::OnActorLoaded.Remove( OnActorLoadedHandle );
 
-	FLevelEditorModule& LevelEditorModule = FModuleManager::LoadModuleChecked<FLevelEditorModule>("LevelEditor");
-	LevelEditorModule.OnMapChanged().Remove(OnMapChangedHandle);
-
 	ClearStageActorDelegates();
 }
 
@@ -553,6 +547,11 @@ void SUsdStage::OnNew()
 
 	if ( UsdFilePath )
 	{
+		FScopedTransaction Transaction(FText::Format(
+			LOCTEXT("NewStageTransaction", "Created new USD stage '{0}'"),
+			FText::FromString(UsdFilePath.GetValue())
+		));
+
 		{
 			FScopedUsdAllocs UsdAllocs;
 
@@ -831,11 +830,6 @@ void SUsdStage::OnStageActorPropertyChanged( UObject* ObjectBeingModified, FProp
 	}
 }
 
-void SUsdStage::OnMapChanged(UWorld* World, EMapChangeType ChangeType)
-{
-	CloseStage();
-}
-
 #endif // #if USE_USD_SDK
 
 #undef LOCTEXT_NAMESPACE
diff --git a/Engine/Plugins/Importers/USDImporter/Source/USDStageEditor/Private/SUSDStage.h b/Engine/Plugins/Importers/USDImporter/Source/USDStageEditor/Private/SUSDStage.h
--- a/Engine/Plugins/Importers/USDImporter/Source/USDStageEditor/Private/SUSDStage.h
+++ b/Engine/Plugins/Importers/USDImporter/Source/USDStageEditor/Private/SUSDStage.h
@@ -75,7 +75,6 @@ protected:
 	FDelegateHandle OnStageChangedHandle;
 	FDelegateHandle OnStageEditTargetChangedHandle;
 	FDelegateHandle OnPrimChangedHandle;
-	FDelegateHandle OnMapChangedHandle;
 
 	FString SelectedPrimPath;
 };
