USD: Remove PurposeVisibility code;

Write back USceneComponent's visibility to the stage;
Fix crash when disabling prims near schemas that can't be translated;

diff --git a/Engine/Plugins/Importers/USDImporter/Source/USDStage/Private/USDStageActor.cpp b/Engine/Plugins/Importers/USDImporter/Source/USDStage/Private/USDStageActor.cpp
--- a/Engine/Plugins/Importers/USDImporter/Source/USDStage/Private/USDStageActor.cpp
+++ b/Engine/Plugins/Importers/USDImporter/Source/USDStage/Private/USDStageActor.cpp
@@ -95,34 +95,6 @@ struct FUsdStageActorImpl
 
 		return TranslationContext;
 	}
-
-	/**
-	 * Will recursively scan down the attach children of Parent and toggle each component's
-	 * Visibility based on whether they have the correct tag
-	 */
-	static void RecursivelyToggleVisibility(USceneComponent* Parent, const TSet<FName>& VisiblePurposeNames)
-	{
-		if (!Parent)
-		{
-			return;
-		}
-
-		bool bVisible = false;
-		for (const FName& Tag : Parent->ComponentTags)
-		{
-			if (VisiblePurposeNames.Contains(Tag))
-			{
-				bVisible = true;
-				break;
-			}
-		}
-		Parent->SetVisibility(bVisible, false);
-
-		for (USceneComponent* Child : Parent->GetAttachChildren())
-		{
-			RecursivelyToggleVisibility(Child, VisiblePurposeNames);
-		}
-	}
 };
 
 #endif // #if USE_USD_SDK
@@ -130,7 +102,6 @@ struct FUsdStageActorImpl
 AUsdStageActor::AUsdStageActor()
 	: InitialLoadSet( EUsdInitialLoadSet::LoadAll )
 	, PurposesToLoad((int32) EUsdPurpose::Proxy)
-	, PurposeVisibility((int32)(EUsdPurpose::Proxy | EUsdPurpose::Render | EUsdPurpose::Guide))
 	, Time( 0.0f )
 	, StartTimeCode( 0.f )
 	, EndTimeCode( 100.f )
@@ -175,6 +146,11 @@ AUsdStageActor::AUsdStageActor()
 
 							TranslationContext = FUsdStageActorImpl::CreateUsdSchemaTranslationContext( this, UsdToUnreal::ConvertPath( UsdPrimPath.Get() ) );
 							SchemaTranslator = UsdSchemasModule.GetTranslatorRegistry().CreateTranslatorForSchema( TranslationContext, pxr::UsdTyped( UsdPrim.Get() ) );
+
+							if ( !SchemaTranslator.IsValid() )
+							{
+								break;
+							}
 						}
 					}
 
@@ -198,8 +174,6 @@ AUsdStageActor::AUsdStageActor()
 					TranslationContext->CompleteTasks();
 				}
 
-				this->RefreshVisibilityBasedOnPurpose();
-
 				if ( this->HasAutorithyOverStage() )
 				{
 					this->OnPrimChanged.Broadcast( OriginalPrimPath, bResync );
@@ -472,12 +446,8 @@ void AUsdStageActor::PostEditChangeProperty(FPropertyChangedEvent& PropertyChang
 			LevelSequenceHelper.UpdateLevelSequence( UsdStage );
 		}
 	}
-	else if (PropertyName == GET_MEMBER_NAME_CHECKED(AUsdStageActor, PurposeVisibility))
-	{
-		RefreshVisibilityBasedOnPurpose();
-	}
-	else if (PropertyName == GET_MEMBER_NAME_CHECKED(AUsdStageActor, InitialLoadSet) ||
-			 PropertyName == GET_MEMBER_NAME_CHECKED(AUsdStageActor, PurposesToLoad))
+	else if ( PropertyName == GET_MEMBER_NAME_CHECKED( AUsdStageActor, InitialLoadSet ) ||
+		      PropertyName == GET_MEMBER_NAME_CHECKED( AUsdStageActor, PurposesToLoad ) )
 	{
 		LoadUsdStage();
 	}
@@ -524,31 +494,6 @@ void AUsdStageActor::OpenUsdStage()
 #endif // #if USE_USD_SDK
 }
 
-void AUsdStageActor::RefreshVisibilityBasedOnPurpose()
-{
-#if USE_USD_SDK
-	USceneComponent* Root = GetRootComponent();
-
-	TSet<FName> VisiblePurposeNames{ IUsdPrim::GetPurposeName(EUsdPurpose::Default) };
-	if (EnumHasAllFlags((EUsdPurpose)PurposeVisibility, EUsdPurpose::Proxy))
-	{
-		VisiblePurposeNames.Add(IUsdPrim::GetPurposeName(EUsdPurpose::Proxy));
-	}
-	if (EnumHasAllFlags((EUsdPurpose)PurposeVisibility, EUsdPurpose::Render))
-	{
-		VisiblePurposeNames.Add(IUsdPrim::GetPurposeName(EUsdPurpose::Render));
-	}
-	if (EnumHasAllFlags((EUsdPurpose)PurposeVisibility, EUsdPurpose::Guide))
-	{
-		VisiblePurposeNames.Add(IUsdPrim::GetPurposeName(EUsdPurpose::Guide));
-	}
-
-	FUsdStageActorImpl::RecursivelyToggleVisibility(Root, VisiblePurposeNames);
-
-	GEditor->RedrawLevelEditingViewports();
-#endif // #if USE_USD_SDK
-}
-
 #if WITH_EDITOR
 void AUsdStageActor::OnMapChanged(UWorld* World, EMapChangeType ChangeType)
 {
@@ -610,8 +555,6 @@ void AUsdStageActor::LoadUsdStage()
 
 	GEditor->BroadcastLevelActorListChanged();
 
-	this->RefreshVisibilityBasedOnPurpose();
-
 	OnTimeChanged.AddUObject( this, &AUsdStageActor::AnimatePrims );
 
 	// Log time spent to load the stage
@@ -763,6 +706,19 @@ void AUsdStageActor::OnPrimObjectPropertyChanged( UObject* ObjectBeingModified,
 				{
 					UnrealToUsd::ConvertSceneComponent( UsdStage, PrimSceneComponent, UsdPrim.Get() );
 				}
+
+				// We want to keep component visibilities in sync with USD, which uses inherited visibilities
+				// To accomplish that while blocking notices we must always propagate component visibility changes
+				if ( PropertyChangedEvent.GetPropertyName() == TEXT( "bVisible" ) )
+				{
+					PrimSceneComponent->SetVisibility( PrimSceneComponent->GetVisibleFlag(), true );
+				}
+
+				// Update stage window in case any of our component changes trigger USD stage changes
+				if ( this->HasAutorithyOverStage() )
+				{
+					this->OnPrimChanged.Broadcast( PrimPath, false );
+				}
 			}
 		}
 	}
diff --git a/Engine/Plugins/Importers/USDImporter/Source/USDStage/Public/USDStageActor.h b/Engine/Plugins/Importers/USDImporter/Source/USDStage/Public/USDStageActor.h
--- a/Engine/Plugins/Importers/USDImporter/Source/USDStage/Public/USDStageActor.h
+++ b/Engine/Plugins/Importers/USDImporter/Source/USDStage/Public/USDStageActor.h
@@ -62,10 +62,6 @@ public:
 	UPROPERTY(EditAnywhere, Category = "USD", meta = (Bitmask, BitmaskEnum=EUsdPurpose))
 	int32 PurposesToLoad;
 
-	/* Quickly toggle visibility of prims with specific purposes in the level based on component tags */
-	UPROPERTY(EditAnywhere, Category = "USD", meta = (Bitmask, BitmaskEnum=EUsdPurpose))
-	int32 PurposeVisibility;
-
 	UFUNCTION(BlueprintCallable, Category = "USD", meta = (CallInEditor = "true"))
 	float GetTime() const { return Time; }
 
@@ -126,8 +122,6 @@ private:
 	void OpenUsdStage();
 	void LoadUsdStage();
 
-	void RefreshVisibilityBasedOnPurpose();
-
 #if WITH_EDITOR
 	void OnMapChanged(UWorld* World, EMapChangeType ChangeType);
 #endif // WITH_EDITOR
diff --git a/Engine/Plugins/Importers/USDImporter/Source/USDStageEditor/Private/SUSDStage.cpp b/Engine/Plugins/Importers/USDImporter/Source/USDStageEditor/Private/SUSDStage.cpp
--- a/Engine/Plugins/Importers/USDImporter/Source/USDStageEditor/Private/SUSDStage.cpp
+++ b/Engine/Plugins/Importers/USDImporter/Source/USDStageEditor/Private/SUSDStage.cpp
@@ -356,14 +356,6 @@ void SUsdStage::FillOptionsMenu(FMenuBuilder& MenuBuilder)
 			false,
 			FSlateIcon(),
 			false);
-
-		MenuBuilder.AddSubMenu(
-			LOCTEXT("PurposeVisibility", "Purpose visibility"),
-			LOCTEXT("PurposeVisibility_ToolTip", "Quickly toggle which of the loaded purposes are visible on the stage"),
-			FNewMenuDelegate::CreateSP(this, &SUsdStage::FillPurposeVisibilitySubMenu),
-			false,
-			FSlateIcon(),
-			false);
 	}
 	MenuBuilder.EndSection();
 }
@@ -491,56 +483,6 @@ void SUsdStage::FillPurposesToLoadSubMenu(FMenuBuilder& MenuBuilder)
 	AddPurposeEntry(EUsdPurpose::Guide,  LOCTEXT("GuidePurpose",  "Guide"));
 }
 
-void SUsdStage::FillPurposeVisibilitySubMenu(FMenuBuilder& MenuBuilder)
-{
-	auto AddPurposeEntry = [&](const EUsdPurpose& Purpose, const FText& Text)
-	{
-		MenuBuilder.AddMenuEntry(
-			Text,
-			FText::GetEmpty(),
-			FSlateIcon(),
-			FUIAction(
-				FExecuteAction::CreateLambda([this, Purpose]()
-				{
-					if(AUsdStageActor* StageActor = UsdStageActor.Get())
-					{
-						FScopedTransaction Transaction(FText::Format(
-							LOCTEXT("PurposeVisibilityTransaction", "Change purpose visibility for USD stage actor '{0}'"),
-							FText::FromString(StageActor->GetActorLabel())
-						));
-
-						StageActor->Modify();
-						StageActor->PurposeVisibility = (int32)((EUsdPurpose)StageActor->PurposeVisibility ^ Purpose);
-
-						FPropertyChangedEvent PropertyChangedEvent(
-							FindFieldChecked< FProperty >( StageActor->GetClass(), GET_MEMBER_NAME_CHECKED( AUsdStageActor, PurposeVisibility ) )
-						);
-						StageActor->PostEditChangeProperty(PropertyChangedEvent);
-					}
-				}),
-				FCanExecuteAction::CreateLambda([this]()
-				{
-					return UsdStageActor.Get() != nullptr;
-				}),
-				FIsActionChecked::CreateLambda([this, Purpose]()
-				{
-					if(AUsdStageActor* StageActor = UsdStageActor.Get())
-					{
-						return EnumHasAllFlags((EUsdPurpose)StageActor->PurposeVisibility, Purpose);
-					}
-					return false;
-				})
-			),
-			NAME_None,
-			EUserInterfaceActionType::Check
-		);
-	};
-
-	AddPurposeEntry(EUsdPurpose::Proxy,  LOCTEXT("ProxyPurpose",  "Proxy"));
-	AddPurposeEntry(EUsdPurpose::Render, LOCTEXT("RenderPurpose", "Render"));
-	AddPurposeEntry(EUsdPurpose::Guide,  LOCTEXT("GuidePurpose",  "Guide"));
-}
-
 void SUsdStage::OnNew()
 {
 	TOptional< FString > UsdFilePath = UsdUtils::BrowseUsdFile( UsdUtils::EBrowseFileMode::Save, AsShared() );
diff --git a/Engine/Plugins/Importers/USDImporter/Source/USDStageEditor/Private/SUSDStage.h b/Engine/Plugins/Importers/USDImporter/Source/USDStageEditor/Private/SUSDStage.h
--- a/Engine/Plugins/Importers/USDImporter/Source/USDStageEditor/Private/SUSDStage.h
+++ b/Engine/Plugins/Importers/USDImporter/Source/USDStageEditor/Private/SUSDStage.h
@@ -39,7 +39,6 @@ protected:
 	void FillOptionsMenu( FMenuBuilder& MenuBuilder );
 	void FillPayloadsSubMenu( FMenuBuilder& MenuBuilder );
 	void FillPurposesToLoadSubMenu( FMenuBuilder& MenuBuilder );
-	void FillPurposeVisibilitySubMenu( FMenuBuilder& MenuBuilder );
 
 	void OnNew();
 	void OnOpen();
diff --git a/Engine/Plugins/Importers/USDImporter/Source/USDUtilities/Private/USDPrimConversion.cpp b/Engine/Plugins/Importers/USDImporter/Source/USDUtilities/Private/USDPrimConversion.cpp
--- a/Engine/Plugins/Importers/USDImporter/Source/USDUtilities/Private/USDPrimConversion.cpp
+++ b/Engine/Plugins/Importers/USDImporter/Source/USDUtilities/Private/USDPrimConversion.cpp
@@ -113,19 +113,19 @@ bool UnrealToUsd::ConvertSceneComponent( const pxr::UsdStageRefPtr& Stage, const
 
 	// Transform
 	pxr::UsdGeomXformable XForm( UsdPrim );
-	if ( XForm )
+	if ( !XForm )
 	{
-		pxr::GfMatrix4d UsdMatrix;
-		bool bResetXFormStack = false;
-		XForm.GetLocalTransformation( &UsdMatrix, &bResetXFormStack );
+		return false;
+	}
 
-		pxr::GfMatrix4d UsdTransform = UnrealToUsd::ConvertTransform( Stage, SceneComponent->GetRelativeTransform() );
+	pxr::GfMatrix4d UsdMatrix;
+	bool bResetXFormStack = false;
+	XForm.GetLocalTransformation( &UsdMatrix, &bResetXFormStack );
 
-		if ( GfIsClose( UsdMatrix, UsdTransform, THRESH_VECTORS_ARE_NEAR ) )
-		{
-			return true;
-		}
+	pxr::GfMatrix4d UsdTransform = UnrealToUsd::ConvertTransform( Stage, SceneComponent->GetRelativeTransform() );
 
+	if ( !GfIsClose( UsdMatrix, UsdTransform, THRESH_VECTORS_ARE_NEAR ) )
+	{
 		bResetXFormStack = false;
 		bool bFoundTransformOp = false;
 
@@ -152,6 +152,17 @@ bool UnrealToUsd::ConvertSceneComponent( const pxr::UsdStageRefPtr& Stage, const
 		}
 	}
 
+	// Visibility
+	bool bVisible = SceneComponent->GetVisibleFlag();
+	if ( bVisible )
+	{
+		XForm.MakeVisible();
+	}
+	else
+	{
+		XForm.MakeInvisible();
+	}
+
 	return true;
 }
 
