USD: Fix mistake when resolving shelved fix cl 11843151

diff --git a/Engine/Plugins/Importers/USDImporter/Source/USDStage/Private/USDStageActor.cpp b/Engine/Plugins/Importers/USDImporter/Source/USDStage/Private/USDStageActor.cpp
--- a/Engine/Plugins/Importers/USDImporter/Source/USDStage/Private/USDStageActor.cpp
+++ b/Engine/Plugins/Importers/USDImporter/Source/USDStage/Private/USDStageActor.cpp
@@ -595,6 +595,7 @@ void AUsdStageActor::ReloadAnimations()
 		LevelSequence = nullptr;
 		SubLayerLevelSequencesByIdentifier.Reset();
 
+		LevelSequenceHelper.InitLevelSequence(UsdStage);
 	}
 #endif // #if USE_USD_SDK
 }
