USD: Minor fixes related to prim visibility.

Fix invisible meshes contributing to the collapsed mesh;
Remove remaining PurposeVisibility code;
Fix child component becoming visible not  updating parent prim tree items;

diff --git a/Engine/Plugins/Importers/USDImporter/Source/USDImporter/Private/USDPrimResolver.cpp b/Engine/Plugins/Importers/USDImporter/Source/USDImporter/Private/USDPrimResolver.cpp
--- a/Engine/Plugins/Importers/USDImporter/Source/USDImporter/Private/USDPrimResolver.cpp
+++ b/Engine/Plugins/Importers/USDImporter/Source/USDImporter/Private/USDPrimResolver.cpp
@@ -223,7 +223,6 @@ AActor* UUSDPrimResolver::SpawnActor(FUSDSceneImportContext& ImportContext, cons
 				{
 					SpawnedActor = ImportContext.World->SpawnActor( AActor::StaticClass() );
 					USceneComponent* SceneComponent = NewObject< USceneComponent >( SpawnedActor, SpawnData.ActorName );
-					SceneComponent->ComponentTags.Add(IUsdPrim::GetPurposeName(IUsdPrim::GetPurpose(SpawnData.ActorPrim.Get())));
 					SpawnedActor->AddInstanceComponent( SceneComponent );
 					SpawnedActor->SetRootComponent( SceneComponent );
 				}
@@ -263,7 +262,6 @@ AActor* UUSDPrimResolver::SpawnActor(FUSDSceneImportContext& ImportContext, cons
 
 					UStaticMeshComponent* StaticMeshComponent = NewObject< UStaticMeshComponent >( SpawnedActor, ComponentName );
 					StaticMeshComponent->SetStaticMesh( ImportedStaticMesh );
-					StaticMeshComponent->ComponentTags.Add(IUsdPrim::GetPurposeName(IUsdPrim::GetPurpose(UsdAssetPrimToImport.Prim.Get())));
 
 					// Don't add the prim transform if its the same prim used for the actor as it's already accounted for in the ActorTransform
 					if ( UsdAssetPrimToImport.Prim.Get() != SpawnData.ActorPrim.Get() )
diff --git a/Engine/Plugins/Importers/USDImporter/Source/USDSchemas/Private/USDGeomXformableTranslator.cpp b/Engine/Plugins/Importers/USDImporter/Source/USDSchemas/Private/USDGeomXformableTranslator.cpp
--- a/Engine/Plugins/Importers/USDImporter/Source/USDSchemas/Private/USDGeomXformableTranslator.cpp
+++ b/Engine/Plugins/Importers/USDImporter/Source/USDSchemas/Private/USDGeomXformableTranslator.cpp
@@ -50,6 +50,15 @@ namespace UsdGeomXformableTranslatorImpl
 					continue;
 				}
 
+				// Ignore invisible prims
+				if ( pxr::UsdGeomImageable UsdGeomImageable = pxr::UsdGeomImageable( ChildPrim ) )
+				{
+					if ( UsdGeomImageable.ComputeVisibility() == pxr::UsdGeomTokens->invisible )
+					{
+						continue;
+					}
+				}
+
 				FTransform ChildTransform = CurrentTransform;
 
 				if ( pxr::UsdGeomXformable ChildXformable = pxr::UsdGeomXformable( ChildPrim ) )
@@ -288,11 +297,6 @@ void FUsdGeomXformableTranslator::UpdateComponents( USceneComponent* SceneCompon
 				StaticMeshComponent->RegisterComponent();
 			}
 		}
-
-		// Set purpose as a tag so its visibility can be toggled later
-		TUsdStore< pxr::UsdPrim > Prim = Schema.Get().GetPrim();
-		SceneComponent->ComponentTags.RemoveAll([&](const FName& Tag){ return Tag.ToString().EndsWith(TEXT("Purpose")); });
-		SceneComponent->ComponentTags.Add(IUsdPrim::GetPurposeName(IUsdPrim::GetPurpose(Prim.Get())));
 	}
 }
 
diff --git a/Engine/Plugins/Importers/USDImporter/Source/USDStageEditor/Private/SUSDStageTreeView.cpp b/Engine/Plugins/Importers/USDImporter/Source/USDStageEditor/Private/SUSDStageTreeView.cpp
--- a/Engine/Plugins/Importers/USDImporter/Source/USDStageEditor/Private/SUSDStageTreeView.cpp
+++ b/Engine/Plugins/Importers/USDImporter/Source/USDStageEditor/Private/SUSDStageTreeView.cpp
@@ -160,11 +160,25 @@ public:
 		RowData->bHasPayload = UsdPrim.Get().HasPayload();
 		RowData->bIsLoaded = UsdPrim.Get().IsLoaded();
 
+		bool bOldVisibility = RowData->bIsVisible;
 		if ( pxr::UsdGeomImageable UsdGeomImageable = pxr::UsdGeomImageable( UsdPrim.Get() ) )
 		{
 			RowData->bIsVisible = ( UsdGeomImageable.ComputeVisibility() != pxr::UsdGeomTokens->invisible );
 		}
 
+		// If our visibility was enabled, it may be that the visibilities of all of our parents were enabled to accomplish
+		// the target change, so we need to refresh them too. This happens when we manually change visibility on
+		// a USceneComponent and write that to the USD Stage, for example
+		if ( bOldVisibility == false && RowData->bIsVisible )
+		{
+			FUsdStageTreeItem* Item = ParentItem;
+			while ( Item )
+			{
+				Item->RefreshData(false);
+				Item = Item->ParentItem;
+			}
+		}
+
 		if ( bRefreshChildren )
 		{
 			for ( FUsdStageTreeItemRef& Child : UpdateChildren() )
diff --git a/Engine/Plugins/Importers/USDImporter/Source/UnrealUSDWrapper/Private/UnrealUSDWrapper.cpp b/Engine/Plugins/Importers/USDImporter/Source/UnrealUSDWrapper/Private/UnrealUSDWrapper.cpp
--- a/Engine/Plugins/Importers/USDImporter/Source/UnrealUSDWrapper/Private/UnrealUSDWrapper.cpp
+++ b/Engine/Plugins/Importers/USDImporter/Source/UnrealUSDWrapper/Private/UnrealUSDWrapper.cpp
@@ -510,26 +510,6 @@ EUsdPurpose IUsdPrim::GetPurpose( const UsdPrim& Prim, bool bComputed )
 	return EUsdPurpose::Default;
 }
 
-UNREALUSDWRAPPER_API FName IUsdPrim::GetPurposeName(EUsdPurpose Purpose)
-{
-	switch (Purpose)
-	{
-	case EUsdPurpose::Proxy:
-		return TEXT("ProxyPurpose");
-		break;
-	case EUsdPurpose::Render:
-		return TEXT("RenderPurpose");
-		break;
-	case EUsdPurpose::Guide:
-		return TEXT("GuidePurpose");
-		break;
-	default:
-		break;
-	}
-
-	return TEXT("DefaultPurpose");
-}
-
 bool IUsdPrim::HasGeometryData(const UsdPrim& Prim)
 {
 	return UsdGeomMesh(Prim) ? true : false;
diff --git a/Engine/Plugins/Importers/USDImporter/Source/UnrealUSDWrapper/Public/UnrealUSDWrapper.h b/Engine/Plugins/Importers/USDImporter/Source/UnrealUSDWrapper/Public/UnrealUSDWrapper.h
--- a/Engine/Plugins/Importers/USDImporter/Source/UnrealUSDWrapper/Public/UnrealUSDWrapper.h
+++ b/Engine/Plugins/Importers/USDImporter/Source/UnrealUSDWrapper/Public/UnrealUSDWrapper.h
@@ -210,7 +210,6 @@ public:
 	static UNREALUSDWRAPPER_API bool IsValidPrimName(const FString& Name, FText& OutReason);
 
 	static UNREALUSDWRAPPER_API EUsdPurpose GetPurpose(const pxr::UsdPrim& Prim, bool bComputed = true);
-	static UNREALUSDWRAPPER_API FName GetPurposeName(EUsdPurpose Purpose);
 
 	static UNREALUSDWRAPPER_API bool HasGeometryData(const pxr::UsdPrim& Prim);
 	static UNREALUSDWRAPPER_API bool HasGeometryDataOrLODVariants(const pxr::UsdPrim& Prim);
