USD - Fixed slow reload layers

Because of prim collapsing. we could receive multiple resyncs for different prims that were all collapsed into the same prim, resulting in us reloading the asset multiple times. We could also receive an update on the root which was triggering a reload of all the assets when we should have just been updating that particular prim.

diff --git a/Engine/Plugins/Importers/USDImporter/Source/USDStage/Private/USDListener.cpp b/Engine/Plugins/Importers/USDImporter/Source/USDStage/Private/USDListener.cpp
--- a/Engine/Plugins/Importers/USDImporter/Source/USDStage/Private/USDListener.cpp
+++ b/Engine/Plugins/Importers/USDImporter/Source/USDStage/Private/USDListener.cpp
@@ -51,18 +51,12 @@ void FUsdListener::HandleUsdNotice( const pxr::UsdNotice::ObjectsChanged& Notice
 	using namespace pxr;
 	using PathRange = UsdNotice::ObjectsChanged::PathRange;
 
-	if ( !OnPrimChanged.IsBound() || IsBlocked.GetValue() > 0 )
+	if ( !OnPrimsChanged.IsBound() || IsBlocked.GetValue() > 0 )
 	{
 		return;
 	}
 
-	auto BroadcastPrimUpdate = [ this ]( const pxr::SdfPath& PrimPath, bool bResync )
-	{
-		FScopedUnrealAllocs UnrealAllocs;
-
-		//FPlatformMisc::LowLevelOutputDebugStringf( TEXT("Prim %s changed (%s)\n"), *UsdToUnreal::ConvertPath( *It ), bResync ? TEXT("resync") : TEXT("update") );
-		OnPrimChanged.Broadcast( UsdToUnreal::ConvertPath( PrimPath ), bResync );
-	};
+	TMap< FString, bool > PrimsChangedList;
 
 	FScopedUsdAllocs UsdAllocs;
 
@@ -71,15 +65,26 @@ void FUsdListener::HandleUsdNotice( const pxr::UsdNotice::ObjectsChanged& Notice
 	for ( PathRange::const_iterator It = PathsToUpdate.begin(); It != PathsToUpdate.end(); ++It )
 	{
 		constexpr bool bResync = true;
-		BroadcastPrimUpdate( It->GetAbsoluteRootOrPrimPath(), bResync );
+		PrimsChangedList.Add( UsdToUnreal::ConvertPath( It->GetAbsoluteRootOrPrimPath() ), bResync );
 	}
 
 	PathsToUpdate = Notice.GetChangedInfoOnlyPaths();
 
 	for ( PathRange::const_iterator It = PathsToUpdate.begin(); It != PathsToUpdate.end(); ++It )
 	{
-		constexpr bool bResync = false;
-		BroadcastPrimUpdate( It->GetAbsoluteRootOrPrimPath(), bResync );
+		const FString PrimPath = UsdToUnreal::ConvertPath( It->GetAbsoluteRootOrPrimPath() );
+
+		if ( !PrimsChangedList.Contains( PrimPath ) )
+		{
+			constexpr bool bResync = false;
+			PrimsChangedList.Add( PrimPath, bResync );
+		}
+	}
+
+	if ( PrimsChangedList.Num() > 0 )
+	{
+		FScopedUnrealAllocs UnrealAllocs;
+		OnPrimsChanged.Broadcast( PrimsChangedList );
 	}
 }
 
diff --git a/Engine/Plugins/Importers/USDImporter/Source/USDStage/Private/USDStageActor.cpp b/Engine/Plugins/Importers/USDImporter/Source/USDStage/Private/USDStageActor.cpp
--- a/Engine/Plugins/Importers/USDImporter/Source/USDStage/Private/USDStageActor.cpp
+++ b/Engine/Plugins/Importers/USDImporter/Source/USDStage/Private/USDStageActor.cpp
@@ -125,16 +125,26 @@ AUsdStageActor::AUsdStageActor()
 #if USE_USD_SDK
 	RootUsdTwin->PrimPath = TEXT("/");
 
-	UsdListener.OnPrimChanged.AddLambda(
-		[ this ]( const FString& OriginalPrimPath, bool bResync )
+	UsdListener.OnPrimsChanged.AddLambda(
+		[ this ]( const TMap< FString, bool >& PrimsChangedList )
+		{
+			FScopedSlowTask RefreshStageTask( PrimsChangedList.Num(), LOCTEXT( "RefreshingUSDStage", "Refreshing USD Stage" ) );
+			RefreshStageTask.MakeDialog();
+
+			TSet< FString > RefreshedAssets;
+			TSet< FString > RefreshedComponents;
+
+			for ( const TPair< FString, bool >& PrimChangedInfo : PrimsChangedList )
 			{
-				auto UnwindToNonCollapsedPrim = [ &OriginalPrimPath, this ]( FUsdSchemaTranslator::ECollapsingType CollapsingType ) -> TUsdStore< pxr::SdfPath >
+				RefreshStageTask.EnterProgressFrame();
+
+				auto UnwindToNonCollapsedPrim = [ &PrimChangedInfo, this ]( FUsdSchemaTranslator::ECollapsingType CollapsingType ) -> TUsdStore< pxr::SdfPath >
 				{
 					IUsdSchemasModule& UsdSchemasModule = FModuleManager::Get().LoadModuleChecked< IUsdSchemasModule >( TEXT("USDSchemas") );
 
-					TSharedRef< FUsdSchemaTranslationContext > TranslationContext = FUsdStageActorImpl::CreateUsdSchemaTranslationContext( this, OriginalPrimPath );
+					TSharedRef< FUsdSchemaTranslationContext > TranslationContext = FUsdStageActorImpl::CreateUsdSchemaTranslationContext( this, PrimChangedInfo.Key );
 
-					TUsdStore< pxr::SdfPath > UsdPrimPath = UnrealToUsd::ConvertPath( *OriginalPrimPath );
+					TUsdStore< pxr::SdfPath > UsdPrimPath = UnrealToUsd::ConvertPath( *PrimChangedInfo.Key );
 					TUsdStore< pxr::UsdPrim > UsdPrim = this->GetUsdStage()->GetPrimAtPath( UsdPrimPath.Get() );
 
 					if ( TSharedPtr< FUsdSchemaTranslator > SchemaTranslator = UsdSchemasModule.GetTranslatorRegistry().CreateTranslatorForSchema( TranslationContext, pxr::UsdTyped( UsdPrim.Get() ) ) )
@@ -161,25 +171,45 @@ AUsdStageActor::AUsdStageActor()
 				{
 					TUsdStore< pxr::SdfPath > AssetsPrimPath = UnwindToNonCollapsedPrim( FUsdSchemaTranslator::ECollapsingType::Assets );
 
-					TSharedRef< FUsdSchemaTranslationContext > TranslationContext = FUsdStageActorImpl::CreateUsdSchemaTranslationContext( this, UsdToUnreal::ConvertPath( AssetsPrimPath.Get() ) );
-					this->LoadAssets( *TranslationContext, this->GetUsdStage()->GetPrimAtPath( AssetsPrimPath.Get() ) );
+					if ( !RefreshedAssets.Contains( UsdToUnreal::ConvertPath( AssetsPrimPath.Get() ) ) )
+					{
+						TSharedRef< FUsdSchemaTranslationContext > TranslationContext = FUsdStageActorImpl::CreateUsdSchemaTranslationContext( this, UsdToUnreal::ConvertPath( AssetsPrimPath.Get() ) );
+						
+						const bool bIsResync = PrimChangedInfo.Value;
+						if ( bIsResync )
+						{
+							this->LoadAssets( *TranslationContext, this->GetUsdStage()->GetPrimAtPath( AssetsPrimPath.Get() ) );
+						}
+						else
+						{
+							this->LoadAsset( *TranslationContext, this->GetUsdStage()->GetPrimAtPath( AssetsPrimPath.Get() ) );
+						}
+
+						RefreshedAssets.Add( UsdToUnreal::ConvertPath( AssetsPrimPath.Get() ) );
+					}
 				}
 
 				// Update components
 				{
 					TUsdStore< pxr::SdfPath > ComponentsPrimPath = UnwindToNonCollapsedPrim( FUsdSchemaTranslator::ECollapsingType::Components );
 
-					TSharedRef< FUsdSchemaTranslationContext > TranslationContext = FUsdStageActorImpl::CreateUsdSchemaTranslationContext( this, UsdToUnreal::ConvertPath( ComponentsPrimPath.Get() ) );
-					this->UpdatePrim( ComponentsPrimPath.Get(), bResync, *TranslationContext );
-					TranslationContext->CompleteTasks();
+					if ( !RefreshedComponents.Contains( UsdToUnreal::ConvertPath( ComponentsPrimPath.Get() ) ) )
+					{
+						TSharedRef< FUsdSchemaTranslationContext > TranslationContext = FUsdStageActorImpl::CreateUsdSchemaTranslationContext( this, UsdToUnreal::ConvertPath( ComponentsPrimPath.Get() ) );
+						this->UpdatePrim( ComponentsPrimPath.Get(), PrimChangedInfo.Value, *TranslationContext );
+						TranslationContext->CompleteTasks();
+
+						RefreshedComponents.Add( UsdToUnreal::ConvertPath( ComponentsPrimPath.Get() ) );
+					}
 				}
 
 				if ( this->HasAutorithyOverStage() )
 				{
-					this->OnPrimChanged.Broadcast( OriginalPrimPath, bResync );
+					this->OnPrimChanged.Broadcast( PrimChangedInfo.Key, PrimChangedInfo.Value );
 				}
 			}
-		);
+			
+		} );
 
 	UsdListener.OnLayersChanged.AddLambda(
 		[&](const pxr::SdfLayerChangeListVec& ChangeVec)
@@ -471,6 +501,8 @@ void AUsdStageActor::OpenUsdStage()
 		return;
 	}
 
+	TRACE_CPUPROFILER_EVENT_SCOPE( AUsdStageActor::OpenUsdStage );
+
 	FString FilePath = IFileManager::Get().ConvertToAbsolutePathForExternalAppForRead( *RootLayer.FilePath );
 	FilePath = FPaths::GetPath(FilePath) + TEXT("/");
 	FString CleanFilename = FPaths::GetCleanFilename( RootLayer.FilePath );
@@ -522,6 +554,8 @@ void AUsdStageActor::OnMapChanged(UWorld* World, EMapChangeType ChangeType)
 void AUsdStageActor::LoadUsdStage()
 {
 #if USE_USD_SDK
+	TRACE_CPUPROFILER_EVENT_SCOPE( AUsdStageActor::LoadUsdStage );
+
 	double StartTime = FPlatformTime::Cycles64();
 
 	FScopedSlowTask SlowTask( 1.f, LOCTEXT( "LoadingUDStage", "Loading USD Stage") );
@@ -578,6 +612,8 @@ void AUsdStageActor::Refresh() const
 void AUsdStageActor::ReloadAnimations()
 {
 #if USE_USD_SDK
+	TRACE_CPUPROFILER_EVENT_SCOPE( AUsdStageActor::ReloadAnimations );
+
 	const pxr::UsdStageRefPtr& UsdStage = GetUsdStage();
 	if ( !UsdStage )
 	{
@@ -732,6 +768,23 @@ bool AUsdStageActor::HasAutorithyOverStage() const
 }
 
 #if USE_USD_SDK
+void AUsdStageActor::LoadAsset( FUsdSchemaTranslationContext& TranslationContext, const pxr::UsdPrim& Prim )
+{
+	TRACE_CPUPROFILER_EVENT_SCOPE( AUsdStageActor::LoadAsset );
+
+	const FString PrimPath = UsdToUnreal::ConvertPath( Prim.GetPrimPath() );
+	PrimPathsToAssets.Remove( PrimPath );
+
+	IUsdSchemasModule& UsdSchemasModule = FModuleManager::Get().LoadModuleChecked< IUsdSchemasModule >( TEXT("USDSchemas") );
+	if ( TSharedPtr< FUsdSchemaTranslator > SchemaTranslator = UsdSchemasModule.GetTranslatorRegistry().CreateTranslatorForSchema( TranslationContext.AsShared(), pxr::UsdTyped( Prim ) ) )
+	{
+		TRACE_CPUPROFILER_EVENT_SCOPE( AUsdStageActor::CreateAssetsForPrim );
+		SchemaTranslator->CreateAssets();
+	}
+
+	TranslationContext.CompleteTasks(); // Finish the asset tasks before moving on
+}
+
 void AUsdStageActor::LoadAssets( FUsdSchemaTranslationContext& TranslationContext, const pxr::UsdPrim& StartPrim )
 {
 	TRACE_CPUPROFILER_EVENT_SCOPE( AUsdStageActor::LoadAssets );
diff --git a/Engine/Plugins/Importers/USDImporter/Source/USDStage/Public/USDListener.h b/Engine/Plugins/Importers/USDImporter/Source/USDStage/Public/USDListener.h
--- a/Engine/Plugins/Importers/USDImporter/Source/USDStage/Public/USDListener.h
+++ b/Engine/Plugins/Importers/USDImporter/Source/USDStage/Public/USDListener.h
@@ -4,6 +4,9 @@
 
 #if USE_USD_SDK
 
+#include "Containers/ContainersFwd.h"
+#include "Containers/Map.h"
+
 #include "USDIncludesStart.h"
 
 #include "pxr/base/tf/weakBase.h"
@@ -32,8 +35,9 @@ public:
 	DECLARE_EVENT( FUsdListener, FOnStageEditTargetChanged );
 	FOnStageEditTargetChanged OnStageEditTargetChanged;
 
-	DECLARE_EVENT_TwoParams( FUsdListener, FOnPrimChanged, const FString&, bool );
-	FOnPrimChanged OnPrimChanged;
+	using FPrimsChangedList = TMap< FString, bool >;
+	DECLARE_EVENT_OneParam( FUsdListener, FOnPrimsChanged, const FPrimsChangedList& );
+	FOnPrimsChanged OnPrimsChanged;
 
 	DECLARE_EVENT_OneParam( FUsdListener, FOnLayersChanged, const pxr::SdfLayerChangeListVec& );
 	FOnLayersChanged OnLayersChanged;
diff --git a/Engine/Plugins/Importers/USDImporter/Source/USDStage/Public/USDStageActor.h b/Engine/Plugins/Importers/USDImporter/Source/USDStage/Public/USDStageActor.h
--- a/Engine/Plugins/Importers/USDImporter/Source/USDStage/Public/USDStageActor.h
+++ b/Engine/Plugins/Importers/USDImporter/Source/USDStage/Public/USDStageActor.h
@@ -162,7 +162,12 @@ public:
 	void UpdatePrim( const pxr::SdfPath& UsdPrimPath, bool bResync, FUsdSchemaTranslationContext& TranslationContext );
 
 protected:
+	/** Loads the asset for a single prim */
+	void LoadAsset( FUsdSchemaTranslationContext& TranslationContext, const pxr::UsdPrim& Prim );
+
+	/** Loads the assets for all prims from StartPrim and its children */
 	void LoadAssets( FUsdSchemaTranslationContext& TranslationContext, const pxr::UsdPrim& StartPrim );
+
 	void AnimatePrims();
 
 private:
