USD: Fix rounding errors on sequence section sizes for optimized builds

diff --git a/Engine/Plugins/Importers/USDImporter/Source/USDStage/Private/USDLevelSequenceHelper.cpp b/Engine/Plugins/Importers/USDImporter/Source/USDStage/Private/USDLevelSequenceHelper.cpp
--- a/Engine/Plugins/Importers/USDImporter/Source/USDStage/Private/USDLevelSequenceHelper.cpp
+++ b/Engine/Plugins/Importers/USDImporter/Source/USDStage/Private/USDLevelSequenceHelper.cpp
@@ -84,6 +84,12 @@ public:
 	 */
 	void InitializeActorTimeCodes(const pxr::UsdStageRefPtr& UsdStage);
 
+	/**
+	 * Similar to FrameRate.AsFrameNumber(TimeSeconds) except that it uses RoundToDouble instead of FloorToDouble, to
+	 * prevent issues with floating point precision
+	 */
+	FFrameNumber RoundAsFrameNumber(const FFrameRate& FrameRate, double TimeSeconds) const;
+
 private:
 	/** Create a FLayerTimeInfo object with data from LayerHandle, and calls itself for all of its sublayers */
 	FLayerTimeInfo* RecursivelyCreateLayerTimeInfo(const TUsdStore<pxr::SdfLayerRefPtr>& LayerHandle, TSet<FString>& ExploredFiles);
@@ -274,8 +280,8 @@ void FUsdLevelSequenceHelperImpl::CreateSubsequenceTrack(const FUsdLevelSequence
 			double SubStartSeconds  = (SubOffsetTimeCode + SubScale * SubInfo->StartTimeCode.GetValue()) / SubTimeCodesPerSecond;
 			double SubEndSeconds 	= (SubOffsetTimeCode + SubScale * SubInfo->EndTimeCode.GetValue())   / SubTimeCodesPerSecond;
 
-			FFrameNumber StartFrame = FrameRate.AsFrameNumber(SubStartSeconds);
-			int32 Duration = FrameRate.AsFrameNumber(SubEndSeconds).Value - StartFrame.Value;
+			FFrameNumber StartFrame = RoundAsFrameNumber(FrameRate, SubStartSeconds);
+			int32 Duration = RoundAsFrameNumber(FrameRate, SubEndSeconds).Value - StartFrame.Value;
 			UMovieSceneSubSection* NewSection = SubTrack->AddSequenceOnRow(*SubSequence, StartFrame, Duration, INDEX_NONE);
 
 			UE_LOG(LogUsd, Verbose, TEXT("CreateSubsequenceSections: Layer '%s' received subsection for layer '%s' starting at frame '%d' and with duration '%d'"),
@@ -290,8 +296,8 @@ void FUsdLevelSequenceHelperImpl::CreateSubsequenceTrack(const FUsdLevelSequence
 		double StartTimeSeconds   = Info->StartTimeCode.GetValue() / TimeCodesPerSecond;
 		double EndTimeSeconds     = Info->EndTimeCode.GetValue() / TimeCodesPerSecond;
 
-		FFrameNumber StartFrame = FrameRate.AsFrameNumber(StartTimeSeconds);
-		FFrameNumber EndFrame = FrameRate.AsFrameNumber(EndTimeSeconds);
+		FFrameNumber StartFrame = RoundAsFrameNumber(FrameRate, StartTimeSeconds);
+		FFrameNumber EndFrame = RoundAsFrameNumber(FrameRate, EndTimeSeconds);
 		TRange<FFrameNumber> TimeRange = TRange<FFrameNumber>::Inclusive(StartFrame, EndFrame);
 
 		MovieScene->SetPlaybackRange(TimeRange);
@@ -361,8 +367,8 @@ void FUsdLevelSequenceHelperImpl::CreateTimeTrack(const FUsdLevelSequenceHelperI
 	double TimeCodesPerSecond = Info->TimeCodesPerSecond.Get(DEFAULT_FRAMERATE);
 
 	FFrameRate DestFrameRate = MovieScene->GetTickResolution();
-	FFrameNumber StartFrame  = DestFrameRate.AsFrameNumber(StartTimeCode / TimeCodesPerSecond);
-	FFrameNumber EndFrame    = DestFrameRate.AsFrameNumber(EndTimeCode / TimeCodesPerSecond);
+	FFrameNumber StartFrame  = RoundAsFrameNumber(DestFrameRate, StartTimeCode / TimeCodesPerSecond);
+	FFrameNumber EndFrame    = RoundAsFrameNumber(DestFrameRate, EndTimeCode / TimeCodesPerSecond);
 
 	bool bSectionAdded = false;
 	UMovieSceneFloatSection* TimeSection = Cast<UMovieSceneFloatSection>(TimeTrack->FindOrAddSection(0, bSectionAdded));
@@ -450,6 +456,12 @@ void FUsdLevelSequenceHelperImpl::InitializeActorTimeCodes(const pxr::UsdStageRe
 	UE_LOG(LogUsd, Verbose, TEXT("Initialized AUsdStageActor's Stage. StartTimeCode: %f, EndTimeCode: %f, TimeCodesPerSeconds: %f"), ValidStageActor->StartTimeCode, ValidStageActor->EndTimeCode, ValidStageActor->TimeCodesPerSecond);
 }
 
+FFrameNumber FUsdLevelSequenceHelperImpl::RoundAsFrameNumber( const FFrameRate& FrameRate, double TimeSeconds ) const
+{
+	const double TimeAsFrame = ( double( TimeSeconds ) * FrameRate.Numerator ) / FrameRate.Denominator;
+	return FFrameNumber(static_cast<int32>( FMath::RoundToDouble( TimeAsFrame ) ));
+}
+
 FUsdLevelSequenceHelperImpl::FLayerTimeInfo* FUsdLevelSequenceHelperImpl::RecursivelyCreateLayerTimeInfo(const TUsdStore<pxr::SdfLayerRefPtr>& LayerHandleWrapper, TSet<FString>& ExploredFiles)
 {
 	const pxr::SdfLayerRefPtr& LayerRefPtr = LayerHandleWrapper.Get();
@@ -653,8 +665,8 @@ void FUsdLevelSequenceHelperImpl::UpdateInfoFromSection(FUsdLevelSequenceHelperI
 	double SubEndSeconds   = (OffsetForSub + ScaleForSub * SubEndTimeCode)   / SubCodesPerSecond;
 
 	FFrameRate FrameRate = MovieScene->GetTickResolution();
-	FFrameNumber OriginalStartFrame = FrameRate.AsFrameNumber(SubStartSeconds);
-	FFrameNumber OriginalEndFrame   = FrameRate.AsFrameNumber(SubEndSeconds);
+	FFrameNumber OriginalStartFrame = RoundAsFrameNumber(FrameRate, SubStartSeconds);
+	FFrameNumber OriginalEndFrame   = RoundAsFrameNumber(FrameRate, SubEndSeconds);
 	FFrameNumber ModifiedStartFrame = Section->GetInclusiveStartFrame();
 	FFrameNumber ModifiedEndFrame   = Section->GetExclusiveEndFrame();
 
