USD - Fixed skelroot anims not playing when reference from a prim that we created through the USD Stage editor
The prims that we create are XForm so when they refer to skelroot, they don't have the skel mesh schema. Now we'll dynamically change the prim type to match the default prim of the referenced file. No guarantee is given when referencing multiple files.

diff --git a/Engine/Plugins/Importers/USDImporter/Source/USDStage/Private/USDStageActor.cpp b/Engine/Plugins/Importers/USDImporter/Source/USDStage/Private/USDStageActor.cpp
--- a/Engine/Plugins/Importers/USDImporter/Source/USDStage/Private/USDStageActor.cpp
+++ b/Engine/Plugins/Importers/USDImporter/Source/USDStage/Private/USDStageActor.cpp
@@ -328,11 +328,16 @@ UUsdPrimTwin* AUsdStageActor::GetOrCreatePrimTwin( const pxr::SdfPath& UsdPrimPa
 			{
 				this->OnUsdPrimTwinDestroyed( UsdPrimTwin );
 			} );
+	}
 
-		if ( UsdUtils::IsAnimated( Prim ) )
-		{
-			PrimsToAnimate.Add( PrimPath );
-		}
+	// Update the prim animated status
+	if ( UsdUtils::IsAnimated( Prim ) )
+	{
+		PrimsToAnimate.Add( PrimPath );
+	}
+	else
+	{
+		PrimsToAnimate.Remove( PrimPath );
 	}
 
 	return UsdPrimTwin;
diff --git a/Engine/Plugins/Importers/USDImporter/Source/USDStageEditor/Private/SUSDPrimInfo.cpp b/Engine/Plugins/Importers/USDImporter/Source/USDStageEditor/Private/SUSDPrimInfo.cpp
--- a/Engine/Plugins/Importers/USDImporter/Source/USDStageEditor/Private/SUSDPrimInfo.cpp
+++ b/Engine/Plugins/Importers/USDImporter/Source/USDStageEditor/Private/SUSDPrimInfo.cpp
@@ -44,7 +44,7 @@ namespace UsdPrimInfoWidgetConstants
 void SUsdPrimInfo::Construct( const FArguments& InArgs, const TUsdStore< pxr::UsdStageRefPtr >& UsdStage, const TCHAR* PrimPath )
 {
 	TSharedRef< SWidget > VariantSetsWidget = GenerateVariantSetsWidget( UsdStage, PrimPath );
-	TSharedRef< SWidget > ReferencesListWidget = GenerateVariantSetsWidget( UsdStage, PrimPath );
+	TSharedRef< SWidget > ReferencesListWidget = GenerateReferencesListWidget( UsdStage, PrimPath );
 
 	ChildSlot
 	[
diff --git a/Engine/Plugins/Importers/USDImporter/Source/USDStageEditor/Private/SUSDStageTreeView.cpp b/Engine/Plugins/Importers/USDImporter/Source/USDStageEditor/Private/SUSDStageTreeView.cpp
--- a/Engine/Plugins/Importers/USDImporter/Source/USDStageEditor/Private/SUSDStageTreeView.cpp
+++ b/Engine/Plugins/Importers/USDImporter/Source/USDStageEditor/Private/SUSDStageTreeView.cpp
@@ -31,6 +31,7 @@
 #include "pxr/usd/usd/modelAPI.h"
 #include "pxr/usd/usd/prim.h"
 #include "pxr/usd/usd/references.h"
+#include "pxr/usd/usd/schemaRegistry.h"
 #include "pxr/usd/usdGeom/xform.h"
 
 #include "USDIncludesEnd.h"
@@ -778,7 +779,7 @@ void SUsdStageTreeView::OnRemovePrim()
 
 void SUsdStageTreeView::OnAddReference()
 {
-	if ( !UsdStageActor.IsValid() )
+	if ( !UsdStageActor.IsValid() || !UsdStageActor->GetUsdStage() || !UsdStageActor->GetUsdStage()->GetEditTarget().IsValid() )
 	{
 		return;
 	}
@@ -794,14 +795,33 @@ void SUsdStageTreeView::OnAddReference()
 
 	for ( FUsdStageTreeItemRef SelectedItem : MySelectedItems )
 	{
+		FString AbsoluteFilePath = FPaths::ConvertRelativePathToFull( PickedFile.GetValue() );
+
 		FScopedUsdAllocs UsdAllocs;
 
+		const std::string UsdAbsoluteFilePath = UnrealToUsd::ConvertString( *AbsoluteFilePath ).Get();
 		pxr::UsdReferences References = SelectedItem->UsdPrim.Get().GetReferences();
 
-		FString AbsoluteFilePath = FPaths::ConvertRelativePathToFull( PickedFile.GetValue() );
+		pxr::SdfLayerRefPtr ReferenceLayer = pxr::SdfLayer::FindOrOpen( UsdAbsoluteFilePath );
+
+		if ( ReferenceLayer )
+		{
+			pxr::SdfPrimSpecHandle DefaultPrimSpec = ReferenceLayer->GetPrimAtPath( pxr::SdfPath( ReferenceLayer->GetDefaultPrim() ) );
+
+			if ( DefaultPrimSpec )
+			{
+				if ( !SelectedItem->UsdPrim.Get().IsA( pxr::UsdSchemaRegistry::GetTypeFromName( DefaultPrimSpec->GetTypeName() ) ) )
+				{
+					SelectedItem->UsdPrim.Get().SetTypeName( DefaultPrimSpec->GetTypeName() ); // Set the same prim type as its reference so that they are compatible
+				}
+			}
+		}
+
+		FString RelativePath = AbsoluteFilePath;
+		FString LayerAbsolutePath = UsdToUnreal::ConvertString( UsdStageActor->GetUsdStage()->GetEditTarget().GetLayer()->GetRepositoryPath() );
+		FPaths::MakePathRelativeTo( RelativePath, *LayerAbsolutePath );
 
-		FPaths::MakePathRelativeTo( AbsoluteFilePath, *UsdStageActor->RootLayer.FilePath );
-		References.AddReference( UnrealToUsd::ConvertString( *AbsoluteFilePath ).Get() );
+		References.AddReference( UnrealToUsd::ConvertString( *RelativePath ).Get() );
 	}
 }
 
@@ -849,7 +869,7 @@ bool SUsdStageTreeView::CanExecutePrimAction() const
 
 	TUsdStore< pxr::UsdStageRefPtr > UsdStage =  UsdStageActor->GetUsdStage();
 
-	if ( !UsdStage.Get() )
+	if ( !UsdStage.Get() || !UsdStage.Get()->GetEditTarget().IsValid() )
 	{
 		return false;
 	}
@@ -860,7 +880,7 @@ bool SUsdStageTreeView::CanExecutePrimAction() const
 
 	for ( FUsdStageTreeItemRef SelectedItem : MySelectedItems )
 	{
-		pxr::SdfPrimSpecHandle PrimSpec = UsdStage.Get()->GetRootLayer()->GetPrimAtPath( SelectedItem->UsdPrim.Get().GetPrimPath() );
+		pxr::SdfPrimSpecHandle PrimSpec = UsdStage.Get()->GetEditTarget().GetLayer()->GetPrimAtPath( SelectedItem->UsdPrim.Get().GetPrimPath() );
 		bHasPrimSpec = bHasPrimSpec || (bool)PrimSpec;
 
 		if ( bHasPrimSpec )
diff --git a/Engine/Plugins/Importers/USDImporter/Source/USDStageEditor/Private/Widgets/SUSDReferencesList.cpp b/Engine/Plugins/Importers/USDImporter/Source/USDStageEditor/Private/Widgets/SUSDReferencesList.cpp
--- a/Engine/Plugins/Importers/USDImporter/Source/USDStageEditor/Private/Widgets/SUSDReferencesList.cpp
+++ b/Engine/Plugins/Importers/USDImporter/Source/USDStageEditor/Private/Widgets/SUSDReferencesList.cpp
@@ -92,7 +92,10 @@ void SUsdReferencesList::UpdateReferences( const TUsdStore< pxr::UsdStageRefPtr
 
 	FScopedUsdAllocs UsdAllocs;
 	
-	pxr::SdfPrimSpecHandle PrimSpec = UsdStage.Get()->GetRootLayer()->GetPrimAtPath( UnrealToUsd::ConvertPath( PrimPath ).Get() );
+	pxr::UsdPrim Prim = UsdStage.Get()->GetPrimAtPath( UnrealToUsd::ConvertPath( PrimPath ).Get() );
+
+	// Retrieve the strongest opinion prim spec, hopefully it's enough to get all references
+	pxr::SdfPrimSpecHandle PrimSpec = Prim.GetPrimStack().size() > 0 ? Prim.GetPrimStack()[0] : pxr::SdfPrimSpecHandle();
 
 	if ( PrimSpec )
 	{
