USD: Fixes related to PIE.

Fix stage disappearing when PIE;
Fix animations not working for PIE;
Fix stage closing when coming out of PIE;
Fix crash when reloading stage during PIE;
Fix stage not closing when switching maps with the USDStage window closed;

diff --git a/Engine/Plugins/Importers/USDImporter/Source/USDSchemas/Private/USDGeomXformableTranslator.cpp b/Engine/Plugins/Importers/USDImporter/Source/USDSchemas/Private/USDGeomXformableTranslator.cpp
--- a/Engine/Plugins/Importers/USDImporter/Source/USDSchemas/Private/USDGeomXformableTranslator.cpp
+++ b/Engine/Plugins/Importers/USDImporter/Source/USDSchemas/Private/USDGeomXformableTranslator.cpp
@@ -281,9 +281,10 @@ void FUsdGeomXformableTranslator::UpdateComponents( USceneComponent* SceneCompon
 
 			if ( PrimStaticMesh != StaticMeshComponent->GetStaticMesh() )
 			{
-				if ( PrimStaticMesh )
+				// Need to make sure the mesh's resources are initialized here as it may have just been built in another thread
+				// Only do this if required though, as this mesh could using these resources currently (e.g. PIE and editor world sharing the mesh)
+				if ( PrimStaticMesh && !PrimStaticMesh->AreRenderingResourcesInitialized() )
 				{
-					// Need to make sure the mesh's resources are initialized here as it may have just been built in another thread
 					PrimStaticMesh->InitResources();
 				}
 
diff --git a/Engine/Plugins/Importers/USDImporter/Source/USDStage/Private/USDPrimTwin.cpp b/Engine/Plugins/Importers/USDImporter/Source/USDStage/Private/USDPrimTwin.cpp
--- a/Engine/Plugins/Importers/USDImporter/Source/USDStage/Private/USDPrimTwin.cpp
+++ b/Engine/Plugins/Importers/USDImporter/Source/USDStage/Private/USDPrimTwin.cpp
@@ -70,13 +70,11 @@ void UUsdPrimTwin::Clear()
 
 	if ( ActorToDestroy && !ActorToDestroy->IsA< AUsdStageActor >() && !ActorToDestroy->IsActorBeingDestroyed() && ActorToDestroy->GetWorld() )
 	{
-		ActorToDestroy->Modify();
 		ActorToDestroy->GetWorld()->DestroyActor( ActorToDestroy );
 		SpawnedActor = nullptr;
 	}
 	else if ( SceneComponent.IsValid() && !SceneComponent->IsBeingDestroyed() )
 	{
-		SceneComponent->Modify();
 		SceneComponent->DestroyComponent();
 		SceneComponent = nullptr;
 	}
diff --git a/Engine/Plugins/Importers/USDImporter/Source/USDStage/Private/USDStageActor.cpp b/Engine/Plugins/Importers/USDImporter/Source/USDStage/Private/USDStageActor.cpp
--- a/Engine/Plugins/Importers/USDImporter/Source/USDStage/Private/USDStageActor.cpp
+++ b/Engine/Plugins/Importers/USDImporter/Source/USDStage/Private/USDStageActor.cpp
@@ -62,7 +62,7 @@
 
 #define LOCTEXT_NAMESPACE "USDStageActor"
 
-static const EObjectFlags DefaultObjFlag = EObjectFlags::RF_Transactional | EObjectFlags::RF_Transient;
+static const EObjectFlags DefaultObjFlag = EObjectFlags::RF_Transactional | EObjectFlags::RF_NonPIEDuplicateTransient;
 
 AUsdStageActor::FOnActorLoaded AUsdStageActor::OnActorLoaded;
 
@@ -113,7 +113,7 @@ AUsdStageActor::AUsdStageActor()
 
 	RootComponent = SceneComponent;
 
-	RootUsdTwin = NewObject<UUsdPrimTwin>(this, TEXT("RootUsdTwin"), RF_Transient | RF_Transactional);
+	RootUsdTwin = NewObject<UUsdPrimTwin>(this, TEXT("RootUsdTwin"), DefaultObjFlag);
 
 #if WITH_EDITOR
 	// We can't use PostLoad to trigger LoadUsdStage when first loading a saved level because LoadUsdStage may trigger
@@ -123,11 +123,24 @@ AUsdStageActor::AUsdStageActor()
 #endif // WITH_EDITOR
 
 #if USE_USD_SDK
+	OnTimeChanged.AddUObject( this, &AUsdStageActor::AnimatePrims );
+
 	RootUsdTwin->PrimPath = TEXT("/");
 
 	UsdListener.OnPrimsChanged.AddLambda(
 		[ this ]( const TMap< FString, bool >& PrimsChangedList )
 		{
+			// During PIE, the PIE and the editor world will respond to notices. We have to prevent any PIE
+			// objects from being added to the transaction however, or else it will be discarded when finalized.
+			// We need to keep the transaction, or else we may end up with actors outside of the transaction
+			// system that want to use assets that will be destroyed by it on an undo.
+			// Note that we can't just make the spawned components/assets nontransactional because the PIE world will transact too
+			TUniquePtr<TGuardValue<ITransaction*>> SuppressTransaction = nullptr;
+			if ( this->GetOutermost()->HasAnyPackageFlags( PKG_PlayInEditor ) )
+			{
+				SuppressTransaction = MakeUnique<TGuardValue<ITransaction*>>(GUndo, nullptr);
+			}
+
 			FScopedSlowTask RefreshStageTask( PrimsChangedList.Num(), LOCTEXT( "RefreshingUSDStage", "Refreshing USD Stage" ) );
 			RefreshStageTask.MakeDialog();
 
@@ -174,7 +187,7 @@ AUsdStageActor::AUsdStageActor()
 					if ( !RefreshedAssets.Contains( UsdToUnreal::ConvertPath( AssetsPrimPath.Get() ) ) )
 					{
 						TSharedRef< FUsdSchemaTranslationContext > TranslationContext = FUsdStageActorImpl::CreateUsdSchemaTranslationContext( this, UsdToUnreal::ConvertPath( AssetsPrimPath.Get() ) );
-						
+
 						const bool bIsResync = PrimChangedInfo.Value;
 						if ( bIsResync )
 						{
@@ -208,12 +221,17 @@ AUsdStageActor::AUsdStageActor()
 					this->OnPrimChanged.Broadcast( PrimChangedInfo.Key, PrimChangedInfo.Value );
 				}
 			}
-			
 		} );
 
 	UsdListener.OnLayersChanged.AddLambda(
 		[&](const pxr::SdfLayerChangeListVec& ChangeVec)
 		{
+			TUniquePtr<TGuardValue<ITransaction*>> SuppressTransaction = nullptr;
+			if ( this->GetOutermost()->HasAnyPackageFlags( PKG_PlayInEditor ) )
+			{
+				SuppressTransaction = MakeUnique<TGuardValue<ITransaction*>>( GUndo, nullptr );
+			}
+
 			// Check to see if any layer reloaded. If so, rebuild all of our animations as a single layer changing
 			// might propagate timecodes through all level sequences
 			for (const std::pair<pxr::SdfLayerHandle, pxr::SdfChangeList>& ChangeVecItem : ChangeVec)
@@ -254,13 +272,6 @@ AUsdStageActor::AUsdStageActor()
 
 AUsdStageActor::~AUsdStageActor()
 {
-#if USE_USD_SDK
-	if ( HasAutorithyOverStage() )
-	{
-		UnrealUSDWrapper::GetUsdStageCache().Erase( UsdStageStore.Get() );
-	}
-#endif // #if USE_USD_SDK
-
 #if WITH_EDITOR
 	FLevelEditorModule& LevelEditorModule = FModuleManager::LoadModuleChecked<FLevelEditorModule>("LevelEditor");
 	LevelEditorModule.OnMapChanged().RemoveAll(this);
@@ -551,6 +562,7 @@ void AUsdStageActor::OnMapChanged(UWorld* World, EMapChangeType ChangeType)
 			// This is in charge of clearing the SUSDStage window when switching away from our level
 			// SUSDStage window needs to update
 			OnActorDestroyed.Broadcast();
+			Reset();
 		}
 	}
 }
@@ -594,8 +606,6 @@ void AUsdStageActor::LoadUsdStage()
 
 	GEditor->BroadcastLevelActorListChanged();
 
-	OnTimeChanged.AddUObject( this, &AUsdStageActor::AnimatePrims );
-
 	// Log time spent to load the stage
 	double ElapsedSeconds = FPlatformTime::ToSeconds64(FPlatformTime::Cycles64() - StartTime);
 
@@ -608,10 +618,7 @@ void AUsdStageActor::LoadUsdStage()
 
 void AUsdStageActor::Refresh() const
 {
-	if ( HasAutorithyOverStage() )
-	{
-		OnTimeChanged.Broadcast();
-	}
+	OnTimeChanged.Broadcast();
 }
 
 void AUsdStageActor::ReloadAnimations()
@@ -649,6 +656,7 @@ void AUsdStageActor::PostTransacted(const FTransactionObjectEvent& TransactionEv
 		if (this->IsPendingKill())
 		{
 			OnActorDestroyed.Broadcast();
+			Reset();
 		}
 		// This fires when being spawned in an existing level, undo delete, redo spawn
 		else
@@ -681,6 +689,15 @@ void AUsdStageActor::PostTransacted(const FTransactionObjectEvent& TransactionEv
 	}
 }
 
+void AUsdStageActor::PostDuplicate( bool bDuplicateForPIE )
+{
+	Super::PostDuplicate( bDuplicateForPIE );
+
+#if USE_USD_SDK
+	AnimatePrims();
+#endif // #if USE_USD_SDK
+}
+
 void AUsdStageActor::OnUsdPrimTwinDestroyed( const UUsdPrimTwin& UsdPrimTwin )
 {
 	PrimsToAnimate.Remove( UsdPrimTwin.PrimPath );
@@ -753,6 +770,7 @@ void AUsdStageActor::OnPrimObjectPropertyChanged( UObject* ObjectBeingModified,
 				// To accomplish that while blocking notices we must always propagate component visibility changes
 				if ( PropertyChangedEvent.GetPropertyName() == TEXT( "bVisible" ) )
 				{
+					PrimSceneComponent->Modify();
 					PrimSceneComponent->SetVisibility( PrimSceneComponent->GetVisibleFlag(), true );
 				}
 
@@ -844,6 +862,13 @@ void AUsdStageActor::LoadAssets( FUsdSchemaTranslationContext& TranslationContex
 
 void AUsdStageActor::AnimatePrims()
 {
+	// Don't try to animate if we don't have a stage opened
+	const pxr::UsdStageRefPtr& UsdStage = GetUsdStage();
+	if (!UsdStage)
+	{
+		return;
+	}
+
 	TSharedRef< FUsdSchemaTranslationContext > TranslationContext = FUsdStageActorImpl::CreateUsdSchemaTranslationContext( this, RootUsdTwin->PrimPath );
 
 	for ( const FString& PrimToAnimate : PrimsToAnimate )
@@ -865,6 +890,7 @@ void AUsdStageActor::AnimatePrims()
 	GEditor->BroadcastLevelActorListChanged();
 	GEditor->RedrawLevelEditingViewports();
 }
+
 #endif // #if USE_USD_SDK
 
 #undef LOCTEXT_NAMESPACE
diff --git a/Engine/Plugins/Importers/USDImporter/Source/USDStage/Public/USDPrimTwin.h b/Engine/Plugins/Importers/USDImporter/Source/USDStage/Public/USDPrimTwin.h
--- a/Engine/Plugins/Importers/USDImporter/Source/USDStage/Public/USDPrimTwin.h
+++ b/Engine/Plugins/Importers/USDImporter/Source/USDStage/Public/USDPrimTwin.h
@@ -11,11 +11,12 @@
 #include "Templates/UniquePtr.h"
 #include "UObject/WeakObjectPtr.h"
 #include "UObject/WeakObjectPtrTemplates.h"
+#include "UObject/SoftObjectPtr.h"
 
 #include "USDPrimTwin.generated.h"
 
 /** The Unreal equivalent (twin) of a USD prim */
-UCLASS(Transient)
+UCLASS()
 class UUsdPrimTwin final : public UObject
 {
 	GENERATED_BODY()
@@ -51,7 +52,7 @@ public:
 	FString PrimPath;
 
 	UPROPERTY()
-	TWeakObjectPtr< class AActor > SpawnedActor;
+	TSoftObjectPtr< class AActor > SpawnedActor;
 
 	UPROPERTY()
 	TWeakObjectPtr< class USceneComponent > SceneComponent;
diff --git a/Engine/Plugins/Importers/USDImporter/Source/USDStage/Public/USDStageActor.h b/Engine/Plugins/Importers/USDImporter/Source/USDStage/Public/USDStageActor.h
--- a/Engine/Plugins/Importers/USDImporter/Source/USDStage/Public/USDStageActor.h
+++ b/Engine/Plugins/Importers/USDImporter/Source/USDStage/Public/USDStageActor.h
@@ -85,10 +85,10 @@ private:
 	UPROPERTY(VisibleAnywhere, Category = "USD")
 	float TimeCodesPerSecond;
 
-	UPROPERTY(VisibleAnywhere, Category = "USD", Transient)
+	UPROPERTY(VisibleAnywhere, Category = "USD", NonPIEDuplicateTransient)
 	ULevelSequence* LevelSequence;
 
-	UPROPERTY(Transient)
+	UPROPERTY(NonPIEDuplicateTransient)
 	TMap<FString, ULevelSequence*> SubLayerLevelSequencesByIdentifier;
 
 public:
@@ -116,6 +116,7 @@ public:
 public:
 	virtual void PostEditChangeProperty( FPropertyChangedEvent& PropertyChangedEvent ) override;
 	virtual void PostTransacted(const FTransactionObjectEvent& TransactionEvent) override;
+	virtual void PostDuplicate( bool bDuplicateForPIE ) override;
 
 private:
 	void Clear();
@@ -132,13 +133,13 @@ private:
 	bool HasAutorithyOverStage() const;
 
 private:
-	UPROPERTY(Transient)
+	UPROPERTY( NonPIEDuplicateTransient )
 	UUsdPrimTwin* RootUsdTwin;
 
-	UPROPERTY(Transient)
+	UPROPERTY( NonPIEDuplicateTransient )
 	TSet< FString > PrimsToAnimate;
 
-	UPROPERTY(Transient)
+	UPROPERTY( Transient )
 	TMap< UObject*, FString > ObjectsToWatch;
 
 private:
diff --git a/Engine/Plugins/Importers/USDImporter/Source/USDStageEditor/Private/SUSDStageTreeView.cpp b/Engine/Plugins/Importers/USDImporter/Source/USDStageEditor/Private/SUSDStageTreeView.cpp
--- a/Engine/Plugins/Importers/USDImporter/Source/USDStageEditor/Private/SUSDStageTreeView.cpp
+++ b/Engine/Plugins/Importers/USDImporter/Source/USDStageEditor/Private/SUSDStageTreeView.cpp
@@ -13,6 +13,7 @@
 #include "Framework/Application/SlateApplication.h"
 #include "Framework/MultiBox/MultiBoxBuilder.h"
 #include "Modules/ModuleManager.h"
+#include "ScopedTransaction.h"
 #include "Widgets/Images/SImage.h"
 #include "Widgets/Input/SButton.h"
 #include "Widgets/Input/SCheckBox.h"
@@ -361,6 +362,11 @@ public:
 	{
 		if ( pxr::UsdGeomImageable UsdGeomImageable = GetUsdGeomImageable( TreeItem ).Get() )
 		{
+			FScopedTransaction Transaction( FText::Format(
+				LOCTEXT( "VisibilityTransaction", "Toggle visibility of prim '{0}'" ),
+				FText::FromString( UsdToUnreal::ConvertString( TreeItem->UsdPrim.Get().GetName() ) )
+			) );
+
 			if ( TreeItem->RowData->IsVisible() )
 			{
 				UsdGeomImageable.MakeInvisible();
diff --git a/Engine/Plugins/Importers/USDImporter/Source/USDStageEditor/Private/Widgets/SUSDVariantSetsList.cpp b/Engine/Plugins/Importers/USDImporter/Source/USDStageEditor/Private/Widgets/SUSDVariantSetsList.cpp
--- a/Engine/Plugins/Importers/USDImporter/Source/USDStageEditor/Private/Widgets/SUSDVariantSetsList.cpp
+++ b/Engine/Plugins/Importers/USDImporter/Source/USDStageEditor/Private/Widgets/SUSDVariantSetsList.cpp
@@ -10,6 +10,7 @@
 
 #include "Engine/World.h"
 #include "Modules/ModuleManager.h"
+#include "ScopedTransaction.h"
 #include "Widgets/Input/STextComboBox.h"
 
 #if USE_USD_SDK
@@ -38,7 +39,7 @@ void SUsdVariantRow::Construct( const FArguments& InArgs, TSharedPtr< FUsdVarian
 }
 
 TSharedRef< SWidget > SUsdVariantRow::GenerateWidgetForColumn( const FName& ColumnName )
-{	
+{
 	TSharedRef< SWidget > ColumnWidget = SNullWidget::NullWidget;
 
 	bool bIsLeftRow = true;
@@ -200,6 +201,12 @@ void SVariantsList::OnVariantSelectionChanged( const TSharedRef< FUsdVariantSet
 		return;
 	}
 
+	FScopedTransaction Transaction( FText::Format(
+		NSLOCTEXT( "USDVariantSetsList", "SwitchVariantSetTransaction", "Switch USD Variant Set '{0}' to option '{1}'" ),
+		FText::FromString( VariantSet->SetName ),
+		FText::FromString( *VariantSet->VariantSelection )
+	) );
+
 	std::string UsdVariantSelection;
 
 	if ( VariantSet->VariantSelection )
diff --git a/Engine/Plugins/Importers/USDImporter/Source/USDUtilities/Private/USDPrimConversion.cpp b/Engine/Plugins/Importers/USDImporter/Source/USDUtilities/Private/USDPrimConversion.cpp
--- a/Engine/Plugins/Importers/USDImporter/Source/USDUtilities/Private/USDPrimConversion.cpp
+++ b/Engine/Plugins/Importers/USDImporter/Source/USDUtilities/Private/USDPrimConversion.cpp
@@ -78,6 +78,8 @@ bool UsdToUnreal::ConvertXformable( const pxr::UsdStageRefPtr& Stage, const pxr:
 
 	// Visibility
 	const bool bIsHidden = ( Xformable.ComputeVisibility( EvalTime ) == pxr::UsdGeomTokens->invisible );
+
+	SceneComponent.Modify();
 	SceneComponent.SetVisibility( !bIsHidden );
 
 	return true;
