USD - Fixed stack overflow when layer A has a sublayer to has layer A as its sublayer

diff --git a/Engine/Plugins/Importers/USDImporter/Source/USDStageEditor/Private/SUSDLayersTreeView.cpp b/Engine/Plugins/Importers/USDImporter/Source/USDStageEditor/Private/SUSDLayersTreeView.cpp
--- a/Engine/Plugins/Importers/USDImporter/Source/USDStageEditor/Private/SUSDLayersTreeView.cpp
+++ b/Engine/Plugins/Importers/USDImporter/Source/USDStageEditor/Private/SUSDLayersTreeView.cpp
@@ -121,11 +121,24 @@ public:
 
 			if ( UsdLayer )
 			{
+				TSet< FString > AllLayerIdentifiers;
+
+				FUsdLayersTreeItem* CurrentItem = this;
+				while ( CurrentItem )
+				{
+					AllLayerIdentifiers.Add( UsdToUnreal::ConvertString( CurrentItem->LayerIdentifier.Get() ) );
+					CurrentItem = CurrentItem->ParentItem;
+				}
+
 				for ( std::string SubLayerPath : UsdLayer->GetSubLayerPaths() )
 				{
 					std::string AssetPathRelativeToLayer = pxr::SdfComputeAssetPathRelativeToLayer( UsdLayer, SubLayerPath );
 
-					Children.Add( MakeShared< FUsdLayersTreeItem >( this, UsdStage, MakeUsdStore< std::string >( AssetPathRelativeToLayer ) ) );
+					// Prevent infinite recursions if a sublayer refers to a parent of the same hierarchy
+					if ( !AllLayerIdentifiers.Contains( UsdToUnreal::ConvertString( AssetPathRelativeToLayer ) ) )
+					{
+						Children.Add( MakeShared< FUsdLayersTreeItem >( this, UsdStage, MakeUsdStore< std::string >( AssetPathRelativeToLayer ) ) );
+					}
 				}
 			}
 		}
