USD - Fix for point instancers spawning a bunch of actors and not unloading them when swtiching to a different stage.
Allow schema translators to create only components and no actors.

#rnx

diff --git a/Engine/Plugins/Importers/USDImporter/Source/USDSchemas/Private/USDGeomCameraTranslator.cpp b/Engine/Plugins/Importers/USDImporter/Source/USDSchemas/Private/USDGeomCameraTranslator.cpp
--- a/Engine/Plugins/Importers/USDImporter/Source/USDSchemas/Private/USDGeomCameraTranslator.cpp
+++ b/Engine/Plugins/Importers/USDImporter/Source/USDSchemas/Private/USDGeomCameraTranslator.cpp
@@ -17,7 +17,7 @@
 USceneComponent* FUsdGeomCameraTranslator::CreateComponents()
 {
 	constexpr bool bNeedsActor = true;
-	USceneComponent* RootComponent = FUsdGeomXformableTranslator::CreateComponents( UCineCameraComponent::StaticClass(), bNeedsActor );
+	USceneComponent* RootComponent = CreateComponentsEx( TSubclassOf< USceneComponent >( UCineCameraComponent::StaticClass() ), bNeedsActor );
 
 	UpdateComponents( RootComponent );
 
diff --git a/Engine/Plugins/Importers/USDImporter/Source/USDSchemas/Private/USDGeomPointInstancerTranslator.cpp b/Engine/Plugins/Importers/USDImporter/Source/USDSchemas/Private/USDGeomPointInstancerTranslator.cpp
--- a/Engine/Plugins/Importers/USDImporter/Source/USDSchemas/Private/USDGeomPointInstancerTranslator.cpp
+++ b/Engine/Plugins/Importers/USDImporter/Source/USDSchemas/Private/USDGeomPointInstancerTranslator.cpp
@@ -85,7 +85,7 @@ namespace UsdGeomPointInstancerTranslatorImpl
 	}
 }
 
-void FUsdGeomPointInstancerTranslator::UpdateComponents(USceneComponent* PointInstancerRootComponent)
+void FUsdGeomPointInstancerTranslator::UpdateComponents( USceneComponent* PointInstancerRootComponent )
 {
 	if ( !PointInstancerRootComponent )
 	{
@@ -114,13 +114,15 @@ void FUsdGeomPointInstancerTranslator::UpdateComponents(USceneComponent* PointIn
 		for ( int32 PrototypeIndex = 0; PrototypeIndex < PrototypesPaths.size(); ++PrototypeIndex)
 		{
 			pxr::UsdPrim PrototypePrim = Prim.GetStage()->GetPrimAtPath( PrototypesPaths[PrototypeIndex] );
-			if (!Prim)
+			if ( !Prim )
 			{
 				continue;
 			}
 
 			FUsdGeomXformableTranslator PrototypeXformTranslator( Context, pxr::UsdTyped( PrototypePrim ) );
-			USceneComponent* PrototypeXformComponent = PrototypeXformTranslator.CreateComponents();
+
+			const bool bNeedsActor = false;
+			USceneComponent* PrototypeXformComponent = PrototypeXformTranslator.CreateComponentsEx( {}, bNeedsActor );
 
 			TGuardValue< USceneComponent* > ParentComponentGuard2( Context->ParentComponent, PrototypeXformComponent );
 
@@ -134,7 +136,7 @@ void FUsdGeomPointInstancerTranslator::UpdateComponents(USceneComponent* PointIn
 
 				FUsdGeomXformableTranslator PrototypeGeomMeshTranslator( UHierarchicalInstancedStaticMeshComponent::StaticClass(), Context, PrototypeGeomMesh );
 
-				USceneComponent* UsdGeomPrimComponent = PrototypeGeomMeshTranslator.CreateComponents();
+				USceneComponent* UsdGeomPrimComponent = PrototypeGeomMeshTranslator.CreateComponentsEx( {}, bNeedsActor );
 
 				if ( UHierarchicalInstancedStaticMeshComponent* HismComponent = Cast< UHierarchicalInstancedStaticMeshComponent >( UsdGeomPrimComponent ) )
 				{
diff --git a/Engine/Plugins/Importers/USDImporter/Source/USDSchemas/Private/USDGeomXformableTranslator.cpp b/Engine/Plugins/Importers/USDImporter/Source/USDSchemas/Private/USDGeomXformableTranslator.cpp
--- a/Engine/Plugins/Importers/USDImporter/Source/USDSchemas/Private/USDGeomXformableTranslator.cpp
+++ b/Engine/Plugins/Importers/USDImporter/Source/USDSchemas/Private/USDGeomXformableTranslator.cpp
@@ -136,46 +136,38 @@ FUsdGeomXformableTranslator::FUsdGeomXformableTranslator( TSubclassOf< USceneCom
 
 USceneComponent* FUsdGeomXformableTranslator::CreateComponents()
 {
-	TSubclassOf< USceneComponent > ComponentType;
+	return CreateComponentsEx( {}, {} );
+}
 
-	if ( ComponentTypeOverride.IsSet() )
-	{
-		ComponentType = ComponentTypeOverride.GetValue();
-	}
-	else
+USceneComponent* FUsdGeomXformableTranslator::CreateComponentsEx( TOptional< TSubclassOf< USceneComponent > > ComponentType, TOptional< bool > bNeedsActor )
+{
+	if ( !Context->IsValid() )
 	{
-		ComponentType = UsdUtils::GetComponentTypeForPrim( Schema.Get().GetPrim() );
-
-		if ( CollapsesChildren( ECollapsingType::Assets ) )
-		{
-			if ( UStaticMesh* PrimStaticMesh = Cast< UStaticMesh >( Context->PrimPathsToAssets.FindRef( UsdToUnreal::ConvertPath( Schema.Get().GetPath() ) ) ) )
-			{
-				// At this time, we only support collapsing static meshes together
-				ComponentType = UStaticMeshComponent::StaticClass();
-			}
-		}
+		return nullptr;
 	}
 
-	USceneComponent* RootComponent = CreateComponents( ComponentType );
+	TUsdStore< pxr::UsdPrim > Prim = Schema.Get().GetPrim();
 
-	return RootComponent;
-}
+	if ( !Prim.Get() )
+	{
+		return nullptr;
+	}
 
-USceneComponent* FUsdGeomXformableTranslator::CreateComponents( TSubclassOf< USceneComponent > ComponentType )
-{
-	TUsdStore< pxr::UsdPrim > Prim = Schema.Get().GetPrim();
+	FScopedUnrealAllocs UnrealAllocs;
 
-	// We should only add components to our transient/instanced actors, never to the AUsdStageActor or any other permanent actor
-	bool bOwnerActorIsPersistent = false;
-	if (USceneComponent* RootComponent = Context->ParentComponent)
+	if ( !bNeedsActor.IsSet() )
 	{
-		if (AActor* Actor = RootComponent->GetOwner())
+		// We should only add components to our transient/instanced actors, never to the AUsdStageActor or any other permanent actor
+		bool bOwnerActorIsPersistent = false;
+		if ( USceneComponent* RootComponent = Context->ParentComponent )
 		{
-			bOwnerActorIsPersistent = !Actor->HasAnyFlags(RF_Transient);
+			if ( AActor* Actor = RootComponent->GetOwner() )
+			{
+				bOwnerActorIsPersistent = !Actor->HasAnyFlags(RF_Transient);
+			}
 		}
-	}
 
-	const bool bNeedsActor =
+		bNeedsActor =
 		(
 			bOwnerActorIsPersistent ||
 			Context->ParentComponent == nullptr ||
@@ -184,30 +176,12 @@ USceneComponent* FUsdGeomXformableTranslator::CreateComponents( TSubclassOf< USc
 			Prim.Get().IsGroup() ||
 			UsdUtils::HasCompositionArcs( Prim.Get() )
 		);
-
-	return CreateComponents( ComponentType, bNeedsActor );
-}
-
-USceneComponent* FUsdGeomXformableTranslator::CreateComponents( TSubclassOf< USceneComponent > ComponentType, const bool bNeedsActor )
-{
-	if ( !Context->IsValid() )
-	{
-		return nullptr;
 	}
 
-	TUsdStore< pxr::UsdPrim > Prim = Schema.Get().GetPrim();
-
-	if ( !Prim.Get() )
-	{
-		return nullptr;
-	}
-
-	FScopedUnrealAllocs UnrealAllocs;
-
 	USceneComponent* SceneComponent = nullptr;
 	UObject* ComponentOuter = nullptr;
 
-	if ( bNeedsActor )
+	if ( bNeedsActor.GetValue() )
 	{
 		// Spawn actor
 		FActorSpawnParameters SpawnParameters;
@@ -234,16 +208,45 @@ USceneComponent* FUsdGeomXformableTranslator::CreateComponents( TSubclassOf< USc
 
 	if ( !SceneComponent )
 	{
-		SceneComponent = NewObject< USceneComponent >( ComponentOuter, ComponentType, FName( Prim.Get().GetName().GetText() ), Context->ObjectFlags );
+		if ( !ComponentType.IsSet() )
+		{
+			if ( ComponentTypeOverride.IsSet() )
+			{
+				ComponentType = ComponentTypeOverride.GetValue();
+			}
+			else
+			{
+				ComponentType = UsdUtils::GetComponentTypeForPrim( Schema.Get().GetPrim() );
 
-		if ( AActor* Owner = SceneComponent->GetOwner() )
+				if ( CollapsesChildren( ECollapsingType::Assets ) )
+				{
+					if ( UStaticMesh* PrimStaticMesh = Cast< UStaticMesh >( Context->PrimPathsToAssets.FindRef( UsdToUnreal::ConvertPath( Schema.Get().GetPath() ) ) ) )
+					{
+						// At this time, we only support collapsing static meshes together
+						ComponentType = UStaticMeshComponent::StaticClass();
+					}
+				}
+			}
+		}
+
+		if ( ComponentType.IsSet() && ComponentType.GetValue() != nullptr )
 		{
-			Owner->AddInstanceComponent( SceneComponent );
+			SceneComponent = NewObject< USceneComponent >( ComponentOuter, ComponentType.GetValue(), FName( Prim.Get().GetName().GetText() ), Context->ObjectFlags );
+		
+			if ( AActor* Owner = SceneComponent->GetOwner() )
+			{
+				Owner->AddInstanceComponent( SceneComponent );
+			}
 		}
 	}
 
 	if ( SceneComponent )
 	{
+		if ( !SceneComponent->GetOwner()->GetRootComponent() )
+		{
+			SceneComponent->GetOwner()->SetRootComponent( SceneComponent );
+		}
+
 		UpdateComponents( SceneComponent );
 
 		if ( Context->ParentComponent && Context->ParentComponent->Mobility == EComponentMobility::Movable )
@@ -257,11 +260,6 @@ USceneComponent* FUsdGeomXformableTranslator::CreateComponents( TSubclassOf< USc
 
 		// Attach to parent
 		SceneComponent->AttachToComponent( Context->ParentComponent, FAttachmentTransformRules::KeepRelativeTransform );
-
-		if ( !SceneComponent->GetOwner()->GetRootComponent() )
-		{
-			SceneComponent->GetOwner()->SetRootComponent( SceneComponent );
-		}
 	}
 
 	return SceneComponent;
diff --git a/Engine/Plugins/Importers/USDImporter/Source/USDSchemas/Public/USDGeomXformableTranslator.h b/Engine/Plugins/Importers/USDImporter/Source/USDSchemas/Public/USDGeomXformableTranslator.h
--- a/Engine/Plugins/Importers/USDImporter/Source/USDSchemas/Public/USDGeomXformableTranslator.h
+++ b/Engine/Plugins/Importers/USDImporter/Source/USDSchemas/Public/USDGeomXformableTranslator.h
@@ -4,6 +4,8 @@
 
 #include "USDSchemaTranslator.h"
 
+#include "Misc/Optional.h"
+
 #if USE_USD_SDK
 
 #include "USDIncludesStart.h"
@@ -30,9 +32,8 @@ public:
 	virtual bool CollapsesChildren( ECollapsingType CollapsingType ) const override;
 	virtual bool CanBeCollapsed( ECollapsingType CollapsingType ) const override { return true; }
 
-protected:
-	USceneComponent* CreateComponents( TSubclassOf< USceneComponent > ComponentType );
-	USceneComponent* CreateComponents( TSubclassOf< USceneComponent > ComponentType, const bool bNeedsActor );
+	// If the optional parameters are not set, we'll figure them out automatically.
+	USceneComponent* CreateComponentsEx( TOptional< TSubclassOf< USceneComponent > > ComponentType, TOptional< bool > bNeedsActor );
 
 private:
 	TOptional< TSubclassOf< USceneComponent > > ComponentTypeOverride;
