USD: Fix error when saving level with a loaded stage; Prevent transient properties from being saved

diff --git a/Engine/Plugins/Importers/USDImporter/Source/USDStage/Private/USDStageActor.cpp b/Engine/Plugins/Importers/USDImporter/Source/USDStage/Private/USDStageActor.cpp
--- a/Engine/Plugins/Importers/USDImporter/Source/USDStage/Private/USDStageActor.cpp
+++ b/Engine/Plugins/Importers/USDImporter/Source/USDStage/Private/USDStageActor.cpp
@@ -62,7 +62,7 @@
 
 #define LOCTEXT_NAMESPACE "USDStageActor"
 
-static const EObjectFlags DefaultObjFlag = EObjectFlags::RF_Transactional | EObjectFlags::RF_NonPIEDuplicateTransient;
+static const EObjectFlags DefaultObjFlag = EObjectFlags::RF_Transactional | EObjectFlags::RF_Transient;
 
 AUsdStageActor::FOnActorLoaded AUsdStageActor::OnActorLoaded;
 
@@ -120,6 +120,8 @@ AUsdStageActor::AUsdStageActor()
 	// Rename() calls, and that is not allowed.
 	FLevelEditorModule& LevelEditorModule = FModuleManager::LoadModuleChecked<FLevelEditorModule>("LevelEditor");
 	LevelEditorModule.OnMapChanged().AddUObject(this, &AUsdStageActor::OnMapChanged);
+	FEditorDelegates::BeginPIE.AddUObject(this, &AUsdStageActor::OnBeginPIE);
+	FEditorDelegates::PostPIEStarted.AddUObject(this, &AUsdStageActor::OnPostPIEStarted);
 #endif // WITH_EDITOR
 
 #if USE_USD_SDK
@@ -275,6 +277,8 @@ AUsdStageActor::~AUsdStageActor()
 #if WITH_EDITOR
 	FLevelEditorModule& LevelEditorModule = FModuleManager::LoadModuleChecked<FLevelEditorModule>("LevelEditor");
 	LevelEditorModule.OnMapChanged().RemoveAll(this);
+	FEditorDelegates::BeginPIE.RemoveAll(this);
+	FEditorDelegates::PostPIEStarted.RemoveAll(this);
 #endif // WITH_EDITOR
 }
 
@@ -566,6 +570,21 @@ void AUsdStageActor::OnMapChanged(UWorld* World, EMapChangeType ChangeType)
 		}
 	}
 }
+
+void AUsdStageActor::OnBeginPIE(bool bIsSimulating)
+{
+	// Remove transient flag from our spawned actors and components so they can be duplicated for PIE
+	const bool bTransient = false;
+	UpdateSpawnedObjectsTransientFlag(bTransient);
+}
+
+void AUsdStageActor::OnPostPIEStarted(bool bIsSimulating)
+{
+	// Restore transient flags to our spawned actors and components so they aren't saved otherwise
+	const bool bTransient = true;
+	UpdateSpawnedObjectsTransientFlag(bTransient);
+}
+
 #endif // WITH_EDITOR
 
 void AUsdStageActor::LoadUsdStage()
@@ -694,10 +713,75 @@ void AUsdStageActor::PostDuplicate( bool bDuplicateForPIE )
 	Super::PostDuplicate( bDuplicateForPIE );
 
 #if USE_USD_SDK
-	AnimatePrims();
+	// Setup for the very first frame when we duplicate into PIE, or else we will just show a T-pose
+	if ( bDuplicateForPIE )
+	{
+		AnimatePrims();
+	}
 #endif // #if USE_USD_SDK
 }
 
+void AUsdStageActor::PostLoad()
+{
+	Super::PostLoad();
+
+	// This may be reset to nullptr when loading a level or serializing, because the property is Transient
+	if (RootUsdTwin == nullptr)
+	{
+		RootUsdTwin = NewObject<UUsdPrimTwin>(this, TEXT("RootUsdTwin"), DefaultObjFlag);
+	}
+}
+
+void AUsdStageActor::Serialize(FArchive& Ar)
+{
+	Super::Serialize(Ar);
+
+	if (Ar.GetPortFlags() & PPF_DuplicateForPIE)
+	{
+		// We want to duplicate these properties for PIE only, as they are required to animate and listen to notices
+		Ar << LevelSequence;
+		Ar << SubLayerLevelSequencesByIdentifier;
+		Ar << RootUsdTwin;
+		Ar << PrimsToAnimate;
+		Ar << ObjectsToWatch;
+		Ar << AssetsCache;
+		Ar << PrimPathsToAssets;
+	}
+}
+
+void AUsdStageActor::UpdateSpawnedObjectsTransientFlag(bool bTransient)
+{
+	if (!RootUsdTwin)
+	{
+		return;
+	}
+
+	EObjectFlags Flag = bTransient ? EObjectFlags::RF_Transient : EObjectFlags::RF_NoFlags;
+	TFunction<void(UUsdPrimTwin&)> UpdateTransient = [=](UUsdPrimTwin& PrimTwin)
+	{
+		if (AActor* SpawnedActor = PrimTwin.SpawnedActor.Get())
+		{
+			SpawnedActor->ClearFlags(EObjectFlags::RF_Transient);
+			SpawnedActor->SetFlags(Flag);
+		}
+
+		if (USceneComponent* Component = PrimTwin.SceneComponent.Get())
+		{
+			Component->ClearFlags(EObjectFlags::RF_Transient);
+			Component->SetFlags(Flag);
+
+			if (AActor* ComponentOwner = Component->GetOwner())
+			{
+				ComponentOwner->ClearFlags(EObjectFlags::RF_Transient);
+				ComponentOwner->SetFlags(Flag);
+			}
+		}
+	};
+
+	const bool bRecursive = true;
+	RootUsdTwin->Iterate(UpdateTransient, bRecursive);
+}
+
 void AUsdStageActor::OnUsdPrimTwinDestroyed( const UUsdPrimTwin& UsdPrimTwin )
 {
 	PrimsToAnimate.Remove( UsdPrimTwin.PrimPath );
diff --git a/Engine/Plugins/Importers/USDImporter/Source/USDStage/Public/USDPrimTwin.h b/Engine/Plugins/Importers/USDImporter/Source/USDStage/Public/USDPrimTwin.h
--- a/Engine/Plugins/Importers/USDImporter/Source/USDStage/Public/USDPrimTwin.h
+++ b/Engine/Plugins/Importers/USDImporter/Source/USDStage/Public/USDPrimTwin.h
@@ -16,7 +16,7 @@
 #include "USDPrimTwin.generated.h"
 
 /** The Unreal equivalent (twin) of a USD prim */
-UCLASS()
+UCLASS( Transient )
 class UUsdPrimTwin final : public UObject
 {
 	GENERATED_BODY()
diff --git a/Engine/Plugins/Importers/USDImporter/Source/USDStage/Public/USDStageActor.h b/Engine/Plugins/Importers/USDImporter/Source/USDStage/Public/USDStageActor.h
--- a/Engine/Plugins/Importers/USDImporter/Source/USDStage/Public/USDStageActor.h
+++ b/Engine/Plugins/Importers/USDImporter/Source/USDStage/Public/USDStageActor.h
@@ -85,10 +85,10 @@ private:
 	UPROPERTY(VisibleAnywhere, Category = "USD")
 	float TimeCodesPerSecond;
 
-	UPROPERTY(VisibleAnywhere, Category = "USD", NonPIEDuplicateTransient)
+	UPROPERTY(VisibleAnywhere, Category = "USD", Transient)
 	ULevelSequence* LevelSequence;
 
-	UPROPERTY(NonPIEDuplicateTransient)
+	UPROPERTY(Transient)
 	TMap<FString, ULevelSequence*> SubLayerLevelSequencesByIdentifier;
 
 public:
@@ -117,6 +117,8 @@ public:
 	virtual void PostEditChangeProperty( FPropertyChangedEvent& PropertyChangedEvent ) override;
 	virtual void PostTransacted(const FTransactionObjectEvent& TransactionEvent) override;
 	virtual void PostDuplicate( bool bDuplicateForPIE ) override;
+	virtual void PostLoad() override;
+	virtual void Serialize(FArchive& Ar) override;
 
 private:
 	void Clear();
@@ -125,30 +127,34 @@ private:
 
 #if WITH_EDITOR
 	void OnMapChanged(UWorld* World, EMapChangeType ChangeType);
+	void OnBeginPIE(bool bIsSimulating);
+	void OnPostPIEStarted(bool bIsSimulating);
 #endif // WITH_EDITOR
 
+	void UpdateSpawnedObjectsTransientFlag( bool bTransient );
+
 	void OnUsdPrimTwinDestroyed( const UUsdPrimTwin& UsdPrimTwin );
 
 	void OnPrimObjectPropertyChanged( UObject* ObjectBeingModified, FPropertyChangedEvent& PropertyChangedEvent );
 	bool HasAutorithyOverStage() const;
 
 private:
-	UPROPERTY( NonPIEDuplicateTransient )
+	UPROPERTY(Transient)
 	UUsdPrimTwin* RootUsdTwin;
 
-	UPROPERTY( NonPIEDuplicateTransient )
+	UPROPERTY(Transient)
 	TSet< FString > PrimsToAnimate;
 
-	UPROPERTY( Transient )
+	UPROPERTY(Transient)
 	TMap< UObject*, FString > ObjectsToWatch;
 
 private:
 	/** Hash based assets cache */
-	UPROPERTY( NonPIEDuplicateTransient )
+	UPROPERTY(Transient)
 	TMap< FString, UObject* > AssetsCache;
 
 	/** Map of USD Prim Paths to UE assets */
-	UPROPERTY( NonPIEDuplicateTransient )
+	UPROPERTY(Transient)
 	TMap< FString, UObject* > PrimPathsToAssets;
 
 #if USE_USD_SDK
