USD - Fixed crash when importing an empty actor over an existing stage. The requested name already exists and creates a conflict.

diff --git a/Engine/Plugins/Importers/USDImporter/Source/USDImporter/Private/USDPrimResolver.cpp b/Engine/Plugins/Importers/USDImporter/Source/USDImporter/Private/USDPrimResolver.cpp
--- a/Engine/Plugins/Importers/USDImporter/Source/USDImporter/Private/USDPrimResolver.cpp
+++ b/Engine/Plugins/Importers/USDImporter/Source/USDImporter/Private/USDPrimResolver.cpp
@@ -288,7 +288,8 @@ AActor* UUSDPrimResolver::SpawnActor(FUSDSceneImportContext& ImportContext, cons
 			}
 			else
 			{
-				SpawnedActor = ActorFactory->CreateActor(ActorAsset, ImportContext.World->GetCurrentLevel(), FTransform::Identity, RF_Transactional, SpawnData.ActorName);
+				// We'll set the name later by using the ActorLabels cache
+				SpawnedActor = ActorFactory->CreateActor(ActorAsset, ImportContext.World->GetCurrentLevel(), FTransform::Identity, RF_Transactional);
 
 				// For empty group actors set their initial mobility to static
 				if ( ActorFactory == ImportContext.EmptyActorFactory )
