USD - Fix some references not resolving. The repository path might be null, use the real path in that case.

diff --git a/Engine/Plugins/Importers/USDImporter/Source/USDStageEditor/Private/SUSDStageTreeView.cpp b/Engine/Plugins/Importers/USDImporter/Source/USDStageEditor/Private/SUSDStageTreeView.cpp
--- a/Engine/Plugins/Importers/USDImporter/Source/USDStageEditor/Private/SUSDStageTreeView.cpp
+++ b/Engine/Plugins/Importers/USDImporter/Source/USDStageEditor/Private/SUSDStageTreeView.cpp
@@ -824,7 +824,11 @@ void SUsdStageTreeView::OnAddReference()
 		}
 
 		FString RelativePath = AbsoluteFilePath;
-		FString LayerAbsolutePath = UsdToUnreal::ConvertString( UsdStageActor->GetUsdStage()->GetEditTarget().GetLayer()->GetRepositoryPath() );
+
+		pxr::SdfLayerHandle EditLayer = UsdStageActor->GetUsdStage()->GetEditTarget().GetLayer();
+
+		std::string RepositoryPath = EditLayer->GetRepositoryPath().empty() ? EditLayer->GetRealPath() : EditLayer->GetRepositoryPath();
+		FString LayerAbsolutePath = UsdToUnreal::ConvertString( RepositoryPath );
 		FPaths::MakePathRelativeTo( RelativePath, *LayerAbsolutePath );
 
 		References.AddReference( UnrealToUsd::ConvertString( *RelativePath ).Get() );
