USD - Fix Camera component transform conversion from Unreal to USD

diff --git a/Engine/Plugins/Importers/USDImporter/Source/USDUtilities/Private/USDPrimConversion.cpp b/Engine/Plugins/Importers/USDImporter/Source/USDUtilities/Private/USDPrimConversion.cpp
--- a/Engine/Plugins/Importers/USDImporter/Source/USDUtilities/Private/USDPrimConversion.cpp
+++ b/Engine/Plugins/Importers/USDImporter/Source/USDUtilities/Private/USDPrimConversion.cpp
@@ -6,6 +6,7 @@
 #include "USDConversionUtils.h"
 #include "USDTypesConversion.h"
 
+#include "CineCameraActor.h"
 #include "CineCameraComponent.h"
 #include "Components/MeshComponent.h"
 #include "Components/SceneComponent.h"
@@ -120,12 +121,26 @@ bool UnrealToUsd::ConvertSceneComponent( const pxr::UsdStageRefPtr& Stage, const
 		return false;
 	}
 
+	FTransform RelativeTransform = SceneComponent->GetRelativeTransform();
+
+	FTransform AdditionalRotation;
+	if ( ACineCameraActor* CineCameraActor = Cast< ACineCameraActor >( SceneComponent->GetOwner() ) )
+	{
+		AdditionalRotation = FTransform( FRotator( 0.0f, 90.f, 0.0f ) );
+
+		if ( UsdUtils::GetUsdStageAxis( Stage ) == pxr::UsdGeomTokens->z )
+		{
+			AdditionalRotation *= FTransform( FRotator( 90.0f, 0.f, 0.0f ) );
+		}
+	}
+
+	RelativeTransform = AdditionalRotation * RelativeTransform;
+	pxr::GfMatrix4d UsdTransform = UnrealToUsd::ConvertTransform( Stage, RelativeTransform );
+
 	pxr::GfMatrix4d UsdMatrix;
 	bool bResetXFormStack = false;
 	XForm.GetLocalTransformation( &UsdMatrix, &bResetXFormStack );
 
-	pxr::GfMatrix4d UsdTransform = UnrealToUsd::ConvertTransform( Stage, SceneComponent->GetRelativeTransform() );
-
 	if ( !GfIsClose( UsdMatrix, UsdTransform, THRESH_VECTORS_ARE_NEAR ) )
 	{
 		bResetXFormStack = false;
