USD - Multithreaded skel mesh conversion

diff --git a/Engine/Plugins/Importers/USDImporter/Source/USDSchemas/Private/USDSkelRootTranslator.cpp b/Engine/Plugins/Importers/USDImporter/Source/USDSchemas/Private/USDSkelRootTranslator.cpp
--- a/Engine/Plugins/Importers/USDImporter/Source/USDSchemas/Private/USDSkelRootTranslator.cpp
+++ b/Engine/Plugins/Importers/USDImporter/Source/USDSchemas/Private/USDSkelRootTranslator.cpp
@@ -31,6 +31,8 @@ namespace UsdSkelRootTranslatorImpl
 {
 	void ProcessMaterials( const pxr::UsdPrim& UsdPrim, FSkeletalMeshImportData& SkelMeshImportData, TMap< FString, UObject* >& PrimPathsToAssets, bool bHasPrimDisplayColor, float Time )
 	{
+		TRACE_CPUPROFILER_EVENT_SCOPE( UsdSkelRootTranslatorImpl::ProcessMaterials );
+
 		for ( SkeletalMeshImportData::FMaterial& ImportedMaterial : SkelMeshImportData.Materials )
 		{
 			if ( !ImportedMaterial.Material.IsValid() )
@@ -80,90 +82,156 @@ namespace UsdSkelRootTranslatorImpl
 
 		return OutHash;
 	}
-}
 
-void FUsdSkelRootTranslator::CreateAssets()
-{
-	TRACE_CPUPROFILER_EVENT_SCOPE( FUsdSkelRootTranslator::CreateAssets );
+	bool LoadSkelMeshImportData( FSkeletalMeshImportData& SkelMeshImportData, pxr::UsdSkelCache& SkeletonCache, const pxr::UsdSkelRoot& SkeletonRoot, const float Time, TMap< FString, UObject* > PrimPathsToAssets )
+	{
+		TRACE_CPUPROFILER_EVENT_SCOPE( UsdSkelRootTranslatorImpl::LoadSkelMeshImportData );
+
+		bool bIsSkeletalDataValid = true;
+
+		// Retrieve the USD skeletal data under the SkeletonRoot into the SkeletalMeshImportData
+		{
+			FScopedUsdAllocs UsdAllocs;
+
+			SkeletonCache.Populate( SkeletonRoot );
+
+			std::vector< pxr::UsdSkelBinding > SkeletonBindings;
+			SkeletonCache.ComputeSkelBindings( SkeletonRoot, &SkeletonBindings );
+
+			// Note that there could be multiple skeleton bindings under the SkeletonRoot
+			// For now, extract just the first one
+			for ( const pxr::UsdSkelBinding& Binding : SkeletonBindings )
+			{
+				const pxr::UsdSkelSkeleton& Skeleton = Binding.GetSkeleton();
+				pxr::UsdSkelSkeletonQuery SkelQuery = SkeletonCache.GetSkelQuery( Skeleton );
+
+				const bool bSkeletonValid = UsdToUnreal::ConvertSkeleton( SkelQuery, SkelMeshImportData );
+				if ( !bSkeletonValid )
+				{
+					bIsSkeletalDataValid = false;
+					break;
+				}
+
+				for (const pxr::UsdSkelSkinningQuery& SkinningQuery : Binding.GetSkinningTargets())
+				{
+					// In USD, the skinning target need not be a mesh, but for Unreal we are only interested in skinning meshes
+					pxr::UsdGeomMesh SkinningMesh = pxr::UsdGeomMesh( SkinningQuery.GetPrim() );
+					if (SkinningMesh)
+					{
+						UsdToUnreal::ConvertSkinnedMesh( SkinningQuery, SkelMeshImportData );
+					}
+				}
 
-	FSkeletalMeshImportData SkelMeshImportData;
-	bool bIsSkeletalDataValid = true;
+				break;
+			}
+		}
 
-	// Retrieve the USD skeletal data under the SkeletonRoot into the SkeletalMeshImportData
+		return bIsSkeletalDataValid;
+	}
+
+	class FSkelRootCreateAssetsTaskChain : public FUsdSchemaTranslatorTaskChain
 	{
-		FScopedUsdAllocs UsdAllocs;
+	public:
+		explicit FSkelRootCreateAssetsTaskChain( const TSharedRef< FUsdSchemaTranslationContext >& InContext, const TUsdStore< pxr::UsdTyped >& InSchema );
 
-		pxr::UsdPrim Prim = Schema.Get().GetPrim();
+	protected:
+		// Inputs
+		TUsdStore< pxr::UsdTyped > Schema;
+		TSharedRef< FUsdSchemaTranslationContext > Context;
 
-		pxr::UsdSkelCache SkeletonCache;
-		pxr::UsdSkelRoot SkeletonRoot( Prim );
-		SkeletonCache.Populate( SkeletonRoot );
+		// Outputs
+		FSkeletalMeshImportData SkeletalMeshImportData;
 
-		pxr::SdfPath PrimPath = Prim.GetPath();
+		TUsdStore< pxr::UsdSkelCache > SkeletonCache;
 
-		std::vector< pxr::UsdSkelBinding > SkeletonBindings;
-		SkeletonCache.ComputeSkelBindings( SkeletonRoot, &SkeletonBindings );
+		void SetupTasks();
+	};
 
-		if ( SkeletonBindings.size() == 0 )
+	FSkelRootCreateAssetsTaskChain::FSkelRootCreateAssetsTaskChain( const TSharedRef< FUsdSchemaTranslationContext >& InContext, const TUsdStore< pxr::UsdTyped >& InSchema )
+		: Schema( InSchema )
+		, Context( InContext )
+	{
+		SetupTasks();
+	}
+
+	void FSkelRootCreateAssetsTaskChain::SetupTasks()
+	{
+		// Ignore prims from disabled purposes
+		if ( !EnumHasAllFlags( Context->PurposesToLoad, IUsdPrim::GetPurpose( Schema.Get().GetPrim() ) ) )
 		{
 			return;
 		}
 
-		const pxr::TfToken StageUpAxis = UsdUtils::GetUsdStageAxis( Prim.GetStage() );
+		{
+			const bool bIsAsyncTask = true;
 
-		pxr::UsdGeomXformable Xformable( Prim );
+			// Create SkeletalMeshImportData (Async)
+			Do( bIsAsyncTask,
+				[ this ]()
+				{
+					const bool bContinueTaskChain = LoadSkelMeshImportData( SkeletalMeshImportData, SkeletonCache.Get(), pxr::UsdSkelRoot( Schema.Get() ), Context->Time, Context->PrimPathsToAssets );
 
-		std::vector< double > TimeSamples;
+					return bContinueTaskChain;
+				} );
 
-		bool bHasPrimDisplayColor = false;
+			// Process materials (Main thread)
+			Then( !bIsAsyncTask,
+				[ this ]()
+				{
+					FScopedUsdAllocs UsdAllocs;
 
-		// Note that there could be multiple skeleton bindings under the SkeletonRoot
-		// For now, extract just the first one
-		for ( const pxr::UsdSkelBinding& Binding : SkeletonBindings )
-		{
-			const pxr::UsdSkelSkeleton& Skeleton = Binding.GetSkeleton();
-			pxr::UsdSkelSkeletonQuery SkelQuery = SkeletonCache.GetSkelQuery( Skeleton );
+					std::vector< pxr::UsdSkelBinding > SkeletonBindings;
+					SkeletonCache.Get().ComputeSkelBindings( pxr::UsdSkelRoot( Schema.Get() ), &SkeletonBindings );
 
-			const bool bSkeletonValid = UsdToUnreal::ConvertSkeleton( SkelQuery, SkelMeshImportData );
-			if ( !bSkeletonValid )
-			{
-				bIsSkeletalDataValid = false;
-				break;
-			}
+					bool bHasPrimDisplayColor = false;
 
-			for (const pxr::UsdSkelSkinningQuery& SkinningQuery : Binding.GetSkinningTargets())
-			{
-				// In USD, the skinning target need not be a mesh, but for Unreal we are only interested in skinning meshes
-				pxr::UsdGeomMesh SkinningMesh = pxr::UsdGeomMesh( SkinningQuery.GetPrim() );
-				if (SkinningMesh)
-				{
-					UsdToUnreal::ConvertSkinnedMesh( SkinningQuery, SkelMeshImportData );
-					bHasPrimDisplayColor = bHasPrimDisplayColor || SkinningMesh.GetDisplayColorPrimvar().IsDefined();
-				}
-			}
+					for ( const pxr::UsdSkelBinding& Binding : SkeletonBindings )
+					{
+						for ( const pxr::UsdSkelSkinningQuery& SkinningQuery : Binding.GetSkinningTargets() )
+						{
+							// In USD, the skinning target need not be a mesh, but for Unreal we are only interested in skinning meshes
+							pxr::UsdGeomMesh SkinningMesh = pxr::UsdGeomMesh( SkinningQuery.GetPrim() );
+							if ( SkinningMesh )
+							{
+								bHasPrimDisplayColor = bHasPrimDisplayColor || SkinningMesh.GetDisplayColorPrimvar().IsDefined();
+							}
+						}
+
+						UsdSkelRootTranslatorImpl::ProcessMaterials( Schema.Get().GetPrim(), SkeletalMeshImportData, Context->PrimPathsToAssets, bHasPrimDisplayColor, Context->Time );
+						break;
+					}
 
-			UsdSkelRootTranslatorImpl::ProcessMaterials( Prim, SkelMeshImportData, Context->PrimPathsToAssets, bHasPrimDisplayColor, Context->Time );
+					return true;
+				} );
 
-			break;
-		}
-	}
+			// Create USkeletalMesh (Main thread)
+			Then( !bIsAsyncTask,
+				[ this ]()
+				{
+					FSHAHash SkeletalMeshHash = UsdSkelRootTranslatorImpl::ComputeSHAHash( SkeletalMeshImportData );
 
-	if ( !bIsSkeletalDataValid )
-	{
-		return;
-	}
+					USkeletalMesh* SkeletalMesh = Cast< USkeletalMesh >( Context->AssetsCache.FindRef( SkeletalMeshHash.ToString() ) );
 
-	FSHAHash SkeletalMeshHash = UsdSkelRootTranslatorImpl::ComputeSHAHash( SkelMeshImportData );
+					if ( !SkeletalMesh )
+					{
+						SkeletalMesh = UsdToUnreal::GetSkeletalMeshFromImportData( SkeletalMeshImportData, Context->ObjectFlags );
+						Context->AssetsCache.Add( SkeletalMeshHash.ToString(), SkeletalMesh );
+					}
 
-	USkeletalMesh* SkeletalMesh = Cast< USkeletalMesh >( Context->AssetsCache.FindRef( SkeletalMeshHash.ToString() ) );
+					Context->PrimPathsToAssets.Add( UsdToUnreal::ConvertPath( Schema.Get().GetPath() ), SkeletalMesh );
 
-	if ( !SkeletalMesh )
-	{
-		SkeletalMesh = UsdToUnreal::GetSkeletalMeshFromImportData( SkelMeshImportData, Context->ObjectFlags );
-		Context->AssetsCache.Add( SkeletalMeshHash.ToString(), SkeletalMesh );
+					return true;
+				} );
+		}
 	}
+}
+
+void FUsdSkelRootTranslator::CreateAssets()
+{
+	TSharedRef< UsdSkelRootTranslatorImpl::FSkelRootCreateAssetsTaskChain > AssetsTaskChain =
+		MakeShared< UsdSkelRootTranslatorImpl::FSkelRootCreateAssetsTaskChain >( Context, Schema );
 
-	Context->PrimPathsToAssets.Add( UsdToUnreal::ConvertPath( Schema.Get().GetPath() ), SkeletalMesh );
+	Context->TranslatorTasks.Add( MoveTemp( AssetsTaskChain ) );
 }
 
 USceneComponent* FUsdSkelRootTranslator::CreateComponents()
