USD - Fixed changing the stage units not updating the geometry.
If the units have been changed, make sure we resync everything.
Also made sure we never resync the same path twice for the same set  of usd notices.

diff --git a/Engine/Plugins/Importers/USDImporter/Source/USDStage/Private/USDListener.cpp b/Engine/Plugins/Importers/USDImporter/Source/USDStage/Private/USDListener.cpp
--- a/Engine/Plugins/Importers/USDImporter/Source/USDStage/Private/USDListener.cpp
+++ b/Engine/Plugins/Importers/USDImporter/Source/USDStage/Private/USDListener.cpp
@@ -7,6 +7,12 @@
 
 #if USE_USD_SDK
 
+#include "USDIncludesStart.h"
+
+#include "pxr/usd/usdGeom/tokens.h"
+
+#include "USDIncludesEnd.h"
+
 FUsdListener::FUsdListener( const pxr::UsdStageRefPtr& Stage )
 {
 	Register( Stage );
@@ -70,13 +76,28 @@ void FUsdListener::HandleUsdNotice( const pxr::UsdNotice::ObjectsChanged& Notice
 
 	PathsToUpdate = Notice.GetChangedInfoOnlyPaths();
 
-	for ( PathRange::const_iterator It = PathsToUpdate.begin(); It != PathsToUpdate.end(); ++It )
+	for ( PathRange::const_iterator  PathToUpdateIt = PathsToUpdate.begin(); PathToUpdateIt != PathsToUpdate.end(); ++PathToUpdateIt )
 	{
-		const FString PrimPath = UsdToUnreal::ConvertPath( It->GetAbsoluteRootOrPrimPath() );
+		const FString PrimPath = UsdToUnreal::ConvertPath( PathToUpdateIt->GetAbsoluteRootOrPrimPath() );
 
 		if ( !PrimsChangedList.Contains( PrimPath ) )
 		{
-			constexpr bool bResync = false;
+			bool bResync = false;
+
+			if ( PathToUpdateIt->GetAbsoluteRootOrPrimPath() == pxr::SdfPath::AbsoluteRootPath() )
+			{
+				pxr::TfTokenVector ChangedFields = PathToUpdateIt.GetChangedFields();
+
+				for ( const pxr::TfToken& ChangeField : ChangedFields )
+				{
+					if ( ChangeField == UsdGeomTokens->metersPerUnit )
+					{
+						bResync = true; // Force a resync when changing the metersPerUnit since it affects all coordinates
+						break;
+					}
+				}
+			}
+
 			PrimsChangedList.Add( PrimPath, bResync );
 		}
 	}
diff --git a/Engine/Plugins/Importers/USDImporter/Source/USDStage/Private/USDStageActor.cpp b/Engine/Plugins/Importers/USDImporter/Source/USDStage/Private/USDStageActor.cpp
--- a/Engine/Plugins/Importers/USDImporter/Source/USDStage/Private/USDStageActor.cpp
+++ b/Engine/Plugins/Importers/USDImporter/Source/USDStage/Private/USDStageActor.cpp
@@ -129,101 +129,7 @@ AUsdStageActor::AUsdStageActor()
 
 	RootUsdTwin->PrimPath = TEXT("/");
 
-	UsdListener.OnPrimsChanged.AddLambda(
-		[ this ]( const TMap< FString, bool >& PrimsChangedList )
-		{
-			// During PIE, the PIE and the editor world will respond to notices. We have to prevent any PIE
-			// objects from being added to the transaction however, or else it will be discarded when finalized.
-			// We need to keep the transaction, or else we may end up with actors outside of the transaction
-			// system that want to use assets that will be destroyed by it on an undo.
-			// Note that we can't just make the spawned components/assets nontransactional because the PIE world will transact too
-			TUniquePtr<TGuardValue<ITransaction*>> SuppressTransaction = nullptr;
-			if ( this->GetOutermost()->HasAnyPackageFlags( PKG_PlayInEditor ) )
-			{
-				SuppressTransaction = MakeUnique<TGuardValue<ITransaction*>>(GUndo, nullptr);
-			}
-
-			FScopedSlowTask RefreshStageTask( PrimsChangedList.Num(), LOCTEXT( "RefreshingUSDStage", "Refreshing USD Stage" ) );
-			RefreshStageTask.MakeDialog();
-
-			TSet< FString > RefreshedAssets;
-			TSet< FString > RefreshedComponents;
-
-			for ( const TPair< FString, bool >& PrimChangedInfo : PrimsChangedList )
-			{
-				RefreshStageTask.EnterProgressFrame();
-
-				auto UnwindToNonCollapsedPrim = [ &PrimChangedInfo, this ]( FUsdSchemaTranslator::ECollapsingType CollapsingType ) -> TUsdStore< pxr::SdfPath >
-				{
-					IUsdSchemasModule& UsdSchemasModule = FModuleManager::Get().LoadModuleChecked< IUsdSchemasModule >( TEXT("USDSchemas") );
-
-					TSharedRef< FUsdSchemaTranslationContext > TranslationContext = FUsdStageActorImpl::CreateUsdSchemaTranslationContext( this, PrimChangedInfo.Key );
-
-					TUsdStore< pxr::SdfPath > UsdPrimPath = UnrealToUsd::ConvertPath( *PrimChangedInfo.Key );
-					TUsdStore< pxr::UsdPrim > UsdPrim = this->GetUsdStage()->GetPrimAtPath( UsdPrimPath.Get() );
-
-					if ( TSharedPtr< FUsdSchemaTranslator > SchemaTranslator = UsdSchemasModule.GetTranslatorRegistry().CreateTranslatorForSchema( TranslationContext, pxr::UsdTyped( UsdPrim.Get() ) ) )
-					{
-						while ( SchemaTranslator->IsCollapsed( CollapsingType ) )
-						{
-							UsdPrimPath = UsdPrimPath.Get().GetParentPath();
-							UsdPrim = this->GetUsdStage()->GetPrimAtPath( UsdPrimPath.Get() );
-
-							TranslationContext = FUsdStageActorImpl::CreateUsdSchemaTranslationContext( this, UsdToUnreal::ConvertPath( UsdPrimPath.Get() ) );
-							SchemaTranslator = UsdSchemasModule.GetTranslatorRegistry().CreateTranslatorForSchema( TranslationContext, pxr::UsdTyped( UsdPrim.Get() ) );
-
-							if ( !SchemaTranslator.IsValid() )
-							{
-								break;
-							}
-						}
-					}
-
-					return UsdPrimPath;
-				};
-
-				// Reload assets
-				{
-					TUsdStore< pxr::SdfPath > AssetsPrimPath = UnwindToNonCollapsedPrim( FUsdSchemaTranslator::ECollapsingType::Assets );
-
-					if ( !RefreshedAssets.Contains( UsdToUnreal::ConvertPath( AssetsPrimPath.Get() ) ) )
-					{
-						TSharedRef< FUsdSchemaTranslationContext > TranslationContext = FUsdStageActorImpl::CreateUsdSchemaTranslationContext( this, UsdToUnreal::ConvertPath( AssetsPrimPath.Get() ) );
-
-						const bool bIsResync = PrimChangedInfo.Value;
-						if ( bIsResync )
-						{
-							this->LoadAssets( *TranslationContext, this->GetUsdStage()->GetPrimAtPath( AssetsPrimPath.Get() ) );
-						}
-						else
-						{
-							this->LoadAsset( *TranslationContext, this->GetUsdStage()->GetPrimAtPath( AssetsPrimPath.Get() ) );
-						}
-
-						RefreshedAssets.Add( UsdToUnreal::ConvertPath( AssetsPrimPath.Get() ) );
-					}
-				}
-
-				// Update components
-				{
-					TUsdStore< pxr::SdfPath > ComponentsPrimPath = UnwindToNonCollapsedPrim( FUsdSchemaTranslator::ECollapsingType::Components );
-
-					if ( !RefreshedComponents.Contains( UsdToUnreal::ConvertPath( ComponentsPrimPath.Get() ) ) )
-					{
-						TSharedRef< FUsdSchemaTranslationContext > TranslationContext = FUsdStageActorImpl::CreateUsdSchemaTranslationContext( this, UsdToUnreal::ConvertPath( ComponentsPrimPath.Get() ) );
-						this->UpdatePrim( ComponentsPrimPath.Get(), PrimChangedInfo.Value, *TranslationContext );
-						TranslationContext->CompleteTasks();
-
-						RefreshedComponents.Add( UsdToUnreal::ConvertPath( ComponentsPrimPath.Get() ) );
-					}
-				}
-
-				if ( this->HasAutorithyOverStage() )
-				{
-					this->OnPrimChanged.Broadcast( PrimChangedInfo.Key, PrimChangedInfo.Value );
-				}
-			}
-		} );
+	UsdListener.OnPrimsChanged.AddUObject( this, &AUsdStageActor::OnPrimsChanged );
 
 	UsdListener.OnLayersChanged.AddLambda(
 		[&](const pxr::SdfLayerChangeListVec& ChangeVec)
@@ -272,6 +178,143 @@ AUsdStageActor::AUsdStageActor()
 	}
 }
 
+void AUsdStageActor::OnPrimsChanged( const TMap< FString, bool >& PrimsChangedList )
+{
+	// Sort paths by length so that we parse the root paths first
+	TMap< FString, bool > SortedPrimsChangedList = PrimsChangedList;
+	SortedPrimsChangedList.KeySort( []( const FString& A, const FString& B ) -> bool { return A.Len() < B.Len(); } );
+
+	// During PIE, the PIE and the editor world will respond to notices. We have to prevent any PIE
+	// objects from being added to the transaction however, or else it will be discarded when finalized.
+	// We need to keep the transaction, or else we may end up with actors outside of the transaction
+	// system that want to use assets that will be destroyed by it on an undo.
+	// Note that we can't just make the spawned components/assets nontransactional because the PIE world will transact too
+	TUniquePtr<TGuardValue<ITransaction*>> SuppressTransaction = nullptr;
+	if ( this->GetOutermost()->HasAnyPackageFlags( PKG_PlayInEditor ) )
+	{
+		SuppressTransaction = MakeUnique<TGuardValue<ITransaction*>>(GUndo, nullptr);
+	}
+
+	FScopedSlowTask RefreshStageTask( SortedPrimsChangedList.Num(), LOCTEXT( "RefreshingUSDStage", "Refreshing USD Stage" ) );
+	RefreshStageTask.MakeDialog();
+
+	TSet< FString > UpdatedAssets;
+	TSet< FString > UpdatedComponents;
+	TSet< FString > ResyncedComponents;
+
+	for ( const TPair< FString, bool >& PrimChangedInfo : SortedPrimsChangedList )
+	{
+		RefreshStageTask.EnterProgressFrame();
+
+		auto UnwindToNonCollapsedPrim = [ &PrimChangedInfo, this ]( FUsdSchemaTranslator::ECollapsingType CollapsingType ) -> TUsdStore< pxr::SdfPath >
+		{
+			IUsdSchemasModule& UsdSchemasModule = FModuleManager::Get().LoadModuleChecked< IUsdSchemasModule >( TEXT("USDSchemas") );
+
+			TSharedRef< FUsdSchemaTranslationContext > TranslationContext = FUsdStageActorImpl::CreateUsdSchemaTranslationContext( this, PrimChangedInfo.Key );
+
+			TUsdStore< pxr::SdfPath > UsdPrimPath = UnrealToUsd::ConvertPath( *PrimChangedInfo.Key );
+			TUsdStore< pxr::UsdPrim > UsdPrim = this->GetUsdStage()->GetPrimAtPath( UsdPrimPath.Get() );
+
+			if ( TSharedPtr< FUsdSchemaTranslator > SchemaTranslator = UsdSchemasModule.GetTranslatorRegistry().CreateTranslatorForSchema( TranslationContext, pxr::UsdTyped( UsdPrim.Get() ) ) )
+			{
+				while ( SchemaTranslator->IsCollapsed( CollapsingType ) )
+				{
+					UsdPrimPath = UsdPrimPath.Get().GetParentPath();
+					UsdPrim = this->GetUsdStage()->GetPrimAtPath( UsdPrimPath.Get() );
+
+					TranslationContext = FUsdStageActorImpl::CreateUsdSchemaTranslationContext( this, UsdToUnreal::ConvertPath( UsdPrimPath.Get() ) );
+					SchemaTranslator = UsdSchemasModule.GetTranslatorRegistry().CreateTranslatorForSchema( TranslationContext, pxr::UsdTyped( UsdPrim.Get() ) );
+
+					if ( !SchemaTranslator.IsValid() )
+					{
+						break;
+					}
+				}
+			}
+
+			return UsdPrimPath;
+		};
+
+		// Return if the path or any of its higher level paths are already processed
+		auto IsPathAlreadyProcessed = []( TSet< FString >& PathsProcessed, FString PathToProcess ) -> bool
+		{
+			FString SubPath;
+			FString ParentPath;
+
+			if ( PathsProcessed.Contains( TEXT("/") ) )
+			{
+				return true;
+			}
+
+			while ( !PathToProcess.IsEmpty() && !PathsProcessed.Contains( PathToProcess ) )
+			{
+				if ( PathToProcess.Split( TEXT("/"), &ParentPath, &SubPath, ESearchCase::IgnoreCase, ESearchDir::FromEnd ) )
+				{
+					PathToProcess = ParentPath;
+				}
+				else
+				{
+					return false;
+				}
+			}
+			
+			return !PathToProcess.IsEmpty() && PathsProcessed.Contains( PathToProcess );
+		};
+
+		// Reload assets
+		{
+			TUsdStore< pxr::SdfPath > AssetsPrimPath = UnwindToNonCollapsedPrim( FUsdSchemaTranslator::ECollapsingType::Assets );
+			const FString UEAssetsPrimPath = UsdToUnreal::ConvertPath( AssetsPrimPath.Get() );
+
+			if ( !IsPathAlreadyProcessed( UpdatedAssets, UEAssetsPrimPath ) )
+			{
+				TSharedRef< FUsdSchemaTranslationContext > TranslationContext = FUsdStageActorImpl::CreateUsdSchemaTranslationContext( this, UsdToUnreal::ConvertPath( AssetsPrimPath.Get() ) );
+				
+				const bool bIsResync = PrimChangedInfo.Value;
+				if ( bIsResync )
+				{
+					this->LoadAssets( *TranslationContext, this->GetUsdStage()->GetPrimAtPath( AssetsPrimPath.Get() ) );
+				}
+				else
+				{
+					this->LoadAsset( *TranslationContext, this->GetUsdStage()->GetPrimAtPath( AssetsPrimPath.Get() ) );
+				}
+
+				UpdatedAssets.Add( UEAssetsPrimPath );
+			}
+		}
+
+		// Update components
+		{
+			TUsdStore< pxr::SdfPath > ComponentsPrimPath = UnwindToNonCollapsedPrim( FUsdSchemaTranslator::ECollapsingType::Components );
+			const FString UEComponentsPrimPath = UsdToUnreal::ConvertPath( ComponentsPrimPath.Get() );
+
+			const bool bResync = PrimChangedInfo.Value;
+			TSet< FString >& RefreshedComponents = bResync ? ResyncedComponents : UpdatedComponents;
+
+			if ( !IsPathAlreadyProcessed( RefreshedComponents, UEComponentsPrimPath ) )
+			{
+				TSharedRef< FUsdSchemaTranslationContext > TranslationContext = FUsdStageActorImpl::CreateUsdSchemaTranslationContext( this, UsdToUnreal::ConvertPath( ComponentsPrimPath.Get() ) );
+				UpdatePrim( ComponentsPrimPath.Get(), PrimChangedInfo.Value, *TranslationContext );
+				TranslationContext->CompleteTasks();
+
+				RefreshedComponents.Add( UEComponentsPrimPath );
+
+				if ( bResync )
+				{
+					// Consider that the path has been updated in the case of a resync
+					UpdatedComponents.Add( UEComponentsPrimPath );
+				}
+			}
+		}
+
+		if ( HasAutorithyOverStage() )
+		{
+			OnPrimChanged.Broadcast( PrimChangedInfo.Key, PrimChangedInfo.Value );
+		}
+	}
+}
+
 AUsdStageActor::~AUsdStageActor()
 {
 #if WITH_EDITOR
diff --git a/Engine/Plugins/Importers/USDImporter/Source/USDStage/Public/USDStageActor.h b/Engine/Plugins/Importers/USDImporter/Source/USDStage/Public/USDStageActor.h
--- a/Engine/Plugins/Importers/USDImporter/Source/USDStage/Public/USDStageActor.h
+++ b/Engine/Plugins/Importers/USDImporter/Source/USDStage/Public/USDStageActor.h
@@ -133,6 +133,7 @@ private:
 
 	void UpdateSpawnedObjectsTransientFlag( bool bTransient );
 
+	void OnPrimsChanged( const TMap< FString, bool >& PrimsChangedList );
 	void OnUsdPrimTwinDestroyed( const UUsdPrimTwin& UsdPrimTwin );
 
 	void OnPrimObjectPropertyChanged( UObject* ObjectBeingModified, FPropertyChangedEvent& PropertyChangedEvent );
@@ -145,7 +146,7 @@ private:
 	UPROPERTY(Transient)
 	TSet< FString > PrimsToAnimate;
 
-	UPROPERTY(Transient)
+	UPROPERTY( Transient )
 	TMap< UObject*, FString > ObjectsToWatch;
 
 private:
