USD - Added a hash function for Usd Shade Materials so that we don't create many identicaly UMaterials
Moved the material conversion code to USDShadeConversion

diff --git a/Engine/Plugins/Importers/USDImporter/Source/USDSchemas/Private/USDShadeMaterialTranslator.cpp b/Engine/Plugins/Importers/USDImporter/Source/USDSchemas/Private/USDShadeMaterialTranslator.cpp
--- a/Engine/Plugins/Importers/USDImporter/Source/USDSchemas/Private/USDShadeMaterialTranslator.cpp
+++ b/Engine/Plugins/Importers/USDImporter/Source/USDSchemas/Private/USDShadeMaterialTranslator.cpp
@@ -2,10 +2,11 @@
 
 #include "USDShadeMaterialTranslator.h"
 
-#include "USDGeomMeshConversion.h"
+#include "USDShadeConversion.h"
 #include "USDTypesConversion.h"
 
 #include "Materials/Material.h"
+#include "Misc/SecureHash.h"
 
 #if USE_USD_SDK
 
@@ -23,7 +24,8 @@ void FUsdShadeMaterialTranslator::CreateAssets()
 		return;
 	}
 
-	UObject*& CachedMaterial = Context->AssetsCache.FindOrAdd( UsdToUnreal::ConvertPath( ShadeMaterial.GetPath() ) );
+	FSHAHash MaterialHash = UsdToUnreal::HashShadeMaterial( ShadeMaterial );
+	UObject*& CachedMaterial = Context->AssetsCache.FindOrAdd( MaterialHash.ToString() );
 
 	if ( !CachedMaterial )
 	{
diff --git a/Engine/Plugins/Importers/USDImporter/Source/USDUtilities/Private/USDGeomMeshConversion.cpp b/Engine/Plugins/Importers/USDImporter/Source/USDUtilities/Private/USDGeomMeshConversion.cpp
--- a/Engine/Plugins/Importers/USDImporter/Source/USDUtilities/Private/USDGeomMeshConversion.cpp
+++ b/Engine/Plugins/Importers/USDImporter/Source/USDUtilities/Private/USDGeomMeshConversion.cpp
@@ -10,12 +10,8 @@
 #include "USDTypesConversion.h"
 
 #include "Engine/StaticMesh.h"
-#include "Factories/TextureFactory.h"
+
 #include "MaterialEditingLibrary.h"
-#include "Materials/Material.h"
-#include "Materials/MaterialExpressionConstant.h"
-#include "Materials/MaterialExpressionConstant4Vector.h"
-#include "Materials/MaterialInstanceConstant.h"
 #include "MeshDescription.h"
 #include "Misc/Paths.h"
 #include "StaticMeshAttributes.h"
@@ -386,240 +382,6 @@ bool UsdToUnreal::ConvertGeomMesh( const pxr::UsdGeomMesh& UsdMesh, FMeshDescrip
 	return true;
 }
 
-namespace UsdGeomMeshConversionImpl
-{
-	// Finds the strongest layer contributing to an attribute
-	TUsdStore< pxr::SdfLayerHandle > FindLayerHandle( const pxr::UsdAttribute& Attr, const pxr::UsdTimeCode& Time )
-	{
-		FScopedUsdAllocs UsdAllocs;
-
-		for ( const pxr::SdfPropertySpecHandle& PropertySpec : Attr.GetPropertyStack( Time ) )
-		{
-			if ( PropertySpec->HasDefaultValue() || PropertySpec->GetLayer()->GetNumTimeSamplesForPath( PropertySpec->GetPath() ) > 0 )
-			{
-				return PropertySpec->GetLayer();
-			}
-		}
-
-		return {};
-	}
-
-	// Given an AssetPath, resolve it to an actual file path
-	FString ResolveAssetPath( const TUsdStore< pxr::SdfLayerHandle >& LayerHandle, const FString& AssetPath )
-	{
-		FString AssetPathToResolve = AssetPath;
-		AssetPathToResolve.ReplaceInline( TEXT("<UDIM>"), TEXT("1001") ); // If it's a UDIM path, we replace the UDIM tag with the start tile
-
-		pxr::ArResolverScopedCache ResolverCache;
-		pxr::ArResolver& Resolver = pxr::ArGetResolver();
-
-		const FString RelativePathToResolve = 
-			LayerHandle.Get()
-			? UsdToUnreal::ConvertString( pxr::SdfComputeAssetPathRelativeToLayer( LayerHandle.Get(), UnrealToUsd::ConvertString( *AssetPathToResolve ).Get() ) )
-			: AssetPathToResolve;
-
-		TUsdStore< std::string > ResolvedPathUsd = Resolver.Resolve( UnrealToUsd::ConvertString( *RelativePathToResolve ).Get().c_str() );
-
-		FString ResolvedAssetPath = UsdToUnreal::ConvertString( ResolvedPathUsd.Get() );
-
-		return ResolvedAssetPath;
-	}
-
-	UMaterialExpression* ParseInputTexture( pxr::UsdShadeInput& ShadeInput, UMaterial& Material, TMap< FString, UObject* >& TexturesCache )
-	{
-		UMaterialExpression* Result = nullptr;
-
-		FScopedUsdAllocs UsdAllocs;
-
-		pxr::UsdShadeConnectableAPI Source;
-		pxr::TfToken SourceName;
-		pxr::UsdShadeAttributeType AttributeType;
-
-		if ( ShadeInput.GetConnectedSource( &Source, &SourceName, &AttributeType ) )
-		{
-			//if ( SourceName == pxr::TfToken( "UsdUVTexture" ) /*pxr::UsdImagingTokens->UsdUVTexture*/ )
-			{
-				pxr::UsdShadeInput FileInput = Source.GetInput( pxr::TfToken("file") );
-
-				if ( FileInput )
-				{
-					pxr::SdfAssetPath TextureAssetPath;
-					FileInput.Get< pxr::SdfAssetPath >( &TextureAssetPath );
-
-					FScopedUnrealAllocs UnrealAllocs;
-
-					FString TexturePath = UsdToUnreal::ConvertString( TextureAssetPath.GetAssetPath() );
-					const bool bIsSupportedUdimTexture = TexturePath.Contains( TEXT("<UDIM>") );
-					FPaths::NormalizeFilename( TexturePath );
-
-					FString ResolvedTexturePath = UsdToUnreal::ConvertString( TextureAssetPath.GetResolvedPath() );
-					ResolvedTexturePath.RemoveFromStart( TEXT("\\\\?\\") );
-					FPaths::NormalizeFilename( ResolvedTexturePath );
-
-					UTexture2D* Texture2D = Cast< UTexture2D >( TexturesCache.FindRef( TexturePath ) );
-
-					if ( !Texture2D && !TexturePath.IsEmpty() )
-					{
-						bool bOutCancelled = false;
-						UTextureFactory* TextureFactory = NewObject< UTextureFactory >();
-						TextureFactory->SuppressImportOverwriteDialog();
-						TextureFactory->bUseHashAsGuid = true;
-
-						if ( bIsSupportedUdimTexture )
-						{
-							FString BaseFileName = FPaths::GetBaseFilename( TexturePath );
-
-							FString BaseFileNameBeforeUdim;
-							FString BaseFileNameAfterUdim;
-							BaseFileName.Split( TEXT("<UDIM>"), &BaseFileNameBeforeUdim, &BaseFileNameAfterUdim );
-
-							FString UdimRegexPattern = FString::Printf( TEXT(R"((%s)(\d{4})(%s))"), *BaseFileNameBeforeUdim, *BaseFileNameAfterUdim );
-							TextureFactory->UdimRegexPattern = MoveTemp( UdimRegexPattern );
-						}
-
-						if ( ResolvedTexturePath.IsEmpty() )
-						{
-							TUsdStore< pxr::SdfLayerHandle > LayerHandle = FindLayerHandle( FileInput.GetAttr(), pxr::UsdTimeCode::EarliestTime() );
-							ResolvedTexturePath = ResolveAssetPath( LayerHandle, TexturePath );
-						}
-
-						Texture2D = Cast< UTexture2D >( TextureFactory->ImportObject( UTexture2D::StaticClass(), GetTransientPackage(), NAME_None, RF_Transient, ResolvedTexturePath, TEXT(""), bOutCancelled ) );
-
-						TexturesCache.Add( TexturePath, Texture2D );
-					}
-
-					if ( Texture2D )
-					{
-						UMaterialExpressionTextureSample* TextureSample = Cast< UMaterialExpressionTextureSample >( UMaterialEditingLibrary::CreateMaterialExpression( &Material, UMaterialExpressionTextureSample::StaticClass() ) );
-						TextureSample->Texture = Texture2D;
-						TextureSample->SamplerType = UMaterialExpressionTextureBase::GetSamplerTypeForTexture( Texture2D );
-
-						Result = TextureSample;
-					}
-				}
-			}
-		}
-
-		return Result;
-	}
-}
-
-bool UsdToUnreal::ConvertMaterial( const pxr::UsdShadeMaterial& UsdShadeMaterial, UMaterial& Material )
-{
-	TMap< FString, UObject* > TexturesCache;
-	return ConvertMaterial( UsdShadeMaterial, Material, TexturesCache );
-}
-
-bool UsdToUnreal::ConvertMaterial( const pxr::UsdShadeMaterial& UsdShadeMaterial, UMaterial& Material, TMap< FString, UObject* >& TexturesCache )
-{
-	pxr::UsdShadeShader SurfaceShader = UsdShadeMaterial.ComputeSurfaceSource();
-	
-	if ( !SurfaceShader )
-	{
-		return false;
-	}
-
-	bool bHasMaterialInfo = false;
-
-	auto ParseFloatInput = [ &Material, &SurfaceShader, &TexturesCache ]( const ANSICHAR* InputName, float DefaultValue ) -> UMaterialExpression*
-	{
-		pxr::UsdShadeInput Input = SurfaceShader.GetInput( pxr::TfToken(InputName) );
-		UMaterialExpression* InputExpression = nullptr;
-
-		if ( Input )
-		{
-			if ( !Input.HasConnectedSource() )
-			{
-				float InputValue = DefaultValue;
-				Input.Get< float >( &InputValue );
-
-				if ( !FMath::IsNearlyEqual( InputValue, DefaultValue ) )
-				{
-					UMaterialExpressionConstant* ExpressionConstant = Cast< UMaterialExpressionConstant >( UMaterialEditingLibrary::CreateMaterialExpression( &Material, UMaterialExpressionConstant::StaticClass() ) );
-					ExpressionConstant->R = InputValue;
-
-					InputExpression = ExpressionConstant;
-				}
-			}
-			else
-			{
-				InputExpression = UsdGeomMeshConversionImpl::ParseInputTexture( Input, Material, TexturesCache );
-			}
-		}
-
-		return InputExpression;
-	};
-
-	// Base color
-	{
-		pxr::UsdShadeInput DiffuseInput = SurfaceShader.GetInput( pxr::TfToken("diffuseColor") );
-
-		if ( DiffuseInput )
-		{
-			UMaterialExpression* BaseColorExpression = nullptr;
-
-			if ( !DiffuseInput.HasConnectedSource() )
-			{
-				pxr::GfVec3f UsdDiffuseColor;
-				DiffuseInput.Get< pxr::GfVec3f >( &UsdDiffuseColor );
-
-				FLinearColor DiffuseColor = UsdToUnreal::ConvertColor( UsdDiffuseColor );
-
-				UMaterialExpressionConstant4Vector* Constant4VectorExpression = Cast< UMaterialExpressionConstant4Vector >( UMaterialEditingLibrary::CreateMaterialExpression( &Material, UMaterialExpressionConstant4Vector::StaticClass() ) );
-				Constant4VectorExpression->Constant = DiffuseColor;
-
-				BaseColorExpression = Constant4VectorExpression;
-			}
-			else
-			{
-				BaseColorExpression = UsdGeomMeshConversionImpl::ParseInputTexture( DiffuseInput, Material, TexturesCache );
-			}
-
-			if ( BaseColorExpression )
-			{
-				Material.BaseColor.Expression = BaseColorExpression;
-				bHasMaterialInfo = true;
-			}
-		}
-	}
-	
-	// Metallic
-	{
-		UMaterialExpression* MetallicExpression = ParseFloatInput( "metallic", 0.f );
-
-		if ( MetallicExpression )
-		{
-			Material.Metallic.Expression = MetallicExpression;
-			bHasMaterialInfo = true;
-		}
-	}
-
-	// Roughness
-	{
-		UMaterialExpression* RoughnessExpression = ParseFloatInput( "roughness", 1.f );
-
-		if ( RoughnessExpression )
-		{
-			Material.Roughness.Expression = RoughnessExpression;
-			bHasMaterialInfo = true;
-		}
-	}
-
-	// Opacity
-	{
-		UMaterialExpression* OpacityExpression = ParseFloatInput( "opacity", 1.f );;
-
-		if ( OpacityExpression )
-		{
-			Material.Opacity.Expression = OpacityExpression;
-			Material.BlendMode = BLEND_Translucent;
-			bHasMaterialInfo = true;
-		}
-	}
-
-	return bHasMaterialInfo;
-}
-
 bool UsdToUnreal::ConvertDisplayColor( const pxr::UsdGeomMesh& UsdMesh, UMaterialInstanceConstant& MaterialInstance, const pxr::UsdTimeCode TimeCode )
 {
 	FScopedUsdAllocs UsdAllocs;
diff --git a/Engine/Plugins/Importers/USDImporter/Source/USDUtilities/Private/USDShadeConversion.cpp b/Engine/Plugins/Importers/USDImporter/Source/USDUtilities/Private/USDShadeConversion.cpp
new file mode 100644
--- /dev/null
+++ b/Engine/Plugins/Importers/USDImporter/Source/USDUtilities/Private/USDShadeConversion.cpp
@@ -0,0 +1,330 @@
+// Copyright Epic Games, Inc. All Rights Reserved.
+
+#include "USDShadeConversion.h"
+
+#if USE_USD_SDK
+
+#include "USDMemory.h"
+#include "USDTypesConversion.h"
+
+#include "Engine/Texture2D.h"
+#include "Factories/TextureFactory.h"
+#include "MaterialEditingLibrary.h"
+#include "Materials/Material.h"
+#include "Materials/MaterialExpressionConstant.h"
+#include "Materials/MaterialExpressionConstant4Vector.h"
+#include "Materials/MaterialInstanceConstant.h"
+
+#include "USDIncludesStart.h"
+
+#include "pxr/usd/ar/resolver.h"
+#include "pxr/usd/ar/resolverScopedCache.h"
+#include "pxr/usd/sdf/layer.h"
+#include "pxr/usd/sdf/layerUtils.h"
+#include "pxr/usd/usdShade/material.h"
+
+#include "USDIncludesEnd.h"
+
+namespace UsdGeomMeshConversionImpl
+{
+	// Finds the strongest layer contributing to an attribute
+	TUsdStore< pxr::SdfLayerHandle > FindLayerHandle( const pxr::UsdAttribute& Attr, const pxr::UsdTimeCode& Time )
+	{
+		FScopedUsdAllocs UsdAllocs;
+
+		for ( const pxr::SdfPropertySpecHandle& PropertySpec : Attr.GetPropertyStack( Time ) )
+		{
+			if ( PropertySpec->HasDefaultValue() || PropertySpec->GetLayer()->GetNumTimeSamplesForPath( PropertySpec->GetPath() ) > 0 )
+			{
+				return PropertySpec->GetLayer();
+			}
+		}
+
+		return {};
+	}
+
+	// Given an AssetPath, resolve it to an actual file path
+	FString ResolveAssetPath( const TUsdStore< pxr::SdfLayerHandle >& LayerHandle, const FString& AssetPath )
+	{
+		FString AssetPathToResolve = AssetPath;
+		AssetPathToResolve.ReplaceInline( TEXT("<UDIM>"), TEXT("1001") ); // If it's a UDIM path, we replace the UDIM tag with the start tile
+
+		pxr::ArResolverScopedCache ResolverCache;
+		pxr::ArResolver& Resolver = pxr::ArGetResolver();
+
+		const FString RelativePathToResolve = 
+			LayerHandle.Get()
+			? UsdToUnreal::ConvertString( pxr::SdfComputeAssetPathRelativeToLayer( LayerHandle.Get(), UnrealToUsd::ConvertString( *AssetPathToResolve ).Get() ) )
+			: AssetPathToResolve;
+
+		TUsdStore< std::string > ResolvedPathUsd = Resolver.Resolve( UnrealToUsd::ConvertString( *RelativePathToResolve ).Get().c_str() );
+
+		FString ResolvedAssetPath = UsdToUnreal::ConvertString( ResolvedPathUsd.Get() );
+
+		return ResolvedAssetPath;
+	}
+
+	UMaterialExpression* ParseInputTexture( pxr::UsdShadeInput& ShadeInput, UMaterial& Material, TMap< FString, UObject* >& TexturesCache )
+	{
+		UMaterialExpression* Result = nullptr;
+
+		FScopedUsdAllocs UsdAllocs;
+
+		pxr::UsdShadeConnectableAPI Source;
+		pxr::TfToken SourceName;
+		pxr::UsdShadeAttributeType AttributeType;
+
+		if ( ShadeInput.GetConnectedSource( &Source, &SourceName, &AttributeType ) )
+		{
+			//if ( SourceName == pxr::TfToken( "UsdUVTexture" ) /*pxr::UsdImagingTokens->UsdUVTexture*/ )
+			{
+				pxr::UsdShadeInput FileInput = Source.GetInput( pxr::TfToken("file") );
+
+				if ( FileInput )
+				{
+					pxr::SdfAssetPath TextureAssetPath;
+					FileInput.Get< pxr::SdfAssetPath >( &TextureAssetPath );
+
+					FScopedUnrealAllocs UnrealAllocs;
+
+					FString TexturePath = UsdToUnreal::ConvertString( TextureAssetPath.GetAssetPath() );
+					const bool bIsSupportedUdimTexture = TexturePath.Contains( TEXT("<UDIM>") );
+					FPaths::NormalizeFilename( TexturePath );
+
+					FString ResolvedTexturePath = UsdToUnreal::ConvertString( TextureAssetPath.GetResolvedPath() );
+					ResolvedTexturePath.RemoveFromStart( TEXT("\\\\?\\") );
+					FPaths::NormalizeFilename( ResolvedTexturePath );
+
+					UTexture2D* Texture2D = Cast< UTexture2D >( TexturesCache.FindRef( TexturePath ) );
+
+					if ( !Texture2D && !TexturePath.IsEmpty() )
+					{
+						bool bOutCancelled = false;
+						UTextureFactory* TextureFactory = NewObject< UTextureFactory >();
+						TextureFactory->SuppressImportOverwriteDialog();
+						TextureFactory->bUseHashAsGuid = true;
+
+						if ( bIsSupportedUdimTexture )
+						{
+							FString BaseFileName = FPaths::GetBaseFilename( TexturePath );
+
+							FString BaseFileNameBeforeUdim;
+							FString BaseFileNameAfterUdim;
+							BaseFileName.Split( TEXT("<UDIM>"), &BaseFileNameBeforeUdim, &BaseFileNameAfterUdim );
+
+							FString UdimRegexPattern = FString::Printf( TEXT(R"((%s)(\d{4})(%s))"), *BaseFileNameBeforeUdim, *BaseFileNameAfterUdim );
+							TextureFactory->UdimRegexPattern = MoveTemp( UdimRegexPattern );
+						}
+
+						if ( ResolvedTexturePath.IsEmpty() )
+						{
+							TUsdStore< pxr::SdfLayerHandle > LayerHandle = FindLayerHandle( FileInput.GetAttr(), pxr::UsdTimeCode::EarliestTime() );
+							ResolvedTexturePath = ResolveAssetPath( LayerHandle, TexturePath );
+						}
+
+						Texture2D = Cast< UTexture2D >( TextureFactory->ImportObject( UTexture2D::StaticClass(), GetTransientPackage(), NAME_None, RF_Transient, ResolvedTexturePath, TEXT(""), bOutCancelled ) );
+
+						TexturesCache.Add( TexturePath, Texture2D );
+					}
+
+					if ( Texture2D )
+					{
+						UMaterialExpressionTextureSample* TextureSample = Cast< UMaterialExpressionTextureSample >( UMaterialEditingLibrary::CreateMaterialExpression( &Material, UMaterialExpressionTextureSample::StaticClass() ) );
+						TextureSample->Texture = Texture2D;
+						TextureSample->SamplerType = UMaterialExpressionTextureBase::GetSamplerTypeForTexture( Texture2D );
+
+						Result = TextureSample;
+					}
+				}
+			}
+		}
+
+		return Result;
+	}
+}
+
+bool UsdToUnreal::ConvertMaterial( const pxr::UsdShadeMaterial& UsdShadeMaterial, UMaterial& Material )
+{
+	TMap< FString, UObject* > TexturesCache;
+	return ConvertMaterial( UsdShadeMaterial, Material, TexturesCache );
+}
+
+bool UsdToUnreal::ConvertMaterial( const pxr::UsdShadeMaterial& UsdShadeMaterial, UMaterial& Material, TMap< FString, UObject* >& TexturesCache )
+{
+	pxr::UsdShadeShader SurfaceShader = UsdShadeMaterial.ComputeSurfaceSource();
+
+	if ( !SurfaceShader )
+	{
+		return false;
+	}
+
+	bool bHasMaterialInfo = false;
+
+	auto ParseFloatInput = [ &Material, &SurfaceShader, &TexturesCache ]( const ANSICHAR* InputName, float DefaultValue ) -> UMaterialExpression*
+	{
+		pxr::UsdShadeInput Input = SurfaceShader.GetInput( pxr::TfToken(InputName) );
+		UMaterialExpression* InputExpression = nullptr;
+
+		if ( Input )
+		{
+			if ( !Input.HasConnectedSource() )
+			{
+				float InputValue = DefaultValue;
+				Input.Get< float >( &InputValue );
+
+				if ( !FMath::IsNearlyEqual( InputValue, DefaultValue ) )
+				{
+					UMaterialExpressionConstant* ExpressionConstant = Cast< UMaterialExpressionConstant >( UMaterialEditingLibrary::CreateMaterialExpression( &Material, UMaterialExpressionConstant::StaticClass() ) );
+					ExpressionConstant->R = InputValue;
+
+					InputExpression = ExpressionConstant;
+				}
+			}
+			else
+			{
+				InputExpression = UsdGeomMeshConversionImpl::ParseInputTexture( Input, Material, TexturesCache );
+			}
+		}
+
+		return InputExpression;
+	};
+
+	// Base color
+	{
+		pxr::UsdShadeInput DiffuseInput = SurfaceShader.GetInput( pxr::TfToken("diffuseColor") );
+
+		if ( DiffuseInput )
+		{
+			UMaterialExpression* BaseColorExpression = nullptr;
+
+			if ( !DiffuseInput.HasConnectedSource() )
+			{
+				pxr::GfVec3f UsdDiffuseColor;
+				DiffuseInput.Get< pxr::GfVec3f >( &UsdDiffuseColor );
+
+				FLinearColor DiffuseColor = UsdToUnreal::ConvertColor( UsdDiffuseColor );
+
+				UMaterialExpressionConstant4Vector* Constant4VectorExpression = Cast< UMaterialExpressionConstant4Vector >( UMaterialEditingLibrary::CreateMaterialExpression( &Material, UMaterialExpressionConstant4Vector::StaticClass() ) );
+				Constant4VectorExpression->Constant = DiffuseColor;
+
+				BaseColorExpression = Constant4VectorExpression;
+			}
+			else
+			{
+				BaseColorExpression = UsdGeomMeshConversionImpl::ParseInputTexture( DiffuseInput, Material, TexturesCache );
+			}
+
+			if ( BaseColorExpression )
+			{
+				Material.BaseColor.Expression = BaseColorExpression;
+				bHasMaterialInfo = true;
+			}
+		}
+	}
+
+	// Metallic
+	{
+		UMaterialExpression* MetallicExpression = ParseFloatInput( "metallic", 0.f );
+
+		if ( MetallicExpression )
+		{
+			Material.Metallic.Expression = MetallicExpression;
+			bHasMaterialInfo = true;
+		}
+	}
+
+	// Roughness
+	{
+		UMaterialExpression* RoughnessExpression = ParseFloatInput( "roughness", 1.f );
+
+		if ( RoughnessExpression )
+		{
+			Material.Roughness.Expression = RoughnessExpression;
+			bHasMaterialInfo = true;
+		}
+	}
+
+	// Opacity
+	{
+		UMaterialExpression* OpacityExpression = ParseFloatInput( "opacity", 1.f );
+
+		if ( OpacityExpression )
+		{
+			Material.Opacity.Expression = OpacityExpression;
+			Material.BlendMode = BLEND_Translucent;
+			bHasMaterialInfo = true;
+		}
+	}
+
+	return bHasMaterialInfo;
+}
+
+namespace UsdShadeConversionImpl
+{
+	uint64 HashShadeInput( const pxr::UsdShadeInput& ShadeInput )
+	{
+		FScopedUsdAllocs UsdAllocs;
+
+		uint64 InputHash = 0;
+
+		if ( ShadeInput )
+		{
+			if ( !ShadeInput.HasConnectedSource() )
+			{
+				pxr::VtValue ShadeInputValue;
+				ShadeInput.Get( &ShadeInputValue );
+				InputHash = (uint64)ShadeInputValue.GetHash();
+			}
+			else
+			{
+				pxr::UsdShadeConnectableAPI Source;
+				pxr::TfToken SourceName;
+				pxr::UsdShadeAttributeType AttributeType;
+
+				if ( ShadeInput.GetConnectedSource( &Source, &SourceName, &AttributeType ) )
+				{
+					pxr::UsdShadeInput FileInput = Source.GetInput( pxr::TfToken("file") );
+
+					if ( FileInput )
+					{
+						pxr::SdfAssetPath TextureAssetPath;
+						FileInput.Get< pxr::SdfAssetPath >( &TextureAssetPath );
+
+						InputHash = (uint64)TextureAssetPath.GetHash();
+					}
+				}
+			}
+		}
+
+		return InputHash;
+	}
+}
+
+FSHAHash UsdToUnreal::HashShadeMaterial( const pxr::UsdShadeMaterial& UsdShadeMaterial )
+{
+	FScopedUsdAllocs UsdAllocs;
+
+	pxr::UsdShadeShader SurfaceShader = UsdShadeMaterial.ComputeSurfaceSource();
+
+	if ( !SurfaceShader )
+	{
+		return {};
+	}
+
+	FSHA1 HashState;
+
+	for ( const pxr::UsdShadeInput& ShadeInput : SurfaceShader.GetInputs() )
+	{
+		uint64 HashValue = UsdShadeConversionImpl::HashShadeInput( ShadeInput );
+		HashState.Update( (uint8*)&HashValue, sizeof( uint64 ) );
+	}
+
+	FSHAHash OutHash;
+
+	HashState.Final();
+	HashState.GetHash( &OutHash.Hash[0] );
+
+	return OutHash;
+}
+
+#endif // #if USE_USD_SDK
diff --git a/Engine/Plugins/Importers/USDImporter/Source/USDUtilities/Public/USDGeomMeshConversion.h b/Engine/Plugins/Importers/USDImporter/Source/USDUtilities/Public/USDGeomMeshConversion.h
--- a/Engine/Plugins/Importers/USDImporter/Source/USDUtilities/Public/USDGeomMeshConversion.h
+++ b/Engine/Plugins/Importers/USDImporter/Source/USDUtilities/Public/USDGeomMeshConversion.h
@@ -19,7 +19,6 @@ PXR_NAMESPACE_CLOSE_SCOPE
 
 struct FMeshDescription;
 struct FStaticMeshLODResources;
-class UMaterial;
 class UMaterialInstanceConstant;
 class UStaticMesh;
 
@@ -27,8 +26,6 @@ namespace UsdToUnreal
 {
 	USDUTILITIES_API bool ConvertGeomMesh( const pxr::UsdGeomMesh& UsdMesh, FMeshDescription& MeshDescription, const pxr::UsdTimeCode TimeCode = pxr::UsdTimeCode::EarliestTime() );
 	USDUTILITIES_API bool ConvertGeomMesh( const pxr::UsdGeomMesh& UsdMesh, FMeshDescription& MeshDescription, const FTransform& AdditionalTransform, const pxr::UsdTimeCode TimeCode = pxr::UsdTimeCode::EarliestTime() );
-	USDUTILITIES_API bool ConvertMaterial( const pxr::UsdShadeMaterial& UsdShadeMaterial, UMaterial& Material );
-	USDUTILITIES_API bool ConvertMaterial( const pxr::UsdShadeMaterial& UsdShadeMaterial, UMaterial& Material, TMap< FString, UObject* >& TexturesCache );
 
 	/** Reads the first display color and opacity value and assigns it as the base color and opacity of the material */
 	USDUTILITIES_API bool ConvertDisplayColor( const pxr::UsdGeomMesh& UsdMesh, UMaterialInstanceConstant& MaterialInstance, const pxr::UsdTimeCode TimeCode = pxr::UsdTimeCode::EarliestTime() );
diff --git a/Engine/Plugins/Importers/USDImporter/Source/USDUtilities/Public/USDShadeConversion.h b/Engine/Plugins/Importers/USDImporter/Source/USDUtilities/Public/USDShadeConversion.h
new file mode 100644
--- /dev/null
+++ b/Engine/Plugins/Importers/USDImporter/Source/USDUtilities/Public/USDShadeConversion.h
@@ -0,0 +1,31 @@
+// Copyright Epic Games, Inc. All Rights Reserved.
+
+#pragma once
+
+#if USE_USD_SDK
+
+#include "CoreMinimal.h"
+
+#include "USDIncludesStart.h"
+
+#include "pxr/usd/usd/timeCode.h"
+
+#include "USDIncludesEnd.h"
+
+PXR_NAMESPACE_OPEN_SCOPE
+	class UsdGeomMesh;
+	class UsdShadeMaterial;
+PXR_NAMESPACE_CLOSE_SCOPE
+
+class FSHAHash;
+class UMaterial;
+
+namespace UsdToUnreal
+{
+	USDUTILITIES_API bool ConvertMaterial( const pxr::UsdShadeMaterial& UsdShadeMaterial, UMaterial& Material );
+	USDUTILITIES_API bool ConvertMaterial( const pxr::UsdShadeMaterial& UsdShadeMaterial, UMaterial& Material, TMap< FString, UObject* >& TexturesCache );
+
+	USDUTILITIES_API FSHAHash HashShadeMaterial( const pxr::UsdShadeMaterial& UsdShadeMaterial );
+}
+
+#endif // #if USE_USD_SDK
