USD - skel conversion optims (from 12.4s to 7.7s)

#rnx

diff --git a/Engine/Plugins/Importers/USDImporter/Source/USDUtilities/Private/USDSkeletalDataConversion.cpp b/Engine/Plugins/Importers/USDImporter/Source/USDUtilities/Private/USDSkeletalDataConversion.cpp
--- a/Engine/Plugins/Importers/USDImporter/Source/USDUtilities/Private/USDSkeletalDataConversion.cpp
+++ b/Engine/Plugins/Importers/USDImporter/Source/USDUtilities/Private/USDSkeletalDataConversion.cpp
@@ -351,18 +351,48 @@ bool UsdToUnreal::ConvertSkinnedMesh(const pxr::UsdSkelSkinningQuery& SkinningQu
 		++UVChannelIndex;
 	}
 
-	// SkeletalMeshImportData uses triangle faces so quads will have to be split into triangles
-	SkeletalMeshImportData::FVertex TmpWedges[3];
+	SkelMeshImportData.Wedges.Reserve( ( NumExistingFaces + NumFaces ) * 6 );
 
 	for (uint32 PolygonIndex = NumExistingFaces, LocalIndex = 0; PolygonIndex < NumExistingFaces + NumFaces; ++PolygonIndex, ++LocalIndex)
 	{
 		const uint32 NumOriginalFaceVertices = FaceCounts[LocalIndex];
 		const uint32 NumFinalFaceVertices = 3;
 
-		// Face must be processed as triangle
-		bool bIsQuad = ( NumOriginalFaceVertices == 4 );
+		// Manage materials
+		int32 LocalMaterialIndex = 0;
+		if ( LocalIndex >= 0 && LocalIndex < (uint32) FaceMaterialIndices.Num() )
+		{
+			LocalMaterialIndex = FaceMaterialIndices[LocalIndex];
+			if (LocalMaterialIndex < 0 || LocalMaterialIndex > MaterialNames.Num())
+			{
+				LocalMaterialIndex = 0;
+			}
+		}
 
-		uint32 NumTriangles = bIsQuad ? 2 : 1;
+		int32 RealMaterialIndex = 0;
+		if ( LocalMaterialIndex < MaterialNames.Num() )
+		{
+			FString MaterialName = MaterialNames[LocalMaterialIndex];
+			if (!ExistingMaterialNames.Contains(MaterialName))
+			{
+				// If new material, add it to the list
+				SkeletalMeshImportData::FMaterial NewMaterial;
+
+				NewMaterial.MaterialImportName = MaterialName;
+				RealMaterialIndex = SkelMeshImportData.Materials.Add(NewMaterial);
+				ExistingMaterialNames.Add(MaterialName, RealMaterialIndex);
+			}
+			else
+			{
+				RealMaterialIndex = ExistingMaterialNames[MaterialName];
+			}
+		}
+
+		SkelMeshImportData.MaxMaterialIndex = FMath::Max< uint32 >( SkelMeshImportData.MaxMaterialIndex, RealMaterialIndex );
+
+		// SkeletalMeshImportData uses triangle faces so quads will have to be split into triangles
+		const bool bIsQuad = ( NumOriginalFaceVertices == 4 );
+		const uint32 NumTriangles = bIsQuad ? 2 : 1;
 
 		for ( uint32 TriangleIndex = 0; TriangleIndex < NumTriangles; ++TriangleIndex )
 		{
@@ -375,47 +405,13 @@ bool UsdToUnreal::ConvertSkinnedMesh(const pxr::UsdSkelSkinningQuery& SkinningQu
 
 			// #ueent_todo: Convert normals to TangentZ (TangentX/TangentY are tangents/binormals)
 
-			// Manage materials
-			int32 LocalMaterialIndex = 0;
-			if (LocalIndex >= 0 && LocalIndex < (uint32) FaceMaterialIndices.Num())
-			{
-				LocalMaterialIndex = FaceMaterialIndices[LocalIndex];
-				if (LocalMaterialIndex < 0 || LocalMaterialIndex > MaterialNames.Num())
-				{
-					LocalMaterialIndex = 0;
-				}
-			}
-
-			int32 RealMaterialIndex = 0;
-			if (LocalMaterialIndex < MaterialNames.Num())
-			{
-				FString MaterialName = MaterialNames[LocalMaterialIndex];
-				if (!ExistingMaterialNames.Contains(MaterialName))
-				{
-					// If new material, add it to the list
-					SkeletalMeshImportData::FMaterial NewMaterial;
-
-					NewMaterial.MaterialImportName = MaterialName;
-					RealMaterialIndex = SkelMeshImportData.Materials.Add(NewMaterial);
-					ExistingMaterialNames.Add(MaterialName, RealMaterialIndex);
-				}
-				else
-				{
-					RealMaterialIndex = ExistingMaterialNames[MaterialName];
-				}
-			}
-
 			Triangle.MatIndex = RealMaterialIndex;
-			SkelMeshImportData.MaxMaterialIndex = FMath::Max<uint32>(SkelMeshImportData.MaxMaterialIndex, Triangle.MatIndex);
-
 			Triangle.AuxMatIndex = 0;
 
 			// Manage vertex colors
 			FColor FaceColor = SkelMeshImportData.bHasVertexColors ? FaceColors[bIsConstantColor ? 0 : LocalIndex] : FColor::White;
 
 			// Fill the wedge data and complete the triangle setup with the wedge indices
-			SkelMeshImportData.Wedges.Reserve( SkelMeshImportData.Wedges.Num() + NumFinalFaceVertices );
-
 			for ( uint32 CornerIndex = 0; CornerIndex < NumFinalFaceVertices; ++CornerIndex )
 			{
 				uint32 OriginalCornerIndex = ( ( TriangleIndex * ( NumOriginalFaceVertices - 2 ) ) + CornerIndex ) % NumOriginalFaceVertices;
@@ -424,10 +420,13 @@ bool UsdToUnreal::ConvertSkinnedMesh(const pxr::UsdSkelSkinningQuery& SkinningQu
 
 				int32 FinalCornerIndex = bReverseOrder ? NumFinalFaceVertices - 1 - CornerIndex : CornerIndex;
 
-				TmpWedges[ FinalCornerIndex ].MatIndex = Triangle.MatIndex;
-				TmpWedges[ FinalCornerIndex ].VertexIndex = NumExistingPoints + OriginalVertexIndex;
-			
-				TmpWedges[ FinalCornerIndex ].Color = FaceColor;
+				const uint32 WedgeIndex = SkelMeshImportData.Wedges.AddUninitialized();
+				SkeletalMeshImportData::FVertex& SkelMeshWedge = SkelMeshImportData.Wedges[ WedgeIndex ];
+
+				SkelMeshWedge.MatIndex = Triangle.MatIndex;
+				SkelMeshWedge.VertexIndex = NumExistingPoints + OriginalVertexIndex;
+				SkelMeshWedge.Color = FaceColor;
+				SkelMeshWedge.Reserved = 0;
 
 				int32 UVLayerIndex = 0;
 				for ( const FUVSet& UVSet : UVSets )
@@ -467,7 +466,7 @@ bool UsdToUnreal::ConvertSkinnedMesh(const pxr::UsdSkelSkinningQuery& SkinningQu
 
 					// Flip V for Unreal uv's which match directx
 					FVector2D FinalUVVector( UV[0], 1.f - UV[1] );
-					TmpWedges[ FinalCornerIndex ].UVs[ UVLayerIndex ] = FinalUVVector;
+					SkelMeshWedge.UVs[ UVLayerIndex ] = FinalUVVector;
 
 					++UVLayerIndex;
 				}
@@ -476,16 +475,7 @@ bool UsdToUnreal::ConvertSkinnedMesh(const pxr::UsdSkelSkinningQuery& SkinningQu
 				Triangle.TangentY[ FinalCornerIndex ] = FVector::ZeroVector;
 				Triangle.TangentZ[ FinalCornerIndex ] = FVector::ZeroVector;
 
-				{
-					uint32 WedgeIndex = SkelMeshImportData.Wedges.AddUninitialized();
-					SkelMeshImportData.Wedges[ WedgeIndex ].VertexIndex = TmpWedges[ FinalCornerIndex ].VertexIndex;
-					SkelMeshImportData.Wedges[ WedgeIndex ].MatIndex = TmpWedges[ FinalCornerIndex ].MatIndex;
-					SkelMeshImportData.Wedges[ WedgeIndex ].Color = TmpWedges[ FinalCornerIndex ].Color;
-					SkelMeshImportData.Wedges[ WedgeIndex ].Reserved = 0;
-					FMemory::Memcpy( SkelMeshImportData.Wedges[ WedgeIndex ].UVs, TmpWedges[ FinalCornerIndex ].UVs, sizeof(FVector2D) * MAX_TEXCOORDS );
-
-					Triangle.WedgeIndex[ FinalCornerIndex ] = WedgeIndex;
-				}
+				Triangle.WedgeIndex[ FinalCornerIndex ] = WedgeIndex;
 			}
 		}
 	}
@@ -560,7 +550,7 @@ USkeletalMesh* UsdToUnreal::GetSkeletalMeshFromImportData(FSkeletalMeshImportDat
 
 	SkeletalMesh->PreEditChange(nullptr);
 
-	FSkeletalMeshModel *ImportedResource = SkeletalMesh->GetImportedModel();
+	FSkeletalMeshModel* ImportedResource = SkeletalMesh->GetImportedModel();
 	ImportedResource->LODModels.Empty();
 	ImportedResource->LODModels.Add(new FSkeletalMeshLODModel());
 
@@ -600,6 +590,8 @@ USkeletalMesh* UsdToUnreal::GetSkeletalMeshFromImportData(FSkeletalMeshImportDat
 
 	// Create the render data
 	{
+		TRACE_CPUPROFILER_EVENT_SCOPE( UsdToUnreal::GetSkeletalMeshFromImportData_CreateRenderData );
+
 		TArray<FVector> LODPoints;
 		TArray<SkeletalMeshImportData::FMeshWedge> LODWedges;
 		TArray<SkeletalMeshImportData::FMeshFace> LODFaces;
@@ -636,6 +628,8 @@ USkeletalMesh* UsdToUnreal::GetSkeletalMeshFromImportData(FSkeletalMeshImportDat
 
 	// Generate a Skeleton and associate it to the SkeletalMesh
 	{
+		TRACE_CPUPROFILER_EVENT_SCOPE( UsdToUnreal::GetSkeletalMeshFromImportData_CreateSkeleton );
+
 		USkeleton* Skeleton = NewObject<USkeleton>(GetTransientPackage(), NAME_None, ObjectFlags | EObjectFlags::RF_Public );
 
 		Skeleton->MergeAllBonesToBoneTree(SkeletalMesh);
