USD - GeomMesh conversion optims
- Don't call Primvar.GetInterpolation for every vertex. The call is costly as USD resolves all contributors to the value each time. (3m 26s -> 2m 20s)
- Clear mesh descriptions when done building the static meshes, saves 12gb (104gb -> 92gb) in a scene with big meshes.

#rnx

diff --git a/Engine/Plugins/Importers/USDImporter/Source/USDSchemas/Private/USDGeomMeshTranslator.cpp b/Engine/Plugins/Importers/USDImporter/Source/USDSchemas/Private/USDGeomMeshTranslator.cpp
--- a/Engine/Plugins/Importers/USDImporter/Source/USDSchemas/Private/USDGeomMeshTranslator.cpp
+++ b/Engine/Plugins/Importers/USDImporter/Source/USDSchemas/Private/USDGeomMeshTranslator.cpp
@@ -339,6 +339,7 @@ namespace UsdGeomMeshTranslatorImpl
 		}
 
 		StaticMesh.CalculateExtendedBounds();
+		StaticMesh.ClearMeshDescriptions(); // Clear mesh descriptions to reduce memory usage, they are kept only in bulk data form
 	}
 }
 
diff --git a/Engine/Plugins/Importers/USDImporter/Source/USDUtilities/Private/USDGeomMeshConversion.cpp b/Engine/Plugins/Importers/USDImporter/Source/USDUtilities/Private/USDGeomMeshConversion.cpp
--- a/Engine/Plugins/Importers/USDImporter/Source/USDUtilities/Private/USDGeomMeshConversion.cpp
+++ b/Engine/Plugins/Importers/USDImporter/Source/USDUtilities/Private/USDGeomMeshConversion.cpp
@@ -226,20 +226,24 @@ bool UsdToUnreal::ConvertGeomMesh( const pxr::UsdGeomMesh& UsdMesh, FMeshDescrip
 		TVertexInstanceAttributesRef< FVector4 > MeshDescriptionColors = StaticMeshAttributes.GetVertexInstanceColors();
 
 		UsdGeomPrimvar ColorPrimvar = UsdMesh.GetDisplayColorPrimvar();
+		pxr::TfToken ColorInterpolation = UsdGeomTokens->constant;
 		pxr::VtArray< pxr::GfVec3f > UsdColors;
 
 		if ( ColorPrimvar )
-		{	
+		{
 			ColorPrimvar.ComputeFlattened( &UsdColors, TimeCode );
+			ColorInterpolation = ColorPrimvar.GetInterpolation();
 		}
 
 		// Vertex opacity
 		UsdGeomPrimvar OpacityPrimvar = UsdMesh.GetDisplayOpacityPrimvar();
+		pxr::TfToken OpacityInterpolation = UsdGeomTokens->constant;
 		pxr::VtArray< float > UsdOpacities;
 
 		if ( OpacityPrimvar )
 		{
 			OpacityPrimvar.ComputeFlattened( &UsdOpacities );
+			OpacityInterpolation = OpacityPrimvar.GetInterpolation();
 		}
 
 		for ( int32 PolygonIndex = 0; PolygonIndex < FaceCounts.size(); ++PolygonIndex )
@@ -313,7 +317,7 @@ bool UsdToUnreal::ConvertGeomMesh( const pxr::UsdGeomMesh& UsdMesh, FMeshDescrip
 						return FLinearColor( FLinearColor( UsdToUnreal::ConvertColor( UsdColor ) ).ToFColor( false ) );
 					};
 
-					const int32 ValueIndex = UsdToUnrealImpl::GetPrimValueIndex( ColorPrimvar.GetInterpolation(), ControlPointIndex, CurrentVertexInstanceIndex, PolygonIndex );
+					const int32 ValueIndex = UsdToUnrealImpl::GetPrimValueIndex( ColorInterpolation, ControlPointIndex, CurrentVertexInstanceIndex, PolygonIndex );
 
 					GfVec3f UsdColor( 1.f, 1.f, 1.f );
 
@@ -327,7 +331,7 @@ bool UsdToUnreal::ConvertGeomMesh( const pxr::UsdGeomMesh& UsdMesh, FMeshDescrip
 
 				// Vertex opacity
 				{
-					const int32 ValueIndex = UsdToUnrealImpl::GetPrimValueIndex( OpacityPrimvar.GetInterpolation(), ControlPointIndex, CurrentVertexInstanceIndex, PolygonIndex );
+					const int32 ValueIndex = UsdToUnrealImpl::GetPrimValueIndex( OpacityInterpolation, ControlPointIndex, CurrentVertexInstanceIndex, PolygonIndex );
 
 					if ( !UsdOpacities.empty() && ensure( UsdOpacities.size() > ValueIndex ) )
 					{
