USD: Fix component not being recorded into the transaction buffer when being cleared

diff --git a/Engine/Plugins/Importers/USDImporter/Source/USDStage/Private/USDPrimTwin.cpp b/Engine/Plugins/Importers/USDImporter/Source/USDStage/Private/USDPrimTwin.cpp
--- a/Engine/Plugins/Importers/USDImporter/Source/USDStage/Private/USDPrimTwin.cpp
+++ b/Engine/Plugins/Importers/USDImporter/Source/USDStage/Private/USDPrimTwin.cpp
@@ -70,11 +70,13 @@ void UUsdPrimTwin::Clear()
 
 	if ( ActorToDestroy && !ActorToDestroy->IsA< AUsdStageActor >() && !ActorToDestroy->IsActorBeingDestroyed() && ActorToDestroy->GetWorld() )
 	{
+		ActorToDestroy->Modify();
 		ActorToDestroy->GetWorld()->DestroyActor( ActorToDestroy );
 		SpawnedActor = nullptr;
 	}
 	else if ( SceneComponent.IsValid() && !SceneComponent->IsBeingDestroyed() )
 	{
+		SceneComponent->Modify();
 		SceneComponent->DestroyComponent();
 		SceneComponent = nullptr;
 	}
