Fix for USD linux - Dont load python exporters if we haven't built against the USD SDK

diff --git a/Engine/Plugins/Importers/USDImporter/Source/USDExporter/Private/StaticMeshExporterUSD.cpp b/Engine/Plugins/Importers/USDImporter/Source/USDExporter/Private/StaticMeshExporterUSD.cpp
--- a/Engine/Plugins/Importers/USDImporter/Source/USDExporter/Private/StaticMeshExporterUSD.cpp
+++ b/Engine/Plugins/Importers/USDImporter/Source/USDExporter/Private/StaticMeshExporterUSD.cpp
@@ -22,6 +22,15 @@
 
 #endif // #if USE_USD_SDK
 
+bool UStaticMeshExporterUsd::IsUsdAvailable()
+{
+#if USE_USD_SDK
+	return true;
+#else
+	return false;
+#endif
+}
+
 UStaticMeshExporterUsd::UStaticMeshExporterUsd()
 {
 #if USE_USD_SDK
diff --git a/Engine/Plugins/Importers/USDImporter/Source/USDExporter/Public/StaticMeshExporterUSD.h b/Engine/Plugins/Importers/USDImporter/Source/USDExporter/Public/StaticMeshExporterUSD.h
--- a/Engine/Plugins/Importers/USDImporter/Source/USDExporter/Public/StaticMeshExporterUSD.h
+++ b/Engine/Plugins/Importers/USDImporter/Source/USDExporter/Public/StaticMeshExporterUSD.h
@@ -11,6 +11,10 @@ class UStaticMeshExporterUsd : public UExporter
 {
 	GENERATED_BODY()
 
+public:
+	UFUNCTION(BlueprintCallable, Category = "USD")
+	static bool IsUsdAvailable();
+
 public:
 	UStaticMeshExporterUsd();
 	
