USD - Fixed multiple minutes delay after we're done opening a USD stage.
- Made the assets non transactional as we don't need undo/redo support around those and they take a long time to serialize.
- Removed call to SetMobility as it was generating a lot of remove/add primitive. Now setting the mobility directly as the component is new and doesn't need to be reregistered yet.

diff --git a/Engine/Plugins/Importers/USDImporter/Source/USDSchemas/Private/USDGeomXformableTranslator.cpp b/Engine/Plugins/Importers/USDImporter/Source/USDSchemas/Private/USDGeomXformableTranslator.cpp
--- a/Engine/Plugins/Importers/USDImporter/Source/USDSchemas/Private/USDGeomXformableTranslator.cpp
+++ b/Engine/Plugins/Importers/USDImporter/Source/USDSchemas/Private/USDGeomXformableTranslator.cpp
@@ -249,17 +249,23 @@ USceneComponent* FUsdGeomXformableTranslator::CreateComponentsEx( TOptional< TSu
 
 		UpdateComponents( SceneComponent );
 
+		// Don't call SetMobility as it would trigger a reregister, queuing unnecessary rhi commands since this is a brand new component
 		if ( Context->ParentComponent && Context->ParentComponent->Mobility == EComponentMobility::Movable )
 		{
-			SceneComponent->SetMobility( EComponentMobility::Movable );
+			SceneComponent->Mobility = EComponentMobility::Movable;
 		}
 		else
 		{
-			SceneComponent->SetMobility( UsdUtils::IsAnimated( Prim.Get() ) ? EComponentMobility::Movable : EComponentMobility::Static ); /*PrimsToAnimate.Contains( UsdToUnreal::ConvertPath( Path.GetPrimPath() ) */
+			SceneComponent->Mobility = UsdUtils::IsAnimated( Prim.Get() ) ? EComponentMobility::Movable : EComponentMobility::Static;
 		}
 
 		// Attach to parent
 		SceneComponent->AttachToComponent( Context->ParentComponent, FAttachmentTransformRules::KeepRelativeTransform );
+
+		if ( !SceneComponent->IsRegistered() )
+		{
+			SceneComponent->RegisterComponent();
+		}
 	}
 
 	return SceneComponent;
diff --git a/Engine/Plugins/Importers/USDImporter/Source/USDStage/Private/USDStageActor.cpp b/Engine/Plugins/Importers/USDImporter/Source/USDStage/Private/USDStageActor.cpp
--- a/Engine/Plugins/Importers/USDImporter/Source/USDStage/Private/USDStageActor.cpp
+++ b/Engine/Plugins/Importers/USDImporter/Source/USDStage/Private/USDStageActor.cpp
@@ -924,6 +924,9 @@ void AUsdStageActor::LoadAsset( FUsdSchemaTranslationContext& TranslationContext
 {
 	TRACE_CPUPROFILER_EVENT_SCOPE( AUsdStageActor::LoadAsset );
 
+	// Mark the assets as non transactional so that they don't get serialized in the transaction buffer
+	TGuardValue< EObjectFlags > ContextFlagsGuard( TranslationContext.ObjectFlags, TranslationContext.ObjectFlags & ~RF_Transactional );
+
 	const FString PrimPath = UsdToUnreal::ConvertPath( Prim.GetPrimPath() );
 	PrimPathsToAssets.Remove( PrimPath );
 
@@ -941,6 +944,9 @@ void AUsdStageActor::LoadAssets( FUsdSchemaTranslationContext& TranslationContex
 {
 	TRACE_CPUPROFILER_EVENT_SCOPE( AUsdStageActor::LoadAssets );
 
+	// Mark the assets as non transactional so that they don't get serialized in the transaction buffer
+	TGuardValue< EObjectFlags > ContextFlagsGuard( TranslationContext.ObjectFlags, TranslationContext.ObjectFlags & ~RF_Transactional );
+
 	// Clear existing prim/asset association
 	FString StartPrimPath = UsdToUnreal::ConvertPath( StartPrim.GetPrimPath() );
 	for ( TMap< FString, UObject* >::TIterator PrimPathToAssetIt = PrimPathsToAssets.CreateIterator(); PrimPathToAssetIt; ++PrimPathToAssetIt )
