USD

Fixed mesh collapsing crash
- Don't collapse meshes together if the resulting mesh would have too many vertices.

Fixed auto instancing
- Removed mesh prim paths from polygon group slot names so that prims with the same geometry will generate the exact same mesh description.
- Skip processing materials for prims that are reusing the same static mesh. Each prim generate its own display color material resulting in the mesh being rebuilt for each prim.

#rnx

diff --git a/Engine/Plugins/Importers/USDImporter/Source/USDSchemas/Private/USDGeomMeshTranslator.cpp b/Engine/Plugins/Importers/USDImporter/Source/USDSchemas/Private/USDGeomMeshTranslator.cpp
--- a/Engine/Plugins/Importers/USDImporter/Source/USDSchemas/Private/USDGeomMeshTranslator.cpp
+++ b/Engine/Plugins/Importers/USDImporter/Source/USDSchemas/Private/USDGeomMeshTranslator.cpp
@@ -141,7 +141,7 @@ namespace UsdGeomMeshTranslatorImpl
 			UMaterialInterface* Material = nullptr;
 
 			// If there was a UE4 material stored in the the main prim for this polygon group, try fetching it
-			if ( MainPrimUEMaterialsAttribute.IsValidIndex( MaterialIndex ))
+			if ( MainPrimUEMaterialsAttribute.IsValidIndex( MaterialIndex ) )
 			{
 				Material = Cast< UMaterialInterface >( FSoftObjectPath( MainPrimUEMaterialsAttribute[ MaterialIndex ] ).TryLoad() );
 			}
@@ -173,14 +173,14 @@ namespace UsdGeomMeshTranslatorImpl
 					[&]()
 					{
 						UMaterialInstanceConstant* ExistingMaterialInstance = Cast< UMaterialInstanceConstant >( StaticMesh.GetMaterial( MaterialIndex ) );
-						const FString& SourcePrimPath = ImportedMaterialSlotName.ToString();
+						const FString& SourcePrimPath = UsdToUnreal::ConvertPath( PolygonGroupPrim.Get().GetPrimPath() );
 
 						// Assuming that we own the material instance and that we can change it as we wish, reuse it
 						if ( ExistingMaterialInstance && ExistingMaterialInstance->GetOuter() == GetTransientPackage() )
 						{
-							if (UAssetImportData* ImportData = Cast<UAssetImportData>(ExistingMaterialInstance->AssetImportData))
+							if ( ExistingMaterialInstance->AssetImportData )
 							{
-								if (ImportData->GetFirstFilename() == SourcePrimPath)
+								if ( ExistingMaterialInstance->AssetImportData->GetFirstFilename() == SourcePrimPath )
 								{
 									return ExistingMaterialInstance;
 								}
@@ -190,8 +190,8 @@ namespace UsdGeomMeshTranslatorImpl
 						// Create new material
 						UMaterialInstanceConstant* NewMaterialInstance = NewObject< UMaterialInstanceConstant >();
 
-						UAssetImportData* ImportData = NewObject< UAssetImportData >(NewMaterialInstance, TEXT("AssetImportData"));
-						ImportData->UpdateFilenameOnly(SourcePrimPath);
+						UAssetImportData* ImportData = NewObject< UAssetImportData >( NewMaterialInstance, TEXT("AssetImportData") );
+						ImportData->UpdateFilenameOnly( SourcePrimPath );
 						NewMaterialInstance->AssetImportData = ImportData;
 
 						return NewMaterialInstance;
@@ -372,15 +372,28 @@ void FBuildStaticMeshTaskChain::SetupTasks()
 				bool bIsNew = true;
 				StaticMesh = UsdGeomMeshTranslatorImpl::CreateStaticMesh( MoveTemp( MeshDescription ), *Context, bIsNew );
 
+				const FString PrimPath = UsdToUnreal::ConvertPath( Schema.Get().GetPath() );
+
 				FScopeLock Lock( &Context->CriticalSection );
 				{
-					Context->PrimPathsToAssets.Add( UsdToUnreal::ConvertPath( Schema.Get().GetPrim().GetPrimPath() ), StaticMesh );
+					Context->PrimPathsToAssets.Add( PrimPath, StaticMesh );
+				}
+
+				if ( bIsNew && StaticMesh )
+				{
+					UAssetImportData* ImportData = NewObject< UAssetImportData >( StaticMesh, TEXT("AssetImportData") );
+					ImportData->UpdateFilenameOnly( PrimPath );
+					StaticMesh->AssetImportData = ImportData;
 				}
 
 				bool bMaterialsHaveChanged = false;
-				if ( StaticMesh )
+				if ( StaticMesh && StaticMesh->AssetImportData )
 				{
-					bMaterialsHaveChanged = UsdGeomMeshTranslatorImpl::ProcessMaterials( Schema.Get().GetPrim(), *StaticMesh, Context->PrimPathsToAssets, Context->Time );
+					// Only process the materials if we have own the mesh
+					if ( StaticMesh->AssetImportData->GetFirstFilename() == PrimPath )
+					{
+						bMaterialsHaveChanged = UsdGeomMeshTranslatorImpl::ProcessMaterials( Schema.Get().GetPrim(), *StaticMesh, Context->PrimPathsToAssets, Context->Time );
+					}
 				}
 
 				const bool bContinueTaskChain = ( bIsNew || bMaterialsHaveChanged );
diff --git a/Engine/Plugins/Importers/USDImporter/Source/USDSchemas/Private/USDGeomXformableTranslator.cpp b/Engine/Plugins/Importers/USDImporter/Source/USDSchemas/Private/USDGeomXformableTranslator.cpp
--- a/Engine/Plugins/Importers/USDImporter/Source/USDSchemas/Private/USDGeomXformableTranslator.cpp
+++ b/Engine/Plugins/Importers/USDImporter/Source/USDSchemas/Private/USDGeomXformableTranslator.cpp
@@ -307,7 +307,7 @@ void FUsdGeomXformableTranslator::UpdateComponents( USceneComponent* SceneCompon
 
 bool FUsdGeomXformableTranslator::CollapsesChildren( ECollapsingType CollapsingType ) const
 {
-	bool bIsCollapsableKind = false;
+	bool bCollapsesChildren = false;
 
 	FScopedUsdAllocs UsdAllocs;
 
@@ -315,15 +315,15 @@ bool FUsdGeomXformableTranslator::CollapsesChildren( ECollapsingType CollapsingT
 
 	if ( Model )
 	{
-		bIsCollapsableKind = Model.IsKind( pxr::KindTokens->component );
+		bCollapsesChildren = Model.IsKind( pxr::KindTokens->component );
 
-		if ( !bIsCollapsableKind )
+		if ( !bCollapsesChildren )
 		{
 			// Temp support for the prop kind
-			bIsCollapsableKind = Model.IsKind( pxr::TfToken( "prop" ), pxr::UsdModelAPI::KindValidationNone );
+			bCollapsesChildren = Model.IsKind( pxr::TfToken( "prop" ), pxr::UsdModelAPI::KindValidationNone );
 		}
 
-		if ( bIsCollapsableKind )
+		if ( bCollapsesChildren )
 		{
 			IUsdSchemasModule& UsdSchemasModule = FModuleManager::Get().LoadModuleChecked< IUsdSchemasModule >( TEXT("USDSchemas") );
 
@@ -342,7 +342,42 @@ bool FUsdGeomXformableTranslator::CollapsesChildren( ECollapsingType CollapsingT
 		}
 	}
 
-	return bIsCollapsableKind;
+	if ( bCollapsesChildren )
+	{
+		TArray< TUsdStore< pxr::UsdPrim > > ChildGeomMeshes = UsdUtils::GetAllPrimsOfType( Schema.Get().GetPrim(), pxr::TfType::Find< pxr::UsdGeomMesh >() );
+
+		// We only support collapsing GeomMeshes for now and we only want to do it when there are multiple meshes as the resulting mesh is considered unique
+		if ( ChildGeomMeshes.Num() < 2 )
+		{
+			bCollapsesChildren = false;
+		}
+		else
+		{
+			const int32 MaxVertices = 500000;
+			int32 NumVertices = 0;
+
+			for ( const TUsdStore< pxr::UsdPrim >& ChildPrim : ChildGeomMeshes )
+			{
+				pxr::UsdGeomMesh ChildGeomMesh( ChildPrim.Get() );
+
+				if ( pxr::UsdAttribute Points = ChildGeomMesh.GetPointsAttr() )
+				{
+					pxr::VtArray< pxr::GfVec3f > PointsArray;
+					Points.Get( &PointsArray, pxr::UsdTimeCode( Context->Time ) );
+
+					NumVertices += PointsArray.size();
+
+					if ( NumVertices > MaxVertices )
+					{
+						bCollapsesChildren = false;
+						break;
+					}
+				}
+			}
+		}
+	}
+
+	return bCollapsesChildren;
 }
 
 #endif // #if USE_USD_SDK
diff --git a/Engine/Plugins/Importers/USDImporter/Source/USDUtilities/Private/USDGeomMeshConversion.cpp b/Engine/Plugins/Importers/USDImporter/Source/USDUtilities/Private/USDGeomMeshConversion.cpp
--- a/Engine/Plugins/Importers/USDImporter/Source/USDUtilities/Private/USDGeomMeshConversion.cpp
+++ b/Engine/Plugins/Importers/USDImporter/Source/USDUtilities/Private/USDGeomMeshConversion.cpp
@@ -355,7 +355,7 @@ bool UsdToUnreal::ConvertGeomMesh( const pxr::UsdGeomMesh& UsdMesh, FMeshDescrip
 
 			if ( !PolygonGroupMapping.Contains( RealMaterialIndex ) )
 			{
-				FName ImportedMaterialSlotName = *UsdToUnreal::ConvertPath( UsdMesh.GetPath() );
+				FName ImportedMaterialSlotName = *LexToString( RealMaterialIndex );
 				if ( MaterialNames.IsValidIndex( MaterialIndex ) )
 				{
 					FString MaterialName = MaterialNames[MaterialIndex];
