USD: Fix material slot still being parsed as prim path in some circumstances

diff --git a/Engine/Plugins/Importers/USDImporter/Source/USDSchemas/Private/USDGeomMeshTranslator.cpp b/Engine/Plugins/Importers/USDImporter/Source/USDSchemas/Private/USDGeomMeshTranslator.cpp
--- a/Engine/Plugins/Importers/USDImporter/Source/USDSchemas/Private/USDGeomMeshTranslator.cpp
+++ b/Engine/Plugins/Importers/USDImporter/Source/USDSchemas/Private/USDGeomMeshTranslator.cpp
@@ -157,7 +157,23 @@ namespace UsdGeomMeshTranslatorImpl
 				else
 				{
 					// We don't have a UE4 material stored, but we may have imported this material already, so check PrimPathsToAssets
-					TUsdStore< pxr::UsdPrim > MaterialPrim = UsdPrim.GetStage()->GetPrimAtPath( UnrealToUsd::ConvertPath( *ImportedMaterialSlotName.ToString() ).Get() );
+
+					FScopedUsdAllocs Allocs;
+
+					// We may have just a material slot index as the slot name (in case we have default displayColor materials) or a path to a valid
+					// material prim (in case we point to a shade material)
+					pxr::SdfPath MaterialPrimPath;
+					FString SlotNameString = ImportedMaterialSlotName.ToString();
+					if (!SlotNameString.IsNumeric() && pxr::SdfPath::IsValidPathString(UnrealToUsd::ConvertString(*SlotNameString).Get()))
+					{
+						MaterialPrimPath = UnrealToUsd::ConvertPath( *SlotNameString ).Get();
+					}
+					else
+					{
+						MaterialPrimPath = PolygonGroupPrim.Get().GetPrimPath();
+					}
+
+					TUsdStore< pxr::UsdPrim > MaterialPrim = UsdPrim.GetStage()->GetPrimAtPath( MaterialPrimPath );
 
 					if ( MaterialPrim.Get() )
 					{
