USD
- Fixed FTransform to GfMatrix4d conversion when the stage units aren't in centimeters.

diff --git a/Engine/Plugins/Importers/USDImporter/Source/USDUtilities/Public/USDTypesConversion.h b/Engine/Plugins/Importers/USDImporter/Source/USDUtilities/Public/USDTypesConversion.h
--- a/Engine/Plugins/Importers/USDImporter/Source/USDUtilities/Public/USDTypesConversion.h
+++ b/Engine/Plugins/Importers/USDImporter/Source/USDUtilities/Public/USDTypesConversion.h
@@ -264,8 +264,17 @@ namespace UnrealToUsd
 	inline pxr::GfMatrix4d ConvertTransform( const pxr::UsdStageRefPtr& Stage, const FTransform& Transform )
 	{
 		pxr::TfToken UpAxisValue = UsdUtils::GetUsdStageAxis( Stage );
+		const float StageMetersPerUnit = UsdUtils::GetUsdStageMetersPerUnit( Stage );
 
-		FTransform TransformInUsdSpace = UsdToUnreal::ConvertTransform( UpAxisValue == pxr::UsdGeomTokens->z, Transform );
+		FTransform TransformInUsdSpace = Transform;
+
+		const float UEMetersPerUnit = 0.01f;
+		if ( !FMath::IsNearlyEqual( StageMetersPerUnit, UEMetersPerUnit ) )
+		{
+			TransformInUsdSpace.ScaleTranslation( StageMetersPerUnit * UEMetersPerUnit );
+		}
+
+		TransformInUsdSpace = UsdToUnreal::ConvertTransform( UpAxisValue == pxr::UsdGeomTokens->z, TransformInUsdSpace );
 
 		return ConvertMatrix( TransformInUsdSpace.ToMatrixWithScale() );
 	}
