Copied fix for crash from main when the usd stage actor doesn't exist in the level and we try to open a usd file.

diff --git a/Engine/Plugins/Importers/USDImporter/Source/USDStageEditorViewModels/Private/USDStageViewModel.cpp b/Engine/Plugins/Importers/USDImporter/Source/USDStageEditorViewModels/Private/USDStageViewModel.cpp
--- a/Engine/Plugins/Importers/USDImporter/Source/USDStageEditorViewModels/Private/USDStageViewModel.cpp
+++ b/Engine/Plugins/Importers/USDImporter/Source/USDStageEditorViewModels/Private/USDStageViewModel.cpp
@@ -9,11 +9,13 @@
 #include "USDStageImportContext.h"
 #include "USDStageImporterModule.h"
 #include "USDStageImportOptions.h"
+#include "USDStageModule.h"
 #include "USDTypesConversion.h"
 
 #include "UsdWrappers/SdfLayer.h"
 #include "UsdWrappers/UsdStage.h"
 
+#include "Engine/World.h"
 #include "Misc/Paths.h"
 #include "ScopedTransaction.h"
 #include "UObject/GCObjectScopeGuard.h"
@@ -72,7 +74,12 @@ void FUsdStageViewModel::OpenStage( const TCHAR* FilePath )
 
 	UsdUtils::StartMonitoringErrors();
 
-	check( UsdStageActor.IsValid() );
+	if ( !UsdStageActor.IsValid() )
+	{
+		IUsdStageModule& UsdStageModule = FModuleManager::GetModuleChecked< IUsdStageModule >( TEXT("USDStage") );
+		UsdStageActor = &UsdStageModule.GetUsdStageActor( GWorld );
+	}
+
 	UsdStageActor->Modify();
 
 	UsdStageActor->RootLayer.FilePath = FilePath;
