CIS deprecation warning fix - transform curves are baked during compression

diff --git a/Engine/Plugins/Importers/USDImporter/Source/USDUtilities/Private/USDSkeletalDataConversion.cpp b/Engine/Plugins/Importers/USDImporter/Source/USDUtilities/Private/USDSkeletalDataConversion.cpp
--- a/Engine/Plugins/Importers/USDImporter/Source/USDUtilities/Private/USDSkeletalDataConversion.cpp
+++ b/Engine/Plugins/Importers/USDImporter/Source/USDUtilities/Private/USDSkeletalDataConversion.cpp
@@ -1620,7 +1620,6 @@ bool UsdToUnreal::ConvertSkelAnim( const pxr::UsdSkelSkeletonQuery& InUsdSkeleto
 	}
 
 	OutSkeletalAnimationAsset->CleanAnimSequenceForImport();
-	const bool bSourceDataExists = OutSkeletalAnimationAsset->HasSourceRawData();
 
 	TUsdStore<std::vector<double>> UsdJointTransformTimeSamples;
 	AnimQuery.Get().GetJointTransformTimeSamples( &( UsdJointTransformTimeSamples.Get() ) );
@@ -1833,15 +1832,9 @@ bool UsdToUnreal::ConvertSkelAnim( const pxr::UsdSkelSkeletonQuery& InUsdSkeleto
 	OutSkeletalAnimationAsset->ImportResampleFramerate = FramesPerSecond;
 	OutSkeletalAnimationAsset->SetSequenceLength(SequenceLengthSeconds);
 	OutSkeletalAnimationAsset->SetRawNumberOfFrame( NumBakedFrames );
-	OutSkeletalAnimationAsset->MarkRawDataAsModified();
-	if ( bSourceDataExists )
-	{
-		OutSkeletalAnimationAsset->BakeTrackCurvesToRawAnimation();
-	}
-	else
-	{
-		OutSkeletalAnimationAsset->PostProcessSequence();
-	}
+	OutSkeletalAnimationAsset->MarkRawDataAsModified();	
+	OutSkeletalAnimationAsset->PostProcessSequence();
+
 	OutSkeletalAnimationAsset->PostEditChange();
 	OutSkeletalAnimationAsset->MarkPackageDirty();
 
