To protect the skeletalmesh when using async operations like build or re-import we deprecate all public member variables of USkeletalMesh class and create the necessary getter and setter. Those member will stay public but will be move to private in a future release.
Adding a public member to USkeletalMesh class is now prohibed since it can break asynchronous operations that will be added next to this submit.

#rnx

diff --git a/Engine/Plugins/Animation/GameplayInsights/Source/GameplayInsightsEditor/Private/InsightsSkeletalMeshComponent.cpp b/Engine/Plugins/Animation/GameplayInsights/Source/GameplayInsightsEditor/Private/InsightsSkeletalMeshComponent.cpp
--- a/Engine/Plugins/Animation/GameplayInsights/Source/GameplayInsightsEditor/Private/InsightsSkeletalMeshComponent.cpp
+++ b/Engine/Plugins/Animation/GameplayInsights/Source/GameplayInsightsEditor/Private/InsightsSkeletalMeshComponent.cpp
@@ -31,12 +31,13 @@ void UInsightsSkeletalMeshComponent::InitAnim(bool bForceReInit)
 {
 	if(SkeletalMesh)
 	{
+		const FReferenceSkeleton& SkeletalMeshRefSkeleton = SkeletalMesh->GetRefSkeleton();
 		// set up bone visibility states as this gets skipped since we allocate the component array before registration
 		for (int32 BaseIndex = 0; BaseIndex < 2; ++BaseIndex)
 		{
-			BoneVisibilityStates[BaseIndex].SetNum(SkeletalMesh->RefSkeleton.GetNum());
+			BoneVisibilityStates[BaseIndex].SetNum(SkeletalMeshRefSkeleton.GetNum());
 
-			for (int32 BoneIndex = 0; BoneIndex < SkeletalMesh->RefSkeleton.GetNum(); BoneIndex++)
+			for (int32 BoneIndex = 0; BoneIndex < SkeletalMeshRefSkeleton.GetNum(); BoneIndex++)
 			{
 				BoneVisibilityStates[BaseIndex][BoneIndex] = BVS_Visible;
 			}
diff --git a/Engine/Plugins/Animation/RigLogic/Source/RigLogicModule/Private/DNAToSkelMeshMap.cpp b/Engine/Plugins/Animation/RigLogic/Source/RigLogicModule/Private/DNAToSkelMeshMap.cpp
--- a/Engine/Plugins/Animation/RigLogic/Source/RigLogicModule/Private/DNAToSkelMeshMap.cpp
+++ b/Engine/Plugins/Animation/RigLogic/Source/RigLogicModule/Private/DNAToSkelMeshMap.cpp
@@ -187,7 +187,7 @@ void FDNAToSkelMeshMap::MapJoints(IDNAReader* DNAReader)
 
 	double StartTime = FPlatformTime::Seconds();
 
-	FReferenceSkeleton* RefSkeleton = &TargetSkelMesh->RefSkeleton;
+	FReferenceSkeleton* RefSkeleton = &TargetSkelMesh->GetRefSkeleton();
 	uint32 JointCount = DNAReader->GetJointCount();
 
 	// Map Joints to Bones.
@@ -222,7 +222,7 @@ void FDNAToSkelMeshMap::MapMorphTargets(IDNAReader* DNAReader)
 	// Assemble the reverse map, for each MorphTarget index save corresponding FDNABlendShapeTarget Index
 	// Keeping MeshIndex and TargetIndex for each MorphTarget.
 	MeshBlendShapeTargets.Empty();
-	MeshBlendShapeTargets.Init(FDNABlendShapeTarget(), TargetSkelMesh->MorphTargets.Num());
+	MeshBlendShapeTargets.Init(FDNABlendShapeTarget(), TargetSkelMesh->GetMorphTargets().Num());
 
 	for (int16 MeshIndex = 0; MeshIndex < MeshCount; MeshIndex++)
 	{
diff --git a/Engine/Plugins/Animation/RigLogic/Source/RigLogicModule/Private/RigUnit_RigLogic_Test.cpp b/Engine/Plugins/Animation/RigLogic/Source/RigLogicModule/Private/RigUnit_RigLogic_Test.cpp
--- a/Engine/Plugins/Animation/RigLogic/Source/RigLogicModule/Private/RigUnit_RigLogic_Test.cpp
+++ b/Engine/Plugins/Animation/RigLogic/Source/RigLogicModule/Private/RigUnit_RigLogic_Test.cpp
@@ -721,7 +721,7 @@ IMPLEMENT_RIGUNIT_AUTOMATION_TEST(FRigUnit_RigLogic)
 	//Create skeleton, skeletal mesh and skeletal mesh component
 	USkeleton* Skeleton = NewObject<USkeleton>(GetTransientPackage());
 	USkeletalMesh* SkeletalMesh = NewObject<USkeletalMesh>(GetTransientPackage());
-	SkeletalMesh->Skeleton = Skeleton;
+	SkeletalMesh->SetSkeleton(Skeleton);
 	USkeletalMeshComponent* SkelMeshComponent = NewObject<USkeletalMeshComponent>();
 	SkelMeshComponent->SetSkeletalMesh(SkeletalMesh);
 
diff --git a/Engine/Plugins/Animation/RigLogic/Source/RigLogicModule/Private/SkelMeshDNAUtils.cpp b/Engine/Plugins/Animation/RigLogic/Source/RigLogicModule/Private/SkelMeshDNAUtils.cpp
--- a/Engine/Plugins/Animation/RigLogic/Source/RigLogicModule/Private/SkelMeshDNAUtils.cpp
+++ b/Engine/Plugins/Animation/RigLogic/Source/RigLogicModule/Private/SkelMeshDNAUtils.cpp
@@ -51,10 +51,10 @@ FDNAToSkelMeshMap* USkelMeshDNAUtils::CreateMapForUpdatingNeutralMesh(IDNAReader
 void USkelMeshDNAUtils::UpdateJoints(USkeletalMesh* SkelMesh, IDNAReader* DNAReader, FDNAToSkelMeshMap* DNAToSkelMeshMap)
 {
 	{	// Scoping of RefSkelModifier
-		FReferenceSkeletonModifier RefSkelModifier(SkelMesh->RefSkeleton, SkelMesh->Skeleton);
+		FReferenceSkeletonModifier RefSkelModifier(SkelMesh->GetRefSkeleton(), SkelMesh->GetSkeleton());
 
 		// copy here
-		TArray<FTransform> RawBonePose = SkelMesh->RefSkeleton.GetRawRefBonePose();
+		TArray<FTransform> RawBonePose = SkelMesh->GetRefSkeleton().GetRawRefBonePose();
 
 		// When we are mounting the head to different bodies than female average, we need to use
 		// component space, as the joint to which the head root is snapped to will be on a different
@@ -62,9 +62,9 @@ void USkelMeshDNAUtils::UpdateJoints(USkeletalMesh* SkelMesh, IDNAReader* DNARea
 
 		// calculate component space ahead of current transform
 		TArray<FTransform> ComponentTransforms;
-		FAnimationRuntime::FillUpComponentSpaceTransforms(SkelMesh->RefSkeleton, RawBonePose, ComponentTransforms);
+		FAnimationRuntime::FillUpComponentSpaceTransforms(SkelMesh->GetRefSkeleton(), RawBonePose, ComponentTransforms);
 
-		const TArray<FMeshBoneInfo>& RawBoneInfo = SkelMesh->RefSkeleton.GetRawRefBoneInfo();
+		const TArray<FMeshBoneInfo>& RawBoneInfo = SkelMesh->GetRefSkeleton().GetRawRefBoneInfo();
 
 		// Skipping root joint (index 0) to avoid blinking of the mesh due to bounding box issue
 		for (uint16 JointIndex = 0; JointIndex < DNAReader->GetJointCount(); JointIndex++)
@@ -127,7 +127,7 @@ void USkelMeshDNAUtils::UpdateJoints(USkeletalMesh* SkelMesh, IDNAReader* DNARea
 		}
 	}
 
-	SkelMesh->RefBasesInvMatrix.Reset();
+	SkelMesh->GetRefBasesInvMatrix().Reset();
 	SkelMesh->CalculateInvRefMatrices(); // Needs to be called after RefSkelModifier is destroyed
 }
 
@@ -188,9 +188,9 @@ void USkelMeshDNAUtils::UpdateMorphTargets(USkeletalMesh* SkelMesh, IDNAReader*
 		return;
 	}
 	TArray<TArray<TArray<int32>>>& BlendShapeVertexDeltas = DNAToSkelMeshMap->GetBlendShapeVertexDeltaIndices();
-	ParallelFor(SkelMesh->MorphTargets.Num(), [&](int32 MorphIndex)
+	ParallelFor(SkelMesh->GetMorphTargets().Num(), [&](int32 MorphIndex)
 	{
-		UMorphTarget* MorphTarget = SkelMesh->MorphTargets[MorphIndex];
+		UMorphTarget* MorphTarget = SkelMesh->GetMorphTargets()[MorphIndex];
 		const FDNABlendShapeTarget& MeshTarget = MeshBlendShapeTargets[MorphIndex];
 		// First get all DNA deltas for current Morph Target.
 		TArrayView<const uint32> BlendShapeVertexIndices = DNAReader->GetBlendShapeTargetVertexIndices(MeshTarget.MeshIndex, MeshTarget.TargetIndex);
diff --git a/Engine/Plugins/Developer/AnimationSharing/Source/AnimationSharing/Private/AnimationSharingManager.cpp b/Engine/Plugins/Developer/AnimationSharing/Source/AnimationSharing/Private/AnimationSharingManager.cpp
--- a/Engine/Plugins/Developer/AnimationSharing/Source/AnimationSharing/Private/AnimationSharingManager.cpp
+++ b/Engine/Plugins/Developer/AnimationSharing/Source/AnimationSharing/Private/AnimationSharingManager.cpp
@@ -363,11 +363,11 @@ void UAnimationSharingManager::RegisterActor(AActor* InActor, FUpdateActorHandle
 
 		const USkeleton* UsedSkeleton = [&OwnedComponents]()
 		{
-			USkeleton* CurrentSkeleton = nullptr;
+			const USkeleton* CurrentSkeleton = nullptr;
 			for (USkeletalMeshComponent* SkeletalMeshComponent : OwnedComponents)
 			{
 				const USkeletalMesh* Mesh = SkeletalMeshComponent->SkeletalMesh;
-				USkeleton* Skeleton = Mesh->Skeleton;
+				const USkeleton* Skeleton = Mesh->GetSkeleton();
 
 				if (CurrentSkeleton == nullptr)
 				{
diff --git a/Engine/Plugins/Experimental/AlembicImporter/Source/AlembicImporter/Private/AlembicImportFactory.cpp b/Engine/Plugins/Experimental/AlembicImporter/Source/AlembicImporter/Private/AlembicImportFactory.cpp
--- a/Engine/Plugins/Experimental/AlembicImporter/Source/AlembicImporter/Private/AlembicImportFactory.cpp
+++ b/Engine/Plugins/Experimental/AlembicImporter/Source/AlembicImporter/Private/AlembicImportFactory.cpp
@@ -273,12 +273,12 @@ UObject* UAlembicImportFactory::ImportSkeletalMesh(FAbcImporter& Importer, UObje
 		if (SkeletalMesh)
 		{
 		// Setup asset import data
-		if (!SkeletalMesh->AssetImportData || !SkeletalMesh->AssetImportData->IsA<UAbcAssetImportData>())
+		if (!SkeletalMesh->GetAssetImportData() || !SkeletalMesh->GetAssetImportData()->IsA<UAbcAssetImportData>())
 		{
-			SkeletalMesh->AssetImportData = NewObject<UAbcAssetImportData>(SkeletalMesh);
+			SkeletalMesh->SetAssetImportData(NewObject<UAbcAssetImportData>(SkeletalMesh));
 		}
-		SkeletalMesh->AssetImportData->Update(UFactory::CurrentFilename);
-		UAbcAssetImportData* AssetImportData = Cast<UAbcAssetImportData>(SkeletalMesh->AssetImportData);
+		SkeletalMesh->GetAssetImportData()->Update(UFactory::CurrentFilename);
+		UAbcAssetImportData* AssetImportData = Cast<UAbcAssetImportData>(SkeletalMesh->GetAssetImportData());
 		if (AssetImportData)
 		{
 			Importer.UpdateAssetImportData(AssetImportData);
@@ -333,7 +333,7 @@ bool UAlembicImportFactory::CanReimport(UObject* Obj, TArray<FString>& OutFilena
 	else if (Obj->GetClass() == USkeletalMesh::StaticClass())
 	{
 		USkeletalMesh* Cache = Cast<USkeletalMesh>(Obj);
-		ImportData = Cache->AssetImportData;
+		ImportData = Cache->GetAssetImportData();
 	}
 	else if (Obj->GetClass() == UAnimSequence::StaticClass())
 	{
@@ -361,9 +361,9 @@ void UAlembicImportFactory::SetReimportPaths(UObject* Obj, const TArray<FString>
 	}
 
 	USkeletalMesh* SkeletalMesh = Cast<USkeletalMesh>(Obj);
-	if (SkeletalMesh && SkeletalMesh->AssetImportData && ensure(NewReimportPaths.Num() == 1))
+	if (SkeletalMesh && SkeletalMesh->GetAssetImportData() && ensure(NewReimportPaths.Num() == 1))
 	{
-		SkeletalMesh->AssetImportData->UpdateFilenameOnly(NewReimportPaths[0]);
+		SkeletalMesh->GetAssetImportData()->UpdateFilenameOnly(NewReimportPaths[0]);
 	}
 
 	UAnimSequence* Sequence = Cast<UAnimSequence>(Obj);
@@ -436,7 +436,7 @@ EReimportResult::Type UAlembicImportFactory::Reimport(UObject* Obj)
 			return EReimportResult::Failed;
 		}
 
-		CurrentFilename = SkeletalMesh->AssetImportData->GetFirstFilename();
+		CurrentFilename = SkeletalMesh->GetAssetImportData()->GetFirstFilename();
 
 		EReimportResult::Type Result = ReimportSkeletalMesh(SkeletalMesh);
 
@@ -468,7 +468,7 @@ EReimportResult::Type UAlembicImportFactory::Reimport(UObject* Obj)
 		for (TObjectIterator<USkeletalMesh> It; It; ++It)
 		{
 			// This works because the skeleton is unique for every imported alembic cache
-			if (It->Skeleton == AnimSequence->GetSkeleton())
+			if (It->GetSkeleton() == AnimSequence->GetSkeleton())
 			{
 				SkeletalMesh = *It;
 				break;
@@ -626,9 +626,9 @@ EReimportResult::Type UAlembicImportFactory::ReimportSkeletalMesh(USkeletalMesh*
 		return EReimportResult::Failed;
 	}
 
-	if (SkeletalMesh->AssetImportData && SkeletalMesh->AssetImportData->IsA<UAbcAssetImportData>())
+	if (SkeletalMesh->GetAssetImportData() && SkeletalMesh->GetAssetImportData()->IsA<UAbcAssetImportData>())
 	{
-		UAbcAssetImportData* ImportData = Cast<UAbcAssetImportData>(SkeletalMesh->AssetImportData);
+		UAbcAssetImportData* ImportData = Cast<UAbcAssetImportData>(SkeletalMesh->GetAssetImportData());
 		PopulateOptionsWithImportData(ImportData);
 		Importer.RetrieveAssetImportData(ImportData);
 	}
@@ -677,13 +677,13 @@ EReimportResult::Type UAlembicImportFactory::ReimportSkeletalMesh(USkeletalMesh*
 	else
 	{
 		// Update file path/timestamp (Path could change if user has to browse for it manually)
-		if (!NewSkeletalMesh->AssetImportData || !NewSkeletalMesh->AssetImportData->IsA<UAbcAssetImportData>())
+		if (!NewSkeletalMesh->GetAssetImportData() || !NewSkeletalMesh->GetAssetImportData()->IsA<UAbcAssetImportData>())
 		{
-			NewSkeletalMesh->AssetImportData = NewObject<UAbcAssetImportData>(NewSkeletalMesh);
+			NewSkeletalMesh->SetAssetImportData(NewObject<UAbcAssetImportData>(NewSkeletalMesh));
 		}
 
-		NewSkeletalMesh->AssetImportData->Update(CurrentFilename);
-		UAbcAssetImportData* AssetImportData = Cast<UAbcAssetImportData>(NewSkeletalMesh->AssetImportData);
+		NewSkeletalMesh->GetAssetImportData()->Update(CurrentFilename);
+		UAbcAssetImportData* AssetImportData = Cast<UAbcAssetImportData>(NewSkeletalMesh->GetAssetImportData());
 		if (AssetImportData)
 		{
 			Importer.UpdateAssetImportData(AssetImportData);
diff --git a/Engine/Plugins/Experimental/AlembicImporter/Source/AlembicLibrary/Private/AbcImporter.cpp b/Engine/Plugins/Experimental/AlembicImporter/Source/AlembicLibrary/Private/AbcImporter.cpp
--- a/Engine/Plugins/Experimental/AlembicImporter/Source/AlembicLibrary/Private/AbcImporter.cpp
+++ b/Engine/Plugins/Experimental/AlembicImporter/Source/AlembicLibrary/Private/AbcImporter.cpp
@@ -597,7 +597,7 @@ TArray<UObject*> FAbcImporter::ImportAsSkeletalMesh(UObject* InParent, EObjectFl
 		const FMeshBoneInfo BoneInfo(FName(TEXT("RootBone"), FNAME_Add), TEXT("RootBone_Export"), INDEX_NONE);
 		const FTransform BoneTransform;
 		{
-			FReferenceSkeletonModifier RefSkelModifier(SkeletalMesh->RefSkeleton, SkeletalMesh->Skeleton);
+			FReferenceSkeletonModifier RefSkelModifier(SkeletalMesh->GetRefSkeleton(), SkeletalMesh->GetSkeleton());
 			RefSkelModifier.Add(BoneInfo, BoneTransform);
 		}
 
@@ -610,8 +610,8 @@ TArray<UObject*> FAbcImporter::ImportAsSkeletalMesh(UObject* InParent, EObjectFl
 		
 		// Forced to 1
 		LODModel.NumTexCoords = MergedMeshSample->NumUVSets;
-		SkeletalMesh->bHasVertexColors = true;
-		SkeletalMesh->VertexColorGuid = FGuid::NewGuid();
+		SkeletalMesh->SetHasVertexColors(true);
+		SkeletalMesh->SetVertexColorGuid(FGuid::NewGuid());
 
 		/* Bounding box according to animation */
 		SkeletalMesh->SetImportedBounds(AbcFile->GetArchiveBounds().GetBox());
@@ -621,7 +621,7 @@ TArray<UObject*> FAbcImporter::ImportAsSkeletalMesh(UObject* InParent, EObjectFl
 		TArray<int32> UsedVertexIndicesForMorphs;
 		MergedMeshSample->TangentX.Empty();
 		MergedMeshSample->TangentY.Empty();
-		bBuildSuccess = BuildSkeletalMesh(LODModel, SkeletalMesh->RefSkeleton, MergedMeshSample, MorphTargetVertexRemapping, UsedVertexIndicesForMorphs);
+		bBuildSuccess = BuildSkeletalMesh(LODModel, SkeletalMesh->GetRefSkeleton(), MergedMeshSample, MorphTargetVertexRemapping, UsedVertexIndicesForMorphs);
 
 		if (!bBuildSuccess)
 		{
@@ -636,9 +636,9 @@ TArray<UObject*> FAbcImporter::ImportAsSkeletalMesh(UObject* InParent, EObjectFl
 		// Merge bones to the selected skeleton
 		check(Skeleton->MergeAllBonesToBoneTree(SkeletalMesh));
 		Skeleton->MarkPackageDirty();
-		if (SkeletalMesh->Skeleton != Skeleton)
+		if (SkeletalMesh->GetSkeleton() != Skeleton)
 		{
-			SkeletalMesh->Skeleton = Skeleton;
+			SkeletalMesh->SetSkeleton(Skeleton);
 			SkeletalMesh->MarkPackageDirty();
 		}
 
@@ -738,7 +738,7 @@ TArray<UObject*> FAbcImporter::ImportAsSkeletalMesh(UObject* InParent, EObjectFl
 						RootBoneTrack.PosKeys.Add(DeltaOffsets[FrameIndex]);
 					}
 
-					const FReferenceSkeleton& RefSkeleton = SkeletalMesh->RefSkeleton;
+					const FReferenceSkeleton& RefSkeleton = SkeletalMesh->GetRefSkeleton();
 					const TArray<FMeshBoneInfo>& BonesInfo = RefSkeleton.GetRawRefBoneInfo();
 					Sequence->AddNewRawTrack(BonesInfo[0].Name, &RootBoneTrack);
 				}
@@ -753,7 +753,7 @@ TArray<UObject*> FAbcImporter::ImportAsSkeletalMesh(UObject* InParent, EObjectFl
 				{
 					const FString& MaterialName = CompressedData.MaterialNames[MaterialIndex];
 					UMaterialInterface* Material = AbcImporterUtilities::RetrieveMaterial(*AbcFile, MaterialName, InParent, Flags);
-					SkeletalMesh->Materials.Add(FSkeletalMaterial(Material, true));
+					SkeletalMesh->GetMaterials().Add(FSkeletalMaterial(Material, true));
 					if (Material != UMaterial::GetDefaultMaterial(MD_Surface))
 					{
 						Material->PostEditChange();
diff --git a/Engine/Plugins/Experimental/ChaosCloth/Source/Chaos/Private/ChaosCloth/ChaosClothingSimulationCollider.cpp b/Engine/Plugins/Experimental/ChaosCloth/Source/Chaos/Private/ChaosCloth/ChaosClothingSimulationCollider.cpp
--- a/Engine/Plugins/Experimental/ChaosCloth/Source/Chaos/Private/ChaosCloth/ChaosClothingSimulationCollider.cpp
+++ b/Engine/Plugins/Experimental/ChaosCloth/Source/Chaos/Private/ChaosCloth/ChaosClothingSimulationCollider.cpp
@@ -390,7 +390,7 @@ void FClothingSimulationCollider::ExtractPhysicsAssetCollision(FClothCollisionDa
 			if (!BodySetup)
 				continue;
 
-			const int32 MeshBoneIndex = TargetMesh->RefSkeleton.FindBoneIndex(BodySetup->BoneName);
+			const int32 MeshBoneIndex = TargetMesh->GetRefSkeleton().FindBoneIndex(BodySetup->BoneName);
 			const int32 MappedBoneIndex = UsedBoneIndices.Add(MeshBoneIndex);
 			
 			// Add capsules
diff --git a/Engine/Plugins/Experimental/ControlRig/Source/ControlRig/Private/ControlRigComponent.cpp b/Engine/Plugins/Experimental/ControlRig/Source/ControlRig/Private/ControlRigComponent.cpp
--- a/Engine/Plugins/Experimental/ControlRig/Source/ControlRig/Private/ControlRigComponent.cpp
+++ b/Engine/Plugins/Experimental/ControlRig/Source/ControlRig/Private/ControlRigComponent.cpp
@@ -390,7 +390,7 @@ void UControlRigComponent::AddMappedSkeletalMesh(USkeletalMeshComponent* Skeleta
 		{
 			if (USkeletalMesh* SkeletalMesh = SkeletalMeshComponent->SkeletalMesh)
 			{
-				if (USkeleton* Skeleton = SkeletalMesh->Skeleton)
+				if (USkeleton* Skeleton = SkeletalMesh->GetSkeleton())
 				{
 					for (const FRigBone& RigBone : CR->GetBoneHierarchy())
 					{
@@ -418,7 +418,7 @@ void UControlRigComponent::AddMappedSkeletalMesh(USkeletalMeshComponent* Skeleta
 		{
 			if (USkeletalMesh* SkeletalMesh = SkeletalMeshComponent->SkeletalMesh)
 			{
-				if (USkeleton* Skeleton = SkeletalMesh->Skeleton)
+				if (USkeleton* Skeleton = SkeletalMesh->GetSkeleton())
 				{
 					for (const FRigCurve& RigCurve : CR->GetCurveContainer())
 					{
@@ -1143,7 +1143,7 @@ void UControlRigComponent::ValidateMappingData()
 				{
 					if (USkeletalMesh* SkeletalMesh = SkeletalMeshComponent->SkeletalMesh)
 					{
-						if (USkeleton* Skeleton = SkeletalMesh->Skeleton)
+						if (USkeleton* Skeleton = SkeletalMesh->GetSkeleton())
 						{
 							if (MappedElement.ElementType == ERigElementType::Bone)
 							{
diff --git a/Engine/Plugins/Experimental/ControlRig/Source/ControlRig/Private/Rigs/AdditiveControlRig.cpp b/Engine/Plugins/Experimental/ControlRig/Source/ControlRig/Private/Rigs/AdditiveControlRig.cpp
--- a/Engine/Plugins/Experimental/ControlRig/Source/ControlRig/Private/Rigs/AdditiveControlRig.cpp
+++ b/Engine/Plugins/Experimental/ControlRig/Source/ControlRig/Private/Rigs/AdditiveControlRig.cpp
@@ -143,8 +143,8 @@ void UAdditiveControlRig::CreateRigElements(const USkeletalMesh* InReferenceMesh
 {
 	if (InReferenceMesh)
 	{
-		USkeleton* Skeleton = InReferenceMesh->Skeleton;
-		CreateRigElements(InReferenceMesh->RefSkeleton, (Skeleton) ? Skeleton->GetSmartNameContainer(USkeleton::AnimCurveMappingName) : nullptr);
+		const USkeleton* Skeleton = InReferenceMesh->GetSkeleton();
+		CreateRigElements(InReferenceMesh->GetRefSkeleton(), (Skeleton) ? Skeleton->GetSmartNameContainer(USkeleton::AnimCurveMappingName) : nullptr);
 	}
 }
 
diff --git a/Engine/Plugins/Experimental/ControlRig/Source/ControlRig/Private/Rigs/FKControlRig.cpp b/Engine/Plugins/Experimental/ControlRig/Source/ControlRig/Private/Rigs/FKControlRig.cpp
--- a/Engine/Plugins/Experimental/ControlRig/Source/ControlRig/Private/Rigs/FKControlRig.cpp
+++ b/Engine/Plugins/Experimental/ControlRig/Source/ControlRig/Private/Rigs/FKControlRig.cpp
@@ -264,8 +264,8 @@ void UFKControlRig::CreateRigElements(const USkeletalMesh* InReferenceMesh)
 {
 	if (InReferenceMesh)
 	{
-		USkeleton* Skeleton = InReferenceMesh->Skeleton;
-		CreateRigElements(InReferenceMesh->RefSkeleton, (Skeleton) ? Skeleton->GetSmartNameContainer(USkeleton::AnimCurveMappingName) : nullptr);
+		const USkeleton* Skeleton = InReferenceMesh->GetSkeleton();
+		CreateRigElements(InReferenceMesh->GetRefSkeleton(), (Skeleton) ? Skeleton->GetSmartNameContainer(USkeleton::AnimCurveMappingName) : nullptr);
 	}
 }
 
diff --git a/Engine/Plugins/Experimental/ControlRig/Source/ControlRigEditor/Private/ControlRigBlueprintActions.cpp b/Engine/Plugins/Experimental/ControlRig/Source/ControlRigEditor/Private/ControlRigBlueprintActions.cpp
--- a/Engine/Plugins/Experimental/ControlRig/Source/ControlRigEditor/Private/ControlRigBlueprintActions.cpp
+++ b/Engine/Plugins/Experimental/ControlRig/Source/ControlRigEditor/Private/ControlRigBlueprintActions.cpp
@@ -150,7 +150,7 @@ UControlRigBlueprint* FControlRigBlueprintActions::CreateControlRigFromSkeletalM
 
 	if(SkeletalMesh)
 	{
-		Skeleton = SkeletalMesh->Skeleton;
+		Skeleton = SkeletalMesh->GetSkeleton();
 	}
 	else if (Skeleton == nullptr)
 	{
@@ -308,7 +308,7 @@ void FControlRigBlueprintActions::OnSpawnedSkeletalMeshActorChanged(UObject* InO
 			if (Track)
 			{
 				USkeletalMesh* SkeletalMesh = MeshActor->GetSkeletalMeshComponent()->SkeletalMesh;
-				USkeleton* Skeleton = SkeletalMesh->Skeleton;
+				USkeleton* Skeleton = SkeletalMesh->GetSkeleton();
 				
 				FString ObjectName = (ControlRigClass->GetName());
 				ObjectName.RemoveFromEnd(TEXT("_C"));
diff --git a/Engine/Plugins/Experimental/ControlRig/Source/ControlRigEditor/Private/ControlRigEditorModule.cpp b/Engine/Plugins/Experimental/ControlRig/Source/ControlRigEditor/Private/ControlRigEditorModule.cpp
--- a/Engine/Plugins/Experimental/ControlRig/Source/ControlRigEditor/Private/ControlRigEditorModule.cpp
+++ b/Engine/Plugins/Experimental/ControlRig/Source/ControlRigEditor/Private/ControlRigEditorModule.cpp
@@ -642,7 +642,7 @@ void FControlRigEditorModule::BakeToControlRig(UClass* ControlRigClass, UAnimSeq
 			if (Track)
 			{
 				USkeletalMesh* SkeletalMesh = MeshActor->GetSkeletalMeshComponent()->SkeletalMesh;
-				USkeleton* Skeleton = SkeletalMesh->Skeleton;
+				USkeleton* Skeleton = SkeletalMesh->GetSkeleton();
 
 				FString ObjectName = (ControlRigClass->GetName());
 				ObjectName.RemoveFromEnd(TEXT("_C"));
diff --git a/Engine/Plugins/Experimental/ControlRig/Source/ControlRigEditor/Private/Editor/SRigCurveContainer.cpp b/Engine/Plugins/Experimental/ControlRig/Source/ControlRigEditor/Private/Editor/SRigCurveContainer.cpp
--- a/Engine/Plugins/Experimental/ControlRig/Source/ControlRigEditor/Private/Editor/SRigCurveContainer.cpp
+++ b/Engine/Plugins/Experimental/ControlRig/Source/ControlRigEditor/Private/Editor/SRigCurveContainer.cpp
@@ -650,7 +650,7 @@ void SRigCurveContainer::ImportCurve(const FAssetData& InAssetData)
 		const USkeleton* Skeleton = nullptr;
 		if (USkeletalMesh* Mesh = Cast<USkeletalMesh>(InAssetData.GetAsset()))
 		{
-			Skeleton = Mesh->Skeleton;
+			Skeleton = Mesh->GetSkeleton();
 			ControlRigBlueprint->SourceCurveImport = Skeleton;
 		}
 		else 
diff --git a/Engine/Plugins/Experimental/ControlRig/Source/ControlRigEditor/Private/Editor/SRigHierarchy.cpp b/Engine/Plugins/Experimental/ControlRig/Source/ControlRigEditor/Private/Editor/SRigHierarchy.cpp
--- a/Engine/Plugins/Experimental/ControlRig/Source/ControlRigEditor/Private/Editor/SRigHierarchy.cpp
+++ b/Engine/Plugins/Experimental/ControlRig/Source/ControlRigEditor/Private/Editor/SRigHierarchy.cpp
@@ -1356,7 +1356,7 @@ void SRigHierarchy::RefreshHierarchy(const FAssetData& InAssetData)
 		FScopedTransaction Transaction(LOCTEXT("HierarchyRefresh", "Refresh Transform"));
 		ControlRigBlueprint->Modify();
 
-		const FReferenceSkeleton& RefSkeleton = Mesh->RefSkeleton;
+		const FReferenceSkeleton& RefSkeleton = Mesh->GetRefSkeleton();
 		Hierarchy->BoneHierarchy.ImportSkeleton(RefSkeleton, NAME_None, true, true, true, false /* notify */);
 	}
 
@@ -1408,12 +1408,12 @@ void SRigHierarchy::ImportHierarchy(const FAssetData& InAssetData)
 		FScopedTransaction Transaction(LOCTEXT("HierarchyImport", "Import Hierarchy"));
 		ControlRigBlueprint->Modify();
 
-		const FReferenceSkeleton& RefSkeleton = Mesh->RefSkeleton;
+		const FReferenceSkeleton& RefSkeleton = Mesh->GetRefSkeleton();
 		Hierarchy->BoneHierarchy.ImportSkeleton(RefSkeleton, NAME_None, false, false, true, false /* notify */);
-		Hierarchy->CurveContainer.ImportCurvesFromSkeleton(Mesh->Skeleton, NAME_None, true, false, false /* notify */);
+		Hierarchy->CurveContainer.ImportCurvesFromSkeleton(Mesh->GetSkeleton(), NAME_None, true, false, false /* notify */);
 
-		ControlRigBlueprint->SourceHierarchyImport = Mesh->Skeleton;
-		ControlRigBlueprint->SourceCurveImport = Mesh->Skeleton;
+		ControlRigBlueprint->SourceHierarchyImport = Mesh->GetSkeleton();
+		ControlRigBlueprint->SourceCurveImport = Mesh->GetSkeleton();
 	}
 
 	FSlateApplication::Get().DismissAllMenus();
diff --git a/Engine/Plugins/Experimental/ControlRig/Source/ControlRigEditor/Private/Sequencer/ControlRigParameterTrackEditor.cpp b/Engine/Plugins/Experimental/ControlRig/Source/ControlRigEditor/Private/Sequencer/ControlRigParameterTrackEditor.cpp
--- a/Engine/Plugins/Experimental/ControlRig/Source/ControlRigEditor/Private/Sequencer/ControlRigParameterTrackEditor.cpp
+++ b/Engine/Plugins/Experimental/ControlRig/Source/ControlRigEditor/Private/Sequencer/ControlRigParameterTrackEditor.cpp
@@ -97,10 +97,10 @@ static USkeletalMeshComponent* AcquireSkeletalMeshFromObject(UObject* BoundObjec
 static USkeleton* GetSkeletonFromComponent(UActorComponent* InComponent)
 {
 	USkeletalMeshComponent* SkeletalMeshComp = Cast<USkeletalMeshComponent>(InComponent);
-	if (SkeletalMeshComp && SkeletalMeshComp->SkeletalMesh && SkeletalMeshComp->SkeletalMesh->Skeleton)
+	if (SkeletalMeshComp && SkeletalMeshComp->SkeletalMesh && SkeletalMeshComp->SkeletalMesh->GetSkeleton())
 	{
 		// @todo Multiple actors, multiple components
-		return SkeletalMeshComp->SkeletalMesh->Skeleton;
+		return SkeletalMeshComp->SkeletalMesh->GetSkeleton();
 	}
 
 	return nullptr;
@@ -1057,9 +1057,9 @@ void FControlRigParameterTrackEditor::AddTrackForComponent(USceneComponent* InCo
 {
 	if (USkeletalMeshComponent* SkelMeshComp = Cast<USkeletalMeshComponent>(InComponent))
 	{
-		if (SkelMeshComp->SkeletalMesh && !SkelMeshComp->SkeletalMesh->DefaultAnimatingRig.IsNull())
+		if (SkelMeshComp->SkeletalMesh && !SkelMeshComp->SkeletalMesh->GetDefaultAnimatingRig().IsNull())
 		{
-			UObject* Object = SkelMeshComp->SkeletalMesh->DefaultAnimatingRig.LoadSynchronous();
+			UObject* Object = SkelMeshComp->SkeletalMesh->GetDefaultAnimatingRig().LoadSynchronous();
 			if (Object != nullptr && Object->IsA<UControlRigBlueprint>())
 			{
 				FGuid Binding = GetSequencer()->GetHandleToObject(InComponent, true /*bCreateHandle*/);
diff --git a/Engine/Plugins/Experimental/GeometryCollectionPlugin/Source/GeometryCollectionEditor/Private/GeometryCollection/GeometryCollectionConversion.cpp b/Engine/Plugins/Experimental/GeometryCollectionPlugin/Source/GeometryCollectionEditor/Private/GeometryCollection/GeometryCollectionConversion.cpp
--- a/Engine/Plugins/Experimental/GeometryCollectionPlugin/Source/GeometryCollectionEditor/Private/GeometryCollection/GeometryCollectionConversion.cpp
+++ b/Engine/Plugins/Experimental/GeometryCollectionPlugin/Source/GeometryCollectionEditor/Private/GeometryCollection/GeometryCollectionConversion.cpp
@@ -266,7 +266,7 @@ void FGeometryCollectionConversion::AppendSkeletalMesh(const USkeletalMesh* Skel
 		if (FGeometryCollection* GeometryCollection = GeometryCollectionPtr.Get())
 		{
 
-			if (USkeleton * Skeleton = SkeletalMesh->Skeleton)
+			if (const USkeleton* Skeleton = SkeletalMesh->GetSkeleton())
 			{
 				if (const FSkeletalMeshRenderData * SkelMeshRenderData = SkeletalMesh->GetResourceForRendering())
 				{
@@ -411,15 +411,17 @@ void FGeometryCollectionConversion::AppendSkeletalMesh(const USkeletalMesh* Skel
 
 						// for each material, add a reference in our GeometryCollectionObject
 						int CurrIdx = 0;
+						
+						const TArray<FSkeletalMaterial>& SkeletalMeshMaterials = SkeletalMesh->GetMaterials();
 
-						UMaterialInterface *CurrMaterial = SkeletalMeshComponent ? SkeletalMeshComponent->GetMaterial(CurrIdx) : ToRawPtr(SkeletalMesh->Materials[CurrIdx].MaterialInterface);
+						UMaterialInterface* CurrMaterial = SkeletalMeshComponent ? SkeletalMeshComponent->GetMaterial(CurrIdx) : ToRawPtr(SkeletalMeshMaterials[CurrIdx].MaterialInterface);
 
 
 						int MaterialStart = GeometryCollectionObject->Materials.Num();
 						while (CurrMaterial)
 						{
 							GeometryCollectionObject->Materials.Add(CurrMaterial);
-							CurrMaterial = SkeletalMeshComponent ? SkeletalMeshComponent->GetMaterial(++CurrIdx) : ToRawPtr(SkeletalMesh->Materials[++CurrIdx].MaterialInterface);
+							CurrMaterial = SkeletalMeshComponent ? SkeletalMeshComponent->GetMaterial(++CurrIdx) : ToRawPtr(SkeletalMeshMaterials[++CurrIdx].MaterialInterface);
 
 						}
 
diff --git a/Engine/Plugins/Experimental/MotionTrailEditorMode/Source/MotionTrailEditorMode/Private/Sequencer/AnimationBoneTrail.cpp b/Engine/Plugins/Experimental/MotionTrailEditorMode/Source/MotionTrailEditorMode/Private/Sequencer/AnimationBoneTrail.cpp
--- a/Engine/Plugins/Experimental/MotionTrailEditorMode/Source/MotionTrailEditorMode/Private/Sequencer/AnimationBoneTrail.cpp
+++ b/Engine/Plugins/Experimental/MotionTrailEditorMode/Source/MotionTrailEditorMode/Private/Sequencer/AnimationBoneTrail.cpp
@@ -28,14 +28,14 @@ void FAnimTrajectoryCache::Evaluate(FTrajectoryCache* ParentTransformCache)
 	TSharedPtr<ISequencer> Sequencer = WeakSequencer.Pin();
 	check(Sequencer);
 
-	if (!SkeletalMeshComponent.IsValid() || !SkeletalMeshComponent->SkeletalMesh || !SkeletalMeshComponent->SkeletalMesh->Skeleton)
+	if (!SkeletalMeshComponent.IsValid() || !SkeletalMeshComponent->SkeletalMesh || !SkeletalMeshComponent->SkeletalMesh->GetSkeleton())
 	{
 		return;
 	}
 
 	SkeletalMeshComponent->MasterPoseComponent = nullptr;
 
-	CachedAnimSequence->SetSkeleton(SkeletalMeshComponent->SkeletalMesh->Skeleton);
+	CachedAnimSequence->SetSkeleton(SkeletalMeshComponent->SkeletalMesh->GetSkeleton());
 
 	// TODO: for some reason SkeletalMeshComponent becomes invalid sometimes when evaluating, no clue why yet
 	FMovieSceneSequenceTransform MovieSceneSequenceTransform;
diff --git a/Engine/Plugins/Experimental/MotionTrailEditorMode/Source/MotionTrailEditorMode/Private/Sequencer/AnimationBoneTrail.h b/Engine/Plugins/Experimental/MotionTrailEditorMode/Source/MotionTrailEditorMode/Private/Sequencer/AnimationBoneTrail.h
--- a/Engine/Plugins/Experimental/MotionTrailEditorMode/Source/MotionTrailEditorMode/Private/Sequencer/AnimationBoneTrail.h
+++ b/Engine/Plugins/Experimental/MotionTrailEditorMode/Source/MotionTrailEditorMode/Private/Sequencer/AnimationBoneTrail.h
@@ -28,7 +28,7 @@ public:
 		, Spacing()
 		, bDirty(true)
 	{
-		CachedAnimSequence->SetSkeleton(InSkeletalMeshComponent->SkeletalMesh->Skeleton);
+		CachedAnimSequence->SetSkeleton(InSkeletalMeshComponent->SkeletalMesh->GetSkeleton());
 	}
 
 	void Evaluate(FTrajectoryCache* ParentTrajectoryCache);
diff --git a/Engine/Plugins/Experimental/MotionTrailEditorMode/Source/MotionTrailEditorMode/Private/Sequencer/MotionTrailTrackEditor.cpp b/Engine/Plugins/Experimental/MotionTrailEditorMode/Source/MotionTrailEditorMode/Private/Sequencer/MotionTrailTrackEditor.cpp
--- a/Engine/Plugins/Experimental/MotionTrailEditorMode/Source/MotionTrailEditorMode/Private/Sequencer/MotionTrailTrackEditor.cpp
+++ b/Engine/Plugins/Experimental/MotionTrailEditorMode/Source/MotionTrailEditorMode/Private/Sequencer/MotionTrailTrackEditor.cpp
@@ -133,9 +133,9 @@ void FMotionTrailTrackEditor::BuildObjectBindingContextMenu(FMenuBuilder& MenuBu
 					);
 				}
 
-				if (SkelMeshComp && SkelMeshComp->SkeletalMesh && SkelMeshComp->SkeletalMesh->Skeleton && Hierarchy->GetBonesTracked().Contains(SkelMeshComp))
+				if (SkelMeshComp && SkelMeshComp->SkeletalMesh && SkelMeshComp->SkeletalMesh->GetSkeleton() && Hierarchy->GetBonesTracked().Contains(SkelMeshComp))
 				{
-					USkeleton* Skeleton = SkelMeshComp->SkeletalMesh->Skeleton;
+					USkeleton* Skeleton = SkelMeshComp->SkeletalMesh->GetSkeleton();
 					const FName RootBoneName = Skeleton->GetReferenceSkeleton().GetBoneName(0);
 					const FGuid RootBoneTrailGuid = Hierarchy->GetBonesTracked().FindChecked(SkelMeshComp).FindChecked(RootBoneName);
 					const FGuid ComponentGuid = Hierarchy->GetObjectsTracked().FindChecked(SkelMeshComp);
diff --git a/Engine/Plugins/Experimental/MotionTrailEditorMode/Source/MotionTrailEditorMode/Private/Sequencer/SequencerTrailHierarchy.cpp b/Engine/Plugins/Experimental/MotionTrailEditorMode/Source/MotionTrailEditorMode/Private/Sequencer/SequencerTrailHierarchy.cpp
--- a/Engine/Plugins/Experimental/MotionTrailEditorMode/Source/MotionTrailEditorMode/Private/Sequencer/SequencerTrailHierarchy.cpp
+++ b/Engine/Plugins/Experimental/MotionTrailEditorMode/Source/MotionTrailEditorMode/Private/Sequencer/SequencerTrailHierarchy.cpp
@@ -223,7 +223,7 @@ void FSequencerTrailHierarchy::OnBoneVisibilityChanged(class USkeleton* Skeleton
 				BoundComponent = BoundActor->FindComponentByClass<USkeletalMeshComponent>();
 			}
 
-			if (!BoundComponent || !BoundComponent->SkeletalMesh || !BoundComponent->SkeletalMesh->Skeleton || !(BoundComponent->SkeletalMesh->Skeleton == Skeleton) || !BonesTracked.Contains(BoundComponent))
+			if (!BoundComponent || !BoundComponent->SkeletalMesh || !BoundComponent->SkeletalMesh->GetSkeleton() || !(BoundComponent->SkeletalMesh->GetSkeleton() == Skeleton) || !BonesTracked.Contains(BoundComponent))
 			{
 				continue;
 			}
@@ -335,7 +335,7 @@ void FSequencerTrailHierarchy::UpdateSequencerBindings(const TArray<FGuid>& Sequ
 					BoundComponent = BoundActor->FindComponentByClass<USkeletalMeshComponent>();
 				}
 
-				if (!BoundComponent || !BoundComponent->SkeletalMesh || !BoundComponent->SkeletalMesh->Skeleton)
+				if (!BoundComponent || !BoundComponent->SkeletalMesh || !BoundComponent->SkeletalMesh->GetSkeleton())
 				{
 					continue;
 				}
@@ -387,7 +387,7 @@ void FSequencerTrailHierarchy::UpdateSequencerBindings(const TArray<FGuid>& Sequ
 					BoundComponent = BoundActor->FindComponentByClass<USkeletalMeshComponent>();
 				}
 
-				if (!BoundComponent || !BoundComponent->SkeletalMesh || !BoundComponent->SkeletalMesh->Skeleton)
+				if (!BoundComponent || !BoundComponent->SkeletalMesh || !BoundComponent->SkeletalMesh->GetSkeleton())
 				{
 					continue;
 				}
@@ -562,7 +562,7 @@ void FSequencerTrailHierarchy::AddSkeletonToHierarchy(class USkeletalMeshCompone
 	TSharedPtr<FAnimTrajectoryCache> AnimTrajectoryCache = MakeShared<FAnimTrajectoryCache>(CompToAdd, Sequencer);
 	TMap<FName, FGuid>& BoneMap = BonesTracked.Add(CompToAdd, TMap<FName, FGuid>());
 	
-	USkeleton* MySkeleton = CompToAdd->SkeletalMesh->Skeleton;
+	USkeleton* MySkeleton = CompToAdd->SkeletalMesh->GetSkeleton();
 	const int32 NumBones = MySkeleton->GetReferenceSkeleton().GetNum();
 	for (int32 BoneIdx = 0; BoneIdx < NumBones; BoneIdx++)
 	{
diff --git a/Engine/Plugins/Experimental/SkeletalReduction/Source/Private/SkeletalMeshReductionPlugin.cpp b/Engine/Plugins/Experimental/SkeletalReduction/Source/Private/SkeletalMeshReductionPlugin.cpp
--- a/Engine/Plugins/Experimental/SkeletalReduction/Source/Private/SkeletalMeshReductionPlugin.cpp
+++ b/Engine/Plugins/Experimental/SkeletalReduction/Source/Private/SkeletalMeshReductionPlugin.cpp
@@ -1843,7 +1843,7 @@ void FQuadricSkeletalMeshReduction::ReduceSkeletalMesh(USkeletalMesh& SkeletalMe
 		ImportantBones.Weight = BonePrioritizationWeight;
 		for (const FBoneReference& BoneReference : BonesToPrioritize)
 		{
-			int32 BoneId = SkeletalMesh.RefSkeleton.FindRawBoneIndex(BoneReference.BoneName);
+			int32 BoneId = SkeletalMesh.GetRefSkeleton().FindRawBoneIndex(BoneReference.BoneName);
 
 			// Q: should we exclude BoneId = 0?
 			ImportantBones.Ids.AddUnique(BoneId);
@@ -1952,10 +1952,10 @@ void FQuadricSkeletalMeshReduction::ReduceSkeletalMesh(USkeletalMesh& SkeletalMe
 	IMeshBoneReduction* MeshBoneReductionInterface = FModuleManager::Get().LoadModuleChecked<IMeshBoneReductionModule>("MeshBoneReduction").GetMeshBoneReductionInterface();
 
 	TArray<FName> BoneNames;
-	const int32 NumBones = SkeletalMesh.RefSkeleton.GetNum();
+	const int32 NumBones = SkeletalMesh.GetRefSkeleton().GetNum();
 	for (int32 BoneIndex = 0; BoneIndex < NumBones; ++BoneIndex)
 	{
-		BoneNames.Add(SkeletalMesh.RefSkeleton.GetBoneName(BoneIndex));
+		BoneNames.Add(SkeletalMesh.GetRefSkeleton().GetBoneName(BoneIndex));
 	}
 
 	// get the relative to ref pose matrices
@@ -1967,7 +1967,7 @@ void FQuadricSkeletalMeshReduction::ReduceSkeletalMesh(USkeletalMesh& SkeletalMe
 		TArray<FTransform> BonePoses;
 		UAnimationBlueprintLibrary::GetBonePosesForFrame(BakePoseAnim, BoneNames, 0, true, BonePoses, &SkeletalMesh);
 
-		const FReferenceSkeleton& RefSkeleton = SkeletalMesh.RefSkeleton;
+		const FReferenceSkeleton& RefSkeleton = SkeletalMesh.GetRefSkeleton();
 		const TArray<FTransform>& RefPoseInLocal = RefSkeleton.GetRefBonePose();
 
 		// get component ref pose
@@ -2058,7 +2058,7 @@ void FQuadricSkeletalMeshReduction::ReduceSkeletalMesh(USkeletalMesh& SkeletalMe
 	}
 
 	// Reduce LOD model with SrcMesh if src mesh has more then 1 triangle
-	if (SrcModel->NumVertices > 3 && ReduceSkeletalLODModel(*SrcModel, *NewModel, SkeletalMesh.GetPathName(), SkeletalMesh.GetImportedBounds(), SkeletalMesh.RefSkeleton, Settings, ImportantBones, RelativeToRefPoseMatrices, LODIndex, bReducingSourceModel))
+	if (SrcModel->NumVertices > 3 && ReduceSkeletalLODModel(*SrcModel, *NewModel, SkeletalMesh.GetPathName(), SkeletalMesh.GetImportedBounds(), SkeletalMesh.GetRefSkeleton(), Settings, ImportantBones, RelativeToRefPoseMatrices, LODIndex, bReducingSourceModel))
 	{
 		FSkeletalMeshLODInfo* ReducedLODInfoPtr = SkeletalMesh.GetLODInfo(LODIndex);
 		check(ReducedLODInfoPtr);
@@ -2079,7 +2079,7 @@ void FQuadricSkeletalMeshReduction::ReduceSkeletalMesh(USkeletalMesh& SkeletalMe
 		}
 		// Flag this LOD as having been simplified.
 		ReducedLODInfoPtr->bHasBeenSimplified = true;
-		SkeletalMesh.bHasBeenSimplified = true;
+		SkeletalMesh.SetHasBeenSimplified(true);
 
 		//Restore the user sections data to what it was. It must be done if we want to avoid changing the DDC key. I.E. UserSectionData is part of the key
 		//DDC key cannot be change during the build
@@ -2155,7 +2155,7 @@ void FQuadricSkeletalMeshReduction::ReduceSkeletalMesh(USkeletalMesh& SkeletalMe
 
 		NewModel->RequiredBones.Empty();
 		SkeletalMesh.GetLODInfo(LODIndex)->bHasBeenSimplified = true;
-		SkeletalMesh.bHasBeenSimplified = true;
+		SkeletalMesh.SetHasBeenSimplified(true);
 	}
 	
 	if (!bLODModelAdded)
@@ -2175,7 +2175,7 @@ void FQuadricSkeletalMeshReduction::ReduceSkeletalMesh(USkeletalMesh& SkeletalMe
 		SkeletalMesh.SetLODImportedDataVersions(LODIndex, GeoImportVersion, SkinningImportVersion);
 	}
 
-	SkeletalMesh.CalculateRequiredBones(SkeletalMeshResource.LODModels[LODIndex], SkeletalMesh.RefSkeleton, &BonesToRemove);
+	SkeletalMesh.CalculateRequiredBones(SkeletalMeshResource.LODModels[LODIndex], SkeletalMesh.GetRefSkeleton(), &BonesToRemove);
 }
 
 #undef LOCTEXT_NAMESPACE
diff --git a/Engine/Plugins/FX/Niagara/Source/Niagara/Classes/NDISkeletalMeshCommon.h b/Engine/Plugins/FX/Niagara/Source/Niagara/Classes/NDISkeletalMeshCommon.h
--- a/Engine/Plugins/FX/Niagara/Source/Niagara/Classes/NDISkeletalMeshCommon.h
+++ b/Engine/Plugins/FX/Niagara/Source/Niagara/Classes/NDISkeletalMeshCommon.h
@@ -105,7 +105,7 @@ struct FSkinnedPositionAccessorHelper<TNDISkelMesh_SkinningModeInvalid>
 	{
 		if (const USkeletalMesh* Mesh = Accessor.Mesh)
 		{
-			return Mesh->RefSkeleton.GetNum();
+			return Mesh->GetRefSkeleton().GetNum();
 		}
 
 		return 0;
@@ -145,13 +145,13 @@ struct FSkinnedPositionAccessorHelper<TNDISkelMesh_SkinningModeInvalid>
 	// The bone accessor functions are valid if a mesh is present, so don't stub them entirely
 	FORCEINLINE_DEBUGGABLE FVector GetSkinnedBonePosition(FSkeletalMeshAccessorHelper& Accessor, int32 BoneIndex)
 	{
-		const int32 NumRealBones = Accessor.Mesh->RefSkeleton.GetRawBoneNum();
+		const int32 NumRealBones = Accessor.Mesh->GetRefSkeleton().GetRawBoneNum();
 		if (BoneIndex < NumRealBones)
 		{
 			return Accessor.Mesh->GetComposedRefPoseMatrix(BoneIndex).GetOrigin();
 		}
 
-		const FTransform& RefTransform = Accessor.Mesh->RefSkeleton.GetRefBonePose()[BoneIndex];
+		const FTransform& RefTransform = Accessor.Mesh->GetRefSkeleton().GetRefBonePose()[BoneIndex];
 		return RefTransform.GetLocation();
 	}
 
@@ -162,13 +162,13 @@ struct FSkinnedPositionAccessorHelper<TNDISkelMesh_SkinningModeInvalid>
 
 	FORCEINLINE_DEBUGGABLE FQuat GetSkinnedBoneRotation(FSkeletalMeshAccessorHelper& Accessor, int32 BoneIndex)
 	{
-		const int32 NumRealBones = Accessor.Mesh->RefSkeleton.GetRawBoneNum();
+		const int32 NumRealBones = Accessor.Mesh->GetRefSkeleton().GetRawBoneNum();
 		if (BoneIndex < NumRealBones)
 		{
 			return Accessor.Mesh->GetComposedRefPoseMatrix(BoneIndex).GetMatrixWithoutScale().ToQuat();
 		}
 
-		const FTransform& RefTransform = Accessor.Mesh->RefSkeleton.GetRefBonePose()[BoneIndex];
+		const FTransform& RefTransform = Accessor.Mesh->GetRefSkeleton().GetRefBonePose()[BoneIndex];
 		return RefTransform.GetRotation();
 	}
 
@@ -185,7 +185,7 @@ struct FSkinnedPositionAccessorHelper<TNDISkelMesh_SkinningModeNone>
 	{
 		if (const USkeletalMesh* Mesh = Accessor.Mesh)
 		{
-			return Mesh->RefSkeleton.GetNum();
+			return Mesh->GetRefSkeleton().GetNum();
 		}
 
 		return 0;
@@ -231,13 +231,13 @@ struct FSkinnedPositionAccessorHelper<TNDISkelMesh_SkinningModeNone>
 
 	FORCEINLINE_DEBUGGABLE FVector GetSkinnedBonePosition(FSkeletalMeshAccessorHelper& Accessor, int32 BoneIndex)
 	{
-		const int32 NumRealBones = Accessor.Mesh->RefSkeleton.GetRawBoneNum();
+		const int32 NumRealBones = Accessor.Mesh->GetRefSkeleton().GetRawBoneNum();
 		if (BoneIndex < NumRealBones)
 		{
 			return Accessor.Mesh->GetComposedRefPoseMatrix(BoneIndex).GetOrigin();
 		}
 
-		const FTransform& RefTransform = Accessor.Mesh->RefSkeleton.GetRefBonePose()[BoneIndex];
+		const FTransform& RefTransform = Accessor.Mesh->GetRefSkeleton().GetRefBonePose()[BoneIndex];
 		return RefTransform.GetLocation();
 	}
 
@@ -248,13 +248,13 @@ struct FSkinnedPositionAccessorHelper<TNDISkelMesh_SkinningModeNone>
 
 	FORCEINLINE_DEBUGGABLE FQuat GetSkinnedBoneRotation(FSkeletalMeshAccessorHelper& Accessor, int32 BoneIndex)
 	{
-		const int32 NumRealBones = Accessor.Mesh->RefSkeleton.GetRawBoneNum();
+		const int32 NumRealBones = Accessor.Mesh->GetRefSkeleton().GetRawBoneNum();
 		if (BoneIndex < NumRealBones)
 		{
 			return Accessor.Mesh->GetComposedRefPoseMatrix(BoneIndex).GetMatrixWithoutScale().ToQuat();
 		}
 
-		const FTransform& RefTransform = Accessor.Mesh->RefSkeleton.GetRefBonePose()[BoneIndex];
+		const FTransform& RefTransform = Accessor.Mesh->GetRefSkeleton().GetRefBonePose()[BoneIndex];
 		return RefTransform.GetRotation();
 	}
 
diff --git a/Engine/Plugins/FX/Niagara/Source/Niagara/Private/NiagaraDataInterfaceSkeletalMesh.cpp b/Engine/Plugins/FX/Niagara/Source/Niagara/Private/NiagaraDataInterfaceSkeletalMesh.cpp
--- a/Engine/Plugins/FX/Niagara/Source/Niagara/Private/NiagaraDataInterfaceSkeletalMesh.cpp
+++ b/Engine/Plugins/FX/Niagara/Source/Niagara/Private/NiagaraDataInterfaceSkeletalMesh.cpp
@@ -274,10 +274,10 @@ void FSkeletalMeshSkinningData::UpdateBoneTransforms()
 		if (NumBones == 0)
 		{
 			// This case indicates an invalid master pose component (e.g. no skeletal mesh)
-			CurrBones.Empty(SkelMesh->RefSkeleton.GetNum());
-			CurrBones.AddDefaulted(SkelMesh->RefSkeleton.GetNum());
-			CurrTransforms.Empty(SkelMesh->RefSkeleton.GetNum());
-			CurrTransforms.AddDefaulted(SkelMesh->RefSkeleton.GetNum());
+			CurrBones.Empty(SkelMesh->GetRefSkeleton().GetNum());
+			CurrBones.AddDefaulted(SkelMesh->GetRefSkeleton().GetNum());
+			CurrTransforms.Empty(SkelMesh->GetRefSkeleton().GetNum());
+			CurrTransforms.AddDefaulted(SkelMesh->GetRefSkeleton().GetNum());
 		}
 		else
 		{
@@ -295,22 +295,22 @@ void FSkeletalMeshSkinningData::UpdateBoneTransforms()
 					if (MasterIndex != INDEX_NONE && MasterIndex < MasterTransforms.Num())
 					{
 						CurrTransforms[BoneIndex] = MasterTransforms[MasterIndex];
-						CurrBones[BoneIndex] = SkelMesh->RefBasesInvMatrix[BoneIndex] * MasterTransforms[MasterIndex].ToMatrixWithScale();
+						CurrBones[BoneIndex] = SkelMesh->GetRefBasesInvMatrix()[BoneIndex] * MasterTransforms[MasterIndex].ToMatrixWithScale();
 						bFoundMaster = true;
 					}
 				}
 
 				if (!bFoundMaster)
 				{
-					const int32 ParentIndex = SkelMesh->RefSkeleton.GetParentIndex(BoneIndex);
+					const int32 ParentIndex = SkelMesh->GetRefSkeleton().GetParentIndex(BoneIndex);
 
 					if (CurrTransforms.IsValidIndex(ParentIndex) && ParentIndex < BoneIndex)
 					{
-						CurrTransforms[BoneIndex] = CurrTransforms[ParentIndex] * SkelMesh->RefSkeleton.GetRefBonePose()[BoneIndex];
+						CurrTransforms[BoneIndex] = CurrTransforms[ParentIndex] * SkelMesh->GetRefSkeleton().GetRefBonePose()[BoneIndex];
 					}
 					else
 					{
-						CurrTransforms[BoneIndex] = SkelMesh->RefSkeleton.GetRefBonePose()[BoneIndex];
+						CurrTransforms[BoneIndex] = SkelMesh->GetRefSkeleton().GetRefBonePose()[BoneIndex];
 					}
 
 					CurrBones[BoneIndex] = CurrTransforms[BoneIndex].ToMatrixWithScale();
@@ -905,7 +905,7 @@ void FSkeletalMeshGpuDynamicBufferProxy::NewFrame(const FNDISkeletalMesh_Instanc
 				{
 					const int32 BoneIndex = Section.BoneMap[m];
 					const FTransform& BoneTransform = BoneTransforms[BoneIndex];
-					const FMatrix BoneMatrix = SkelMesh->RefBasesInvMatrix.IsValidIndex(BoneIndex) ? SkelMesh->RefBasesInvMatrix[BoneIndex] * BoneTransform.ToMatrixWithScale() : BoneTransform.ToMatrixWithScale();
+					const FMatrix BoneMatrix = SkelMesh->GetRefBasesInvMatrix().IsValidIndex(BoneIndex) ? SkelMesh->GetRefBasesInvMatrix()[BoneIndex] * BoneTransform.ToMatrixWithScale() : BoneTransform.ToMatrixWithScale();
 					BoneMatrix.To3x4MatrixTranspose(&AllSectionsRefToLocalMatrices[Float4Count].X);
 					Float4Count += 3;
 				}
@@ -962,8 +962,8 @@ void FSkeletalMeshGpuDynamicBufferProxy::NewFrame(const FNDISkeletalMesh_Instanc
 						}
 					}
 
-					const int32 ParentIndex = SkelMesh->RefSkeleton.GetParentIndex(BoneIndex);
-					FTransform BoneTransform =SkelMesh->RefSkeleton.GetRefBonePose()[BoneIndex];
+					const int32 ParentIndex = SkelMesh->GetRefSkeleton().GetParentIndex(BoneIndex);
+					FTransform BoneTransform =SkelMesh->GetRefSkeleton().GetRefBonePose()[BoneIndex];
 					if (TempBoneTransforms.IsValidIndex(ParentIndex))
 					{
 						BoneTransform = BoneTransform * TempBoneTransforms[ParentIndex];
@@ -985,11 +985,11 @@ void FSkeletalMeshGpuDynamicBufferProxy::NewFrame(const FNDISkeletalMesh_Instanc
 		TArray<FTransform> TempBoneTransforms;
 		TempBoneTransforms.Reserve(SamplingBoneCount);
 
-		const TArray<FTransform>& RefTransforms = SkelMesh->RefSkeleton.GetRefBonePose();
+		const TArray<FTransform>& RefTransforms = SkelMesh->GetRefSkeleton().GetRefBonePose();
 		for (int32 i=0; i < RefTransforms.Num(); ++i)
 		{
 			FTransform BoneTransform = RefTransforms[i];
-			const int32 ParentIndex = SkelMesh->RefSkeleton.GetParentIndex(i);
+			const int32 ParentIndex = SkelMesh->GetRefSkeleton().GetParentIndex(i);
 			if (TempBoneTransforms.IsValidIndex(ParentIndex))
 			{
 				BoneTransform = BoneTransform * TempBoneTransforms[ParentIndex];
@@ -1609,7 +1609,7 @@ bool FNDISkeletalMesh_InstanceData::Init(UNiagaraDataInterfaceSkeletalMesh* Inte
 		if (UNiagaraComponent* NiagaraComponent = Cast<UNiagaraComponent>(SystemInstance->GetAttachComponent()))
 		{
 			Mesh->GetOnMeshChanged().AddUObject(NiagaraComponent, &UNiagaraComponent::ReinitializeSystem);
-			if (USkeleton* Skeleton = Mesh->Skeleton)
+			if (USkeleton* Skeleton = Mesh->GetSkeleton())
 			{
 				Skeleton->RegisterOnSkeletonHierarchyChanged(USkeleton::FOnSkeletonHierarchyChanged::CreateUObject(NiagaraComponent, &UNiagaraComponent::ReinitializeSystem));
 			}
@@ -1635,7 +1635,7 @@ bool FNDISkeletalMesh_InstanceData::Init(UNiagaraDataInterfaceSkeletalMesh* Inte
 			return false;
 		}
 
-		MinLODIdx = Mesh->MinLod.GetValue();
+		MinLODIdx = Mesh->GetMinLod().GetValue();
 		const int32 PendingFirstLODIndex = RenderData->GetPendingFirstLODIdx(MinLODIdx);
 		if (PendingFirstLODIndex == INDEX_NONE)
 		{
@@ -1729,7 +1729,7 @@ bool FNDISkeletalMesh_InstanceData::Init(UNiagaraDataInterfaceSkeletalMesh* Inte
 		}
 
 		// Generate excluded root bone index (if any)
-		FReferenceSkeleton& RefSkel = Mesh->RefSkeleton;
+		FReferenceSkeleton& RefSkel = Mesh->GetRefSkeleton();
 		ExcludedBoneIndex = INDEX_NONE;
 		if (Interface->bExcludeBone && !Interface->ExcludeBoneName.IsNone())
 		{
@@ -1832,7 +1832,7 @@ bool FNDISkeletalMesh_InstanceData::Init(UNiagaraDataInterfaceSkeletalMesh* Inte
 				}
 			}
 
-			FilteredSocketBoneOffset = Mesh->RefSkeleton.GetNum();
+			FilteredSocketBoneOffset = Mesh->GetRefSkeleton().GetNum();
 
 			FilteredSocketTransformsIndex = 0;
 			FilteredSocketTransforms[0].Reset(FilteredSockets.Num());
@@ -2269,7 +2269,7 @@ void UNiagaraDataInterfaceSkeletalMesh::DestroyPerInstanceData(void* PerInstance
 		if (UNiagaraComponent* NiagaraComponent = Cast<UNiagaraComponent>(SystemInstance->GetAttachComponent()))
 		{
 			SkeletalMesh->GetOnMeshChanged().RemoveAll(NiagaraComponent);
-			if (USkeleton* Skeleton = SkeletalMesh->Skeleton)
+			if (USkeleton* Skeleton = SkeletalMesh->GetSkeleton())
 			{
 				Skeleton->UnregisterOnSkeletonHierarchyChanged(NiagaraComponent);
 			}
diff --git a/Engine/Plugins/FX/Niagara/Source/NiagaraEditorWidgets/Private/DetailCustomizations/NiagaraDataInterfaceSkeletalMeshDetails.cpp b/Engine/Plugins/FX/Niagara/Source/NiagaraEditorWidgets/Private/DetailCustomizations/NiagaraDataInterfaceSkeletalMeshDetails.cpp
--- a/Engine/Plugins/FX/Niagara/Source/NiagaraEditorWidgets/Private/DetailCustomizations/NiagaraDataInterfaceSkeletalMeshDetails.cpp
+++ b/Engine/Plugins/FX/Niagara/Source/NiagaraEditorWidgets/Private/DetailCustomizations/NiagaraDataInterfaceSkeletalMeshDetails.cpp
@@ -206,7 +206,7 @@ void FNiagaraDataInterfaceSkeletalMeshDetails::GenerateRegionsArray(TArray<TShar
 		USkeletalMeshComponent* FoundSkelComp = nullptr;
 		if (USkeletalMesh* Mesh = Interface->GetSkeletalMesh(SystemInstance, SceneComponent, FoundSkelComp))
 		{
-			for (const FMeshBoneInfo& Bone : Mesh->RefSkeleton.GetRefBoneInfo())
+			for (const FMeshBoneInfo& Bone : Mesh->GetRefSkeleton().GetRefBoneInfo())
 			{
 				SourceArray.Add(MakeShared<FName>(Bone.Name));
 			}
diff --git a/Engine/Plugins/Importers/USDImporter/Source/USDSchemas/Private/USDSkelRootTranslator.cpp b/Engine/Plugins/Importers/USDImporter/Source/USDSchemas/Private/USDSkelRootTranslator.cpp
--- a/Engine/Plugins/Importers/USDImporter/Source/USDSchemas/Private/USDSkelRootTranslator.cpp
+++ b/Engine/Plugins/Importers/USDImporter/Source/USDSchemas/Private/USDSkelRootTranslator.cpp
@@ -64,7 +64,7 @@ namespace UsdSkelRootTranslatorImpl
 		}
 
 		TArray<UMaterialInterface*> ExistingAssignments;
-		for ( const FSkeletalMaterial& SkeletalMaterial : SkeletalMesh->Materials )
+		for ( const FSkeletalMaterial& SkeletalMaterial : SkeletalMesh->GetMaterials())
 		{
 			ExistingAssignments.Add( SkeletalMaterial.MaterialInterface );
 		}
@@ -117,9 +117,9 @@ namespace UsdSkelRootTranslatorImpl
 				FName MaterialSlotName = *LexToString( SkeletalMeshSlotIndex );
 
 				// Already have a material at that skeletal mesh slot, need to reassign
-				if ( SkeletalMesh->Materials.IsValidIndex( SkeletalMeshSlotIndex ) )
+				if ( SkeletalMesh->GetMaterials().IsValidIndex( SkeletalMeshSlotIndex ) )
 				{
-					FSkeletalMaterial& ExistingMaterial = SkeletalMesh->Materials[ SkeletalMeshSlotIndex ];
+					FSkeletalMaterial& ExistingMaterial = SkeletalMesh->GetMaterials()[ SkeletalMeshSlotIndex ];
 
 					if ( ExistingMaterial.MaterialInterface != Material ||
 						 ExistingMaterial.MaterialSlotName != MaterialSlotName ||
@@ -136,7 +136,7 @@ namespace UsdSkelRootTranslatorImpl
 				{
 					const bool bEnableShadowCasting = true;
 					const bool bRecomputeTangents = false;
-					SkeletalMesh->Materials.Add( FSkeletalMaterial( Material, bEnableShadowCasting, bRecomputeTangents, MaterialSlotName, MaterialSlotName ) );
+					SkeletalMesh->GetMaterials().Add( FSkeletalMaterial( Material, bEnableShadowCasting, bRecomputeTangents, MaterialSlotName, MaterialSlotName ) );
 					bMaterialsHaveChanged = true;
 				}
 
@@ -252,7 +252,7 @@ namespace UsdSkelRootTranslatorImpl
 	{
 		USkeletalMesh* SkeletalMesh = PoseableMeshComponent.SkeletalMesh;
 
-		// We try keeping a perfect correspondence between SkeletalMesh->MorphTargets and PoseableMeshComponent.ActiveMorphTargets
+		// We try keeping a perfect correspondence between SkeletalMesh->GetMorphTargets() and PoseableMeshComponent.ActiveMorphTargets
 		int32 IndexInSkeletalMesh = INDEX_NONE;
 		PoseableMeshComponent.SkeletalMesh->FindMorphTargetAndIndex( *MorphTargetName, IndexInSkeletalMesh );
 		if ( IndexInSkeletalMesh == INDEX_NONE )
@@ -260,7 +260,7 @@ namespace UsdSkelRootTranslatorImpl
 			return;
 		}
 
-		UMorphTarget* MorphTarget = PoseableMeshComponent.SkeletalMesh->MorphTargets[ IndexInSkeletalMesh ];
+		UMorphTarget* MorphTarget = PoseableMeshComponent.SkeletalMesh->GetMorphTargets()[ IndexInSkeletalMesh ];
 		if ( !MorphTarget )
 		{
 			return;
@@ -283,10 +283,10 @@ namespace UsdSkelRootTranslatorImpl
 			PoseableMeshComponent.ActiveMorphTargets.Reset();
 			PoseableMeshComponent.MorphTargetWeights.Reset();
 
-			for ( int32 MorphTargetIndex = 0; MorphTargetIndex < SkeletalMesh->MorphTargets.Num(); ++MorphTargetIndex )
+			for ( int32 MorphTargetIndex = 0; MorphTargetIndex < SkeletalMesh->GetMorphTargets().Num(); ++MorphTargetIndex )
 			{
 				FActiveMorphTarget ActiveMorphTarget;
-				ActiveMorphTarget.MorphTarget = SkeletalMesh->MorphTargets[ MorphTargetIndex ];
+				ActiveMorphTarget.MorphTarget = SkeletalMesh->GetMorphTargets()[ MorphTargetIndex ];
 				ActiveMorphTarget.WeightIndex = MorphTargetIndex;
 
 				PoseableMeshComponent.ActiveMorphTargets.Add( ActiveMorphTarget );
@@ -753,7 +753,7 @@ namespace UsdSkelRootTranslatorImpl
 				if ( SkeletalMesh )
 				{
 					Context->CurrentlyUsedAssets.Add( SkeletalMesh );
-					Context->CurrentlyUsedAssets.Add( SkeletalMesh->Skeleton );
+					Context->CurrentlyUsedAssets.Add( SkeletalMesh->GetSkeleton() );
 
 					FScopeLock Lock( &Context->CriticalSection );
 					{
@@ -781,10 +781,10 @@ namespace UsdSkelRootTranslatorImpl
 
 						UUsdAssetImportData* ImportData = NewObject< UUsdAssetImportData >( SkeletalMesh, TEXT( "USDAssetImportData" ) );
 						ImportData->PrimPath = SkelRootPath;
-						SkeletalMesh->AssetImportData = ImportData;
+						SkeletalMesh->SetAssetImportData(ImportData);
 
 						Context->AssetsCache.Add( SkeletalMeshHash.ToString(), SkeletalMesh );
-						Context->AssetsCache.Add( SkeletalMeshHash.ToString() + TEXT( "_Skeleton" ), SkeletalMesh->Skeleton );
+						Context->AssetsCache.Add( SkeletalMeshHash.ToString() + TEXT( "_Skeleton" ), SkeletalMesh->GetSkeleton() );
 					}
 
 					// We may be reusing a skeletal mesh we got in the cache, but we always need the BlendShapesByPath stored on the
@@ -794,7 +794,7 @@ namespace UsdSkelRootTranslatorImpl
 						Context->BlendShapesByPath->Append( NewBlendShapes );
 					}
 
-					for ( const FSkeletalMaterial& SkeletalMaterial : SkeletalMesh->Materials )
+					for ( const FSkeletalMaterial& SkeletalMaterial : SkeletalMesh->GetMaterials())
 					{
 						Context->CurrentlyUsedAssets.Add( SkeletalMaterial.MaterialInterface );
 					}
@@ -848,7 +848,7 @@ namespace UsdSkelRootTranslatorImpl
 						FString HashString = Hash.ToString();
 						UAnimSequence* AnimSequence = Cast< UAnimSequence >( Context->AssetsCache.FindRef( HashString ) );
 
-						if ( !AnimSequence || AnimSequence->GetSkeleton() != SkeletalMesh->Skeleton )
+						if ( !AnimSequence || AnimSequence->GetSkeleton() != SkeletalMesh->GetSkeleton() )
 						{
 							FScopedUnrealAllocs UEAllocs;
 
@@ -859,7 +859,7 @@ namespace UsdSkelRootTranslatorImpl
 							// Besides, this particular asset type is only ever created when we import to content folder assets (so never for realtime), and
 							// in that case we don't need it to be transactional anyway
 							AnimSequence = NewObject<UAnimSequence>( GetTransientPackage(), *UsdToUnreal::ConvertString( SkelAnimPrim.GetName() ), Context->ObjectFlags & ~EObjectFlags::RF_Transactional );
-							AnimSequence->SetSkeleton(SkeletalMesh->Skeleton);
+							AnimSequence->SetSkeleton(SkeletalMesh->GetSkeleton());
 
 							TUsdStore<pxr::VtArray<pxr::UsdSkelSkinningQuery>> SkinningTargets = Binding.GetSkinningTargets();
 							UsdToUnreal::ConvertSkelAnim( SkelQuery, &SkinningTargets.Get(), &NewBlendShapes, Context->bAllowInterpretingLODs, AnimSequence );
@@ -922,7 +922,7 @@ void FUsdSkelRootTranslator::UpdateComponents( USceneComponent* SceneComponent )
 		// Handle material overrides
 		if ( TargetSkeletalMesh )
 		{
-			if ( UUsdAssetImportData* UsdImportData = Cast<UUsdAssetImportData>( TargetSkeletalMesh->AssetImportData ) )
+			if ( UUsdAssetImportData* UsdImportData = Cast<UUsdAssetImportData>( TargetSkeletalMesh->GetAssetImportData() ) )
 			{
 				// If the prim paths match, it means that it was this prim that created (and so "owns") the mesh,
 				// so its material assignments will already be directly on the mesh. If they differ, we're using some other prim's mesh,
@@ -930,7 +930,7 @@ void FUsdSkelRootTranslator::UpdateComponents( USceneComponent* SceneComponent )
 				if ( UsdImportData->PrimPath != PrimPath.GetString() )
 				{
 					TArray<UMaterialInterface*> ExistingAssignments;
-					for ( FSkeletalMaterial& SkeletalMaterial : TargetSkeletalMesh->Materials )
+					for ( FSkeletalMaterial& SkeletalMaterial : TargetSkeletalMesh->GetMaterials())
 					{
 						ExistingAssignments.Add( SkeletalMaterial.MaterialInterface );
 					}
@@ -1002,7 +1002,7 @@ void FUsdSkelRootTranslator::UpdateComponents( USceneComponent* SceneComponent )
 			{
 				FScopedUnrealAllocs UnrealAllocs;
 
-				if ( PoseableMeshComponent->BoneSpaceTransforms.Num() != PoseableMeshComponent->SkeletalMesh->RefSkeleton.GetNum() )
+				if ( PoseableMeshComponent->BoneSpaceTransforms.Num() != PoseableMeshComponent->SkeletalMesh->GetRefSkeleton().GetNum() )
 				{
 					PoseableMeshComponent->AllocateTransformData();
 				}
diff --git a/Engine/Plugins/Importers/USDImporter/Source/USDStageImporter/Private/USDStageImporter.cpp b/Engine/Plugins/Importers/USDImporter/Source/USDStageImporter/Private/USDStageImporter.cpp
--- a/Engine/Plugins/Importers/USDImporter/Source/USDStageImporter/Private/USDStageImporter.cpp
+++ b/Engine/Plugins/Importers/USDImporter/Source/USDStageImporter/Private/USDStageImporter.cpp
@@ -327,7 +327,7 @@ namespace UsdStageImporterImpl
 						if ( OldSkeletalMesh->GetLODNum() != NumUsdLODs )
 						{
 							AssetsToDiscard.Add( OldSkeletalMesh );
-							AssetsToDiscard.Add( OldSkeletalMesh->Skeleton );
+							AssetsToDiscard.Add( OldSkeletalMesh->GetSkeleton() );
 							ImportContext.PrimPathsToAssets.Remove( SkelRootPrimPath );
 						}
 					}
@@ -429,7 +429,7 @@ namespace UsdStageImporterImpl
 		{
 			AssetPrefix = TEXT("SK_");
 
-			if (UUsdAssetImportData* AssetImportData = Cast<UUsdAssetImportData>(SkMesh->AssetImportData))
+			if (UUsdAssetImportData* AssetImportData = Cast<UUsdAssetImportData>(SkMesh->GetAssetImportData()))
 			{
 				AssetPath = AssetImportData->PrimPath;
 			}
@@ -442,7 +442,7 @@ namespace UsdStageImporterImpl
 			// as it can determine the destination path and how the asset conflicts are resolved
 			if (USkeletalMesh* SkeletalMesh = Skeleton->GetPreviewMesh())
 			{
-				if (UUsdAssetImportData* AssetImportData = Cast<UUsdAssetImportData>(SkeletalMesh->AssetImportData))
+				if (UUsdAssetImportData* AssetImportData = Cast<UUsdAssetImportData>(SkeletalMesh->GetAssetImportData()))
 				{
 					AssetPath = AssetImportData->PrimPath;
 				}
@@ -666,7 +666,7 @@ namespace UsdStageImporterImpl
 		// Note that we need to do this even if we never replaced/overwrote assets, as we will at least move from transient to the published folder
 		if ( USkeletalMesh* MovedMeshAsset = Cast<USkeletalMesh>( MovedAsset ) )
 		{
-			if ( USkeleton* Skeleton = MovedMeshAsset->Skeleton )
+			if ( USkeleton* Skeleton = MovedMeshAsset->GetSkeleton() )
 			{
 				Skeleton->SetPreviewMesh( MovedMeshAsset );
 			}
@@ -1002,7 +1002,7 @@ namespace UsdStageImporterImpl
 		USkeletalMesh* NewSkeletalMesh = Cast<USkeletalMesh>(NewAsset);
 		if (ExistingSkeletalMesh && NewSkeletalMesh)
 		{
-			NewSkeletalMesh->Materials = ExistingSkeletalMesh->Materials;
+			NewSkeletalMesh->SetMaterials(ExistingSkeletalMesh->GetMaterials());
 			return;
 		}
 	}
@@ -1014,14 +1014,14 @@ namespace UsdStageImporterImpl
 		if ( ExistingSkeletalMesh && NewSkeletalMesh )
 		{
 			// Never assign a transient skeleton
-			if ( ExistingSkeletalMesh->Skeleton && ExistingSkeletalMesh->Skeleton->GetOutermost() == GetTransientPackage() )
+			if ( ExistingSkeletalMesh->GetSkeleton() && ExistingSkeletalMesh->GetSkeleton()->GetOutermost() == GetTransientPackage() )
 			{
 				return;
 			}
 
 			// Assign even if ExistingSkeletalMesh has nullptr skeleton because we must be able to cleanup the
 			// abandoned Skeleton in the transient package
-			NewSkeletalMesh->Skeleton = ExistingSkeletalMesh->Skeleton;
+			NewSkeletalMesh->SetSkeleton(ExistingSkeletalMesh->GetSkeleton());
 		}
 
 		UAnimSequence* ExistingAnimSequence = Cast<UAnimSequence>( ExistingAsset );
diff --git a/Engine/Plugins/Importers/USDImporter/Source/USDUtilities/Private/USDConversionUtils.cpp b/Engine/Plugins/Importers/USDImporter/Source/USDUtilities/Private/USDConversionUtils.cpp
--- a/Engine/Plugins/Importers/USDImporter/Source/USDUtilities/Private/USDConversionUtils.cpp
+++ b/Engine/Plugins/Importers/USDImporter/Source/USDUtilities/Private/USDConversionUtils.cpp
@@ -693,12 +693,12 @@ UUsdAssetImportData* UsdUtils::GetAssetImportData( UObject* Asset )
 	{
 		if ( USkeletalMesh* SkMesh = Skeleton->GetPreviewMesh() )
 		{
-			ImportData = Cast<UUsdAssetImportData>( SkMesh->AssetImportData );
+			ImportData = Cast<UUsdAssetImportData>( SkMesh->GetAssetImportData() );
 		}
 	}
 	else if ( USkeletalMesh* SkMesh = Cast<USkeletalMesh>( Asset ) )
 	{
-		ImportData = Cast<UUsdAssetImportData>( SkMesh->AssetImportData );
+		ImportData = Cast<UUsdAssetImportData>(SkMesh->GetAssetImportData());
 	}
 	else if ( UAnimSequence* SkelAnim = Cast<UAnimSequence>( Asset ) )
 	{
diff --git a/Engine/Plugins/Importers/USDImporter/Source/USDUtilities/Private/USDSkeletalDataConversion.cpp b/Engine/Plugins/Importers/USDImporter/Source/USDUtilities/Private/USDSkeletalDataConversion.cpp
--- a/Engine/Plugins/Importers/USDImporter/Source/USDUtilities/Private/USDSkeletalDataConversion.cpp
+++ b/Engine/Plugins/Importers/USDImporter/Source/USDUtilities/Private/USDSkeletalDataConversion.cpp
@@ -627,7 +627,7 @@ namespace UsdToUnrealImpl
 			BlendShape.Vertices.Empty();
 
 			MorphTarget->BaseSkelMesh = SkeletalMesh;
-			SkeletalMesh->MorphTargets.Add( MorphTarget );
+			SkeletalMesh->GetMorphTargets().Add( MorphTarget );
 		}
 
 		if ( bHasValidMorphTarget )
@@ -2033,11 +2033,11 @@ USkeletalMesh* UsdToUnreal::GetSkeletalMeshFromImportData(
 	int32 SkeletalDepth = 0;
 	FSkeletalMeshImportData DummyData;
 	DummyData.RefBonesBinary = InSkeletonBones;
-	if ( !SkeletalMeshHelper::ProcessImportMeshSkeleton( SkeletalMesh->Skeleton, SkeletalMesh->RefSkeleton, SkeletalDepth, DummyData ) )
+	if ( !SkeletalMeshHelper::ProcessImportMeshSkeleton( SkeletalMesh->GetSkeleton(), SkeletalMesh->GetRefSkeleton(), SkeletalDepth, DummyData ) )
 	{
 		return nullptr;
 	}
-	if ( SkeletalMesh->RefSkeleton.GetRawBoneNum() == 0 )
+	if ( SkeletalMesh->GetRefSkeleton().GetRawBoneNum() == 0 )
 	{
 		SkeletalMesh->MarkPendingKill();
 		return nullptr;
@@ -2110,7 +2110,7 @@ USkeletalMesh* UsdToUnreal::GetSkeletalMeshFromImportData(
 		TArray<FText> WarningMessages;
 		TArray<FName> WarningNames;
 
-		bool bBuildSuccess = MeshUtilities.BuildSkeletalMesh(LODModel, SkeletalMesh->GetPathName(), SkeletalMesh->RefSkeleton, LODInfluences, LODWedges, LODFaces, LODPoints, LODPointToRawMap, BuildOptions, &WarningMessages, &WarningNames );
+		bool bBuildSuccess = MeshUtilities.BuildSkeletalMesh(LODModel, SkeletalMesh->GetPathName(), SkeletalMesh->GetRefSkeleton(), LODInfluences, LODWedges, LODFaces, LODPoints, LODPointToRawMap, BuildOptions, &WarningMessages, &WarningNames );
 		if ( !bBuildSuccess )
 		{
 			SkeletalMesh->MarkPendingKill();
@@ -2123,15 +2123,15 @@ USkeletalMesh* UsdToUnreal::GetSkeletalMeshFromImportData(
 	}
 
 	SkeletalMesh->SetImportedBounds( FBoxSphereBounds( BoundingBox ) );
-	SkeletalMesh->bHasVertexColors = bHasVertexColors;
-	SkeletalMesh->VertexColorGuid = SkeletalMesh->bHasVertexColors ? FGuid::NewGuid() : FGuid();
+	SkeletalMesh->SetHasVertexColors(bHasVertexColors);
+	SkeletalMesh->SetVertexColorGuid(SkeletalMesh->GetHasVertexColors() ? FGuid::NewGuid() : FGuid());
 	SkeletalMesh->CalculateInvRefMatrices();
 
 	// Generate a Skeleton and associate it to the SkeletalMesh
 	USkeleton* Skeleton = NewObject<USkeleton>(GetTransientPackage(), NAME_None, ObjectFlags | EObjectFlags::RF_Public );
 	Skeleton->MergeAllBonesToBoneTree(SkeletalMesh);
 	Skeleton->SetPreviewMesh(SkeletalMesh);
-	SkeletalMesh->Skeleton = Skeleton;
+	SkeletalMesh->SetSkeleton(Skeleton);
 
 	UsdToUnrealImpl::CreateMorphTargets(InBlendShapesByPath, LODIndexToSkeletalMeshImportData, SkeletalMesh);
 
@@ -2310,7 +2310,7 @@ bool UnrealToUsd::ConvertSkeleton( const USkeleton* Skeleton, pxr::UsdSkelSkelet
 bool UnrealToUsd::ConvertSkeletalMesh( const USkeletalMesh* SkeletalMesh, pxr::UsdPrim& SkelRootPrim, const pxr::UsdTimeCode TimeCode )
 {
 	pxr::UsdSkelRoot SkelRoot{ SkelRootPrim };
-	if ( !SkeletalMesh || !SkeletalMesh->Skeleton || !SkelRoot )
+	if ( !SkeletalMesh || !SkeletalMesh->GetSkeleton() || !SkelRoot )
 	{
 		return false;
 	}
@@ -2346,7 +2346,7 @@ bool UnrealToUsd::ConvertSkeletalMesh( const USkeletalMesh* SkeletalMesh, pxr::U
 	// Collect all material assignments, referenced by the sections' material indices
 	bool bHasMaterialAssignments = false;
 	pxr::VtArray< std::string > MaterialAssignments;
-	for ( const FSkeletalMaterial& SkeletalMaterial : SkeletalMesh->Materials )
+	for ( const FSkeletalMaterial& SkeletalMaterial : SkeletalMesh->GetMaterials())
 	{
 		FString AssignedMaterialPathName;
 		if ( UMaterialInterface* Material = SkeletalMaterial.MaterialInterface )
@@ -2378,7 +2378,7 @@ bool UnrealToUsd::ConvertSkeletalMesh( const USkeletalMesh* SkeletalMesh, pxr::U
 		pxr::UsdRelationship SkelRel = SkelBindingAPI.CreateSkeletonRel();
 		SkelRel.SetTargets({SkeletonPrim.GetPath()});
 
-		UnrealToUsd::ConvertSkeleton( SkeletalMesh->Skeleton, SkelSkeleton );
+		UnrealToUsd::ConvertSkeleton( SkeletalMesh->GetSkeleton(), SkelSkeleton );
 	}
 
 	// Actual meshes
@@ -2423,17 +2423,17 @@ bool UnrealToUsd::ConvertSkeletalMesh( const USkeletalMesh* SkeletalMesh, pxr::U
 			LODMaterialMap = LODInfo->LODMaterialMap;
 		}
 
-		UnrealToUsdImpl::ConvertSkeletalMeshLOD( LODModel, UsdLODPrimGeomMesh, SkeletalMesh->bHasVertexColors, MaterialAssignments, LODMaterialMap, TimeCode );
+		UnrealToUsdImpl::ConvertSkeletalMeshLOD( LODModel, UsdLODPrimGeomMesh, SkeletalMesh->GetHasVertexColors(), MaterialAssignments, LODMaterialMap, TimeCode );
 
 		// Relationships can't target prims inside variants, so if we have BlendShapes to export we have to disable the edit target
 		// so that the blend shapes end up outside the variants and the Meshes can have their blendShapeTargets relationships pointing at them
-		if ( bExportMultipleLODs && SkeletalMesh->MorphTargets.Num() > 0)
+		if ( bExportMultipleLODs && SkeletalMesh->GetMorphTargets().Num() > 0)
 		{
 			EditContext.Reset();
 		}
 
 		pxr::SdfPathVector AddedBlendShapeTargets;
-		for ( UMorphTarget* MorphTarget : SkeletalMesh->MorphTargets )
+		for ( UMorphTarget* MorphTarget : SkeletalMesh->GetMorphTargets())
 		{
 			if ( !MorphTarget || !MorphTarget->HasDataForLOD( LODIndex ) )
 			{
diff --git a/Engine/Plugins/MeshPainting/Source/MeshPaintEditorMode/Private/MeshPaintModeHelpers.cpp b/Engine/Plugins/MeshPainting/Source/MeshPaintEditorMode/Private/MeshPaintModeHelpers.cpp
--- a/Engine/Plugins/MeshPainting/Source/MeshPaintEditorMode/Private/MeshPaintModeHelpers.cpp
+++ b/Engine/Plugins/MeshPainting/Source/MeshPaintEditorMode/Private/MeshPaintModeHelpers.cpp
@@ -302,9 +302,9 @@ void UMeshPaintModeHelpers::ImportVertexColorsToSkeletalMesh(USkeletalMesh* Skel
 	}
 
 	//Make sure we change the import data so the re-import do not replace the new data
-	if (SkeletalMesh->AssetImportData)
+	if (SkeletalMesh->GetAssetImportData())
 	{
-		UFbxSkeletalMeshImportData* ImportData = Cast<UFbxSkeletalMeshImportData>(SkeletalMesh->AssetImportData);
+		UFbxSkeletalMeshImportData* ImportData = Cast<UFbxSkeletalMeshImportData>(SkeletalMesh->GetAssetImportData());
 		if (ImportData && ImportData->VertexColorImportOption != EVertexColorImportOption::Ignore)
 		{
 			ImportData->SetFlags(RF_Transactional);
diff --git a/Engine/Plugins/MeshPainting/Source/MeshPaintingToolset/Private/MeshPaintHelpers.cpp b/Engine/Plugins/MeshPainting/Source/MeshPaintingToolset/Private/MeshPaintHelpers.cpp
--- a/Engine/Plugins/MeshPainting/Source/MeshPaintingToolset/Private/MeshPaintHelpers.cpp
+++ b/Engine/Plugins/MeshPainting/Source/MeshPaintingToolset/Private/MeshPaintHelpers.cpp
@@ -781,8 +781,8 @@ void UMeshPaintingToolset::FillSkeletalMeshVertexColors(USkeletalMeshComponent*
 		// Dirty the mesh
 		Mesh->SetFlags(RF_Transactional);
 		Mesh->Modify();
-		Mesh->bHasVertexColors = true;
-		Mesh->VertexColorGuid = FGuid::NewGuid();
+		Mesh->SetHasVertexColors(true);
+		Mesh->SetVertexColorGuid(FGuid::NewGuid());
 
 		// Release the static mesh's resources.
 		Mesh->ReleaseResources();
diff --git a/Engine/Plugins/MeshPainting/Source/MeshPaintingToolset/Private/MeshPaintSkeletalMeshAdapter.cpp b/Engine/Plugins/MeshPainting/Source/MeshPaintingToolset/Private/MeshPaintSkeletalMeshAdapter.cpp
--- a/Engine/Plugins/MeshPainting/Source/MeshPaintingToolset/Private/MeshPaintSkeletalMeshAdapter.cpp
+++ b/Engine/Plugins/MeshPainting/Source/MeshPaintingToolset/Private/MeshPaintSkeletalMeshAdapter.cpp
@@ -334,8 +334,8 @@ void FMeshPaintSkeletalMeshComponentAdapter::PreEdit()
 	ReferencedSkeletalMesh->SetFlags(RF_Transactional);
 	ReferencedSkeletalMesh->Modify();
 
-	ReferencedSkeletalMesh->bHasVertexColors = true;
-	ReferencedSkeletalMesh->VertexColorGuid = FGuid::NewGuid();
+	ReferencedSkeletalMesh->SetHasVertexColors(true);
+	ReferencedSkeletalMesh->SetVertexColorGuid(FGuid::NewGuid());
 
 	// Release the static mesh's resources.
 	ReferencedSkeletalMesh->ReleaseResources();
@@ -348,14 +348,14 @@ void FMeshPaintSkeletalMeshComponentAdapter::PreEdit()
 	{
 		// Mesh doesn't have a color vertex buffer yet!  We'll create one now.
 		LODData->StaticVertexBuffers.ColorVertexBuffer.InitFromSingleColor(FColor(255, 255, 255, 255), LODData->GetNumVertices());
-		ReferencedSkeletalMesh->bHasVertexColors = true;
-		ReferencedSkeletalMesh->VertexColorGuid = FGuid::NewGuid();
+		ReferencedSkeletalMesh->SetHasVertexColors(true);
+		ReferencedSkeletalMesh->SetVertexColorGuid(FGuid::NewGuid());
 		BeginInitResource(&LODData->StaticVertexBuffers.ColorVertexBuffer);
 	}
 	//Make sure we change the import data so the re-import do not replace the new data
-	if (ReferencedSkeletalMesh->AssetImportData)
+	if (ReferencedSkeletalMesh->GetAssetImportData())
 	{
-		UFbxSkeletalMeshImportData* ImportData = Cast<UFbxSkeletalMeshImportData>(ReferencedSkeletalMesh->AssetImportData);
+		UFbxSkeletalMeshImportData* ImportData = Cast<UFbxSkeletalMeshImportData>(ReferencedSkeletalMesh->GetAssetImportData());
 		if (ImportData && ImportData->VertexColorImportOption != EVertexColorImportOption::Ignore)
 		{
 			ImportData->SetFlags(RF_Transactional);
diff --git a/Engine/Plugins/MovieScene/SequencerScripting/Source/SequencerScriptingEditor/Private/SequencerTools.cpp b/Engine/Plugins/MovieScene/SequencerScripting/Source/SequencerScriptingEditor/Private/SequencerTools.cpp
--- a/Engine/Plugins/MovieScene/SequencerScripting/Source/SequencerScriptingEditor/Private/SequencerTools.cpp
+++ b/Engine/Plugins/MovieScene/SequencerScripting/Source/SequencerScriptingEditor/Private/SequencerTools.cpp
@@ -282,9 +282,9 @@ bool USequencerToolsFunctionLibrary::ExportAnimSequence(UWorld* World, ULevelSeq
 		}
  
 		USkeletalMeshComponent* SkeletalMeshComp =  GetSkelMeshComponent(Player, Binding);
-		if (SkeletalMeshComp && SkeletalMeshComp->SkeletalMesh && SkeletalMeshComp->SkeletalMesh->Skeleton)
+		if (SkeletalMeshComp && SkeletalMeshComp->SkeletalMesh && SkeletalMeshComp->SkeletalMesh->GetSkeleton())
 		{
-			AnimSequence->SetSkeleton(SkeletalMeshComp->SkeletalMesh->Skeleton);
+			AnimSequence->SetSkeleton(SkeletalMeshComp->SkeletalMesh->GetSkeleton());
 			bResult = MovieSceneToolHelpers::ExportToAnimSequence(AnimSequence,ExportOptions, MovieScene, Player, SkeletalMeshComp, Template, RootToLocalTransform);
 		}
 	}
diff --git a/Engine/Plugins/Runtime/ApexDestruction/Source/ApexDestruction/Private/DestructibleComponent.cpp b/Engine/Plugins/Runtime/ApexDestruction/Source/ApexDestruction/Private/DestructibleComponent.cpp
--- a/Engine/Plugins/Runtime/ApexDestruction/Source/ApexDestruction/Private/DestructibleComponent.cpp
+++ b/Engine/Plugins/Runtime/ApexDestruction/Source/ApexDestruction/Private/DestructibleComponent.cpp
@@ -457,7 +457,8 @@ UBodySetup* UDestructibleComponent::GetBodySetup()
 
 		if (TheDestructibleMesh != NULL)
 		{
-			return TheDestructibleMesh->BodySetup;
+			const UDestructibleMesh* TheDestructibleMeshConst = TheDestructibleMesh;
+			return TheDestructibleMeshConst->GetBodySetup();
 		}
 	}
 
@@ -1640,10 +1641,10 @@ void UDestructibleComponent::SetMaterial(int32 ElementIndex, UMaterialInterface*
 	// Update physical properties for individual bone instances as well
 	if (SkeletalMesh)
 	{
-		int32 NumBones = SkeletalMesh->RefSkeleton.GetRawBoneNum();
+		int32 NumBones = SkeletalMesh->GetRefSkeleton().GetRawBoneNum();
 		for (int32 BoneIdx = 0; BoneIdx < NumBones; ++BoneIdx)
 		{
-			FName BoneName = SkeletalMesh->RefSkeleton.GetBoneName(BoneIdx);
+			FName BoneName = SkeletalMesh->GetRefSkeleton().GetBoneName(BoneIdx);
 			FBodyInstance* Instance = GetBodyInstance(BoneName);
 			if (Instance)
 			{
diff --git a/Engine/Plugins/Runtime/HairStrands/Source/HairStrandsCore/Private/Niagara/NiagaraDataInterfacePhysicsAsset.cpp b/Engine/Plugins/Runtime/HairStrands/Source/HairStrandsCore/Private/Niagara/NiagaraDataInterfacePhysicsAsset.cpp
--- a/Engine/Plugins/Runtime/HairStrands/Source/HairStrandsCore/Private/Niagara/NiagaraDataInterfacePhysicsAsset.cpp
+++ b/Engine/Plugins/Runtime/HairStrands/Source/HairStrandsCore/Private/Niagara/NiagaraDataInterfacePhysicsAsset.cpp
@@ -148,26 +148,22 @@ void CreateInternalArrays(const TArray<TWeakObjectPtr<UPhysicsAsset>>& PhysicsAs
 					continue;
 				}
 
-				const FReferenceSkeleton* RefSkeleton = &SkelMesh->RefSkeleton;
-				if (RefSkeleton != nullptr)
+				const FReferenceSkeleton& RefSkeleton = SkelMesh->GetRefSkeleton();
+				if (RefSkeleton.GetNum() > 0)
 				{
-					if (RefSkeleton->GetNum() > 0)
+					for (const UBodySetup* BodySetup : PhysicsAsset->SkeletalBodySetups)
 					{
-
-						for (const UBodySetup* BodySetup : PhysicsAsset->SkeletalBodySetups)
+						const FName BoneName = BodySetup->BoneName;
+						const int32 BoneIndex = RefSkeleton.FindBoneIndex(BoneName);
+						if (BoneIndex != INDEX_NONE && BoneIndex < RefSkeleton.GetNum())
 						{
-							const FName BoneName = BodySetup->BoneName;
-							const int32 BoneIndex = RefSkeleton->FindBoneIndex(BoneName);
-							if (BoneIndex != INDEX_NONE && BoneIndex < RefSkeleton->GetNum())
-							{
-								NumBoxes += BodySetup->AggGeom.BoxElems.Num();
-								NumSpheres += BodySetup->AggGeom.SphereElems.Num();
-								NumCapsules += BodySetup->AggGeom.SphylElems.Num();
-							}
+							NumBoxes += BodySetup->AggGeom.BoxElems.Num();
+							NumSpheres += BodySetup->AggGeom.SphereElems.Num();
+							NumCapsules += BodySetup->AggGeom.SphylElems.Num();
 						}
-						//UE_LOG(LogPhysicsAsset, Warning, TEXT("PhysicsAsset = %s | SkeletalMesh = %d | Num Capsules = %d | Num Spheres = %d | Num Boxes = %d"), *PhysicsAsset->GetName(), SkeletalMeshs[ComponentIndex].Get(), 
-						//				NumCapsules, NumSpheres, NumBoxes);
 					}
+					//UE_LOG(LogPhysicsAsset, Warning, TEXT("PhysicsAsset = %s | SkeletalMesh = %d | Num Capsules = %d | Num Spheres = %d | Num Boxes = %d"), *PhysicsAsset->GetName(), SkeletalMeshs[ComponentIndex].Get(), 
+					//				NumCapsules, NumSpheres, NumBoxes);
 				}
 			}
 		}
@@ -199,7 +195,7 @@ void CreateInternalArrays(const TArray<TWeakObjectPtr<UPhysicsAsset>>& PhysicsAs
 					{
 						continue;
 					}
-					const FReferenceSkeleton* RefSkeleton = &SkelMesh->RefSkeleton;
+					const FReferenceSkeleton* RefSkeleton = &SkelMesh->GetRefSkeleton();
 					if (RefSkeleton != nullptr)
 					{
 						TArray<FTransform> RestTransforms;
@@ -290,7 +286,7 @@ void UpdateInternalArrays(const TArray<TWeakObjectPtr<UPhysicsAsset>>& PhysicsAs
 				{
 					continue;
 				}
-				const FReferenceSkeleton* RefSkeleton = &SkelMesh->RefSkeleton;
+				const FReferenceSkeleton* RefSkeleton = &SkelMesh->GetRefSkeleton();
 
 				if (RefSkeleton != nullptr)
 				{
diff --git a/Engine/Plugins/Runtime/Oculus/OculusVR/Source/OculusInput/Private/OculusHandComponent.cpp b/Engine/Plugins/Runtime/Oculus/OculusVR/Source/OculusInput/Private/OculusHandComponent.cpp
--- a/Engine/Plugins/Runtime/Oculus/OculusVR/Source/OculusInput/Private/OculusHandComponent.cpp
+++ b/Engine/Plugins/Runtime/Oculus/OculusVR/Source/OculusInput/Private/OculusHandComponent.cpp
@@ -161,7 +161,7 @@ void UOculusHandComponent::UpdateBonePose()
 			else
 			{
 				// Set Remaing Bone Rotations
-				int32 BoneIndex = SkeletalMesh->RefSkeleton.FindBoneIndex(BoneElem.Value);
+				int32 BoneIndex = SkeletalMesh->GetRefSkeleton().FindBoneIndex(BoneElem.Value);
 				if (BoneIndex >= 0)
 				{
 					FQuat BoneRotation = UOculusInputFunctionLibrary::GetBoneRotation(SkeletonType, (EBone)BoneElem.Key);
@@ -179,7 +179,7 @@ void UOculusHandComponent::UpdateBonePose()
 		BoneSpaceTransforms[0].SetRotation(RootBoneRotation);
 
 		// Set Remaining Bone Rotations
-		for (uint32 BoneIndex = 1; BoneIndex < (uint32)SkeletalMesh->RefSkeleton.GetNum(); BoneIndex++)
+		for (uint32 BoneIndex = 1; BoneIndex < (uint32)SkeletalMesh->GetRefSkeleton().GetNum(); BoneIndex++)
 		{
 			FQuat BoneRotation = UOculusInputFunctionLibrary::GetBoneRotation(SkeletonType, (EBone)BoneIndex);
 			BoneSpaceTransforms[BoneIndex].SetRotation(BoneRotation);
diff --git a/Engine/Plugins/Runtime/Oculus/OculusVR/Source/OculusInput/Private/OculusHandTracking.cpp b/Engine/Plugins/Runtime/Oculus/OculusVR/Source/OculusInput/Private/OculusHandTracking.cpp
--- a/Engine/Plugins/Runtime/Oculus/OculusVR/Source/OculusInput/Private/OculusHandTracking.cpp
+++ b/Engine/Plugins/Runtime/Oculus/OculusVR/Source/OculusInput/Private/OculusHandTracking.cpp
@@ -224,13 +224,13 @@ bool FOculusHandTracking::GetHandSkeletalMesh(USkeletalMesh* HandSkeletalMesh, c
 		// Add default material as backup
 		LodInfo.LODMaterialMap.Add(0);
 		UMaterialInterface* DefaultMaterial = UMaterial::GetDefaultMaterial(MD_Surface);
-		HandSkeletalMesh->Materials.Add(DefaultMaterial);
-		HandSkeletalMesh->Materials[0].UVChannelData.bInitialized = true;
+		HandSkeletalMesh->GetMaterials().Add(DefaultMaterial);
+		HandSkeletalMesh->GetMaterials()[0].UVChannelData.bInitialized = true;
 
 		// Set skeletal mesh properties
-		HandSkeletalMesh->bHasVertexColors = true;
-		HandSkeletalMesh->bHasBeenSimplified = false;
-		HandSkeletalMesh->bEnablePerPolyCollision = false;
+		HandSkeletalMesh->SetHasVertexColors(true);
+		HandSkeletalMesh->SetHasBeenSimplified(false);
+		HandSkeletalMesh->SetEnablePerPolyCollision(false);
 
 		InitializeHandMesh(HandSkeletalMesh, OvrMesh, WorldToMeters);
 
@@ -240,8 +240,8 @@ bool FOculusHandTracking::GetHandSkeletalMesh(USkeletalMesh* HandSkeletalMesh, c
 #endif
 
 		// Create Skeleton object and merge all bones
-		HandSkeletalMesh->Skeleton = NewObject<USkeleton>();
-		HandSkeletalMesh->Skeleton->MergeAllBonesToBoneTree(HandSkeletalMesh);
+		HandSkeletalMesh->SetSkeleton(NewObject<USkeleton>());
+		HandSkeletalMesh->GetSkeleton()->MergeAllBonesToBoneTree(HandSkeletalMesh);
 		HandSkeletalMesh->PostLoad();
 
 		delete OvrMesh;
@@ -313,7 +313,7 @@ void FOculusHandTracking::InitializeHandMesh(USkeletalMesh* SkeletalMesh, const
 	}
 
 	// Update bone map
-	for (uint32 BoneIndex = 0; BoneIndex < (uint32)SkeletalMesh->RefSkeleton.GetNum(); BoneIndex++)
+	for (uint32 BoneIndex = 0; BoneIndex < (uint32)SkeletalMesh->GetRefSkeleton().GetNum(); BoneIndex++)
 	{
 		MeshSection.BoneMap.Add(BoneIndex);
 	}
@@ -411,7 +411,7 @@ void FOculusHandTracking::InitializeHandMesh(USkeletalMesh* SkeletalMesh, const
 	}
 
 	// Update bone map for mesh section
-	for (uint32 BoneIndex = 0; BoneIndex < (uint32)SkeletalMesh->RefSkeleton.GetNum(); BoneIndex++)
+	for (uint32 BoneIndex = 0; BoneIndex < (uint32)SkeletalMesh->GetRefSkeleton().GetNum(); BoneIndex++)
 	{
 		MeshSection.BoneMap.Add(BoneIndex);
 	}
@@ -439,15 +439,15 @@ void FOculusHandTracking::InitializeHandMesh(USkeletalMesh* SkeletalMesh, const
 
 void FOculusHandTracking::InitializeHandSkeleton(USkeletalMesh* SkeletalMesh, const ovrpSkeleton* OvrSkeleton, const float WorldToMeters)
 {
-	SkeletalMesh->RefSkeleton.Empty(OvrSkeleton->NumBones);
+	SkeletalMesh->GetRefSkeleton().Empty(OvrSkeleton->NumBones);
 
 #if WITH_EDITOR
 	FSkeletalMeshLODModel* LodRenderData = &SkeletalMesh->GetImportedModel()->LODModels[0];
 #else
 	FSkeletalMeshLODRenderData* LodRenderData = &SkeletalMesh->GetResourceForRendering()->LODRenderData[0];
 #endif
-	SkeletalMesh->bHasBeenSimplified = false;
-	SkeletalMesh->bHasVertexColors = true;
+	SkeletalMesh->SetHasBeenSimplified(false);
+	SkeletalMesh->SetHasVertexColors(true);
 
 	for (uint32 BoneIndex = 0; BoneIndex < OvrSkeleton->NumBones; BoneIndex++)
 	{
@@ -463,7 +463,7 @@ void FOculusHandTracking::InitializeHandSkeleton(USkeletalMesh* SkeletalMesh, co
 		Transform.SetLocation(BonePosition);
 		Transform.SetRotation(BoneRotation);
 
-		FReferenceSkeletonModifier Modifier = FReferenceSkeletonModifier(SkeletalMesh->RefSkeleton, nullptr);
+		FReferenceSkeletonModifier Modifier = FReferenceSkeletonModifier(SkeletalMesh->GetRefSkeleton(), nullptr);
 		int32 ParentIndex = -1;
 		if (BoneIndex > 0)
 		{
@@ -510,7 +510,7 @@ TArray<FOculusCapsuleCollider> FOculusHandTracking::InitializeHandPhysics(const
 		FVector CapsulePointOne = OvrBoneVectorToFVector(OvrBoneCapsule.Points[1], WorldToMeters);
 		FVector Delta = (CapsulePointOne - CapsulePointZero);
 
-		FName BoneName = HandComponent->SkeletalMesh->RefSkeleton.GetBoneName(OvrBoneCapsule.BoneIndex);
+		FName BoneName = HandComponent->SkeletalMesh->GetRefSkeleton().GetBoneName(OvrBoneCapsule.BoneIndex);
 
 		float CapsuleHeight = Delta.Size();
 		float CapsuleRadius = OvrBoneCapsule.Radius * WorldToMeters;
diff --git a/Engine/Plugins/Tests/EditorTests/Source/EditorTests/Private/UnrealEd/EditorBuildPromotionTests.cpp b/Engine/Plugins/Tests/EditorTests/Source/EditorTests/Private/UnrealEd/EditorBuildPromotionTests.cpp
--- a/Engine/Plugins/Tests/EditorTests/Source/EditorTests/Private/UnrealEd/EditorBuildPromotionTests.cpp
+++ b/Engine/Plugins/Tests/EditorTests/Source/EditorTests/Private/UnrealEd/EditorBuildPromotionTests.cpp
@@ -1272,7 +1272,7 @@ namespace BuildPromotionTestHelper
 			if (SkeletalMesh_Test)
 			{
 				//Animation
-				IMPORT_ASSET_WITH_FACTORY(UFbxFactory, UAnimSequence, Animation, AnimationTest, FactoryInst->ImportUI->MeshTypeToImport = FBXIT_Animation; FactoryInst->ImportUI->Skeleton = SkeletalMesh_Test->Skeleton;);
+				IMPORT_ASSET_WITH_FACTORY(UFbxFactory, UAnimSequence, Animation, AnimationTest, FactoryInst->ImportUI->MeshTypeToImport = FBXIT_Animation; FactoryInst->ImportUI->Skeleton = SkeletalMesh_Test->GetSkeleton(););
 			}
 			else
 			{
diff --git a/Engine/Plugins/Tests/EditorTests/Source/EditorTests/Private/UnrealEd/FbxAutomationTests.cpp b/Engine/Plugins/Tests/EditorTests/Source/EditorTests/Private/UnrealEd/FbxAutomationTests.cpp
--- a/Engine/Plugins/Tests/EditorTests/Source/EditorTests/Private/UnrealEd/FbxAutomationTests.cpp
+++ b/Engine/Plugins/Tests/EditorTests/Source/EditorTests/Private/UnrealEd/FbxAutomationTests.cpp
@@ -318,7 +318,7 @@ bool FFbxImportAssetsAutomationTest::RunTest(const FString& Parameters)
 					FbxSkeletalMeshReimportFactory->ImportUI = TestPlan->ImportUI;
 
 					USkeletalMesh *ReimportSkeletalMesh = Cast<USkeletalMesh>(GlobalImportedObjects[0]);
-					UFbxSkeletalMeshImportData* ImportData = Cast<UFbxSkeletalMeshImportData>(ReimportSkeletalMesh->AssetImportData);
+					UFbxSkeletalMeshImportData* ImportData = Cast<UFbxSkeletalMeshImportData>(ReimportSkeletalMesh->GetAssetImportData());
 
 					//Copy UFbxSkeletalMeshImportData
 					ImportData->ImportContentType = TestPlan->ImportUI->SkeletalMeshImportData->ImportContentType;
@@ -415,7 +415,7 @@ bool FFbxImportAssetsAutomationTest::RunTest(const FString& Parameters)
 				else if (GlobalImportedObjects[0]->IsA(USkeletalMesh::StaticClass()))
 				{
 					USkeletalMesh *ExistingSkeletalMesh = Cast<USkeletalMesh>(GlobalImportedObjects[0]);
-					UFbxSkeletalMeshImportData* ImportData = Cast<UFbxSkeletalMeshImportData>(ExistingSkeletalMesh->AssetImportData);
+					UFbxSkeletalMeshImportData* ImportData = Cast<UFbxSkeletalMeshImportData>(ExistingSkeletalMesh->GetAssetImportData());
 
 					//Copy UFbxSkeletalMeshImportData
 					ImportData->ImportContentType = TestPlan->ImportUI->SkeletalMeshImportData->ImportContentType;
@@ -603,14 +603,14 @@ bool FFbxImportAssetsAutomationTest::RunTest(const FString& Parameters)
 					else if (Object->IsA(USkeletalMesh::StaticClass()))
 					{
 						USkeletalMesh *Mesh = Cast<USkeletalMesh>(Object);
-						if (Mesh->Materials.Num() <= MaterialSlotIndex)
+						if (Mesh->GetMaterials().Num() <= MaterialSlotIndex)
 						{
 							BadSlotIndex = true;
-							MeshMaterialNumber = Mesh->Materials.Num();
+							MeshMaterialNumber = Mesh->GetMaterials().Num();
 						}
 						else
 						{
-							MaterialImportedName = Mesh->Materials[MaterialSlotIndex].ImportedMaterialSlotName.ToString();
+							MaterialImportedName = Mesh->GetMaterials()[MaterialSlotIndex].ImportedMaterialSlotName.ToString();
 						}
 					}
 				}
@@ -750,7 +750,7 @@ bool FFbxImportAssetsAutomationTest::RunTest(const FString& Parameters)
 					else if (Object->IsA(USkeletalMesh::StaticClass()))
 					{
 						USkeletalMesh *Mesh = Cast<USkeletalMesh>(Object);
-						MaterialIndexNumber = Mesh->Materials.Num();
+						MaterialIndexNumber = Mesh->GetMaterials().Num();
 					}
 				}
 				if (MaterialIndexNumber != ExpectedResult.ExpectedPresetsDataInteger[0])
@@ -1034,9 +1034,9 @@ bool FFbxImportAssetsAutomationTest::RunTest(const FString& Parameters)
 							else
 							{
 								int32 MaterialIndex = Mesh->GetResourceForRendering()->LODRenderData[LODIndex].RenderSections[SectionIndex].MaterialIndex;
-								if (MaterialIndex >= 0 && MaterialIndex < Mesh->Materials.Num())
+								if (MaterialIndex >= 0 && MaterialIndex < Mesh->GetMaterials().Num())
 								{
-									MaterialName = Mesh->Materials[MaterialIndex].MaterialInterface->GetName();
+									MaterialName = Mesh->GetMaterials()[MaterialIndex].MaterialInterface->GetName();
 								}
 							}
 						}
@@ -1200,9 +1200,9 @@ bool FFbxImportAssetsAutomationTest::RunTest(const FString& Parameters)
 							else
 							{
 								int32 MaterialIndex = Mesh->GetResourceForRendering()->LODRenderData[LODIndex].RenderSections[SectionIndex].MaterialIndex;
-								if (MaterialIndex >= 0 && MaterialIndex < Mesh->Materials.Num())
+								if (MaterialIndex >= 0 && MaterialIndex < Mesh->GetMaterials().Num())
 								{
-									MaterialName = Mesh->Materials[MaterialIndex].ImportedMaterialSlotName.ToString();
+									MaterialName = Mesh->GetMaterials()[MaterialIndex].ImportedMaterialSlotName.ToString();
 								}
 							}
 						}
@@ -1468,9 +1468,9 @@ bool FFbxImportAssetsAutomationTest::RunTest(const FString& Parameters)
 					if (Object->IsA(USkeletalMesh::StaticClass()))
 					{
 						USkeletalMesh *Mesh = Cast<USkeletalMesh>(Object);
-						if (Mesh->Skeleton != nullptr)
+						if (Mesh->GetSkeleton() != nullptr)
 						{
-							BoneNumber = Mesh->Skeleton->GetReferenceSkeleton().GetNum();
+							BoneNumber = Mesh->GetSkeleton()->GetReferenceSkeleton().GetNum();
 						}
 					}
 				}
@@ -1505,11 +1505,11 @@ bool FFbxImportAssetsAutomationTest::RunTest(const FString& Parameters)
 						FoundSkeletalMesh = true;
 						USkeletalMesh *Mesh = Cast<USkeletalMesh>(Object);
 						FoundSkeleton = true;
-						BoneNumber = Mesh->RefSkeleton.GetNum();
+						BoneNumber = Mesh->GetRefSkeleton().GetNum();
 						if (BoneIndex >= 0 && BoneIndex < BoneNumber)
 						{
 							FoundBoneIndex = true;
-							BoneIndexPosition = Mesh->RefSkeleton.GetRefBonePose()[BoneIndex].GetLocation();
+							BoneIndexPosition = Mesh->GetRefSkeleton().GetRefBonePose()[BoneIndex].GetLocation();
 						}
 					}
 				}
diff --git a/Engine/Plugins/VirtualProduction/Takes/Source/TakeTrackRecorders/Private/MovieScene3DTransformTrackRecorder.cpp b/Engine/Plugins/VirtualProduction/Takes/Source/TakeTrackRecorders/Private/MovieScene3DTransformTrackRecorder.cpp
--- a/Engine/Plugins/VirtualProduction/Takes/Source/TakeTrackRecorders/Private/MovieScene3DTransformTrackRecorder.cpp
+++ b/Engine/Plugins/VirtualProduction/Takes/Source/TakeTrackRecorders/Private/MovieScene3DTransformTrackRecorder.cpp
@@ -502,7 +502,7 @@ void UMovieScene3DTransformTrackRecorder::PostProcessAnimationData(UMovieSceneAn
 					if (BoneTreeIndex != INDEX_NONE)
 					{
 						int32 BoneIndex = AnimSkeleton->GetMeshBoneIndexFromSkeletonBoneIndex(SkeletalMesh, BoneTreeIndex);
-						int32 ParentIndex = SkeletalMesh->RefSkeleton.GetParentIndex(BoneIndex);
+						int32 ParentIndex = SkeletalMesh->GetRefSkeleton().GetParentIndex(BoneIndex);
 						if (ParentIndex == INDEX_NONE)
 						{
 							// We've found the root (root bones do not have a valid parent)
diff --git a/Engine/Plugins/VirtualProduction/Takes/Source/TakeTrackRecorders/Private/MovieSceneAnimationTrackRecorder.cpp b/Engine/Plugins/VirtualProduction/Takes/Source/TakeTrackRecorders/Private/MovieSceneAnimationTrackRecorder.cpp
--- a/Engine/Plugins/VirtualProduction/Takes/Source/TakeTrackRecorders/Private/MovieSceneAnimationTrackRecorder.cpp
+++ b/Engine/Plugins/VirtualProduction/Takes/Source/TakeTrackRecorders/Private/MovieSceneAnimationTrackRecorder.cpp
@@ -66,7 +66,7 @@ void UMovieSceneAnimationTrackRecorder::CreateAnimationAssetAndSequence(const AA
 			FAssetRegistryModule::AssetCreated(AnimSequence.Get());
 
 			// Assign the skeleton we're recording to the newly created Animation Sequence.
-			AnimSequence->SetSkeleton(SkeletalMeshComponent->SkeletalMesh->Skeleton);
+			AnimSequence->SetSkeleton(SkeletalMeshComponent->SkeletalMesh->GetSkeleton());
 		}
 	}
 
diff --git a/Engine/Source/Developer/AssetTools/Private/AssetTypeActions/AssetTypeActions_SkeletalMesh.cpp b/Engine/Source/Developer/AssetTools/Private/AssetTypeActions/AssetTypeActions_SkeletalMesh.cpp
--- a/Engine/Source/Developer/AssetTools/Private/AssetTypeActions/AssetTypeActions_SkeletalMesh.cpp
+++ b/Engine/Source/Developer/AssetTools/Private/AssetTypeActions/AssetTypeActions_SkeletalMesh.cpp
@@ -373,9 +373,9 @@ FDlgMergeSkeleton::EResult FDlgMergeSkeleton::ShowModal()
 		BoneIndicesMap.Add(BoneName, BoneTreeId);
 	}
 
-	for ( int32 RefBoneId=0 ; RefBoneId< Mesh->RefSkeleton.GetRawBoneNum() ; ++RefBoneId )
+	for ( int32 RefBoneId=0 ; RefBoneId< Mesh->GetRefSkeleton().GetRawBoneNum() ; ++RefBoneId )
 	{
-		const FName& BoneName = Mesh->RefSkeleton.GetBoneName(RefBoneId);
+		const FName& BoneName = Mesh->GetRefSkeleton().GetBoneName(RefBoneId);
 		// if I can't find this from Skeleton
 		if (BoneIndicesMap.Find(BoneName)==NULL)
 		{
@@ -390,9 +390,9 @@ FDlgMergeSkeleton::EResult FDlgMergeSkeleton::ShowModal()
 	{
 		// it's all identical, but still need to return RequiredBones
 		// for the case, where they'd like to replace the one exactly same hierarchy but different skeleton 
-		for ( int32 RefBoneId= 0 ; RefBoneId< Mesh->RefSkeleton.GetRawBoneNum() ; ++RefBoneId )
+		for ( int32 RefBoneId= 0 ; RefBoneId< Mesh->GetRefSkeleton().GetRawBoneNum() ; ++RefBoneId )
 		{
-			const FName& BoneName = Mesh->RefSkeleton.GetBoneName(RefBoneId);
+			const FName& BoneName = Mesh->GetRefSkeleton().GetBoneName(RefBoneId);
 			RequiredBones.Add(RefBoneId);
 		}
 
@@ -407,14 +407,14 @@ FDlgMergeSkeleton::EResult FDlgMergeSkeleton::ShowModal()
 
 	if(UserResponse == EResult::Confirm)
 	{
-		for ( int32 RefBoneId= 0 ; RefBoneId< Mesh->RefSkeleton.GetRawBoneNum() ; ++RefBoneId )
+		for ( int32 RefBoneId= 0 ; RefBoneId< Mesh->GetRefSkeleton().GetRawBoneNum() ; ++RefBoneId )
 		{
 			if ( DialogWidget->IsBoneIncluded(RefBoneId) )
 			{
 				TArray<int32> ParentList;
 				
 				// I need to make sure parent exists first
-				int32 ParentIndex = Mesh->RefSkeleton.GetParentIndex(RefBoneId);
+				int32 ParentIndex = Mesh->GetRefSkeleton().GetParentIndex(RefBoneId);
 				
 				// make sure RequiredBones already have ParentIndex
 				while (ParentIndex >= 0 )
@@ -425,7 +425,7 @@ FDlgMergeSkeleton::EResult FDlgMergeSkeleton::ShowModal()
 						ParentList.Add(ParentIndex);
 					}
 
-					ParentIndex = Mesh->RefSkeleton.GetParentIndex(ParentIndex);
+					ParentIndex = Mesh->GetRefSkeleton().GetParentIndex(ParentIndex);
 				}
 
 				if ( ParentList.Num() > 0 )
@@ -514,9 +514,9 @@ void FAssetTypeActions_SkeletalMesh::OpenAssetEditor( const TArray<UObject*>& In
 	for (auto ObjIt = InObjects.CreateConstIterator(); ObjIt; ++ObjIt)
 	{
 		auto Mesh = Cast<USkeletalMesh>(*ObjIt);
-		if (Mesh != NULL)
+		if (Mesh != nullptr)
 		{
-			if (Mesh->Skeleton == NULL)
+			if (Mesh->GetSkeleton() == nullptr)
 			{
 				if ( FMessageDialog::Open(EAppMsgType::YesNo, LOCTEXT("MissingSkeleton", "This mesh currently has no valid Skeleton. Would you like to create a new Skeleton?")) == EAppReturnType::Yes )
 				{
@@ -538,14 +538,14 @@ void FAssetTypeActions_SkeletalMesh::OpenAssetEditor( const TArray<UObject*>& In
 					AssignSkeletonToMesh(Mesh);
 				}
 
-				if( Mesh->Skeleton == NULL )
+				if( Mesh->GetSkeleton() == NULL )
 				{
 					// error message
 					FMessageDialog::Open( EAppMsgType::Ok, LOCTEXT("CreateSkeletonOrAssign", "You need to create a Skeleton or assign one in order to open this in Persona."));
 				}
 			}
 
-			if ( Mesh->Skeleton != NULL )
+			if ( Mesh->GetSkeleton() != NULL )
 			{
 				const bool bBringToFrontIfOpen = true;
 #if WITH_EDITOR
@@ -567,11 +567,11 @@ void FAssetTypeActions_SkeletalMesh::OpenAssetEditor( const TArray<UObject*>& In
 UThumbnailInfo* FAssetTypeActions_SkeletalMesh::GetThumbnailInfo(UObject* Asset) const
 {
 	USkeletalMesh* SkeletalMesh = CastChecked<USkeletalMesh>(Asset);
-	UThumbnailInfo* ThumbnailInfo = SkeletalMesh->ThumbnailInfo;
+	UThumbnailInfo* ThumbnailInfo = SkeletalMesh->GetThumbnailInfo();
 	if ( ThumbnailInfo == NULL )
 	{
 		ThumbnailInfo = NewObject<USceneThumbnailInfo>(SkeletalMesh, NAME_None, RF_Transactional);
-		SkeletalMesh->ThumbnailInfo = ThumbnailInfo;
+		SkeletalMesh->SetThumbnailInfo(ThumbnailInfo);
 	}
 
 	return ThumbnailInfo;
@@ -603,7 +603,7 @@ EVisibility FAssetTypeActions_SkeletalMesh::GetThumbnailSkinningOverlayVisibilit
 	UObject* Obj = AssetData.GetAsset();
 	if(USkeletalMesh* SkeletalMesh = Cast<USkeletalMesh>(Obj))
 	{
-		UAssetImportData* GenericImportData = SkeletalMesh->AssetImportData;
+		UAssetImportData* GenericImportData = SkeletalMesh->GetAssetImportData();
 		if (GenericImportData != nullptr)
 		{
 			UFbxSkeletalMeshImportData* ImportData = Cast<UFbxSkeletalMeshImportData>(GenericImportData);
@@ -649,7 +649,7 @@ void FAssetTypeActions_SkeletalMesh::GetResolvedSourceFilePaths(const TArray<UOb
 	for (auto& Asset : TypeAssets)
 	{
 		const auto SkeletalMesh = CastChecked<USkeletalMesh>(Asset);
-		SkeletalMesh->AssetImportData->ExtractFilenames(OutSourceFilePaths);
+		SkeletalMesh->GetAssetImportData()->ExtractFilenames(OutSourceFilePaths);
 	}
 }
 
@@ -659,7 +659,7 @@ void FAssetTypeActions_SkeletalMesh::GetSourceFileLabels(const TArray<UObject*>&
 	{
 		const auto SkeletalMesh = CastChecked<USkeletalMesh>(Asset);
 		TArray<FString> SourceFilePaths;
-		SkeletalMesh->AssetImportData->ExtractFilenames(SourceFilePaths);
+		SkeletalMesh->GetAssetImportData()->ExtractFilenames(SourceFilePaths);
 		for (int32 SourceIndex = 0; SourceIndex < SourceFilePaths.Num(); ++SourceIndex)
 		{
 			FText SourceIndexLabel = USkeletalMesh::GetSourceFileLabelFromIndex(SourceIndex);
@@ -801,7 +801,7 @@ void FAssetTypeActions_SkeletalMesh::ExecuteFindSkeleton(TArray<TWeakObjectPtr<U
 		auto Object = (*ObjIt).Get();
 		if ( Object )
 		{
-			USkeleton* Skeleton = Object->Skeleton;
+			USkeleton* Skeleton = Object->GetSkeleton();
 			if (Skeleton)
 			{
 				ObjectsToSync.AddUnique(Skeleton);
@@ -916,9 +916,9 @@ void FAssetTypeActions_SkeletalMesh::AssignSkeletonToMesh(USkeletalMesh* SkelMes
 				bool bSuccess = SelectedSkeleton->MergeBonesToBoneTree( SkelMesh, RequiredBones );
 				if ( bSuccess )
 				{
-					if (SkelMesh->Skeleton != SelectedSkeleton)
+					if (SkelMesh->GetSkeleton() != SelectedSkeleton)
 					{
-						SkelMesh->Skeleton = SelectedSkeleton;
+						SkelMesh->SetSkeleton(SelectedSkeleton);
 						SkelMesh->MarkPackageDirty();
 					}
 					FAssetNotifications::SkeletonNeedsToBeSaved(SelectedSkeleton);
diff --git a/Engine/Source/Developer/AssetTools/Private/AssetTypeActions/AssetTypeActions_Skeleton.cpp b/Engine/Source/Developer/AssetTools/Private/AssetTypeActions/AssetTypeActions_Skeleton.cpp
--- a/Engine/Source/Developer/AssetTools/Private/AssetTypeActions/AssetTypeActions_Skeleton.cpp
+++ b/Engine/Source/Developer/AssetTools/Private/AssetTypeActions/AssetTypeActions_Skeleton.cpp
@@ -628,7 +628,7 @@ void FAssetTypeActions_Skeleton::ExecuteRetargetSkeleton(TArray<TWeakObjectPtr<U
 			}
 			else if (USkeletalMesh* SkeletalMesh = Cast<USkeletalMesh>(EditedAsset))
 			{
-				bCloseAssetEditor = OldSkeleton == SkeletalMesh->Skeleton;
+				bCloseAssetEditor = OldSkeleton == SkeletalMesh->GetSkeleton();
 			}
 			else if (UAnimBlueprint* AnimBlueprint = Cast<UAnimBlueprint>(EditedAsset))
 			{
@@ -639,7 +639,7 @@ void FAssetTypeActions_Skeleton::ExecuteRetargetSkeleton(TArray<TWeakObjectPtr<U
 				USkeletalMesh* PhysicsSkeletalMesh = PhysicsAsset->PreviewSkeletalMesh.LoadSynchronous();
 				if(PhysicsSkeletalMesh != nullptr)
 				{
-					bCloseAssetEditor = OldSkeleton == PhysicsSkeletalMesh->Skeleton;
+					bCloseAssetEditor = OldSkeleton == PhysicsSkeletalMesh->GetSkeleton();
 				}
 			}
 
@@ -975,7 +975,7 @@ void FAssetTypeActions_Skeleton::RetargetSkeleton(TArray<FAssetToRemapSkeleton>&
 	for (FObjectIterator Iter(USkeletalMeshComponent::StaticClass()); Iter; ++Iter)
 	{
 		USkeletalMeshComponent * MeshComponent = Cast<USkeletalMeshComponent>(*Iter);
-		if (MeshComponent->SkeletalMesh && MeshComponent->SkeletalMesh->Skeleton == OldSkeleton)
+		if (MeshComponent->SkeletalMesh && MeshComponent->SkeletalMesh->GetSkeleton() == OldSkeleton)
 		{
 			MeshComponent->InitAnim(true);
 		}
