Fixed string usage issues found by UE_NODISCARD

#rnx

diff --git a/Engine/Plugins/Importers/USDImporter/Source/USDImporter/Private/USDImporter.cpp b/Engine/Plugins/Importers/USDImporter/Source/USDImporter/Private/USDImporter.cpp
--- a/Engine/Plugins/Importers/USDImporter/Source/USDImporter/Private/USDImporter.cpp
+++ b/Engine/Plugins/Importers/USDImporter/Source/USDImporter/Private/USDImporter.cpp
@@ -219,7 +219,7 @@ TArray<UObject*> UDEPRECATED_UUSDImporter::ImportMeshes(FUsdImportContext& Impor
 						ExistingNamesToCount.Add(MeshName, 1);
 					}
 
-					FinalPackagePathName / MeshName;
+					FinalPackagePathName /= MeshName;
 				}
 			}
 
