Edigrate cl 15418432
Update USD Linux plugin dynamic libraries relative paths
- These paths are not valid at runtime but there's no good way to link on a .so in editor that would also be loaded by the USD plugin system from the TargetOutputDir.

diff --git a/Engine/Plugins/Importers/USDImporter/Resources/UsdResources/Linux/plugins/ar/resources/plugInfo.json b/Engine/Plugins/Importers/USDImporter/Resources/UsdResources/Linux/plugins/ar/resources/plugInfo.json
--- a/Engine/Plugins/Importers/USDImporter/Resources/UsdResources/Linux/plugins/ar/resources/plugInfo.json
+++ b/Engine/Plugins/Importers/USDImporter/Resources/UsdResources/Linux/plugins/ar/resources/plugInfo.json
@@ -7,7 +7,7 @@
                     "ArPackageResolver": {}
                 }
             },
-            "LibraryPath": "../../../../../Binaries/Linux/x86_64-unknown-linux-gnu/libar.so", 
+            "LibraryPath": "../../../../../Source/ThirdParty/Linux/bin/x86_64-unknown-linux-gnu/libar.so", 
             "Name": "ar", 
             "ResourcePath": "resources", 
             "Root": "..", 
diff --git a/Engine/Plugins/Importers/USDImporter/Resources/UsdResources/Linux/plugins/ndr/resources/plugInfo.json b/Engine/Plugins/Importers/USDImporter/Resources/UsdResources/Linux/plugins/ndr/resources/plugInfo.json
--- a/Engine/Plugins/Importers/USDImporter/Resources/UsdResources/Linux/plugins/ndr/resources/plugInfo.json
+++ b/Engine/Plugins/Importers/USDImporter/Resources/UsdResources/Linux/plugins/ndr/resources/plugInfo.json
@@ -11,7 +11,7 @@
                     "NdrParserPlugin": {}
                 }
             },
-            "LibraryPath": "../../../../../Binaries/Linux/x86_64-unknown-linux-gnu/libndr.so", 
+            "LibraryPath": "../../../../../Source/ThirdParty/Linux/bin/x86_64-unknown-linux-gnu/libndr.so", 
             "Name": "ndr", 
             "ResourcePath": "resources", 
             "Root": "..", 
diff --git a/Engine/Plugins/Importers/USDImporter/Resources/UsdResources/Linux/plugins/sdf/resources/plugInfo.json b/Engine/Plugins/Importers/USDImporter/Resources/UsdResources/Linux/plugins/sdf/resources/plugInfo.json
--- a/Engine/Plugins/Importers/USDImporter/Resources/UsdResources/Linux/plugins/sdf/resources/plugInfo.json
+++ b/Engine/Plugins/Importers/USDImporter/Resources/UsdResources/Linux/plugins/sdf/resources/plugInfo.json
@@ -26,7 +26,7 @@
                     }
                 }
             },
-            "LibraryPath": "../../../../../Binaries/Linux/x86_64-unknown-linux-gnu/libsdf.so",
+            "LibraryPath": "../../../../../Source/ThirdParty/Linux/bin/x86_64-unknown-linux-gnu/libsdf.so",
             "Name": "sdf",
             "ResourcePath": "resources",
             "Root": "..",
diff --git a/Engine/Plugins/Importers/USDImporter/Resources/UsdResources/Linux/plugins/usd/resources/plugInfo.json b/Engine/Plugins/Importers/USDImporter/Resources/UsdResources/Linux/plugins/usd/resources/plugInfo.json
--- a/Engine/Plugins/Importers/USDImporter/Resources/UsdResources/Linux/plugins/usd/resources/plugInfo.json
+++ b/Engine/Plugins/Importers/USDImporter/Resources/UsdResources/Linux/plugins/usd/resources/plugInfo.json
@@ -175,7 +175,7 @@
                     }
                 }
             }, 
-            "LibraryPath": "../../../../../Binaries/Linux/x86_64-unknown-linux-gnu/libusd.so", 
+            "LibraryPath": "../../../../../Source/ThirdParty/Linux/bin/x86_64-unknown-linux-gnu/libusd.so", 
             "Name": "usd", 
             "ResourcePath": "resources", 
             "Root": "..", 
diff --git a/Engine/Plugins/Importers/USDImporter/Resources/UsdResources/Linux/plugins/usdGeom/resources/plugInfo.json b/Engine/Plugins/Importers/USDImporter/Resources/UsdResources/Linux/plugins/usdGeom/resources/plugInfo.json
--- a/Engine/Plugins/Importers/USDImporter/Resources/UsdResources/Linux/plugins/usdGeom/resources/plugInfo.json
+++ b/Engine/Plugins/Importers/USDImporter/Resources/UsdResources/Linux/plugins/usdGeom/resources/plugInfo.json
@@ -292,7 +292,7 @@
                     }
                 }
             }, 
-            "LibraryPath": "../../../../../Binaries/Linux/x86_64-unknown-linux-gnu/libusdGeom.so", 
+            "LibraryPath": "../../../../../Source/ThirdParty/Linux/bin/x86_64-unknown-linux-gnu/libusdGeom.so", 
             "Name": "usdGeom", 
             "ResourcePath": "resources", 
             "Root": "..", 
diff --git a/Engine/Plugins/Importers/USDImporter/Resources/UsdResources/Linux/plugins/usdHydra/resources/plugInfo.json b/Engine/Plugins/Importers/USDImporter/Resources/UsdResources/Linux/plugins/usdHydra/resources/plugInfo.json
--- a/Engine/Plugins/Importers/USDImporter/Resources/UsdResources/Linux/plugins/usdHydra/resources/plugInfo.json
+++ b/Engine/Plugins/Importers/USDImporter/Resources/UsdResources/Linux/plugins/usdHydra/resources/plugInfo.json
@@ -12,7 +12,7 @@
                     }
                 }
             }, 
-            "LibraryPath": "../../../../../Binaries/Linux/x86_64-unknown-linux-gnu/libusdHydra.so", 
+            "LibraryPath": "../../../../../Source/ThirdParty/Linux/bin/x86_64-unknown-linux-gnu/libusdHydra.so", 
             "Name": "usdHydra", 
             "ResourcePath": "resources", 
             "Root": "..", 
diff --git a/Engine/Plugins/Importers/USDImporter/Resources/UsdResources/Linux/plugins/usdLux/resources/plugInfo.json b/Engine/Plugins/Importers/USDImporter/Resources/UsdResources/Linux/plugins/usdLux/resources/plugInfo.json
--- a/Engine/Plugins/Importers/USDImporter/Resources/UsdResources/Linux/plugins/usdLux/resources/plugInfo.json
+++ b/Engine/Plugins/Importers/USDImporter/Resources/UsdResources/Linux/plugins/usdLux/resources/plugInfo.json
@@ -122,7 +122,7 @@
                     }
                 }
             }, 
-            "LibraryPath": "../../../../../Binaries/Linux/x86_64-unknown-linux-gnu/libusdLux.so", 
+            "LibraryPath": "../../../../../Source/ThirdParty/Linux/bin/x86_64-unknown-linux-gnu/libusdLux.so", 
             "Name": "usdLux", 
             "ResourcePath": "resources", 
             "Root": "..", 
diff --git a/Engine/Plugins/Importers/USDImporter/Resources/UsdResources/Linux/plugins/usdRi/resources/plugInfo.json b/Engine/Plugins/Importers/USDImporter/Resources/UsdResources/Linux/plugins/usdRi/resources/plugInfo.json
--- a/Engine/Plugins/Importers/USDImporter/Resources/UsdResources/Linux/plugins/usdRi/resources/plugInfo.json
+++ b/Engine/Plugins/Importers/USDImporter/Resources/UsdResources/Linux/plugins/usdRi/resources/plugInfo.json
@@ -188,7 +188,7 @@
                     }
                 }
             }, 
-            "LibraryPath": "../../../../../Binaries/Linux/x86_64-unknown-linux-gnu/libusdRi.so", 
+            "LibraryPath": "../../../../../Source/ThirdParty/Linux/bin/x86_64-unknown-linux-gnu/libusdRi.so", 
             "Name": "usdRi", 
             "ResourcePath": "resources", 
             "Root": "..", 
diff --git a/Engine/Plugins/Importers/USDImporter/Resources/UsdResources/Linux/plugins/usdShade/resources/plugInfo.json b/Engine/Plugins/Importers/USDImporter/Resources/UsdResources/Linux/plugins/usdShade/resources/plugInfo.json
--- a/Engine/Plugins/Importers/USDImporter/Resources/UsdResources/Linux/plugins/usdShade/resources/plugInfo.json
+++ b/Engine/Plugins/Importers/USDImporter/Resources/UsdResources/Linux/plugins/usdShade/resources/plugInfo.json
@@ -109,7 +109,7 @@
                     }
                 }
             }, 
-            "LibraryPath": "../../../../../Binaries/Linux/x86_64-unknown-linux-gnu/libusdShade.so", 
+            "LibraryPath": "../../../../../Source/ThirdParty/Linux/bin/x86_64-unknown-linux-gnu/libusdShade.so", 
             "Name": "usdShade", 
             "ResourcePath": "resources", 
             "Root": "..", 
diff --git a/Engine/Plugins/Importers/USDImporter/Resources/UsdResources/Linux/plugins/usdSkel/resources/plugInfo.json b/Engine/Plugins/Importers/USDImporter/Resources/UsdResources/Linux/plugins/usdSkel/resources/plugInfo.json
--- a/Engine/Plugins/Importers/USDImporter/Resources/UsdResources/Linux/plugins/usdSkel/resources/plugInfo.json
+++ b/Engine/Plugins/Importers/USDImporter/Resources/UsdResources/Linux/plugins/usdSkel/resources/plugInfo.json
@@ -75,7 +75,7 @@
                     }
                 }
             }, 
-            "LibraryPath": "../../../../../Binaries/Linux/x86_64-unknown-linux-gnu/libusdSkel.so", 
+            "LibraryPath": "../../../../../Source/ThirdParty/Linux/bin/x86_64-unknown-linux-gnu/libusdSkel.so", 
             "Name": "usdSkel", 
             "ResourcePath": "resources", 
             "Root": "..", 
diff --git a/Engine/Plugins/Importers/USDImporter/Resources/UsdResources/Linux/plugins/usdUI/resources/plugInfo.json b/Engine/Plugins/Importers/USDImporter/Resources/UsdResources/Linux/plugins/usdUI/resources/plugInfo.json
--- a/Engine/Plugins/Importers/USDImporter/Resources/UsdResources/Linux/plugins/usdUI/resources/plugInfo.json
+++ b/Engine/Plugins/Importers/USDImporter/Resources/UsdResources/Linux/plugins/usdUI/resources/plugInfo.json
@@ -35,7 +35,7 @@
                     }
                 }
             }, 
-            "LibraryPath": "../../../../../Binaries/Linux/x86_64-unknown-linux-gnu/libusdUI.so", 
+            "LibraryPath": "../../../../../Source/ThirdParty/Linux/bin/x86_64-unknown-linux-gnu/libusdUI.so", 
             "Name": "usdUI", 
             "ResourcePath": "resources", 
             "Root": "..", 
diff --git a/Engine/Plugins/Importers/USDImporter/Resources/UsdResources/Linux/plugins/usdVol/resources/plugInfo.json b/Engine/Plugins/Importers/USDImporter/Resources/UsdResources/Linux/plugins/usdVol/resources/plugInfo.json
--- a/Engine/Plugins/Importers/USDImporter/Resources/UsdResources/Linux/plugins/usdVol/resources/plugInfo.json
+++ b/Engine/Plugins/Importers/USDImporter/Resources/UsdResources/Linux/plugins/usdVol/resources/plugInfo.json
@@ -47,7 +47,7 @@
                     }
                 }
             }, 
-            "LibraryPath": "../../../../../Binaries/Linux/x86_64-unknown-linux-gnu/libusdVol.so", 
+            "LibraryPath": "../../../../../Source/ThirdParty/Linux/bin/x86_64-unknown-linux-gnu/libusdVol.so", 
             "Name": "usdVol", 
             "ResourcePath": "resources", 
             "Root": "..", 
