Disable USD on Linux until the TBB thread crash on editor shutdown is resolved

diff --git a/Engine/Plugins/Importers/USDImporter/Source/UnrealUSDWrapper/UnrealUSDWrapper.Build.cs b/Engine/Plugins/Importers/USDImporter/Source/UnrealUSDWrapper/UnrealUSDWrapper.Build.cs
--- a/Engine/Plugins/Importers/USDImporter/Source/UnrealUSDWrapper/UnrealUSDWrapper.Build.cs
+++ b/Engine/Plugins/Importers/USDImporter/Source/UnrealUSDWrapper/UnrealUSDWrapper.Build.cs
@@ -150,7 +150,8 @@ namespace UnrealBuildTool.Rules
 
 			bool bEnableUsdSdk = (
 				Target.WindowsPlatform.Compiler != WindowsCompiler.Clang &&
-				Target.WindowsPlatform.StaticAnalyzer == WindowsStaticAnalyzer.None
+				Target.WindowsPlatform.StaticAnalyzer == WindowsStaticAnalyzer.None &&
+				Target.Platform != UnrealTargetPlatform.Linux
 			);
 
 			// Don't enable USD when running the include tool because it has issues parsing Boost headers
