USD: Update StageEditor UI to UE5 style.

Fix missing visibility icons for prims and for layers;
Use font styles that naturally invert colors when selected;
Update visibility buttons to new style where they're only visible when necessary;
Update spinbox labels to new style with inner narrow color strips;
Fix text value widgets having their font colors flipped unnecessarily when selected;
Other minor tweaks;

diff --git a/Engine/Plugins/Importers/USDImporter/Source/USDStageEditor/Private/SUSDLayersTreeView.cpp b/Engine/Plugins/Importers/USDImporter/Source/USDStageEditor/Private/SUSDLayersTreeView.cpp
--- a/Engine/Plugins/Importers/USDImporter/Source/USDStageEditor/Private/SUSDLayersTreeView.cpp
+++ b/Engine/Plugins/Importers/USDImporter/Source/USDStageEditor/Private/SUSDLayersTreeView.cpp
@@ -20,6 +20,7 @@
 #include "Styling/SlateTypes.h"
 #include "Widgets/Images/SImage.h"
 #include "Widgets/Input/SButton.h"
+#include "Widgets/SToolTip.h"
 
 #if USE_USD_SDK
 
@@ -28,7 +29,7 @@
 class FUsdLayerNameColumn : public FUsdTreeViewColumn
 {
 public:
-	virtual TSharedRef< SWidget > GenerateWidget( const TSharedPtr< IUsdTreeViewItem > InTreeItem ) override
+	virtual TSharedRef< SWidget > GenerateWidget( const TSharedPtr< IUsdTreeViewItem > InTreeItem, const TSharedPtr< ITableRow > TableRow ) override
 	{
 		FUsdLayerViewModelRef TreeItem = StaticCastSharedRef< FUsdLayerViewModel >( InTreeItem.ToSharedRef() );
 
@@ -36,9 +37,7 @@ public:
 			.VAlign( VAlign_Center )
 			[
 				SNew(STextBlock)
-				.TextStyle( FEditorStyle::Get(), "LargeText" )
 				.Text( TreeItem->LayerModel, &FUsdLayerModel::GetDisplayName )
-				.Font( FEditorStyle::GetFontStyle( "ContentBrowser.SourceTreeItemFont" ) )
 			];
 	}
 };
@@ -53,23 +52,53 @@ public:
 		return FReply::Handled();
 	}
 
-	const FSlateBrush* GetBrush( const FUsdLayerViewModelRef TreeItem ) const
+	const FSlateBrush* GetBrush( const FUsdLayerViewModelRef TreeItem, const TSharedPtr< SButton > Button ) const
 	{
+		const bool bIsButtonHovered = Button.IsValid() && Button->IsHovered();
+
 		if ( !CanMuteLayer( TreeItem ) )
 		{
 			return nullptr;
 		}
 		else if ( TreeItem->LayerModel->bIsMuted )
 		{
-			return FEditorStyle::GetBrush( "Level.NotVisibleIcon16x" );
+			return bIsButtonHovered
+				? FEditorStyle::GetBrush( "Level.NotVisibleHighlightIcon16x" )
+				: FEditorStyle::GetBrush( "Level.NotVisibleIcon16x" );
 		}
 		else
 		{
-			return FEditorStyle::GetBrush( "Level.VisibleIcon16x" );
+			return bIsButtonHovered
+				? FEditorStyle::GetBrush( "Level.VisibleHighlightIcon16x" )
+				: FEditorStyle::GetBrush( "Level.VisibleIcon16x" );
 		}
 	}
 
-	virtual TSharedRef< SWidget > GenerateWidget( const TSharedPtr< IUsdTreeViewItem > InTreeItem ) override
+	FSlateColor GetForegroundColor( const FUsdLayerViewModelRef TreeItem, const TSharedPtr< ITableRow > TableRow, const TSharedPtr< SButton > Button ) const
+	{
+		if ( !TableRow.IsValid() || !Button.IsValid() )
+		{
+			return FSlateColor::UseForeground();
+		}
+
+		const bool bIsRowHovered = TableRow->AsWidget()->IsHovered();
+		const bool bIsButtonHovered = Button->IsHovered();
+		const bool bIsRowSelected = TableRow->IsItemSelected();
+		const bool bIsLayerMuted = TreeItem->IsLayerMuted();
+
+		if ( !bIsLayerMuted && !bIsRowHovered && !bIsRowSelected )
+		{
+			return FLinearColor::Transparent;
+		}
+		else if ( bIsButtonHovered && !bIsRowSelected )
+		{
+			return FEditorStyle::GetSlateColor( TEXT( "Colors.ForegroundHover" ) );
+		}
+
+		return FSlateColor::UseForeground();
+	}
+
+	virtual TSharedRef< SWidget > GenerateWidget( const TSharedPtr< IUsdTreeViewItem > InTreeItem, const TSharedPtr< ITableRow > TableRow ) override
 	{
 		if ( !InTreeItem )
 		{
@@ -77,29 +106,37 @@ public:
 		}
 
 		FUsdLayerViewModelRef TreeItem = StaticCastSharedRef< FUsdLayerViewModel >( InTreeItem.ToSharedRef() );
+		const float ItemSize = FUsdStageEditorStyle::Get()->GetFloat( "UsdStageEditor.ListItemHeight" );
 
-		TSharedRef< SWidget > Item =
-			SNew( SButton )
-			.ContentPadding(0)
-			.ButtonStyle( FEditorStyle::Get(), "ToggleButton" )
+		if ( !TreeItem->CanMuteLayer() )
+		{
+			return SNew( SBox )
+				.HeightOverride( ItemSize )
+				.WidthOverride( ItemSize )
+				.Visibility( EVisibility::Visible )
+				.ToolTip( SNew( SToolTip ).Text( LOCTEXT( "CantMuteLayerTooltip", "This layer cannot be muted!" ) ) );
+		}
+
+		TSharedPtr<SButton> Button = SNew( SButton )
+			.ContentPadding( 0 )
+			.ButtonStyle( FUsdStageEditorStyle::Get(), TEXT("NoBorder") )
 			.OnClicked( this, &FUsdLayerMutedColumn::OnClicked, TreeItem )
+			.ToolTip( SNew( SToolTip ).Text( LOCTEXT( "MuteLayerTooltip", "Mute or unmute this layer" ) ) )
 			.HAlign( HAlign_Center )
-			.VAlign( VAlign_Center )
-			.Content()
-			[
-				SNew( SImage )
-				.Image( this, &FUsdLayerMutedColumn::GetBrush, TreeItem )
-			];
+			.VAlign( VAlign_Center );
 
-		float ItemSize = FUsdStageEditorStyle::Get()->GetFloat( "UsdStageEditor.ListItemHeight" );
+		TSharedPtr<SImage> Image = SNew( SImage )
+			.Image( this, &FUsdLayerMutedColumn::GetBrush, TreeItem, Button )
+			.ColorAndOpacity( this, &FUsdLayerMutedColumn::GetForegroundColor, TreeItem, TableRow, Button );
+
+		Button->SetContent( Image.ToSharedRef() );
 
 		return SNew( SBox )
 			.HeightOverride( ItemSize )
 			.WidthOverride( ItemSize )
-			.HAlign( HAlign_Center )
-			.VAlign( VAlign_Center )
+			.Visibility( EVisibility::Visible )
 			[
-				Item
+				Button.ToSharedRef()
 			];
 	}
 
@@ -135,14 +172,14 @@ public:
 			nullptr;
 	}
 
-	virtual TSharedRef< SWidget > GenerateWidget( const TSharedPtr< IUsdTreeViewItem > InTreeItem ) override
+	virtual TSharedRef< SWidget > GenerateWidget( const TSharedPtr< IUsdTreeViewItem > InTreeItem, const TSharedPtr< ITableRow > TableRow ) override
 	{
 		const FUsdLayerViewModelRef TreeItem = StaticCastSharedRef< FUsdLayerViewModel >( InTreeItem.ToSharedRef() );
 
 		TSharedRef< SWidget > Item =
 			SNew(SImage)
 				.Image( this, &FUsdLayerEditColumn::GetCheckedImage, TreeItem )
-				.ColorAndOpacity( FEditorStyle::Get().GetWidgetStyle< FCheckBoxStyle >( "Checkbox" ).ForegroundColor );
+				.ColorAndOpacity( FSlateColor::UseForeground() );
 
 		float ItemSize = FUsdStageEditorStyle::Get()->GetFloat( "UsdStageEditor.ListItemHeight" );
 
diff --git a/Engine/Plugins/Importers/USDImporter/Source/USDStageEditor/Private/SUSDStage.cpp b/Engine/Plugins/Importers/USDImporter/Source/USDStageEditor/Private/SUSDStage.cpp
--- a/Engine/Plugins/Importers/USDImporter/Source/USDStageEditor/Private/SUSDStage.cpp
+++ b/Engine/Plugins/Importers/USDImporter/Source/USDStageEditor/Private/SUSDStage.cpp
@@ -389,7 +389,7 @@ void SUsdStage::FillActionsMenu( FMenuBuilder& MenuBuilder )
 
 void SUsdStage::FillOptionsMenu(FMenuBuilder& MenuBuilder)
 {
-	MenuBuilder.BeginSection( "Options", LOCTEXT("Options", "Options") );
+	MenuBuilder.BeginSection( TEXT( "StageOptions" ), LOCTEXT( "StageOptions", "Stage options" ) );
 	{
 		MenuBuilder.AddSubMenu(
 			LOCTEXT("Payloads", "Payloads"),
@@ -408,9 +408,11 @@ void SUsdStage::FillOptionsMenu(FMenuBuilder& MenuBuilder)
 			LOCTEXT("RenderContext", "Render Context"),
 			LOCTEXT("RenderContext_ToolTip", "Choose which render context to use when parsing materials"),
 			FNewMenuDelegate::CreateSP(this, &SUsdStage::FillRenderContextSubMenu));
+	}
+	MenuBuilder.EndSection();
 
-		MenuBuilder.AddMenuSeparator();
-
+	MenuBuilder.BeginSection( TEXT( "EditorOptions" ), LOCTEXT( "EditorOptions", "Editor options" ) );
+	{
 		MenuBuilder.AddSubMenu(
 			LOCTEXT( "SelectionText", "Selection" ),
 			LOCTEXT( "SelectionText_ToolTip", "How the selection of prims, actors and components should behave" ),
diff --git a/Engine/Plugins/Importers/USDImporter/Source/USDStageEditor/Private/SUSDStageEditorStyle.cpp b/Engine/Plugins/Importers/USDImporter/Source/USDStageEditor/Private/SUSDStageEditorStyle.cpp
--- a/Engine/Plugins/Importers/USDImporter/Source/USDStageEditor/Private/SUSDStageEditorStyle.cpp
+++ b/Engine/Plugins/Importers/USDImporter/Source/USDStageEditor/Private/SUSDStageEditorStyle.cpp
@@ -8,6 +8,7 @@
 #include "Styling/CoreStyle.h"
 #include "Styling/SlateStyleRegistry.h"
 #include "Styling/SlateTypes.h"
+#include "Styling/StyleColors.h"
 
 #define IMAGE_PLUGIN_BRUSH( RelativePath, ... ) FSlateImageBrush( FUsdStageEditorStyle::InContent( RelativePath, ".png" ), __VA_ARGS__ )
 #define IMAGE_BRUSH(RelativePath, ...) FSlateImageBrush(StyleSet->RootToContentDir(RelativePath, TEXT(".png")), __VA_ARGS__)
@@ -52,6 +53,25 @@ void FUsdStageEditorStyle::Initialize()
 	{
 		/** Height of items in the stage/layer tree views, prim info panel, etc. */
 		StyleSet->Set( "UsdStageEditor.ListItemHeight", 20.0f );
+
+		/** Color of Prim labels that have composition arcs */
+		StyleSet->Set( "UsdStageEditor.PrimCompositionArcColor", FLinearColor( FColor::Orange ) );
+
+		/** Highlighted color of Prim labels that have composition arcs */
+		StyleSet->Set( "UsdStageEditor.HighlightPrimCompositionArcColor", FLinearColor( 0.2f, 0.08f, 0.0f, 1.0f ) );
+
+		/**
+		 * Button without a background, border or foreground color that doesn't move when pressed.
+		 * No foreground color is important to let the automatic slate foreground color propagate down.
+		 */
+		const FButtonStyle Style = FButtonStyle()
+			.SetNormal( FSlateNoResource() )
+			.SetHovered( FSlateNoResource() )
+			.SetPressed( FSlateNoResource() )
+			.SetDisabled( FSlateNoResource() )
+			.SetNormalPadding( FMargin( 2, 2, 2, 2 ) )
+			.SetPressedPadding( FMargin( 2, 2, 2, 2 ) );
+		StyleSet->Set( "NoBorder", Style );
 	}
 
 	FSlateStyleRegistry::RegisterSlateStyle(*StyleSet.Get());
diff --git a/Engine/Plugins/Importers/USDImporter/Source/USDStageEditor/Private/SUSDStageTreeView.cpp b/Engine/Plugins/Importers/USDImporter/Source/USDStageEditor/Private/SUSDStageTreeView.cpp
--- a/Engine/Plugins/Importers/USDImporter/Source/USDStageEditor/Private/SUSDStageTreeView.cpp
+++ b/Engine/Plugins/Importers/USDImporter/Source/USDStageEditor/Private/SUSDStageTreeView.cpp
@@ -51,7 +51,7 @@ public:
 
 	TWeakPtr< SUsdStageTreeView > OwnerTree;
 
-	virtual TSharedRef< SWidget > GenerateWidget( const TSharedPtr< IUsdTreeViewItem > InTreeItem ) override
+	virtual TSharedRef< SWidget > GenerateWidget( const TSharedPtr< IUsdTreeViewItem > InTreeItem, const TSharedPtr< ITableRow > TableRow ) override
 	{
 		if ( !InTreeItem )
 		{
@@ -63,7 +63,6 @@ public:
 		TSharedRef< SInlineEditableTextBlock > Item =
 			SNew( SInlineEditableTextBlock )
 			.Text( TreeItem->RowData, &FUsdPrimModel::GetName )
-			.Font( FEditorStyle::GetFontStyle( "ContentBrowser.SourceTreeItemFont" ) )
 			.ColorAndOpacity( this, &FUsdStageNameColumn::GetTextColor, TreeItem )
 			.OnTextCommitted( this, &FUsdStageNameColumn::OnTextCommitted, TreeItem )
 			.OnVerifyTextChanged( this, &FUsdStageNameColumn::OnTextUpdated, TreeItem )
@@ -106,18 +105,22 @@ protected:
 
 	FSlateColor GetTextColor( TSharedPtr< FUsdPrimViewModel > TreeItem ) const
 	{
-		FSlateColor TextColor = FLinearColor::White;
+		FSlateColor TextColor = FSlateColor::UseForeground();
+		if ( !TreeItem )
+		{
+			return TextColor;
+		}
 
-		if ( TreeItem && TreeItem->RowData->HasCompositionArcs() )
+		if ( TreeItem->RowData->HasCompositionArcs() )
 		{
 			const TSharedPtr< SUsdStageTreeView > OwnerTreePtr = OwnerTree.Pin();
 			if ( OwnerTreePtr && OwnerTreePtr->IsItemSelected( TreeItem.ToSharedRef() ) )
 			{
-				TextColor = FSlateColor::UseForeground();
+				TextColor = FUsdStageEditorStyle::Get()->GetColor( "UsdStageEditor.HighlightPrimCompositionArcColor" );
 			}
 			else
 			{
-				TextColor = FLinearColor( FColor::Orange );
+				TextColor = FUsdStageEditorStyle::Get()->GetColor( "UsdStageEditor.PrimCompositionArcColor" );
 			}
 		}
 
@@ -160,7 +163,7 @@ public:
 		}
 	}
 
-	virtual TSharedRef< SWidget > GenerateWidget( const TSharedPtr< IUsdTreeViewItem > InTreeItem ) override
+	virtual TSharedRef< SWidget > GenerateWidget( const TSharedPtr< IUsdTreeViewItem > InTreeItem, const TSharedPtr< ITableRow > TableRow ) override
 	{
 		bool bHasPayload = false;
 
@@ -198,26 +201,60 @@ public:
 		return FReply::Handled();
 	}
 
-	const FSlateBrush* GetBrush( const FUsdPrimViewModelPtr TreeItem ) const
+	const FSlateBrush* GetBrush( const FUsdPrimViewModelPtr TreeItem, const TSharedPtr< SButton > Button ) const
 	{
+		const bool bIsButtonHovered = Button.IsValid() && Button->IsHovered();
+
 		if ( TreeItem->RowData->IsVisible() )
 		{
-			return FEditorStyle::GetBrush( "Level.VisibleIcon16x" );
+			return bIsButtonHovered
+				? FEditorStyle::GetBrush( "Level.VisibleHighlightIcon16x" )
+				: FEditorStyle::GetBrush( "Level.VisibleIcon16x" );
 		}
 		else
 		{
-			return FEditorStyle::GetBrush( "Level.NotVisibleIcon16x" );
+			return bIsButtonHovered
+				? FEditorStyle::GetBrush( "Level.NotVisibleHighlightIcon16x" )
+				: FEditorStyle::GetBrush( "Level.NotVisibleIcon16x" );
 		}
 	}
 
-	virtual TSharedRef< SWidget > GenerateWidget( const TSharedPtr< IUsdTreeViewItem > InTreeItem ) override
+	FSlateColor GetForegroundColor( const FUsdPrimViewModelPtr TreeItem, const TSharedPtr< ITableRow > TableRow, const TSharedPtr< SButton > Button ) const
 	{
+		if ( !TreeItem.IsValid() || !TableRow.IsValid() || !Button.IsValid() )
+		{
+			return FSlateColor::UseForeground();
+		}
+
+		const bool bIsRowHovered = TableRow->AsWidget()->IsHovered();
+		const bool bIsButtonHovered = Button->IsHovered();
+		const bool bIsRowSelected = TableRow->IsItemSelected();
+		const bool bIsPrimVisible = TreeItem->RowData->IsVisible();
+
+		if ( bIsPrimVisible && !bIsRowHovered && !bIsRowSelected )
+		{
+			return FLinearColor::Transparent;
+		}
+		else if ( bIsButtonHovered && !bIsRowSelected )
+		{
+			return FEditorStyle::GetSlateColor( TEXT( "Colors.ForegroundHover" ) );
+		}
+
+		return FSlateColor::UseForeground();
+	}
+
+	virtual TSharedRef< SWidget > GenerateWidget( const TSharedPtr< IUsdTreeViewItem > InTreeItem, const TSharedPtr< ITableRow > TableRow ) override
+	{
+		if ( !InTreeItem )
+		{
+			return SNullWidget::NullWidget;
+		}
+
 		const TSharedPtr< FUsdPrimViewModel > TreeItem = StaticCastSharedPtr< FUsdPrimViewModel >( InTreeItem );
+		const float ItemSize = FUsdStageEditorStyle::Get()->GetFloat( "UsdStageEditor.ListItemHeight" );
 
 		if ( !TreeItem->HasVisibilityAttribute() )
 		{
-			float ItemSize = FUsdStageEditorStyle::Get()->GetFloat( "UsdStageEditor.ListItemHeight" );
-
 			return SNew(SBox)
 				.HeightOverride( ItemSize )
 				.WidthOverride( ItemSize )
@@ -225,27 +262,34 @@ public:
 				.ToolTip( SNew( SToolTip ).Text( LOCTEXT( "NoGeomImageable", "Only prims with the GeomImageable schema (or derived) have the visibility attribute!" ) ) );
 		}
 
-		TSharedRef< SWidget > Item = SNew( SButton )
-										.ContentPadding(0)
-										.ButtonStyle( FEditorStyle::Get(), "ToggleButton" )
-										.OnClicked( this, &FUsdStageVisibilityColumn::OnToggleVisibility, TreeItem )
-										.ToolTip( SNew( SToolTip ).Text( LOCTEXT( "GeomImageable", "Toggle the visibility of this prim" ) ) )
-										.HAlign( HAlign_Center )
-										.VAlign( VAlign_Center )
-										.Content()
-										[
-											SNew( SImage )
-											.Image( this, &FUsdStageVisibilityColumn::GetBrush, TreeItem )
-										];
+		TSharedPtr<SButton> Button = SNew( SButton )
+			.ContentPadding( 0 )
+			.ButtonStyle( FUsdStageEditorStyle::Get(), TEXT("NoBorder") )
+			.OnClicked( this, &FUsdStageVisibilityColumn::OnToggleVisibility, TreeItem )
+			.ToolTip( SNew( SToolTip ).Text( LOCTEXT( "GeomImageable", "Toggle the visibility of this prim" ) ) )
+			.HAlign( HAlign_Center )
+			.VAlign( VAlign_Center );
 
-		return Item;
+		TSharedPtr<SImage> Image = SNew( SImage )
+			.Image( this, &FUsdStageVisibilityColumn::GetBrush, TreeItem, Button )
+			.ColorAndOpacity( this, &FUsdStageVisibilityColumn::GetForegroundColor, TreeItem, TableRow, Button );
+
+		Button->SetContent( Image.ToSharedRef() );
+
+		return SNew( SBox )
+			.HeightOverride( ItemSize )
+			.WidthOverride( ItemSize )
+			.Visibility( EVisibility::Visible )
+			[
+				Button.ToSharedRef()
+			];
 	}
 };
 
 class FUsdStagePrimTypeColumn : public FUsdTreeViewColumn
 {
 public:
-	virtual TSharedRef< SWidget > GenerateWidget( const TSharedPtr< IUsdTreeViewItem > InTreeItem ) override
+	virtual TSharedRef< SWidget > GenerateWidget( const TSharedPtr< IUsdTreeViewItem > InTreeItem, const TSharedPtr< ITableRow > TableRow ) override
 	{
 		TSharedPtr< FUsdPrimViewModel > TreeItem = StaticCastSharedPtr< FUsdPrimViewModel >( InTreeItem );
 
@@ -253,9 +297,7 @@ public:
 			.VAlign( VAlign_Center )
 			[
 				SNew(STextBlock)
-				.TextStyle( FEditorStyle::Get(), "LargeText" )
 				.Text( TreeItem->RowData, &FUsdPrimModel::GetType )
-				.Font( FEditorStyle::GetFontStyle( "ContentBrowser.SourceTreeItemFont" ) )
 			];
 	}
 };
diff --git a/Engine/Plugins/Importers/USDImporter/Source/USDStageEditor/Private/Widgets/SUSDPrimPropertiesList.cpp b/Engine/Plugins/Importers/USDImporter/Source/USDStageEditor/Private/Widgets/SUSDPrimPropertiesList.cpp
--- a/Engine/Plugins/Importers/USDImporter/Source/USDStageEditor/Private/Widgets/SUSDPrimPropertiesList.cpp
+++ b/Engine/Plugins/Importers/USDImporter/Source/USDStageEditor/Private/Widgets/SUSDPrimPropertiesList.cpp
@@ -171,8 +171,6 @@ namespace UsdPrimPropertiesListImpl
 	/** Always returns `NumLabels` widgets, regardless of LabelTypes. Those may be the NullWidget though */
 	TArray<TSharedRef<SWidget>> GetNumericEntryBoxLabels( int32 NumLabels, EPrimPropertyLabelTypes LabelType )
 	{
-		const static TArray<FString> XYZWLabels = { TEXT( "X" ), TEXT( "Y" ), TEXT( "Z" ), TEXT( "W" ) };
-		const static TArray<FString> RGBALabels = { TEXT( "R" ), TEXT( "G" ), TEXT( "B" ), TEXT( "A" ) };
 		const static TArray<const FLinearColor*> Colors = {
 			&SNumericEntryBox<int32>::RedLabelBackgroundColor,
 			&SNumericEntryBox<int32>::GreenLabelBackgroundColor,
@@ -193,13 +191,13 @@ namespace UsdPrimPropertiesListImpl
 		case EPrimPropertyLabelTypes::RGBA:
 			for ( int32 Index = 0; Index < NumLabels; ++Index )
 			{
-				Labels.Add( SNumericEntryBox<int32>::BuildLabel( FText::FromString( RGBALabels[ Index ] ), Index == 3 ? FLinearColor::Black : FLinearColor::White, *Colors[ Index ] ) );
+				Labels.Add( SNumericEntryBox<int32>::BuildNarrowColorLabel( *Colors[ Index ] ) );
 			}
 			break;
 		case EPrimPropertyLabelTypes::XYZW:
 			for ( int32 Index = 0; Index < NumLabels; ++Index )
 			{
-				Labels.Add( SNumericEntryBox<int32>::BuildLabel( FText::FromString( XYZWLabels[ Index ] ), Index == 3 ? FLinearColor::Black : FLinearColor::White, *Colors[ Index ] ) );
+				Labels.Add( SNumericEntryBox<int32>::BuildNarrowColorLabel( *Colors[ Index ] ) );
 			}
 			break;
 		case EPrimPropertyLabelTypes::NoLabel:
@@ -324,7 +322,6 @@ TSharedRef< SWidget > SUsdPrimPropertyRow::GenerateWidgetForColumn( const FName&
 
 	if ( ColumnName == TEXT("PropertyName") )
 	{
-		const bool bIsReadOnly = true;
 		ColumnWidget = GenerateTextWidget( { this, &SUsdPrimPropertyRow::GetLabel } );
 	}
 	else
@@ -373,7 +370,6 @@ TSharedRef< SWidget > SUsdPrimPropertyRow::GenerateWidgetForColumn( const FName&
 					.OptionsSource( Options )
 					.OnGenerateWidget_Lambda( [ & ]( TSharedPtr<FString> Option )
 					{
-						const bool bIsReadOnly = false;
 						return SUsdPrimPropertyRow::GenerateTextWidget( FText::FromString( *Option ) );
 					})
 					.OnSelectionChanged( this, &SUsdPrimPropertyRow::OnComboBoxSelectionChanged )
@@ -382,9 +378,10 @@ TSharedRef< SWidget > SUsdPrimPropertyRow::GenerateWidgetForColumn( const FName&
 						// also specify a custom kind/purpose/etc. if he wants to
 						SNew( SEditableTextBox )
 						.Text( this, &SUsdPrimPropertyRow::GetValueText )
-						.Font( FEditorStyle::GetFontStyle( UsdPrimPropertiesListConstants::NormalFont ) )
-						.Padding( FMargin( 0.0f ) )
-						.BackgroundColor( FLinearColor::White ) // This could be transparent, but white also hints that it's a regular text input field
+						.Font( FEditorStyle::GetFontStyle( "PropertyWindow.NormalFont" ) )
+						.Padding( FMargin( 3.0f ) )
+						// Fixed foreground color or else it will flip when our row is selected, and we already have a background
+						.ForegroundColor( FEditorStyle::GetSlateColor( TEXT( "Colors.ForegroundHover" ) ) )
 						.OnTextCommitted( this, &SUsdPrimPropertyRow::OnTextBoxTextCommitted )
 					]
 				];
@@ -452,7 +449,7 @@ TSharedRef< SWidget > SUsdPrimPropertyRow::GenerateTextWidget( const TAttribute<
 		[
 			SNew( STextBlock )
 			.Text( Attribute )
-			.Font( FEditorStyle::GetFontStyle( UsdPrimPropertiesListConstants::NormalFont ) )
+			.Font( FEditorStyle::GetFontStyle( "PropertyWindow.NormalFont" ) )
 			.Margin( UsdPrimPropertiesListConstants::RightRowPadding )
 		];
 }
@@ -467,8 +464,10 @@ TSharedRef< SWidget > SUsdPrimPropertyRow::GenerateEditableTextWidget( const TAt
 			SNew( SEditableTextBox )
 			.Text( Attribute )
 			.IsReadOnly( bIsReadOnly )
-			.Font( FEditorStyle::GetFontStyle( UsdPrimPropertiesListConstants::NormalFont ) )
+			.Font( FEditorStyle::GetFontStyle( "PropertyWindow.NormalFont" ) )
 			.OnTextCommitted( this, &SUsdPrimPropertyRow::OnTextBoxTextCommitted )
+			// Fixed foreground color or else it will flip when our row is selected, and we already have a background
+			.ForegroundColor( FEditorStyle::GetSlateColor( bIsReadOnly ? TEXT( "Colors.Foreground" ) : TEXT( "Colors.ForegroundHover" ) ) )
 		];
 }
 
@@ -518,7 +517,7 @@ TSharedRef< SWidget > SUsdPrimPropertyRow::GenerateSpinboxWidgets( const TArray<
 
 			TSharedRef<SWidget> EntryBox = SNew( SNumericEntryBox<T> )
 				.AllowSpin( true )
-				.LabelPadding( 0 )
+				.Font( FEditorStyle::GetFontStyle( "PropertyWindow.NormalFont" ) )
 				.ShiftMouseMovePixelPerDelta( 1 )
 				.SupportDynamicSliderMaxValue( true )
 				.SupportDynamicSliderMinValue( true )
@@ -531,6 +530,8 @@ TSharedRef< SWidget > SUsdPrimPropertyRow::GenerateSpinboxWidgets( const TArray<
 				.MaxSliderValue( TOptional<T>() )
 				.SliderExponent( T( 1 ) )
 				.Delta( T( 0 ) )
+				.LabelPadding( FMargin(3) )
+				.LabelLocation( SNumericEntryBox<T>::ELabelLocation::Inside )
 				.Label()
 				[
 					Labels[ ComponentIndex ]
diff --git a/Engine/Plugins/Importers/USDImporter/Source/USDStageEditor/Private/Widgets/SUSDTreeView.h b/Engine/Plugins/Importers/USDImporter/Source/USDStageEditor/Private/Widgets/SUSDTreeView.h
--- a/Engine/Plugins/Importers/USDImporter/Source/USDStageEditor/Private/Widgets/SUSDTreeView.h
+++ b/Engine/Plugins/Importers/USDImporter/Source/USDStageEditor/Private/Widgets/SUSDTreeView.h
@@ -16,7 +16,7 @@ class FUsdTreeViewColumn
 public:
 	virtual ~FUsdTreeViewColumn() = default;
 
-	virtual TSharedRef< SWidget > GenerateWidget( const TSharedPtr< IUsdTreeViewItem > TreeItem ) = 0;
+	virtual TSharedRef< SWidget > GenerateWidget( const TSharedPtr< IUsdTreeViewItem > TreeItem, const TSharedPtr< ITableRow > TableRow ) = 0;
 
 public:
 	bool bIsMainColumn = false;
@@ -109,12 +109,15 @@ private:
 
 template< typename ItemType >
 inline TSharedRef< SWidget > SUsdTreeRow< ItemType >::GenerateWidgetForColumn( const FName& ColumnName )
-{	
+{
 	TSharedRef< SWidget > ColumnWidget = SNullWidget::NullWidget;
 
 	if ( SharedData && SharedData->Columns.Contains( ColumnName ) )
 	{
-		ColumnWidget = SharedData->Columns[ ColumnName ]->GenerateWidget( StaticCastSharedPtr< IUsdTreeViewItem >( TSharedPtr< typename ItemType::ElementType >( TreeItem ) ) );
+		ColumnWidget = SharedData->Columns[ ColumnName ]->GenerateWidget(
+			StaticCastSharedPtr< IUsdTreeViewItem >( TSharedPtr< typename ItemType::ElementType >( TreeItem ) ),
+			SharedThis(this)
+		);
 	}
 
 	if ( SharedData->Columns[ ColumnName ]->bIsMainColumn )
diff --git a/Engine/Plugins/Importers/USDImporter/Source/USDStageEditorViewModels/Private/USDLayersViewModel.cpp b/Engine/Plugins/Importers/USDImporter/Source/USDStageEditorViewModels/Private/USDLayersViewModel.cpp
--- a/Engine/Plugins/Importers/USDImporter/Source/USDStageEditorViewModels/Private/USDLayersViewModel.cpp
+++ b/Engine/Plugins/Importers/USDImporter/Source/USDStageEditorViewModels/Private/USDLayersViewModel.cpp
@@ -163,6 +163,23 @@ UE::FSdfLayer FUsdLayerViewModel::GetLayer() const
 	return UE::FSdfLayer::FindOrOpen( *LayerIdentifier );
 }
 
+bool FUsdLayerViewModel::IsLayerMuted() const
+{
+	bool bIsMuted = false;
+
+#if USE_USD_SDK
+	FScopedUsdAllocs UsdAllocs;
+
+	const TUsdStore< std::string > UsdLayerIdentifier = UnrealToUsd::ConvertString( *LayerIdentifier );
+
+	pxr::UsdStageRefPtr UsdStageRef( UsdStage );
+
+	bIsMuted = UsdStageRef->IsLayerMuted( UsdLayerIdentifier.Get() );
+#endif // USE_USD_SDK
+
+	return bIsMuted;
+}
+
 bool FUsdLayerViewModel::CanMuteLayer() const
 {
 	if ( !IsValid() )
diff --git a/Engine/Plugins/Importers/USDImporter/Source/USDStageEditorViewModels/Public/USDLayersViewModel.h b/Engine/Plugins/Importers/USDImporter/Source/USDStageEditorViewModels/Public/USDLayersViewModel.h
--- a/Engine/Plugins/Importers/USDImporter/Source/USDStageEditorViewModels/Public/USDLayersViewModel.h
+++ b/Engine/Plugins/Importers/USDImporter/Source/USDStageEditorViewModels/Public/USDLayersViewModel.h
@@ -38,6 +38,7 @@ public:
 
 	UE::FSdfLayer GetLayer() const;
 
+	bool IsLayerMuted() const;
 	bool CanMuteLayer() const;
 	void ToggleMuteLayer();
 
