Fix build.

#rnx

diff --git a/Engine/Plugins/Importers/USDImporter/Source/USDStageEditor/Private/Widgets/SUSDTreeView.h b/Engine/Plugins/Importers/USDImporter/Source/USDStageEditor/Private/Widgets/SUSDTreeView.h
--- a/Engine/Plugins/Importers/USDImporter/Source/USDStageEditor/Private/Widgets/SUSDTreeView.h
+++ b/Engine/Plugins/Importers/USDImporter/Source/USDStageEditor/Private/Widgets/SUSDTreeView.h
@@ -116,7 +116,7 @@ inline TSharedRef< SWidget > SUsdTreeRow< ItemType >::GenerateWidgetForColumn( c
 	{
 		ColumnWidget = SharedData->Columns[ ColumnName ]->GenerateWidget(
 			StaticCastSharedPtr< IUsdTreeViewItem >( TSharedPtr< typename ItemType::ElementType >( TreeItem ) ),
-			SharedThis(this)
+			this->SharedThis(this)
 		);
 	}
 
diff --git a/Engine/Source/Editor/TextureEditor/Private/TextureEditorToolkit.cpp b/Engine/Source/Editor/TextureEditor/Private/TextureEditorToolkit.cpp
--- a/Engine/Source/Editor/TextureEditor/Private/TextureEditorToolkit.cpp
+++ b/Engine/Source/Editor/TextureEditor/Private/TextureEditorToolkit.cpp
@@ -1083,10 +1083,10 @@ int32 FTextureEditorToolkit::GetNumMips( ) const
 	{
 		return Texture2DDynamic->NumMips;
 	}
-	else if (const UMediaTexture* MediaTexture = Cast<UMediaTexture>(Texture))
-	{
-		return MediaTexture->GetTextureNumMips();
-	}
+	//else if (const UMediaTexture* MediaTexture = Cast<UMediaTexture>(Texture))
+	//{
+	//	return MediaTexture->GetTextureNumMips();
+	//}
 
 	return MIPLEVEL_MAX;
 }
