Fixed USDImporter loading two copies of its plugin .dylibs. The PublicAdditionalLibraries.Add() in the UnrealUSDImporter.Build.cs causes the plugins output UnrealEditor-UnrealUSDImporter.dylib to dynamically link against the .dylibs in Source/ThirdParty. However, the plugin descriptions were set to use copies in Engine/Binaries/Mac that RuntimeDependencies.Add() does. This would result in two libusd.dylib libraries being loaded and cause errors as they would both register their type info to the same singleton. This change makes everything use the Source/ThirdParty/.dylibs, and nukes the copy to avoid this in the future.

#rnx

diff --git a/Engine/Plugins/Importers/USDImporter/Resources/UsdResources/Mac/plugins/ar/resources/plugInfo.json b/Engine/Plugins/Importers/USDImporter/Resources/UsdResources/Mac/plugins/ar/resources/plugInfo.json
--- a/Engine/Plugins/Importers/USDImporter/Resources/UsdResources/Mac/plugins/ar/resources/plugInfo.json
+++ b/Engine/Plugins/Importers/USDImporter/Resources/UsdResources/Mac/plugins/ar/resources/plugInfo.json
@@ -7,7 +7,7 @@
                     "ArPackageResolver": {}
                 }
             },
-            "LibraryPath": "../../../../../../../../Binaries/Mac/libar.dylib", 
+            "LibraryPath": "../../../../../Source/ThirdParty/Mac/bin/libar.dylib", 
             "Name": "ar", 
             "ResourcePath": "resources", 
             "Root": "..", 
diff --git a/Engine/Plugins/Importers/USDImporter/Resources/UsdResources/Mac/plugins/ndr/resources/plugInfo.json b/Engine/Plugins/Importers/USDImporter/Resources/UsdResources/Mac/plugins/ndr/resources/plugInfo.json
--- a/Engine/Plugins/Importers/USDImporter/Resources/UsdResources/Mac/plugins/ndr/resources/plugInfo.json
+++ b/Engine/Plugins/Importers/USDImporter/Resources/UsdResources/Mac/plugins/ndr/resources/plugInfo.json
@@ -11,7 +11,7 @@
                     "NdrParserPlugin": {}
                 }
             },
-            "LibraryPath": "../../../../../../../../Binaries/Mac/libndr.dylib", 
+            "LibraryPath": "../../../../../Source/ThirdParty/Mac/bin/libndr.dylib", 
             "Name": "ndr", 
             "ResourcePath": "resources", 
             "Root": "..", 
diff --git a/Engine/Plugins/Importers/USDImporter/Resources/UsdResources/Mac/plugins/sdf/resources/plugInfo.json b/Engine/Plugins/Importers/USDImporter/Resources/UsdResources/Mac/plugins/sdf/resources/plugInfo.json
--- a/Engine/Plugins/Importers/USDImporter/Resources/UsdResources/Mac/plugins/sdf/resources/plugInfo.json
+++ b/Engine/Plugins/Importers/USDImporter/Resources/UsdResources/Mac/plugins/sdf/resources/plugInfo.json
@@ -26,7 +26,7 @@
                     }
                 }
             },
-            "LibraryPath": "../../../../../../../../Binaries/Mac/libsdf.dylib",
+            "LibraryPath": "../../../../../Source/ThirdParty/Mac/bin/libsdf.dylib",
             "Name": "sdf",
             "ResourcePath": "resources",
             "Root": "..",
diff --git a/Engine/Plugins/Importers/USDImporter/Resources/UsdResources/Mac/plugins/usd/resources/plugInfo.json b/Engine/Plugins/Importers/USDImporter/Resources/UsdResources/Mac/plugins/usd/resources/plugInfo.json
--- a/Engine/Plugins/Importers/USDImporter/Resources/UsdResources/Mac/plugins/usd/resources/plugInfo.json
+++ b/Engine/Plugins/Importers/USDImporter/Resources/UsdResources/Mac/plugins/usd/resources/plugInfo.json
@@ -175,7 +175,7 @@
                     }
                 }
             }, 
-            "LibraryPath": "../../../../../../../../Binaries/Mac/libusd.dylib", 
+            "LibraryPath": "../../../../../Source/ThirdParty/Mac/bin/libusd.dylib", 
             "Name": "usd", 
             "ResourcePath": "resources", 
             "Root": "..", 
diff --git a/Engine/Plugins/Importers/USDImporter/Resources/UsdResources/Mac/plugins/usdGeom/resources/plugInfo.json b/Engine/Plugins/Importers/USDImporter/Resources/UsdResources/Mac/plugins/usdGeom/resources/plugInfo.json
--- a/Engine/Plugins/Importers/USDImporter/Resources/UsdResources/Mac/plugins/usdGeom/resources/plugInfo.json
+++ b/Engine/Plugins/Importers/USDImporter/Resources/UsdResources/Mac/plugins/usdGeom/resources/plugInfo.json
@@ -292,7 +292,7 @@
                     }
                 }
             }, 
-            "LibraryPath": "../../../../../../../../Binaries/Mac/libusdGeom.dylib", 
+            "LibraryPath": "../../../../../Source/ThirdParty/Mac/bin/libusdGeom.dylib", 
             "Name": "usdGeom", 
             "ResourcePath": "resources", 
             "Root": "..", 
diff --git a/Engine/Plugins/Importers/USDImporter/Resources/UsdResources/Mac/plugins/usdLux/resources/plugInfo.json b/Engine/Plugins/Importers/USDImporter/Resources/UsdResources/Mac/plugins/usdLux/resources/plugInfo.json
--- a/Engine/Plugins/Importers/USDImporter/Resources/UsdResources/Mac/plugins/usdLux/resources/plugInfo.json
+++ b/Engine/Plugins/Importers/USDImporter/Resources/UsdResources/Mac/plugins/usdLux/resources/plugInfo.json
@@ -122,7 +122,7 @@
                     }
                 }
             }, 
-            "LibraryPath": "../../../../../../../../Binaries/Mac/libusdLux.dylib", 
+            "LibraryPath": "../../../../../Source/ThirdParty/Mac/bin/libusdLux.dylib", 
             "Name": "usdLux", 
             "ResourcePath": "resources", 
             "Root": "..", 
diff --git a/Engine/Plugins/Importers/USDImporter/Resources/UsdResources/Mac/plugins/usdRi/resources/plugInfo.json b/Engine/Plugins/Importers/USDImporter/Resources/UsdResources/Mac/plugins/usdRi/resources/plugInfo.json
--- a/Engine/Plugins/Importers/USDImporter/Resources/UsdResources/Mac/plugins/usdRi/resources/plugInfo.json
+++ b/Engine/Plugins/Importers/USDImporter/Resources/UsdResources/Mac/plugins/usdRi/resources/plugInfo.json
@@ -188,7 +188,7 @@
                     }
                 }
             }, 
-            "LibraryPath": "../../../../../../../../Binaries/Mac/libusdRi.dylib", 
+            "LibraryPath": "../../../../../Source/ThirdParty/Mac/bin/libusdRi.dylib", 
             "Name": "usdRi", 
             "ResourcePath": "resources", 
             "Root": "..", 
diff --git a/Engine/Plugins/Importers/USDImporter/Resources/UsdResources/Mac/plugins/usdShade/resources/plugInfo.json b/Engine/Plugins/Importers/USDImporter/Resources/UsdResources/Mac/plugins/usdShade/resources/plugInfo.json
--- a/Engine/Plugins/Importers/USDImporter/Resources/UsdResources/Mac/plugins/usdShade/resources/plugInfo.json
+++ b/Engine/Plugins/Importers/USDImporter/Resources/UsdResources/Mac/plugins/usdShade/resources/plugInfo.json
@@ -109,7 +109,7 @@
                     }
                 }
             }, 
-            "LibraryPath": "../../../../../../../../Binaries/Mac/libusdShade.dylib", 
+            "LibraryPath": "../../../../../Source/ThirdParty/Mac/bin/libusdShade.dylib", 
             "Name": "usdShade", 
             "ResourcePath": "resources", 
             "Root": "..", 
diff --git a/Engine/Plugins/Importers/USDImporter/Resources/UsdResources/Mac/plugins/usdSkel/resources/plugInfo.json b/Engine/Plugins/Importers/USDImporter/Resources/UsdResources/Mac/plugins/usdSkel/resources/plugInfo.json
--- a/Engine/Plugins/Importers/USDImporter/Resources/UsdResources/Mac/plugins/usdSkel/resources/plugInfo.json
+++ b/Engine/Plugins/Importers/USDImporter/Resources/UsdResources/Mac/plugins/usdSkel/resources/plugInfo.json
@@ -75,7 +75,7 @@
                     }
                 }
             }, 
-            "LibraryPath": "../../../../../../../../Binaries/Mac/libusdSkel.dylib", 
+            "LibraryPath": "../../../../../Source/ThirdParty/Mac/bin/libusdSkel.dylib", 
             "Name": "usdSkel", 
             "ResourcePath": "resources", 
             "Root": "..", 
diff --git a/Engine/Plugins/Importers/USDImporter/Resources/UsdResources/Mac/plugins/usdUI/resources/plugInfo.json b/Engine/Plugins/Importers/USDImporter/Resources/UsdResources/Mac/plugins/usdUI/resources/plugInfo.json
--- a/Engine/Plugins/Importers/USDImporter/Resources/UsdResources/Mac/plugins/usdUI/resources/plugInfo.json
+++ b/Engine/Plugins/Importers/USDImporter/Resources/UsdResources/Mac/plugins/usdUI/resources/plugInfo.json
@@ -35,7 +35,7 @@
                     }
                 }
             }, 
-            "LibraryPath": "../../../../../../../../Binaries/Mac/libusdUI.dylib", 
+            "LibraryPath": "../../../../../Source/ThirdParty/Mac/bin/libusdUI.dylib", 
             "Name": "usdUI", 
             "ResourcePath": "resources", 
             "Root": "..", 
diff --git a/Engine/Plugins/Importers/USDImporter/Resources/UsdResources/Mac/plugins/usdVol/resources/plugInfo.json b/Engine/Plugins/Importers/USDImporter/Resources/UsdResources/Mac/plugins/usdVol/resources/plugInfo.json
--- a/Engine/Plugins/Importers/USDImporter/Resources/UsdResources/Mac/plugins/usdVol/resources/plugInfo.json
+++ b/Engine/Plugins/Importers/USDImporter/Resources/UsdResources/Mac/plugins/usdVol/resources/plugInfo.json
@@ -47,7 +47,7 @@
                     }
                 }
             }, 
-            "LibraryPath": "../../../../../../../../Binaries/Mac/libusdVol.dylib", 
+            "LibraryPath": "../../../../../Source/ThirdParty/Mac/bin/libusdVol.dylib", 
             "Name": "usdVol", 
             "ResourcePath": "resources", 
             "Root": "..", 
diff --git a/Engine/Plugins/Importers/USDImporter/Source/UnrealUSDWrapper/UnrealUSDWrapper.Build.cs b/Engine/Plugins/Importers/USDImporter/Source/UnrealUSDWrapper/UnrealUSDWrapper.Build.cs
--- a/Engine/Plugins/Importers/USDImporter/Source/UnrealUSDWrapper/UnrealUSDWrapper.Build.cs
+++ b/Engine/Plugins/Importers/USDImporter/Source/UnrealUSDWrapper/UnrealUSDWrapper.Build.cs
@@ -122,8 +122,6 @@ namespace UnrealBuildTool.Rules
 					foreach (string LibPath in Directory.EnumerateFiles(USDBinDir, "*.dylib", SearchOption.AllDirectories))
 					{
 						PublicAdditionalLibraries.Add(LibPath);
-						
-						RuntimeDependencies.Add(Path.Combine("$(TargetOutputDir)", Path.GetFileName(LibPath)), LibPath);
 					}
 				}
 
