USD: Fix crash when trying to delete GeometryCaches.

It was moving a TStrongObjectPtr to the UUsdAssetCache into a static lifetime lambda, so that the cache object itself could never be collected again.

diff --git a/Engine/Plugins/Importers/USDImporter/Source/USDSchemas/Private/USDGeomMeshTranslator.cpp b/Engine/Plugins/Importers/USDImporter/Source/USDSchemas/Private/USDGeomMeshTranslator.cpp
--- a/Engine/Plugins/Importers/USDImporter/Source/USDSchemas/Private/USDGeomMeshTranslator.cpp
+++ b/Engine/Plugins/Importers/USDImporter/Source/USDSchemas/Private/USDGeomMeshTranslator.cpp
@@ -602,8 +602,18 @@ namespace UsdGeomMeshTranslatorImpl
 			// Create and configure a new USDTrack to be added to the GeometryCache
 			UGeometryCacheTrackUsd* UsdTrack = NewObject< UGeometryCacheTrackUsd >( GeometryCache );
 
-			// #ueent_todo: Remove the context from the read function if possible
-			TSharedPtr< FUsdSchemaTranslationContext > ContextPtr( Context );
+			// Pull out this data from the Context because the translation context itself cannot ever go into the read function, since it also holds
+			// a strong pointer to the asset cache. If that cache can't be collected we can't ever delete the stage actor that owns it, or the read function closure.
+			// Note that the read function closure *will* get discarded as soon as the UGeometryCache that owns the track is collected, but moving the data out of the
+			// context like this prevents the AssetCache from being pinned, and allows UGeometryCacheTrackUsd's to run, which will unregister itself and
+			// cause the closure to be destroyed in the first place
+			const UE::FUsdStage Stage = Context->Stage;
+			const FName RenderContext = Context->RenderContext;
+			const bool bAllowInterpretingLODs = Context->bAllowInterpretingLODs;
+			const TMap< FString, TMap< FString, int32 > > MaterialToPrimvarToUVIndex = Context->MaterialToPrimvarToUVIndex
+				? *Context->MaterialToPrimvarToUVIndex
+				: TMap< FString, TMap< FString, int32 > >{};
+
 			UsdTrack->Initialize( [ = ]( FGeometryCacheMeshData& OutMeshData, const FString& InPrimPath, float Time )
 				{
 					// Get MeshDescription associated with the prim
@@ -612,19 +622,16 @@ namespace UsdGeomMeshTranslatorImpl
 					TArray< UsdUtils::FUsdPrimMaterialAssignmentInfo > LODIndexToMaterialInfo;
 
 					UE::FSdfPath PrimPath( *InPrimPath );
-					UE::FUsdPrim Prim = ContextPtr->Stage.GetPrimAtPath( PrimPath );
-
-					TMap< FString, TMap< FString, int32 > > Unused;
-					TMap< FString, TMap< FString, int32 > >* MaterialToPrimvarToUVIndex = ContextPtr->MaterialToPrimvarToUVIndex ? ContextPtr->MaterialToPrimvarToUVIndex : &Unused;
+					UE::FUsdPrim Prim = Stage.GetPrimAtPath( PrimPath );
 
 					UsdGeomMeshTranslatorImpl::LoadMeshDescriptions(
 						pxr::UsdTyped( Prim ),
 						LODIndexToMeshDescription,
 						LODIndexToMaterialInfo,
-						*MaterialToPrimvarToUVIndex,
+						MaterialToPrimvarToUVIndex,
 						pxr::UsdTimeCode( Time ),
-						ContextPtr->bAllowInterpretingLODs,
-						ContextPtr->RenderContext
+						bAllowInterpretingLODs,
+						RenderContext
 					);
 
 					// Convert the MeshDescription to MeshData
