USD: Replace usage of deprecated ComputePolygonTangentsAndNormals

diff --git a/Engine/Plugins/Importers/USDImporter/Source/USDSchemas/Private/USDGeomMeshTranslator.cpp b/Engine/Plugins/Importers/USDImporter/Source/USDSchemas/Private/USDGeomMeshTranslator.cpp
--- a/Engine/Plugins/Importers/USDImporter/Source/USDSchemas/Private/USDGeomMeshTranslator.cpp
+++ b/Engine/Plugins/Importers/USDImporter/Source/USDSchemas/Private/USDGeomMeshTranslator.cpp
@@ -528,7 +528,7 @@ namespace UsdGeomMeshTranslatorImpl
 			// (i.e. StaticMeshBuilder is not available at runtime)
 			// We need polygon info because ComputeTangentsAndNormals uses it to repair the invalid vertex normals/tangents
 			// Can't calculate just the required polygons as ComputeTangentsAndNormals is parallel and we can't guarantee thread-safe access patterns
-			FStaticMeshOperations::ComputePolygonTangentsAndNormals( MeshDescription );
+			FStaticMeshOperations::ComputeTriangleTangentsAndNormals( MeshDescription );
 			FStaticMeshOperations::ComputeTangentsAndNormals( MeshDescription, EComputeNTBsFlags::UseMikkTSpace );
 
 			// Manually set this as it seems the UStaticMesh only sets this whenever the mesh is serialized, which we won't do
