Renamed libtbbmalloc on Linux so that TBB doesn't find it and doesn't use it.
There's a race condition on application shutdown where TBB would try to deallocate a buffer after TBB Malloc has been unloaded.

if private_server removes the last reference to init_once, tbbmalloc gets unloaded and private_server will try to deallocate my_thread_array and crash on: tbb::cache_aligned_allocator<padded_private_worker>().deallocate( my_thread_array, my_n_thread );

diff --git a/Engine/Plugins/Importers/USDImporter/Source/UnrealUSDWrapper/UnrealUSDWrapper.Build.cs b/Engine/Plugins/Importers/USDImporter/Source/UnrealUSDWrapper/UnrealUSDWrapper.Build.cs
--- a/Engine/Plugins/Importers/USDImporter/Source/UnrealUSDWrapper/UnrealUSDWrapper.Build.cs
+++ b/Engine/Plugins/Importers/USDImporter/Source/UnrealUSDWrapper/UnrealUSDWrapper.Build.cs
@@ -76,11 +76,8 @@ namespace UnrealBuildTool.Rules
 					PublicSystemIncludePaths.Add(IntelTBBIncludes);
 					PrivateRuntimeLibraryPaths.Add(IntelTBBBinaries);
 					PublicAdditionalLibraries.Add(Path.Combine(IntelTBBBinaries, "libtbb.so"));
-					PublicAdditionalLibraries.Add(Path.Combine(IntelTBBBinaries, "libtbbmalloc.so"));
 					RuntimeDependencies.Add(Path.Combine(IntelTBBBinaries, "libtbb.so"));
 					RuntimeDependencies.Add(Path.Combine(IntelTBBBinaries, "libtbb.so.2"));
-					RuntimeDependencies.Add(Path.Combine(IntelTBBBinaries, "libtbbmalloc.so"));
-					RuntimeDependencies.Add(Path.Combine(IntelTBBBinaries, "libtbbmalloc.so.2"));
 
 					// Python3
 					PublicIncludePaths.Add(Path.Combine(PythonSourceTPSDir, "include"));
@@ -148,8 +145,7 @@ namespace UnrealBuildTool.Rules
 
 			bool bEnableUsdSdk = (
 				Target.WindowsPlatform.Compiler != WindowsCompiler.Clang &&
-				Target.WindowsPlatform.StaticAnalyzer == WindowsStaticAnalyzer.None &&
-				Target.Platform != UnrealTargetPlatform.Linux
+				Target.WindowsPlatform.StaticAnalyzer == WindowsStaticAnalyzer.None
 			);
 
 			// Don't enable USD when running the include tool because it has issues parsing Boost headers
