USD: Temporarily disable runtime support

diff --git a/Engine/Plugins/Importers/USDImporter/Source/UnrealUSDWrapper/UnrealUSDWrapper.Build.cs b/Engine/Plugins/Importers/USDImporter/Source/UnrealUSDWrapper/UnrealUSDWrapper.Build.cs
--- a/Engine/Plugins/Importers/USDImporter/Source/UnrealUSDWrapper/UnrealUSDWrapper.Build.cs
+++ b/Engine/Plugins/Importers/USDImporter/Source/UnrealUSDWrapper/UnrealUSDWrapper.Build.cs
@@ -21,7 +21,8 @@ namespace UnrealBuildTool.Rules
 				}
 			);
 
-			if (EnableUsdSdk(Target))
+			// Temporarily disabled runtime USD support until Mac and Linux dynamic linking issues are resolved
+			if (EnableUsdSdk(Target) && (Target.Type == TargetType.Editor || Target.Platform == UnrealTargetPlatform.Win64))
 			{
 				PublicDependencyModuleNames.Add("Python3");
 
