USD: Avoid trying to make a relative path if the anchor is empty.

Fixes some issues where e.g. we add a reference on a prim, but we haven't saved our root layer yet.

diff --git a/Engine/Plugins/Importers/USDImporter/Source/USDUtilities/Private/USDConversionUtils.cpp b/Engine/Plugins/Importers/USDImporter/Source/USDUtilities/Private/USDConversionUtils.cpp
--- a/Engine/Plugins/Importers/USDImporter/Source/USDUtilities/Private/USDConversionUtils.cpp
+++ b/Engine/Plugins/Importers/USDImporter/Source/USDUtilities/Private/USDConversionUtils.cpp
@@ -789,10 +789,19 @@ void UsdUtils::AddReference( UE::FUsdPrim& Prim, const TCHAR* AbsoluteFilePath )
 		}
 	}
 
+	FString RelativePath = AbsoluteFilePath;
+
 	pxr::SdfLayerHandle EditLayer = UsdPrim.GetStage()->GetEditTarget().GetLayer();
 
-	FString RelativePath = AbsoluteFilePath;
-	MakePathRelativeToLayer( UE::FSdfLayer( EditLayer ), RelativePath );
+	std::string RepositoryPath = EditLayer->GetRepositoryPath().empty() ? EditLayer->GetRealPath() : EditLayer->GetRepositoryPath();
+
+	// If we're editing an in-memory stage our root layer may not have a path yet
+	// Giving an empty InRelativeTo to MakePathRelativeTo causes it to use the engine binary
+	if ( !RepositoryPath.empty() )
+	{
+		FString LayerAbsolutePath = UsdToUnreal::ConvertString( RepositoryPath );
+		FPaths::MakePathRelativeTo( RelativePath, *LayerAbsolutePath );
+	}
 
 	References.AddReference( UnrealToUsd::ConvertString( *RelativePath ).Get() );
 #endif // #if USE_USD_SDK
diff --git a/Engine/Plugins/Importers/USDImporter/Source/USDUtilities/Private/USDLayerUtils.cpp b/Engine/Plugins/Importers/USDImporter/Source/USDUtilities/Private/USDLayerUtils.cpp
--- a/Engine/Plugins/Importers/USDImporter/Source/USDUtilities/Private/USDLayerUtils.cpp
+++ b/Engine/Plugins/Importers/USDImporter/Source/USDUtilities/Private/USDLayerUtils.cpp
@@ -57,7 +57,10 @@ bool UsdUtils::InsertSubLayer( const TUsdStore< pxr::SdfLayerRefPtr >& ParentLay
 	FString LayerFilePath = UsdToUnreal::ConvertString( UsdLayerFilePath );
 
 	FString SubLayerFilePath = FPaths::ConvertRelativePathToFull( SubLayerFile );
-	FPaths::MakePathRelativeTo( SubLayerFilePath, *LayerFilePath );
+	if ( !LayerFilePath.IsEmpty() )
+	{
+		FPaths::MakePathRelativeTo( SubLayerFilePath, *LayerFilePath );
+	}
 
 	ParentLayer.Get()->InsertSubLayerPath( UnrealToUsd::ConvertString( *SubLayerFilePath ).Get(), Index );
 
diff --git a/Engine/Plugins/Importers/USDImporter/Source/USDUtilities/Private/USDShadeConversion.cpp b/Engine/Plugins/Importers/USDImporter/Source/USDUtilities/Private/USDShadeConversion.cpp
--- a/Engine/Plugins/Importers/USDImporter/Source/USDUtilities/Private/USDShadeConversion.cpp
+++ b/Engine/Plugins/Importers/USDImporter/Source/USDUtilities/Private/USDShadeConversion.cpp
@@ -1197,7 +1197,10 @@ namespace UsdShadeConversionImpl
 
 				pxr::UsdShadeInput TextureFileInput = UsdUVTextureShader.CreateInput( UnrealIdentifiers::File, pxr::SdfValueTypeNames->Asset );
 				FString TextureRelativePath = *TextureFilePath;
-				FPaths::MakePathRelativeTo( TextureRelativePath, *UsdFilePath );
+				if ( !UsdFilePath.IsEmpty() )
+				{
+					FPaths::MakePathRelativeTo( TextureRelativePath, *UsdFilePath );
+				}
 				TextureFileInput.Set( pxr::SdfAssetPath( UnrealToUsd::ConvertString( *TextureRelativePath ).Get() ) );
 
 				pxr::UsdShadeInput TextureStInput = UsdUVTextureShader.CreateInput( UnrealIdentifiers::St, pxr::SdfValueTypeNames->Float2 );
@@ -1378,7 +1381,6 @@ bool UsdToUnreal::ConvertMaterial( const pxr::UsdShadeMaterial& UsdShadeMaterial
 		{
 			Material.BaseColor.Expression = Expression;
 			SetOutputIndex( ParameterValue, Material.BaseColor.OutputIndex );
-			
 			bHasMaterialInfo = true;
 		}
 	}
