Added missing PreEditChange and PosteditChange when importing a texture from usd.

diff --git a/Engine/Plugins/Importers/USDImporter/Source/USDUtilities/Private/USDShadeConversion.cpp b/Engine/Plugins/Importers/USDImporter/Source/USDUtilities/Private/USDShadeConversion.cpp
--- a/Engine/Plugins/Importers/USDImporter/Source/USDUtilities/Private/USDShadeConversion.cpp
+++ b/Engine/Plugins/Importers/USDImporter/Source/USDUtilities/Private/USDShadeConversion.cpp
@@ -401,10 +401,12 @@ namespace UsdShadeConversionImpl
 
 					if ( Texture )
 					{
-						if ( bIsNormalInput )
+						if ( bIsNormalInput && ( Texture->SRGB || Texture->CompressionSettings != TextureCompressionSettings::TC_Normalmap ) )
 						{
+							Texture->PreEditChange(nullptr);
 							Texture->SRGB = false;
 							Texture->CompressionSettings = TextureCompressionSettings::TC_Normalmap;
+							Texture->PostEditChange();
 						}
 
 						FString PrimvarName = GetPrimvarUsedAsST(Source);
