Fix loading of USDWrapper dynamic library on macOS.

#rnx

diff --git a/Engine/Plugins/Importers/USDImporter/Source/UnrealUSDWrapper/UnrealUSDWrapper.Build.cs b/Engine/Plugins/Importers/USDImporter/Source/UnrealUSDWrapper/UnrealUSDWrapper.Build.cs
--- a/Engine/Plugins/Importers/USDImporter/Source/UnrealUSDWrapper/UnrealUSDWrapper.Build.cs
+++ b/Engine/Plugins/Importers/USDImporter/Source/UnrealUSDWrapper/UnrealUSDWrapper.Build.cs
@@ -2,6 +2,7 @@
 
 using System.Linq;
 using System.IO;
+using System.Collections.Generic;
 
 namespace UnrealBuildTool.Rules
 {
@@ -104,22 +105,37 @@ namespace UnrealBuildTool.Rules
 				{
 					PublicDefinitions.Add("USD_USES_SYSTEM_MALLOC=0");
 
+					List<string> RuntimeModulePaths = new List<string>();
+
 					// TBB
-					RuntimeDependencies.Add(Path.Combine(IntelTBBBinaries, "libtbb.dylib"));
-					RuntimeDependencies.Add(Path.Combine(IntelTBBBinaries, "libtbbmalloc.dylib"));
+					RuntimeModulePaths.Add(Path.Combine(IntelTBBBinaries, "libtbb.dylib"));
+					RuntimeModulePaths.Add(Path.Combine(IntelTBBBinaries, "libtbbmalloc.dylib"));
 
 					// Python3
 					PublicIncludePaths.Add(Path.Combine(PythonSourceTPSDir, "include"));
 					PublicSystemLibraryPaths.Add(Path.Combine(PythonBinaryTPSDir, "lib"));
 					PrivateRuntimeLibraryPaths.Add(Path.Combine(PythonBinaryTPSDir, "bin"));
-					PublicAdditionalLibraries.Add(Path.Combine(PythonBinaryTPSDir, "lib", "libpython3.7.dylib"));
+					RuntimeModulePaths.Add(Path.Combine(PythonBinaryTPSDir, "lib", "libpython3.7.dylib"));
 
 					// USD
 					PublicIncludePaths.Add(Path.Combine(ModuleDirectory, "..", "ThirdParty", "USD", "include"));
 					var USDBinDir = Path.Combine(ModuleDirectory, "..", "ThirdParty", "Mac", "bin");
 					foreach (string LibPath in Directory.EnumerateFiles(USDBinDir, "*.dylib", SearchOption.AllDirectories))
 					{
-						PublicAdditionalLibraries.Add(LibPath);
+						RuntimeModulePaths.Add(LibPath);
+					}
+
+					foreach (string RuntimeModulePath in RuntimeModulePaths)
+					{
+						if (!File.Exists(RuntimeModulePath))
+						{
+							string Err = string.Format("USD SDK module '{0}' not found.", RuntimeModulePath);
+							System.Console.WriteLine(Err);
+							throw new BuildException(Err);
+						}
+
+						PublicDelayLoadDLLs.Add(RuntimeModulePath);
+						RuntimeDependencies.Add(RuntimeModulePath);
 					}
 				}
 
