USD: Compile fix for targets without editor

diff --git a/Engine/Plugins/Importers/USDImporter/Source/USDUtilities/Private/USDShadeConversion.cpp b/Engine/Plugins/Importers/USDImporter/Source/USDUtilities/Private/USDShadeConversion.cpp
--- a/Engine/Plugins/Importers/USDImporter/Source/USDUtilities/Private/USDShadeConversion.cpp
+++ b/Engine/Plugins/Importers/USDImporter/Source/USDUtilities/Private/USDShadeConversion.cpp
@@ -403,10 +403,14 @@ namespace UsdShadeConversionImpl
 					{
 						if ( bIsNormalInput && ( Texture->SRGB || Texture->CompressionSettings != TextureCompressionSettings::TC_Normalmap ) )
 						{
+#if WITH_EDITOR
 							Texture->PreEditChange(nullptr);
+#endif // WITH_EDITOR
 							Texture->SRGB = false;
 							Texture->CompressionSettings = TextureCompressionSettings::TC_Normalmap;
+#if WITH_EDITOR
 							Texture->PostEditChange();
+#endif // WITH_EDITOR
 						}
 
 						FString PrimvarName = GetPrimvarUsedAsST(Source);
