fix bad resolve due to ingore + edit in other branch causing build break. (p4 copy from RES)

diff --git a/Engine/Plugins/Importers/USDImporter/Source/USDUtilities/Public/USDShadeConversion.h b/Engine/Plugins/Importers/USDImporter/Source/USDUtilities/Public/USDShadeConversion.h
--- a/Engine/Plugins/Importers/USDImporter/Source/USDUtilities/Public/USDShadeConversion.h
+++ b/Engine/Plugins/Importers/USDImporter/Source/USDUtilities/Public/USDShadeConversion.h
@@ -25,7 +25,14 @@ class UUsdAssetCache;
 class UMaterial;
 class UMaterialOptions;
 class UTexture;
+struct FFlattenMaterial;
 struct FPropertyEntry;
+enum class EFlattenMaterialProperties : uint8;
+enum EMaterialProperty;
+namespace UE
+{
+	class FUsdPrim;
+}
 
 namespace UsdToUnreal
 {
@@ -68,6 +75,19 @@ namespace UnrealToUsd
 	 * @return Whether the conversion was successful or not.
 	 */
 	USDUTILITIES_API bool ConvertMaterialToBakedSurface( const UMaterialInterface& InMaterial, const TArray<FPropertyEntry>& InMaterialProperties, const FIntPoint& InDefaultTextureSize, const FDirectoryPath& InTexturesDir, pxr::UsdPrim& OutUsdShadeMaterialPrim );
+
+	/**
+	 * Converts a flattened material's data into textures placed at InTexturesDir, and configures OutUsdShadeMaterial to use the baked textures.
+	 * Note that to avoid a potentially useless copy, InMaterial's samples will be modified in place to have 255 alpha before being exported to textures.
+	 *
+	 * @param MaterialName - Name of the material, used as prefix on the exported texture filenames
+	 * @param InMaterial - Source material to bake
+	 * @param InMaterialProperties - Object used *exclusively* to provide floating point constant values if necessary
+	 * @param InTexturesDir - Directory where the baked textures will be placed
+	 * @param OutUsdShadeMaterialPrim - UsdPrim with the UsdShadeMaterial schema that will be configured to use the baked textures and constants
+	 * @return Whether the conversion was successful or not.
+	 */
+	USDUTILITIES_API bool ConvertFlattenMaterial( const FString& InMaterialName, FFlattenMaterial& InMaterial, const TArray<FPropertyEntry>& InMaterialProperties, const FDirectoryPath& InTexturesDir, UE::FUsdPrim& OutUsdShadeMaterialPrim );
 }
 #endif // WITH_EDITOR
 
@@ -86,6 +106,20 @@ namespace UsdUtils
 
 	/** Creates a texture from a pxr::SdfAssetPath attribute. PrimPath is optional, and should point to the source shadematerial prim path. It will be placed in its UUsdAssetImportData */
 	USDUTILITIES_API UTexture* CreateTexture( const pxr::UsdAttribute& TextureAssetPathAttr, const FString& PrimPath = FString(), TextureGroup LODGroup = TEXTUREGROUP_World, UObject* Outer = GetTransientPackage() );
+
+#if WITH_EDITOR
+	/** Convert between the two different types used to represent material channels to bake */
+	USDUTILITIES_API EFlattenMaterialProperties MaterialPropertyToFlattenProperty( EMaterialProperty MaterialProperty );
+
+	/** Convert between the two different types used to represent material channels to bake */
+	USDUTILITIES_API EMaterialProperty FlattenPropertyToMaterialProperty( EFlattenMaterialProperties FlattenProperty );
+#endif // WITH_EDITOR
+
+	/** Converts channels that have the same value for every pixel into a channel that only has a single pixel with that value */
+	USDUTILITIES_API void CollapseConstantChannelsToSinglePixel( FFlattenMaterial& InMaterial );
+
+	/** Temporary function until UnrealWrappers can create attributes, just adds a custom bool attribute 'worldSpaceNormals' as true */
+	USDUTILITIES_API bool MarkMaterialPrimWithWorldSpaceNormals( const UE::FUsdPrim& MaterialPrim );
 }
 
 #endif // #if USE_USD_SDK
