USD: Manual merge of cl 16452236 from //UE5/Release-Engine-Staging to fix a CIS issue.

diff --git a/Engine/Plugins/Importers/USDImporter/Source/USDExporter/Private/AnimSequenceExporterUSD.cpp b/Engine/Plugins/Importers/USDImporter/Source/USDExporter/Private/AnimSequenceExporterUSD.cpp
--- a/Engine/Plugins/Importers/USDImporter/Source/USDExporter/Private/AnimSequenceExporterUSD.cpp
+++ b/Engine/Plugins/Importers/USDImporter/Source/USDExporter/Private/AnimSequenceExporterUSD.cpp
@@ -78,7 +78,10 @@ bool UAnimSequenceExporterUSD::ExportBinary( UObject* Object, const TCHAR* Type,
 		FString ExtensionPart;
 		FPaths::Split( PayloadFilename, PathPart, FilenamePart, ExtensionPart );
 
-		ExtensionPart = Options->PreviewMeshOptions.PayloadFormat;
+		if ( FormatExtension.Contains( Options->PreviewMeshOptions.PayloadFormat ) )
+		{
+			ExtensionPart = Options->PreviewMeshOptions.PayloadFormat;
+		}
 
 		PayloadFilename = FPaths::Combine( PathPart, FilenamePart + TEXT( "_payload." ) + ExtensionPart );
 	}
diff --git a/Engine/Plugins/Importers/USDImporter/Source/USDExporter/Private/SkeletalMeshExporterUSD.cpp b/Engine/Plugins/Importers/USDImporter/Source/USDExporter/Private/SkeletalMeshExporterUSD.cpp
--- a/Engine/Plugins/Importers/USDImporter/Source/USDExporter/Private/SkeletalMeshExporterUSD.cpp
+++ b/Engine/Plugins/Importers/USDImporter/Source/USDExporter/Private/SkeletalMeshExporterUSD.cpp
@@ -75,7 +75,10 @@ bool USkeletalMeshExporterUsd::ExportBinary( UObject* Object, const TCHAR* Type,
 		FString ExtensionPart;
 		FPaths::Split( PayloadFilename, PathPart, FilenamePart, ExtensionPart );
 
-		ExtensionPart = Options->Inner.PayloadFormat;
+		if ( FormatExtension.Contains( Options->Inner.PayloadFormat ) )
+		{
+			ExtensionPart = Options->Inner.PayloadFormat;
+		}
 
 		PayloadFilename = FPaths::Combine( PathPart, FilenamePart + TEXT( "_payload." ) + ExtensionPart );
 	}
diff --git a/Engine/Plugins/Importers/USDImporter/Source/USDExporter/Private/StaticMeshExporterUSD.cpp b/Engine/Plugins/Importers/USDImporter/Source/USDExporter/Private/StaticMeshExporterUSD.cpp
--- a/Engine/Plugins/Importers/USDImporter/Source/USDExporter/Private/StaticMeshExporterUSD.cpp
+++ b/Engine/Plugins/Importers/USDImporter/Source/USDExporter/Private/StaticMeshExporterUSD.cpp
@@ -85,7 +85,10 @@ bool UStaticMeshExporterUsd::ExportBinary( UObject* Object, const TCHAR* Type, F
 		FString ExtensionPart;
 		FPaths::Split( PayloadFilename, PathPart, FilenamePart, ExtensionPart );
 
-		ExtensionPart = Options->PayloadFormat;
+		if ( FormatExtension.Contains( Options->PayloadFormat ) )
+		{
+			ExtensionPart = Options->PayloadFormat;
+		}
 
 		PayloadFilename = FPaths::Combine( PathPart, FilenamePart + TEXT( "_payload." ) + ExtensionPart );
 	}
