USD: Fix crash when stage actor destructor is called after garbage collection.

diff --git a/Engine/Plugins/Importers/USDImporter/Source/USDStage/Private/USDStageActor.cpp b/Engine/Plugins/Importers/USDImporter/Source/USDStage/Private/USDStageActor.cpp
--- a/Engine/Plugins/Importers/USDImporter/Source/USDStage/Private/USDStageActor.cpp
+++ b/Engine/Plugins/Importers/USDImporter/Source/USDStage/Private/USDStageActor.cpp
@@ -901,38 +901,6 @@ void AUsdStageActor::OnUsdObjectsChanged( const UsdUtils::FObjectChangesByPath&
 	}
 }
 
-AUsdStageActor::~AUsdStageActor()
-{
-#if WITH_EDITOR
-	if ( !IsEngineExitRequested() && HasAuthorityOverStage() )
-	{
-		FEditorDelegates::BeginPIE.RemoveAll(this);
-		FEditorDelegates::PostPIEStarted.RemoveAll(this);
-		FUsdDelegates::OnPostUsdImport.RemoveAll(this);
-		FUsdDelegates::OnPreUsdImport.RemoveAll(this);
-		if ( UTransBuffer* TransBuffer = GUnrealEd ? Cast<UTransBuffer>( GUnrealEd->Trans ) : nullptr )
-		{
-			TransBuffer->OnTransactionStateChanged().RemoveAll( this );
-			TransBuffer->OnRedo().Remove( OnRedoHandle );
-		}
-
-		GEngine->OnLevelActorDeleted().AddUObject( this, &AUsdStageActor::OnLevelActorDeleted );
-
-		// This clears the SUSDStage window whenever the level we're currently in gets destroyed.
-		// Note that this is not called when deleting from the Editor, as the actor goes into the undo buffer.
-		OnActorDestroyed.Broadcast();
-		CloseUsdStage();
-
-		if ( RootUsdTwin )
-		{
-			RootUsdTwin->Clear();
-		}
-
-		FCoreUObjectDelegates::OnObjectsReplaced.RemoveAll( this );
-	}
-#endif // WITH_EDITOR
-}
-
 USDSTAGE_API void AUsdStageActor::Reset()
 {
 	Super::Reset();
@@ -1856,6 +1824,40 @@ void AUsdStageActor::PostActorCreated()
 	}
 }
 
+void AUsdStageActor::BeginDestroy()
+{
+#if WITH_EDITOR
+	if ( !IsEngineExitRequested() && HasAuthorityOverStage() )
+	{
+		FEditorDelegates::BeginPIE.RemoveAll( this );
+		FEditorDelegates::PostPIEStarted.RemoveAll( this );
+		FUsdDelegates::OnPostUsdImport.RemoveAll( this );
+		FUsdDelegates::OnPreUsdImport.RemoveAll( this );
+		if ( UTransBuffer* TransBuffer = GUnrealEd ? Cast<UTransBuffer>( GUnrealEd->Trans ) : nullptr )
+		{
+			TransBuffer->OnTransactionStateChanged().RemoveAll( this );
+			TransBuffer->OnRedo().Remove( OnRedoHandle );
+		}
+
+		GEngine->OnLevelActorDeleted().RemoveAll( this );
+		FCoreUObjectDelegates::OnObjectsReplaced.RemoveAll( this );
+	}
+
+	// This clears the SUSDStage window whenever the level we're currently in gets destroyed.
+	// Note that this is not called when deleting from the Editor, as the actor goes into the undo buffer.
+	OnActorDestroyed.Broadcast();
+	CloseUsdStage();
+
+	// If our prims are already destroyed then likely the entire map has been destroyed anyway, so don't need to clear it
+	if ( RootUsdTwin && !RootUsdTwin->HasAnyFlags( RF_BeginDestroyed ) )
+	{
+		RootUsdTwin->Clear();
+	}
+#endif // WITH_EDITOR
+
+	Super::BeginDestroy();
+}
+
 void AUsdStageActor::PostRegisterAllComponents()
 {
 	Super::PostRegisterAllComponents();
diff --git a/Engine/Plugins/Importers/USDImporter/Source/USDStage/Public/USDStageActor.h b/Engine/Plugins/Importers/USDImporter/Source/USDStage/Public/USDStageActor.h
--- a/Engine/Plugins/Importers/USDImporter/Source/USDStage/Public/USDStageActor.h
+++ b/Engine/Plugins/Importers/USDImporter/Source/USDStage/Public/USDStageActor.h
@@ -131,7 +131,6 @@ public:
 
 public:
 	AUsdStageActor();
-	virtual ~AUsdStageActor();
 
 	USDSTAGE_API void Reset() override;
 	void Refresh() const;
@@ -152,6 +151,7 @@ public:
 	virtual void Serialize(FArchive& Ar) override;
 	virtual void Destroyed() override;
 	virtual void PostActorCreated() override;
+	virtual void BeginDestroy() override;
 
 	virtual void PostRegisterAllComponents() override;
 	virtual void PostUnregisterAllComponents() override;
