USD: Fix Action->Import ignoring source stage actor transform;

diff --git a/Engine/Plugins/Importers/USDImporter/Source/USDStageEditorViewModels/Private/USDStageViewModel.cpp b/Engine/Plugins/Importers/USDImporter/Source/USDStageEditorViewModels/Private/USDStageViewModel.cpp
--- a/Engine/Plugins/Importers/USDImporter/Source/USDStageEditorViewModels/Private/USDStageViewModel.cpp
+++ b/Engine/Plugins/Importers/USDImporter/Source/USDStageEditorViewModels/Private/USDStageViewModel.cpp
@@ -263,6 +263,9 @@ void FUsdStageViewModel::ImportStage()
 			ImportContext.AssetCache = StageActor->GetAssetCache();
 			ImportContext.MaterialToPrimvarToUVIndex = StageActor->GetMaterialToPrimvarToUVIndex();
 
+			ImportContext.TargetSceneActorAttachParent = StageActor->GetRootComponent()->GetAttachParent();
+			ImportContext.TargetSceneActorTargetTransform = StageActor->GetActorTransform();
+
 			UUsdStageImporter* USDImporter = IUsdStageImporterModule::Get().GetImporter();
 			USDImporter->ImportFromFile(ImportContext);
 
diff --git a/Engine/Plugins/Importers/USDImporter/Source/USDStageImporter/Private/USDStageImportContext.cpp b/Engine/Plugins/Importers/USDImporter/Source/USDStageImporter/Private/USDStageImportContext.cpp
--- a/Engine/Plugins/Importers/USDImporter/Source/USDStageImporter/Private/USDStageImportContext.cpp
+++ b/Engine/Plugins/Importers/USDImporter/Source/USDStageImporter/Private/USDStageImportContext.cpp
@@ -25,6 +25,8 @@ FUsdStageImportContext::FUsdStageImportContext()
 	ImportedAsset = nullptr;
 	AssetCache = nullptr;
 	OriginalMetersPerUnit = 0.01f;
+	TargetSceneActorTargetTransform = FTransform::Identity;
+	TargetSceneActorAttachParent = nullptr;
 }
 
 bool FUsdStageImportContext::Init(const FString& InName, const FString& InFilePath, const FString& InInitialPackagePath, EObjectFlags InFlags, bool bInIsAutomated, bool bIsReimport, bool bAllowActorImport)
diff --git a/Engine/Plugins/Importers/USDImporter/Source/USDStageImporter/Private/USDStageImporter.cpp b/Engine/Plugins/Importers/USDImporter/Source/USDStageImporter/Private/USDStageImporter.cpp
--- a/Engine/Plugins/Importers/USDImporter/Source/USDStageImporter/Private/USDStageImporter.cpp
+++ b/Engine/Plugins/Importers/USDImporter/Source/USDStageImporter/Private/USDStageImporter.cpp
@@ -185,6 +185,13 @@ namespace UsdStageImporterImpl
 			RootComponent->RegisterComponent();
 		}
 
+		if ( ImportContext.TargetSceneActorAttachParent )
+		{
+			RootComponent->AttachToComponent( ImportContext.TargetSceneActorAttachParent, FAttachmentTransformRules::KeepRelativeTransform );
+		}
+
+		Actor->SetActorTransform( ImportContext.TargetSceneActorTargetTransform );
+
 		ImportContext.SceneActor = Actor;
 	}
 
diff --git a/Engine/Plugins/Importers/USDImporter/Source/USDStageImporter/Public/USDStageImportContext.h b/Engine/Plugins/Importers/USDImporter/Source/USDStageImporter/Public/USDStageImportContext.h
--- a/Engine/Plugins/Importers/USDImporter/Source/USDStageImporter/Public/USDStageImportContext.h
+++ b/Engine/Plugins/Importers/USDImporter/Source/USDStageImporter/Public/USDStageImportContext.h
@@ -24,6 +24,13 @@ struct USDSTAGEIMPORTER_API FUsdStageImportContext
 
 	UWorld* World;
 
+	/**
+	 * Whenever we spawn the scene actor, it should have this local transform and be attached to this parent.
+	 * We use this so that Actions->Import can spawn the scene actor exactly where the original stage actor was
+	 */
+	FTransform TargetSceneActorTargetTransform;
+	USceneComponent* TargetSceneActorAttachParent;
+
 	/** Spawned actor that contains the imported scene as a child hierarchy */
 	UPROPERTY()
 	AActor* SceneActor;
