USD: Fix reloading not resyncing prims anymore;

Fix worldSpaceNormal custom attribute not contributing to material hash.

diff --git a/Engine/Plugins/Importers/USDImporter/Source/USDUtilities/Private/USDShadeConversion.cpp b/Engine/Plugins/Importers/USDImporter/Source/USDUtilities/Private/USDShadeConversion.cpp
--- a/Engine/Plugins/Importers/USDImporter/Source/USDUtilities/Private/USDShadeConversion.cpp
+++ b/Engine/Plugins/Importers/USDImporter/Source/USDUtilities/Private/USDShadeConversion.cpp
@@ -2056,8 +2056,14 @@ FSHAHash UsdUtils::HashShadeMaterial( const pxr::UsdShadeMaterial& UsdShadeMater
 		UsdShadeConversionImpl::HashShadeInput( ShadeInput, HashState );
 	}
 
-	FSHAHash OutHash;
+	bool bValue = false;
+	if ( pxr::UsdAttribute Attr = UsdShadeMaterial.GetPrim().GetAttribute( UnrealIdentifiers::WorldSpaceNormals ) )
+	{
+		Attr.Get<bool>( &bValue );
+	}
+	HashState.Update( reinterpret_cast< uint8* >( &bValue ), sizeof( bValue ) );
 
+	FSHAHash OutHash;
 	HashState.Final();
 	HashState.GetHash( &OutHash.Hash[0] );
 
diff --git a/Engine/Plugins/Importers/USDImporter/Source/UnrealUSDWrapper/Private/USDListener.cpp b/Engine/Plugins/Importers/USDImporter/Source/UnrealUSDWrapper/Private/USDListener.cpp
--- a/Engine/Plugins/Importers/USDImporter/Source/UnrealUSDWrapper/Private/USDListener.cpp
+++ b/Engine/Plugins/Importers/USDImporter/Source/UnrealUSDWrapper/Private/USDListener.cpp
@@ -145,6 +145,27 @@ namespace UsdToUnreal
 		ConvertPathRange( InNotice.GetChangedInfoOnlyPaths(), OutInfoChanges );
 		ConvertPathRange( InNotice.GetResyncedPaths(), OutResyncChanges );
 
+		// Upgrade info changes with content reloads into resync changes
+		for ( TPair< FString, TArray<UsdUtils::FObjectChangeNotice> >& InfoPair : OutInfoChanges )
+		{
+			const FString& PrimPath = InfoPair.Key;
+			TArray< UsdUtils::FObjectChangeNotice >* AnalogueResyncChanges = nullptr;
+
+			for ( TArray<UsdUtils::FObjectChangeNotice>::TIterator ChangeIt = InfoPair.Value.CreateIterator(); ChangeIt; ++ChangeIt )
+			{
+				if ( ChangeIt->Flags.bDidReloadContent )
+				{
+					if ( !AnalogueResyncChanges )
+					{
+						AnalogueResyncChanges = &OutResyncChanges.FindOrAdd(PrimPath);
+					}
+
+					AnalogueResyncChanges->Add( *ChangeIt );
+					ChangeIt.RemoveCurrent();
+				}
+			}
+		}
+
 		return true;
 	}
 #endif // USE_USD_SDK
