USD: Fix missing inverse transform compensation when parent prims are lights or cameras.

diff --git a/Engine/Plugins/Importers/USDImporter/Source/USDUtilities/Private/USDPrimConversion.cpp b/Engine/Plugins/Importers/USDImporter/Source/USDUtilities/Private/USDPrimConversion.cpp
--- a/Engine/Plugins/Importers/USDImporter/Source/USDUtilities/Private/USDPrimConversion.cpp
+++ b/Engine/Plugins/Importers/USDImporter/Source/USDUtilities/Private/USDPrimConversion.cpp
@@ -62,36 +62,38 @@ bool UsdToUnreal::ConvertXformable( const pxr::UsdStageRefPtr& Stage, const pxr:
 	Xformable.GetLocalTransformation( &UsdMatrix, &bResetXFormStack, EvalTime );
 
 	FUsdStageInfo StageInfo( Stage );
+	OutTransform = UsdToUnreal::ConvertMatrix( StageInfo, UsdMatrix );
 
 	// Extra rotation to match different camera facing direction convention
 	// Note: The camera space is always Y-up, yes, but this is not what this is: This is the camera's transform wrt the stage,
 	// which follows the stage up axis
-	FRotator AdditionalRotation( ForceInit );
-	if ( Xformable.GetPrim().IsA< pxr::UsdGeomCamera >() )
+	if ( Xformable.GetPrim().IsA< pxr::UsdGeomCamera >() || Xformable.GetPrim().IsA< pxr::UsdLuxLight >() )
 	{
 		if ( StageInfo.UpAxis == EUsdUpAxis::YAxis )
 		{
-			AdditionalRotation = FRotator( 0.0f, -90.f, 0.0f );
+			OutTransform = FTransform( FRotator( 0.0f, -90.f, 0.0f ) ) * OutTransform;
 		}
 		else
 		{
-			AdditionalRotation = FRotator( -90.0f, -90.f, 0.0f );
+			OutTransform = FTransform( FRotator( -90.0f, -90.f, 0.0f ) ) * OutTransform;
 		}
 	}
-	else if ( Xformable.GetPrim().IsA< pxr::UsdLuxLight >() )
+	// Invert the compensation applied to our parents, in case they're a camera or a light
+	if ( pxr::UsdPrim Parent = Xformable.GetPrim().GetParent() )
 	{
-		if ( StageInfo.UpAxis == EUsdUpAxis::YAxis )
+		if ( Parent.IsA< pxr::UsdGeomCamera >() || Parent.IsA< pxr::UsdLuxLight >() )
 		{
-			AdditionalRotation = FRotator( 0.0f, -90.f, 0.0f );
-		}
-		else
-		{
-			AdditionalRotation = FRotator( -90.0f, -90.f, 0.0f );
+			if ( StageInfo.UpAxis == EUsdUpAxis::YAxis )
+			{
+				OutTransform = OutTransform * FTransform( FRotator( 0.0f, -90.f, 0.0f ).GetInverse() );
+			}
+			else
+			{
+				OutTransform = OutTransform * FTransform( FRotator( -90.0f, -90.f, 0.0f ).GetInverse() );
+			}
 		}
 	}
 
-	OutTransform = FTransform( AdditionalRotation ) * UsdToUnreal::ConvertMatrix( StageInfo, UsdMatrix );
-
 	return true;
 }
 
