USD: Fix silent incorrect comparison between UsdWrapper types

diff --git a/Engine/Plugins/Importers/USDImporter/Source/USDStage/Private/USDStageActor.cpp b/Engine/Plugins/Importers/USDImporter/Source/USDStage/Private/USDStageActor.cpp
--- a/Engine/Plugins/Importers/USDImporter/Source/USDStage/Private/USDStageActor.cpp
+++ b/Engine/Plugins/Importers/USDImporter/Source/USDStage/Private/USDStageActor.cpp
@@ -203,7 +203,7 @@ struct FUsdStageActorImpl
 			if ( const AUsdStageActor* Actor = Cast<const AUsdStageActor>( Instance ) )
 			{
 				const UE::FUsdStage& OtherStage = Actor->GetUsdStage();
-				if ( OtherStage && RootLayer == OtherStage )
+				if ( OtherStage && Stage == OtherStage )
 				{
 					// Some other actor is still using our stage, so don't close it
 					return;
diff --git a/Engine/Plugins/Importers/USDImporter/Source/UnrealUSDWrapper/Private/UsdWrappers/UsdStage.cpp b/Engine/Plugins/Importers/USDImporter/Source/UnrealUSDWrapper/Private/UsdWrappers/UsdStage.cpp
--- a/Engine/Plugins/Importers/USDImporter/Source/UnrealUSDWrapper/Private/UsdWrappers/UsdStage.cpp
+++ b/Engine/Plugins/Importers/USDImporter/Source/UnrealUSDWrapper/Private/UsdWrappers/UsdStage.cpp
@@ -85,6 +85,20 @@ namespace UE
 #endif // #if USE_USD_SDK
 	}
 
+	bool FUsdStage::operator==( const FUsdStage& Other ) const
+	{
+#if USE_USD_SDK
+		return Impl->PxrUsdStageRefPtr.Get() == Other.Impl->PxrUsdStageRefPtr.Get();
+#else
+		return false;
+#endif // #if USE_USD_SDK
+	}
+
+	bool FUsdStage::operator!=( const FUsdStage& Other ) const
+	{
+		return !( *this == Other );
+	}
+
 	FUsdStage::~FUsdStage()
 	{
 		FScopedUnrealAllocs UnrealAllocs;
diff --git a/Engine/Plugins/Importers/USDImporter/Source/UnrealUSDWrapper/Public/UsdWrappers/SdfLayer.h b/Engine/Plugins/Importers/USDImporter/Source/UnrealUSDWrapper/Public/UsdWrappers/SdfLayer.h
--- a/Engine/Plugins/Importers/USDImporter/Source/UnrealUSDWrapper/Public/UsdWrappers/SdfLayer.h
+++ b/Engine/Plugins/Importers/USDImporter/Source/UnrealUSDWrapper/Public/UsdWrappers/SdfLayer.h
@@ -60,7 +60,7 @@ namespace UE
 		bool operator==( const FSdfLayer& Other ) const;
 		bool operator!=( const FSdfLayer& Other ) const;
 
-		operator bool() const;
+		explicit operator bool() const;
 
 	// Auto conversion from/to pxr::SdfLayerRefPtr
 	public:
diff --git a/Engine/Plugins/Importers/USDImporter/Source/UnrealUSDWrapper/Public/UsdWrappers/UsdAttribute.h b/Engine/Plugins/Importers/USDImporter/Source/UnrealUSDWrapper/Public/UsdWrappers/UsdAttribute.h
--- a/Engine/Plugins/Importers/USDImporter/Source/UnrealUSDWrapper/Public/UsdWrappers/UsdAttribute.h
+++ b/Engine/Plugins/Importers/USDImporter/Source/UnrealUSDWrapper/Public/UsdWrappers/UsdAttribute.h
@@ -46,7 +46,7 @@ namespace UE
 		bool operator==( const FUsdAttribute& Other ) const;
 		bool operator!=( const FUsdAttribute& Other ) const;
 
-		operator bool() const;
+		explicit operator bool() const;
 
 	// Auto conversion from/to pxr::UsdAttribute
 	public:
diff --git a/Engine/Plugins/Importers/USDImporter/Source/UnrealUSDWrapper/Public/UsdWrappers/UsdGeomXformable.h b/Engine/Plugins/Importers/USDImporter/Source/UnrealUSDWrapper/Public/UsdWrappers/UsdGeomXformable.h
--- a/Engine/Plugins/Importers/USDImporter/Source/UnrealUSDWrapper/Public/UsdWrappers/UsdGeomXformable.h
+++ b/Engine/Plugins/Importers/USDImporter/Source/UnrealUSDWrapper/Public/UsdWrappers/UsdGeomXformable.h
@@ -45,7 +45,7 @@ namespace UE
 
 		~FUsdGeomXformable();
 
-		operator bool() const;
+		explicit operator bool() const;
 
 		// Auto conversion from/to pxr::UsdGeomXformable
 	public:
diff --git a/Engine/Plugins/Importers/USDImporter/Source/UnrealUSDWrapper/Public/UsdWrappers/UsdPrim.h b/Engine/Plugins/Importers/USDImporter/Source/UnrealUSDWrapper/Public/UsdWrappers/UsdPrim.h
--- a/Engine/Plugins/Importers/USDImporter/Source/UnrealUSDWrapper/Public/UsdWrappers/UsdPrim.h
+++ b/Engine/Plugins/Importers/USDImporter/Source/UnrealUSDWrapper/Public/UsdWrappers/UsdPrim.h
@@ -47,7 +47,7 @@ namespace UE
 		bool operator==( const FUsdPrim& Other ) const;
 		bool operator!=( const FUsdPrim& Other ) const;
 
-		operator bool() const;
+		explicit operator bool() const;
 
 	// Auto conversion from/to pxr::UsdPrim
 	public:
diff --git a/Engine/Plugins/Importers/USDImporter/Source/UnrealUSDWrapper/Public/UsdWrappers/UsdStage.h b/Engine/Plugins/Importers/USDImporter/Source/UnrealUSDWrapper/Public/UsdWrappers/UsdStage.h
--- a/Engine/Plugins/Importers/USDImporter/Source/UnrealUSDWrapper/Public/UsdWrappers/UsdStage.h
+++ b/Engine/Plugins/Importers/USDImporter/Source/UnrealUSDWrapper/Public/UsdWrappers/UsdStage.h
@@ -50,7 +50,10 @@ namespace UE
 
 		~FUsdStage();
 
-		operator bool() const;
+		explicit operator bool() const;
+
+		bool operator==( const FUsdStage& Other ) const;
+		bool operator!=( const FUsdStage& Other ) const;
 
 	// Auto conversion from/to pxr::UsdStageRefPtr
 	public:
diff --git a/Engine/Plugins/Importers/USDImporter/Source/UnrealUSDWrapper/Public/UsdWrappers/UsdTyped.h b/Engine/Plugins/Importers/USDImporter/Source/UnrealUSDWrapper/Public/UsdWrappers/UsdTyped.h
--- a/Engine/Plugins/Importers/USDImporter/Source/UnrealUSDWrapper/Public/UsdWrappers/UsdTyped.h
+++ b/Engine/Plugins/Importers/USDImporter/Source/UnrealUSDWrapper/Public/UsdWrappers/UsdTyped.h
@@ -45,7 +45,7 @@ namespace UE
 
 		~FUsdTyped();
 
-		operator bool() const;
+		explicit operator bool() const;
 
 	// Auto conversion from/to pxr::UsdTyped
 	public:
