USD: Don't spawn an AUsdStageActor on the level when just opening the USD Stage editor.

diff --git a/Engine/Plugins/Importers/USDImporter/Source/USDStageEditor/Private/SUSDStage.cpp b/Engine/Plugins/Importers/USDImporter/Source/USDStageEditor/Private/SUSDStage.cpp
--- a/Engine/Plugins/Importers/USDImporter/Source/USDStageEditor/Private/SUSDStage.cpp
+++ b/Engine/Plugins/Importers/USDImporter/Source/USDStageEditor/Private/SUSDStage.cpp
@@ -102,7 +102,7 @@ void SUsdStage::Construct( const FArguments& InArgs )
 	OnViewportSelectionChangedHandle = USelection::SelectionChangedEvent.AddRaw( this, &SUsdStage::OnViewportSelectionChanged );
 
 	IUsdStageModule& UsdStageModule = FModuleManager::Get().LoadModuleChecked< IUsdStageModule >( "UsdStage" );
-	ViewModel.UsdStageActor = &UsdStageModule.GetUsdStageActor( GWorld );
+	ViewModel.UsdStageActor = UsdStageModule.FindUsdStageActor( GWorld );
 
 	bUpdatingViewportSelection = false;
 
