USD: Add public accessor for AUsdStageActor's LevelSequence.

diff --git a/Engine/Plugins/Importers/USDImporter/Source/USDStage/Public/USDStageActor.h b/Engine/Plugins/Importers/USDImporter/Source/USDStage/Public/USDStageActor.h
--- a/Engine/Plugins/Importers/USDImporter/Source/USDStage/Public/USDStageActor.h
+++ b/Engine/Plugins/Importers/USDImporter/Source/USDStage/Public/USDStageActor.h
@@ -70,6 +70,9 @@ public:
 	UFUNCTION(BlueprintCallable, Category = "USD", meta = (CallInEditor = "true"))
 	USDSTAGE_API void SetTime(float InTime);
 
+	UFUNCTION( BlueprintCallable, Category = "USD", meta = ( CallInEditor = "true" ) )
+	USDSTAGE_API ULevelSequence* GetLevelSequence() { return LevelSequence; }
+
 	/**
 	 * Gets the transient component that was generated for a prim with a given prim path.
 	 * Warning: The lifetime of the component is managed by the AUsdStageActor, and it may be force-destroyed at any time (e.g. when closing the stage)
