USD: Fix materials being reimported when reloading USDZ stages that didn't change.

diff --git a/Engine/Plugins/Importers/USDImporter/Source/USDUtilities/Private/USDShadeConversion.cpp b/Engine/Plugins/Importers/USDImporter/Source/USDUtilities/Private/USDShadeConversion.cpp
--- a/Engine/Plugins/Importers/USDImporter/Source/USDUtilities/Private/USDShadeConversion.cpp
+++ b/Engine/Plugins/Importers/USDImporter/Source/USDUtilities/Private/USDShadeConversion.cpp
@@ -1498,8 +1498,22 @@ namespace UE
 				{
 					pxr::VtValue ShadeInputValue;
 					ShadeInput.Get( &ShadeInputValue );
-					uint64 ValueHash = ( uint64 ) ShadeInputValue.GetHash();
-					InOutHashState.Update( reinterpret_cast< uint8* >( &ValueHash ), sizeof( uint64 ) );
+
+					// We have to manually resolve and hash these file paths or else resolvedpaths inside usdz archives will have upper case driver letters
+					// when we first open the stage, but will switch to lower case drive letters if we reload them. This is not something we're doing,
+					// as it happens with a pure USD python script. (this with USD 21.05 in June 2021)
+					if ( ShadeInputValue.IsHolding<pxr::SdfAssetPath>() )
+					{
+						pxr::ArResolver& Resolver = pxr::ArGetResolver();
+						pxr::SdfAssetPath UsdAssetPath = ShadeInputValue.UncheckedGet<pxr::SdfAssetPath>();
+						FString ResolvedPath = UsdToUnreal::ConvertString( Resolver.ComputeNormalizedPath( UsdAssetPath.GetResolvedPath() ) );
+						InOutHashState.UpdateWithString( *ResolvedPath, ResolvedPath.Len() );
+					}
+					else
+					{
+						uint64 ValueHash = ( uint64 ) ShadeInputValue.GetHash();
+						InOutHashState.Update( reinterpret_cast< uint8* >( &ValueHash ), sizeof( uint64 ) );
+					}
 				}
 			}
 		}
