USD: Use prim names for transient asset names as well, when in the editor.

diff --git a/Engine/Plugins/Importers/USDImporter/Source/USDSchemas/Private/USDGeomMeshTranslator.cpp b/Engine/Plugins/Importers/USDImporter/Source/USDSchemas/Private/USDGeomMeshTranslator.cpp
--- a/Engine/Plugins/Importers/USDImporter/Source/USDSchemas/Private/USDGeomMeshTranslator.cpp
+++ b/Engine/Plugins/Importers/USDImporter/Source/USDSchemas/Private/USDGeomMeshTranslator.cpp
@@ -420,7 +420,7 @@ namespace UsdGeomMeshTranslatorImpl
 		}
 	}
 
-	UStaticMesh* CreateStaticMesh( TArray<FMeshDescription>& LODIndexToMeshDescription, FUsdSchemaTranslationContext& Context, bool& bOutIsNew )
+	UStaticMesh* CreateStaticMesh( TArray<FMeshDescription>& LODIndexToMeshDescription, FUsdSchemaTranslationContext& Context, const FString& PrimPath, bool& bOutIsNew )
 	{
 		UStaticMesh* StaticMesh = nullptr;
 
@@ -447,7 +447,8 @@ namespace UsdGeomMeshTranslatorImpl
 		{
 			bOutIsNew = true;
 
-			StaticMesh = NewObject< UStaticMesh >( GetTransientPackage(), NAME_None, Context.ObjectFlags | EObjectFlags::RF_Public );
+			FName AssetName = MakeUniqueObjectName( GetTransientPackage(), UStaticMesh::StaticClass(), *FPaths::GetBaseFilename( PrimPath ) );
+			StaticMesh = NewObject< UStaticMesh >( GetTransientPackage(), AssetName, Context.ObjectFlags | EObjectFlags::RF_Public );
 
 #if WITH_EDITOR
 			for ( int32 LODIndex = 0; LODIndex < LODIndexToMeshDescription.Num(); ++LODIndex )
@@ -597,7 +598,8 @@ namespace UsdGeomMeshTranslatorImpl
 		{
 			bOutIsNew = true;
 
-			GeometryCache = NewObject< UGeometryCache >( GetTransientPackage(), NAME_None, Context->ObjectFlags | EObjectFlags::RF_Public );
+			const FName AssetName = MakeUniqueObjectName( GetTransientPackage(), UGeometryCache::StaticClass(), *FPaths::GetBaseFilename( InPrimPath ) );
+			GeometryCache = NewObject< UGeometryCache >( GetTransientPackage(), AssetName, Context->ObjectFlags | EObjectFlags::RF_Public );
 
 			// Create and configure a new USDTrack to be added to the GeometryCache
 			UGeometryCacheTrackUsd* UsdTrack = NewObject< UGeometryCacheTrackUsd >( GeometryCache );
@@ -805,9 +807,8 @@ void FBuildStaticMeshTaskChain::SetupTasks()
 #endif // WITH_EDITOR
 
 			bool bIsNew = true;
-			StaticMesh = UsdGeomMeshTranslatorImpl::CreateStaticMesh( LODIndexToMeshDescription, *Context, bIsNew );
-
 			const FString PrimPathString = PrimPath.GetString();
+			StaticMesh = UsdGeomMeshTranslatorImpl::CreateStaticMesh( LODIndexToMeshDescription, *Context, PrimPathString, bIsNew );
 
 			if ( StaticMesh )
 			{
diff --git a/Engine/Plugins/Importers/USDImporter/Source/USDSchemas/Private/USDShadeMaterialTranslator.cpp b/Engine/Plugins/Importers/USDImporter/Source/USDSchemas/Private/USDShadeMaterialTranslator.cpp
--- a/Engine/Plugins/Importers/USDImporter/Source/USDSchemas/Private/USDShadeMaterialTranslator.cpp
+++ b/Engine/Plugins/Importers/USDImporter/Source/USDSchemas/Private/USDShadeMaterialTranslator.cpp
@@ -45,10 +45,11 @@ void FUsdShadeMaterialTranslator::CreateAssets()
 
 		if ( UMaterialInterface* MasterMaterial = Cast< UMaterialInterface >( FSoftObjectPath( MaterialPath ).TryLoad() ) )
 		{
+			FName InstanceName = MakeUniqueObjectName( GetTransientPackage(), UMaterialInstance::StaticClass(), *FPaths::GetBaseFilename( PrimPath.GetString() ) );
 			if ( GIsEditor ) // Also have to prevent Standalone game from going with MaterialInstanceConstants
 			{
 #if WITH_EDITOR
-				if ( UMaterialInstanceConstant* NewMaterial = NewObject<UMaterialInstanceConstant>( GetTransientPackage(), NAME_None, Context->ObjectFlags ) )
+				if ( UMaterialInstanceConstant* NewMaterial = NewObject<UMaterialInstanceConstant>( GetTransientPackage(), InstanceName, Context->ObjectFlags ) )
 				{
 					NewMaterial->SetParentEditorOnly( MasterMaterial );
 
@@ -70,7 +71,7 @@ void FUsdShadeMaterialTranslator::CreateAssets()
 				}
 #endif // WITH_EDITOR
 			}
-			else if ( UMaterialInstanceDynamic* NewMaterial = UMaterialInstanceDynamic::Create( MasterMaterial, GetTransientPackage() ) )
+			else if ( UMaterialInstanceDynamic* NewMaterial = UMaterialInstanceDynamic::Create( MasterMaterial, GetTransientPackage(), InstanceName ) )
 			{
 				TMap<FString, int32> Unused;
 				TMap<FString, int32>& PrimvarToUVIndex = Context->MaterialToPrimvarToUVIndex ? Context->MaterialToPrimvarToUVIndex->FindOrAdd( PrimPath.GetString() ) : Unused;
diff --git a/Engine/Plugins/Importers/USDImporter/Source/USDSchemas/Private/USDSkelRootTranslator.cpp b/Engine/Plugins/Importers/USDImporter/Source/USDSchemas/Private/USDSkelRootTranslator.cpp
--- a/Engine/Plugins/Importers/USDImporter/Source/USDSchemas/Private/USDSkelRootTranslator.cpp
+++ b/Engine/Plugins/Importers/USDImporter/Source/USDSchemas/Private/USDSkelRootTranslator.cpp
@@ -759,7 +759,7 @@ namespace UsdSkelRootTranslatorImpl
 				if ( !SkeletalMesh )
 				{
 					bIsNew = true;
-					SkeletalMesh = UsdToUnreal::GetSkeletalMeshFromImportData( LODIndexToSkeletalMeshImportData, SkeletonBones, NewBlendShapes, Context->ObjectFlags );
+					SkeletalMesh = UsdToUnreal::GetSkeletalMeshFromImportData( LODIndexToSkeletalMeshImportData, SkeletonBones, NewBlendShapes, Context->ObjectFlags, *FPaths::GetBaseFilename( SkelRootPath ) );
 				}
 
 				if ( SkeletalMesh )
@@ -852,8 +852,6 @@ namespace UsdSkelRootTranslatorImpl
 						{
 							FScopedUnrealAllocs UEAllocs;
 
-							pxr::UsdPrim SkelAnimPrim = AnimQuery.GetPrim();
-
 							// The UAnimSequence can't be created with the RF_Transactional flag, or else it will be serialized without
 							// Bone/CurveCompressionSettings. Undoing that transaction would call UAnimSequence::Serialize with nullptr values for both, which crashes.
 							// Besides, this particular asset type is only ever created when we import to content folder assets (so never for realtime), and
diff --git a/Engine/Plugins/Importers/USDImporter/Source/USDUtilities/Private/USDGeomMeshConversion.cpp b/Engine/Plugins/Importers/USDImporter/Source/USDUtilities/Private/USDGeomMeshConversion.cpp
--- a/Engine/Plugins/Importers/USDImporter/Source/USDUtilities/Private/USDGeomMeshConversion.cpp
+++ b/Engine/Plugins/Importers/USDImporter/Source/USDUtilities/Private/USDGeomMeshConversion.cpp
@@ -936,7 +936,16 @@ UMaterialInstanceDynamic* UsdUtils::CreateDisplayColorMaterialInstanceDynamic( c
 
 	if ( UMaterialInterface* ParentMaterial = Cast< UMaterialInterface >( FSoftObjectPath( ParentPath ).TryLoad() ) )
 	{
-		if ( UMaterialInstanceDynamic* NewMaterial = UMaterialInstanceDynamic::Create( ParentMaterial, GetTransientPackage() ) )
+		FName AssetName = MakeUniqueObjectName(
+			GetTransientPackage(),
+			UMaterialInstanceConstant::StaticClass(),
+			*FString::Printf( TEXT( "DisplayColor_%s_%s" ),
+				DisplayColorDescription.bHasOpacity ? TEXT( "Opacity" ) : TEXT( "NoOpacity" ),
+				DisplayColorDescription.bIsDoubleSided ? TEXT( "DoubleSided" ) : TEXT( "SingleSided" )
+			)
+		);
+
+		if ( UMaterialInstanceDynamic* NewMaterial = UMaterialInstanceDynamic::Create( ParentMaterial, GetTransientPackage(), AssetName ) )
 		{
 			return NewMaterial;
 		}
@@ -974,7 +983,16 @@ UMaterialInstanceConstant* UsdUtils::CreateDisplayColorMaterialInstanceConstant(
 
 	if ( UMaterialInterface* ParentMaterial = Cast< UMaterialInterface >( FSoftObjectPath( ParentPath ).TryLoad() ) )
 	{
-		if ( UMaterialInstanceConstant* MaterialInstance = NewObject< UMaterialInstanceConstant >( GetTransientPackage(), NAME_None, RF_NoFlags ) )
+		FName AssetName = MakeUniqueObjectName(
+			GetTransientPackage(),
+			UMaterialInstanceConstant::StaticClass(),
+			*FString::Printf( TEXT( "DisplayColor_%s_%s" ),
+				DisplayColorDescription.bHasOpacity ? TEXT( "Opacity" ) : TEXT( "NoOpacity" ),
+				DisplayColorDescription.bIsDoubleSided ? TEXT( "DoubleSided" ) : TEXT( "SingleSided" )
+			)
+		);
+
+		if ( UMaterialInstanceConstant* MaterialInstance = NewObject< UMaterialInstanceConstant >( GetTransientPackage(), AssetName, RF_NoFlags ) )
 		{
 			UMaterialEditingLibrary::SetMaterialInstanceParent( MaterialInstance, ParentMaterial );
 			return MaterialInstance;
diff --git a/Engine/Plugins/Importers/USDImporter/Source/USDUtilities/Private/USDShadeConversion.cpp b/Engine/Plugins/Importers/USDImporter/Source/USDUtilities/Private/USDShadeConversion.cpp
--- a/Engine/Plugins/Importers/USDImporter/Source/USDUtilities/Private/USDShadeConversion.cpp
+++ b/Engine/Plugins/Importers/USDImporter/Source/USDUtilities/Private/USDShadeConversion.cpp
@@ -200,10 +200,12 @@ namespace UE
 
 				FScopedUnrealAllocs UEAllocs;
 
+				FName TextureName = MakeUniqueObjectName( Outer, UTexture::StaticClass(), *FPaths::GetBaseFilename( ResolvedTexturePath ) );
+
 				return Cast<UTexture>(TextureFactory->FactoryCreateBinary(
 					UTexture::StaticClass(),
 					Outer,
-					NAME_None,
+					TextureName,
 					ObjectFlags,
 					nullptr,
 					*TextureExtension,
@@ -967,7 +969,8 @@ namespace UE
 						// Not inside an USDZ archive, just a regular texture
 						else
 						{
-							Texture = Cast< UTexture >( TextureFactory->ImportObject( UTexture::StaticClass(), Outer, NAME_None, ObjectFlags, ResolvedTexturePath, TEXT( "" ), bOutCancelled ) );
+							FName TextureName = MakeUniqueObjectName( Outer, UTexture::StaticClass(), *FPaths::GetBaseFilename( ResolvedTexturePath ) );
+							Texture = Cast< UTexture >( TextureFactory->ImportObject( UTexture::StaticClass(), Outer, TextureName, ObjectFlags, ResolvedTexturePath, TEXT(""), bOutCancelled));
 						}
 
 						if ( Texture )
diff --git a/Engine/Plugins/Importers/USDImporter/Source/USDUtilities/Private/USDSkeletalDataConversion.cpp b/Engine/Plugins/Importers/USDImporter/Source/USDUtilities/Private/USDSkeletalDataConversion.cpp
--- a/Engine/Plugins/Importers/USDImporter/Source/USDUtilities/Private/USDSkeletalDataConversion.cpp
+++ b/Engine/Plugins/Importers/USDImporter/Source/USDUtilities/Private/USDSkeletalDataConversion.cpp
@@ -2053,7 +2053,9 @@ USkeletalMesh* UsdToUnreal::GetSkeletalMeshFromImportData(
 	TArray<FSkeletalMeshImportData>& LODIndexToSkeletalMeshImportData,
 	const TArray<SkeletalMeshImportData::FBone>& InSkeletonBones,
 	UsdUtils::FBlendShapeMap& InBlendShapesByPath,
-	EObjectFlags ObjectFlags
+	EObjectFlags ObjectFlags,
+	const FName& MeshName,
+	const FName& SkeletonName
 )
 {
 	TRACE_CPUPROFILER_EVENT_SCOPE( UsdToUnreal::GetSkeletalMeshFromImportData );
@@ -2066,7 +2068,7 @@ USkeletalMesh* UsdToUnreal::GetSkeletalMeshFromImportData(
 	// A SkeletalMesh could be retrieved for re-use and updated for animations
 	// For now, create a new USkeletalMesh
 	// Note: Remember to initialize UsedMorphTargetNames with existing morph targets, whenever the SkeletalMesh is reused
-	USkeletalMesh* SkeletalMesh = NewObject<USkeletalMesh>(GetTransientPackage(), NAME_None, ObjectFlags | EObjectFlags::RF_Public);
+	USkeletalMesh* SkeletalMesh = NewObject<USkeletalMesh>(GetTransientPackage(), MeshName, ObjectFlags | EObjectFlags::RF_Public);
 
 	// Process reference skeleton from import data
 	int32 SkeletalDepth = 0;
@@ -2172,7 +2174,7 @@ USkeletalMesh* UsdToUnreal::GetSkeletalMeshFromImportData(
 	SkeletalMesh->CalculateInvRefMatrices();
 
 	// Generate a Skeleton and associate it to the SkeletalMesh
-	USkeleton* Skeleton = NewObject<USkeleton>(GetTransientPackage(), NAME_None, ObjectFlags | EObjectFlags::RF_Public );
+	USkeleton* Skeleton = NewObject<USkeleton>( GetTransientPackage(), SkeletonName, ObjectFlags | EObjectFlags::RF_Public );
 	Skeleton->MergeAllBonesToBoneTree(SkeletalMesh);
 	Skeleton->SetPreviewMesh(SkeletalMesh);
 	SkeletalMesh->SetSkeleton(Skeleton);
diff --git a/Engine/Plugins/Importers/USDImporter/Source/USDUtilities/Public/USDSkeletalDataConversion.h b/Engine/Plugins/Importers/USDImporter/Source/USDUtilities/Public/USDSkeletalDataConversion.h
--- a/Engine/Plugins/Importers/USDImporter/Source/USDUtilities/Public/USDSkeletalDataConversion.h
+++ b/Engine/Plugins/Importers/USDImporter/Source/USDUtilities/Public/USDSkeletalDataConversion.h
@@ -168,9 +168,11 @@ namespace UsdToUnreal
 	 * @param InSkeletonBones - Bones to use for the reference skeleton (skeleton data on each LODIndexToSkeletalMeshImportData will be ignored).
 	 * @param BlendShapesByPath - Blend shapes to convert to morph targets
 	 * @param ObjectFlags - Flags to use when creating the USkeletalMesh and corresponding USkeleton
+	 * @param MeshName - Name to use for the new USkeletalMesh asset
+	 * @param SkeletonName - Name to use for the new USkeleton asset
 	 * @return Newly created USkeletalMesh, or nullptr in case of failure
 	 */
-	USDUTILITIES_API USkeletalMesh* GetSkeletalMeshFromImportData( TArray<FSkeletalMeshImportData>& LODIndexToSkeletalMeshImportData, const TArray<SkeletalMeshImportData::FBone>& InSkeletonBones, UsdUtils::FBlendShapeMap& InBlendShapesByPath, EObjectFlags ObjectFlags );
+	USDUTILITIES_API USkeletalMesh* GetSkeletalMeshFromImportData( TArray<FSkeletalMeshImportData>& LODIndexToSkeletalMeshImportData, const TArray<SkeletalMeshImportData::FBone>& InSkeletonBones, UsdUtils::FBlendShapeMap& InBlendShapesByPath, EObjectFlags ObjectFlags, const FName& MeshName = NAME_None, const FName& SkeletonName = NAME_None );
 }
 
 namespace UnrealToUsd
