USD: Fix incorrect sRGB if the texture is shared between two different channels with different sRGB values.

Use filenames for texture FNames when opening USD stage;

diff --git a/Engine/Plugins/Importers/USDImporter/Source/USDUtilities/Private/USDShadeConversion.cpp b/Engine/Plugins/Importers/USDImporter/Source/USDUtilities/Private/USDShadeConversion.cpp
--- a/Engine/Plugins/Importers/USDImporter/Source/USDUtilities/Private/USDShadeConversion.cpp
+++ b/Engine/Plugins/Importers/USDImporter/Source/USDUtilities/Private/USDShadeConversion.cpp
@@ -299,8 +299,11 @@ namespace UE
 			// Computes and returns the hash string for the texture at the given path.
 			// Handles regular texture asset paths as well as asset paths identifying textures inside Usdz archives.
 			// Returns an empty string if the texture could not be hashed.
-			FString GetTextureHash(const FString& ResolvedTexturePath)
+			FString GetTextureHash( const FString& ResolvedTexturePath, bool bSRGB, TextureCompressionSettings CompressionSettings )
 			{
+				FMD5 MD5;
+
+				// Hash the actual texture
 				FString TextureExtension;
 				if (IsInsideUsdzArchive(ResolvedTexturePath, TextureExtension))
 				{
@@ -310,15 +313,40 @@ namespace UE
 
 					if (BufferSize > 0 && BufferStart != nullptr)
 					{
-						return FMD5::HashBytes(BufferStart, BufferSize);
+						MD5.Update( BufferStart, BufferSize );
+					}
+				}
+				// Copied from FMD5Hash::HashFileFromArchive as it doesn't expose its FMD5
+				else if ( FArchive* Ar = IFileManager::Get().CreateFileReader( *ResolvedTexturePath ) )
+				{
+					TArray<uint8> LocalScratch;
+					LocalScratch.SetNumUninitialized( 1024 * 64 );
+
+					const int64 Size = Ar->TotalSize();
+					int64 Position = 0;
+
+					// Read in BufferSize chunks
+					while ( Position < Size )
+					{
+						const auto ReadNum = FMath::Min( Size - Position, ( int64 ) LocalScratch.Num() );
+						Ar->Serialize( LocalScratch.GetData(), ReadNum );
+						MD5.Update( LocalScratch.GetData(), ReadNum );
+
+						Position += ReadNum;
 					}
 				}
 				else
 				{
-					return LexToString(FMD5Hash::HashFile(*ResolvedTexturePath));
+					UE_LOG( LogUsd, Warning, TEXT( "Failed to find texture at path '%s' when trying to generate a hash for it" ), *ResolvedTexturePath );
 				}
 
-				return {};
+				// Hash the additional data
+				MD5.Update( reinterpret_cast< uint8* >( &bSRGB ), sizeof( bool ) );
+				MD5.Update( reinterpret_cast< uint8* >( &CompressionSettings ), sizeof( CompressionSettings ) );
+
+				FMD5Hash Hash;
+				Hash.Set( MD5 );
+				return LexToString( Hash );
 			}
 
 			// Will traverse the shade material graph backwards looking for a string/token value and return it.
@@ -473,7 +501,24 @@ namespace UE
 							return false;
 						}
 
-						const FString TextureHash = GetTextureHash(TexturePath);
+						// We'll add these to the hash because the materials are built to try and reuse the same textures for multiple channels,
+						// and those may expect linear or sRGB values. Without this we may parse a texture as linear because we hit the opacity channel first,
+						// and then reuse it as linear for the base color channel even though it should have been sRGB.
+						// Plus we may have something weird like a normal map being plugged into the base color and the normal channel.
+						bool bSRGB = true;
+						TextureCompressionSettings CompressionSettings = TextureCompressionSettings::TC_Default;
+						if ( bIsNormalInput )
+						{
+							// Disable SRGB when parsing float textures, as they're likely specular/roughness maps
+							bSRGB = false;
+							CompressionSettings = TextureCompressionSettings::TC_Normalmap;
+						}
+						else if ( LODGroup == TEXTUREGROUP_WorldSpecular )
+						{
+							bSRGB = false;
+						}
+
+						const FString TextureHash = GetTextureHash( TexturePath, bSRGB, CompressionSettings );
 
 						// We only actually want to retrieve the textures if we have a cache to put them in
 						if ( TexturesCache )
@@ -498,18 +543,8 @@ namespace UE
 
 							if ( Texture )
 							{
-								if ( bIsNormalInput )
-								{
-									Texture->SRGB = false;
-									Texture->CompressionSettings = TextureCompressionSettings::TC_Normalmap;
-								}
-
-								// Disable SRGB when parsing float textures, as they're likely specular/roughness maps
-								if ( LODGroup == TEXTUREGROUP_WorldSpecular )
-								{
-									Texture->SRGB = false;
-								}
-
+								Texture->SRGB = bSRGB;
+								Texture->CompressionSettings = CompressionSettings;
 								Texture->UpdateResource();
 
 								FString PrimvarName = GetPrimvarUsedAsST(Source);
