USD: Call Modify() before applying new values from exposed setters

diff --git a/Engine/Plugins/Importers/USDImporter/Source/USDStage/Private/USDStageActor.cpp b/Engine/Plugins/Importers/USDImporter/Source/USDStage/Private/USDStageActor.cpp
--- a/Engine/Plugins/Importers/USDImporter/Source/USDStage/Private/USDStageActor.cpp
+++ b/Engine/Plugins/Importers/USDImporter/Source/USDStage/Private/USDStageActor.cpp
@@ -1129,6 +1129,8 @@ UE::FUsdStage& AUsdStageActor::GetOrLoadUsdStage()
 
 void AUsdStageActor::SetRootLayer( const FString& RootFilePath )
 {
+	Modify();
+
 	FString RelativeFilePath = RootFilePath;
 #if USE_USD_SDK
 	if ( !RelativeFilePath.IsEmpty() && !FPaths::IsRelative( RelativeFilePath ) && !RelativeFilePath.StartsWith( UnrealIdentifiers::IdentifierPrefix ) )
@@ -1157,24 +1159,32 @@ void AUsdStageActor::SetRootLayer( const FString& RootFilePath )
 
 void AUsdStageActor::SetInitialLoadSet( EUsdInitialLoadSet NewLoadSet )
 {
+	Modify();
+
 	InitialLoadSet = NewLoadSet;
 	LoadUsdStage();
 }
 
 void AUsdStageActor::SetPurposesToLoad( int32 NewPurposesToLoad )
 {
+	Modify();
+
 	PurposesToLoad = NewPurposesToLoad;
 	LoadUsdStage();
 }
 
 void AUsdStageActor::SetRenderContext( const FName& NewRenderContext )
 {
+	Modify();
+
 	RenderContext = NewRenderContext;
 	LoadUsdStage();
 }
 
 void AUsdStageActor::SetTime(float InTime)
 {
+	Modify();
+
 	Time = InTime;
 	Refresh();
 }
