USD: Fix UAnimSequences being imported with 60000 fps due to UAnimSequence 5.0 changes

diff --git a/Engine/Plugins/Importers/USDImporter/Source/USDUtilities/Private/USDSkeletalDataConversion.cpp b/Engine/Plugins/Importers/USDImporter/Source/USDUtilities/Private/USDSkeletalDataConversion.cpp
--- a/Engine/Plugins/Importers/USDImporter/Source/USDUtilities/Private/USDSkeletalDataConversion.cpp
+++ b/Engine/Plugins/Importers/USDImporter/Source/USDUtilities/Private/USDSkeletalDataConversion.cpp
@@ -19,6 +19,7 @@
 #include "Interfaces/ITargetPlatform.h"
 #include "Interfaces/ITargetPlatformManagerModule.h"
 #include "Misc/CoreMisc.h"
+#include "Misc/MemStack.h"
 #include "Modules/ModuleManager.h"
 #include "Rendering/SkeletalMeshLODImporterData.h"
 #include "Rendering/SkeletalMeshLODModel.h"
@@ -1694,7 +1695,7 @@ bool UsdToUnreal::ConvertSkelAnim( const pxr::UsdSkelSkeletonQuery& InUsdSkeleto
 	const double StartSeconds = StartTimeCode / TimeCodesPerSecond;
 	const double SequenceLengthTimeCodes = EndTimeCode - StartTimeCode;
 	const double SequenceLengthSeconds = FMath::Max<double>( SequenceLengthTimeCodes / TimeCodesPerSecond, MINIMUM_ANIMATION_LENGTH );
-	const int32 NumBakedFrames = FMath::CeilToInt( FMath::Max( SequenceLengthSeconds * FramesPerSecond + 1.0, 1.0 ) );
+	const int32 NumBakedFrames = FMath::RoundToInt( FMath::Max( SequenceLengthSeconds * FramesPerSecond + 1.0, 1.0 ) );
 	const double IntervalTimeCodes = ( NumBakedFrames > 1 ) ? ( SequenceLengthTimeCodes / ( NumBakedFrames - 1 ) ) : MINIMUM_ANIMATION_LENGTH;
 
 
@@ -1936,6 +1937,7 @@ bool UsdToUnreal::ConvertSkelAnim( const pxr::UsdSkelSkeletonQuery& InUsdSkeleto
 
 	Controller.SetPlayLength(SequenceLengthSeconds);
 	Controller.SetFrameRate(FFrameRate(FramesPerSecond, 1));
+	Controller.NotifyPopulated(); // This call is important to get the controller to not use the sampling frequency as framerate
 	Controller.CloseBracket();
 
 	OutSkeletalAnimationAsset->PostEditChange();
@@ -2608,6 +2610,10 @@ bool UnrealToUsd::ConvertAnimSequence( UAnimSequence* AnimSequence, pxr::UsdPrim
 				}
 			}
 
+			// We need to make sure we have at least one mark on the memstack allocator because FBlendedCurve
+			// will allocate using one and will assert if there aren't any marks yet
+			FMemMark Mark( FMemStack::Get() );
+
 			// Blend shape weights
 			for ( int32 TimeCode = 0; TimeCode < NumTimeCodes; ++TimeCode )
 			{
