USD: Fix crash when changing display rate on the sequencer for a USD-generated level sequence.

Automatically compensate how USD also sets timecodespersecond when we set framespersecond;
Fix xform tracks disappearing when reloading the stage;
Tactical renaming to help prevent some confusion regarding framerate/tick resolution/time codes/frames per second.

diff --git a/Engine/Plugins/Importers/USDImporter/Source/USDStage/Private/USDLevelSequenceHelper.cpp b/Engine/Plugins/Importers/USDImporter/Source/USDStage/Private/USDLevelSequenceHelper.cpp
--- a/Engine/Plugins/Importers/USDImporter/Source/USDStage/Private/USDLevelSequenceHelper.cpp
+++ b/Engine/Plugins/Importers/USDImporter/Source/USDStage/Private/USDLevelSequenceHelper.cpp
@@ -633,7 +633,7 @@ void FUsdLevelSequenceHelperImpl::CreateSubSequenceSection( ULevelSequence& Sequ
 	const bool bReadonly = false;
 	UsdLevelSequenceHelperImpl::FMovieSceneReadonlyGuard MovieSceneReadonlyGuard{ *MovieScene, bReadonly };
 
-	FFrameRate FrameRate = MovieScene->GetTickResolution();
+	FFrameRate TickResolution = MovieScene->GetTickResolution();
 
 	UMovieSceneSubTrack* SubTrack = MovieScene->FindMasterTrack< UMovieSceneSubTrack >();
 	if ( !SubTrack )
@@ -700,8 +700,8 @@ void FUsdLevelSequenceHelperImpl::CreateSubSequenceSection( ULevelSequence& Sequ
 	const double SubStartTimeSeconds  = ( SubLayerOffset.Offset + SubLayerOffset.Scale * SubLayerTimeInfo->StartTimeCode.Get( 0.0 ) ) / TimeCodesPerSecond;
 	const double SubEndTimeSeconds = ( SubLayerOffset.Offset + SubLayerOffset.Scale * SubLayerTimeInfo->EndTimeCode.Get( 0.0 ) ) / TimeCodesPerSecond;
 
-	const FFrameNumber StartFrame = UsdLevelSequenceHelperImpl::RoundAsFrameNumber( FrameRate, SubStartTimeSeconds );
-	const int32 Duration = UsdLevelSequenceHelperImpl::RoundAsFrameNumber( FrameRate, SubEndTimeSeconds ).Value - StartFrame.Value;
+	const FFrameNumber StartFrame = UsdLevelSequenceHelperImpl::RoundAsFrameNumber( TickResolution, SubStartTimeSeconds );
+	const int32 Duration = UsdLevelSequenceHelperImpl::RoundAsFrameNumber( TickResolution, SubEndTimeSeconds ).Value - StartFrame.Value;
 
 	TRange< FFrameNumber > SubSectionRange( StartFrame, StartFrame + Duration );
 	SubSectionRange = SequenceTransform.TransformRangeUnwarped( SubSectionRange );
@@ -1156,12 +1156,12 @@ void FUsdLevelSequenceHelperImpl::UpdateUsdLayerOffsetFromSection(const UMovieSc
 	const double SubStartTimeCode = SubLayerTimeInfo->StartTimeCode.Get(0.0);
 	const double SubEndTimeCode = SubLayerTimeInfo->EndTimeCode.Get(0.0);
 
-	FFrameRate FrameRate = MovieScene->GetTickResolution();
+	FFrameRate TickResolution = MovieScene->GetTickResolution();
 	FFrameNumber ModifiedStartFrame = Section->GetInclusiveStartFrame();
 	FFrameNumber ModifiedEndFrame   = Section->GetExclusiveEndFrame();
 
 	// This will obviously be quantized to frame intervals for now
-	double SubSectionStartTimeCode = FrameRate.AsSeconds(ModifiedStartFrame) * TimeCodesPerSecond;
+	double SubSectionStartTimeCode = TickResolution.AsSeconds(ModifiedStartFrame) * TimeCodesPerSecond;
 
 	UE::FSdfLayerOffset NewLayerOffset;
 	NewLayerOffset.Scale = FMath::IsNearlyZero( Section->Parameters.TimeScale ) ? 0.f : 1.f / Section->Parameters.TimeScale;
@@ -1265,9 +1265,9 @@ void FUsdLevelSequenceHelperImpl::UpdateMovieSceneTimeRanges( UMovieScene& Movie
 		const double EndTimeCode = LayerTimeInfo.EndTimeCode.Get(0.0);
 		const double TimeCodesPerSecond = GetTimeCodesPerSecond();
 
-		const FFrameRate DestFrameRate = MovieScene.GetTickResolution();
-		const FFrameNumber StartFrame  = UsdLevelSequenceHelperImpl::RoundAsFrameNumber( DestFrameRate, StartTimeCode / TimeCodesPerSecond );
-		const FFrameNumber EndFrame    = UsdLevelSequenceHelperImpl::RoundAsFrameNumber( DestFrameRate, EndTimeCode / TimeCodesPerSecond );
+		const FFrameRate TickResolution  = MovieScene.GetTickResolution();
+		const FFrameNumber StartFrame    = UsdLevelSequenceHelperImpl::RoundAsFrameNumber( TickResolution, StartTimeCode / TimeCodesPerSecond );
+		const FFrameNumber EndFrame      = UsdLevelSequenceHelperImpl::RoundAsFrameNumber( TickResolution, EndTimeCode / TimeCodesPerSecond );
 		TRange< FFrameNumber > TimeRange = TRange<FFrameNumber>::Inclusive( StartFrame, EndFrame );
 
 		MovieScene.SetPlaybackRange( TimeRange );
@@ -1407,6 +1407,7 @@ void FUsdLevelSequenceHelperImpl::HandleMovieSceneChange( UMovieScene& MovieScen
 		return;
 	}
 
+	const double StageTimeCodesPerSecond = GetTimeCodesPerSecond();
 	const TRange< FFrameNumber > PlaybackRange = MovieScene.GetPlaybackRange();
 	const FFrameRate DisplayRate = MovieScene.GetDisplayRate();
 	const FFrameRate LayerTimeCodesPerSecond( Layer.GetTimeCodesPerSecond(), 1 );
@@ -1414,10 +1415,16 @@ void FUsdLevelSequenceHelperImpl::HandleMovieSceneChange( UMovieScene& MovieScen
 	const FFrameTime EndTime = FFrameRate::TransformTime(UE::MovieScene::DiscreteExclusiveUpper( PlaybackRange ).Value, MovieScene.GetTickResolution(), LayerTimeCodesPerSecond );
 
 	UE::FSdfChangeBlock ChangeBlock;
+	FScopedBlockNoticeListening BlockNotices( StageActor.Get() );
 	if ( !FMath::IsNearlyEqual( DisplayRate.AsDecimal(), GetFramesPerSecond() ) )
 	{
 		UsdStage.SetFramesPerSecond( DisplayRate.AsDecimal() );
 
+		// For whatever reason setting a stage FramesPerSecond also automatically sets its TimeCodesPerSecond to the same value, so we need to undo it.
+		// This because all the sequencer does is change display rate, which is the analogue to USD's frames per second (i.e. we are only changing how many
+		// frames we'll display between any two timecodes, not how many timecodes we'll display per second)
+		UsdStage.SetTimeCodesPerSecond( StageTimeCodesPerSecond );
+
 		// Propagate to all movie scenes, as USD only uses the stage FramesPerSecond so the sequences should have a unified DisplayRate to reflect that
 		for ( TPair< FString, ULevelSequence* >& SequenceByIdentifier : LevelSequencesByIdentifier )
 		{
diff --git a/Engine/Plugins/Importers/USDImporter/Source/USDStage/Private/USDStageActor.cpp b/Engine/Plugins/Importers/USDImporter/Source/USDStage/Private/USDStageActor.cpp
--- a/Engine/Plugins/Importers/USDImporter/Source/USDStage/Private/USDStageActor.cpp
+++ b/Engine/Plugins/Importers/USDImporter/Source/USDStage/Private/USDStageActor.cpp
@@ -1617,6 +1617,10 @@ void AUsdStageActor::ReloadAnimations()
 		// We need to guarantee we'll record our change of LevelSequence into the transaction, as Init() will create a new one
 		Modify();
 
+		// Make sure our PrimsToAnimate and the LevelSequenceHelper are kept in sync, because we'll use PrimsToAnimate to
+		// check whether we need to call LevelSequenceHelper::AddPrim within AUsdStageActor::ExpandPrim
+		PrimsToAnimate.Reset();
+
 		LevelSequence = LevelSequenceHelper.Init( UsdStage );
 		LevelSequenceHelper.BindToUsdStageActor( this );
 
