USD: Fix bugs related to importing UAnimSequence assets.

Fix UAnimSequence asset retaining a soft reference to its preview mesh when it was in the transient package;

diff --git a/Engine/Plugins/Importers/USDImporter/Source/USDStageEditorViewModels/Private/USDStageViewModel.cpp b/Engine/Plugins/Importers/USDImporter/Source/USDStageEditorViewModels/Private/USDStageViewModel.cpp
--- a/Engine/Plugins/Importers/USDImporter/Source/USDStageEditorViewModels/Private/USDStageViewModel.cpp
+++ b/Engine/Plugins/Importers/USDImporter/Source/USDStageEditorViewModels/Private/USDStageViewModel.cpp
@@ -245,7 +245,6 @@ void FUsdStageViewModel::ImportStage()
 		// Preload some settings according to USDStage options. These will overwrite whatever is loaded from config
 		ImportContext.ImportOptions->PurposesToImport = StageActor->PurposesToLoad;
 		ImportContext.ImportOptions->RenderContextToImport = StageActor->RenderContext;
-		ImportContext.ImportOptions->ImportTime = StageActor->GetTime();
 		ImportContext.ImportOptions->StageOptions.MetersPerUnit = UsdUtils::GetUsdStageMetersPerUnit( UsdStage );
 		ImportContext.ImportOptions->StageOptions.UpAxis = UsdUtils::GetUsdStageUpAxisAsEnum( UsdStage );
 		ImportContext.bReadFromStageCache = true; // So that we import whatever the user has open right now, even if the file has changes
diff --git a/Engine/Plugins/Importers/USDImporter/Source/USDStageImporter/Private/USDStageImportOptions.cpp b/Engine/Plugins/Importers/USDImporter/Source/USDStageImporter/Private/USDStageImportOptions.cpp
--- a/Engine/Plugins/Importers/USDImporter/Source/USDStageImporter/Private/USDStageImportOptions.cpp
+++ b/Engine/Plugins/Importers/USDImporter/Source/USDStageImporter/Private/USDStageImportOptions.cpp
@@ -19,7 +19,6 @@ UUsdStageImportOptions::UUsdStageImportOptions(const FObjectInitializer& ObjectI
 	PurposesToImport = (int32) (EUsdPurpose::Default | EUsdPurpose::Proxy | EUsdPurpose::Render | EUsdPurpose::Guide);
 	IUsdSchemasModule& UsdSchemasModule = FModuleManager::Get().LoadModuleChecked< IUsdSchemasModule >( TEXT("USDSchemas") );
 	RenderContextToImport = UsdSchemasModule.GetRenderContextRegistry().GetUniversalRenderContext();
-	ImportTime = 0.0f;
 	bOverrideStageOptions = false;
 	StageOptions.MetersPerUnit = 0.01f;
 	StageOptions.UpAxis = EUsdUpAxis::ZAxis;
diff --git a/Engine/Plugins/Importers/USDImporter/Source/USDStageImporter/Private/USDStageImporter.cpp b/Engine/Plugins/Importers/USDImporter/Source/USDStageImporter/Private/USDStageImporter.cpp
--- a/Engine/Plugins/Importers/USDImporter/Source/USDStageImporter/Private/USDStageImporter.cpp
+++ b/Engine/Plugins/Importers/USDImporter/Source/USDStageImporter/Private/USDStageImporter.cpp
@@ -472,7 +472,7 @@ namespace UsdStageImporterImpl
 
 	// Moves Asset from its folder to the package at DestFullContentPath and sets up its flags.
 	// Depending on ReplacePolicy it may replace the existing actor (if it finds one) or just abort
-	UObject* PublishAsset(FUsdStageImportContext& ImportContext, UObject* Asset, const FString& DestFullPackagePath, TMap<UObject*, UObject*>& ObjectsToRemap)
+	UObject* PublishAsset(FUsdStageImportContext& ImportContext, UObject* Asset, const FString& DestFullPackagePath, TMap<UObject*, UObject*>& ObjectsToRemap, TMap<FSoftObjectPath, FSoftObjectPath>& SoftObjectsToRemap )
 	{
 		if (!Asset)
 		{
@@ -513,6 +513,7 @@ namespace UsdStageImporterImpl
 			{
 				// Redirect any users of our new transient asset to the old, existing asset
 				ObjectsToRemap.Add( Asset, ExistingAsset );
+				SoftObjectsToRemap.Add( Asset, ExistingAsset );
 				return nullptr;
 			}
 		}
@@ -534,15 +535,13 @@ namespace UsdStageImporterImpl
 		}
 		Package->FullyLoad();
 
-		FString OldAssetPathName;
+		FSoftObjectPath OldPath = Asset;
 
 		// Strategy copied from FDatasmithImporterImpl::PublicizeAsset
 		// Replace existing asset (reimport or conflict) with new asset
 		UObject* MovedAsset = ExistingAsset;
 		if (ExistingAsset != nullptr && ExistingAsset != Asset && ReplacePolicy == EReplaceAssetPolicy::Replace)
 		{
-			OldAssetPathName = ExistingAsset->GetPathName();
-
 			MovedAsset = DuplicateObject<UObject>(Asset, Package, ExistingAsset->GetFName());
 
 			// If mesh's label has changed, update its name
@@ -572,21 +571,12 @@ namespace UsdStageImporterImpl
 			MovedAsset = Asset;
 		}
 
+		SoftObjectsToRemap.Add( OldPath, MovedAsset );
 		if (MovedAsset != Asset)
 		{
 			ObjectsToRemap.Add(Asset, MovedAsset);
 		}
 
-		// Soft object ptrs won't update, so we need to do it manually as we use PreviewMesh when fetching AssetImportData
-		// Note that we need to do this even if we never replaced/overwrote assets, as we will at least move from transient to the published folder
-		if ( USkeletalMesh* MovedMeshAsset = Cast<USkeletalMesh>( MovedAsset ) )
-		{
-			if ( USkeleton* Skeleton = MovedMeshAsset->GetSkeleton() )
-			{
-				Skeleton->SetPreviewMesh( MovedMeshAsset );
-			}
-		}
-
 		// Important as some assets (e.g. material instances) are created with no flags
 		MovedAsset->SetFlags(ImportContext.ImportObjectFlags | EObjectFlags::RF_Public | EObjectFlags::RF_Standalone );
 		MovedAsset->ClearFlags(EObjectFlags::RF_Transient | EObjectFlags::RF_DuplicateTransient | EObjectFlags::RF_NonPIEDuplicateTransient);
@@ -600,15 +590,6 @@ namespace UsdStageImporterImpl
 			Package->MarkPackageDirty();
 		}
 
-		if (!ExistingAsset)
-		{
-			FAssetRegistryModule::AssetCreated(MovedAsset);
-		}
-		else if (!OldAssetPathName.IsEmpty())
-		{
-			FAssetRegistryModule::AssetRenamed(MovedAsset, OldAssetPathName);
-		}
-
 		// Reopen asset editor if we were editing the asset
 		if (bAssetWasOpen)
 		{
@@ -621,7 +602,7 @@ namespace UsdStageImporterImpl
 	}
 
 	// Move imported assets from transient folder to their final package, updating AssetCache to point to the moved assets
-	void PublishAssets(FUsdStageImportContext& ImportContext, const TSet<UObject*>& AssetsToPublish, TMap<UObject*, UObject*>& ObjectsToRemap)
+	void PublishAssets(FUsdStageImportContext& ImportContext, const TSet<UObject*>& AssetsToPublish, TMap<UObject*, UObject*>& ObjectsToRemap, TMap<FSoftObjectPath, FSoftObjectPath>& SoftObjectsToRemap )
 	{
 		TSet<FString> UniqueAssetNames;
 
@@ -671,7 +652,7 @@ namespace UsdStageImporterImpl
 
 			FString TargetAssetName = GetUserFriendlyName(Asset, UniqueAssetNames);
 			FString DestPackagePath = FPaths::Combine(ImportContext.PackagePath, AssetTypeFolder, TargetAssetName);
-			PublishAsset(ImportContext, Asset, DestPackagePath, ObjectsToRemap);
+			PublishAsset(ImportContext, Asset, DestPackagePath, ObjectsToRemap, SoftObjectsToRemap);
 		}
 
 		// Publish the level sequences if there's data in them
@@ -688,12 +669,12 @@ namespace UsdStageImporterImpl
 			for (ULevelSequence* LevelSequence : LevelSequences)
 			{
 				const FString DestPackagePath = FPaths::Combine(ImportContext.PackagePath, AssetTypeFolder, LevelSequence->GetName());
-				UObject* PublishedLevelSequence = PublishAsset(ImportContext, LevelSequence, DestPackagePath, ObjectsToRemap);
+				UObject* PublishedLevelSequence = PublishAsset( ImportContext, LevelSequence, DestPackagePath, ObjectsToRemap, SoftObjectsToRemap );
 			}
 		}
 	}
 
-	void ResolveComponentConflict(USceneComponent* NewRoot, USceneComponent* ExistingRoot, EReplaceActorPolicy ReplacePolicy, TMap<UObject*, UObject*>& ObjectsToRemap)
+	void ResolveComponentConflict(USceneComponent* NewRoot, USceneComponent* ExistingRoot, EReplaceActorPolicy ReplacePolicy, TMap<UObject*, UObject*>& ObjectsToRemap, TMap<FSoftObjectPath, FSoftObjectPath>& SoftObjectsToRemap )
 	{
 		if (!NewRoot || !ExistingRoot || ReplacePolicy == EReplaceActorPolicy::Append)
 		{
@@ -701,6 +682,7 @@ namespace UsdStageImporterImpl
 		}
 
 		ObjectsToRemap.Add(ExistingRoot, NewRoot);
+		SoftObjectsToRemap.Add(ExistingRoot, NewRoot);
 
 		TArray<USceneComponent*> ExistingComponents = ExistingRoot->GetAttachChildren();
 		TArray<USceneComponent*> NewComponents = NewRoot->GetAttachChildren();
@@ -756,7 +738,7 @@ namespace UsdStageImporterImpl
 
 				if (bRecurse)
 				{
-					ResolveComponentConflict(NewComponent, *FoundExistingComponent, ReplacePolicy, ObjectsToRemap);
+					ResolveComponentConflict(NewComponent, *FoundExistingComponent, ReplacePolicy, ObjectsToRemap, SoftObjectsToRemap);
 				}
 			}
 		}
@@ -795,7 +777,7 @@ namespace UsdStageImporterImpl
 		Actor->GetWorld()->DestroyActor(Actor);
 	}
 
-	void ResolveActorConflict(AActor* NewActor, AActor* ExistingActor, EReplaceActorPolicy ReplacePolicy, TMap<UObject*, UObject*>& ObjectsToRemap)
+	void ResolveActorConflict(AActor* NewActor, AActor* ExistingActor, EReplaceActorPolicy ReplacePolicy, TMap<UObject*, UObject*>& ObjectsToRemap, TMap<FSoftObjectPath, FSoftObjectPath>& SoftObjectsToRemap)
 	{
 		if (!NewActor || !ExistingActor || ReplacePolicy == EReplaceActorPolicy::Append)
 		{
@@ -803,6 +785,7 @@ namespace UsdStageImporterImpl
 		}
 
 		ObjectsToRemap.Add(ExistingActor, NewActor);
+		SoftObjectsToRemap.Add( ExistingActor, NewActor );
 
 		// Collect new and existing actors by label
 		const bool bResetArray = false;
@@ -855,7 +838,7 @@ namespace UsdStageImporterImpl
 
 				if (bRecurse)
 				{
-					ResolveActorConflict(NewChild, *ExistingChild, ReplacePolicy, ObjectsToRemap);
+					ResolveActorConflict(NewChild, *ExistingChild, ReplacePolicy, ObjectsToRemap, SoftObjectsToRemap);
 				}
 			}
 		}
@@ -863,7 +846,7 @@ namespace UsdStageImporterImpl
 		// Handle component hierarchy collisions
 		USceneComponent* ExistingRoot = ExistingActor->GetRootComponent();
 		USceneComponent* NewRoot = NewActor->GetRootComponent();
-		ResolveComponentConflict(NewRoot, ExistingRoot, ReplacePolicy, ObjectsToRemap);
+		ResolveComponentConflict(NewRoot, ExistingRoot, ReplacePolicy, ObjectsToRemap, SoftObjectsToRemap);
 
 		// Move child actors over from existing hierarchy that don't conflict with anything in new hierarchy
 		// Do these later so that we don't recurse into them
@@ -880,7 +863,7 @@ namespace UsdStageImporterImpl
 		}
 	}
 
-	void ResolveActorConflicts(FUsdStageImportContext& ImportContext, AActor* ExistingSceneActor, TMap<UObject*, UObject*>& ObjectsToRemap)
+	void ResolveActorConflicts(FUsdStageImportContext& ImportContext, AActor* ExistingSceneActor, TMap<UObject*, UObject*>& ObjectsToRemap, TMap<FSoftObjectPath, FSoftObjectPath>& SoftObjectsToRemap)
 	{
 		if ( !ImportContext.ImportOptions->bImportActors )
 		{
@@ -901,7 +884,7 @@ namespace UsdStageImporterImpl
 			return;
 		}
 
-		ResolveActorConflict(ImportContext.SceneActor, ExistingSceneActor, ReplacePolicy, ObjectsToRemap);
+		ResolveActorConflict(ImportContext.SceneActor, ExistingSceneActor, ReplacePolicy, ObjectsToRemap, SoftObjectsToRemap);
 	}
 
 	// If we just reimported a static mesh, we use this to remap the material references to the existing materials, as any
@@ -1176,6 +1159,48 @@ namespace UsdStageImporterImpl
 			}
 		}
 	}
+
+	/**
+	 * Remaps asset's soft object pointers to point to the post-publish paths of their target assets.
+	 * It's important to run this *after* RemapReferences, as we will sometimes rely on those references to find our target assets.
+	 */
+	void RemapSoftReferences( const FUsdStageImportContext& ImportContext, const TSet<UObject*>& UsedAssetsAndDependencies, const TMap<FSoftObjectPath, FSoftObjectPath>& SoftObjectsToRemap )
+	{
+		TSet<UPackage*> Packages;
+		for ( UObject* Object : UsedAssetsAndDependencies )
+		{
+			if ( Object )
+			{
+				Packages.Add( Object->GetOutermost() );
+			}
+		}
+
+		if ( AActor* SceneActor = ImportContext.SceneActor )
+		{
+			Packages.Add( ImportContext.SceneActor->GetWorld()->GetOutermost() );
+		}
+
+		// In case one our used assets was left on the transient package.
+		// We don't care about anything that was left on the transient package, and doing this may actually cause some reference counting issues
+		// if we try deleting those assets afterwards
+		Packages.Remove( GetTransientPackage() );
+
+		IAssetTools& AssetTools = FModuleManager::LoadModuleChecked<FAssetToolsModule>( "AssetTools" ).Get();
+		AssetTools.RenameReferencingSoftObjectPaths( Packages.Array(), SoftObjectsToRemap );
+	}
+
+	/** After we remapped everything, notify the AssetRegistry that we created some new assets */
+	void NotifyAssetRegistry( const TSet<UObject*>& UsedAssetsAndDependencies )
+	{
+		for ( UObject* Object : UsedAssetsAndDependencies )
+		{
+			// If it's still on the transient package it means we abandoned this one (maybe we had asset replace policy ignore and hit a conflict)
+			if ( Object->GetOutermost() != GetTransientPackage() )
+			{
+				FAssetRegistryModule::AssetCreated( Object );
+			}
+		}
+	}
 }
 
 void UUsdStageImporter::ImportFromFile(FUsdStageImportContext& ImportContext)
@@ -1208,6 +1233,7 @@ void UUsdStageImporter::ImportFromFile(FUsdStageImportContext& ImportContext)
 
 	ImportContext.LevelSequenceHelper.Init(ImportContext.Stage);
 
+	TMap<FSoftObjectPath, FSoftObjectPath> SoftObjectsToRemap;
 	TMap<UObject*, UObject*> ObjectsToRemap;
 	TSet<UObject*> UsedAssetsAndDependencies;
 	UsdUtils::FBlendShapeMap BlendShapesByPath;
@@ -1225,7 +1251,7 @@ void UUsdStageImporter::ImportFromFile(FUsdStageImportContext& ImportContext)
 	TSharedRef<FUsdSchemaTranslationContext> TranslationContext = MakeShared<FUsdSchemaTranslationContext>( ImportContext.Stage, *ImportContext.AssetCache );
 	TranslationContext->Level = ImportContext.World->GetCurrentLevel();
 	TranslationContext->ObjectFlags = ImportContext.ImportObjectFlags;
-	TranslationContext->Time = ImportContext.ImportOptions->ImportTime;
+	TranslationContext->Time = static_cast< float >( UsdUtils::GetDefaultTimeCode() );
 	TranslationContext->PurposesToLoad = (EUsdPurpose) ImportContext.ImportOptions->PurposesToImport;
 	TranslationContext->RenderContext = ImportContext.ImportOptions->RenderContextToImport;
 	TranslationContext->ParentComponent = ImportContext.SceneActor ? ImportContext.SceneActor->GetRootComponent() : nullptr;
@@ -1243,10 +1269,12 @@ void UUsdStageImporter::ImportFromFile(FUsdStageImportContext& ImportContext)
 
 	UsdStageImporterImpl::CollectUsedAssetDependencies( ImportContext.AssetCache, UsedAssetsAndDependencies );
 	UsdStageImporterImpl::UpdateAssetImportData( UsedAssetsAndDependencies, ImportContext.FilePath, ImportContext.ImportOptions );
-	UsdStageImporterImpl::PublishAssets( ImportContext, UsedAssetsAndDependencies, ObjectsToRemap );
-	UsdStageImporterImpl::ResolveActorConflicts( ImportContext, ExistingSceneActor, ObjectsToRemap );
+	UsdStageImporterImpl::PublishAssets( ImportContext, UsedAssetsAndDependencies, ObjectsToRemap, SoftObjectsToRemap );
+	UsdStageImporterImpl::ResolveActorConflicts( ImportContext, ExistingSceneActor, ObjectsToRemap, SoftObjectsToRemap );
 	UsdStageImporterImpl::RemapReferences( ImportContext, UsedAssetsAndDependencies, ObjectsToRemap );
+	UsdStageImporterImpl::RemapSoftReferences( ImportContext, UsedAssetsAndDependencies, SoftObjectsToRemap );
 	UsdStageImporterImpl::Cleanup( ImportContext.SceneActor, ExistingSceneActor, ImportContext.ImportOptions->ExistingActorPolicy );
+	UsdStageImporterImpl::NotifyAssetRegistry( UsedAssetsAndDependencies );
 	UsdStageImporterImpl::CloseStageIfNeeded( ImportContext );
 
 	FUsdDelegates::OnPostUsdImport.Broadcast( ImportContext.FilePath );
@@ -1274,6 +1302,7 @@ bool UUsdStageImporter::ReimportSingleAsset(FUsdStageImportContext& ImportContex
 
 	UsdStageImporterImpl::SetupStageForImport( ImportContext );
 
+	TMap<FSoftObjectPath, FSoftObjectPath> SoftObjectsToRemap;
 	TMap<UObject*, UObject*> ObjectsToRemap;
 	UsdUtils::FBlendShapeMap BlendShapesByPath;
 
@@ -1290,7 +1319,7 @@ bool UUsdStageImporter::ReimportSingleAsset(FUsdStageImportContext& ImportContex
 	TSharedRef<FUsdSchemaTranslationContext> TranslationContext = MakeShared<FUsdSchemaTranslationContext>( ImportContext.Stage, *ImportContext.AssetCache );
 	TranslationContext->Level = ImportContext.World->GetCurrentLevel();
 	TranslationContext->ObjectFlags = ImportContext.ImportObjectFlags;
-	TranslationContext->Time = ImportContext.ImportOptions->ImportTime;
+	TranslationContext->Time = static_cast< float >( UsdUtils::GetDefaultTimeCode() );
 	TranslationContext->PurposesToLoad = (EUsdPurpose) ImportContext.ImportOptions->PurposesToImport;
 	TranslationContext->bAllowCollapsing = ImportContext.ImportOptions->bCollapse;
 	TranslationContext->bAllowInterpretingLODs = ImportContext.ImportOptions->bInterpretLODs;
@@ -1331,13 +1360,15 @@ bool UUsdStageImporter::ReimportSingleAsset(FUsdStageImportContext& ImportContex
 		UsdStageImporterImpl::CopySkeletonAssignment(ImportContext, OriginalAsset, ReimportedObject );
 
 		// Just publish the one asset we wanted to reimport. Note that we may have other assets here too, but we'll ignore those e.g. a displayColor material or a skeleton
-		OutReimportedAsset = UsdStageImporterImpl::PublishAsset(ImportContext, ReimportedObject, OriginalAsset->GetOutermost()->GetPathName(), ObjectsToRemap);
+		OutReimportedAsset = UsdStageImporterImpl::PublishAsset(ImportContext, ReimportedObject, OriginalAsset->GetOutermost()->GetPathName(), ObjectsToRemap, SoftObjectsToRemap);
 		UsdStageImporterImpl::RemapReferences( ImportContext, ImportContext.AssetCache->GetActiveAssets(), ObjectsToRemap );
+		UsdStageImporterImpl::RemapSoftReferences( ImportContext, ImportContext.AssetCache->GetActiveAssets(), SoftObjectsToRemap );
 
 		bSuccess = OutReimportedAsset != nullptr && ImportContext.AssetCache->GetActiveAssets().Contains( ReimportedObject );
 	}
 
 	UsdStageImporterImpl::Cleanup( ImportContext.SceneActor, nullptr, ImportContext.ImportOptions->ExistingActorPolicy );
+	UsdStageImporterImpl::NotifyAssetRegistry( { ReimportedObject } );
 	UsdStageImporterImpl::CloseStageIfNeeded( ImportContext );
 
 	FUsdDelegates::OnPostUsdImport.Broadcast(ImportContext.FilePath);
diff --git a/Engine/Plugins/Importers/USDImporter/Source/USDStageImporter/Public/USDStageImportOptions.h b/Engine/Plugins/Importers/USDImporter/Source/USDStageImporter/Public/USDStageImportOptions.h
--- a/Engine/Plugins/Importers/USDImporter/Source/USDStageImporter/Public/USDStageImportOptions.h
+++ b/Engine/Plugins/Importers/USDImporter/Source/USDStageImporter/Public/USDStageImportOptions.h
@@ -70,10 +70,6 @@ public:
 	UPROPERTY(BlueprintReadWrite, config, EditAnywhere, Category= "USD options")
 	FName RenderContextToImport;
 
-	/** Time to evaluate the USD Stage for import */
-	UPROPERTY(BlueprintReadWrite, config, EditAnywhere, Category= "USD options", meta = (DisplayName = "Time"))
-	float ImportTime;
-
 	/** Whether to use the specified StageOptions instead of the stage's own settings */
 	UPROPERTY( BlueprintReadWrite, config, EditAnywhere, Category = "USD options" )
 	bool bOverrideStageOptions;
