USD: Fix spawned HISM components showing as invisible if the source Prototype scope prim is marked as invisible.

Fix hiding and unhiding collapsed meshes permanently hiding the meshes;
Fix crash when we can't find the prototype for a point instancer;
Fix point instancer prims secretly spawning many duplicates of their HISM components.

diff --git a/Engine/Plugins/Importers/USDImporter/Source/USDSchemas/Private/USDGeomPointInstancerTranslator.cpp b/Engine/Plugins/Importers/USDImporter/Source/USDSchemas/Private/USDGeomPointInstancerTranslator.cpp
--- a/Engine/Plugins/Importers/USDImporter/Source/USDSchemas/Private/USDGeomPointInstancerTranslator.cpp
+++ b/Engine/Plugins/Importers/USDImporter/Source/USDSchemas/Private/USDGeomPointInstancerTranslator.cpp
@@ -6,11 +6,16 @@
 #include "USDAssetCache.h"
 #include "USDAssetImportData.h"
 #include "USDConversionUtils.h"
+#include "USDLog.h"
 #include "USDSchemasModule.h"
 #include "USDTypesConversion.h"
 
 #include "UsdWrappers/UsdTyped.h"
 
+#if WITH_EDITOR
+#include "Editor.h"
+#endif // WITH_EDITOR
+
 #if USE_USD_SDK
 
 #include "Components/HierarchicalInstancedStaticMeshComponent.h"
@@ -136,6 +141,44 @@ void FUsdGeomPointInstancerTranslator::UpdateComponents( USceneComponent* PointI
 		return;
 	}
 
+	// The components that we will spawn in here will not have prim twins, and so would otherwise never be destroyed.
+	// Here we make sure we get rid of the old ones because we will create new ones right afterwards
+	PointInstancerRootComponent->Modify();
+	TArray<USceneComponent*> Children;
+	const bool bIncludeAllDescendants = true;
+	PointInstancerRootComponent->GetChildrenComponents( bIncludeAllDescendants, Children );
+	for ( USceneComponent* Child : Children )
+	{
+		// Attempt at making sure it's one of our components just in case the user happened to put something in there
+		if ( !Child->ComponentTags.Contains( UnrealIdentifiers::Invisible) &&
+			 !Child->ComponentTags.Contains( UnrealIdentifiers::Inherited) )
+		{
+			continue;
+		}
+
+		Child->Modify();
+
+		const bool bPromoteChildren = true;
+		Child->DestroyComponent( bPromoteChildren );
+	}
+
+#if WITH_EDITOR
+	// If we have this actor selected we need to refresh the details panel manually on the next frame because
+	// deleting those components will leave it in some type of invalid state
+	if ( PointInstancerRootComponent->GetOwner()->IsSelected() )
+	{
+		GEditor->GetTimerManager()->SetTimerForNextTick( []()
+		{
+			GEditor->NoteSelectionChange();
+		} );
+	}
+#endif // WITH_EDITOR
+
+	// We have to manage the visibilities of the components that we spawn ourselves, because they will not have prim twins
+	// and so will not be handled by the stage actor. Plus there are some weird rules involved, like how we don't care about the
+	// visibilities of *parents* of prototype prims, only theirs and down from there
+	const bool bPointInstancerVisible = UsdUtils::IsVisible( UE::FUsdPrim{ Prim }, Context->Time );
+
 	// For each prototype
 	const pxr::UsdRelationship& Prototypes = PointInstancer.GetPrototypesRel();
 	pxr::SdfPathVector PrototypesPaths;
@@ -148,30 +191,40 @@ void FUsdGeomPointInstancerTranslator::UpdateComponents( USceneComponent* PointI
 
 		for ( int32 PrototypeIndex = 0; PrototypeIndex < PrototypesPaths.size(); ++PrototypeIndex)
 		{
-			pxr::UsdPrim PrototypePrim = Prim.GetStage()->GetPrimAtPath( PrototypesPaths[PrototypeIndex] );
-			if ( !Prim )
+			pxr::UsdPrim PrototypeUsdPrim = Prim.GetStage()->GetPrimAtPath( PrototypesPaths[PrototypeIndex] );
+			if ( !PrototypeUsdPrim )
 			{
+				UE_LOG( LogUsd, Warning, TEXT( "Failed to find prototype '%s' for PointInstancer '%s'" ), *UsdToUnreal::ConvertPath( PrototypesPaths[ PrototypeIndex ] ), *PrimPath.GetString() );
 				continue;
 			}
+			UE::FUsdPrim PrototypePrim{ PrototypeUsdPrim };
 
 			USceneComponent* PrototypeParentComponent = Context->ParentComponent;
 
+			bool bPrototypeRootVisible = true;
+
 			// Temp fix to prevent us from creating yet another UStaticMeshComponent as a parent prim to the HISM component
 			const bool bNeedsActor = false;
 			TOptional<TGuardValue< USceneComponent* >> ParentComponentGuard2;
-			if ( !PrototypePrim.IsA<pxr::UsdGeomMesh>() )
+			if ( !PrototypeUsdPrim.IsA<pxr::UsdGeomMesh>() )
 			{
-				FUsdGeomXformableTranslator PrototypeXformTranslator( Context, UE::FUsdTyped( PrototypePrim ) );
+				FUsdGeomXformableTranslator PrototypeXformTranslator( Context, UE::FUsdTyped( PrototypeUsdPrim ) );
 
 				if ( USceneComponent* PrototypeXformComponent = PrototypeXformTranslator.CreateComponentsEx( {}, bNeedsActor ) )
 				{
+					// We have to be careful with prototype visibilities and can't defer to just using ComputeVisibility() because sometimes source apps
+					// can mark the prototype parent prim as invisible, as an attempt to prevent them from being parsed directly as scene meshes.
+					// We will handle the visibility of all the components we spawn manually
+					bPrototypeRootVisible = UsdUtils::HasInheritedVisibility( PrototypePrim, Context->Time );
+					PrototypeXformComponent->SetVisibility( bPointInstancerVisible && bPrototypeRootVisible );
+
 					PrototypeParentComponent = PrototypeXformComponent;
 				}
 
 				ParentComponentGuard2.Emplace(Context->ParentComponent, PrototypeParentComponent );
 			}
 
-			TArray< TUsdStore< pxr::UsdPrim > > ChildGeomMeshPrims = UsdUtils::GetAllPrimsOfType( PrototypePrim, pxr::TfType::Find< pxr::UsdGeomMesh >() );
+			TArray< TUsdStore< pxr::UsdPrim > > ChildGeomMeshPrims = UsdUtils::GetAllPrimsOfType( PrototypeUsdPrim, pxr::TfType::Find< pxr::UsdGeomMesh >() );
 
 			for ( const TUsdStore< pxr::UsdPrim >& PrototypeGeomMeshPrim : ChildGeomMeshPrims )
 			{
@@ -183,15 +236,22 @@ void FUsdGeomPointInstancerTranslator::UpdateComponents( USceneComponent* PointI
 				{
 					ProcessedPrims.Add( UEPrototypeTargetPrimPath );
 
-					if ( pxr::UsdPrim PrototypeTargetPrim = Prim.GetStage()->GetPrimAtPath( PrototypeTargetPrimPath.Get() ) )
+					if ( pxr::UsdPrim PrototypeTargetUsdPrim = Prim.GetStage()->GetPrimAtPath( PrototypeTargetPrimPath.Get() ) )
 					{
-						pxr::UsdGeomMesh PrototypeGeomMesh( PrototypeTargetPrim );
+						pxr::UsdGeomMesh PrototypeGeomMesh( PrototypeTargetUsdPrim );
 						FUsdGeomXformableTranslator PrototypeGeomXformableTranslator( UHierarchicalInstancedStaticMeshComponent::StaticClass(), Context, UE::FUsdTyped( PrototypeGeomMesh ) );
 
 						USceneComponent* UsdGeomPrimComponent = PrototypeGeomXformableTranslator.CreateComponentsEx( {}, bNeedsActor );
 
 						if ( UHierarchicalInstancedStaticMeshComponent* HismComponent = Cast< UHierarchicalInstancedStaticMeshComponent >( UsdGeomPrimComponent ) )
 						{
+							UE::FUsdPrim PrototypeTargetPrim{ PrototypeTargetUsdPrim };
+
+							// Traverse up every time because we have no idea how far deep each UsdGeomMesh prim is
+							const bool bHasInvisibleParent = UsdUtils::HasInvisibleParent( PrototypeTargetPrim, PrototypePrim, Context->Time );
+							const bool bPrototypeMeshVisible = UsdUtils::HasInheritedVisibility( PrototypeTargetPrim, Context->Time );
+							HismComponent->SetVisibility( bPointInstancerVisible && bPrototypeRootVisible && bPrototypeMeshVisible && !bHasInvisibleParent );
+
 							UUsdAssetCache& AssetCache = *Context->AssetCache.Get();
 							UStaticMesh* StaticMesh = UsdGeomPointInstancerTranslatorImpl::SetStaticMesh( PrototypeGeomMesh, *HismComponent, AssetCache );
 							UsdGeomPointInstancerTranslatorImpl::ConvertGeomPointInstancer( Prim.GetStage(), PointInstancer, PrototypeIndex, *HismComponent, pxr::UsdTimeCode( Context->Time ) );
@@ -214,7 +274,7 @@ void FUsdGeomPointInstancerTranslator::UpdateComponents( USceneComponent* PointI
 									}
 
 									MeshTranslationImpl::SetMaterialOverrides(
-										PrototypeTargetPrim,
+										PrototypeTargetUsdPrim,
 										ExistingAssignments,
 										*HismComponent,
 										AssetCache,
diff --git a/Engine/Plugins/Importers/USDImporter/Source/USDSchemas/Private/USDGeomXformableTranslator.cpp b/Engine/Plugins/Importers/USDImporter/Source/USDSchemas/Private/USDGeomXformableTranslator.cpp
--- a/Engine/Plugins/Importers/USDImporter/Source/USDSchemas/Private/USDGeomXformableTranslator.cpp
+++ b/Engine/Plugins/Importers/USDImporter/Source/USDSchemas/Private/USDGeomXformableTranslator.cpp
@@ -318,9 +318,13 @@ USceneComponent* FUsdGeomXformableTranslator::CreateComponentsEx( TOptional< TSu
 
 				if ( CollapsesChildren( ECollapsingType::Assets ) )
 				{
-					if ( UStaticMesh* PrimStaticMesh = Cast< UStaticMesh >( Context->AssetCache->GetAssetForPrim( PrimPath.GetString() ) ) )
+					// If we're a type that collapses assets, we should probably be a static mesh component as we only really collapse static meshes together right now.
+					// We can't just check if there's a static mesh for this prim on the cache, because the prims with meshes could be potentially invisible (and so
+					// we don't have parsed their meshes yet), so here we traverse our child hierarchy and if we have any chance of ever generating a Mesh, we go
+					// for a static mesh component
+					TArray< UE::FUsdPrim > ChildMeshPrims = UsdUtils::GetAllPrimsOfType( Prim, TEXT( "UsdGeomMesh" ) );
+					if ( ChildMeshPrims.Num() > 0 )
 					{
-						// At this time, we only support collapsing static meshes together
 						ComponentType = UStaticMeshComponent::StaticClass();
 					}
 				}
diff --git a/Engine/Plugins/Importers/USDImporter/Source/USDUtilities/Private/USDConversionUtils.cpp b/Engine/Plugins/Importers/USDImporter/Source/USDUtilities/Private/USDConversionUtils.cpp
--- a/Engine/Plugins/Importers/USDImporter/Source/USDUtilities/Private/USDConversionUtils.cpp
+++ b/Engine/Plugins/Importers/USDImporter/Source/USDUtilities/Private/USDConversionUtils.cpp
@@ -1029,7 +1029,7 @@ void UsdUtils::MakeInvisible( UE::FUsdPrim& Prim, double TimeCode )
 #endif // USE_USD_SDK
 }
 
-bool UsdUtils::IsVisible( UE::FUsdPrim& Prim, double TimeCode )
+bool UsdUtils::IsVisible( const UE::FUsdPrim& Prim, double TimeCode )
 {
 #if USE_USD_SDK
 	FScopedUsdAllocs Allocs;
@@ -1046,7 +1046,7 @@ bool UsdUtils::IsVisible( UE::FUsdPrim& Prim, double TimeCode )
 #endif // USE_USD_SDK
 }
 
-bool UsdUtils::HasInheritedVisibility( UE::FUsdPrim& Prim, double TimeCode )
+bool UsdUtils::HasInheritedVisibility( const UE::FUsdPrim& Prim, double TimeCode )
 {
 #if USE_USD_SDK
 	FScopedUsdAllocs Allocs;
@@ -1073,6 +1073,38 @@ bool UsdUtils::HasInheritedVisibility( UE::FUsdPrim& Prim, double TimeCode )
 #endif // USE_USD_SDK
 }
 
+bool UsdUtils::HasInvisibleParent( const UE::FUsdPrim& Prim, const UE::FUsdPrim& RootPrim, double TimeCode )
+{
+#if USE_USD_SDK
+	FScopedUsdAllocs Allocs;
+
+	pxr::UsdPrim PxrUsdPrim{ Prim };
+	pxr::UsdPrim Parent = PxrUsdPrim.GetParent();
+
+	while ( Parent && Parent != RootPrim )
+	{
+		if ( pxr::UsdGeomImageable Imageable{ Parent } )
+		{
+			if ( pxr::UsdAttribute VisibilityAttr = Imageable.GetVisibilityAttr() )
+			{
+				pxr::TfToken Visibility;
+				if ( VisibilityAttr.Get<pxr::TfToken>( &Visibility, TimeCode ) )
+				{
+					if ( Visibility == pxr::UsdGeomTokens->invisible )
+					{
+						return true;
+					}
+				}
+			}
+		}
+
+		Parent = Parent.GetParent();
+	}
+#endif // USE_USD_SDK
+
+	return false;
+}
+
 UE::FSdfPath UsdUtils::GetPrimSpecPathForLayer( const UE::FUsdPrim& Prim, const UE::FSdfLayer& Layer )
 {
 	UE::FSdfPath Result;
diff --git a/Engine/Plugins/Importers/USDImporter/Source/USDUtilities/Public/USDConversionUtils.h b/Engine/Plugins/Importers/USDImporter/Source/USDUtilities/Public/USDConversionUtils.h
--- a/Engine/Plugins/Importers/USDImporter/Source/USDUtilities/Public/USDConversionUtils.h
+++ b/Engine/Plugins/Importers/USDImporter/Source/USDUtilities/Public/USDConversionUtils.h
@@ -180,10 +180,17 @@ namespace UsdUtils
 	USDUTILITIES_API void MakeInvisible( UE::FUsdPrim& Prim, double TimeCode = UsdUtils::GetDefaultTimeCode() );
 
 	/** Returns if the ComputedVisibility for Prim says it should be visible */
-	USDUTILITIES_API bool IsVisible( UE::FUsdPrim& Prim, double TimeCode = UsdUtils::GetDefaultTimeCode() );
+	USDUTILITIES_API bool IsVisible( const UE::FUsdPrim& Prim, double TimeCode = UsdUtils::GetDefaultTimeCode() );
 
 	/** Returns whether Prim has visibility set to 'inherited' */
-	USDUTILITIES_API bool HasInheritedVisibility( UE::FUsdPrim& Prim, double TimeCode = UsdUtils::GetDefaultTimeCode() );
+	USDUTILITIES_API bool HasInheritedVisibility( const UE::FUsdPrim& Prim, double TimeCode = UsdUtils::GetDefaultTimeCode() );
+
+	/**
+	 * Travels up from Prim and returns true if we hit any invisible prim before we hit the stage pseudoroot.
+	 * Does not check `Prim` itself (or `RootPrim`, so that you can just pass the PseudoRoot in most cases).
+	 * Completely ignores whether non-UsdGeomImageable prims are visible or not and keeps travelling up.
+	 */
+	USDUTILITIES_API bool HasInvisibleParent( const UE::FUsdPrim& Prim, const UE::FUsdPrim& RootPrim, double TimeCode = UsdUtils::GetDefaultTimeCode() );
 
 	/**
 	 * Returns a path exactly like Prim.GetPrimPath(), except that if the prim is within variant sets, it will return the
