USD: Fixup bug introduced with cl 17006546, that broke PoseableMeshComponents from ever being spawned.

I didn't notice that PoseableMeshComponents also fit the new criteria for when we should upgrade a scene component to a static mesh component. This change just prevents that case and only ever "upgrades" from a scene component to a static mesh component instead.

diff --git a/Engine/Plugins/Importers/USDImporter/Source/USDSchemas/Private/USDGeomXformableTranslator.cpp b/Engine/Plugins/Importers/USDImporter/Source/USDSchemas/Private/USDGeomXformableTranslator.cpp
--- a/Engine/Plugins/Importers/USDImporter/Source/USDSchemas/Private/USDGeomXformableTranslator.cpp
+++ b/Engine/Plugins/Importers/USDImporter/Source/USDSchemas/Private/USDGeomXformableTranslator.cpp
@@ -316,7 +316,9 @@ USceneComponent* FUsdGeomXformableTranslator::CreateComponentsEx( TOptional< TSu
 			{
 				ComponentType = UsdUtils::GetComponentTypeForPrim( Prim );
 
-				if ( CollapsesChildren( ECollapsingType::Assets ) )
+				// For now only upgrade actual scene components to static mesh components (important because poseable mesh components will also fit this
+				// criteria but we don't want to use a static mesh component for those)
+				if ( CollapsesChildren( ECollapsingType::Assets ) && ComponentType.IsSet() && ComponentType.GetValue() == USceneComponent::StaticClass() )
 				{
 					// If we're a type that collapses assets, we should probably be a static mesh component as we only really collapse static meshes together right now.
 					// We can't just check if there's a static mesh for this prim on the cache, because the prims with meshes could be potentially invisible (and so
