USD: Port static/skeletal mesh exporter LOD range option from cl 16922714

It includes Jean-Luc's comments on swarm review 16901915. Merging the change didn't so much as lead to a merge conflict so I don't think there was any point in creating another review

diff --git a/Engine/Plugins/Importers/USDImporter/Source/USDExporter/Private/AnimSequenceExporterUSD.cpp b/Engine/Plugins/Importers/USDImporter/Source/USDExporter/Private/AnimSequenceExporterUSD.cpp
--- a/Engine/Plugins/Importers/USDImporter/Source/USDExporter/Private/AnimSequenceExporterUSD.cpp
+++ b/Engine/Plugins/Importers/USDImporter/Source/USDExporter/Private/AnimSequenceExporterUSD.cpp
@@ -133,7 +133,7 @@ bool UAnimSequenceExporterUSD::ExportBinary( UObject* Object, const TCHAR* Type,
 
 					if ( UE::FUsdPrim PayloadSkelRootPrim = PayloadStage.DefinePrim( UE::FSdfPath( *PrimPath ), TEXT( "SkelRoot" ) ) )
 					{
-						UnrealToUsd::ConvertSkeletalMesh( SkeletalMesh, PayloadSkelRootPrim, UsdUtils::GetDefaultTimeCode(), &AssetStage );
+						UnrealToUsd::ConvertSkeletalMesh( SkeletalMesh, PayloadSkelRootPrim, UsdUtils::GetDefaultTimeCode(), &AssetStage, Options->PreviewMeshOptions.LowestMeshLOD, Options->PreviewMeshOptions.HighestMeshLOD );
 						PayloadStage.SetDefaultPrim( PayloadSkelRootPrim );
 					}
 
@@ -149,7 +149,7 @@ bool UAnimSequenceExporterUSD::ExportBinary( UObject* Object, const TCHAR* Type,
 			// Not using payload: Just author everything on the current edit target of the payload (== asset) layer
 			else
 			{
-				UnrealToUsd::ConvertSkeletalMesh( SkeletalMesh, SkelRootPrim );
+				UnrealToUsd::ConvertSkeletalMesh( SkeletalMesh, SkelRootPrim, UsdUtils::GetDefaultTimeCode(), nullptr, Options->PreviewMeshOptions.LowestMeshLOD, Options->PreviewMeshOptions.HighestMeshLOD );
 			}
 		}
 		else
diff --git a/Engine/Plugins/Importers/USDImporter/Source/USDExporter/Private/SkeletalMeshExporterUSD.cpp b/Engine/Plugins/Importers/USDImporter/Source/USDExporter/Private/SkeletalMeshExporterUSD.cpp
--- a/Engine/Plugins/Importers/USDImporter/Source/USDExporter/Private/SkeletalMeshExporterUSD.cpp
+++ b/Engine/Plugins/Importers/USDImporter/Source/USDExporter/Private/SkeletalMeshExporterUSD.cpp
@@ -123,7 +123,7 @@ bool USkeletalMeshExporterUsd::ExportBinary( UObject* Object, const TCHAR* Type,
 				UsdUtils::AddPayload( AssetRootPrim, *PayloadFilename );
 			}
 
-			UnrealToUsd::ConvertSkeletalMesh( SkeletalMesh, RootPrim, UsdUtils::GetDefaultTimeCode(), &AssetStage );
+			UnrealToUsd::ConvertSkeletalMesh( SkeletalMesh, RootPrim, UsdUtils::GetDefaultTimeCode(), &AssetStage, Options->Inner.LowestMeshLOD, Options->Inner.HighestMeshLOD );
 
 			AssetStage.GetRootLayer().Save();
 		}
@@ -131,7 +131,15 @@ bool USkeletalMeshExporterUsd::ExportBinary( UObject* Object, const TCHAR* Type,
 	// Not using payload: Just author everything on the current edit target of the payload (== asset) layer
 	else
 	{
-		UnrealToUsd::ConvertSkeletalMesh( SkeletalMesh, RootPrim );
+		int32 LowestLOD = 0;
+		int32 HighestLOD = MAX_MESH_LOD_COUNT - 1;
+		if ( Options )
+		{
+			LowestLOD = Options->Inner.LowestMeshLOD;
+			HighestLOD = Options->Inner.HighestMeshLOD;
+		}
+
+		UnrealToUsd::ConvertSkeletalMesh( SkeletalMesh, RootPrim, UsdUtils::GetDefaultTimeCode(), nullptr, LowestLOD, HighestLOD );
 	}
 
 	UsdStage.GetRootLayer().Save();
diff --git a/Engine/Plugins/Importers/USDImporter/Source/USDExporter/Private/StaticMeshExporterUSD.cpp b/Engine/Plugins/Importers/USDImporter/Source/USDExporter/Private/StaticMeshExporterUSD.cpp
--- a/Engine/Plugins/Importers/USDImporter/Source/USDExporter/Private/StaticMeshExporterUSD.cpp
+++ b/Engine/Plugins/Importers/USDImporter/Source/USDExporter/Private/StaticMeshExporterUSD.cpp
@@ -131,7 +131,7 @@ bool UStaticMeshExporterUsd::ExportBinary( UObject* Object, const TCHAR* Type, F
 				UsdUtils::AddPayload( AssetRootPrim, *PayloadFilename );
 			}
 
-			UnrealToUsd::ConvertStaticMesh( StaticMesh, RootPrim, UsdUtils::GetDefaultTimeCode(), &AssetStage );
+			UnrealToUsd::ConvertStaticMesh( StaticMesh, RootPrim, UsdUtils::GetDefaultTimeCode(), &AssetStage, Options->LowestMeshLOD, Options->HighestMeshLOD );
 
 			AssetStage.GetRootLayer().Save();
 		}
@@ -139,7 +139,15 @@ bool UStaticMeshExporterUsd::ExportBinary( UObject* Object, const TCHAR* Type, F
 	// Not using payload: Just author everything on the current edit target of the payload (== asset) layer
 	else
 	{
-		UnrealToUsd::ConvertStaticMesh( StaticMesh, RootPrim );
+		int32 LowestLOD = 0;
+		int32 HighestLOD = MAX_MESH_LOD_COUNT - 1;
+		if ( Options )
+		{
+			LowestLOD = Options->LowestMeshLOD;
+			HighestLOD = Options->HighestMeshLOD;
+		}
+
+		UnrealToUsd::ConvertStaticMesh( StaticMesh, RootPrim, UsdUtils::GetDefaultTimeCode(), nullptr, LowestLOD, HighestLOD );
 	}
 
 	UsdStage.GetRootLayer().Save();
diff --git a/Engine/Plugins/Importers/USDImporter/Source/USDExporter/Public/LevelExporterUSDOptions.h b/Engine/Plugins/Importers/USDImporter/Source/USDExporter/Public/LevelExporterUSDOptions.h
--- a/Engine/Plugins/Importers/USDImporter/Source/USDExporter/Public/LevelExporterUSDOptions.h
+++ b/Engine/Plugins/Importers/USDImporter/Source/USDExporter/Public/LevelExporterUSDOptions.h
@@ -8,6 +8,7 @@
 
 #include "AssetExportTask.h"
 #include "Engine/EngineTypes.h"
+#include "RHIDefinitions.h"
 
 #include "LevelExporterUSDOptions.generated.h"
 
@@ -56,6 +57,14 @@ public:
 	UPROPERTY( EditAnywhere, config, BlueprintReadWrite, Category = "Export settings", meta = ( EditCondition = "bUsePayload", GetOptions = GetUsdExtensions ) )
 	FString PayloadFormat;
 
+	/** Lowest of the LOD indices to export static and skeletal meshes with (use 0 for maximum detail) */
+	UPROPERTY( EditAnywhere, config, BlueprintReadWrite, Category = "Export settings", meta = ( ClampMin = "0" ) )
+	int32 LowestMeshLOD = 0;
+
+	/** Highest of the LOD indices to export static and skeletal meshes with */
+	UPROPERTY( EditAnywhere, config, BlueprintReadWrite, Category = "Export settings", meta = ( ClampMin = "0" ) )
+	int32 HighestMeshLOD = MAX_MESH_LOD_COUNT - 1;
+
 	/** Whether to use UE actor folders as empty prims */
     UPROPERTY( EditAnywhere, config, BlueprintReadWrite, Category = "Export settings" )
     bool bExportActorFolders;
diff --git a/Engine/Plugins/Importers/USDImporter/Source/USDExporter/Public/SkeletalMeshExporterUSDOptions.h b/Engine/Plugins/Importers/USDImporter/Source/USDExporter/Public/SkeletalMeshExporterUSDOptions.h
--- a/Engine/Plugins/Importers/USDImporter/Source/USDExporter/Public/SkeletalMeshExporterUSDOptions.h
+++ b/Engine/Plugins/Importers/USDImporter/Source/USDExporter/Public/SkeletalMeshExporterUSDOptions.h
@@ -7,6 +7,7 @@
 #include "USDStageOptions.h"
 
 #include "Engine/EngineTypes.h"
+#include "RHIDefinitions.h"
 
 #include "SkeletalMeshExporterUSDOptions.generated.h"
 
@@ -29,6 +30,14 @@ struct USDEXPORTER_API FSkeletalMeshExporterUSDInnerOptions
 	/** USD format to use for exported payload files */
 	UPROPERTY( EditAnywhere, config, BlueprintReadWrite, Category = USDSettings, meta = ( EditCondition = "bUsePayload", GetOptions = "USDExporter.LevelExporterUSDOptions.GetUsdExtensions" ) )
 	FString PayloadFormat;
+
+	/** Lowest of the LOD indices to export static and skeletal meshes with (use 0 for maximum detail) */
+	UPROPERTY( EditAnywhere, config, BlueprintReadWrite, Category = "Export settings", meta = ( ClampMin = "0" ) )
+	int32 LowestMeshLOD = 0;
+
+	/** Highest of the LOD indices to export static and skeletal meshes with */
+	UPROPERTY( EditAnywhere, config, BlueprintReadWrite, Category = "Export settings", meta = ( ClampMin = "0" ) )
+	int32 HighestMeshLOD = MAX_MESH_LOD_COUNT - 1;
 };
 
 /**
diff --git a/Engine/Plugins/Importers/USDImporter/Source/USDExporter/Public/StaticMeshExporterUSDOptions.h b/Engine/Plugins/Importers/USDImporter/Source/USDExporter/Public/StaticMeshExporterUSDOptions.h
--- a/Engine/Plugins/Importers/USDImporter/Source/USDExporter/Public/StaticMeshExporterUSDOptions.h
+++ b/Engine/Plugins/Importers/USDImporter/Source/USDExporter/Public/StaticMeshExporterUSDOptions.h
@@ -7,6 +7,7 @@
 #include "USDStageOptions.h"
 
 #include "Engine/EngineTypes.h"
+#include "RHIDefinitions.h"
 
 #include "StaticMeshExporterUSDOptions.generated.h"
 
@@ -29,4 +30,12 @@ public:
 	/** USD format to use for exported payload files */
 	UPROPERTY( EditAnywhere, config, BlueprintReadWrite, Category = USDSettings, meta = ( EditCondition = "bUsePayload", GetOptions = "USDExporter.LevelExporterUSDOptions.GetUsdExtensions" ) )
 	FString PayloadFormat;
+
+	/** Lowest of the LOD indices to export static and skeletal meshes with (use 0 for maximum detail) */
+	UPROPERTY( EditAnywhere, config, BlueprintReadWrite, Category = "Export settings", meta = ( ClampMin = "0" ) )
+	int32 LowestMeshLOD = 0;
+
+	/** Highest of the LOD indices to export static and skeletal meshes with */
+	UPROPERTY( EditAnywhere, config, BlueprintReadWrite, Category = "Export settings", meta = ( ClampMin = "0" ) )
+	int32 HighestMeshLOD = MAX_MESH_LOD_COUNT - 1;
 };
diff --git a/Engine/Plugins/Importers/USDImporter/Source/USDUtilities/Private/USDGeomMeshConversion.cpp b/Engine/Plugins/Importers/USDImporter/Source/USDUtilities/Private/USDGeomMeshConversion.cpp
--- a/Engine/Plugins/Importers/USDImporter/Source/USDUtilities/Private/USDGeomMeshConversion.cpp
+++ b/Engine/Plugins/Importers/USDImporter/Source/USDUtilities/Private/USDGeomMeshConversion.cpp
@@ -1322,7 +1322,7 @@ UsdUtils::FUsdPrimMaterialAssignmentInfo UsdUtils::GetPrimMaterialAssignments( c
 	return Result;
 }
 
-bool UnrealToUsd::ConvertStaticMesh( const UStaticMesh* StaticMesh, pxr::UsdPrim& UsdPrim, const pxr::UsdTimeCode TimeCode, UE::FUsdStage* StageForMaterialAssignments )
+bool UnrealToUsd::ConvertStaticMesh( const UStaticMesh* StaticMesh, pxr::UsdPrim& UsdPrim, const pxr::UsdTimeCode TimeCode, UE::FUsdStage* StageForMaterialAssignments, int32 LowestMeshLOD, int32 HighestMeshLOD )
 {
 	FScopedUsdAllocs UsdAllocs;
 
@@ -1340,6 +1340,18 @@ bool UnrealToUsd::ConvertStaticMesh( const UStaticMesh* StaticMesh, pxr::UsdPrim
 		return false;
 	}
 
+	// Make sure they're both >= 0 (the options dialog slider is clamped, but this may be called directly)
+	LowestMeshLOD = FMath::Clamp( LowestMeshLOD, 0, NumLODs -1 );
+	HighestMeshLOD = FMath::Clamp( HighestMeshLOD, 0, NumLODs - 1 );
+
+	// Make sure Lowest <= Highest
+	int32 Temp = FMath::Min( LowestMeshLOD, HighestMeshLOD );
+	HighestMeshLOD = FMath::Max( LowestMeshLOD, HighestMeshLOD );
+	LowestMeshLOD = Temp;
+
+	// Make sure it's at least 1 LOD level
+	NumLODs = FMath::Max( HighestMeshLOD - LowestMeshLOD + 1, 1 );
+
 	pxr::UsdVariantSets VariantSets = UsdPrim.GetVariantSets();
 	if ( NumLODs > 1 && VariantSets.HasVariantSet( UnrealIdentifiers::LOD ) )
 	{
@@ -1375,7 +1387,7 @@ bool UnrealToUsd::ConvertStaticMesh( const UStaticMesh* StaticMesh, pxr::UsdPrim
 		MaterialAssignments.clear();
 	}
 
-	for ( int32 LODIndex = 0; LODIndex < NumLODs; ++LODIndex )
+	for ( int32 LODIndex = LowestMeshLOD; LODIndex <= HighestMeshLOD; ++LODIndex )
 	{
 		const FStaticMeshLODResources& RenderMesh = StaticMesh->GetLODForExport( LODIndex );
 
diff --git a/Engine/Plugins/Importers/USDImporter/Source/USDUtilities/Private/USDSkeletalDataConversion.cpp b/Engine/Plugins/Importers/USDImporter/Source/USDUtilities/Private/USDSkeletalDataConversion.cpp
--- a/Engine/Plugins/Importers/USDImporter/Source/USDUtilities/Private/USDSkeletalDataConversion.cpp
+++ b/Engine/Plugins/Importers/USDImporter/Source/USDUtilities/Private/USDSkeletalDataConversion.cpp
@@ -2348,7 +2348,7 @@ bool UnrealToUsd::ConvertSkeleton( const USkeleton* Skeleton, pxr::UsdSkelSkelet
 	return UnrealToUsd::ConvertSkeleton( Skeleton->GetReferenceSkeleton(), UsdSkeleton );
 }
 
-bool UnrealToUsd::ConvertSkeletalMesh( const USkeletalMesh* SkeletalMesh, pxr::UsdPrim& SkelRootPrim, const pxr::UsdTimeCode TimeCode, UE::FUsdStage* StageForMaterialAssignments )
+bool UnrealToUsd::ConvertSkeletalMesh( const USkeletalMesh* SkeletalMesh, pxr::UsdPrim& SkelRootPrim, const pxr::UsdTimeCode TimeCode, UE::FUsdStage* StageForMaterialAssignments, int32 LowestMeshLOD, int32 HighestMeshLOD )
 {
 	pxr::UsdSkelRoot SkelRoot{ SkelRootPrim };
 	if ( !SkeletalMesh || !SkeletalMesh->GetSkeleton() || !SkelRoot )
@@ -2372,6 +2372,18 @@ bool UnrealToUsd::ConvertSkeletalMesh( const USkeletalMesh* SkeletalMesh, pxr::U
 		return false;
 	}
 
+	// Make sure they're both >= 0 (the options dialog slider is clamped, but this may be called directly)
+	LowestMeshLOD = FMath::Clamp( LowestMeshLOD, 0, NumLODs - 1 );
+	HighestMeshLOD = FMath::Clamp( HighestMeshLOD, 0, NumLODs - 1 );
+
+	// Make sure Lowest <= Highest
+	int32 Temp = FMath::Min( LowestMeshLOD, HighestMeshLOD );
+	HighestMeshLOD = FMath::Max( LowestMeshLOD, HighestMeshLOD );
+	LowestMeshLOD = Temp;
+
+	// Make sure it's at least 1 LOD level
+	NumLODs = FMath::Max( HighestMeshLOD - LowestMeshLOD + 1, 1 );
+
 	pxr::UsdVariantSets VariantSets = SkelRootPrim.GetVariantSets();
 	if ( NumLODs > 1 && VariantSets.HasVariantSet( UnrealIdentifiers::LOD ) )
 	{
@@ -2423,7 +2435,7 @@ bool UnrealToUsd::ConvertSkeletalMesh( const USkeletalMesh* SkeletalMesh, pxr::U
 	}
 
 	// Actual meshes
-	for ( int32 LODIndex = 0; LODIndex < NumLODs; ++LODIndex )
+	for ( int32 LODIndex = LowestMeshLOD; LODIndex <= HighestMeshLOD; ++LODIndex )
 	{
 		const FSkeletalMeshLODModel& LODModel = SkelMeshResource->LODModels[ LODIndex ];
 
@@ -2490,7 +2502,7 @@ bool UnrealToUsd::ConvertSkeletalMesh( const USkeletalMesh* SkeletalMesh, pxr::U
 			}
 
 			int32 NumDeltas = 0;
-			FMorphTargetDelta* DeltaArray = MorphTarget->GetMorphTargetDelta(LODIndex, NumDeltas);
+			FMorphTargetDelta* DeltaArray = MorphTarget->GetMorphTargetDelta( LODIndex, NumDeltas );
 			if ( !DeltaArray || NumDeltas == 0 )
 			{
 				continue;
diff --git a/Engine/Plugins/Importers/USDImporter/Source/USDUtilities/Public/USDGeomMeshConversion.h b/Engine/Plugins/Importers/USDImporter/Source/USDUtilities/Public/USDGeomMeshConversion.h
--- a/Engine/Plugins/Importers/USDImporter/Source/USDUtilities/Public/USDGeomMeshConversion.h
+++ b/Engine/Plugins/Importers/USDImporter/Source/USDUtilities/Public/USDGeomMeshConversion.h
@@ -72,9 +72,11 @@ namespace UnrealToUsd
 	 * @param UsdPrim - Prim to receive the mesh data or LOD variant set
 	 * @param TimeCode - TimeCode to author the attributes with
 	 * @param StageForMaterialAssignments - Stage to use when authoring material assignments (we use this when we want to export the mesh to a payload layer, but the material assignments to an asset layer)
+	 * @param LowestMeshLOD - Lowest LOD of the UStaticMesh to export (start of the LOD range)
+	 * @param HighestMeshLOD - Lowest LOD of the UStaticMesh to export (end of the LOD range)
 	 * @return Whether the conversion was successful or not.
 	 */
-	USDUTILITIES_API bool ConvertStaticMesh( const UStaticMesh* StaticMesh, pxr::UsdPrim& UsdPrim, pxr::UsdTimeCode TimeCode = pxr::UsdTimeCode::Default(), UE::FUsdStage* StageForMaterialAssignments = nullptr );
+	USDUTILITIES_API bool ConvertStaticMesh( const UStaticMesh* StaticMesh, pxr::UsdPrim& UsdPrim, pxr::UsdTimeCode TimeCode = pxr::UsdTimeCode::Default(), UE::FUsdStage* StageForMaterialAssignments = nullptr, int32 LowestMeshLOD = 0, int32 HighestMeshLOD = INT32_MAX );
 
 	/**
 	 * Converts an array of mesh descriptions into mesh data, and places that data within the UsdGeomMesh UsdPrim.
diff --git a/Engine/Plugins/Importers/USDImporter/Source/USDUtilities/Public/USDSkeletalDataConversion.h b/Engine/Plugins/Importers/USDImporter/Source/USDUtilities/Public/USDSkeletalDataConversion.h
--- a/Engine/Plugins/Importers/USDImporter/Source/USDUtilities/Public/USDSkeletalDataConversion.h
+++ b/Engine/Plugins/Importers/USDImporter/Source/USDUtilities/Public/USDSkeletalDataConversion.h
@@ -195,9 +195,11 @@ namespace UnrealToUsd
 	 * @param SkelRoot - Root prim of the output source data. Child UsdSkelSkeleton, UsdGeomMesh, and UsdSkelBlendShape will be created as children of it, containing the converted data
 	 * @param TimeCode - TimeCode with which the converted data will be placed in the USD stage
 	 * @param StageForMaterialAssignments - Stage to use when authoring material assignments (we use this when we want to export the mesh to a payload layer, but the material assignments to an asset layer)
+	 * @param LowestMeshLOD - Lowest LOD of the UStaticMesh to export (start of the LOD range)
+	 * @param HighestMeshLOD - Lowest LOD of the UStaticMesh to export (end of the LOD range)
 	 * @return Whether the conversion was successful or not.
 	 */
-	USDUTILITIES_API bool ConvertSkeletalMesh( const USkeletalMesh* SkeletalMesh, pxr::UsdPrim& SkelRootPrim, const pxr::UsdTimeCode TimeCode = pxr::UsdTimeCode::Default(), UE::FUsdStage* StageForMaterialAssignments = nullptr );
+	USDUTILITIES_API bool ConvertSkeletalMesh( const USkeletalMesh* SkeletalMesh, pxr::UsdPrim& SkelRootPrim, const pxr::UsdTimeCode TimeCode = pxr::UsdTimeCode::Default(), UE::FUsdStage* StageForMaterialAssignments = nullptr, int32 LowestMeshLOD = 0, int32 HighestMeshLOD = INT32_MAX );
 
 	/**
 	 * Converts an AnimSequence to a UsdSkelAnimation. Includes bone transforms and blend shape weights.
