USD: Fix crash when changing LOD variants of skeletal meshes

diff --git a/Engine/Plugins/Importers/USDImporter/Source/USDUtilities/Private/USDSkeletalDataConversion.cpp b/Engine/Plugins/Importers/USDImporter/Source/USDUtilities/Private/USDSkeletalDataConversion.cpp
--- a/Engine/Plugins/Importers/USDImporter/Source/USDUtilities/Private/USDSkeletalDataConversion.cpp
+++ b/Engine/Plugins/Importers/USDImporter/Source/USDUtilities/Private/USDSkeletalDataConversion.cpp
@@ -2070,7 +2070,8 @@ USkeletalMesh* UsdToUnreal::GetSkeletalMeshFromImportData(
 	// A SkeletalMesh could be retrieved for re-use and updated for animations
 	// For now, create a new USkeletalMesh
 	// Note: Remember to initialize UsedMorphTargetNames with existing morph targets, whenever the SkeletalMesh is reused
-	USkeletalMesh* SkeletalMesh = NewObject<USkeletalMesh>(GetTransientPackage(), MeshName, ObjectFlags | EObjectFlags::RF_Public);
+	FName UniqueMeshName = MakeUniqueObjectName( GetTransientPackage(), USkeletalMesh::StaticClass(), MeshName );
+	USkeletalMesh* SkeletalMesh = NewObject<USkeletalMesh>(GetTransientPackage(), UniqueMeshName, ObjectFlags | EObjectFlags::RF_Public);
 
 	// Process reference skeleton from import data
 	int32 SkeletalDepth = 0;
@@ -2176,7 +2177,8 @@ USkeletalMesh* UsdToUnreal::GetSkeletalMeshFromImportData(
 	SkeletalMesh->CalculateInvRefMatrices();
 
 	// Generate a Skeleton and associate it to the SkeletalMesh
-	USkeleton* Skeleton = NewObject<USkeleton>( GetTransientPackage(), SkeletonName, ObjectFlags | EObjectFlags::RF_Public );
+	FName UniqueSkeletonName = MakeUniqueObjectName( GetTransientPackage(), USkeleton::StaticClass(), SkeletonName );
+	USkeleton* Skeleton = NewObject<USkeleton>( GetTransientPackage(), UniqueSkeletonName, ObjectFlags | EObjectFlags::RF_Public );
 	Skeleton->MergeAllBonesToBoneTree(SkeletalMesh);
 	Skeleton->SetPreviewMesh(SkeletalMesh);
 	SkeletalMesh->SetSkeleton(Skeleton);
