USD: Minor fixes related to the USD Stage editor.

Fix leaked references to the stage due to new 5.0 changes on list and tree views;
Sort prim attributes by label;
Remove old references to stage infos widget.

diff --git a/Engine/Plugins/Importers/USDImporter/Source/USDStage/Private/USDStageActor.cpp b/Engine/Plugins/Importers/USDImporter/Source/USDStage/Private/USDStageActor.cpp
--- a/Engine/Plugins/Importers/USDImporter/Source/USDStage/Private/USDStageActor.cpp
+++ b/Engine/Plugins/Importers/USDImporter/Source/USDStage/Private/USDStageActor.cpp
@@ -20,10 +20,11 @@
 #include "USDTransactor.h"
 #include "USDTypesConversion.h"
 
+#include "UsdWrappers/SdfChangeBlock.h"
+#include "UsdWrappers/SdfLayer.h"
 #include "UsdWrappers/UsdAttribute.h"
 #include "UsdWrappers/UsdGeomXformable.h"
-#include "UsdWrappers/SdfLayer.h"
-#include "UsdWrappers/SdfChangeBlock.h"
+#include "UsdWrappers/UsdStage.h"
 
 #include "Async/ParallelFor.h"
 #include "CineCameraActor.h"
diff --git a/Engine/Plugins/Importers/USDImporter/Source/USDStage/Public/USDLevelSequenceHelper.h b/Engine/Plugins/Importers/USDImporter/Source/USDStage/Public/USDLevelSequenceHelper.h
--- a/Engine/Plugins/Importers/USDImporter/Source/USDStage/Public/USDLevelSequenceHelper.h
+++ b/Engine/Plugins/Importers/USDImporter/Source/USDStage/Public/USDLevelSequenceHelper.h
@@ -4,16 +4,13 @@
 
 #include "CoreMinimal.h"
 
+#include "UsdWrappers/ForwardDeclarations.h"
+
 class AUsdStageActor;
 class FUsdLevelSequenceHelperImpl;
 class ULevelSequence;
 class UUsdPrimTwin;
 
-namespace UE
-{
-	class FUsdStage;
-}
-
 /**
  * Builds and maintains the level sequence and subsequences for a Usd Stage
  */
diff --git a/Engine/Plugins/Importers/USDImporter/Source/USDStageEditor/Private/SUSDPrimInfo.cpp b/Engine/Plugins/Importers/USDImporter/Source/USDStageEditor/Private/SUSDPrimInfo.cpp
--- a/Engine/Plugins/Importers/USDImporter/Source/USDStageEditor/Private/SUSDPrimInfo.cpp
+++ b/Engine/Plugins/Importers/USDImporter/Source/USDStageEditor/Private/SUSDPrimInfo.cpp
@@ -34,7 +34,7 @@ namespace UsdPrimInfoWidgetConstants
 	const FMargin CategoryHeaderPadding( 4.0f, 4.0f, 4.0f, 4.0f );
 }
 
-void SUsdPrimInfo::Construct( const FArguments& InArgs, const UE::FUsdStage& UsdStage, const TCHAR* PrimPath )
+void SUsdPrimInfo::Construct( const FArguments& InArgs, const UE::FUsdStageWeak& UsdStage, const TCHAR* PrimPath )
 {
 	TSharedRef< SWidget > VariantSetsWidget = GenerateVariantSetsWidget( UsdStage, PrimPath );
 	TSharedRef< SWidget > ReferencesListWidget = GenerateReferencesListWidget( UsdStage, PrimPath );
@@ -77,7 +77,7 @@ void SUsdPrimInfo::Construct( const FArguments& InArgs, const UE::FUsdStage& Usd
 	SetPrimPath( UsdStage, PrimPath );
 }
 
-void SUsdPrimInfo::SetPrimPath( const UE::FUsdStage& UsdStage, const TCHAR* PrimPath )
+void SUsdPrimInfo::SetPrimPath( const UE::FUsdStageWeak& UsdStage, const TCHAR* PrimPath )
 {
 	if ( PropertiesList )
 	{
@@ -100,7 +100,7 @@ void SUsdPrimInfo::SetPrimPath( const UE::FUsdStage& UsdStage, const TCHAR* Prim
 	}
 }
 
-TSharedRef< SWidget > SUsdPrimInfo::GenerateVariantSetsWidget( const UE::FUsdStage& UsdStage, const TCHAR* PrimPath )
+TSharedRef< SWidget > SUsdPrimInfo::GenerateVariantSetsWidget( const UE::FUsdStageWeak& UsdStage, const TCHAR* PrimPath )
 {
 	TSharedRef< SWidget > VariantSetsWidget = SNullWidget::NullWidget;
 
@@ -125,7 +125,7 @@ TSharedRef< SWidget > SUsdPrimInfo::GenerateVariantSetsWidget( const UE::FUsdSta
 	return VariantSetsWidget;
 }
 
-TSharedRef< SWidget > SUsdPrimInfo::GenerateReferencesListWidget( const UE::FUsdStage& UsdStage, const TCHAR* PrimPath )
+TSharedRef< SWidget > SUsdPrimInfo::GenerateReferencesListWidget( const UE::FUsdStageWeak& UsdStage, const TCHAR* PrimPath )
 {
 	TSharedRef< SWidget > ReferencesListWidget = SNullWidget::NullWidget;
 
diff --git a/Engine/Plugins/Importers/USDImporter/Source/USDStageEditor/Private/SUSDPrimInfo.h b/Engine/Plugins/Importers/USDImporter/Source/USDStageEditor/Private/SUSDPrimInfo.h
--- a/Engine/Plugins/Importers/USDImporter/Source/USDStageEditor/Private/SUSDPrimInfo.h
+++ b/Engine/Plugins/Importers/USDImporter/Source/USDStageEditor/Private/SUSDPrimInfo.h
@@ -5,10 +5,7 @@
 #include "Widgets/DeclarativeSyntaxSupport.h"
 #include "Widgets/SCompoundWidget.h"
 
-namespace UE
-{
-	class FUsdStage;
-}
+#include "UsdWrappers/ForwardDeclarations.h"
 
 #if USE_USD_SDK
 
@@ -20,13 +17,13 @@ public:
 	SLATE_BEGIN_ARGS( SUsdPrimInfo ) {}
 	SLATE_END_ARGS()
 
-	void Construct( const FArguments& InArgs, const UE::FUsdStage& UsdStage, const TCHAR* PrimPath );
+	void Construct( const FArguments& InArgs, const UE::FUsdStageWeak& UsdStage, const TCHAR* PrimPath );
 
-	void SetPrimPath( const UE::FUsdStage& UsdStage, const TCHAR* PrimPath );
+	void SetPrimPath( const UE::FUsdStageWeak& UsdStage, const TCHAR* PrimPath );
 
 protected:
-	TSharedRef< SWidget > GenerateVariantSetsWidget( const UE::FUsdStage& UsdStage, const TCHAR* PrimPath );
-	TSharedRef< SWidget > GenerateReferencesListWidget( const UE::FUsdStage& UsdStage, const TCHAR* PrimPath );
+	TSharedRef< SWidget > GenerateVariantSetsWidget( const UE::FUsdStageWeak& UsdStage, const TCHAR* PrimPath );
+	TSharedRef< SWidget > GenerateReferencesListWidget( const UE::FUsdStageWeak& UsdStage, const TCHAR* PrimPath );
 
 private:
 	TSharedPtr< class SUsdPrimPropertiesList > PropertiesList;
diff --git a/Engine/Plugins/Importers/USDImporter/Source/USDStageEditor/Private/SUSDStage.cpp b/Engine/Plugins/Importers/USDImporter/Source/USDStageEditor/Private/SUSDStage.cpp
--- a/Engine/Plugins/Importers/USDImporter/Source/USDStageEditor/Private/SUSDStage.cpp
+++ b/Engine/Plugins/Importers/USDImporter/Source/USDStageEditor/Private/SUSDStage.cpp
@@ -4,7 +4,6 @@
 
 #include "SUSDLayersTreeView.h"
 #include "SUSDPrimInfo.h"
-#include "SUSDStageInfo.h"
 #include "SUSDStageTreeView.h"
 #include "UnrealUSDWrapper.h"
 #include "USDErrorUtils.h"
@@ -96,7 +95,6 @@ namespace SUSDStageImpl
 
 void SUsdStage::Construct( const FArguments& InArgs )
 {
-	OnStageActorPropertyChangedHandle = FCoreUObjectDelegates::OnObjectPropertyChanged.AddSP( SharedThis( this ), &SUsdStage::OnStageActorPropertyChanged );
 	OnActorLoadedHandle = AUsdStageActor::OnActorLoaded.AddSP( SharedThis( this ), &SUsdStage::OnStageActorLoaded );
 
 	OnViewportSelectionChangedHandle = USelection::SelectionChangedEvent.AddRaw( this, &SUsdStage::OnViewportSelectionChanged );
@@ -204,11 +202,6 @@ void SUsdStage::SetupStageActorDelegates()
 				{
 					this->UsdPrimInfoWidget->SetPrimPath( ViewModel.UsdStageActor->GetOrLoadUsdStage(), *PrimPath );
 				}
-
-				if ( PrimPath == TEXT("/") && this->UsdStageInfoWidget )
-				{
-					this->UsdStageInfoWidget->RefreshStageInfos( ViewModel.UsdStageActor.Get() );
-				}
 			}
 		);
 
@@ -789,11 +782,6 @@ void SUsdStage::Refresh()
 		UsdLayersTreeView->Refresh( StageActor, true );
 	}
 
-	if (UsdStageInfoWidget)
-	{
-		UsdStageInfoWidget->RefreshStageInfos( StageActor );
-	}
-
 	if (UsdStageTreeView)
 	{
 		UsdStageTreeView->Refresh( StageActor );
@@ -817,17 +805,6 @@ void SUsdStage::OnStageActorLoaded( AUsdStageActor* InUsdStageActor )
 	Refresh();
 }
 
-void SUsdStage::OnStageActorPropertyChanged( UObject* ObjectBeingModified, FPropertyChangedEvent& PropertyChangedEvent )
-{
-	if ( ObjectBeingModified == ViewModel.UsdStageActor )
-	{
-		if ( UsdStageInfoWidget )
-		{
-			UsdStageInfoWidget->RefreshStageInfos( ViewModel.UsdStageActor.Get() );
-		}
-	}
-}
-
 void SUsdStage::OnViewportSelectionChanged( UObject* NewSelection )
 {
 	// This may be called when first opening a project, before the our widgets are fully initialized
diff --git a/Engine/Plugins/Importers/USDImporter/Source/USDStageEditor/Private/SUSDStage.h b/Engine/Plugins/Importers/USDImporter/Source/USDStageEditor/Private/SUSDStage.h
--- a/Engine/Plugins/Importers/USDImporter/Source/USDStageEditor/Private/SUSDStage.h
+++ b/Engine/Plugins/Importers/USDImporter/Source/USDStageEditor/Private/SUSDStage.h
@@ -59,14 +59,12 @@ protected:
 	void Refresh();
 
 	void OnStageActorLoaded( AUsdStageActor* InUsdStageActor );
-	void OnStageActorPropertyChanged( UObject* ObjectBeingModified, FPropertyChangedEvent& PropertyChangedEvent );
 
 	void OnMapChanged( UWorld* World, EMapChangeType ChangeType );
 
 	void OnViewportSelectionChanged( UObject* NewSelection );
 
 protected:
-	TSharedPtr< class SUsdStageInfo > UsdStageInfoWidget;
 	TSharedPtr< class SUsdStageTreeView > UsdStageTreeView;
 	TSharedPtr< class SUsdPrimInfo > UsdPrimInfoWidget;
 	TSharedPtr< class SUsdLayersTreeView > UsdLayersTreeView;
@@ -75,7 +73,6 @@ protected:
 	FDelegateHandle OnActorDestroyedHandle;
 	FDelegateHandle OnStageActorPropertyChangedHandle;
 	FDelegateHandle OnStageChangedHandle;
-	FDelegateHandle OnStageInfoChangedHandle;
 	FDelegateHandle OnStageEditTargetChangedHandle;
 	FDelegateHandle OnPrimChangedHandle;
 
diff --git a/Engine/Plugins/Importers/USDImporter/Source/USDStageEditor/Private/SUSDStageInfo.cpp b/Engine/Plugins/Importers/USDImporter/Source/USDStageEditor/Private/SUSDStageInfo.cpp
deleted file mode 100644
diff --git a/Engine/Plugins/Importers/USDImporter/Source/USDStageEditor/Private/SUSDStageInfo.h b/Engine/Plugins/Importers/USDImporter/Source/USDStageEditor/Private/SUSDStageInfo.h
deleted file mode 100644
diff --git a/Engine/Plugins/Importers/USDImporter/Source/USDStageEditor/Private/SUSDStageTreeView.cpp b/Engine/Plugins/Importers/USDImporter/Source/USDStageEditor/Private/SUSDStageTreeView.cpp
--- a/Engine/Plugins/Importers/USDImporter/Source/USDStageEditor/Private/SUSDStageTreeView.cpp
+++ b/Engine/Plugins/Importers/USDImporter/Source/USDStageEditor/Private/SUSDStageTreeView.cpp
@@ -358,24 +358,15 @@ void SUsdStageTreeView::OnGetChildren( FUsdPrimViewModelRef InParent, TArray< FU
 
 void SUsdStageTreeView::Refresh( AUsdStageActor* InUsdStageActor )
 {
-	UE::FUsdStage OldStage = RootItems.Num() > 0 ? RootItems[0]->UsdStage : UE::FUsdStage();
-	UE::FUsdStage NewStage = InUsdStageActor ? static_cast< const AUsdStageActor* >( InUsdStageActor )->GetUsdStage() : UE::FUsdStage();
+	UE::FUsdStageWeak OldStage = RootItems.Num() > 0 ? RootItems[0]->UsdStage : UE::FUsdStageWeak();
+	UE::FUsdStageWeak NewStage = InUsdStageActor
+		? static_cast< UE::FUsdStageWeak > ( static_cast< const AUsdStageActor* >( InUsdStageActor )->GetUsdStage() )
+		: UE::FUsdStageWeak();
 
 	RootItems.Empty();
 	if ( UsdStageActor.Get() != InUsdStageActor || NewStage != OldStage )
 	{
-		// This is very important: Internally the tree will store FUsdPrimViewModelRef in its SparseItemInfos member if we have
-		// any member manually expanded/collapsed. These can prevent the FUsdPrimViewModels from being collected, and prevent the
-		// stage from being fully closed, so we must do this whenever the stage changes
-		ClearExpandedItems();
 		TreeItemExpansionStates.Reset();
-
-		// Clear other things that may hold FUsdPrimViewModelRefs
-		LinearizedItems.Empty();
-		SelectorItem = SUsdStageTreeView::NullableItemType(nullptr);
-		RangeSelectionStart = SUsdStageTreeView::NullableItemType(nullptr);
-		ItemToScrollIntoView = SUsdStageTreeView::NullableItemType(nullptr);
-		ItemToNotifyWhenInView = SUsdStageTreeView::NullableItemType(nullptr);
 	}
 
 	UsdStageActor = InUsdStageActor;
@@ -1027,7 +1018,7 @@ void SUsdStageTreeView::OnPrimNameUpdated(const FUsdPrimViewModelRef& TreeItem,
 	}
 
 	{
-		const UE::FUsdStage& Stage = TreeItem->UsdStage;
+		const UE::FUsdStageWeak& Stage = TreeItem->UsdStage;
 		if ( !Stage )
 		{
 			return;
diff --git a/Engine/Plugins/Importers/USDImporter/Source/USDStageEditor/Private/Widgets/SUSDPrimPropertiesList.cpp b/Engine/Plugins/Importers/USDImporter/Source/USDStageEditor/Private/Widgets/SUSDPrimPropertiesList.cpp
--- a/Engine/Plugins/Importers/USDImporter/Source/USDStageEditor/Private/Widgets/SUSDPrimPropertiesList.cpp
+++ b/Engine/Plugins/Importers/USDImporter/Source/USDStageEditor/Private/Widgets/SUSDPrimPropertiesList.cpp
@@ -650,7 +650,7 @@ void SUsdPrimPropertyRow::OnCheckBoxCheckStateChanged( ECheckBoxState NewState )
 	UsdPrimAttribute->SetAttributeValue( NewValue );
 }
 
-void SUsdPrimPropertiesList::Construct( const FArguments& InArgs, const UE::FUsdStage& UsdStage, const TCHAR* InPrimPath )
+void SUsdPrimPropertiesList::Construct( const FArguments& InArgs, const UE::FUsdStageWeak& UsdStage, const TCHAR* InPrimPath )
 {
 	GeneratePropertiesList( UsdStage, InPrimPath );
 
@@ -681,14 +681,14 @@ TSharedRef< ITableRow > SUsdPrimPropertiesList::OnGenerateRow( TSharedPtr< FUsdP
 	return SNew( SUsdPrimPropertyRow, InDisplayNode, OwnerTable );
 }
 
-void SUsdPrimPropertiesList::GeneratePropertiesList( const UE::FUsdStage& UsdStage, const TCHAR* InPrimPath )
+void SUsdPrimPropertiesList::GeneratePropertiesList( const UE::FUsdStageWeak& UsdStage, const TCHAR* InPrimPath )
 {
 	float TimeCode = 0.f;
 	ViewModel.UsdStage = UsdStage;
 	ViewModel.Refresh( InPrimPath, TimeCode );
 }
 
-void SUsdPrimPropertiesList::SetPrimPath( const UE::FUsdStage& UsdStage, const TCHAR* InPrimPath )
+void SUsdPrimPropertiesList::SetPrimPath( const UE::FUsdStageWeak& UsdStage, const TCHAR* InPrimPath )
 {
 	GeneratePropertiesList( UsdStage, InPrimPath );
 	RequestListRefresh();
diff --git a/Engine/Plugins/Importers/USDImporter/Source/USDStageEditor/Private/Widgets/SUSDPrimPropertiesList.h b/Engine/Plugins/Importers/USDImporter/Source/USDStageEditor/Private/Widgets/SUSDPrimPropertiesList.h
--- a/Engine/Plugins/Importers/USDImporter/Source/USDStageEditor/Private/Widgets/SUSDPrimPropertiesList.h
+++ b/Engine/Plugins/Importers/USDImporter/Source/USDStageEditor/Private/Widgets/SUSDPrimPropertiesList.h
@@ -13,12 +13,12 @@ class SUsdPrimPropertiesList : public SListView< TSharedPtr< FUsdPrimAttributeVi
 	SLATE_END_ARGS()
 
 public:
-	void Construct( const FArguments& InArgs, const UE::FUsdStage& UsdStage, const TCHAR* InPrimPath );
-	void SetPrimPath( const UE::FUsdStage& UsdStage, const TCHAR* InPrimPath );
+	void Construct( const FArguments& InArgs, const UE::FUsdStageWeak& UsdStage, const TCHAR* InPrimPath );
+	void SetPrimPath( const UE::FUsdStageWeak& UsdStage, const TCHAR* InPrimPath );
 
 protected:
 	TSharedRef< ITableRow > OnGenerateRow( TSharedPtr< FUsdPrimAttributeViewModel > InDisplayNode, const TSharedRef< STableViewBase >& OwnerTable );
-	void GeneratePropertiesList( const UE::FUsdStage& UsdStage, const TCHAR* InPrimPath );
+	void GeneratePropertiesList( const UE::FUsdStageWeak& UsdStage, const TCHAR* InPrimPath );
 
 private:
 	FUsdPrimAttributesViewModel ViewModel;
diff --git a/Engine/Plugins/Importers/USDImporter/Source/USDStageEditor/Private/Widgets/SUSDReferencesList.cpp b/Engine/Plugins/Importers/USDImporter/Source/USDStageEditor/Private/Widgets/SUSDReferencesList.cpp
--- a/Engine/Plugins/Importers/USDImporter/Source/USDStageEditor/Private/Widgets/SUSDReferencesList.cpp
+++ b/Engine/Plugins/Importers/USDImporter/Source/USDStageEditor/Private/Widgets/SUSDReferencesList.cpp
@@ -55,7 +55,7 @@ TSharedRef< SWidget > SUsdReferenceRow::GenerateWidgetForColumn( const FName& Co
 		];
 }
 
-void SUsdReferencesList::Construct( const FArguments& InArgs, const UE::FUsdStage& UsdStage, const TCHAR* PrimPath )
+void SUsdReferencesList::Construct( const FArguments& InArgs, const UE::FUsdStageWeak& UsdStage, const TCHAR* PrimPath )
 {
 	ViewModel.UpdateReferences( UsdStage, PrimPath );
 
@@ -79,7 +79,7 @@ TSharedRef< ITableRow > SUsdReferencesList::OnGenerateRow( TSharedPtr< FUsdRefer
 	return SNew( SUsdReferenceRow, InDisplayNode, OwnerTable );
 }
 
-void SUsdReferencesList::SetPrimPath( const UE::FUsdStage& UsdStage, const TCHAR* PrimPath )
+void SUsdReferencesList::SetPrimPath( const UE::FUsdStageWeak& UsdStage, const TCHAR* PrimPath )
 {
 	ViewModel.UpdateReferences( UsdStage, PrimPath );
 	RequestListRefresh();
diff --git a/Engine/Plugins/Importers/USDImporter/Source/USDStageEditor/Private/Widgets/SUSDReferencesList.h b/Engine/Plugins/Importers/USDImporter/Source/USDStageEditor/Private/Widgets/SUSDReferencesList.h
--- a/Engine/Plugins/Importers/USDImporter/Source/USDStageEditor/Private/Widgets/SUSDReferencesList.h
+++ b/Engine/Plugins/Importers/USDImporter/Source/USDStageEditor/Private/Widgets/SUSDReferencesList.h
@@ -7,11 +7,7 @@
 #if USE_USD_SDK
 
 #include "USDReferencesViewModel.h"
-
-namespace UE
-{
-	class FUsdStage;
-}
+#include "UsdWrappers/ForwardDeclarations.h"
 
 class SUsdReferenceRow : public SMultiColumnTableRow< TSharedPtr< FUsdReference > >
 {
@@ -35,8 +31,8 @@ class SUsdReferencesList : public SListView< TSharedPtr< FUsdReference > >
 	SLATE_END_ARGS()
 
 public:
-	void Construct( const FArguments& InArgs, const UE::FUsdStage& UsdStage, const TCHAR* PrimPath );
-	void SetPrimPath( const UE::FUsdStage& UsdStage, const TCHAR* PrimPath );
+	void Construct( const FArguments& InArgs, const UE::FUsdStageWeak& UsdStage, const TCHAR* PrimPath );
+	void SetPrimPath( const UE::FUsdStageWeak& UsdStage, const TCHAR* PrimPath );
 
 protected:
 	TSharedRef< ITableRow > OnGenerateRow( TSharedPtr< FUsdReference > InDisplayNode, const TSharedRef< STableViewBase >& OwnerTable );
diff --git a/Engine/Plugins/Importers/USDImporter/Source/USDStageEditor/Private/Widgets/SUSDVariantSetsList.cpp b/Engine/Plugins/Importers/USDImporter/Source/USDStageEditor/Private/Widgets/SUSDVariantSetsList.cpp
--- a/Engine/Plugins/Importers/USDImporter/Source/USDStageEditor/Private/Widgets/SUSDVariantSetsList.cpp
+++ b/Engine/Plugins/Importers/USDImporter/Source/USDStageEditor/Private/Widgets/SUSDVariantSetsList.cpp
@@ -84,7 +84,7 @@ void SUsdVariantRow::OnSelectionChanged( TSharedPtr< FString > NewValue, ESelect
 	OnVariantSelectionChanged.ExecuteIfBound( VariantSet.ToSharedRef(), NewValue );
 }
 
-void SVariantsList::Construct( const FArguments& InArgs, const UE::FUsdStage& UsdStage, const TCHAR* InPrimPath )
+void SVariantsList::Construct( const FArguments& InArgs, const UE::FUsdStageWeak& UsdStage, const TCHAR* InPrimPath )
 {
 	PrimPath = InPrimPath;
 
@@ -115,7 +115,7 @@ TSharedRef< ITableRow > SVariantsList::OnGenerateRow( TSharedPtr< FUsdVariantSet
 			.OnVariantSelectionChanged( this, &SVariantsList::OnVariantSelectionChanged );
 }
 
-void SVariantsList::SetPrimPath( const UE::FUsdStage& UsdStage, const TCHAR* InPrimPath )
+void SVariantsList::SetPrimPath( const UE::FUsdStageWeak& UsdStage, const TCHAR* InPrimPath )
 {
 	PrimPath = InPrimPath;
 	ViewModel.UpdateVariantSets( UsdStage, *PrimPath );
diff --git a/Engine/Plugins/Importers/USDImporter/Source/USDStageEditor/Private/Widgets/SUSDVariantSetsList.h b/Engine/Plugins/Importers/USDImporter/Source/USDStageEditor/Private/Widgets/SUSDVariantSetsList.h
--- a/Engine/Plugins/Importers/USDImporter/Source/USDStageEditor/Private/Widgets/SUSDVariantSetsList.h
+++ b/Engine/Plugins/Importers/USDImporter/Source/USDStageEditor/Private/Widgets/SUSDVariantSetsList.h
@@ -45,8 +45,8 @@ class SVariantsList : public SListView< TSharedPtr< FUsdVariantSetViewModel > >
 	SLATE_END_ARGS()
 
 public:
-	void Construct( const FArguments& InArgs, const UE::FUsdStage& UsdStage, const TCHAR* InPrimPath );
-	void SetPrimPath( const UE::FUsdStage& UsdStage, const TCHAR* InPrimPath );
+	void Construct( const FArguments& InArgs, const UE::FUsdStageWeak& UsdStage, const TCHAR* InPrimPath );
+	void SetPrimPath( const UE::FUsdStageWeak& UsdStage, const TCHAR* InPrimPath );
 
 protected:
 	void UpdateVariantSets( const TCHAR* InPrimPath );
diff --git a/Engine/Plugins/Importers/USDImporter/Source/USDStageEditorViewModels/Private/USDLayersViewModel.cpp b/Engine/Plugins/Importers/USDImporter/Source/USDStageEditorViewModels/Private/USDLayersViewModel.cpp
--- a/Engine/Plugins/Importers/USDImporter/Source/USDStageEditorViewModels/Private/USDLayersViewModel.cpp
+++ b/Engine/Plugins/Importers/USDImporter/Source/USDStageEditorViewModels/Private/USDLayersViewModel.cpp
@@ -23,7 +23,7 @@
 
 #endif // #if USE_USD_SDK
 
-FUsdLayerViewModel::FUsdLayerViewModel( FUsdLayerViewModel* InParentItem, const UE::FUsdStage& InUsdStage, const FString& InLayerIdentifier )
+FUsdLayerViewModel::FUsdLayerViewModel( FUsdLayerViewModel* InParentItem, const UE::FUsdStageWeak& InUsdStage, const FString& InLayerIdentifier )
 	: LayerModel( MakeShared< FUsdLayerModel >() )
 	, ParentItem( InParentItem )
 	, UsdStage( InUsdStage )
diff --git a/Engine/Plugins/Importers/USDImporter/Source/USDStageEditorViewModels/Private/USDPrimAttributesViewModel.cpp b/Engine/Plugins/Importers/USDImporter/Source/USDStageEditorViewModels/Private/USDPrimAttributesViewModel.cpp
--- a/Engine/Plugins/Importers/USDImporter/Source/USDStageEditorViewModels/Private/USDPrimAttributesViewModel.cpp
+++ b/Engine/Plugins/Importers/USDImporter/Source/USDStageEditorViewModels/Private/USDPrimAttributesViewModel.cpp
@@ -303,6 +303,11 @@ void FUsdPrimAttributesViewModel::Refresh( const TCHAR* InPrimPath, float TimeCo
 			}
 		}
 	}
+
+	PrimAttributes.Sort([](const TSharedPtr<FUsdPrimAttributeViewModel>& A, const TSharedPtr<FUsdPrimAttributeViewModel>& B)
+	{
+		return A->Label < B->Label;
+	});
 #endif // #if USE_USD_SDK
 }
 
diff --git a/Engine/Plugins/Importers/USDImporter/Source/USDStageEditorViewModels/Private/USDPrimViewModel.cpp b/Engine/Plugins/Importers/USDImporter/Source/USDStageEditorViewModels/Private/USDPrimViewModel.cpp
--- a/Engine/Plugins/Importers/USDImporter/Source/USDStageEditorViewModels/Private/USDPrimViewModel.cpp
+++ b/Engine/Plugins/Importers/USDImporter/Source/USDStageEditorViewModels/Private/USDPrimViewModel.cpp
@@ -29,7 +29,7 @@
 
 #endif // #if USE_USD_SDK
 
-FUsdPrimViewModel::FUsdPrimViewModel( FUsdPrimViewModel* InParentItem, const UE::FUsdStage& InUsdStage, const UE::FUsdPrim& InUsdPrim )
+FUsdPrimViewModel::FUsdPrimViewModel( FUsdPrimViewModel* InParentItem, const UE::FUsdStageWeak& InUsdStage, const UE::FUsdPrim& InUsdPrim )
 	: FUsdPrimViewModel( InParentItem, InUsdStage )
 {
 	UsdPrim = InUsdPrim;
@@ -38,7 +38,7 @@ FUsdPrimViewModel::FUsdPrimViewModel( FUsdPrimViewModel* InParentItem, const UE:
 	FillChildren();
 }
 
-FUsdPrimViewModel::FUsdPrimViewModel( FUsdPrimViewModel* InParentItem, const UE::FUsdStage& InUsdStage )
+FUsdPrimViewModel::FUsdPrimViewModel( FUsdPrimViewModel* InParentItem, const UE::FUsdStageWeak& InUsdStage )
 	: UsdStage( InUsdStage )
 	, ParentItem( InParentItem )
 	, RowData( MakeShared< FUsdPrimModel >() )
diff --git a/Engine/Plugins/Importers/USDImporter/Source/USDStageEditorViewModels/Private/USDReferencesViewModel.cpp b/Engine/Plugins/Importers/USDImporter/Source/USDStageEditorViewModels/Private/USDReferencesViewModel.cpp
--- a/Engine/Plugins/Importers/USDImporter/Source/USDStageEditorViewModels/Private/USDReferencesViewModel.cpp
+++ b/Engine/Plugins/Importers/USDImporter/Source/USDStageEditorViewModels/Private/USDReferencesViewModel.cpp
@@ -18,7 +18,7 @@
 
 #endif // #if USE_USD_SDK
 
-void FUsdReferencesViewModel::UpdateReferences( const UE::FUsdStage& UsdStage, const TCHAR* PrimPath )
+void FUsdReferencesViewModel::UpdateReferences( const UE::FUsdStageWeak& UsdStage, const TCHAR* PrimPath )
 {
 	References.Reset();
 
diff --git a/Engine/Plugins/Importers/USDImporter/Source/USDStageEditorViewModels/Private/USDVariantSetsViewModel.cpp b/Engine/Plugins/Importers/USDImporter/Source/USDStageEditorViewModels/Private/USDVariantSetsViewModel.cpp
--- a/Engine/Plugins/Importers/USDImporter/Source/USDStageEditorViewModels/Private/USDVariantSetsViewModel.cpp
+++ b/Engine/Plugins/Importers/USDImporter/Source/USDStageEditorViewModels/Private/USDVariantSetsViewModel.cpp
@@ -53,7 +53,7 @@ void FUsdVariantSetViewModel::SetVariantSelection( const TSharedPtr< FString >&
 #endif // #if USE_USD_SDK
 }
 
-void FUsdVariantSetsViewModel::UpdateVariantSets( const UE::FUsdStage& InUsdStage, const TCHAR* PrimPath )
+void FUsdVariantSetsViewModel::UpdateVariantSets( const UE::FUsdStageWeak& InUsdStage, const TCHAR* PrimPath )
 {
 #if USE_USD_SDK
 	VariantSets.Reset();
diff --git a/Engine/Plugins/Importers/USDImporter/Source/USDStageEditorViewModels/Public/USDLayersViewModel.h b/Engine/Plugins/Importers/USDImporter/Source/USDStageEditorViewModels/Public/USDLayersViewModel.h
--- a/Engine/Plugins/Importers/USDImporter/Source/USDStageEditorViewModels/Public/USDLayersViewModel.h
+++ b/Engine/Plugins/Importers/USDImporter/Source/USDStageEditorViewModels/Public/USDLayersViewModel.h
@@ -8,11 +8,6 @@
 #include "USDTreeItemViewModel.h"
 #include "UsdWrappers/UsdStage.h"
 
-namespace UE
-{
-	class FSdfLayer;
-}
-
 class USDSTAGEEDITORVIEWMODELS_API FUsdLayerModel : public TSharedFromThis< FUsdLayerModel >
 {
 public:
@@ -26,7 +21,7 @@ public:
 class USDSTAGEEDITORVIEWMODELS_API FUsdLayerViewModel : public IUsdTreeViewItem
 {
 public:
-	explicit FUsdLayerViewModel( FUsdLayerViewModel* InParentItem, const UE::FUsdStage& InUsdStage, const FString& InLayerIdentifier );
+	explicit FUsdLayerViewModel( FUsdLayerViewModel* InParentItem, const UE::FUsdStageWeak& InUsdStage, const FString& InLayerIdentifier );
 
 	bool IsValid() const;
 
@@ -54,6 +49,6 @@ public:
 	FUsdLayerViewModel* ParentItem;
 	TArray< TSharedRef< FUsdLayerViewModel > > Children;
 
-	UE::FUsdStage UsdStage;
+	UE::FUsdStageWeak UsdStage;
 	FString LayerIdentifier;
 };
diff --git a/Engine/Plugins/Importers/USDImporter/Source/USDStageEditorViewModels/Public/USDPrimAttributesViewModel.h b/Engine/Plugins/Importers/USDImporter/Source/USDStageEditorViewModels/Public/USDPrimAttributesViewModel.h
--- a/Engine/Plugins/Importers/USDImporter/Source/USDStageEditorViewModels/Public/USDPrimAttributesViewModel.h
+++ b/Engine/Plugins/Importers/USDImporter/Source/USDStageEditorViewModels/Public/USDPrimAttributesViewModel.h
@@ -43,7 +43,7 @@ public:
 	void Refresh( const TCHAR* InPrimPath, float TimeCode );
 
 public:
-	UE::FUsdStage UsdStage;
+	UE::FUsdStageWeak UsdStage;
 	TArray< TSharedPtr< FUsdPrimAttributeViewModel > > PrimAttributes;
 
 private:
diff --git a/Engine/Plugins/Importers/USDImporter/Source/USDStageEditorViewModels/Public/USDPrimViewModel.h b/Engine/Plugins/Importers/USDImporter/Source/USDStageEditorViewModels/Public/USDPrimViewModel.h
--- a/Engine/Plugins/Importers/USDImporter/Source/USDStageEditorViewModels/Public/USDPrimViewModel.h
+++ b/Engine/Plugins/Importers/USDImporter/Source/USDStageEditorViewModels/Public/USDPrimViewModel.h
@@ -33,9 +33,9 @@ public:
 class USDSTAGEEDITORVIEWMODELS_API FUsdPrimViewModel : public IUsdTreeViewItem, public TSharedFromThis< FUsdPrimViewModel >
 {
 public:
-	FUsdPrimViewModel( FUsdPrimViewModel* InParentItem, const UE::FUsdStage& InUsdStage, const UE::FUsdPrim& InUsdPrim );
+	FUsdPrimViewModel( FUsdPrimViewModel* InParentItem, const UE::FUsdStageWeak& InUsdStage, const UE::FUsdPrim& InUsdPrim );
 
-	FUsdPrimViewModel( FUsdPrimViewModel* InParentItem, const UE::FUsdStage& InUsdStage );
+	FUsdPrimViewModel( FUsdPrimViewModel* InParentItem, const UE::FUsdStageWeak& InUsdStage );
 
 	TArray< FUsdPrimViewModelRef >& UpdateChildren();
 
@@ -61,7 +61,7 @@ public:
 	void ClearReferences();
 
 public:
-	UE::FUsdStage UsdStage;
+	UE::FUsdStageWeak UsdStage;
 	UE::FUsdPrim UsdPrim;
 	FUsdPrimViewModel* ParentItem;
 
diff --git a/Engine/Plugins/Importers/USDImporter/Source/USDStageEditorViewModels/Public/USDReferencesViewModel.h b/Engine/Plugins/Importers/USDImporter/Source/USDStageEditorViewModels/Public/USDReferencesViewModel.h
--- a/Engine/Plugins/Importers/USDImporter/Source/USDStageEditorViewModels/Public/USDReferencesViewModel.h
+++ b/Engine/Plugins/Importers/USDImporter/Source/USDStageEditorViewModels/Public/USDReferencesViewModel.h
@@ -4,10 +4,7 @@
 
 #include "Templates/SharedPointer.h"
 
-namespace UE
-{
-	class FUsdStage;
-}
+#include "UsdWrappers/ForwardDeclarations.h"
 
 class USDSTAGEEDITORVIEWMODELS_API FUsdReference : public TSharedFromThis< FUsdReference >
 {
@@ -18,7 +15,7 @@ public:
 class USDSTAGEEDITORVIEWMODELS_API FUsdReferencesViewModel
 {
 public:
-	void UpdateReferences( const UE::FUsdStage& UsdStage, const TCHAR* PrimPath );
+	void UpdateReferences( const UE::FUsdStageWeak& UsdStage, const TCHAR* PrimPath );
 
 public:
 	TArray< TSharedPtr< FUsdReference > > References;
diff --git a/Engine/Plugins/Importers/USDImporter/Source/USDStageEditorViewModels/Public/USDVariantSetsViewModel.h b/Engine/Plugins/Importers/USDImporter/Source/USDStageEditorViewModels/Public/USDVariantSetsViewModel.h
--- a/Engine/Plugins/Importers/USDImporter/Source/USDStageEditorViewModels/Public/USDVariantSetsViewModel.h
+++ b/Engine/Plugins/Importers/USDImporter/Source/USDStageEditorViewModels/Public/USDVariantSetsViewModel.h
@@ -28,10 +28,10 @@ private:
 class USDSTAGEEDITORVIEWMODELS_API FUsdVariantSetsViewModel
 {
 public:
-	void UpdateVariantSets( const UE::FUsdStage& InUsdStage, const TCHAR* PrimPath );
+	void UpdateVariantSets( const UE::FUsdStageWeak& InUsdStage, const TCHAR* PrimPath );
 
 public:
-	UE::FUsdStage UsdStage;
+	UE::FUsdStageWeak UsdStage;
 	UE::FUsdPrim UsdPrim;
 
 	TArray< TSharedPtr< FUsdVariantSetViewModel > > VariantSets;
diff --git a/Engine/Plugins/Importers/USDImporter/Source/USDUtilities/Private/USDGeomMeshConversion.cpp b/Engine/Plugins/Importers/USDImporter/Source/USDUtilities/Private/USDGeomMeshConversion.cpp
--- a/Engine/Plugins/Importers/USDImporter/Source/USDUtilities/Private/USDGeomMeshConversion.cpp
+++ b/Engine/Plugins/Importers/USDImporter/Source/USDUtilities/Private/USDGeomMeshConversion.cpp
@@ -1430,7 +1430,7 @@ bool UnrealToUsd::ConvertStaticMesh( const UStaticMesh* StaticMesh, pxr::UsdPrim
 		pxr::UsdStageRefPtr MaterialStage;
 		if ( StageForMaterialAssignments )
 		{
-			MaterialStage = *StageForMaterialAssignments;
+			MaterialStage = static_cast< pxr::UsdStageRefPtr >( *StageForMaterialAssignments );
 		}
 		else
 		{
diff --git a/Engine/Plugins/Importers/USDImporter/Source/USDUtilities/Public/USDAttributeUtils.h b/Engine/Plugins/Importers/USDImporter/Source/USDUtilities/Public/USDAttributeUtils.h
--- a/Engine/Plugins/Importers/USDImporter/Source/USDUtilities/Public/USDAttributeUtils.h
+++ b/Engine/Plugins/Importers/USDImporter/Source/USDUtilities/Public/USDAttributeUtils.h
@@ -4,13 +4,13 @@
 
 #include "CoreMinimal.h"
 
+#include "UsdWrappers/ForwardDeclarations.h"
+
 #if USE_USD_SDK
 
 namespace UE
 {
-	class FSdfLayer;
 	class FUsdAttribute;
-	class FUsdStage;
 }
 
 namespace UsdUtils
diff --git a/Engine/Plugins/Importers/USDImporter/Source/USDUtilities/Public/USDConversionUtils.h b/Engine/Plugins/Importers/USDImporter/Source/USDUtilities/Public/USDConversionUtils.h
--- a/Engine/Plugins/Importers/USDImporter/Source/USDUtilities/Public/USDConversionUtils.h
+++ b/Engine/Plugins/Importers/USDImporter/Source/USDUtilities/Public/USDConversionUtils.h
@@ -7,6 +7,7 @@
 #include "UObject/UObjectGlobals.h"
 
 #include "USDMemory.h"
+#include "UsdWrappers/ForwardDeclarations.h"
 
 #if USE_USD_SDK
 #include "USDIncludesStart.h"
@@ -39,9 +40,7 @@ class UUsdAssetImportData;
 enum class EUsdUpAxis : uint8;
 namespace UE
 {
-	class FSdfLayer;
 	class FUsdPrim;
-	class FUsdStage;
 	class FSdfPath;
 }
 
diff --git a/Engine/Plugins/Importers/USDImporter/Source/USDUtilities/Public/USDGeomMeshConversion.h b/Engine/Plugins/Importers/USDImporter/Source/USDUtilities/Public/USDGeomMeshConversion.h
--- a/Engine/Plugins/Importers/USDImporter/Source/USDUtilities/Public/USDGeomMeshConversion.h
+++ b/Engine/Plugins/Importers/USDImporter/Source/USDUtilities/Public/USDGeomMeshConversion.h
@@ -6,11 +6,11 @@
 
 #include "CoreMinimal.h"
 
-#include "USDIncludesStart.h"
-
-#include "pxr/usd/usd/timeCode.h"
-#include "pxr/usd/usdShade/tokens.h"
+#include "UsdWrappers/ForwardDeclarations.h"
 
+#include "USDIncludesStart.h"
+	#include "pxr/usd/usd/timeCode.h"
+	#include "pxr/usd/usdShade/tokens.h"
 #include "USDIncludesEnd.h"
 
 PXR_NAMESPACE_OPEN_SCOPE
@@ -31,10 +31,6 @@ namespace UsdUtils
 	struct FDisplayColorMaterial;
 	struct FUsdPrimMaterialAssignmentInfo;
 }
-namespace UE
-{
-	class FUsdStage;
-}
 
 namespace UsdToUnreal
 {
diff --git a/Engine/Plugins/Importers/USDImporter/Source/USDUtilities/Public/USDLayerUtils.h b/Engine/Plugins/Importers/USDImporter/Source/USDUtilities/Public/USDLayerUtils.h
--- a/Engine/Plugins/Importers/USDImporter/Source/USDUtilities/Public/USDLayerUtils.h
+++ b/Engine/Plugins/Importers/USDImporter/Source/USDUtilities/Public/USDLayerUtils.h
@@ -5,6 +5,7 @@
 #if USE_USD_SDK
 
 #include "USDMemory.h"
+#include "UsdWrappers/ForwardDeclarations.h"
 
 #include "USDIncludesStart.h"
 	#include "pxr/pxr.h"
@@ -26,8 +27,6 @@ PXR_NAMESPACE_CLOSE_SCOPE
 
 namespace UE
 {
-	class FSdfLayer;
-	class FUsdStage;
 	struct FSdfLayerOffset;
 }
 
diff --git a/Engine/Plugins/Importers/USDImporter/Source/USDUtilities/Public/USDSkeletalDataConversion.h b/Engine/Plugins/Importers/USDImporter/Source/USDUtilities/Public/USDSkeletalDataConversion.h
--- a/Engine/Plugins/Importers/USDImporter/Source/USDUtilities/Public/USDSkeletalDataConversion.h
+++ b/Engine/Plugins/Importers/USDImporter/Source/USDUtilities/Public/USDSkeletalDataConversion.h
@@ -8,6 +8,7 @@
 #include "UObject/ObjectMacros.h"
 
 #include "USDConversionUtils.h"
+#include "UsdWrappers/ForwardDeclarations.h"
 
 #if USE_USD_SDK
 #include "USDIncludesStart.h"
@@ -39,10 +40,6 @@ namespace UsdUtils
 {
 	struct FUsdPrimMaterialSlot;
 }
-namespace UE
-{
-	class FUsdStage;
-}
 
 #endif // #if USE_USD_SDK
 
diff --git a/Engine/Plugins/Importers/USDImporter/Source/USDUtilities/Public/USDValueConversion.h b/Engine/Plugins/Importers/USDImporter/Source/USDUtilities/Public/USDValueConversion.h
--- a/Engine/Plugins/Importers/USDImporter/Source/USDUtilities/Public/USDValueConversion.h
+++ b/Engine/Plugins/Importers/USDImporter/Source/USDUtilities/Public/USDValueConversion.h
@@ -5,13 +5,12 @@
 #include "CoreMinimal.h"
 #include "Misc/TVariant.h"
 
+#include "UsdWrappers/ForwardDeclarations.h"
 #include "UsdWrappers/UsdAttribute.h"
 
 namespace UE
 {
 	class FVtValue;
-	class FUsdStage;
-	class FSdfLayer;
 }
 
 namespace UsdUtils
diff --git a/Engine/Plugins/Importers/USDImporter/Source/UnrealUSDWrapper/Private/UsdWrappers/SdfLayer.cpp b/Engine/Plugins/Importers/USDImporter/Source/UnrealUSDWrapper/Private/UsdWrappers/SdfLayer.cpp
--- a/Engine/Plugins/Importers/USDImporter/Source/UnrealUSDWrapper/Private/UsdWrappers/SdfLayer.cpp
+++ b/Engine/Plugins/Importers/USDImporter/Source/UnrealUSDWrapper/Private/UsdWrappers/SdfLayer.cpp
@@ -19,128 +19,216 @@ namespace UE
 {
 	namespace Internal
 	{
+		template<typename PtrType>
 		class FSdfLayerImpl
 		{
 		public:
 			FSdfLayerImpl() = default;
 
-#if USE_USD_SDK
-			explicit FSdfLayerImpl( const pxr::SdfLayerRefPtr& InSdfLayer )
-				: PxrSdfLayer( InSdfLayer )
+			explicit FSdfLayerImpl( const PtrType& InUsdPtr )
+				: UsdPtr( InUsdPtr )
 			{
 			}
 
-			explicit FSdfLayerImpl( pxr::SdfLayerRefPtr&& InSdfLayer )
-				: PxrSdfLayer( MoveTemp( InSdfLayer ) )
+			explicit FSdfLayerImpl( PtrType&& InUsdPtr )
+				: UsdPtr( MoveTemp( InUsdPtr ) )
 			{
 			}
 
-			TUsdStore< pxr::SdfLayerRefPtr > PxrSdfLayer;
-#endif // #if USE_USD_SDK
+			PtrType& GetInner()
+			{
+				return UsdPtr.Get();
+			}
+
+			const PtrType& GetInner() const
+			{
+				return UsdPtr.Get();
+			}
+
+			TUsdStore< PtrType > UsdPtr;
 		};
 	}
 
-	FSdfLayer::FSdfLayer()
+	template<typename PtrType>
+	FSdfLayerBase<PtrType>::FSdfLayerBase()
 	{
 		FScopedUnrealAllocs UnrealAllocs;
-		Impl = MakeUnique< Internal::FSdfLayerImpl >();
+		Impl = MakeUnique< Internal::FSdfLayerImpl<PtrType> >();
 	}
 
-	FSdfLayer::FSdfLayer( const FSdfLayer& Other )
+	template<typename PtrType>
+	FSdfLayerBase<PtrType>::FSdfLayerBase( const FSdfLayer& Other )
 	{
-#if USE_USD_SDK
 		FScopedUnrealAllocs UnrealAllocs;
-		Impl = MakeUnique< Internal::FSdfLayerImpl >( Other.Impl->PxrSdfLayer.Get() );
-#endif // #if USE_USD_SDK
+		Impl = MakeUnique< Internal::FSdfLayerImpl<PtrType> >( Other.Impl->GetInner() );
 	}
 
-	FSdfLayer::FSdfLayer( FSdfLayer&& Other ) = default;
+	template<typename PtrType>
+	FSdfLayerBase<PtrType>::FSdfLayerBase( FSdfLayer&& Other )
+	{
+		FScopedUnrealAllocs UnrealAllocs;
+		Impl = MakeUnique< Internal::FSdfLayerImpl<PtrType> >( MoveTemp( Other.Impl->GetInner() ) );
+	}
 
-	FSdfLayer& FSdfLayer::operator=( const FSdfLayer& Other )
+	template<typename PtrType>
+	FSdfLayerBase<PtrType>::FSdfLayerBase( const FSdfLayerWeak& Other )
 	{
-#if USE_USD_SDK
 		FScopedUnrealAllocs UnrealAllocs;
-		Impl = MakeUnique< Internal::FSdfLayerImpl >( Other.Impl->PxrSdfLayer.Get() );
-#endif // #if USE_USD_SDK
+		Impl = MakeUnique< Internal::FSdfLayerImpl<PtrType> >( Other.Impl->GetInner() );
+	}
+
+	template<typename PtrType>
+	FSdfLayerBase<PtrType>::FSdfLayerBase( FSdfLayerWeak&& Other )
+	{
+		FScopedUnrealAllocs UnrealAllocs;
+		Impl = MakeUnique< Internal::FSdfLayerImpl<PtrType> >( MoveTemp( Other.Impl->GetInner() ) );
+	}
+
+	template<typename PtrType>
+	FSdfLayerBase<PtrType>& FSdfLayerBase<PtrType>::operator=( const FSdfLayer& Other )
+	{
+		FScopedUnrealAllocs UnrealAllocs;
+		Impl = MakeUnique< Internal::FSdfLayerImpl<PtrType> >( Other.Impl->GetInner() );
 		return *this;
 	}
 
-	FSdfLayer& FSdfLayer::operator=( FSdfLayer&& Other )
+	template<typename PtrType>
+	FSdfLayerBase<PtrType>& FSdfLayerBase<PtrType>::operator=( FSdfLayer&& Other )
 	{
 		FScopedUnrealAllocs UnrealAllocs;
-		Impl = MoveTemp( Other.Impl );
+		Impl = MakeUnique< Internal::FSdfLayerImpl<PtrType> >( MoveTemp( Other.Impl->GetInner() ) );
+		return *this;
+	}
 
+	template<typename PtrType>
+	FSdfLayerBase<PtrType>& FSdfLayerBase<PtrType>::operator=( const FSdfLayerWeak& Other )
+	{
+		FScopedUnrealAllocs UnrealAllocs;
+		Impl = MakeUnique< Internal::FSdfLayerImpl<PtrType> >( Other.Impl->GetInner() );
 		return *this;
 	}
 
-	FSdfLayer::~FSdfLayer()
+	template<typename PtrType>
+	FSdfLayerBase<PtrType>& FSdfLayerBase<PtrType>::operator=( FSdfLayerWeak&& Other )
+	{
+		FScopedUnrealAllocs UnrealAllocs;
+		Impl = MakeUnique< Internal::FSdfLayerImpl<PtrType> >( MoveTemp( Other.Impl->GetInner() ) );
+		return *this;
+	}
+
+	template<typename PtrType>
+	FSdfLayerBase<PtrType>::~FSdfLayerBase()
 	{
 		FScopedUnrealAllocs UnrealAllocs;
 		Impl.Reset();
 	}
 
-	bool FSdfLayer::operator==( const FSdfLayer& Other ) const
+	template<typename PtrType>
+	bool FSdfLayerBase<PtrType>::operator==( const FSdfLayerBase& Other ) const
 	{
-#if USE_USD_SDK
-		return Impl->PxrSdfLayer.Get() == Other.Impl->PxrSdfLayer.Get();
-#else
-		return false;
-#endif // #if USE_USD_SDK
+		return Impl->GetInner() == Other.Impl->GetInner();
 	}
 
-	bool FSdfLayer::operator!=( const FSdfLayer& Other ) const
+	template<typename PtrType>
+	bool FSdfLayerBase<PtrType>::operator!=( const FSdfLayerBase& Other ) const
 	{
 		return !( *this == Other );
 	}
 
-	FSdfLayer::operator bool() const
+	template<typename PtrType>
+	FSdfLayerBase<PtrType>::operator bool() const
 	{
-#if USE_USD_SDK
-		return (bool)Impl->PxrSdfLayer.Get();
-#else
-		return false;
-#endif // #if USE_USD_SDK
+		return ( bool ) Impl->GetInner();
 	}
 
 #if USE_USD_SDK
-	FSdfLayer::FSdfLayer( const pxr::SdfLayerRefPtr& InSdfLayer )
+	template<typename PtrType>
+	FSdfLayerBase<PtrType>::FSdfLayerBase( const pxr::SdfLayerRefPtr& InSdfLayer )
 	{
 		FScopedUnrealAllocs UnrealAllocs;
-		Impl = MakeUnique< Internal::FSdfLayerImpl >( InSdfLayer );
+		Impl = MakeUnique< Internal::FSdfLayerImpl<PtrType> >( InSdfLayer );
 	}
 
-	FSdfLayer::FSdfLayer( pxr::SdfLayerRefPtr&& InSdfLayer )
+	template<typename PtrType>
+	FSdfLayerBase<PtrType>::FSdfLayerBase( pxr::SdfLayerRefPtr&& InSdfLayer )
 	{
 		FScopedUnrealAllocs UnrealAllocs;
-		Impl = MakeUnique< Internal::FSdfLayerImpl >( MoveTemp( InSdfLayer ) );
+		Impl = MakeUnique< Internal::FSdfLayerImpl<PtrType> >( InSdfLayer );
 	}
 
-	FSdfLayer& FSdfLayer::operator=( const pxr::SdfLayerRefPtr& InSdfLayer )
+	template<typename PtrType>
+	FSdfLayerBase<PtrType>::FSdfLayerBase( const pxr::SdfLayerWeakPtr& InSdfLayer )
 	{
 		FScopedUnrealAllocs UnrealAllocs;
-		Impl = MakeUnique< Internal::FSdfLayerImpl >( InSdfLayer );
+		Impl = MakeUnique< Internal::FSdfLayerImpl<PtrType> >( InSdfLayer );
+	}
+
+	template<typename PtrType>
+	FSdfLayerBase<PtrType>::FSdfLayerBase( pxr::SdfLayerWeakPtr&& InSdfLayer )
+	{
+		FScopedUnrealAllocs UnrealAllocs;
+		Impl = MakeUnique< Internal::FSdfLayerImpl<PtrType> >( InSdfLayer );
+	}
+
+	template<typename PtrType>
+	FSdfLayerBase<PtrType>& FSdfLayerBase<PtrType>::operator=( const pxr::SdfLayerRefPtr& InSdfLayer )
+	{
+		FScopedUnrealAllocs UnrealAllocs;
+		Impl = MakeUnique< Internal::FSdfLayerImpl<PtrType> >( InSdfLayer );
 		return *this;
 	}
 
-	FSdfLayer& FSdfLayer::operator=( pxr::SdfLayerRefPtr&& InSdfLayer )
+	template<typename PtrType>
+	FSdfLayerBase<PtrType>& FSdfLayerBase<PtrType>::operator=( pxr::SdfLayerRefPtr&& InSdfLayer )
 	{
 		FScopedUnrealAllocs UnrealAllocs;
-		Impl = MakeUnique< Internal::FSdfLayerImpl >( MoveTemp( InSdfLayer ) );
+		Impl = MakeUnique< Internal::FSdfLayerImpl<PtrType> >( InSdfLayer );
 		return *this;
 	}
 
-	FSdfLayer::operator pxr::SdfLayerRefPtr&()
+	template<typename PtrType>
+	FSdfLayerBase<PtrType>& FSdfLayerBase<PtrType>::operator=( const pxr::SdfLayerWeakPtr& InSdfLayer )
 	{
-		return Impl->PxrSdfLayer.Get();
+		FScopedUnrealAllocs UnrealAllocs;
+		Impl = MakeUnique< Internal::FSdfLayerImpl<PtrType> >( InSdfLayer );
+		return *this;
 	}
 
-	FSdfLayer::operator const pxr::SdfLayerRefPtr&() const
+	template<typename PtrType>
+	FSdfLayerBase<PtrType>& FSdfLayerBase<PtrType>::operator=( pxr::SdfLayerWeakPtr&& InSdfLayer )
 	{
-		return Impl->PxrSdfLayer.Get();
+		FScopedUnrealAllocs UnrealAllocs;
+		Impl = MakeUnique< Internal::FSdfLayerImpl<PtrType> >( InSdfLayer );
+		return *this;
+	}
+
+	template<typename PtrType>
+	FSdfLayerBase<PtrType>::operator PtrType&()
+	{
+		return Impl->GetInner();
+	}
+
+	template<typename PtrType>
+	FSdfLayerBase<PtrType>::operator const PtrType&() const
+	{
+		return Impl->GetInner();
+	}
+
+	template<typename PtrType>
+	FSdfLayerBase<PtrType>::operator pxr::SdfLayerRefPtr() const
+	{
+		return Impl->GetInner();
+	}
+
+	template<typename PtrType>
+	FSdfLayerBase<PtrType>::operator pxr::SdfLayerWeakPtr() const
+	{
+		return Impl->GetInner();
 	}
 #endif // #if USE_USD_SDK
 
-	FSdfLayer FSdfLayer::FindOrOpen( const TCHAR* Identifier )
+	template<typename PtrType>
+	FSdfLayer FSdfLayerBase<PtrType>::FindOrOpen( const TCHAR* Identifier )
 	{
 #if USE_USD_SDK
 		return FSdfLayer( pxr::SdfLayer::FindOrOpen( TCHAR_TO_ANSI( Identifier ) ) );
@@ -149,192 +237,268 @@ namespace UE
 #endif // #if USE_USD_SDK
 	}
 
-	bool FSdfLayer::Save( bool bForce ) const
+	template<typename PtrType>
+	bool FSdfLayerBase<PtrType>::Save( bool bForce ) const
 	{
 #if USE_USD_SDK
-		return Impl->PxrSdfLayer.Get()->Save( bForce );
-#else
-		return false;
+		if ( const PtrType& Ptr = Impl->GetInner() )
+		{
+			return Ptr->Save( bForce );
+		}
 #endif // #if USE_USD_SDK
+
+		return false;
 	}
 
-	FString FSdfLayer::GetRealPath() const
+	template<typename PtrType>
+	FString FSdfLayerBase<PtrType>::GetRealPath() const
 	{
 #if USE_USD_SDK
-		FScopedUsdAllocs UsdAllocs;
-		return FString( ANSI_TO_TCHAR( Impl->PxrSdfLayer.Get()->GetRealPath().c_str() ) );
-#else
-		return FString();
+		if ( const PtrType& Ptr = Impl->GetInner() )
+		{
+			FScopedUsdAllocs UsdAllocs;
+			return FString( ANSI_TO_TCHAR( Ptr->GetRealPath().c_str() ) );
+		}
 #endif // #if USE_USD_SDK
+
+		return FString();
 	}
 
-	FString FSdfLayer::GetIdentifier() const
+	template<typename PtrType>
+	FString FSdfLayerBase<PtrType>::GetIdentifier() const
 	{
 #if USE_USD_SDK
-		FScopedUsdAllocs UsdAllocs;
-		return FString( ANSI_TO_TCHAR( Impl->PxrSdfLayer.Get()->GetIdentifier().c_str() ) );
-#else
-		return FString();
+		if ( const PtrType& Ptr = Impl->GetInner() )
+		{
+			FScopedUsdAllocs UsdAllocs;
+			return FString( ANSI_TO_TCHAR( Ptr->GetIdentifier().c_str() ) );
+		}
 #endif // #if USE_USD_SDK
+
+		return FString();
 	}
 
-	FString FSdfLayer::GetDisplayName() const
+	template<typename PtrType>
+	FString FSdfLayerBase<PtrType>::GetDisplayName() const
 	{
 #if USE_USD_SDK
-		FScopedUsdAllocs UsdAllocs;
-		return FString( ANSI_TO_TCHAR( Impl->PxrSdfLayer.Get()->GetDisplayName().c_str() ) );
-#else
-		return FString();
+		if ( const PtrType& Ptr = Impl->GetInner() )
+		{
+			FScopedUsdAllocs UsdAllocs;
+			return FString( ANSI_TO_TCHAR( Ptr->GetDisplayName().c_str() ) );
+		}
 #endif // #if USE_USD_SDK
+
+		return FString();
 	}
 
-	bool FSdfLayer::HasStartTimeCode() const
+	template<typename PtrType>
+	bool FSdfLayerBase<PtrType>::HasStartTimeCode() const
 	{
 #if USE_USD_SDK
-		return Impl->PxrSdfLayer.Get()->HasStartTimeCode();
-#else
-		return false;
+		if ( const PtrType& Ptr = Impl->GetInner() )
+		{
+			return Ptr->HasStartTimeCode();
+		}
 #endif // #if USE_USD_SDK
+
+		return false;
 	}
 
-	bool FSdfLayer::HasEndTimeCode() const
+	template<typename PtrType>
+	bool FSdfLayerBase<PtrType>::HasEndTimeCode() const
 	{
 #if USE_USD_SDK
-		return Impl->PxrSdfLayer.Get()->HasEndTimeCode();
-#else
-		return false;
+		if ( const PtrType& Ptr = Impl->GetInner() )
+		{
+			return Ptr->HasEndTimeCode();
+		}
 #endif // #if USE_USD_SDK
+
+		return false;
 	}
 
-	int64 FSdfLayer::GetNumSubLayerPaths() const
+	template<typename PtrType>
+	int64 FSdfLayerBase<PtrType>::GetNumSubLayerPaths() const
 	{
 #if USE_USD_SDK
-		return (int64)Impl->PxrSdfLayer.Get()->GetNumSubLayerPaths();
-#else
-		return 0;
+		if ( const PtrType& Ptr = Impl->GetInner() )
+		{
+			return ( int64 ) Ptr->GetNumSubLayerPaths();
+		}
 #endif // #if USE_USD_SDK
+
+		return 0;
 	}
 
-	double FSdfLayer::GetStartTimeCode() const
+	template<typename PtrType>
+	double FSdfLayerBase<PtrType>::GetStartTimeCode() const
 	{
 #if USE_USD_SDK
-		return Impl->PxrSdfLayer.Get()->GetStartTimeCode();
-#else
-		return 0.0;
+		if ( const PtrType& Ptr = Impl->GetInner() )
+		{
+			return Ptr->GetStartTimeCode();
+		}
 #endif // #if USE_USD_SDK
+
+		return 0.0;
 	}
 
-	double FSdfLayer::GetEndTimeCode() const
+	template<typename PtrType>
+	double FSdfLayerBase<PtrType>::GetEndTimeCode() const
 	{
 #if USE_USD_SDK
-		return Impl->PxrSdfLayer.Get()->GetEndTimeCode();
-#else
-		return 0.0;
+		if ( const PtrType& Ptr = Impl->GetInner() )
+		{
+			return Ptr->GetEndTimeCode();
+		}
 #endif // #if USE_USD_SDK
+
+		return 0.0;
 	}
 
-	void FSdfLayer::SetStartTimeCode( double TimeCode )
+	template<typename PtrType>
+	void FSdfLayerBase<PtrType>::SetStartTimeCode( double TimeCode )
 	{
 #if USE_USD_SDK
-		if ( !Impl->PxrSdfLayer.Get()->HasStartTimeCode() || !FMath::IsNearlyEqual( TimeCode, Impl->PxrSdfLayer.Get()->GetStartTimeCode() ) )
+		if ( PtrType& Ptr = Impl->GetInner() )
 		{
-			Impl->PxrSdfLayer.Get()->SetStartTimeCode( TimeCode );
+			if ( !Ptr->HasStartTimeCode() || !FMath::IsNearlyEqual( TimeCode, Ptr->GetStartTimeCode() ) )
+			{
+				Ptr->SetStartTimeCode( TimeCode );
+			}
 		}
 #endif // #if USE_USD_SDK
 	}
 
-	void FSdfLayer::SetEndTimeCode( double TimeCode )
+	template<typename PtrType>
+	void FSdfLayerBase<PtrType>::SetEndTimeCode( double TimeCode )
 	{
 #if USE_USD_SDK
-		if ( !Impl->PxrSdfLayer.Get()->HasEndTimeCode() || !FMath::IsNearlyEqual( TimeCode, Impl->PxrSdfLayer.Get()->GetEndTimeCode() ) )
+		if ( PtrType& Ptr = Impl->GetInner() )
 		{
-			Impl->PxrSdfLayer.Get()->SetEndTimeCode( TimeCode );
+			if ( !Ptr->HasEndTimeCode() || !FMath::IsNearlyEqual( TimeCode, Ptr->GetEndTimeCode() ) )
+			{
+				Ptr->SetEndTimeCode( TimeCode );
+			}
 		}
 #endif // #if USE_USD_SDK
 	}
 
-	bool FSdfLayer::HasTimeCodesPerSecond() const
+	template<typename PtrType>
+	bool FSdfLayerBase<PtrType>::HasTimeCodesPerSecond() const
 	{
 #if USE_USD_SDK
-		return Impl->PxrSdfLayer.Get()->HasTimeCodesPerSecond();
-#else
-		return false;
+		if ( const PtrType& Ptr = Impl->GetInner() )
+		{
+			return Ptr->HasTimeCodesPerSecond();
+		}
 #endif // #if USE_USD_SDK
+
+		return false;
 	}
 
-	double FSdfLayer::GetTimeCodesPerSecond() const
+	template<typename PtrType>
+	double FSdfLayerBase<PtrType>::GetTimeCodesPerSecond() const
 	{
 #if USE_USD_SDK
-		return Impl->PxrSdfLayer.Get()->GetTimeCodesPerSecond();
-#else
-		return 0.0;
+		if ( const PtrType& Ptr = Impl->GetInner() )
+		{
+			return Ptr->GetTimeCodesPerSecond();
+		}
 #endif // #if USE_USD_SDK
+
+		return 0.0;
 	}
 
-	void FSdfLayer::SetTimeCodesPerSecond( double TimeCodesPerSecond )
+	template<typename PtrType>
+	void FSdfLayerBase<PtrType>::SetTimeCodesPerSecond( double TimeCodesPerSecond )
 	{
 #if USE_USD_SDK
-		return Impl->PxrSdfLayer.Get()->SetTimeCodesPerSecond( TimeCodesPerSecond );
+		if ( PtrType& Ptr = Impl->GetInner() )
+		{
+			return Ptr->SetTimeCodesPerSecond( TimeCodesPerSecond );
+		}
 #endif // #if USE_USD_SDK
 	}
 
-	bool FSdfLayer::HasFramesPerSecond() const
+	template<typename PtrType>
+	bool FSdfLayerBase<PtrType>::HasFramesPerSecond() const
 	{
 #if USE_USD_SDK
-		return Impl->PxrSdfLayer.Get()->HasFramesPerSecond();
-#else
-		return false;
+		if ( const PtrType& Ptr = Impl->GetInner() )
+		{
+			return Ptr->HasFramesPerSecond();
+		}
 #endif // #if USE_USD_SDK
+
+		return false;
 	}
 
-	double FSdfLayer::GetFramesPerSecond() const
+	template<typename PtrType>
+	double FSdfLayerBase<PtrType>::GetFramesPerSecond() const
 	{
 #if USE_USD_SDK
-		return Impl->PxrSdfLayer.Get()->GetFramesPerSecond();
-#else
-		return 0.0;
+		if ( const PtrType& Ptr = Impl->GetInner() )
+		{
+			return Ptr->GetFramesPerSecond();
+		}
 #endif // #if USE_USD_SDK
+
+		return 0.0;
 	}
 
-	void FSdfLayer::SetFramesPerSecond( double FramesPerSecond )
+	template<typename PtrType>
+	void FSdfLayerBase<PtrType>::SetFramesPerSecond( double FramesPerSecond )
 	{
 #if USE_USD_SDK
-		return Impl->PxrSdfLayer.Get()->SetFramesPerSecond( FramesPerSecond );
+		if ( PtrType& Ptr = Impl->GetInner() )
+		{
+			return Ptr->SetFramesPerSecond( FramesPerSecond );
+		}
 #endif // #if USE_USD_SDK
 	}
 
-	TArray< FString > FSdfLayer::GetSubLayerPaths() const
+	template<typename PtrType>
+	TArray< FString > FSdfLayerBase<PtrType>::GetSubLayerPaths() const
 	{
 		TArray< FString > SubLayerPaths;
 
 #if USE_USD_SDK
-		FScopedUsdAllocs UsdAllocs;
-
-		for ( const std::string& SubLayerPath : Impl->PxrSdfLayer.Get()->GetSubLayerPaths() )
+		if ( const PtrType& Ptr = Impl->GetInner() )
 		{
-			SubLayerPaths.Emplace( ANSI_TO_TCHAR( SubLayerPath.c_str() ) );
+			FScopedUsdAllocs UsdAllocs;
+
+			for ( const std::string& SubLayerPath : Ptr->GetSubLayerPaths() )
+			{
+				SubLayerPaths.Emplace( ANSI_TO_TCHAR( SubLayerPath.c_str() ) );
+			}
 		}
 #endif // #if USE_USD_SDK
 
 		return SubLayerPaths;
 	}
 
-	TArray< FSdfLayerOffset > FSdfLayer::GetSubLayerOffsets() const
+	template<typename PtrType>
+	TArray< FSdfLayerOffset > FSdfLayerBase<PtrType>::GetSubLayerOffsets() const
 	{
 		TArray< FSdfLayerOffset > SubLayerOffsets;
 
 #if USE_USD_SDK
-		FScopedUsdAllocs UsdAllocs;
-
-		for ( const pxr::SdfLayerOffset& SubLayerOffset : Impl->PxrSdfLayer.Get()->GetSubLayerOffsets() )
+		if ( const PtrType& Ptr = Impl->GetInner() )
 		{
-			if ( SubLayerOffset.IsValid() )
-			{
-				SubLayerOffsets.Emplace( SubLayerOffset.GetOffset(), SubLayerOffset.GetScale() );
-			}
-			else
+			FScopedUsdAllocs UsdAllocs;
+
+			for ( const pxr::SdfLayerOffset& SubLayerOffset : Ptr->GetSubLayerOffsets() )
 			{
-				SubLayerOffsets.AddDefaulted();
+				if ( SubLayerOffset.IsValid() )
+				{
+					SubLayerOffsets.Emplace( SubLayerOffset.GetOffset(), SubLayerOffset.GetScale() );
+				}
+				else
+				{
+					SubLayerOffsets.AddDefaulted();
+				}
 			}
 		}
 #endif // #if USE_USD_SDK
@@ -342,75 +506,103 @@ namespace UE
 		return SubLayerOffsets;
 	}
 
-	void FSdfLayer::SetSubLayerOffset( const FSdfLayerOffset& LayerOffset, int32 Index )
+	template<typename PtrType>
+	void FSdfLayerBase<PtrType>::SetSubLayerOffset( const FSdfLayerOffset& LayerOffset, int32 Index )
 	{
 #if USE_USD_SDK
-		FScopedUsdAllocs UsdAllocs;
+		if ( PtrType& Ptr = Impl->GetInner() )
+		{
+			FScopedUsdAllocs UsdAllocs;
 
-		pxr::SdfLayerOffset UsdLayerOffset( LayerOffset.Offset, LayerOffset.Scale );
-		Impl->PxrSdfLayer.Get()->SetSubLayerOffset( MoveTemp( UsdLayerOffset ), Index );
+			pxr::SdfLayerOffset UsdLayerOffset( LayerOffset.Offset, LayerOffset.Scale );
+			Ptr->SetSubLayerOffset( MoveTemp( UsdLayerOffset ), Index );
+		}
 #endif // #if USE_USD_SDK
 	}
 
-	bool FSdfLayer::HasSpec( const FSdfPath& Path ) const
+	template<typename PtrType>
+	bool FSdfLayerBase<PtrType>::HasSpec( const FSdfPath& Path ) const
 	{
 #if USE_USD_SDK
-		return Impl->PxrSdfLayer.Get()->HasSpec( Path );
-#else
-		return false;
+		if ( const PtrType& Ptr = Impl->GetInner() )
+		{
+			return Ptr->HasSpec( Path );
+		}
 #endif // #if USE_USD_SDK
+
+		return false;
 	}
 
-	TSet< double > FSdfLayer::ListTimeSamplesForPath( const FSdfPath& Path ) const
+	template<typename PtrType>
+	TSet< double > FSdfLayerBase<PtrType>::ListTimeSamplesForPath( const FSdfPath& Path ) const
 	{
 		TSet< double > TimeSamples;
 
 #if USE_USD_SDK
-		FScopedUsdAllocs UsdAllocs;
+		if ( const PtrType& Ptr = Impl->GetInner() )
+		{
+			FScopedUsdAllocs UsdAllocs;
 
-		std::set< double > UsdTimeSamples = Impl->PxrSdfLayer.Get()->ListTimeSamplesForPath( Path );
+			std::set< double > UsdTimeSamples = Ptr->ListTimeSamplesForPath( Path );
 
-		for ( double UsdTimeSample : UsdTimeSamples )
-		{
-			TimeSamples.Add( UsdTimeSample );
+			for ( double UsdTimeSample : UsdTimeSamples )
+			{
+				TimeSamples.Add( UsdTimeSample );
+			}
 		}
 #endif // #if USE_USD_SDK
 
 		return TimeSamples;
 	}
 
-	void FSdfLayer::EraseTimeSample( const FSdfPath& Path, double Time )
+	template<typename PtrType>
+	void FSdfLayerBase<PtrType>::EraseTimeSample( const FSdfPath& Path, double Time )
 	{
 #if USE_USD_SDK
-		return Impl->PxrSdfLayer.Get()->EraseTimeSample( Path, Time );
+		if ( PtrType& Ptr = Impl->GetInner() )
+		{
+			return Ptr->EraseTimeSample( Path, Time );
+		}
 #endif // #if USE_USD_SDK
 	}
 
-	bool FSdfLayer::IsEmpty() const
+	template<typename PtrType>
+	bool FSdfLayerBase<PtrType>::IsEmpty() const
 	{
 #if USE_USD_SDK
-		return Impl->PxrSdfLayer.Get()->IsEmpty();
-#else
-		return false;
+		if ( const PtrType& Ptr = Impl->GetInner() )
+		{
+			return Ptr->IsEmpty();
+		}
 #endif // #if USE_USD_SDK
+
+		return false;
 	}
 
-	bool FSdfLayer::IsAnonymous() const
+	template<typename PtrType>
+	bool FSdfLayerBase<PtrType>::IsAnonymous() const
 	{
 #if USE_USD_SDK
-		return Impl->PxrSdfLayer.Get()->IsAnonymous();
-#else
-		return false;
+		if ( const PtrType& Ptr = Impl->GetInner() )
+		{
+			return Ptr->IsAnonymous();
+		}
 #endif // #if USE_USD_SDK
+
+		return false;
 	}
 
-	bool FSdfLayer::Export( const TCHAR* Filename ) const
+	template<typename PtrType>
+	bool FSdfLayerBase<PtrType>::Export( const TCHAR* Filename ) const
 	{
 #if USE_USD_SDK
-		return Impl->PxrSdfLayer.Get()->Export( TCHAR_TO_ANSI( Filename ) );
-#else
-		return false;
+		if ( const PtrType& Ptr = Impl->GetInner() )
+		{
+			return Ptr->Export( TCHAR_TO_ANSI( Filename ) );
+		}
 #endif // #if USE_USD_SDK
+
+		return false;
 	}
 
 	FString FSdfLayerUtils::SdfComputeAssetPathRelativeToLayer( const FSdfLayer& Anchor, const TCHAR* AssetPath )
@@ -423,19 +615,35 @@ namespace UE
 #endif // #if USE_USD_SDK
 	}
 
-	bool FSdfLayer::IsMuted() const
+	template<typename PtrType>
+	bool FSdfLayerBase<PtrType>::IsMuted() const
 	{
 #if USE_USD_SDK
-		return Impl->PxrSdfLayer.Get()->IsMuted();
-#else
-		return false;
+		if ( const PtrType& Ptr = Impl->GetInner() )
+		{
+			return Ptr->IsMuted();
+		}
 #endif // #if USE_USD_SDK
+
+		return false;
 	}
 
-	void FSdfLayer::SetMuted( bool bMuted )
+	template<typename PtrType>
+	void FSdfLayerBase<PtrType>::SetMuted( bool bMuted )
 	{
 #if USE_USD_SDK
-		return Impl->PxrSdfLayer.Get()->SetMuted( bMuted );
+		if ( PtrType& Ptr = Impl->GetInner() )
+		{
+			return Ptr->SetMuted( bMuted );
+		}
 #endif // #if USE_USD_SDK
 	}
+
+#if USE_USD_SDK
+	template class UNREALUSDWRAPPER_API FSdfLayerBase<pxr::SdfLayerRefPtr>;
+	template class UNREALUSDWRAPPER_API FSdfLayerBase<pxr::SdfLayerWeakPtr>;
+#else
+	template class UNREALUSDWRAPPER_API FSdfLayerBase<FDummyRefPtrType>;
+	template class UNREALUSDWRAPPER_API FSdfLayerBase<FDummyWeakPtrType>;
+#endif // #if USE_USD_SDK
 }
diff --git a/Engine/Plugins/Importers/USDImporter/Source/UnrealUSDWrapper/Private/UsdWrappers/UsdStage.cpp b/Engine/Plugins/Importers/USDImporter/Source/UnrealUSDWrapper/Private/UsdWrappers/UsdStage.cpp
--- a/Engine/Plugins/Importers/USDImporter/Source/UnrealUSDWrapper/Private/UsdWrappers/UsdStage.cpp
+++ b/Engine/Plugins/Importers/USDImporter/Source/UnrealUSDWrapper/Private/UsdWrappers/UsdStage.cpp
@@ -23,338 +23,507 @@ namespace UE
 {
 	namespace Internal
 	{
+		template<typename PtrType>
 		class FUsdStageImpl
 		{
 		public:
 			FUsdStageImpl() = default;
 
-#if USE_USD_SDK
-			explicit FUsdStageImpl( const pxr::UsdStageRefPtr& InUsdStageRefPtr )
-				: PxrUsdStageRefPtr( InUsdStageRefPtr )
+			explicit FUsdStageImpl( const PtrType& InUsdPtr )
+				: UsdPtr( InUsdPtr )
 			{
 			}
 
-			explicit FUsdStageImpl( pxr::UsdStageRefPtr&& InUsdStageRefPtr )
-				: PxrUsdStageRefPtr( MoveTemp( InUsdStageRefPtr ) )
+			explicit FUsdStageImpl( PtrType&& InUsdPtr )
+				: UsdPtr( MoveTemp( InUsdPtr ) )
 			{
 			}
 
-			TUsdStore< pxr::UsdStageRefPtr > PxrUsdStageRefPtr;
-#endif // #if USE_USD_SDK
+			PtrType& GetInner()
+			{
+				return UsdPtr.Get();
+			}
+
+			const PtrType& GetInner() const
+			{
+				return UsdPtr.Get();
+			}
+
+			TUsdStore< PtrType > UsdPtr;
 		};
 	}
 
-	FUsdStage::FUsdStage()
+	template<typename PtrType>
+	FUsdStageBase<PtrType>::FUsdStageBase()
 	{
 		FScopedUnrealAllocs UnrealAllocs;
-		Impl =  MakeUnique< Internal::FUsdStageImpl >();
+		Impl = MakeUnique< Internal::FUsdStageImpl<PtrType> >();
 	}
 
-	FUsdStage::FUsdStage( const FUsdStage& Other )
-#if USE_USD_SDK
-		: Impl( MakeUnique< Internal::FUsdStageImpl >( Other.Impl->PxrUsdStageRefPtr.Get() ) )
-#endif // #if USE_USD_SDK
+	template<typename PtrType>
+	FUsdStageBase<PtrType>::FUsdStageBase( const FUsdStage& Other )
+		: Impl( MakeUnique< Internal::FUsdStageImpl<PtrType> >( Other.Impl->GetInner() ) )
 	{
 	}
 
-	FUsdStage::FUsdStage( FUsdStage&& Other ) = default;
+	template<typename PtrType>
+	FUsdStageBase<PtrType>::FUsdStageBase( FUsdStage&& Other )
+		: Impl( MakeUnique< Internal::FUsdStageImpl<PtrType> >( MoveTemp( Other.Impl->GetInner() ) ) )
+	{
+	}
 
-	FUsdStage& FUsdStage::operator=( const FUsdStage& Other )
+	template<typename PtrType>
+	FUsdStageBase<PtrType>::FUsdStageBase( const FUsdStageWeak& Other )
+		: Impl( MakeUnique< Internal::FUsdStageImpl<PtrType> >( Other.Impl->GetInner() ) )
+	{
+	}
+
+	template<typename PtrType>
+	FUsdStageBase<PtrType>::FUsdStageBase( FUsdStageWeak&& Other )
+		: Impl( MakeUnique< Internal::FUsdStageImpl<PtrType> >( MoveTemp( Other.Impl->GetInner() ) ) )
+	{
+	}
+
+	template<typename PtrType>
+	FUsdStageBase<PtrType>& FUsdStageBase<PtrType>::operator=( const FUsdStage& Other )
 	{
-#if USE_USD_SDK
 		FScopedUnrealAllocs UnrealAllocs;
-		Impl = MakeUnique< Internal::FUsdStageImpl >( Other.Impl->PxrUsdStageRefPtr.Get() );
-#endif // #if USE_USD_SDK
+		Impl = MakeUnique< Internal::FUsdStageImpl<PtrType> >( Other.Impl->GetInner() );
 		return *this;
 	}
 
-	FUsdStage& FUsdStage::operator=( FUsdStage&& Other )
+	template<typename PtrType>
+	FUsdStageBase<PtrType>& FUsdStageBase<PtrType>::operator=( FUsdStage&& Other )
 	{
 		FScopedUnrealAllocs UnrealAllocs;
-		Impl = MoveTemp( Other.Impl );
+		Impl = MakeUnique< Internal::FUsdStageImpl<PtrType> >( MoveTemp( Other.Impl->GetInner() ) );
+		return *this;
+	}
 
+	template<typename PtrType>
+	FUsdStageBase<PtrType>& FUsdStageBase<PtrType>::operator=( const FUsdStageWeak& Other )
+	{
+		FScopedUnrealAllocs UnrealAllocs;
+		Impl = MakeUnique< Internal::FUsdStageImpl<PtrType> >( Other.Impl->GetInner() );
 		return *this;
 	}
 
-	FUsdStage::operator bool() const
+	template<typename PtrType>
+	FUsdStageBase<PtrType>& FUsdStageBase<PtrType>::operator=( FUsdStageWeak&& Other )
 	{
-#if USE_USD_SDK
-		return (bool)Impl->PxrUsdStageRefPtr.Get();
-#else
-		return false;
-#endif // #if USE_USD_SDK
+		FScopedUnrealAllocs UnrealAllocs;
+		Impl = MakeUnique< Internal::FUsdStageImpl<PtrType> >( MoveTemp( Other.Impl->GetInner() ) );
+		return *this;
 	}
 
-	bool FUsdStage::operator==( const FUsdStage& Other ) const
+	template<typename PtrType>
+	FUsdStageBase<PtrType>::operator bool() const
 	{
-#if USE_USD_SDK
-		return Impl->PxrUsdStageRefPtr.Get() == Other.Impl->PxrUsdStageRefPtr.Get();
-#else
-		return false;
-#endif // #if USE_USD_SDK
+		return ( bool ) Impl->GetInner();
 	}
 
-	bool FUsdStage::operator!=( const FUsdStage& Other ) const
+	template<typename PtrType>
+	bool FUsdStageBase<PtrType>::operator==( const FUsdStageBase<PtrType>& Other ) const
+	{
+		return Impl->GetInner() == Other.Impl->GetInner();
+	}
+
+	template<typename PtrType>
+	bool FUsdStageBase<PtrType>::operator!=( const FUsdStageBase<PtrType>& Other ) const
 	{
 		return !( *this == Other );
 	}
 
-	FUsdStage::~FUsdStage()
+	template<typename PtrType>
+	FUsdStageBase<PtrType>::~FUsdStageBase()
 	{
 		FScopedUnrealAllocs UnrealAllocs;
 		Impl.Reset();
 	}
 
+	// Auto conversion from/to PtrType
 #if USE_USD_SDK
-	FUsdStage::FUsdStage( const pxr::UsdStageRefPtr& InUsdStageRefPtr )
+	template<typename PtrType>
+	FUsdStageBase<PtrType>::FUsdStageBase( const pxr::UsdStageRefPtr& InUsdPtr )
 	{
 		FScopedUnrealAllocs UnrealAllocs;
-		Impl = MakeUnique< Internal::FUsdStageImpl >( InUsdStageRefPtr );
+
+		Impl = MakeUnique< Internal::FUsdStageImpl<PtrType> >( InUsdPtr );
 	}
 
-	FUsdStage::FUsdStage( pxr::UsdStageRefPtr&& InUsdStageRefPtr )
+	template<typename PtrType>
+	FUsdStageBase<PtrType>::FUsdStageBase( pxr::UsdStageRefPtr&& InUsdPtr )
 	{
 		FScopedUnrealAllocs UnrealAllocs;
-		Impl = MakeUnique< Internal::FUsdStageImpl >( MoveTemp( InUsdStageRefPtr ) );
+
+		Impl = MakeUnique< Internal::FUsdStageImpl<PtrType> >( InUsdPtr );
 	}
 
-	FUsdStage::operator pxr::UsdStageRefPtr&()
+	template<typename PtrType>
+	FUsdStageBase<PtrType>::FUsdStageBase( const pxr::UsdStageWeakPtr& InUsdPtr )
 	{
-		return Impl->PxrUsdStageRefPtr.Get();
+		FScopedUnrealAllocs UnrealAllocs;
+
+		Impl = MakeUnique< Internal::FUsdStageImpl<PtrType> >( InUsdPtr );
 	}
 
-	FUsdStage::operator const pxr::UsdStageRefPtr&() const
+	template<typename PtrType>
+	FUsdStageBase<PtrType>::FUsdStageBase( pxr::UsdStageWeakPtr&& InUsdPtr )
 	{
-		return Impl->PxrUsdStageRefPtr.Get();
+		FScopedUnrealAllocs UnrealAllocs;
+
+		Impl = MakeUnique< Internal::FUsdStageImpl<PtrType> >( InUsdPtr );
 	}
 
-	FUsdStage::operator pxr::UsdStageWeakPtr() const
+	template<typename PtrType>
+	FUsdStageBase<PtrType>::operator PtrType& ( )
 	{
-		return pxr::UsdStageWeakPtr( Impl->PxrUsdStageRefPtr.Get() );
+		return Impl->GetInner();
 	}
-#endif // #if USE_USD_SDK
 
-	FSdfLayer FUsdStage::GetRootLayer() const
+	template<typename PtrType>
+	FUsdStageBase<PtrType>::operator const PtrType& ( ) const
+	{
+		return Impl->GetInner();
+	}
+
+	template<typename PtrType>
+	FUsdStageBase<PtrType>::operator pxr::UsdStageRefPtr ( ) const
+	{
+		return Impl->GetInner();
+	}
+
+	template<typename PtrType>
+	FUsdStageBase<PtrType>::operator pxr::UsdStageWeakPtr ( ) const
+	{
+		return Impl->GetInner();
+	}
+#endif // USE_USD_SDK
+
+	template<typename PtrType>
+	FSdfLayer FUsdStageBase<PtrType>::GetRootLayer() const
 	{
 #if USE_USD_SDK
-		return FSdfLayer( Impl->PxrUsdStageRefPtr.Get()->GetRootLayer() );
-#else
-		return FSdfLayer();
+		if ( PtrType& Ptr = Impl->GetInner() )
+		{
+			return FSdfLayer( Ptr->GetRootLayer() );
+		}
 #endif // #if USE_USD_SDK
+		return FSdfLayer();
 	}
 
-	FSdfLayer FUsdStage::GetSessionLayer() const
+	template<typename PtrType>
+	FSdfLayer FUsdStageBase<PtrType>::GetSessionLayer() const
 	{
 #if USE_USD_SDK
-		return FSdfLayer( Impl->PxrUsdStageRefPtr.Get()->GetSessionLayer() );
-#else
-		return FSdfLayer();
+		if ( PtrType& Ptr = Impl->GetInner() )
+		{
+			return FSdfLayer( Ptr->GetSessionLayer() );
+		}
 #endif // #if USE_USD_SDK
+
+		return FSdfLayer();
 	}
 
-	bool FUsdStage::HasLocalLayer( const FSdfLayer& Layer ) const
+	template<typename PtrType>
+	bool FUsdStageBase<PtrType>::HasLocalLayer( const FSdfLayer & Layer ) const
 	{
 #if USE_USD_SDK
-		return Impl->PxrUsdStageRefPtr.Get()->HasLocalLayer( pxr::SdfLayerRefPtr( Layer ) );
-#else
-		return false;
+		if ( PtrType& Ptr = Impl->GetInner() )
+		{
+			return Ptr->HasLocalLayer( pxr::SdfLayerRefPtr( Layer ) );
+		}
 #endif // #if USE_USD_SDK
+
+		return false;
 	}
 
-	FUsdPrim FUsdStage::GetPseudoRoot() const
+	template<typename PtrType>
+	FUsdPrim FUsdStageBase<PtrType>::GetPseudoRoot() const
 	{
 #if USE_USD_SDK
-		return FUsdPrim( Impl->PxrUsdStageRefPtr.Get()->GetPseudoRoot() );
-#else
-		return FUsdPrim();
+		if ( PtrType& Ptr = Impl->GetInner() )
+		{
+			return FUsdPrim( Ptr->GetPseudoRoot() );
+		}
 #endif // #if USE_USD_SDK
+
+		return FUsdPrim();
 	}
 
-	FUsdPrim FUsdStage::GetDefaultPrim() const
+	template<typename PtrType>
+	FUsdPrim FUsdStageBase<PtrType>::GetDefaultPrim() const
 	{
 #if USE_USD_SDK
-		return FUsdPrim( Impl->PxrUsdStageRefPtr.Get()->GetDefaultPrim() );
-#else
-		return FUsdPrim();
+		if ( PtrType& Ptr = Impl->GetInner() )
+		{
+			return FUsdPrim( Ptr->GetDefaultPrim() );
+		}
 #endif // #if USE_USD_SDK
+
+		return FUsdPrim();
 	}
 
-	FUsdPrim FUsdStage::GetPrimAtPath( const FSdfPath& Path ) const
+	template<typename PtrType>
+	FUsdPrim FUsdStageBase<PtrType>::GetPrimAtPath( const FSdfPath & Path ) const
 	{
 #if USE_USD_SDK
-		return FUsdPrim( Impl->PxrUsdStageRefPtr.Get()->GetPrimAtPath( Path ) );
-#else
-		return FUsdPrim();
+		if ( PtrType& Ptr = Impl->GetInner() )
+		{
+			return FUsdPrim( Ptr->GetPrimAtPath( Path ) );
+		}
 #endif // #if USE_USD_SDK
+
+		return FUsdPrim();
 	}
 
-	bool FUsdStage::IsEditTargetValid() const
+	template<typename PtrType>
+	bool FUsdStageBase<PtrType>::IsEditTargetValid() const
 	{
 #if USE_USD_SDK
-		return Impl->PxrUsdStageRefPtr.Get()->GetEditTarget().IsValid();
-#else
-		return false;
+		if ( PtrType& Ptr = Impl->GetInner() )
+		{
+			return Ptr->GetEditTarget().IsValid();
+		}
 #endif // #if USE_USD_SDK
+
+		return false;
 	}
 
-	void FUsdStage::SetEditTarget( const FSdfLayer& Layer )
+	template<typename PtrType>
+	void FUsdStageBase<PtrType>::SetEditTarget( const FSdfLayer & Layer )
 	{
 #if USE_USD_SDK
-		FScopedUsdAllocs UsdAllocs;
+		if ( PtrType& Ptr = Impl->GetInner() )
+		{
+			FScopedUsdAllocs UsdAllocs;
 
-		pxr::SdfLayerRefPtr LayerRef( Layer );
-		const pxr::UsdEditTarget EditTarget = Impl->PxrUsdStageRefPtr.Get()->GetEditTargetForLocalLayer( LayerRef );
+			pxr::SdfLayerRefPtr LayerRef( Layer );
+			const pxr::UsdEditTarget EditTarget = Ptr->GetEditTargetForLocalLayer( LayerRef );
 
-		Impl->PxrUsdStageRefPtr.Get()->SetEditTarget( EditTarget );
+			Ptr->SetEditTarget( EditTarget );
+		}
 #endif // #if USE_USD_SDK
 	}
 
-	FSdfLayer FUsdStage::GetEditTarget() const
+	template<typename PtrType>
+	FSdfLayer FUsdStageBase<PtrType>::GetEditTarget() const
 	{
 #if USE_USD_SDK
-		if ( IsEditTargetValid() )
+		if ( PtrType& Ptr = Impl->GetInner() )
 		{
-			return FSdfLayer( Impl->PxrUsdStageRefPtr.Get()->GetEditTarget().GetLayer() );
-		}
-		else
-		{
-			return FSdfLayer();
+			if ( IsEditTargetValid() )
+			{
+				return FSdfLayer( Ptr->GetEditTarget().GetLayer() );
+			}
 		}
-#else
-		return FSdfLayer();
 #endif // #if USE_USD_SDK
+
+		return FSdfLayer();
 	}
 
-	bool FUsdStage::GetMetadata( const TCHAR* Key, UE::FVtValue& Value ) const
+	template<typename PtrType>
+	bool FUsdStageBase<PtrType>::GetMetadata( const TCHAR * Key, UE::FVtValue & Value ) const
 	{
 #if USE_USD_SDK
-		return Impl->PxrUsdStageRefPtr.Get()->GetMetadata( pxr::TfToken{ TCHAR_TO_ANSI( Key ) }, &Value.GetUsdValue() );
-#else
-		return false;
+		if ( PtrType& Ptr = Impl->GetInner() )
+		{
+			return Ptr->GetMetadata( pxr::TfToken{ TCHAR_TO_ANSI( Key ) }, &Value.GetUsdValue() );
+		}
 #endif // #if USE_USD_SDK
+
+		return false;
 	}
 
-	bool FUsdStage::HasMetadata( const TCHAR* Key ) const
+	template<typename PtrType>
+	bool FUsdStageBase<PtrType>::HasMetadata( const TCHAR * Key ) const
 	{
 #if USE_USD_SDK
-		return Impl->PxrUsdStageRefPtr.Get()->HasMetadata( pxr::TfToken{ TCHAR_TO_ANSI( Key ) } );
-#else
-		return false;
+		if ( PtrType& Ptr = Impl->GetInner() )
+		{
+			return Ptr->HasMetadata( pxr::TfToken{ TCHAR_TO_ANSI( Key ) } );
+		}
 #endif // #if USE_USD_SDK
+
+		return false;
 	}
 
-	bool FUsdStage::SetMetadata( const TCHAR* Key, const UE::FVtValue& Value ) const
+	template<typename PtrType>
+	bool FUsdStageBase<PtrType>::SetMetadata( const TCHAR * Key, const UE::FVtValue & Value ) const
 	{
 #if USE_USD_SDK
-		return Impl->PxrUsdStageRefPtr.Get()->SetMetadata( pxr::TfToken{ TCHAR_TO_ANSI( Key ) }, Value.GetUsdValue() );
-#else
-		return false;
+		if ( PtrType& Ptr = Impl->GetInner() )
+		{
+			return Ptr->SetMetadata( pxr::TfToken{ TCHAR_TO_ANSI( Key ) }, Value.GetUsdValue() );
+		}
 #endif // #if USE_USD_SDK
+
+		return false;
 	}
 
-	bool FUsdStage::ClearMetadata( const TCHAR* Key ) const
+	template<typename PtrType>
+	bool FUsdStageBase<PtrType>::ClearMetadata( const TCHAR * Key ) const
 	{
 #if USE_USD_SDK
-		return Impl->PxrUsdStageRefPtr.Get()->ClearMetadata( pxr::TfToken{ TCHAR_TO_ANSI( Key ) } );
-#else
-		return false;
+		if ( PtrType& Ptr = Impl->GetInner() )
+		{
+			return Ptr->ClearMetadata( pxr::TfToken{ TCHAR_TO_ANSI( Key ) } );
+		}
 #endif // #if USE_USD_SDK
+
+		return false;
 	}
 
-	double FUsdStage::GetStartTimeCode() const
+	template<typename PtrType>
+	double FUsdStageBase<PtrType>::GetStartTimeCode() const
 	{
 #if USE_USD_SDK
-		return Impl->PxrUsdStageRefPtr.Get()->GetStartTimeCode();
-#else
-		return 0.0;
+		if ( PtrType& Ptr = Impl->GetInner() )
+		{
+			return Ptr->GetStartTimeCode();
+		}
 #endif // #if USE_USD_SDK
+
+		return 0.0;
 	}
 
-	double FUsdStage::GetEndTimeCode() const
+	template<typename PtrType>
+	double FUsdStageBase<PtrType>::GetEndTimeCode() const
 	{
 #if USE_USD_SDK
-		return Impl->PxrUsdStageRefPtr.Get()->GetEndTimeCode();
-#else
-		return 0.0;
+		if ( PtrType& Ptr = Impl->GetInner() )
+		{
+			return Ptr->GetEndTimeCode();
+		}
 #endif // #if USE_USD_SDK
+
+		return 0.0;
 	}
 
-	void FUsdStage::SetStartTimeCode( double TimeCode )
+	template<typename PtrType>
+	void FUsdStageBase<PtrType>::SetStartTimeCode( double TimeCode )
 	{
 #if USE_USD_SDK
-		if ( !FMath::IsNearlyEqual( TimeCode, Impl->PxrUsdStageRefPtr.Get()->GetStartTimeCode() ) )
+		if ( PtrType& Ptr = Impl->GetInner() )
 		{
-			Impl->PxrUsdStageRefPtr.Get()->SetStartTimeCode( TimeCode );
+			if ( !FMath::IsNearlyEqual( TimeCode, Ptr->GetStartTimeCode() ) )
+			{
+				Ptr->SetStartTimeCode( TimeCode );
+			}
 		}
 #endif // #if USE_USD_SDK
 	}
 
-	void FUsdStage::SetEndTimeCode( double TimeCode )
+	template<typename PtrType>
+	void FUsdStageBase<PtrType>::SetEndTimeCode( double TimeCode )
 	{
 #if USE_USD_SDK
-		if ( !FMath::IsNearlyEqual( TimeCode, Impl->PxrUsdStageRefPtr.Get()->GetEndTimeCode() ) )
+		if ( PtrType& Ptr = Impl->GetInner() )
 		{
-			Impl->PxrUsdStageRefPtr.Get()->SetEndTimeCode( TimeCode );
+			if ( !FMath::IsNearlyEqual( TimeCode, Ptr->GetEndTimeCode() ) )
+			{
+				Ptr->SetEndTimeCode( TimeCode );
+			}
 		}
 #endif // #if USE_USD_SDK
 	}
 
-	double FUsdStage::GetTimeCodesPerSecond() const
+	template<typename PtrType>
+	double FUsdStageBase<PtrType>::GetTimeCodesPerSecond() const
 	{
 #if USE_USD_SDK
-		return Impl->PxrUsdStageRefPtr.Get()->GetTimeCodesPerSecond();
-#else
-		return 24.0;
+		if ( PtrType& Ptr = Impl->GetInner() )
+		{
+			return Ptr->GetTimeCodesPerSecond();
+		}
 #endif // #if USE_USD_SDK
+
+		return 24.0;
 	}
 
-	void FUsdStage::SetTimeCodesPerSecond( double TimeCodesPerSecond )
+	template<typename PtrType>
+	void FUsdStageBase<PtrType>::SetTimeCodesPerSecond( double TimeCodesPerSecond )
 	{
 #if USE_USD_SDK
-		pxr::UsdEditContext EditContext( Impl->PxrUsdStageRefPtr.Get(), Impl->PxrUsdStageRefPtr.Get()->GetRootLayer() );
-		Impl->PxrUsdStageRefPtr.Get()->SetTimeCodesPerSecond( TimeCodesPerSecond );
+		if ( PtrType& Ptr = Impl->GetInner() )
+		{
+			FScopedUsdAllocs Allocs;
+
+			pxr::UsdEditContext EditContext( Ptr, Ptr->GetRootLayer() );
+			Ptr->SetTimeCodesPerSecond( TimeCodesPerSecond );
+		}
 #endif
 	}
 
-	double FUsdStage::GetFramesPerSecond() const
+	template<typename PtrType>
+	double FUsdStageBase<PtrType>::GetFramesPerSecond() const
 	{
 #if USE_USD_SDK
-		return Impl->PxrUsdStageRefPtr.Get()->GetFramesPerSecond();
-#else
-		return 24.0;
+		if ( PtrType& Ptr = Impl->GetInner() )
+		{
+			return Ptr->GetFramesPerSecond();
+		}
 #endif // #if USE_USD_SDK
+
+		return 24.0;
 	}
 
-	void FUsdStage::SetFramesPerSecond( double FramesPerSecond )
+	template<typename PtrType>
+	void FUsdStageBase<PtrType>::SetFramesPerSecond( double FramesPerSecond )
 	{
 #if USE_USD_SDK
-		pxr::UsdEditContext EditContext( Impl->PxrUsdStageRefPtr.Get(), Impl->PxrUsdStageRefPtr.Get()->GetRootLayer() );
-		Impl->PxrUsdStageRefPtr.Get()->SetFramesPerSecond( FramesPerSecond );
+		if ( PtrType& Ptr = Impl->GetInner() )
+		{
+			FScopedUsdAllocs Allocs;
+
+			pxr::UsdEditContext EditContext( Ptr, Ptr->GetRootLayer() );
+			Ptr->SetFramesPerSecond( FramesPerSecond );
+		}
 #endif
 	}
 
-	void FUsdStage::SetDefaultPrim( const FUsdPrim& Prim )
+	template<typename PtrType>
+	void FUsdStageBase<PtrType>::SetDefaultPrim( const FUsdPrim & Prim )
 	{
 #if USE_USD_SDK
-		Impl->PxrUsdStageRefPtr.Get()->SetDefaultPrim( Prim );
+		if ( PtrType& Ptr = Impl->GetInner() )
+		{
+			Ptr->SetDefaultPrim( Prim );
+		}
 #endif // #if USE_USD_SDK
 	}
 
-	FUsdPrim FUsdStage::DefinePrim( const FSdfPath& Path, const TCHAR* TypeName )
+	template<typename PtrType>
+	FUsdPrim FUsdStageBase<PtrType>::DefinePrim( const FSdfPath & Path, const TCHAR * TypeName )
 	{
 #if USE_USD_SDK
-		return FUsdPrim( Impl->PxrUsdStageRefPtr.Get()->DefinePrim( Path, pxr::TfToken( TCHAR_TO_ANSI( TypeName ) ) ) );
-#else
-		return FUsdPrim();
+		if ( PtrType& Ptr = Impl->GetInner() )
+		{
+			return FUsdPrim( Ptr->DefinePrim( Path, pxr::TfToken( TCHAR_TO_ANSI( TypeName ) ) ) );
+		}
 #endif // #if USE_USD_SDK
+
+		return FUsdPrim();
 	}
 
-	bool FUsdStage::RemovePrim( const FSdfPath& Path )
+	template<typename PtrType>
+	bool FUsdStageBase<PtrType>::RemovePrim( const FSdfPath & Path )
 	{
 #if USE_USD_SDK
-		return Impl->PxrUsdStageRefPtr.Get()->RemovePrim( Path );
-#else
-		return false;
+		if ( PtrType& Ptr = Impl->GetInner() )
+		{
+			return Ptr->RemovePrim( Path );
+		}
 #endif // #if USE_USD_SDK
+
+		return false;
 	}
+
+#if USE_USD_SDK
+	template class UNREALUSDWRAPPER_API FUsdStageBase<pxr::UsdStageRefPtr>;
+	template class UNREALUSDWRAPPER_API FUsdStageBase<pxr::UsdStageWeakPtr>;
+#else
+	template class UNREALUSDWRAPPER_API FUsdStageBase<FDummyRefPtrType>;
+	template class UNREALUSDWRAPPER_API FUsdStageBase<FDummyWeakPtrType>;
+#endif // #if USE_USD_SDK
 }
diff --git a/Engine/Plugins/Importers/USDImporter/Source/UnrealUSDWrapper/Public/USDListener.h b/Engine/Plugins/Importers/USDImporter/Source/UnrealUSDWrapper/Public/USDListener.h
--- a/Engine/Plugins/Importers/USDImporter/Source/UnrealUSDWrapper/Public/USDListener.h
+++ b/Engine/Plugins/Importers/USDImporter/Source/UnrealUSDWrapper/Public/USDListener.h
@@ -6,15 +6,11 @@
 #include "Containers/Map.h"
 #include "Delegates/Delegate.h"
 
+#include "UsdWrappers/ForwardDeclarations.h"
 #include "UsdWrappers/VtValue.h"
 
 class FUsdListenerImpl;
 
-namespace UE
-{
-	class FUsdStage;
-}
-
 namespace UsdUtils
 {
 	/** Analogous to pxr::SdfChangeList::Entry::_Flags */
diff --git a/Engine/Plugins/Importers/USDImporter/Source/UnrealUSDWrapper/Public/UnrealUSDWrapper.h b/Engine/Plugins/Importers/USDImporter/Source/UnrealUSDWrapper/Public/UnrealUSDWrapper.h
--- a/Engine/Plugins/Importers/USDImporter/Source/UnrealUSDWrapper/Public/UnrealUSDWrapper.h
+++ b/Engine/Plugins/Importers/USDImporter/Source/UnrealUSDWrapper/Public/UnrealUSDWrapper.h
@@ -9,6 +9,8 @@
 #include "UObject/ObjectMacros.h"
 #include "USDMemory.h"
 
+#include "UsdWrappers/ForwardDeclarations.h"
+
 #include <string>
 #include <vector>
 #include <memory>
@@ -40,7 +42,6 @@ class FUsdDiagnosticDelegate;
 namespace UE
 {
 	class FUsdAttribute;
-	class FUsdStage;
 }
 
 enum class EUsdInterpolationMethod
diff --git a/Engine/Plugins/Importers/USDImporter/Source/UnrealUSDWrapper/Public/UsdWrappers/ForwardDeclarations.h b/Engine/Plugins/Importers/USDImporter/Source/UnrealUSDWrapper/Public/UsdWrappers/ForwardDeclarations.h
new file mode 100644
--- /dev/null
+++ b/Engine/Plugins/Importers/USDImporter/Source/UnrealUSDWrapper/Public/UsdWrappers/ForwardDeclarations.h
@@ -0,0 +1,72 @@
+// Copyright Epic Games, Inc. All Rights Reserved.
+
+#pragma once
+
+#if USE_USD_SDK
+
+#include "USDIncludesStart.h"
+	#include "pxr/pxr.h"
+#include "USDIncludesEnd.h"
+
+PXR_NAMESPACE_OPEN_SCOPE
+class UsdStage;
+class SdfLayer;
+
+template< typename T > class TfRefPtr;
+template< typename T > class TfWeakPtr;
+
+using UsdStageRefPtr = TfRefPtr< UsdStage >;
+using UsdStageWeakPtr = TfWeakPtr< UsdStage >;
+
+using SdfLayerRefPtr = TfRefPtr< SdfLayer >;
+using SdfLayerWeakPtr = TfWeakPtr< SdfLayer >;
+PXR_NAMESPACE_CLOSE_SCOPE
+
+namespace UE
+{
+	template<typename PtrType> class FUsdStageBase;
+	using FUsdStage = FUsdStageBase< pxr::UsdStageRefPtr >;
+	using FUsdStageWeak = FUsdStageBase< pxr::UsdStageWeakPtr >;
+
+	template<typename PtrType> class FSdfLayerBase;
+	using FSdfLayer = FSdfLayerBase< pxr::SdfLayerRefPtr >;
+	using FSdfLayerWeak = FSdfLayerBase< pxr::SdfLayerWeakPtr >;
+}
+
+#else // When the USD SDK is not available
+namespace UE
+{
+	// We use these types so that we can keep the same FUsdStage, etc. type aliases
+	// and not change the rest of the code. It's the bare minimum to be able to
+	// be tested for validity (and always fail).
+	// We use two to prevent multiple symbol definitions on the wrapper copy
+	// constructor/assignment operators.
+	class FDummyWeakPtrType;
+	class FDummyRefPtrType
+	{
+	public:
+		FDummyRefPtrType(){};
+		FDummyRefPtrType( const FDummyWeakPtrType& Other ) {};
+		explicit operator bool() const { return false; }
+		bool operator==( const FDummyRefPtrType& Other ) const { return false; }
+		bool operator!=( const FDummyRefPtrType& Other ) const { return false; }
+	};
+	class FDummyWeakPtrType
+	{
+	public:
+		FDummyWeakPtrType(){};
+		FDummyWeakPtrType( const FDummyRefPtrType& Other ) {};
+		explicit operator bool() const { return false; }
+		bool operator==( const FDummyWeakPtrType& Other ) const { return false; }
+		bool operator!=( const FDummyWeakPtrType& Other ) const { return false; }
+	};
+
+	template<typename PtrType> class FUsdStageBase;
+	using FUsdStage = FUsdStageBase< FDummyRefPtrType >;
+	using FUsdStageWeak = FUsdStageBase< FDummyWeakPtrType >;
+
+	template<typename PtrType> class FSdfLayerBase;
+	using FSdfLayer = FSdfLayerBase< FDummyRefPtrType >;
+	using FSdfLayerWeak = FSdfLayerBase< FDummyWeakPtrType >;
+}
+#endif // #if USE_USD_SDK
\ No newline at end of file
diff --git a/Engine/Plugins/Importers/USDImporter/Source/UnrealUSDWrapper/Public/UsdWrappers/SdfLayer.h b/Engine/Plugins/Importers/USDImporter/Source/UnrealUSDWrapper/Public/UsdWrappers/SdfLayer.h
--- a/Engine/Plugins/Importers/USDImporter/Source/UnrealUSDWrapper/Public/UsdWrappers/SdfLayer.h
+++ b/Engine/Plugins/Importers/USDImporter/Source/UnrealUSDWrapper/Public/UsdWrappers/SdfLayer.h
@@ -4,21 +4,7 @@
 
 #include "CoreMinimal.h"
 
-#if USE_USD_SDK
-
-#include "USDIncludesStart.h"
-	#include "pxr/pxr.h"
-#include "USDIncludesEnd.h"
-
-PXR_NAMESPACE_OPEN_SCOPE
-	class SdfLayer;
-	class SdfLayerOffset;
-
-	template< typename T > class TfRefPtr;
-	using SdfLayerRefPtr = TfRefPtr< SdfLayer >;
-PXR_NAMESPACE_CLOSE_SCOPE
-
-#endif // #if USE_USD_SDK
+#include "UsdWrappers/ForwardDeclarations.h"
 
 namespace UE
 {
@@ -26,7 +12,7 @@ namespace UE
 
 	namespace Internal
 	{
-		class FSdfLayerImpl;
+		template < typename PtrType > class FSdfLayerImpl;
 	}
 
 	struct UNREALUSDWRAPPER_API FSdfLayerOffset
@@ -43,36 +29,53 @@ namespace UE
 	};
 
 	/**
-	 * Minimal pxr::SdfLayerRefPtr wrapper for Unreal that can be used from no-rtti modules.
+	 * Minimal pxr::SdfLayer pointer wrapper for Unreal that can be used from no-rtti modules.
+	 * Use the aliases FSdfLayer and FSdfLayerWeak instead (defined on ForwardDeclarations.h)
 	 */
-	class UNREALUSDWRAPPER_API FSdfLayer
+	template< typename PtrType >
+	class UNREALUSDWRAPPER_API FSdfLayerBase
 	{
 	public:
-		FSdfLayer();
+		FSdfLayerBase();
+
+		FSdfLayerBase( const FSdfLayer& Other );
+		FSdfLayerBase( FSdfLayer&& Other );
+		FSdfLayerBase( const FSdfLayerWeak& Other );
+		FSdfLayerBase( FSdfLayerWeak&& Other );
 
-		FSdfLayer( const FSdfLayer& Other );
-		FSdfLayer( FSdfLayer&& Other );
-		~FSdfLayer();
+		~FSdfLayerBase();
 
-		FSdfLayer& operator=( const FSdfLayer& Other );
-		FSdfLayer& operator=( FSdfLayer&& Other );
+		FSdfLayerBase& operator=( const FSdfLayer& Other );
+		FSdfLayerBase& operator=( FSdfLayer&& Other );
+		FSdfLayerBase& operator=( const FSdfLayerWeak& Other );
+		FSdfLayerBase& operator=( FSdfLayerWeak&& Other );
 
-		bool operator==( const FSdfLayer& Other ) const;
-		bool operator!=( const FSdfLayer& Other ) const;
+		bool operator==( const FSdfLayerBase& Other ) const;
+		bool operator!=( const FSdfLayerBase& Other ) const;
 
 		explicit operator bool() const;
 
-	// Auto conversion from/to pxr::SdfLayerRefPtr
+		// Auto conversion from/to PtrType. We use concrete pointer types here
+		// because we should also be able to convert between them
 	public:
 #if USE_USD_SDK
-		explicit FSdfLayer( const pxr::SdfLayerRefPtr& InSdfLayer );
-		explicit FSdfLayer( pxr::SdfLayerRefPtr&& InSdfLayer );
-
-		FSdfLayer& operator=( const pxr::SdfLayerRefPtr& InSdfLayer );
-		FSdfLayer& operator=( pxr::SdfLayerRefPtr&& InSdfLayer );
-
-		operator pxr::SdfLayerRefPtr&();
-		operator const pxr::SdfLayerRefPtr&() const;
+		explicit FSdfLayerBase( const pxr::SdfLayerRefPtr& InSdfLayer );
+		explicit FSdfLayerBase( pxr::SdfLayerRefPtr&& InSdfLayer );
+		explicit FSdfLayerBase( const pxr::SdfLayerWeakPtr& InSdfLayer );
+		explicit FSdfLayerBase( pxr::SdfLayerWeakPtr&& InSdfLayer );
+
+		FSdfLayerBase& operator=( const pxr::SdfLayerRefPtr& InSdfLayer );
+		FSdfLayerBase& operator=( pxr::SdfLayerRefPtr&& InSdfLayer );
+		FSdfLayerBase& operator=( const pxr::SdfLayerWeakPtr& InSdfLayer );
+		FSdfLayerBase& operator=( pxr::SdfLayerWeakPtr&& InSdfLayer );
+
+		// We can provide reference cast operators for the type we do have, but
+		// need to settle on providing copy cast operators for the other types
+		operator PtrType&();
+		operator const PtrType&() const;
+
+		operator pxr::SdfLayerRefPtr() const;
+		operator pxr::SdfLayerWeakPtr() const;
 #endif // #if USE_USD_SDK
 
 	// Wrapped pxr::SdfLayer functions, refer to the USD SDK documentation
@@ -121,7 +124,10 @@ namespace UE
 		void SetMuted(bool bMuted);
 
 	private:
-		TUniquePtr< Internal::FSdfLayerImpl > Impl;
+		friend FSdfLayer;
+		friend FSdfLayerWeak;
+
+		TUniquePtr< Internal::FSdfLayerImpl<PtrType> > Impl;
 	};
 
 	/**
diff --git a/Engine/Plugins/Importers/USDImporter/Source/UnrealUSDWrapper/Public/UsdWrappers/UsdPrim.h b/Engine/Plugins/Importers/USDImporter/Source/UnrealUSDWrapper/Public/UsdWrappers/UsdPrim.h
--- a/Engine/Plugins/Importers/USDImporter/Source/UnrealUSDWrapper/Public/UsdWrappers/UsdPrim.h
+++ b/Engine/Plugins/Importers/USDImporter/Source/UnrealUSDWrapper/Public/UsdWrappers/UsdPrim.h
@@ -6,6 +6,8 @@
 #include "Templates/UniquePtr.h"
 #include "UObject/NameTypes.h"
 
+#include "UsdWrappers/ForwardDeclarations.h"
+
 #if USE_USD_SDK
 
 #include "USDIncludesStart.h"
@@ -22,7 +24,6 @@ namespace UE
 {
 	class FSdfPath;
 	class FUsdAttribute;
-	class FUsdStage;
 
 	namespace Internal
 	{
diff --git a/Engine/Plugins/Importers/USDImporter/Source/UnrealUSDWrapper/Public/UsdWrappers/UsdStage.h b/Engine/Plugins/Importers/USDImporter/Source/UnrealUSDWrapper/Public/UsdWrappers/UsdStage.h
--- a/Engine/Plugins/Importers/USDImporter/Source/UnrealUSDWrapper/Public/UsdWrappers/UsdStage.h
+++ b/Engine/Plugins/Importers/USDImporter/Source/UnrealUSDWrapper/Public/UsdWrappers/UsdStage.h
@@ -2,69 +2,63 @@
 
 #pragma once
 
-#include "HAL/Platform.h"
 #include "Templates/UniquePtr.h"
 
-#if USE_USD_SDK
-
-#include "USDIncludesStart.h"
-	#include "pxr/pxr.h"
-#include "USDIncludesEnd.h"
-
-PXR_NAMESPACE_OPEN_SCOPE
-	class UsdStage;
-	template< typename T > class TfRefPtr;
-	template< typename T > class TfWeakPtr;
-
-	using UsdStageRefPtr = TfRefPtr< UsdStage >;
-	using UsdStageWeakPtr = TfWeakPtr< UsdStage >;
-PXR_NAMESPACE_CLOSE_SCOPE
-
-#endif // #if USE_USD_SDK
+#include "UsdWrappers/ForwardDeclarations.h"
 
 namespace UE
 {
-	class FSdfLayer;
 	class FSdfPath;
 	class FUsdPrim;
-	class FUsdStage;
 	class FVtValue;
 
 	namespace Internal
 	{
-		class FUsdStageImpl;
+		template< typename PtrType > class FUsdStageImpl;
 	}
 
 	/**
-	 * Minimal pxr::UsdStageRefPtr wrapper for Unreal that can be used from no-rtti modules.
+	 * Minimal pxr::UsdStage pointer wrapper for Unreal that can be used from no-rtti modules.
+	 * Use the aliases FUsdStage and FUsdStageWeak instead (defined on ForwardDeclarations.h)
 	 */
-	class UNREALUSDWRAPPER_API FUsdStage
+	template< typename PtrType >
+	class UNREALUSDWRAPPER_API FUsdStageBase
 	{
 	public:
-		FUsdStage();
-		FUsdStage( const FUsdStage& Other );
-		FUsdStage( FUsdStage&& Other );
+		FUsdStageBase();
+
+		FUsdStageBase( const FUsdStage& Other );
+		FUsdStageBase( FUsdStage&& Other );
+		FUsdStageBase( const FUsdStageWeak& Other );
+		FUsdStageBase( FUsdStageWeak&& Other );
 
-		FUsdStage& operator=( const FUsdStage& Other );
-		FUsdStage& operator=( FUsdStage&& Other );
+		FUsdStageBase& operator=( const FUsdStage& Other );
+		FUsdStageBase& operator=( FUsdStage&& Other );
+		FUsdStageBase& operator=( const FUsdStageWeak& Other );
+		FUsdStageBase& operator=( FUsdStageWeak&& Other );
 
-		~FUsdStage();
+		~FUsdStageBase();
 
 		explicit operator bool() const;
 
-		bool operator==( const FUsdStage& Other ) const;
-		bool operator!=( const FUsdStage& Other ) const;
+		bool operator==( const FUsdStageBase& Other ) const;
+		bool operator!=( const FUsdStageBase& Other ) const;
 
-	// Auto conversion from/to pxr::UsdStageRefPtr
+		// Auto conversion from/to PtrType. We use concrete pointer types here
+		// because we should also be able to convert between them
 	public:
 #if USE_USD_SDK
-		explicit FUsdStage( const pxr::UsdStageRefPtr& InUsdStageRefPtr );
-		explicit FUsdStage( pxr::UsdStageRefPtr&& InUsdStageRefPtr );
+		explicit FUsdStageBase( const pxr::UsdStageRefPtr& InUsdPtr );
+		explicit FUsdStageBase( pxr::UsdStageRefPtr&& InUsdPtr );
+		explicit FUsdStageBase( const pxr::UsdStageWeakPtr& InUsdPtr );
+		explicit FUsdStageBase( pxr::UsdStageWeakPtr&& InUsdPtr );
 
-		operator pxr::UsdStageRefPtr&();
-		operator const pxr::UsdStageRefPtr&() const;
-		operator pxr::UsdStageWeakPtr() const;
-#endif // #if USE_USD_SDK
+		operator PtrType&();
+		operator const PtrType&() const;
+
+		operator pxr::UsdStageRefPtr ( ) const;
+		operator pxr::UsdStageWeakPtr ( ) const;
+#endif // USE_USD_SDK
 
 	// Wrapped pxr::UsdStage functions, refer to the USD SDK documentation
 	public:
@@ -99,6 +93,10 @@ namespace UE
 		bool RemovePrim( const FSdfPath& Path );
 
 	private:
-		TUniquePtr< Internal::FUsdStageImpl > Impl;
+		// So we can use the Other's Impl on copy constructor/operators
+		friend FUsdStage;
+		friend FUsdStageWeak;
+
+		TUniquePtr< Internal::FUsdStageImpl<PtrType> > Impl;
 	};
 }
\ No newline at end of file
