USD: Allow picking the USD Stage Editor's bound AUsdStageActor

diff --git a/Engine/Plugins/Importers/USDImporter/Source/USDStageEditor/Private/SUSDStage.cpp b/Engine/Plugins/Importers/USDImporter/Source/USDStageEditor/Private/SUSDStage.cpp
--- a/Engine/Plugins/Importers/USDImporter/Source/USDStageEditor/Private/SUSDStage.cpp
+++ b/Engine/Plugins/Importers/USDImporter/Source/USDStageEditor/Private/SUSDStage.cpp
@@ -22,16 +22,18 @@
 #include "UsdWrappers/SdfLayer.h"
 #include "UsdWrappers/UsdStage.h"
 
+#include "ActorTreeItem.h"
 #include "Dialogs/DlgPickPath.h"
 #include "EditorStyleSet.h"
+#include "Engine/Selection.h"
 #include "Engine/World.h"
 #include "Framework/MultiBox/MultiBoxBuilder.h"
 #include "LevelEditor.h"
 #include "Modules/ModuleManager.h"
+#include "SceneOutlinerModule.h"
 #include "ScopedTransaction.h"
 #include "UObject/StrongObjectPtr.h"
 #include "Widgets/Layout/SSplitter.h"
-#include "Engine/Selection.h"
 
 #define LOCTEXT_NAMESPACE "SUsdStage"
 
@@ -122,7 +124,23 @@ void SUsdStage::Construct( const FArguments& InArgs )
 			+ SVerticalBox::Slot()
 			.AutoHeight()
 			[
-				MakeMainMenu()
+				SNew( SHorizontalBox )
+				+ SHorizontalBox::Slot()
+				.FillWidth( 1 )
+				[
+					MakeMainMenu()
+				]
+
+				+ SHorizontalBox::Slot()
+				.AutoWidth()
+				[
+					SNew( SBox )
+					.HAlign( HAlign_Right )
+					.VAlign( VAlign_Fill )
+					[
+						MakeActorPickerMenu()
+					]
+				]
 			]
 
 			+SVerticalBox::Slot()
@@ -299,6 +317,51 @@ TSharedRef< SWidget > SUsdStage::MakeMainMenu()
 	return MenuBarWidget;
 }
 
+TSharedRef< SWidget > SUsdStage::MakeActorPickerMenu()
+{
+	return SNew( SComboButton )
+		.ComboButtonStyle( FAppStyle::Get(), "SimpleComboButton" )
+		.OnGetMenuContent( this, &SUsdStage::MakeActorPickerMenuContent )
+		.MenuPlacement( EMenuPlacement::MenuPlacement_BelowRightAnchor )
+		.ToolTipText( LOCTEXT( "ActorPicker_ToolTip", "Switch the active stage actor" ) )
+		.ButtonContent()
+		[
+			SNew( STextBlock )
+			.Text_Lambda( [this]()
+			{
+				if ( AUsdStageActor* StageActor = ViewModel.UsdStageActor.Get() )
+				{
+					return FText::FromString( StageActor->GetActorLabel() );
+				}
+
+				return FText::FromString( "None" );
+			})
+		];
+}
+
+TSharedRef< SWidget > SUsdStage::MakeActorPickerMenuContent()
+{
+	FSceneOutlinerInitializationOptions InitOptions;
+	InitOptions.bShowHeaderRow = false;
+	InitOptions.bShowSearchBox = true;
+	InitOptions.bShowCreateNewFolder = false;
+	InitOptions.bFocusSearchBoxWhenOpened = true;
+	InitOptions.ColumnMap.Add(FSceneOutlinerBuiltInColumnTypes::Label(), FSceneOutlinerColumnInfo(ESceneOutlinerColumnVisibility::Visible, 0));
+	InitOptions.Filters->AddFilterPredicate<FActorTreeItem>(FActorTreeItem::FFilterPredicate::CreateLambda([]( const AActor* Actor )
+	{
+		return Actor && Actor->IsA<AUsdStageActor>();
+	}));
+
+	FSceneOutlinerModule& SceneOutlinerModule = FModuleManager::LoadModuleChecked<FSceneOutlinerModule>("SceneOutliner");
+	return SceneOutlinerModule.CreateActorPicker(InitOptions, FOnActorPicked::CreateLambda([this](AActor* Actor)
+	{
+		if ( Actor && Actor->IsA<AUsdStageActor>() )
+		{
+			this->SetActor( Cast<AUsdStageActor>( Actor ) );
+		}
+	}));
+}
+
 void SUsdStage::FillFileMenu( FMenuBuilder& MenuBuilder )
 {
 	MenuBuilder.BeginSection( "File", LOCTEXT("File", "File") );
@@ -759,9 +822,7 @@ void SUsdStage::OpenStage( const TCHAR* FilePath )
 	if ( !ViewModel.UsdStageActor.IsValid() )
 	{
 		IUsdStageModule& UsdStageModule = FModuleManager::Get().LoadModuleChecked< IUsdStageModule >( "UsdStage" );
-		ViewModel.UsdStageActor = &UsdStageModule.GetUsdStageActor( GWorld );
-
-		SetupStageActorDelegates();
+		SetActor( &UsdStageModule.GetUsdStageActor( GWorld ) );
 	}
 
 	// Block writing level sequence changes back to the USD stage until we finished this transaction, because once we do
@@ -772,6 +833,21 @@ void SUsdStage::OpenStage( const TCHAR* FilePath )
 	ViewModel.OpenStage( FilePath );
 }
 
+void SUsdStage::SetActor( AUsdStageActor* InUsdStageActor )
+{
+	ViewModel.UsdStageActor = InUsdStageActor;
+
+	SetupStageActorDelegates();
+
+	if ( this->UsdPrimInfoWidget && InUsdStageActor )
+	{
+		// Just reset to the pseudoroot for now
+		this->UsdPrimInfoWidget->SetPrimPath( ViewModel.UsdStageActor->GetOrLoadUsdStage(), TEXT( "/" ) );
+	}
+
+	Refresh();
+}
+
 void SUsdStage::Refresh()
 {
 	// May be nullptr, but that is ok. Its how the widgets are reset
diff --git a/Engine/Plugins/Importers/USDImporter/Source/USDStageEditor/Private/SUSDStage.h b/Engine/Plugins/Importers/USDImporter/Source/USDStageEditor/Private/SUSDStage.h
--- a/Engine/Plugins/Importers/USDImporter/Source/USDStageEditor/Private/SUSDStage.h
+++ b/Engine/Plugins/Importers/USDImporter/Source/USDStageEditor/Private/SUSDStage.h
@@ -36,6 +36,8 @@ protected:
 	void ClearStageActorDelegates();
 
 	TSharedRef< SWidget > MakeMainMenu();
+	TSharedRef< SWidget > MakeActorPickerMenu();
+	TSharedRef< SWidget > MakeActorPickerMenuContent();
 	void FillFileMenu( FMenuBuilder& MenuBuilder );
 	void FillActionsMenu( FMenuBuilder& MenuBuilder );
 	void FillOptionsMenu( FMenuBuilder& MenuBuilder );
@@ -56,6 +58,8 @@ protected:
 
 	void OpenStage( const TCHAR* FilePath );
 
+	void SetActor( AUsdStageActor* InUsdStageActor );
+
 	void Refresh();
 
 	void OnStageActorLoaded( AUsdStageActor* InUsdStageActor );
diff --git a/Engine/Plugins/Importers/USDImporter/Source/USDStageEditor/USDStageEditor.Build.cs b/Engine/Plugins/Importers/USDImporter/Source/USDStageEditor/USDStageEditor.Build.cs
--- a/Engine/Plugins/Importers/USDImporter/Source/USDStageEditor/USDStageEditor.Build.cs
+++ b/Engine/Plugins/Importers/USDImporter/Source/USDStageEditor/USDStageEditor.Build.cs
@@ -21,14 +21,15 @@ namespace UnrealBuildTool.Rules
 					"InputCore",
 					"LevelEditor",
 					"Projects", // So that we can use the IPluginManager, required for our custom style
+					"SceneOutliner",
 					"Slate",
 					"SlateCore",
 					"UnrealEd",
 					"UnrealUSDWrapper",
 					"USDSchemas",
-					"USDStageImporter",
 					"USDStage",
 					"USDStageEditorViewModels",
+					"USDStageImporter",
 					"USDUtilities",
 					"WorkspaceMenuStructure",
 				}
