USD: Less agressive way of deselecting actors and components before deleting them.

diff --git a/Engine/Plugins/Importers/USDImporter/Source/USDStage/Private/USDStageActor.cpp b/Engine/Plugins/Importers/USDImporter/Source/USDStage/Private/USDStageActor.cpp
--- a/Engine/Plugins/Importers/USDImporter/Source/USDStage/Private/USDStageActor.cpp
+++ b/Engine/Plugins/Importers/USDImporter/Source/USDStage/Private/USDStageActor.cpp
@@ -62,6 +62,7 @@
 #include "Editor.h"
 #include "Editor/TransBuffer.h"
 #include "Editor/UnrealEdEngine.h"
+#include "Engine/Selection.h"
 #include "LevelEditor.h"
 #include "PropertyEditorModule.h"
 #include "Subsystems/AssetEditorSubsystem.h"
@@ -127,27 +128,51 @@ struct FUsdStageActorImpl
 		// This can get called when an actor is being destroyed due to GC.
 		// Don't do this during garbage collecting if we need to delay-create the root twin (can't NewObject during garbage collection).
 		// If we have no root twin we don't have any tracked spawned actors and components, so we don't need to deselect anything in the first place
+		bool bDeselected = false;
 		if ( !IsGarbageCollecting() || StageActor->RootUsdTwin )
 		{
-			TArray<UObject*> ObjectsToDelete;
+			TArray<UObject*> ActorsToDeselect;
+			TArray<UObject*> ComponentsToDeselect;
+
 			const bool bRecursive = true;
-			StageActor->GetRootPrimTwin()->Iterate( [ &ObjectsToDelete ]( UUsdPrimTwin& PrimTwin )
+			StageActor->GetRootPrimTwin()->Iterate( [&ActorsToDeselect, &ComponentsToDeselect]( UUsdPrimTwin& PrimTwin )
 			{
 				if ( AActor* ReferencedActor = PrimTwin.SpawnedActor.Get() )
 				{
-					ObjectsToDelete.Add( ReferencedActor );
+					ActorsToDeselect.Add( ReferencedActor );
 				}
 				if ( USceneComponent* ReferencedComponent = PrimTwin.SceneComponent.Get() )
 				{
-					ObjectsToDelete.Add( ReferencedComponent );
+					ComponentsToDeselect.Add( ReferencedComponent );
 				}
 			}, bRecursive );
 
-			FPropertyEditorModule& PropertyEditorModule = FModuleManager::GetModuleChecked<FPropertyEditorModule>( "PropertyEditor" );
-			PropertyEditorModule.RemoveDeletedObjects( ObjectsToDelete );
+			if ( USelection* SelectedComponents = GEditor->GetSelectedComponents() )
+			{
+				for ( UObject* Component : ComponentsToDeselect )
+				{
+					if ( SelectedComponents->IsSelected( Component ) )
+					{
+						SelectedComponents->Deselect( Component );
+						bDeselected = true;
+					}
+				}
+			}
+
+			if ( USelection* SelectedActors = GEditor->GetSelectedActors() )
+			{
+				for ( UObject* Actor : ActorsToDeselect )
+				{
+					if ( SelectedActors->IsSelected( Actor ) )
+					{
+						SelectedActors->Deselect( Actor );
+						bDeselected = true;
+					}
+				}
+			}
 		}
 
-		if ( GIsEditor && GEditor ) // Make sure we're not in standalone either
+		if ( bDeselected && GIsEditor && GEditor ) // Make sure we're not in standalone either
 		{
 			GEditor->NoteSelectionChange();
 		}
