USD: Fix static analysis warnings

diff --git a/Engine/Plugins/Importers/USDImporter/Source/USDStage/Private/USDStageActor.cpp b/Engine/Plugins/Importers/USDImporter/Source/USDStage/Private/USDStageActor.cpp
--- a/Engine/Plugins/Importers/USDImporter/Source/USDStage/Private/USDStageActor.cpp
+++ b/Engine/Plugins/Importers/USDImporter/Source/USDStage/Private/USDStageActor.cpp
@@ -120,7 +120,7 @@ struct FUsdStageActorImpl
 	static void DeselectActorsAndComponents( AUsdStageActor* StageActor )
 	{
 #if WITH_EDITOR
-		if ( !StageActor )
+		if ( !StageActor || !GEditor )
 		{
 			return;
 		}
@@ -172,7 +172,7 @@ struct FUsdStageActorImpl
 			}
 		}
 
-		if ( bDeselected && GIsEditor && GEditor ) // Make sure we're not in standalone either
+		if ( GEditor && bDeselected && GIsEditor ) // Make sure we're not in standalone either
 		{
 			GEditor->NoteSelectionChange();
 		}
