USD: Another attempt at fixing incorrect static analysis warning

I think it was getting confused with the check at 175 and usage at 162, so I flipped them

diff --git a/Engine/Plugins/Importers/USDImporter/Source/USDStage/Private/USDStageActor.cpp b/Engine/Plugins/Importers/USDImporter/Source/USDStage/Private/USDStageActor.cpp
--- a/Engine/Plugins/Importers/USDImporter/Source/USDStage/Private/USDStageActor.cpp
+++ b/Engine/Plugins/Importers/USDImporter/Source/USDStage/Private/USDStageActor.cpp
@@ -120,7 +120,7 @@ struct FUsdStageActorImpl
 	static void DeselectActorsAndComponents( AUsdStageActor* StageActor )
 	{
 #if WITH_EDITOR
-		if ( !StageActor || !GEditor )
+		if ( !StageActor )
 		{
 			return;
 		}
@@ -129,7 +129,7 @@ struct FUsdStageActorImpl
 		// Don't do this during garbage collecting if we need to delay-create the root twin (can't NewObject during garbage collection).
 		// If we have no root twin we don't have any tracked spawned actors and components, so we don't need to deselect anything in the first place
 		bool bDeselected = false;
-		if ( !IsGarbageCollecting() || StageActor->RootUsdTwin )
+		if ( GEditor && !IsGarbageCollecting() || StageActor->RootUsdTwin )
 		{
 			TArray<UObject*> ActorsToDeselect;
 			TArray<UObject*> ComponentsToDeselect;
@@ -170,11 +170,11 @@ struct FUsdStageActorImpl
 					}
 				}
 			}
-		}
 
-		if ( GEditor && bDeselected && GIsEditor ) // Make sure we're not in standalone either
-		{
-			GEditor->NoteSelectionChange();
+			if ( bDeselected && GIsEditor ) // Make sure we're not in standalone either
+			{
+				GEditor->NoteSelectionChange();
+			}
 		}
 #endif // WITH_EDITOR
 	}
