USD: Fixup for linux compile error on cl 17099511

diff --git a/Engine/Plugins/Importers/USDImporter/Source/USDStage/Private/USDStageActor.cpp b/Engine/Plugins/Importers/USDImporter/Source/USDStage/Private/USDStageActor.cpp
--- a/Engine/Plugins/Importers/USDImporter/Source/USDStage/Private/USDStageActor.cpp
+++ b/Engine/Plugins/Importers/USDImporter/Source/USDStage/Private/USDStageActor.cpp
@@ -129,7 +129,7 @@ struct FUsdStageActorImpl
 		// Don't do this during garbage collecting if we need to delay-create the root twin (can't NewObject during garbage collection).
 		// If we have no root twin we don't have any tracked spawned actors and components, so we don't need to deselect anything in the first place
 		bool bDeselected = false;
-		if ( GEditor && !IsGarbageCollecting() || StageActor->RootUsdTwin )
+		if ( GEditor && !IsGarbageCollecting() && StageActor->RootUsdTwin )
 		{
 			TArray<UObject*> ActorsToDeselect;
 			TArray<UObject*> ComponentsToDeselect;
