USD: Fixes for UDIM textures
- Added VT variants of the UsdPreviewSurface master materials where the texture samplers are using the virtual types.
- Fixed an issue where the file path hash was wrong for UDIM textures when hashing USD materials which resulted in many materials having the same hash. Changed the path resolving function to the UsdUtils one, which will resolve UDIM paths properly.

diff --git a/Engine/Plugins/Importers/USDImporter/Source/USDSchemas/Private/USDShadeMaterialTranslator.cpp b/Engine/Plugins/Importers/USDImporter/Source/USDSchemas/Private/USDShadeMaterialTranslator.cpp
--- a/Engine/Plugins/Importers/USDImporter/Source/USDSchemas/Private/USDShadeMaterialTranslator.cpp
+++ b/Engine/Plugins/Importers/USDImporter/Source/USDSchemas/Private/USDShadeMaterialTranslator.cpp
@@ -39,11 +39,24 @@ void FUsdShadeMaterialTranslator::CreateAssets()
 
 	if ( !ConvertedMaterial )
 	{
-		FString MaterialPath = UsdUtils::IsMaterialTranslucent( ShadeMaterial )
-			? TEXT( "Material'/USDImporter/Materials/UsdPreviewSurfaceTranslucent.UsdPreviewSurfaceTranslucent'" )
-			: TEXT( "Material'/USDImporter/Materials/UsdPreviewSurface.UsdPreviewSurface'" );
+		const bool bIsTranslucent = UsdUtils::IsMaterialTranslucent( ShadeMaterial );
+		const bool bNeedsVirtualTextures = UsdUtils::IsMaterialUsingUDIMs( ShadeMaterial );
 
-		if ( UMaterialInterface* MasterMaterial = Cast< UMaterialInterface >( FSoftObjectPath( MaterialPath ).TryLoad() ) )
+		FString MasterMaterialName = TEXT("UsdPreviewSurface");
+
+		if ( bIsTranslucent )
+		{
+			MasterMaterialName += TEXT("Translucent");
+		}
+
+		if ( bNeedsVirtualTextures )
+		{
+			MasterMaterialName += TEXT("VT");
+		}
+
+		const FString MasterMaterialPath = FString::Printf( TEXT("Material'/USDImporter/Materials/%s.%s"), *MasterMaterialName, *MasterMaterialName );
+
+		if ( UMaterialInterface* MasterMaterial = Cast< UMaterialInterface >( FSoftObjectPath( MasterMaterialPath ).TryLoad() ) )
 		{
 			FName InstanceName = MakeUniqueObjectName( GetTransientPackage(), UMaterialInstance::StaticClass(), *FPaths::GetBaseFilename( PrimPath.GetString() ) );
 			if ( GIsEditor ) // Also have to prevent Standalone game from going with MaterialInstanceConstants
diff --git a/Engine/Plugins/Importers/USDImporter/Source/USDUtilities/Private/USDShadeConversion.cpp b/Engine/Plugins/Importers/USDImporter/Source/USDUtilities/Private/USDShadeConversion.cpp
--- a/Engine/Plugins/Importers/USDImporter/Source/USDUtilities/Private/USDShadeConversion.cpp
+++ b/Engine/Plugins/Importers/USDImporter/Source/USDUtilities/Private/USDShadeConversion.cpp
@@ -129,7 +129,7 @@ namespace UE
 					? UsdToUnreal::ConvertString( pxr::SdfComputeAssetPathRelativeToLayer( LayerHandle, UnrealToUsd::ConvertString( *AssetPathToResolve ).Get() ) )
 					: AssetPathToResolve;
 
-				std::string ResolvedPathUsd = Resolver.Resolve( UnrealToUsd::ConvertString( *RelativePathToResolve ).Get().c_str() );
+				std::string ResolvedPathUsd = Resolver.ComputeNormalizedPath( Resolver.Resolve( UnrealToUsd::ConvertString( *RelativePathToResolve ).Get().c_str() ) );
 
 				FString ResolvedAssetPath = UsdToUnreal::ConvertString( ResolvedPathUsd );
 
@@ -459,7 +459,8 @@ namespace UE
 			};
 			using FParameterValue = TVariant<float, FVector, FTextureParameterValue, bool>;
 
-			bool GetTextureParameterValue( pxr::UsdShadeInput& ShadeInput, TextureGroup LODGroup, FParameterValue& OutValue, UMaterialInterface* Material = nullptr, UUsdAssetCache* TexturesCache = nullptr, TMap<FString, int32>* PrimvarToUVIndex = nullptr )
+			bool GetTextureParameterValue( pxr::UsdShadeInput& ShadeInput, TextureGroup LODGroup, FParameterValue& OutValue,
+				UMaterialInterface* Material = nullptr, UUsdAssetCache* TexturesCache = nullptr, TMap<FString, int32>* PrimvarToUVIndex = nullptr, bool bForceVirtualTextures = false )
 			{
 				FScopedUsdAllocs UsdAllocs;
 
@@ -531,7 +532,7 @@ namespace UE
 #if WITH_EDITOR
 								if ( Material )
 								{
-									if ( UUsdAssetImportData* ImportData = Cast<UUsdAssetImportData>( Material->AssetImportData ) )
+									if ( UUsdAssetImportData* ImportData = Cast< UUsdAssetImportData >( Material->AssetImportData ) )
 									{
 										MaterialPrimPath = ImportData->PrimPath;
 									}
@@ -545,6 +546,12 @@ namespace UE
 							{
 								Texture->SRGB = bSRGB;
 								Texture->CompressionSettings = CompressionSettings;
+
+								if ( bForceVirtualTextures )
+								{
+									Texture->VirtualTextureStreaming = true;
+								}
+
 								Texture->UpdateResource();
 
 								FString PrimvarName = GetPrimvarUsedAsST(Source);
@@ -619,7 +626,8 @@ namespace UE
 				return OutValue.IsType<FTextureParameterValue>();
 			}
 
-			bool GetFloatParameterValue( pxr::UsdShadeConnectableAPI& Connectable, const pxr::TfToken& InputName, float DefaultValue, FParameterValue& OutValue, UMaterialInterface* Material = nullptr, UUsdAssetCache* TexturesCache = nullptr, TMap<FString, int32>* PrimvarToUVIndex = nullptr)
+			bool GetFloatParameterValue( pxr::UsdShadeConnectableAPI& Connectable, const pxr::TfToken& InputName, float DefaultValue, FParameterValue& OutValue,
+				UMaterialInterface* Material = nullptr, UUsdAssetCache* TexturesCache = nullptr, TMap<FString, int32>* PrimvarToUVIndex = nullptr, bool bForceVirtualTextures = false)
 			{
 				FScopedUsdAllocs Allocs;
 
@@ -635,11 +643,11 @@ namespace UE
 				pxr::UsdShadeAttributeType AttributeType;
 				if ( Input.GetConnectedSource( &Source, &SourceName, &AttributeType ) )
 				{
-					if ( !GetTextureParameterValue( Input, TEXTUREGROUP_WorldSpecular, OutValue, Material, TexturesCache, PrimvarToUVIndex ) )
+					if ( !GetTextureParameterValue( Input, TEXTUREGROUP_WorldSpecular, OutValue, Material, TexturesCache, PrimvarToUVIndex, bForceVirtualTextures ) )
 					{
 						// Recurse because the attribute may just be pointing at some other attribute that has the data
 						// (e.g. when shader input is just "hoisted" and connected to the parent material input)
-						return GetFloatParameterValue( Source, SourceName, DefaultValue, OutValue, Material, TexturesCache, PrimvarToUVIndex );
+						return GetFloatParameterValue( Source, SourceName, DefaultValue, OutValue, Material, TexturesCache, PrimvarToUVIndex, bForceVirtualTextures );
 					}
 				}
 				// No other node connected, so we must have some value
@@ -654,7 +662,8 @@ namespace UE
 				return true;
 			}
 
-			bool GetVec3ParameterValue( pxr::UsdShadeConnectableAPI& Connectable, const pxr::TfToken& InputName, const FLinearColor& DefaultValue, FParameterValue& OutValue, bool bIsNormalMap = false, UMaterialInterface* Material = nullptr, UUsdAssetCache* TexturesCache = nullptr, TMap<FString, int32>* PrimvarToUVIndex = nullptr )
+			bool GetVec3ParameterValue( pxr::UsdShadeConnectableAPI& Connectable, const pxr::TfToken& InputName, const FLinearColor& DefaultValue, FParameterValue& OutValue,
+				bool bIsNormalMap = false, UMaterialInterface* Material = nullptr, UUsdAssetCache* TexturesCache = nullptr, TMap<FString, int32>* PrimvarToUVIndex = nullptr, bool bForceVirtualTextures = false )
 			{
 				FScopedUsdAllocs Allocs;
 
@@ -670,9 +679,9 @@ namespace UE
 				pxr::UsdShadeAttributeType AttributeType;
 				if ( Input.GetConnectedSource( &Source, &SourceName, &AttributeType ) )
 				{
-					if ( !GetTextureParameterValue( Input, bIsNormalMap ? TEXTUREGROUP_WorldNormalMap : TEXTUREGROUP_World, OutValue, Material, TexturesCache, PrimvarToUVIndex ) )
+					if ( !GetTextureParameterValue( Input, bIsNormalMap ? TEXTUREGROUP_WorldNormalMap : TEXTUREGROUP_World, OutValue, Material, TexturesCache, PrimvarToUVIndex, bForceVirtualTextures ) )
 					{
-						return GetVec3ParameterValue( Source, SourceName, DefaultValue, OutValue, bIsNormalMap, Material, TexturesCache, PrimvarToUVIndex );
+						return GetVec3ParameterValue( Source, SourceName, DefaultValue, OutValue, bIsNormalMap, Material, TexturesCache, PrimvarToUVIndex, bForceVirtualTextures );
 					}
 				}
 				// No other node connected, so we must have some value
@@ -1542,9 +1551,7 @@ namespace UE
 					// as it happens with a pure USD python script. (this with USD 21.05 in June 2021)
 					if ( ShadeInputValue.IsHolding<pxr::SdfAssetPath>() )
 					{
-						pxr::ArResolver& Resolver = pxr::ArGetResolver();
-						pxr::SdfAssetPath UsdAssetPath = ShadeInputValue.UncheckedGet<pxr::SdfAssetPath>();
-						FString ResolvedPath = UsdToUnreal::ConvertString( Resolver.ComputeNormalizedPath( UsdAssetPath.GetResolvedPath() ) );
+						FString ResolvedPath = UsdUtils::GetResolvedTexturePath( ShadeInput.GetAttr() );
 						InOutHashState.UpdateWithString( *ResolvedPath, ResolvedPath.Len() );
 					}
 					else
@@ -1571,6 +1578,10 @@ bool UsdToUnreal::ConvertMaterial( const pxr::UsdShadeMaterial& UsdShadeMaterial
 
 	bool bHasMaterialInfo = false;
 
+	// We assume that all samplers for a material instance that supports UDIMs would be virtual so we'll force all the textures to be virtual.
+	// The idea being that when building a master material, we can't know which textures will be virtual so it's safer to set the samplers as if they were all virtual.
+	const bool bForceVirtualTextures = UsdUtils::IsMaterialUsingUDIMs( UsdShadeMaterial );
+
 	pxr::UsdShadeShader SurfaceShader = UsdShadeMaterial.ComputeSurfaceSource();
 	if ( !SurfaceShader )
 	{
@@ -1611,7 +1622,7 @@ bool UsdToUnreal::ConvertMaterial( const pxr::UsdShadeMaterial& UsdShadeMaterial
 	// Base color
 	{
 		const bool bIsNormalMap = false;
-		if ( UsdShadeConversionImpl::GetVec3ParameterValue( Connectable, UnrealIdentifiers::DiffuseColor, FLinearColor( 0, 0, 0 ), ParameterValue, bIsNormalMap, &Material, TexturesCache, &PrimvarToUVIndex ) )
+		if ( UsdShadeConversionImpl::GetVec3ParameterValue( Connectable, UnrealIdentifiers::DiffuseColor, FLinearColor( 0, 0, 0 ), ParameterValue, bIsNormalMap, &Material, TexturesCache, &PrimvarToUVIndex, bForceVirtualTextures ) )
 		{
 			UsdShadeConversionImpl::SetParameterValue( Material, TEXT( "BaseColor" ), ParameterValue, bForUsdPreviewSurface );
 			bHasMaterialInfo = true;
@@ -1621,7 +1632,7 @@ bool UsdToUnreal::ConvertMaterial( const pxr::UsdShadeMaterial& UsdShadeMaterial
 	// Emissive color
 	{
 		const bool bIsNormalMap = false;
-		if ( UsdShadeConversionImpl::GetVec3ParameterValue( Connectable, UnrealIdentifiers::EmissiveColor, FLinearColor( 0, 0, 0 ), ParameterValue, bIsNormalMap, &Material, TexturesCache, &PrimvarToUVIndex ) )
+		if ( UsdShadeConversionImpl::GetVec3ParameterValue( Connectable, UnrealIdentifiers::EmissiveColor, FLinearColor( 0, 0, 0 ), ParameterValue, bIsNormalMap, &Material, TexturesCache, &PrimvarToUVIndex, bForceVirtualTextures ) )
 		{
 			UsdShadeConversionImpl::SetParameterValue( Material, TEXT( "EmissiveColor" ), ParameterValue, bForUsdPreviewSurface );
 			bHasMaterialInfo = true;
@@ -1630,7 +1641,7 @@ bool UsdToUnreal::ConvertMaterial( const pxr::UsdShadeMaterial& UsdShadeMaterial
 
 	// Metallic
 	{
-		if ( UsdShadeConversionImpl::GetFloatParameterValue( Connectable, UnrealIdentifiers::Metallic, 0.f, ParameterValue, &Material, TexturesCache, &PrimvarToUVIndex ) )
+		if ( UsdShadeConversionImpl::GetFloatParameterValue( Connectable, UnrealIdentifiers::Metallic, 0.f, ParameterValue, &Material, TexturesCache, &PrimvarToUVIndex, bForceVirtualTextures ) )
 		{
 			UsdShadeConversionImpl::SetParameterValue( Material, TEXT( "Metallic" ), ParameterValue, bForUsdPreviewSurface );
 			SetTextureComponentParamForScalarInput( ParameterValue, TEXT( "MetallicTextureComponent" ) );
@@ -1640,7 +1651,7 @@ bool UsdToUnreal::ConvertMaterial( const pxr::UsdShadeMaterial& UsdShadeMaterial
 
 	// Roughness
 	{
-		if ( UsdShadeConversionImpl::GetFloatParameterValue( Connectable, UnrealIdentifiers::Roughness, 1.f, ParameterValue, &Material, TexturesCache, &PrimvarToUVIndex ) )
+		if ( UsdShadeConversionImpl::GetFloatParameterValue( Connectable, UnrealIdentifiers::Roughness, 1.f, ParameterValue, &Material, TexturesCache, &PrimvarToUVIndex, bForceVirtualTextures ) )
 		{
 			UsdShadeConversionImpl::SetParameterValue( Material, TEXT( "Roughness" ), ParameterValue, bForUsdPreviewSurface );
 			SetTextureComponentParamForScalarInput( ParameterValue, TEXT( "RoughnessTextureComponent" ) );
@@ -1650,7 +1661,7 @@ bool UsdToUnreal::ConvertMaterial( const pxr::UsdShadeMaterial& UsdShadeMaterial
 
 	// Opacity
 	{
-		if ( UsdShadeConversionImpl::GetFloatParameterValue( Connectable, UnrealIdentifiers::Opacity, 1.f, ParameterValue, &Material, TexturesCache, &PrimvarToUVIndex ) )
+		if ( UsdShadeConversionImpl::GetFloatParameterValue( Connectable, UnrealIdentifiers::Opacity, 1.f, ParameterValue, &Material, TexturesCache, &PrimvarToUVIndex, bForceVirtualTextures ) )
 		{
 			UsdShadeConversionImpl::SetParameterValue( Material, TEXT( "Opacity" ), ParameterValue, bForUsdPreviewSurface );
 			SetTextureComponentParamForScalarInput( ParameterValue, TEXT( "OpacityTextureComponent" ) );
@@ -1661,7 +1672,7 @@ bool UsdToUnreal::ConvertMaterial( const pxr::UsdShadeMaterial& UsdShadeMaterial
 	// Normal
 	{
 		const bool bIsNormalMap = true;
-		if ( UsdShadeConversionImpl::GetVec3ParameterValue( Connectable, UnrealIdentifiers::Normal, FLinearColor( 0, 0, 1 ), ParameterValue, bIsNormalMap, &Material, TexturesCache, &PrimvarToUVIndex ) )
+		if ( UsdShadeConversionImpl::GetVec3ParameterValue( Connectable, UnrealIdentifiers::Normal, FLinearColor( 0, 0, 1 ), ParameterValue, bIsNormalMap, &Material, TexturesCache, &PrimvarToUVIndex, bForceVirtualTextures ) )
 		{
 			UsdShadeConversionImpl::SetParameterValue( Material, TEXT( "Normal" ), ParameterValue, bForUsdPreviewSurface );
 			bHasMaterialInfo = true;
@@ -1946,7 +1957,9 @@ FString UsdUtils::GetResolvedTexturePath( const pxr::UsdAttribute& TextureAssetP
 	pxr::SdfAssetPath TextureAssetPath;
 	TextureAssetPathAttr.Get< pxr::SdfAssetPath >( &TextureAssetPath );
 
-	FString ResolvedTexturePath = UsdToUnreal::ConvertString( TextureAssetPath.GetResolvedPath() );
+	pxr::ArResolver& Resolver = pxr::ArGetResolver();
+
+	FString ResolvedTexturePath = UsdToUnreal::ConvertString( Resolver.ComputeNormalizedPath( TextureAssetPath.GetResolvedPath() ) );
 	FPaths::NormalizeFilename( ResolvedTexturePath );
 
 	if ( ResolvedTexturePath.IsEmpty() )
@@ -2074,6 +2087,8 @@ bool UsdUtils::MarkMaterialPrimWithWorldSpaceNormals( const UE::FUsdPrim& Materi
 
 bool UsdUtils::IsMaterialTranslucent( const pxr::UsdShadeMaterial& UsdShadeMaterial )
 {
+	FScopedUsdAllocs UsdAllocs;
+
 	pxr::UsdShadeShader SurfaceShader = UsdShadeMaterial.ComputeSurfaceSource();
 	if ( !SurfaceShader )
 	{
@@ -2091,6 +2106,55 @@ bool UsdUtils::IsMaterialTranslucent( const pxr::UsdShadeMaterial& UsdShadeMater
 	return bHasOpacityConnection && ( bIsTranslucentFloat || bHasBoundTexture );
 }
 
+bool UsdUtils::IsMaterialUsingUDIMs( const pxr::UsdShadeMaterial& UsdShadeMaterial )
+{
+	FScopedUsdAllocs UsdAllocs;
+
+	pxr::UsdShadeShader SurfaceShader = UsdShadeMaterial.ComputeSurfaceSource();
+	if ( !SurfaceShader )
+	{
+		return false;
+	}
+
+	for ( const pxr::UsdShadeInput& ShadeInput : SurfaceShader.GetInputs() )
+	{
+		pxr::UsdShadeConnectableAPI Source;
+		pxr::TfToken SourceName;
+		pxr::UsdShadeAttributeType AttributeType;
+
+		if ( ShadeInput.GetConnectedSource(&Source, &SourceName, &AttributeType))
+		{
+			pxr::UsdShadeInput FileInput;
+
+			// UsdUVTexture: Get its file input
+			if ( AttributeType == pxr::UsdShadeAttributeType::Output )
+			{
+				FileInput = Source.GetInput( UnrealIdentifiers::File );
+			}
+			// Check if we are being directly passed an asset
+			else
+			{
+				FileInput = Source.GetInput( SourceName );
+			}
+
+			if ( FileInput && FileInput.GetTypeName() == pxr::SdfValueTypeNames->Asset ) // Check that FileInput is of type Asset
+			{
+				pxr::SdfAssetPath TextureAssetPath;
+				FileInput.GetAttr().Get< pxr::SdfAssetPath >( &TextureAssetPath );
+
+				FString TexturePath = UsdToUnreal::ConvertString( TextureAssetPath.GetAssetPath() );
+
+				if ( TexturePath.Contains( TEXT("<UDIM>") ) )
+				{
+					return true;
+				}
+			}
+		}
+	}
+
+	return false;
+}
+
 FSHAHash UsdUtils::HashShadeMaterial( const pxr::UsdShadeMaterial& UsdShadeMaterial )
 {
 	FScopedUsdAllocs UsdAllocs;
diff --git a/Engine/Plugins/Importers/USDImporter/Source/USDUtilities/Public/USDShadeConversion.h b/Engine/Plugins/Importers/USDImporter/Source/USDUtilities/Public/USDShadeConversion.h
--- a/Engine/Plugins/Importers/USDImporter/Source/USDUtilities/Public/USDShadeConversion.h
+++ b/Engine/Plugins/Importers/USDImporter/Source/USDUtilities/Public/USDShadeConversion.h
@@ -99,6 +99,12 @@ namespace UsdUtils
 	 */
 	USDUTILITIES_API bool IsMaterialTranslucent( const pxr::UsdShadeMaterial& UsdShadeMaterial );
 
+	/**
+	* Returns whether the material uses UDIMs textures.
+	* This function exists because we need this information *before* we pick the right parent for a material instance and properly convert it.
+	*/
+	USDUTILITIES_API bool IsMaterialUsingUDIMs( const pxr::UsdShadeMaterial& UsdShadeMaterial );
+
 	USDUTILITIES_API FSHAHash HashShadeMaterial( const pxr::UsdShadeMaterial& UsdShadeMaterial );
 
 	/** Returns the resolved path from a pxr::SdfAssetPath attribute. For UDIMs path, returns the path to the 1001 tile. */
