USD: Fixed asset paths not being resolved because Usd path normalizing returns "." for empty strings (non-resolved paths)

diff --git a/Engine/Plugins/Importers/USDImporter/Source/USDUtilities/Private/USDShadeConversion.cpp b/Engine/Plugins/Importers/USDImporter/Source/USDUtilities/Private/USDShadeConversion.cpp
--- a/Engine/Plugins/Importers/USDImporter/Source/USDUtilities/Private/USDShadeConversion.cpp
+++ b/Engine/Plugins/Importers/USDImporter/Source/USDUtilities/Private/USDShadeConversion.cpp
@@ -129,7 +129,13 @@ namespace UE
 					? UsdToUnreal::ConvertString( pxr::SdfComputeAssetPathRelativeToLayer( LayerHandle, UnrealToUsd::ConvertString( *AssetPathToResolve ).Get() ) )
 					: AssetPathToResolve;
 
-				std::string ResolvedPathUsd = Resolver.ComputeNormalizedPath( Resolver.Resolve( UnrealToUsd::ConvertString( *RelativePathToResolve ).Get().c_str() ) );
+				std::string ResolvedPathUsd = Resolver.Resolve( UnrealToUsd::ConvertString( *RelativePathToResolve ).Get().c_str() );
+
+				// Don't normalize an empty path as the result will be "."
+				if ( ResolvedPathUsd.size() > 0 )
+				{
+					ResolvedPathUsd = Resolver.ComputeNormalizedPath( ResolvedPathUsd );
+				}
 
 				FString ResolvedAssetPath = UsdToUnreal::ConvertString( ResolvedPathUsd );
 
@@ -1959,7 +1965,14 @@ FString UsdUtils::GetResolvedTexturePath( const pxr::UsdAttribute& TextureAssetP
 
 	pxr::ArResolver& Resolver = pxr::ArGetResolver();
 
-	FString ResolvedTexturePath = UsdToUnreal::ConvertString( Resolver.ComputeNormalizedPath( TextureAssetPath.GetResolvedPath() ) );
+	std::string UsdResolvedPath = TextureAssetPath.GetResolvedPath();
+	// Don't normalize an empty path as the result will be "."
+	if ( UsdResolvedPath.size() > 0 )
+	{
+		UsdResolvedPath = Resolver.ComputeNormalizedPath( UsdResolvedPath );
+	}
+
+	FString ResolvedTexturePath = UsdToUnreal::ConvertString( UsdResolvedPath );
 	FPaths::NormalizeFilename( ResolvedTexturePath );
 
 	if ( ResolvedTexturePath.IsEmpty() )
