USD: Add analytics to common operations

diff --git a/Engine/Plugins/Importers/USDImporter/Source/USDClasses/Private/USDClassesModule.cpp b/Engine/Plugins/Importers/USDImporter/Source/USDClasses/Private/USDClassesModule.cpp
--- a/Engine/Plugins/Importers/USDImporter/Source/USDClasses/Private/USDClassesModule.cpp
+++ b/Engine/Plugins/Importers/USDImporter/Source/USDClasses/Private/USDClassesModule.cpp
@@ -4,7 +4,10 @@
 
 #include "USDLog.h"
 
+#include "AnalyticsEventAttribute.h"
+#include "EngineAnalytics.h"
 #include "HAL/FileManager.h"
+#include "Misc/EngineVersion.h"
 #include "Misc/FileHelper.h"
 #include "Misc/Paths.h"
 #include "Modules/ModuleManager.h"
@@ -177,6 +180,26 @@ void IUsdClassesModule::UpdatePlugInfoFiles( const FString& PluginDirectory, con
 	}
 }
 
+void IUsdClassesModule::SendAnalytics( TArray<FAnalyticsEventAttribute>&& InAttributes, const FString& EventName, bool bAutomated, double ElapsedSeconds, double NumberOfFrames, const FString& Extension )
+{
+	if ( FEngineAnalytics::IsAvailable() )
+	{
+		TArray<FAnalyticsEventAttribute> Attributes( InAttributes );
+
+		Attributes.Emplace( TEXT( "Automated" ), LexToString( bAutomated ) );
+		Attributes.Emplace( TEXT( "FileExtension" ), Extension );
+		Attributes.Emplace( TEXT( "NumberOfFrames" ), LexToString( FMath::Abs( NumberOfFrames ) ) );
+		Attributes.Emplace( TEXT( "TimeTaken.Seconds" ), ElapsedSeconds );
+		Attributes.Emplace( TEXT( "Platform" ), FPlatformProperties::IniPlatformName() );
+		Attributes.Emplace( TEXT( "EngineVersion" ), FEngineVersion::Current().ToString() );
+		Attributes.Emplace( TEXT( "EngineMode" ), FPlatformMisc::GetEngineMode() );
+		Attributes.Emplace( TEXT( "EpicAccountId" ), FPlatformMisc::GetEpicAccountId() );
+
+		const FString EventText = FString::Printf( TEXT( "Engine.Usage.USD.%s" ), *EventName);
+		FEngineAnalytics::GetProvider().RecordEvent( EventText, Attributes );
+	}
+}
+
 class FUsdClassesModule : public IUsdClassesModule
 {
 };
diff --git a/Engine/Plugins/Importers/USDImporter/Source/USDClasses/Public/USDClassesModule.h b/Engine/Plugins/Importers/USDImporter/Source/USDClasses/Public/USDClassesModule.h
--- a/Engine/Plugins/Importers/USDImporter/Source/USDClasses/Public/USDClassesModule.h
+++ b/Engine/Plugins/Importers/USDImporter/Source/USDClasses/Public/USDClassesModule.h
@@ -5,9 +5,22 @@
 #include "CoreMinimal.h"
 #include "Modules/ModuleInterface.h"
 
+struct FAnalyticsEventAttribute;
+
 class IUsdClassesModule : public IModuleInterface
 {
 public:
 	/** Updates all plugInfo.json to point their LibraryPaths to TargetDllFolder */
 	USDCLASSES_API static void UpdatePlugInfoFiles( const FString& PluginDirectory, const FString& TargetDllFolder );
+
+	/**
+	 * Sends analytics about a USD operation
+	 * @param InAttributes - Additional analytics events attributes to send, along with new ones collected within this function
+	 * @param EventName - Name of the analytics event (e.g. "Export.StaticMesh", so that the full event name is "Engine.Usage.USD.Export.StaticMesh")
+	 * @param bAutomated - If the operation was automated (e.g. came from a Python script)
+	 * @param ElapsedSeconds - How long the operation took in seconds
+	 * @param NumberOfFrames - Number of time codes in the exported/imported/opened stage
+	 * @param Extension - Extension of the main USD file opened/emitted/imported (e.g. "usda" or "usd")
+	 */
+	USDCLASSES_API static void SendAnalytics( TArray<FAnalyticsEventAttribute>&& InAttributes, const FString& EventName, bool bAutomated, double ElapsedSeconds, double NumberOfFrames, const FString& Extension );
 };
\ No newline at end of file
diff --git a/Engine/Plugins/Importers/USDImporter/Source/USDClasses/USDClasses.Build.cs b/Engine/Plugins/Importers/USDImporter/Source/USDClasses/USDClasses.Build.cs
--- a/Engine/Plugins/Importers/USDImporter/Source/USDClasses/USDClasses.Build.cs
+++ b/Engine/Plugins/Importers/USDImporter/Source/USDClasses/USDClasses.Build.cs
@@ -12,6 +12,7 @@ namespace UnrealBuildTool.Rules
 			PublicDependencyModuleNames.AddRange(
 				new string[]
 				{
+					"Analytics",
 					"Core",
 					"CoreUObject",
 					"Engine",
diff --git a/Engine/Plugins/Importers/USDImporter/Source/USDExporter/Private/AnimSequenceExporterUSD.cpp b/Engine/Plugins/Importers/USDImporter/Source/USDExporter/Private/AnimSequenceExporterUSD.cpp
--- a/Engine/Plugins/Importers/USDImporter/Source/USDExporter/Private/AnimSequenceExporterUSD.cpp
+++ b/Engine/Plugins/Importers/USDImporter/Source/USDExporter/Private/AnimSequenceExporterUSD.cpp
@@ -3,7 +3,9 @@
 #include "AnimSequenceExporterUSD.h"
 
 #include "AnimSequenceExporterUSDOptions.h"
+#include "EngineAnalytics.h"
 #include "UnrealUSDWrapper.h"
+#include "USDClassesModule.h"
 #include "USDErrorUtils.h"
 #include "USDLayerUtils.h"
 #include "USDOptionsWindow.h"
@@ -17,6 +19,46 @@
 #include "Animation/AnimSequence.h"
 #include "AssetExportTask.h"
 
+namespace UE
+{
+	namespace AnimSequenceExporterUSD
+	{
+		namespace Private
+		{
+			void SendAnalytics( UObject* Asset, UAnimSequenceExporterUSDOptions* Options, bool bAutomated, double ElapsedSeconds, double NumberOfFrames, const FString& Extension )
+			{
+				if ( !Asset )
+				{
+					return;
+				}
+
+				if ( FEngineAnalytics::IsAvailable() )
+				{
+					FString ClassName = Asset->GetClass()->GetName();
+
+					TArray<FAnalyticsEventAttribute> EventAttributes;
+
+					EventAttributes.Emplace( TEXT( "AssetType" ), ClassName );
+
+					if ( Options )
+					{
+						EventAttributes.Emplace( TEXT( "ExportPreviewMesh" ), LexToString( Options->bExportPreviewMesh ) );
+						if ( Options->bExportPreviewMesh )
+						{
+							EventAttributes.Emplace( TEXT( "UsePayload" ), LexToString( Options->PreviewMeshOptions.bUsePayload ) );
+							EventAttributes.Emplace( TEXT( "PayloadFormat" ), Options->PreviewMeshOptions.PayloadFormat );
+							EventAttributes.Emplace( TEXT( "LowestMeshLOD" ), LexToString( Options->PreviewMeshOptions.LowestMeshLOD ) );
+							EventAttributes.Emplace( TEXT( "HighestMeshLOD" ), LexToString( Options->PreviewMeshOptions.HighestMeshLOD ) );
+						}
+					}
+
+					IUsdClassesModule::SendAnalytics( MoveTemp( EventAttributes ), FString::Printf( TEXT( "Export.%s" ), *ClassName ), bAutomated, ElapsedSeconds, NumberOfFrames, Extension );
+				}
+			}
+		}
+	}
+}
+
 UAnimSequenceExporterUSD::UAnimSequenceExporterUSD()
 {
 #if USE_USD_SDK
@@ -66,6 +108,8 @@ bool UAnimSequenceExporterUSD::ExportBinary( UObject* Object, const TCHAR* Type,
 		}
 	}
 
+	double StartTime = FPlatformTime::Cycles64();
+
 	// If bUsePayload is true, we'll intercept the filename so that we write the mesh data to
 	// "C:/MyFolder/file_payload.usda" and create an "asset" file "C:/MyFolder/file.usda" that uses it
 	// as a payload, pointing at the default prim
@@ -187,6 +231,15 @@ bool UAnimSequenceExporterUSD::ExportBinary( UObject* Object, const TCHAR* Type,
 
 	AssetStage.GetRootLayer().Save();
 
+	// Analytics
+	{
+		bool bAutomated = ExportTask ? ExportTask->bAutomated : false;
+		double ElapsedSeconds = FPlatformTime::ToSeconds64( FPlatformTime::Cycles64() - StartTime );
+		FString Extension = FPaths::GetExtension( UExporter::CurrentFilename );
+
+		UE::AnimSequenceExporterUSD::Private::SendAnalytics( Object, Options, bAutomated, ElapsedSeconds, EndTimeCode - StartTimeCode, Extension );
+	}
+
 	return true;
 #else
 	return false;
diff --git a/Engine/Plugins/Importers/USDImporter/Source/USDExporter/Private/SkeletalMeshExporterUSD.cpp b/Engine/Plugins/Importers/USDImporter/Source/USDExporter/Private/SkeletalMeshExporterUSD.cpp
--- a/Engine/Plugins/Importers/USDImporter/Source/USDExporter/Private/SkeletalMeshExporterUSD.cpp
+++ b/Engine/Plugins/Importers/USDImporter/Source/USDExporter/Private/SkeletalMeshExporterUSD.cpp
@@ -2,7 +2,9 @@
 
 #include "SkeletalMeshExporterUSD.h"
 
+#include "EngineAnalytics.h"
 #include "SkeletalMeshExporterUSDOptions.h"
+#include "USDClassesModule.h"
 #include "USDConversionUtils.h"
 #include "USDMemory.h"
 #include "USDOptionsWindow.h"
@@ -17,6 +19,42 @@
 #include "AssetExportTask.h"
 #include "Engine/SkeletalMesh.h"
 
+namespace UE
+{
+	namespace SkeletalMeshExporterUSD
+	{
+		namespace Private
+		{
+			void SendAnalytics( UObject* Asset, USkeletalMeshExporterUSDOptions* Options, bool bAutomated, double ElapsedSeconds, double NumberOfFrames, const FString& Extension )
+			{
+				if ( !Asset )
+				{
+					return;
+				}
+
+				if ( FEngineAnalytics::IsAvailable() )
+				{
+					FString ClassName = Asset->GetClass()->GetName();
+
+					TArray<FAnalyticsEventAttribute> EventAttributes;
+
+					EventAttributes.Emplace( TEXT( "AssetType" ), ClassName );
+
+					if ( Options )
+					{
+						EventAttributes.Emplace( TEXT( "UsePayload" ), LexToString( Options->Inner.bUsePayload ) );
+						EventAttributes.Emplace( TEXT( "PayloadFormat" ), Options->Inner.PayloadFormat );
+						EventAttributes.Emplace( TEXT( "LowestMeshLOD" ), LexToString( Options->Inner.LowestMeshLOD ) );
+						EventAttributes.Emplace( TEXT( "HighestMeshLOD" ), LexToString( Options->Inner.HighestMeshLOD ) );
+					}
+
+					IUsdClassesModule::SendAnalytics( MoveTemp( EventAttributes ), FString::Printf( TEXT( "Export.%s" ), *ClassName ), bAutomated, ElapsedSeconds, NumberOfFrames, Extension );
+				}
+			}
+		}
+	}
+}
+
 USkeletalMeshExporterUsd::USkeletalMeshExporterUsd()
 {
 #if USE_USD_SDK
@@ -64,6 +102,8 @@ bool USkeletalMeshExporterUsd::ExportBinary( UObject* Object, const TCHAR* Type,
 		}
 	}
 
+	double StartTime = FPlatformTime::Cycles64();
+
 	// If bUsePayload is true, we'll intercept the filename so that we write the mesh data to
 	// "C:/MyFolder/file_payload.usda" and create an "asset" file "C:/MyFolder/file.usda" that uses it
 	// as a payload, pointing at the default prim
@@ -144,6 +184,16 @@ bool USkeletalMeshExporterUsd::ExportBinary( UObject* Object, const TCHAR* Type,
 
 	UsdStage.GetRootLayer().Save();
 
+	// Analytics
+	{
+		bool bAutomated = ExportTask ? ExportTask->bAutomated : false;
+		double ElapsedSeconds = FPlatformTime::ToSeconds64( FPlatformTime::Cycles64() - StartTime );
+		FString Extension = FPaths::GetExtension( UExporter::CurrentFilename );
+		double NumberOfFrames = UsdStage.GetEndTimeCode() - UsdStage.GetStartTimeCode();
+
+		UE::SkeletalMeshExporterUSD::Private::SendAnalytics( Object, Options, bAutomated, ElapsedSeconds, NumberOfFrames, Extension );
+	}
+
 	return true;
 #else
 	return false;
diff --git a/Engine/Plugins/Importers/USDImporter/Source/USDExporter/Private/StaticMeshExporterUSD.cpp b/Engine/Plugins/Importers/USDImporter/Source/USDExporter/Private/StaticMeshExporterUSD.cpp
--- a/Engine/Plugins/Importers/USDImporter/Source/USDExporter/Private/StaticMeshExporterUSD.cpp
+++ b/Engine/Plugins/Importers/USDImporter/Source/USDExporter/Private/StaticMeshExporterUSD.cpp
@@ -2,7 +2,9 @@
 
 #include "StaticMeshExporterUSD.h"
 
+#include "EngineAnalytics.h"
 #include "StaticMeshExporterUSDOptions.h"
+#include "USDClassesModule.h"
 #include "USDConversionUtils.h"
 #include "USDGeomMeshConversion.h"
 #include "USDMemory.h"
@@ -18,6 +20,42 @@
 #include "CoreMinimal.h"
 #include "Engine/StaticMesh.h"
 
+namespace UE
+{
+	namespace StaticMeshExporterUSD
+	{
+		namespace Private
+		{
+			void SendAnalytics( UObject* Asset, UStaticMeshExporterUSDOptions* Options, bool bAutomated, double ElapsedSeconds, double NumberOfFrames, const FString& Extension )
+			{
+				if ( !Asset )
+				{
+					return;
+				}
+
+				if ( FEngineAnalytics::IsAvailable() )
+				{
+					FString ClassName = Asset->GetClass()->GetName();
+
+					TArray<FAnalyticsEventAttribute> EventAttributes;
+
+					EventAttributes.Emplace( TEXT( "AssetType" ), ClassName );
+
+					if ( Options )
+					{
+						EventAttributes.Emplace( TEXT( "UsePayload" ), LexToString( Options->bUsePayload ) );
+						EventAttributes.Emplace( TEXT( "PayloadFormat" ), Options->PayloadFormat );
+						EventAttributes.Emplace( TEXT( "LowestMeshLOD" ), LexToString( Options->LowestMeshLOD ) );
+						EventAttributes.Emplace( TEXT( "HighestMeshLOD" ), LexToString( Options->HighestMeshLOD ) );
+					}
+
+					IUsdClassesModule::SendAnalytics( MoveTemp( EventAttributes ), FString::Printf( TEXT( "Export.%s" ), *ClassName ), bAutomated, ElapsedSeconds, NumberOfFrames, Extension );
+				}
+			}
+		}
+	}
+}
+
 bool UStaticMeshExporterUsd::IsUsdAvailable()
 {
 #if USE_USD_SDK
@@ -74,6 +112,8 @@ bool UStaticMeshExporterUsd::ExportBinary( UObject* Object, const TCHAR* Type, F
 		}
 	}
 
+	double StartTime = FPlatformTime::Cycles64();
+
 	// If bUsePayload is true, we'll intercept the filename so that we write the mesh data to
 	// "C:/MyFolder/file_payload.usda" and create an "asset" file "C:/MyFolder/file.usda" that uses it
 	// as a payload, pointing at the default prim
@@ -152,6 +192,16 @@ bool UStaticMeshExporterUsd::ExportBinary( UObject* Object, const TCHAR* Type, F
 
 	UsdStage.GetRootLayer().Save();
 
+	// Analytics
+	{
+		bool bAutomated = ExportTask ? ExportTask->bAutomated : false;
+		double ElapsedSeconds = FPlatformTime::ToSeconds64( FPlatformTime::Cycles64() - StartTime );
+		FString Extension = FPaths::GetExtension( UExporter::CurrentFilename );
+		double NumberOfFrames = UsdStage.GetEndTimeCode() - UsdStage.GetStartTimeCode();
+
+		UE::StaticMeshExporterUSD::Private::SendAnalytics( Object, Options, bAutomated, ElapsedSeconds, NumberOfFrames, Extension );
+	}
+
 	return true;
 #else
 	return false;
diff --git a/Engine/Plugins/Importers/USDImporter/Source/USDExporter/Private/USDExporterBlueprintLibrary.cpp b/Engine/Plugins/Importers/USDImporter/Source/USDExporter/Private/USDExporterBlueprintLibrary.cpp
--- a/Engine/Plugins/Importers/USDImporter/Source/USDExporter/Private/USDExporterBlueprintLibrary.cpp
+++ b/Engine/Plugins/Importers/USDImporter/Source/USDExporter/Private/USDExporterBlueprintLibrary.cpp
@@ -2,12 +2,15 @@
 
 #include "USDExporterBlueprintLibrary.h"
 
+#include "AnalyticsBlueprintLibrary.h"
+#include "AnalyticsEventAttribute.h"
 #include "CoreMinimal.h"
 #include "Editor.h"
 #include "InstancedFoliageActor.h"
 #include "UObject/ObjectMacros.h"
+#include "USDClassesModule.h"
 
-AInstancedFoliageActor* USDExporterBlueprintLibrary::GetInstancedFoliageActorForLevel( bool bCreateIfNone /*= false */, ULevel* Level /*= nullptr */ )
+AInstancedFoliageActor* UUsdExporterBlueprintLibrary::GetInstancedFoliageActorForLevel( bool bCreateIfNone /*= false */, ULevel* Level /*= nullptr */ )
 {
 	if ( !Level )
 	{
@@ -28,7 +31,7 @@ AInstancedFoliageActor* USDExporterBlueprintLibrary::GetInstancedFoliageActorFor
 	return AInstancedFoliageActor::GetInstancedFoliageActorForLevel( Level, bCreateIfNone );
 }
 
-TArray<UFoliageType*> USDExporterBlueprintLibrary::GetUsedFoliageTypes( AInstancedFoliageActor* Actor )
+TArray<UFoliageType*> UUsdExporterBlueprintLibrary::GetUsedFoliageTypes( AInstancedFoliageActor* Actor )
 {
 	TArray<UFoliageType*> Result;
 	if ( !Actor )
@@ -44,7 +47,7 @@ TArray<UFoliageType*> USDExporterBlueprintLibrary::GetUsedFoliageTypes( AInstanc
 	return Result;
 }
 
-UObject* USDExporterBlueprintLibrary::GetSource( UFoliageType* FoliageType )
+UObject* UUsdExporterBlueprintLibrary::GetSource( UFoliageType* FoliageType )
 {
 	if ( FoliageType )
 	{
@@ -54,7 +57,7 @@ UObject* USDExporterBlueprintLibrary::GetSource( UFoliageType* FoliageType )
 	return nullptr;
 }
 
-TArray<FTransform> USDExporterBlueprintLibrary::GetInstanceTransforms( AInstancedFoliageActor* Actor, UFoliageType* FoliageType, ULevel* InstancesLevel )
+TArray<FTransform> UUsdExporterBlueprintLibrary::GetInstanceTransforms( AInstancedFoliageActor* Actor, UFoliageType* FoliageType, ULevel* InstancesLevel )
 {
 	TArray<FTransform> Result;
 	if ( !Actor || !FoliageType )
@@ -98,3 +101,15 @@ TArray<FTransform> USDExporterBlueprintLibrary::GetInstanceTransforms( AInstance
 	return Result;
 }
 
+void UUsdExporterBlueprintLibrary::SendAnalytics( const TArray<FAnalyticsEventAttr>& Attrs, const FString& EventName, bool bAutomated, double ElapsedSeconds, double NumberOfFrames, const FString& Extension )
+{
+	TArray<FAnalyticsEventAttribute> Converted;
+	Converted.Reserve( Attrs.Num() );
+	for ( const FAnalyticsEventAttr& Attr : Attrs )
+	{
+		Converted.Emplace( Attr.Name, Attr.Value );
+	}
+
+	IUsdClassesModule::SendAnalytics( MoveTemp( Converted ), EventName, bAutomated, ElapsedSeconds, NumberOfFrames, Extension );
+}
+
diff --git a/Engine/Plugins/Importers/USDImporter/Source/USDExporter/Public/USDExporterBlueprintLibrary.h b/Engine/Plugins/Importers/USDImporter/Source/USDExporter/Public/USDExporterBlueprintLibrary.h
--- a/Engine/Plugins/Importers/USDImporter/Source/USDExporter/Public/USDExporterBlueprintLibrary.h
+++ b/Engine/Plugins/Importers/USDImporter/Source/USDExporter/Public/USDExporterBlueprintLibrary.h
@@ -2,9 +2,10 @@
 
 #pragma once
 
+#include "AnalyticsBlueprintLibrary.h"
 #include "CoreMinimal.h"
-#include "UObject/ObjectMacros.h"
 #include "Kismet/BlueprintFunctionLibrary.h"
+#include "UObject/ObjectMacros.h"
 
 #include "USDExporterBlueprintLibrary.generated.h"
 
@@ -16,7 +17,7 @@ class ULevel;
  * Library of functions that can be used via Python scripting to help export Unreal scenes and assets to USD
  */
 UCLASS(meta=(ScriptName="USDExporterLibrary"))
-class USDEXPORTER_API USDExporterBlueprintLibrary : public UBlueprintFunctionLibrary
+class USDEXPORTER_API UUsdExporterBlueprintLibrary : public UBlueprintFunctionLibrary
 {
 	GENERATED_BODY()
 
@@ -51,4 +52,8 @@ public:
 	 */
 	UFUNCTION( BlueprintCallable, meta = ( ScriptMethod ), Category = "USD Foliage Exporter" )
 	static TArray<FTransform> GetInstanceTransforms( AInstancedFoliageActor* Actor, UFoliageType* FoliageType, ULevel* InstancesLevel = nullptr);
+
+	/** Defer to the USDClasses module to actually send analytics information */
+	UFUNCTION( BlueprintCallable, Category = "Analytics" )
+	static void SendAnalytics( const TArray<FAnalyticsEventAttr>& Attrs, const FString& EventName, bool bAutomated, double ElapsedSeconds, double NumberOfFrames, const FString& Extension );
 };
\ No newline at end of file
diff --git a/Engine/Plugins/Importers/USDImporter/Source/USDExporter/USDExporter.Build.cs b/Engine/Plugins/Importers/USDImporter/Source/USDExporter/USDExporter.Build.cs
--- a/Engine/Plugins/Importers/USDImporter/Source/USDExporter/USDExporter.Build.cs
+++ b/Engine/Plugins/Importers/USDImporter/Source/USDExporter/USDExporter.Build.cs
@@ -23,7 +23,9 @@ namespace UnrealBuildTool.Rules
 			PrivateDependencyModuleNames.AddRange(
 				new string[]
 				{
-                    "EditorStyle",
+					"Analytics",
+					"AnalyticsBlueprintLibrary",
+					"EditorStyle",
                     "GeometryCache",
                     "PropertyEditor",
                     "RawMesh",
diff --git a/Engine/Plugins/Importers/USDImporter/Source/USDStage/Private/USDStageActor.cpp b/Engine/Plugins/Importers/USDImporter/Source/USDStage/Private/USDStageActor.cpp
--- a/Engine/Plugins/Importers/USDImporter/Source/USDStage/Private/USDStageActor.cpp
+++ b/Engine/Plugins/Importers/USDImporter/Source/USDStage/Private/USDStageActor.cpp
@@ -5,6 +5,7 @@
 #include "UnrealUSDWrapper.h"
 #include "USDAssetCache.h"
 #include "USDAssetImportData.h"
+#include "USDClassesModule.h"
 #include "USDConversionUtils.h"
 #include "USDErrorUtils.h"
 #include "USDGeomMeshTranslator.h"
@@ -43,6 +44,7 @@
 #include "Engine/SkeletalMesh.h"
 #include "Engine/StaticMesh.h"
 #include "Engine/World.h"
+#include "EngineAnalytics.h"
 #include "HAL/FileManager.h"
 #include "LevelSequence.h"
 #include "Materials/Material.h"
@@ -533,6 +535,26 @@ struct FUsdStageActorImpl
 		const bool bVisible = false;
 		PrimSceneComponent->SetVisibility( bVisible, bPropagateToChildren );
 	}
+
+	static void SendAnalytics( AUsdStageActor* StageActor, double ElapsedSeconds, double NumberOfFrames, const FString& Extension )
+	{
+		if ( !StageActor )
+		{
+			return;
+		}
+
+		if ( FEngineAnalytics::IsAvailable() )
+		{
+			TArray<FAnalyticsEventAttribute> EventAttributes;
+
+			EventAttributes.Emplace( TEXT( "InitialLoadSet" ), LexToString( (uint8)StageActor->InitialLoadSet ) );
+			EventAttributes.Emplace( TEXT( "PurposesToLoad" ), LexToString( StageActor->PurposesToLoad ) );
+			EventAttributes.Emplace( TEXT( "RenderContext" ), StageActor->RenderContext.ToString() );
+
+			const bool bAutomated = false;
+			IUsdClassesModule::SendAnalytics( MoveTemp( EventAttributes ), TEXT( "Open" ), bAutomated, ElapsedSeconds, NumberOfFrames, Extension );
+		}
+	}
 };
 
 /**
@@ -1543,6 +1565,10 @@ void AUsdStageActor::LoadUsdStage()
 	ElapsedSeconds -= 60.0 * (double)ElapsedMin;
 
 	UE_LOG( LogUsd, Log, TEXT("%s %s in [%d min %.3f s]"), TEXT("Stage loaded"), *FPaths::GetBaseFilename( RootLayer.FilePath ), ElapsedMin, ElapsedSeconds );
+
+	double NumberOfFrames = UsdStage ? (UsdStage.GetEndTimeCode() - UsdStage.GetStartTimeCode()) : 0.0;
+
+	FUsdStageActorImpl::SendAnalytics( this, ElapsedSeconds, NumberOfFrames, FPaths::GetExtension( RootLayer.FilePath ) );
 }
 
 void AUsdStageActor::UnloadUsdStage()
diff --git a/Engine/Plugins/Importers/USDImporter/Source/USDStageImporter/Private/USDStageImporter.cpp b/Engine/Plugins/Importers/USDImporter/Source/USDStageImporter/Private/USDStageImporter.cpp
--- a/Engine/Plugins/Importers/USDImporter/Source/USDStageImporter/Private/USDStageImporter.cpp
+++ b/Engine/Plugins/Importers/USDImporter/Source/USDStageImporter/Private/USDStageImporter.cpp
@@ -4,6 +4,7 @@
 
 #include "USDAssetCache.h"
 #include "USDAssetImportData.h"
+#include "USDClassesModule.h"
 #include "USDConversionUtils.h"
 #include "USDErrorUtils.h"
 #include "USDGeomMeshConversion.h"
@@ -29,6 +30,7 @@
 #include "Engine/StaticMesh.h"
 #include "Engine/Texture.h"
 #include "Engine/World.h"
+#include "EngineAnalytics.h"
 #include "EngineUtils.h"
 #include "GeometryCache.h"
 #include "HAL/FileManager.h"
@@ -1201,6 +1203,96 @@ namespace UsdStageImporterImpl
 			}
 		}
 	}
+
+	void SendAnalytics( FUsdStageImportContext& ImportContext, UObject* Asset, const FString& Operation, const TSet<UObject*>& ImportedAssets, double ElapsedSeconds )
+	{
+		if ( FEngineAnalytics::IsAvailable() )
+		{
+			TArray<FAnalyticsEventAttribute> EventAttributes;
+
+			FString EventName = Operation;
+			if ( Asset )
+			{
+				FString ClassName = Asset->GetClass()->GetName();
+
+				// e.g. "Reimport.StaticMesh"
+				EventName = FString::Printf( TEXT( "%s.%s" ), *EventName, *ClassName );
+				EventAttributes.Emplace( TEXT( "AssetType" ), ClassName );
+			}
+
+			if ( ImportContext.ImportOptions )
+			{
+				EventAttributes.Emplace( TEXT( "ImportActors" ), LexToString( ImportContext.ImportOptions->bImportActors ) );
+				EventAttributes.Emplace( TEXT( "ImportGeometry" ), LexToString( ImportContext.ImportOptions->bImportGeometry ) );
+				EventAttributes.Emplace( TEXT( "ImportSkeletalAnimations" ), LexToString( ImportContext.ImportOptions->bImportSkeletalAnimations ) );
+				EventAttributes.Emplace( TEXT( "ImportMaterials" ), LexToString( ImportContext.ImportOptions->bImportMaterials ) );
+				EventAttributes.Emplace( TEXT( "PurposesToImport" ), LexToString( ImportContext.ImportOptions->PurposesToImport ) );
+				EventAttributes.Emplace( TEXT( "RenderContextToImport" ), ImportContext.ImportOptions->RenderContextToImport.ToString() );
+				EventAttributes.Emplace( TEXT( "ReuseIdenticalAssets" ), ImportContext.ImportOptions->bReuseIdenticalAssets );
+				EventAttributes.Emplace( TEXT( "ReplaceActorPolicy" ), LexToString( (uint8)ImportContext.ImportOptions->ExistingActorPolicy ) );
+				EventAttributes.Emplace( TEXT( "ReplaceAssetPolicy" ), LexToString( (uint8)ImportContext.ImportOptions->ExistingAssetPolicy ) );
+				EventAttributes.Emplace( TEXT( "PrimPathFolderStructure" ), LexToString( ImportContext.ImportOptions->bPrimPathFolderStructure ) );
+				EventAttributes.Emplace( TEXT( "Collapse" ), LexToString( ImportContext.ImportOptions->bCollapse ) );
+				EventAttributes.Emplace( TEXT( "InterpretLODs" ), LexToString( ImportContext.ImportOptions->bInterpretLODs ) );
+			}
+
+			int32 NumStaticMeshes = 0;
+			int32 NumSkeletalMeshes = 0;
+			int32 NumMaterials = 0;
+			int32 NumAnimSequences = 0;
+			int32 NumTextures = 0;
+			int32 NumGeometryCaches = 0;
+			for ( UObject* ImportedAsset : ImportedAssets )
+			{
+				if ( !ImportedAsset )
+				{
+					continue;
+				}
+
+				if ( ImportedAsset->IsA<UStaticMesh>() )
+				{
+					++NumStaticMeshes;
+				}
+				else if ( ImportedAsset->IsA<USkeletalMesh>() )
+				{
+					++NumSkeletalMeshes;
+				}
+				else if ( ImportedAsset->IsA<UMaterialInterface>() )
+				{
+					++NumMaterials;
+				}
+				else if ( ImportedAsset->IsA<UAnimSequence>() )
+				{
+					++NumAnimSequences;
+				}
+				else if ( ImportedAsset->IsA<UTexture>() )
+				{
+					++NumTextures;
+				}
+				else if ( ImportedAsset->IsA<UGeometryCache>() )
+				{
+					++NumGeometryCaches;
+				}
+			}
+			EventAttributes.Emplace( TEXT( "NumStaticMeshes" ), LexToString( NumStaticMeshes ) );
+			EventAttributes.Emplace( TEXT( "NumSkeletalMeshes" ), LexToString( NumSkeletalMeshes ) );
+			EventAttributes.Emplace( TEXT( "NumMaterials" ), LexToString( NumMaterials ) );
+			EventAttributes.Emplace( TEXT( "NumAnimSequences" ), LexToString( NumAnimSequences ) );
+			EventAttributes.Emplace( TEXT( "NumTextures" ), LexToString( NumTextures ) );
+			EventAttributes.Emplace( TEXT( "NumGeometryCaches" ), LexToString( NumGeometryCaches ) );
+
+			double NumberOfFrames = ImportContext.Stage.GetEndTimeCode() - ImportContext.Stage.GetStartTimeCode();
+
+			IUsdClassesModule::SendAnalytics(
+				MoveTemp( EventAttributes ),
+				EventName,
+				ImportContext.bIsAutomated,
+				ElapsedSeconds,
+				NumberOfFrames,
+				FPaths::GetExtension( ImportContext.FilePath )
+			);
+		}
+	}
 }
 
 void UUsdStageImporter::ImportFromFile(FUsdStageImportContext& ImportContext)
@@ -1212,6 +1304,8 @@ void UUsdStageImporter::ImportFromFile(FUsdStageImportContext& ImportContext)
 		return;
 	}
 
+	double StartTime = FPlatformTime::Cycles64();
+
 	ImportContext.Stage = UsdStageImporterImpl::ReadUsdFile(ImportContext);
 	if (!ImportContext.Stage)
 	{
@@ -1278,6 +1372,13 @@ void UUsdStageImporter::ImportFromFile(FUsdStageImportContext& ImportContext)
 	UsdStageImporterImpl::CloseStageIfNeeded( ImportContext );
 
 	FUsdDelegates::OnPostUsdImport.Broadcast( ImportContext.FilePath );
+
+	// Analytics
+	{
+		double ElapsedSeconds = FPlatformTime::ToSeconds64( FPlatformTime::Cycles64() - StartTime );
+		UsdStageImporterImpl::SendAnalytics( ImportContext, nullptr, TEXT("Import"), UsedAssetsAndDependencies, ElapsedSeconds);
+	}
+
 #endif // #if USE_USD_SDK
 }
 
@@ -1287,6 +1388,8 @@ bool UUsdStageImporter::ReimportSingleAsset(FUsdStageImportContext& ImportContex
 	bool bSuccess = false;
 
 #if USE_USD_SDK
+	double StartTime = FPlatformTime::Cycles64();
+
 	ImportContext.Stage = UsdStageImporterImpl::ReadUsdFile(ImportContext);
 	if (!ImportContext.Stage)
 	{
@@ -1294,6 +1397,7 @@ bool UUsdStageImporter::ReimportSingleAsset(FUsdStageImportContext& ImportContex
 		return bSuccess;
 	}
 
+
 	FUsdDelegates::OnPreUsdImport.Broadcast(ImportContext.FilePath);
 
 	// We still need the scene actor to remap all other users of the mesh to the new reimported one. It's not critical if we fail though,
@@ -1373,6 +1477,12 @@ bool UUsdStageImporter::ReimportSingleAsset(FUsdStageImportContext& ImportContex
 
 	FUsdDelegates::OnPostUsdImport.Broadcast(ImportContext.FilePath);
 
+	// Analytics
+	{
+		double ElapsedSeconds = FPlatformTime::ToSeconds64( FPlatformTime::Cycles64() - StartTime );
+		UsdStageImporterImpl::SendAnalytics( ImportContext, ReimportedObject, TEXT( "Reimport" ), { ReimportedObject }, ElapsedSeconds );
+	}
+
 #endif // #if USE_USD_SDK
 	return bSuccess;
 }
diff --git a/Engine/Plugins/Importers/USDImporter/USDImporter.uplugin b/Engine/Plugins/Importers/USDImporter/USDImporter.uplugin
--- a/Engine/Plugins/Importers/USDImporter/USDImporter.uplugin
+++ b/Engine/Plugins/Importers/USDImporter/USDImporter.uplugin
@@ -98,6 +98,10 @@
 		{
 			"Name": "SequencerScripting",
 			"Enabled": true
+		},
+		{
+			"Name": "AnalyticsBlueprintLibrary",
+			"Enabled": true
 		}
 	]
 }
