USD: Fix mismatched stop/start monitoring call when the change is not enveloped in a transaction.

diff --git a/Engine/Plugins/Importers/USDImporter/Source/USDStage/Private/USDLevelSequenceHelper.cpp b/Engine/Plugins/Importers/USDImporter/Source/USDStage/Private/USDLevelSequenceHelper.cpp
--- a/Engine/Plugins/Importers/USDImporter/Source/USDStage/Private/USDLevelSequenceHelper.cpp
+++ b/Engine/Plugins/Importers/USDImporter/Source/USDStage/Private/USDLevelSequenceHelper.cpp
@@ -1294,10 +1294,9 @@ void FUsdLevelSequenceHelperImpl::BlockMonitoringChangesForThisTransaction()
 		}
 
 		BlockedTransactionGuids.Add( Context.TransactionId );
-	}
 
-	// Also block via the regular way in case we're receiving a notice due to a Python change and the user hasn't manually created a transaction
-	StopMonitoringChanges();
+		StopMonitoringChanges();
+	}
 }
 
 void FUsdLevelSequenceHelperImpl::OnObjectTransacted(UObject* Object, const class FTransactionObjectEvent& Event)
@@ -1832,3 +1831,30 @@ TArray< ULevelSequence* > FUsdLevelSequenceHelper::GetSubSequences() const
 		return {};
 	}
 }
+
+FScopedBlockMonitoringChangesForTransaction::FScopedBlockMonitoringChangesForTransaction( FUsdLevelSequenceHelper& InHelper )
+	: Helper( InHelper )
+{
+	// If we're transacting we can just call this and the helper will unblock itself once the transaction is finished, because
+	// we need to make sure the unblocking happens after any call to OnObjectTransacted.
+	if ( GUndo )
+	{
+		Helper.BlockMonitoringChangesForThisTransaction();
+	}
+	// If we're not in a transaction we still need to block this (can also happen e.g. if a Python change triggers a stage notice),
+	// but since we don't have to worry about the OnObjectTransacted calls we can just use this RAII object here to wrap over
+	// any potential changes to level sequence assets
+	else
+	{
+		bStoppedMonitoringChanges = true;
+		Helper.StopMonitoringChanges();
+	}
+}
+
+FScopedBlockMonitoringChangesForTransaction::~FScopedBlockMonitoringChangesForTransaction()
+{
+	if ( bStoppedMonitoringChanges )
+	{
+		Helper.StartMonitoringChanges();
+	}
+}
\ No newline at end of file
diff --git a/Engine/Plugins/Importers/USDImporter/Source/USDStage/Private/USDStageActor.cpp b/Engine/Plugins/Importers/USDImporter/Source/USDStage/Private/USDStageActor.cpp
--- a/Engine/Plugins/Importers/USDImporter/Source/USDStage/Private/USDStageActor.cpp
+++ b/Engine/Plugins/Importers/USDImporter/Source/USDStage/Private/USDStageActor.cpp
@@ -798,8 +798,8 @@ void AUsdStageActor::OnUsdObjectsChanged( const UsdUtils::FObjectChangesByPath&
 	}
 
 	// We may update our levelsequence objects (tracks, moviescene, sections, etc.) due to these changes. We definitely don't want to write anything
-	// back to USD when these objects change though
-	LevelSequenceHelper.BlockMonitoringChangesForThisTransaction();
+	// back to USD when these objects change though.
+	FScopedBlockMonitoringChangesForTransaction BlockMonitoring{ LevelSequenceHelper };
 
 	// The most important thing here is to iterate in parent to child order, so build SortedPrimsChangedList
 	TMap< FString, bool > SortedPrimsChangedList;
diff --git a/Engine/Plugins/Importers/USDImporter/Source/USDStage/Public/USDLevelSequenceHelper.h b/Engine/Plugins/Importers/USDImporter/Source/USDStage/Public/USDLevelSequenceHelper.h
--- a/Engine/Plugins/Importers/USDImporter/Source/USDStage/Public/USDLevelSequenceHelper.h
+++ b/Engine/Plugins/Importers/USDImporter/Source/USDStage/Public/USDLevelSequenceHelper.h
@@ -60,3 +60,20 @@ public:
 private:
 	TUniquePtr<FUsdLevelSequenceHelperImpl> UsdSequencerImpl;
 };
+
+class USDSTAGE_API FScopedBlockMonitoringChangesForTransaction final
+{
+public:
+	explicit FScopedBlockMonitoringChangesForTransaction( FUsdLevelSequenceHelper& InHelper );
+	~FScopedBlockMonitoringChangesForTransaction();
+
+	FScopedBlockMonitoringChangesForTransaction() = delete;
+	FScopedBlockMonitoringChangesForTransaction( const FScopedBlockMonitoringChangesForTransaction& ) = delete;
+	FScopedBlockMonitoringChangesForTransaction( FScopedBlockMonitoringChangesForTransaction&& ) = delete;
+	FScopedBlockMonitoringChangesForTransaction& operator=( const FScopedBlockMonitoringChangesForTransaction& ) = delete;
+	FScopedBlockMonitoringChangesForTransaction& operator=( FScopedBlockMonitoringChangesForTransaction&& ) = delete;
+
+private:
+	FUsdLevelSequenceHelper& Helper;
+	bool bStoppedMonitoringChanges = false;
+};
diff --git a/Engine/Plugins/Importers/USDImporter/Source/USDUtilities/Private/USDPrimConversion.cpp b/Engine/Plugins/Importers/USDImporter/Source/USDUtilities/Private/USDPrimConversion.cpp
--- a/Engine/Plugins/Importers/USDImporter/Source/USDUtilities/Private/USDPrimConversion.cpp
+++ b/Engine/Plugins/Importers/USDImporter/Source/USDUtilities/Private/USDPrimConversion.cpp
@@ -965,7 +965,6 @@ bool UnrealToUsd::ConvertXformable( const UMovieScene3DTransformTrack& MovieScen
 		pxr::SdfChangeBlock ChangeBlock;
 
 		// Compensate different orientation for light or camera components
-		// TODO: Handle inverse compensation when the bound component is a child of a light/camera component
 		FTransform AdditionalRotation = FTransform::Identity;
 		if ( UsdPrim.IsA< pxr::UsdGeomCamera >() || UsdPrim.IsA< pxr::UsdLuxLight >() )
 		{
