USD
- Fixed the UsdToUnreal::ConvertMaterial version that works on a UMaterial.
- Fixed some meshes not having UVs when its material had the same hash as another material. The primvars mapping for the uvs were missing for the material prim path.

diff --git a/Engine/Plugins/Importers/USDImporter/Source/USDClasses/Private/USDAssetCache.cpp b/Engine/Plugins/Importers/USDImporter/Source/USDClasses/Private/USDAssetCache.cpp
--- a/Engine/Plugins/Importers/USDImporter/Source/USDClasses/Private/USDAssetCache.cpp
+++ b/Engine/Plugins/Importers/USDImporter/Source/USDClasses/Private/USDAssetCache.cpp
@@ -153,6 +153,22 @@ UObject* UUsdAssetCache::GetAssetForPrim( const FString& PrimPath ) const
 	return nullptr;
 }
 
+FString UUsdAssetCache::GetPrimForAsset( UObject* Asset ) const
+{
+	FString PrimForAsset;
+
+	{
+		FScopeLock Lock( &CriticalSection );
+
+		if ( const FString* KeyPtr = PrimPathToAssets.FindKey( Asset ) )
+		{
+			PrimForAsset = *KeyPtr;
+		}
+	}
+
+	return PrimForAsset;
+}
+
 void UUsdAssetCache::Reset()
 {
 	FScopeLock Lock( &CriticalSection );
diff --git a/Engine/Plugins/Importers/USDImporter/Source/USDClasses/Public/USDAssetCache.h b/Engine/Plugins/Importers/USDImporter/Source/USDClasses/Public/USDAssetCache.h
--- a/Engine/Plugins/Importers/USDImporter/Source/USDClasses/Public/USDAssetCache.h
+++ b/Engine/Plugins/Importers/USDImporter/Source/USDClasses/Public/USDAssetCache.h
@@ -28,6 +28,10 @@ public:
 	void LinkAssetToPrim( const FString& PrimPath, UObject* Asset );
 	void RemoveAssetPrimLink( const FString& PrimPath );
 	UObject* GetAssetForPrim( const FString& PrimPath ) const;
+
+	/** Returns the first prim path associated with an asset (see TMap::FindKey) */
+	FString GetPrimForAsset( UObject* Asset ) const;
+
 	TMap< FString, UObject* > GetAssetPrimLinks() const { return PrimPathToAssets; }; // Can't return a reference as it wouldn't be thread-safe
 
 	bool IsAssetOwnedByCache( UObject* Asset ) const { return OwnedAssets.Contains( Asset ); }
diff --git a/Engine/Plugins/Importers/USDImporter/Source/USDSchemas/Private/USDShadeMaterialTranslator.cpp b/Engine/Plugins/Importers/USDImporter/Source/USDSchemas/Private/USDShadeMaterialTranslator.cpp
--- a/Engine/Plugins/Importers/USDImporter/Source/USDSchemas/Private/USDShadeMaterialTranslator.cpp
+++ b/Engine/Plugins/Importers/USDImporter/Source/USDSchemas/Private/USDShadeMaterialTranslator.cpp
@@ -35,7 +35,12 @@ void FUsdShadeMaterialTranslator::CreateAssets()
 
 	FString MaterialHashString = UsdUtils::HashShadeMaterial( ShadeMaterial ).ToString();
 
-	UMaterialInterface* ConvertedMaterial = Cast<UMaterialInterface>( Context->AssetCache->GetCachedAsset( MaterialHashString ) );
+	UMaterialInterface* ConvertedMaterial = nullptr;
+	
+	if ( Context->AssetCache )
+	{
+		ConvertedMaterial = Cast< UMaterialInterface >( Context->AssetCache->GetCachedAsset( MaterialHashString ) );
+	}
 
 	if ( !ConvertedMaterial )
 	{
@@ -95,6 +100,11 @@ void FUsdShadeMaterialTranslator::CreateAssets()
 			}
 		}
 	}
+	else if ( Context->MaterialToPrimvarToUVIndex && Context->AssetCache )
+	{
+		// Copy the Material -> Primvar -> UV index mapping from the cached material prim path to this prim path
+		Context->MaterialToPrimvarToUVIndex->FindOrAdd( PrimPath.GetString() ) = Context->MaterialToPrimvarToUVIndex->FindRef( Context->AssetCache->GetPrimForAsset( ConvertedMaterial ) );
+	}
 
 	if ( ConvertedMaterial )
 	{
diff --git a/Engine/Plugins/Importers/USDImporter/Source/USDUtilities/Private/USDShadeConversion.cpp b/Engine/Plugins/Importers/USDImporter/Source/USDUtilities/Private/USDShadeConversion.cpp
--- a/Engine/Plugins/Importers/USDImporter/Source/USDUtilities/Private/USDShadeConversion.cpp
+++ b/Engine/Plugins/Importers/USDImporter/Source/USDUtilities/Private/USDShadeConversion.cpp
@@ -784,9 +784,10 @@ namespace UE
 			{
 #if WITH_EDITOR
 				FGetExpressionForValueVisitor Visitor{ Material };
-				Visit( Visitor, ParameterValue );
-#endif // WITH_EDITOR
+				return Visit( Visitor, ParameterValue );
+#else
 				return nullptr;
+#endif // WITH_EDITOR
 			}
 
 			void SetScalarParameterValue( UMaterialInstance& Material, const TCHAR* ParameterName, float ParameterValue )
@@ -1733,32 +1734,36 @@ bool UsdToUnreal::ConvertMaterial( const pxr::UsdShadeMaterial& UsdShadeMaterial
 	// Base color
 	{
 		const bool bIsNormalMap = false;
-		UsdShadeConversionImpl::GetVec3ParameterValue( Connectable, UnrealIdentifiers::DiffuseColor, FLinearColor( 0, 0, 0 ), ParameterValue, bIsNormalMap, &Material, TexturesCache, &PrimvarToUVIndex );
-		if ( UMaterialExpression* Expression = UsdShadeConversionImpl::GetExpressionForValue( Material, ParameterValue ) )
+		if ( UsdShadeConversionImpl::GetVec3ParameterValue( Connectable, UnrealIdentifiers::DiffuseColor, FLinearColor( 0, 0, 0 ), ParameterValue, bIsNormalMap, &Material, TexturesCache, &PrimvarToUVIndex ) )
 		{
-			Material.BaseColor.Expression = Expression;
-			SetOutputIndex( ParameterValue, Material.BaseColor.OutputIndex );
+			if (UMaterialExpression* Expression = UsdShadeConversionImpl::GetExpressionForValue(Material, ParameterValue))
+			{
+				Material.BaseColor.Expression = Expression;
+				SetOutputIndex(ParameterValue, Material.BaseColor.OutputIndex);
 
-			bHasMaterialInfo = true;
+				bHasMaterialInfo = true;
+			}
 		}
 	}
 
 	// Emissive color
 	{
 		const bool bIsNormalMap = false;
-		UsdShadeConversionImpl::GetVec3ParameterValue( Connectable, UnrealIdentifiers::EmissiveColor, FLinearColor( 0, 0, 0 ), ParameterValue, bIsNormalMap, &Material, TexturesCache, &PrimvarToUVIndex );
-		if ( UMaterialExpression* Expression = UsdShadeConversionImpl::GetExpressionForValue( Material, ParameterValue ) )
+		if ( UsdShadeConversionImpl::GetVec3ParameterValue( Connectable, UnrealIdentifiers::EmissiveColor, FLinearColor( 0, 0, 0 ), ParameterValue, bIsNormalMap, &Material, TexturesCache, &PrimvarToUVIndex ) )
 		{
-			Material.EmissiveColor.Expression = Expression;
-			SetOutputIndex( ParameterValue, Material.EmissiveColor.OutputIndex );
+			if (UMaterialExpression* Expression = UsdShadeConversionImpl::GetExpressionForValue(Material, ParameterValue))
+			{
+				Material.EmissiveColor.Expression = Expression;
+				SetOutputIndex(ParameterValue, Material.EmissiveColor.OutputIndex);
 
-			bHasMaterialInfo = true;
+				bHasMaterialInfo = true;
+			}
 		}
 	}
 
 	// Metallic
+	if ( UsdShadeConversionImpl::GetFloatParameterValue( Connectable, UnrealIdentifiers::Metallic, 0.f, ParameterValue, &Material, TexturesCache, &PrimvarToUVIndex ) )
 	{
-		UsdShadeConversionImpl::GetFloatParameterValue( Connectable, UnrealIdentifiers::Metallic, 0.f, ParameterValue, &Material, TexturesCache, &PrimvarToUVIndex );
 		if ( UMaterialExpression* Expression = UsdShadeConversionImpl::GetExpressionForValue( Material, ParameterValue ) )
 		{
 			Material.Metallic.Expression = Expression;
@@ -1769,8 +1774,8 @@ bool UsdToUnreal::ConvertMaterial( const pxr::UsdShadeMaterial& UsdShadeMaterial
 	}
 
 	// Roughness
+	if ( UsdShadeConversionImpl::GetFloatParameterValue( Connectable, UnrealIdentifiers::Roughness, 1.f, ParameterValue, &Material, TexturesCache, &PrimvarToUVIndex ) )
 	{
-		UsdShadeConversionImpl::GetFloatParameterValue( Connectable, UnrealIdentifiers::Roughness, 1.f, ParameterValue, &Material, TexturesCache, &PrimvarToUVIndex );
 		if ( UMaterialExpression* Expression = UsdShadeConversionImpl::GetExpressionForValue( Material, ParameterValue ) )
 		{
 			Material.Roughness.Expression = Expression;
@@ -1781,8 +1786,8 @@ bool UsdToUnreal::ConvertMaterial( const pxr::UsdShadeMaterial& UsdShadeMaterial
 	}
 
 	// Opacity
+	if ( UsdShadeConversionImpl::GetFloatParameterValue( Connectable, UnrealIdentifiers::Opacity, 1.f, ParameterValue, &Material, TexturesCache, &PrimvarToUVIndex ) )
 	{
-		UsdShadeConversionImpl::GetFloatParameterValue( Connectable, UnrealIdentifiers::Opacity, 1.f, ParameterValue, &Material, TexturesCache, &PrimvarToUVIndex );
 		if ( UMaterialExpression* Expression = UsdShadeConversionImpl::GetExpressionForValue( Material, ParameterValue ) )
 		{
 			Material.Opacity.Expression = Expression;
@@ -1796,13 +1801,15 @@ bool UsdToUnreal::ConvertMaterial( const pxr::UsdShadeMaterial& UsdShadeMaterial
 	// Normal
 	{
 		const bool bIsNormalMap = true;
-		UsdShadeConversionImpl::GetVec3ParameterValue( Connectable, UnrealIdentifiers::Normal, FLinearColor( 0, 0, 1 ), ParameterValue, bIsNormalMap, &Material, TexturesCache, &PrimvarToUVIndex );
-		if ( UMaterialExpression* Expression = UsdShadeConversionImpl::GetExpressionForValue( Material, ParameterValue ) )
+		if ( UsdShadeConversionImpl::GetVec3ParameterValue( Connectable, UnrealIdentifiers::Normal, FLinearColor( 0, 0, 1 ), ParameterValue, bIsNormalMap, &Material, TexturesCache, &PrimvarToUVIndex ) )
 		{
-			Material.Normal.Expression = Expression;
-			SetOutputIndex( ParameterValue, Material.Normal.OutputIndex );
+			if ( UMaterialExpression* Expression = UsdShadeConversionImpl::GetExpressionForValue( Material, ParameterValue ) )
+			{
+				Material.Normal.Expression = Expression;
+				SetOutputIndex(ParameterValue, Material.Normal.OutputIndex);
 
-			bHasMaterialInfo = true;
+				bHasMaterialInfo = true;
+			}
 		}
 	}
 
