USD
- Added support for the IOR attribute in UsdPreviewSurface

diff --git a/Engine/Plugins/Importers/USDImporter/Source/USDUtilities/Private/USDShadeConversion.cpp b/Engine/Plugins/Importers/USDImporter/Source/USDUtilities/Private/USDShadeConversion.cpp
--- a/Engine/Plugins/Importers/USDImporter/Source/USDUtilities/Private/USDShadeConversion.cpp
+++ b/Engine/Plugins/Importers/USDImporter/Source/USDUtilities/Private/USDShadeConversion.cpp
@@ -1686,6 +1686,16 @@ bool UsdToUnreal::ConvertMaterial( const pxr::UsdShadeMaterial& UsdShadeMaterial
 		}
 	}
 
+	// Refraction
+	{
+		if ( UsdShadeConversionImpl::GetFloatParameterValue( Connectable, UnrealIdentifiers::Refraction, 1.5f, ParameterValue, &Material, TexturesCache, &PrimvarToUVIndex, bForceVirtualTextures ) )
+		{
+			UsdShadeConversionImpl::SetParameterValue( Material, TEXT( "Refraction" ), ParameterValue, bForUsdPreviewSurface );
+			SetTextureComponentParamForScalarInput( ParameterValue, TEXT( "RefractionTextureComponent" ) );
+			bHasMaterialInfo = true;
+		}
+	}
+
 	// Handle world space normals
 	pxr::UsdPrim UsdPrim = UsdShadeMaterial.GetPrim();
 	if ( pxr::UsdAttribute Attr = UsdPrim.GetAttribute( UnrealIdentifiers::WorldSpaceNormals ) )
@@ -1813,6 +1823,19 @@ bool UsdToUnreal::ConvertMaterial( const pxr::UsdShadeMaterial& UsdShadeMaterial
 		}
 	}
 
+	// Refraction
+	const bool bHasRefractionValue = UsdShadeConversionImpl::GetFloatParameterValue( Connectable, UnrealIdentifiers::Refraction, 1.5f, ParameterValue, &Material, TexturesCache, &PrimvarToUVIndex );
+	if ( bHasRefractionValue || Material.BlendMode == BLEND_Translucent ) // Force a 1.5 IOR if USD didn't specify a value, as it's USD's fallback value
+	{
+		if ( UMaterialExpression* Expression = UsdShadeConversionImpl::GetExpressionForValue( Material, ParameterValue ) )
+		{
+			Material.Refraction.Expression = Expression;
+			SetOutputIndex( ParameterValue, Material.Refraction.OutputIndex );
+
+			bHasMaterialInfo = true;
+		}
+	}
+
 	// Handle world space normals
 	pxr::UsdPrim UsdPrim = UsdShadeMaterial.GetPrim();
 	if ( pxr::UsdAttribute Attr = UsdPrim.GetAttribute( UnrealIdentifiers::WorldSpaceNormals ) )
diff --git a/Engine/Plugins/Importers/USDImporter/Source/UnrealUSDWrapper/Private/UnrealUSDWrapper.cpp b/Engine/Plugins/Importers/USDImporter/Source/UnrealUSDWrapper/Private/UnrealUSDWrapper.cpp
--- a/Engine/Plugins/Importers/USDImporter/Source/UnrealUSDWrapper/Private/UnrealUSDWrapper.cpp
+++ b/Engine/Plugins/Importers/USDImporter/Source/UnrealUSDWrapper/Private/UnrealUSDWrapper.cpp
@@ -102,6 +102,7 @@ namespace UnrealIdentifiers
 	const TfToken Tangent = TfToken("tangent");
 	const TfToken SubsurfaceColor = TfToken("subsurfaceColor");
 	const TfToken AmbientOcclusion = TfToken("ambientOcclusion");
+	const TfToken Refraction = TfToken("ior");
 
 	const TfToken Surface = TfToken("surface");
 	const TfToken St = TfToken("st");
diff --git a/Engine/Plugins/Importers/USDImporter/Source/UnrealUSDWrapper/Public/UnrealUSDWrapper.h b/Engine/Plugins/Importers/USDImporter/Source/UnrealUSDWrapper/Public/UnrealUSDWrapper.h
--- a/Engine/Plugins/Importers/USDImporter/Source/UnrealUSDWrapper/Public/UnrealUSDWrapper.h
+++ b/Engine/Plugins/Importers/USDImporter/Source/UnrealUSDWrapper/Public/UnrealUSDWrapper.h
@@ -296,6 +296,7 @@ namespace UnrealIdentifiers
 	extern UNREALUSDWRAPPER_API const pxr::TfToken Tangent;
 	extern UNREALUSDWRAPPER_API const pxr::TfToken SubsurfaceColor;
 	extern UNREALUSDWRAPPER_API const pxr::TfToken AmbientOcclusion;
+	extern UNREALUSDWRAPPER_API const pxr::TfToken Refraction;
 
 	// Tokens used mostly for shade material conversion
 	extern UNREALUSDWRAPPER_API const pxr::TfToken Surface;
