USD
- Fixed meshes coming in with no UVs when changing purposes. Materials are still cached in the asset cache but the links between assets and prims are cleared when reparsing the stage. This lead to overwritting the data in MaterialToPrimvarToUVIndex with empty mappings.

diff --git a/Engine/Plugins/Importers/USDImporter/Source/USDSchemas/Private/USDShadeMaterialTranslator.cpp b/Engine/Plugins/Importers/USDImporter/Source/USDSchemas/Private/USDShadeMaterialTranslator.cpp
--- a/Engine/Plugins/Importers/USDImporter/Source/USDSchemas/Private/USDShadeMaterialTranslator.cpp
+++ b/Engine/Plugins/Importers/USDImporter/Source/USDSchemas/Private/USDShadeMaterialTranslator.cpp
@@ -102,8 +102,11 @@ void FUsdShadeMaterialTranslator::CreateAssets()
 	}
 	else if ( Context->MaterialToPrimvarToUVIndex && Context->AssetCache )
 	{
-		// Copy the Material -> Primvar -> UV index mapping from the cached material prim path to this prim path
-		Context->MaterialToPrimvarToUVIndex->FindOrAdd( PrimPath.GetString() ) = Context->MaterialToPrimvarToUVIndex->FindRef( Context->AssetCache->GetPrimForAsset( ConvertedMaterial ) );
+		if ( TMap<FString, int32>* PrimvarToUVIndex = Context->MaterialToPrimvarToUVIndex->Find( Context->AssetCache->GetPrimForAsset( ConvertedMaterial ) ) )
+		{
+			// Copy the Material -> Primvar -> UV index mapping from the cached material prim path to this prim path
+			Context->MaterialToPrimvarToUVIndex->FindOrAdd( PrimPath.GetString() ) = *PrimvarToUVIndex;
+		}
 	}
 
 	if ( ConvertedMaterial )
diff --git a/Engine/Plugins/Importers/USDImporter/Source/USDSchemas/Private/USDSkelRootTranslator.cpp b/Engine/Plugins/Importers/USDImporter/Source/USDSchemas/Private/USDSkelRootTranslator.cpp
--- a/Engine/Plugins/Importers/USDImporter/Source/USDSchemas/Private/USDSkelRootTranslator.cpp
+++ b/Engine/Plugins/Importers/USDImporter/Source/USDSchemas/Private/USDSkelRootTranslator.cpp
@@ -727,7 +727,7 @@ namespace UsdSkelRootTranslatorImpl
 				UsdUtils::FBlendShapeMap* OutBlendShapes = Context->BlendShapesByPath ? &NewBlendShapes : nullptr;
 
 				TMap< FString, TMap< FString, int32 > > Unused;
-				TMap< FString, TMap< FString, int32 > >* MaterialToPrimvarToUVIndex = Context->MaterialToPrimvarToUVIndex ? Context->MaterialToPrimvarToUVIndex : &Unused;
+				const TMap< FString, TMap< FString, int32 > >* MaterialToPrimvarToUVIndex = Context->MaterialToPrimvarToUVIndex ? Context->MaterialToPrimvarToUVIndex : &Unused;
 
 				const bool bContinueTaskChain = UsdSkelRootTranslatorImpl::LoadAllSkeletalData(
 					SkeletonCache.Get(),
