USD Stage Editor: Add a "*" after the layer name if it's dirty.

diff --git a/Engine/Plugins/Importers/USDImporter/Source/USDStageEditor/Private/SUSDLayersTreeView.cpp b/Engine/Plugins/Importers/USDImporter/Source/USDStageEditor/Private/SUSDLayersTreeView.cpp
--- a/Engine/Plugins/Importers/USDImporter/Source/USDStageEditor/Private/SUSDLayersTreeView.cpp
+++ b/Engine/Plugins/Importers/USDImporter/Source/USDStageEditor/Private/SUSDLayersTreeView.cpp
@@ -37,7 +37,7 @@ public:
 			.VAlign( VAlign_Center )
 			[
 				SNew(STextBlock)
-				.Text( TreeItem->LayerModel, &FUsdLayerModel::GetDisplayName )
+				.Text( TreeItem, &FUsdLayerViewModel::GetDisplayName )
 			];
 	}
 };
diff --git a/Engine/Plugins/Importers/USDImporter/Source/USDStageEditor/Private/SUSDStage.cpp b/Engine/Plugins/Importers/USDImporter/Source/USDStageEditor/Private/SUSDStage.cpp
--- a/Engine/Plugins/Importers/USDImporter/Source/USDStageEditor/Private/SUSDStage.cpp
+++ b/Engine/Plugins/Importers/USDImporter/Source/USDStageEditor/Private/SUSDStage.cpp
@@ -258,7 +258,19 @@ void SUsdStage::SetupStageActorDelegates()
 			{
 				if ( this->UsdLayersTreeView && ViewModel.UsdStageActor.IsValid() )
 				{
-					this->UsdLayersTreeView->Refresh( ViewModel.UsdStageActor.Get(), false );
+					constexpr bool bResync = false;
+					this->UsdLayersTreeView->Refresh( ViewModel.UsdStageActor.Get(), bResync );
+				}
+			}
+		);
+
+		OnLayersChangedHandle = ViewModel.UsdStageActor->GetUsdListener().GetOnLayersChanged().AddLambda(
+			[ this ]( const TArray< FString >& LayersNames )
+			{
+				if ( this->UsdLayersTreeView && ViewModel.UsdStageActor.IsValid() )
+				{
+					constexpr bool bResync = false;
+					this->UsdLayersTreeView->Refresh( ViewModel.UsdStageActor.Get(), bResync );
 				}
 			}
 		);
@@ -275,6 +287,7 @@ void SUsdStage::ClearStageActorDelegates()
 		StageActor->OnActorDestroyed.Remove ( OnActorDestroyedHandle );
 
 		StageActor->GetUsdListener().GetOnStageEditTargetChanged().Remove( OnStageEditTargetChangedHandle );
+		StageActor->GetUsdListener().GetOnLayersChanged().Remove( OnLayersChangedHandle );
 	}
 }
 
diff --git a/Engine/Plugins/Importers/USDImporter/Source/USDStageEditor/Private/SUSDStage.h b/Engine/Plugins/Importers/USDImporter/Source/USDStageEditor/Private/SUSDStage.h
--- a/Engine/Plugins/Importers/USDImporter/Source/USDStageEditor/Private/SUSDStage.h
+++ b/Engine/Plugins/Importers/USDImporter/Source/USDStageEditor/Private/SUSDStage.h
@@ -79,6 +79,7 @@ protected:
 	FDelegateHandle OnStageChangedHandle;
 	FDelegateHandle OnStageEditTargetChangedHandle;
 	FDelegateHandle OnPrimChangedHandle;
+	FDelegateHandle OnLayersChangedHandle;
 
 	FDelegateHandle OnViewportSelectionChangedHandle;
 
diff --git a/Engine/Plugins/Importers/USDImporter/Source/USDStageEditorViewModels/Private/USDLayersViewModel.cpp b/Engine/Plugins/Importers/USDImporter/Source/USDStageEditorViewModels/Private/USDLayersViewModel.cpp
--- a/Engine/Plugins/Importers/USDImporter/Source/USDStageEditorViewModels/Private/USDLayersViewModel.cpp
+++ b/Engine/Plugins/Importers/USDImporter/Source/USDStageEditorViewModels/Private/USDLayersViewModel.cpp
@@ -143,9 +143,15 @@ void FUsdLayerViewModel::RefreshData()
 
 	const TUsdStore< std::string > UsdLayerIdentifier = UnrealToUsd::ConvertString( *LayerIdentifier );
 
-	LayerModel->DisplayName = FText::FromString( UsdToUnreal::ConvertString( pxr::SdfLayer::GetDisplayNameFromIdentifier( UsdLayerIdentifier.Get() ) ) );
+	LayerModel->DisplayName = UsdToUnreal::ConvertString( pxr::SdfLayer::GetDisplayNameFromIdentifier( UsdLayerIdentifier.Get() ) );
 	LayerModel->bIsMuted = UsdStageRef->IsLayerMuted( UsdLayerIdentifier.Get() );
 
+	pxr::SdfLayerRefPtr UsdLayer = UE::FSdfLayer::FindOrOpen( *LayerIdentifier );
+	if ( UsdLayer )
+	{
+		LayerModel->bIsDirty = UsdLayer->IsDirty();
+	}
+
 	const pxr::SdfLayerHandle& EditTargetLayer = UsdStageRef->GetEditTarget().GetLayer();
 	LayerModel->bIsEditTarget = ( EditTargetLayer
 		? FCStringAnsi::Stricmp( EditTargetLayer->GetIdentifier().c_str(), UsdLayerIdentifier.Get().c_str() ) == 0
@@ -163,21 +169,21 @@ UE::FSdfLayer FUsdLayerViewModel::GetLayer() const
 	return UE::FSdfLayer::FindOrOpen( *LayerIdentifier );
 }
 
-bool FUsdLayerViewModel::IsLayerMuted() const
+FText FUsdLayerViewModel::GetDisplayName() const
 {
-	bool bIsMuted = false;
-
-#if USE_USD_SDK
-	FScopedUsdAllocs UsdAllocs;
-
-	const TUsdStore< std::string > UsdLayerIdentifier = UnrealToUsd::ConvertString( *LayerIdentifier );
+	FString DisplayName = LayerModel->DisplayName;
 
-	pxr::UsdStageRefPtr UsdStageRef( UsdStage );
+	if ( IsLayerDirty() )
+	{
+		DisplayName += TEXT("*");
+	}
 
-	bIsMuted = UsdStageRef->IsLayerMuted( UsdLayerIdentifier.Get() );
-#endif // USE_USD_SDK
+	return FText::FromString(DisplayName);
+}
 
-	return bIsMuted;
+bool FUsdLayerViewModel::IsLayerMuted() const
+{
+	return LayerModel->bIsMuted;
 }
 
 bool FUsdLayerViewModel::CanMuteLayer() const
@@ -270,3 +276,8 @@ bool FUsdLayerViewModel::RemoveSubLayer( int32 SubLayerIndex )
 
 	return bLayerRemoved;
 }
+
+bool FUsdLayerViewModel::IsLayerDirty() const
+{
+	return LayerModel->bIsDirty;
+}
diff --git a/Engine/Plugins/Importers/USDImporter/Source/USDStageEditorViewModels/Public/USDLayersViewModel.h b/Engine/Plugins/Importers/USDImporter/Source/USDStageEditorViewModels/Public/USDLayersViewModel.h
--- a/Engine/Plugins/Importers/USDImporter/Source/USDStageEditorViewModels/Public/USDLayersViewModel.h
+++ b/Engine/Plugins/Importers/USDImporter/Source/USDStageEditorViewModels/Public/USDLayersViewModel.h
@@ -11,11 +11,10 @@
 class USDSTAGEEDITORVIEWMODELS_API FUsdLayerModel : public TSharedFromThis< FUsdLayerModel >
 {
 public:
-	FText GetDisplayName() const { return DisplayName; }
-
-	FText DisplayName;
+	FString DisplayName;
 	bool bIsEditTarget = false;
 	bool bIsMuted = false;
+	bool bIsDirty = false;
 };
 
 class USDSTAGEEDITORVIEWMODELS_API FUsdLayerViewModel : public IUsdTreeViewItem
@@ -32,6 +31,7 @@ public:
 	void RefreshData();
 
 	UE::FSdfLayer GetLayer() const;
+	FText GetDisplayName() const;
 
 	bool IsLayerMuted() const;
 	bool CanMuteLayer() const;
@@ -44,6 +44,8 @@ public:
 	void NewSubLayer( const TCHAR* SubLayerIdentifier );
 	bool RemoveSubLayer( int32 SubLayerIndex );
 
+	bool IsLayerDirty() const;
+
 public:
 	TSharedRef< FUsdLayerModel > LayerModel;
 	FUsdLayerViewModel* ParentItem;
diff --git a/Engine/Plugins/Importers/USDImporter/Source/UnrealUSDWrapper/Private/USDListener.cpp b/Engine/Plugins/Importers/USDImporter/Source/UnrealUSDWrapper/Private/USDListener.cpp
--- a/Engine/Plugins/Importers/USDImporter/Source/UnrealUSDWrapper/Private/USDListener.cpp
+++ b/Engine/Plugins/Importers/USDImporter/Source/UnrealUSDWrapper/Private/USDListener.cpp
@@ -197,7 +197,8 @@ public:
 protected:
 	void HandleObjectsChangedNotice( const pxr::UsdNotice::ObjectsChanged& Notice, const pxr::UsdStageWeakPtr& Sender );
 	void HandleStageEditTargetChangedNotice( const pxr::UsdNotice::StageEditTargetChanged& Notice, const pxr::UsdStageWeakPtr& Sender );
-	void HandleLayersChangedNotice ( const pxr::SdfNotice::LayersDidChange& Notice );
+	void HandleLayersChangedNotice( const pxr::SdfNotice::LayersDidChange& Notice );
+	void HandleLayerDirtinessChangedNotice( const pxr::SdfNotice::LayerDirtinessChanged& Notice );
 
 	void EmitDeprecatedEvents( const pxr::UsdNotice::ObjectsChanged & Notice );
 
@@ -205,6 +206,7 @@ private:
 	pxr::TfNotice::Key RegisteredObjectsChangedKey;
 	pxr::TfNotice::Key RegisteredStageEditTargetChangedKey;
 	pxr::TfNotice::Key RegisteredLayersChangedKey;
+	pxr::TfNotice::Key RegisteredLayerDirtinessChangedKey;
 #endif // #if USE_USD_SDK
 };
 
@@ -293,6 +295,7 @@ void FUsdListenerImpl::Register( const pxr::UsdStageRefPtr& Stage )
 	}
 
 	RegisteredLayersChangedKey = pxr::TfNotice::Register( pxr::TfWeakPtr< FUsdListenerImpl >( this ), &FUsdListenerImpl::HandleLayersChangedNotice );
+	RegisteredLayerDirtinessChangedKey = pxr::TfNotice::Register( pxr::TfWeakPtr< FUsdListenerImpl >( this ), &FUsdListenerImpl::HandleLayerDirtinessChangedNotice );
 }
 #endif // #if USE_USD_SDK
 
@@ -303,11 +306,12 @@ FUsdListenerImpl::~FUsdListenerImpl()
 	pxr::TfNotice::Revoke( RegisteredObjectsChangedKey );
 	pxr::TfNotice::Revoke( RegisteredStageEditTargetChangedKey );
 	pxr::TfNotice::Revoke( RegisteredLayersChangedKey );
+	pxr::TfNotice::Revoke( RegisteredLayerDirtinessChangedKey );
 #endif // #if USE_USD_SDK
 }
 
 #if USE_USD_SDK
-void FUsdListenerImpl::HandleObjectsChangedNotice( const pxr::UsdNotice::ObjectsChanged & Notice, const pxr::UsdStageWeakPtr & Sender )
+void FUsdListenerImpl::HandleObjectsChangedNotice( const pxr::UsdNotice::ObjectsChanged& Notice, const pxr::UsdStageWeakPtr& Sender )
 {
 	if ( !OnPrimsChanged.IsBound() && !OnStageInfoChanged.IsBound() && !OnObjectsChanged.IsBound() )
 	{
@@ -379,6 +383,19 @@ void FUsdListenerImpl::HandleLayersChangedNotice( const pxr::SdfNotice::LayersDi
 	OnLayersChanged.Broadcast( LayersNames );
 }
 
+void FUsdListenerImpl::HandleLayerDirtinessChangedNotice( const pxr::SdfNotice::LayerDirtinessChanged& Notice )
+{
+	if ( !OnLayersChanged.IsBound() || IsBlocked.GetValue() > 0 )
+	{
+		return;
+	}
+
+	FScopedUnrealAllocs UnrealAllocs;
+
+	TArray< FString > LayersNames; // We don't know which layer dirtiness has changed unfortunately so we'll provide an empty list for now.
+	OnLayersChanged.Broadcast(LayersNames);
+}
+
 void FUsdListenerImpl::EmitDeprecatedEvents( const pxr::UsdNotice::ObjectsChanged& Notice )
 {
 	using namespace pxr;
