USD: Update code that deals with 3D transform tracks to use underlying double channels/values instead of float

diff --git a/Engine/Plugins/Importers/USDImporter/Source/USDUtilities/Private/USDPrimConversion.cpp b/Engine/Plugins/Importers/USDImporter/Source/USDUtilities/Private/USDPrimConversion.cpp
--- a/Engine/Plugins/Importers/USDImporter/Source/USDUtilities/Private/USDPrimConversion.cpp
+++ b/Engine/Plugins/Importers/USDImporter/Source/USDUtilities/Private/USDPrimConversion.cpp
@@ -8,7 +8,7 @@
 #include "USDTypesConversion.h"
 
 #include "Channels/MovieSceneChannelProxy.h"
-#include "Channels/MovieSceneFloatChannel.h"
+#include "Channels/MovieSceneDoubleChannel.h"
 #include "CineCameraActor.h"
 #include "CineCameraComponent.h"
 #include "Components/HierarchicalInstancedStaticMeshComponent.h"
@@ -177,31 +177,31 @@ bool UsdToUnreal::ConvertXformable( const pxr::UsdTyped& Schema, UMovieScene3DTr
 	TArray< FFrameNumber > FrameNumbers;
 	FrameNumbers.Reserve( UsdTimeSamples.size() );
 
-	TArray< FMovieSceneFloatValue > TranslationValuesX;
+	TArray< FMovieSceneDoubleValue > TranslationValuesX;
 	TranslationValuesX.Reserve( UsdTimeSamples.size() );
 
-	TArray< FMovieSceneFloatValue > TranslationValuesY;
+	TArray< FMovieSceneDoubleValue > TranslationValuesY;
 	TranslationValuesY.Reserve( UsdTimeSamples.size() );
 
-	TArray< FMovieSceneFloatValue > TranslationValuesZ;
+	TArray< FMovieSceneDoubleValue > TranslationValuesZ;
 	TranslationValuesZ.Reserve( UsdTimeSamples.size() );
 
-	TArray< FMovieSceneFloatValue > RotationValuesX;
+	TArray< FMovieSceneDoubleValue > RotationValuesX;
 	RotationValuesX.Reserve( UsdTimeSamples.size() );
 
-	TArray< FMovieSceneFloatValue > RotationValuesY;
+	TArray< FMovieSceneDoubleValue > RotationValuesY;
 	RotationValuesY.Reserve( UsdTimeSamples.size() );
 
-	TArray< FMovieSceneFloatValue > RotationValuesZ;
+	TArray< FMovieSceneDoubleValue > RotationValuesZ;
 	RotationValuesZ.Reserve( UsdTimeSamples.size() );
 
-	TArray< FMovieSceneFloatValue > ScaleValuesX;
+	TArray< FMovieSceneDoubleValue > ScaleValuesX;
 	ScaleValuesX.Reserve( UsdTimeSamples.size() );
 
-	TArray< FMovieSceneFloatValue > ScaleValuesY;
+	TArray< FMovieSceneDoubleValue > ScaleValuesY;
 	ScaleValuesY.Reserve( UsdTimeSamples.size() );
 
-	TArray< FMovieSceneFloatValue > ScaleValuesZ;
+	TArray< FMovieSceneDoubleValue > ScaleValuesZ;
 	ScaleValuesZ.Reserve( UsdTimeSamples.size() );
 
 	pxr::UsdStageRefPtr Stage = Schema.GetPrim().GetStage();
@@ -248,7 +248,7 @@ bool UsdToUnreal::ConvertXformable( const pxr::UsdTyped& Schema, UMovieScene3DTr
 	TransformSection->EvalOptions.CompletionMode = EMovieSceneCompletionMode::KeepState;
 	TransformSection->SetRange( TRange< FFrameNumber >::All() );
 
-	TArrayView< FMovieSceneFloatChannel* > Channels = TransformSection->GetChannelProxy().GetChannels< FMovieSceneFloatChannel >();
+	TArrayView< FMovieSceneDoubleChannel* > Channels = TransformSection->GetChannelProxy().GetChannels< FMovieSceneDoubleChannel >();
 
 	check( Channels.Num() >= 9 );
 
@@ -851,7 +851,7 @@ bool UnrealToUsd::ConvertXformable( const UMovieScene3DTransformTrack& MovieScen
 	const double StageTimeCodesPerSecond = UsdPrim.GetStage()->GetTimeCodesPerSecond();
 	const FFrameRate StageFrameRate( StageTimeCodesPerSecond, 1 );
 
-	auto EvaluateChannel = [ &PlaybackRange, &Resolution, &DisplayRate, &SequenceTransform ]( const FMovieSceneFloatChannel* Channel, float DefaultValue ) -> TArray< TPair< FFrameNumber, float > >
+	auto EvaluateChannel = [ &PlaybackRange, &Resolution, &DisplayRate, &SequenceTransform ]( const FMovieSceneDoubleChannel* Channel, double DefaultValue ) -> TArray< TPair< FFrameNumber, float > >
 	{
 		TArray< TPair< FFrameNumber, float > > Values;
 
@@ -865,7 +865,7 @@ bool UnrealToUsd::ConvertXformable( const UMovieScene3DTransformTrack& MovieScen
 			{
 				FFrameNumber KeyTime = FFrameRate::Snap( EvalTime, Resolution, DisplayRate ).FloorToFrame();
 
-				float Result = DefaultValue;
+				double Result = DefaultValue;
 				if ( Channel )
 				{
 					Result = Channel->GetDefault().Get( DefaultValue );
@@ -882,10 +882,10 @@ bool UnrealToUsd::ConvertXformable( const UMovieScene3DTransformTrack& MovieScen
 		return Values;
 	};
 
-	TArrayView< FMovieSceneFloatChannel* > Channels = TransformSection->GetChannelProxy().GetChannels< FMovieSceneFloatChannel >();
+	TArrayView< FMovieSceneDoubleChannel* > Channels = TransformSection->GetChannelProxy().GetChannels< FMovieSceneDoubleChannel >();
 	check( Channels.Num() >= 9 );
 
-	auto GetChannel = [ &Channels ]( const int32 ChannelIndex ) -> const FMovieSceneFloatChannel*
+	auto GetChannel = [ &Channels ]( const int32 ChannelIndex ) -> const FMovieSceneDoubleChannel*
 	{
 		if ( Channels.IsValidIndex( ChannelIndex ) )
 		{
@@ -898,19 +898,19 @@ bool UnrealToUsd::ConvertXformable( const UMovieScene3DTransformTrack& MovieScen
 	};
 
 	// Translation
-	TArray< TPair< FFrameNumber, float > > LocationValuesX = EvaluateChannel( GetChannel(0), 0.f );
-	TArray< TPair< FFrameNumber, float > > LocationValuesY = EvaluateChannel( GetChannel(1), 0.f );
-	TArray< TPair< FFrameNumber, float > > LocationValuesZ = EvaluateChannel( GetChannel(2), 0.f );
+	TArray< TPair< FFrameNumber, float > > LocationValuesX = EvaluateChannel( GetChannel(0), 0.0 );
+	TArray< TPair< FFrameNumber, float > > LocationValuesY = EvaluateChannel( GetChannel(1), 0.0 );
+	TArray< TPair< FFrameNumber, float > > LocationValuesZ = EvaluateChannel( GetChannel(2), 0.0 );
 
 	// Rotation
-	TArray< TPair< FFrameNumber, float > > RotationValuesX = EvaluateChannel( GetChannel(3), 0.f );
-	TArray< TPair< FFrameNumber, float > > RotationValuesY = EvaluateChannel( GetChannel(4), 0.f );
-	TArray< TPair< FFrameNumber, float > > RotationValuesZ = EvaluateChannel( GetChannel(5), 0.f );
+	TArray< TPair< FFrameNumber, float > > RotationValuesX = EvaluateChannel( GetChannel(3), 0.0 );
+	TArray< TPair< FFrameNumber, float > > RotationValuesY = EvaluateChannel( GetChannel(4), 0.0 );
+	TArray< TPair< FFrameNumber, float > > RotationValuesZ = EvaluateChannel( GetChannel(5), 0.0 );
 
 	// Scale
-	TArray< TPair< FFrameNumber, float > > ScaleValuesX = EvaluateChannel( GetChannel(6), 1.f );
-	TArray< TPair< FFrameNumber, float > > ScaleValuesY = EvaluateChannel( GetChannel(7), 1.f );
-	TArray< TPair< FFrameNumber, float > > ScaleValuesZ = EvaluateChannel( GetChannel(8), 1.f );
+	TArray< TPair< FFrameNumber, float > > ScaleValuesX = EvaluateChannel( GetChannel(6), 1.0 );
+	TArray< TPair< FFrameNumber, float > > ScaleValuesY = EvaluateChannel( GetChannel(7), 1.0 );
+	TArray< TPair< FFrameNumber, float > > ScaleValuesZ = EvaluateChannel( GetChannel(8), 1.0 );
 
 	bool bIsDataOutOfSync = false;
 	{
