USD: Fix missing include.

For some reason this only presented an error when I was trying some other targets (for UE-121152). Not sure how it survived this far

diff --git a/Engine/Plugins/Importers/USDImporter/Source/USDStageEditorViewModels/Private/USDStageEditorViewModelsModule.cpp b/Engine/Plugins/Importers/USDImporter/Source/USDStageEditorViewModels/Private/USDStageEditorViewModelsModule.cpp
--- a/Engine/Plugins/Importers/USDImporter/Source/USDStageEditorViewModels/Private/USDStageEditorViewModelsModule.cpp
+++ b/Engine/Plugins/Importers/USDImporter/Source/USDStageEditorViewModels/Private/USDStageEditorViewModelsModule.cpp
@@ -4,6 +4,7 @@
 
 #include "USDMemory.h"
 
+#include "Modules/ModuleManager.h"
 
 class FUsdStageEditorViewModelsModule : public IUsdStageEditorViewModelsModule
 {
