USD: Add option to revert to unanimated state before exporting the level

diff --git a/Engine/Plugins/Importers/USDImporter/Source/USDExporter/Private/LevelExporterUSD.cpp b/Engine/Plugins/Importers/USDImporter/Source/USDExporter/Private/LevelExporterUSD.cpp
--- a/Engine/Plugins/Importers/USDImporter/Source/USDExporter/Private/LevelExporterUSD.cpp
+++ b/Engine/Plugins/Importers/USDImporter/Source/USDExporter/Private/LevelExporterUSD.cpp
@@ -19,6 +19,8 @@
 #include "Engine/LevelStreaming.h"
 #include "Engine/World.h"
 #include "IPythonScriptPlugin.h"
+#include "ISequencer.h"
+#include "LevelEditorSequencerIntegration.h"
 
 namespace UE
 {
@@ -154,6 +156,18 @@ bool ULevelExporterUSD::ExportBinary( UObject* Object, const TCHAR* Type, FArchi
 	// Not only to force actors to spawn, but also to make sure that when we want to bake a landscape it is visible
 	TArray<ULevel*> StreamedInLevels = UE::LevelExporterUSD::Private::StreamInRequiredLevels( World, Options->LevelsToIgnore );
 
+	if ( Options->bIgnoreSequencerAnimations )
+	{
+		for ( const TWeakPtr<ISequencer>& Sequencer : FLevelEditorSequencerIntegration::Get().GetSequencers() )
+		{
+			if ( TSharedPtr<ISequencer> PinnedSequencer = Sequencer.Pin() )
+			{
+				PinnedSequencer->EnterSilentMode();
+				PinnedSequencer->RestorePreAnimatedState();
+			}
+		}
+	}
+
 	// Note how we don't explicitly pass the Options down to Python here: We stash our desired export options on the CDO, and
 	// those are read from Python by executing export_with_cdo_options().
 	Options->CurrentTask = ExportTask;
@@ -163,6 +177,19 @@ bool ULevelExporterUSD::ExportBinary( UObject* Object, const TCHAR* Type, FArchi
 	}
 	Options->CurrentTask = nullptr;
 
+	if ( Options->bIgnoreSequencerAnimations )
+	{
+		for ( const TWeakPtr<ISequencer>& Sequencer : FLevelEditorSequencerIntegration::Get().GetSequencers() )
+		{
+			if ( TSharedPtr<ISequencer> PinnedSequencer = Sequencer.Pin() )
+			{
+				PinnedSequencer->InvalidateCachedData();
+				PinnedSequencer->ForceEvaluate();
+				PinnedSequencer->ExitSilentMode();
+			}
+		}
+	}
+
 	// Return the newly streamed in levels to their old visibilities
 	UE::LevelExporterUSD::Private::StreamOutLevels( StreamedInLevels );
 
diff --git a/Engine/Plugins/Importers/USDImporter/Source/USDExporter/Public/LevelExporterUSDOptions.h b/Engine/Plugins/Importers/USDImporter/Source/USDExporter/Public/LevelExporterUSDOptions.h
--- a/Engine/Plugins/Importers/USDImporter/Source/USDExporter/Public/LevelExporterUSDOptions.h
+++ b/Engine/Plugins/Importers/USDImporter/Source/USDExporter/Public/LevelExporterUSDOptions.h
@@ -84,6 +84,10 @@ public:
 	UPROPERTY( EditAnywhere, config, BlueprintReadWrite, Category = "Export settings", meta = ( ClampMin = "1" ) )
 	FIntPoint LandscapeBakeResolution = FIntPoint( 1024, 1024 );
 
+	/** If true, and if we have a level sequence animating the level during export, it will revert any actor or component to its unanimated state before writing to USD */
+	UPROPERTY( EditAnywhere, config, BlueprintReadWrite, Category = "Export settings" )
+	bool bIgnoreSequencerAnimations = false;
+
 	/** If true, will export sub-levels as separate layers (referenced as sublayers). If false, will collapse all sub-levels in a single exported root layer */
     UPROPERTY( EditAnywhere, config, BlueprintReadWrite, Category = "Sublayers" )
     bool bExportSublayers;
diff --git a/Engine/Plugins/Importers/USDImporter/Source/USDExporter/USDExporter.Build.cs b/Engine/Plugins/Importers/USDImporter/Source/USDExporter/USDExporter.Build.cs
--- a/Engine/Plugins/Importers/USDImporter/Source/USDExporter/USDExporter.Build.cs
+++ b/Engine/Plugins/Importers/USDImporter/Source/USDExporter/USDExporter.Build.cs
@@ -25,13 +25,9 @@ namespace UnrealBuildTool.Rules
 				{
 					"Analytics",
 					"AnalyticsBlueprintLibrary",
-					"EditorStyle",
-                    "GeometryCache",
-                    "PropertyEditor",
-                    "RawMesh",
-                    "RenderCore",
-                    "RHI",
 					"CinematicCamera",
+					"EditorStyle",
+					"GeometryCache",
 					"InputCore",
 					"JsonUtilities",
 					"Landscape",
@@ -41,7 +37,13 @@ namespace UnrealBuildTool.Rules
 					"MeshDescription",
 					"MeshUtilities",
 					"MessageLog",
+					"MovieScene",
+					"PropertyEditor",
 					"PythonScriptPlugin",
+					"RawMesh",
+					"RenderCore",
+					"RHI",
+					"Sequencer",
 					"Slate",
 					"SlateCore",
 					"StaticMeshDescription",
@@ -49,7 +51,7 @@ namespace UnrealBuildTool.Rules
 					"USDClasses",
 					"USDStageImporter", // For USDOptionsWindow
 					"USDUtilities",
-                }
+				}
 			);
 		}
 	}
