USD: Also make subobjects non-transient when importing assets

diff --git a/Engine/Plugins/Importers/USDImporter/Source/USDStageImporter/Private/USDStageImporter.cpp b/Engine/Plugins/Importers/USDImporter/Source/USDStageImporter/Private/USDStageImporter.cpp
--- a/Engine/Plugins/Importers/USDImporter/Source/USDStageImporter/Private/USDStageImporter.cpp
+++ b/Engine/Plugins/Importers/USDImporter/Source/USDStageImporter/Private/USDStageImporter.cpp
@@ -583,6 +583,15 @@ namespace UsdStageImporterImpl
 		MovedAsset->SetFlags(ImportContext.ImportObjectFlags | EObjectFlags::RF_Public | EObjectFlags::RF_Standalone );
 		MovedAsset->ClearFlags(EObjectFlags::RF_Transient | EObjectFlags::RF_DuplicateTransient | EObjectFlags::RF_NonPIEDuplicateTransient);
 
+		// Some subobjects like UStaticMesh::HiResSourceModel->StaticMeshDescriptionBulkData can't be left transient, or else they won't serialize their data.
+		// We probably never want to make them public or standalone if they aren't already though
+		TArray<UObject*> Subobjects;
+		MovedAsset->GetDefaultSubobjects( Subobjects );
+		for ( UObject* Subobject : Subobjects )
+		{
+			Subobject->ClearFlags( EObjectFlags::RF_Transient | EObjectFlags::RF_DuplicateTransient | EObjectFlags::RF_NonPIEDuplicateTransient );
+		}
+
 		// We need to make sure that "dirtying the final package" is not added to the transaction, because if we undo this transaction
 		// the assets should remain on their final destination, so we still want the packages to remain marked as dirty (as they're really not on the disk yet).
 		// If we didn't suppress, the package would become transactional by this call. When undoing, the assets would still remain on the final package,
