USD: Deprecate the AUSDLevelInfo class

diff --git a/Engine/Plugins/Importers/USDImporter/Source/USDExporter/Private/USDLevelInfo.cpp b/Engine/Plugins/Importers/USDImporter/Source/USDExporter/Private/USDLevelInfo.cpp
--- a/Engine/Plugins/Importers/USDImporter/Source/USDExporter/Private/USDLevelInfo.cpp
+++ b/Engine/Plugins/Importers/USDImporter/Source/USDExporter/Private/USDLevelInfo.cpp
@@ -1,21 +1,37 @@
 // Copyright Epic Games, Inc. All Rights Reserved.
 
 #include "USDLevelInfo.h"
+
+#include "LevelExporterUSDOptions.h"
+
+#include "Engine/World.h"
+#include "Exporters/Exporter.h"
 #include "HAL/FileManager.h"
-#include "Misc/Paths.h"
 #include "IPythonScriptPlugin.h"
+#include "Misc/Paths.h"
+#include "UObject/GCObjectScopeGuard.h"
 
-AUSDLevelInfo::AUSDLevelInfo(const FObjectInitializer& ObjectInitializer)
-	: Super(ObjectInitializer)
+ADEPRECATED_USDLevelInfo::ADEPRECATED_USDLevelInfo( const FObjectInitializer& ObjectInitializer )
+	: Super( ObjectInitializer )
 {
 	FileScale = 1.0;
 }
 
-void AUSDLevelInfo::SaveUSD()
+void ADEPRECATED_USDLevelInfo::SaveUSD()
 {
-	if (IPythonScriptPlugin::Get()->IsPythonAvailable())
-	{
-		IPythonScriptPlugin::Get()->ExecPythonCommand(TEXT("import usd_unreal.export_level; usd_unreal.export_level.export_current_level(None)"));
-	}
+	UAssetExportTask* LevelExportTask = NewObject<UAssetExportTask>();
+	FGCObjectScopeGuard ExportTaskGuard( LevelExportTask );
+
+	LevelExportTask->Object = GWorld;
+	LevelExportTask->Options = GetMutableDefault<ULevelExporterUSDOptions>();
+	LevelExportTask->Exporter = nullptr;
+	LevelExportTask->Filename = FilePath.FilePath;
+	LevelExportTask->bSelected = false;
+	LevelExportTask->bReplaceIdentical = true;
+	LevelExportTask->bPrompt = false;
+	LevelExportTask->bUseFileArchive = false;
+	LevelExportTask->bWriteEmptyFiles = false;
+	LevelExportTask->bAutomated = true;
+	UExporter::RunAssetExportTask( LevelExportTask );
 }
 
diff --git a/Engine/Plugins/Importers/USDImporter/Source/USDExporter/Private/USDLevelInfo.h b/Engine/Plugins/Importers/USDImporter/Source/USDExporter/Private/USDLevelInfo.h
--- a/Engine/Plugins/Importers/USDImporter/Source/USDExporter/Private/USDLevelInfo.h
+++ b/Engine/Plugins/Importers/USDImporter/Source/USDExporter/Private/USDLevelInfo.h
@@ -7,12 +7,13 @@
 #include "USDLevelInfo.generated.h"
 
 // This is used by the USDExporter module to dispatch the commands to the Python exporter code
-UCLASS(Blueprintable)
-class AUSDLevelInfo : public AActor
+UCLASS(Blueprintable, Deprecated)
+class ADEPRECATED_USDLevelInfo : public AActor
 {
 	GENERATED_UCLASS_BODY()
 
 public:
+	UE_DEPRECATED( 5.0, "Use the ULevelExporterUSD class to export levels to USD" )
 	UFUNCTION(BlueprintCallable, Category = USD, meta = (CallInEditor = "true"))
 	void SaveUSD();
 
