USD: Update mesh components with the new material whenever a material prim resyncs

diff --git a/Engine/Plugins/Importers/USDImporter/Source/USDSchemas/Private/USDGeomMeshTranslator.cpp b/Engine/Plugins/Importers/USDImporter/Source/USDSchemas/Private/USDGeomMeshTranslator.cpp
--- a/Engine/Plugins/Importers/USDImporter/Source/USDSchemas/Private/USDGeomMeshTranslator.cpp
+++ b/Engine/Plugins/Importers/USDImporter/Source/USDSchemas/Private/USDGeomMeshTranslator.cpp
@@ -1107,31 +1107,22 @@ USceneComponent* FUsdGeomMeshTranslator::CreateComponents()
 	{
 		if ( UStaticMesh* StaticMesh = Cast< UStaticMesh >( Context->AssetCache->GetAssetForPrim( PrimPath.GetString() ) ) )
 		{
-#if WITH_EDITOR
-			// If the prim paths match, it means that it was this prim that created (and so "owns") the static mesh,
-			// so its material assignments will already be directly on the mesh. If they differ, we're using some other prim's mesh,
-			// so we may need material overrides on our component
-			UUsdAssetImportData* UsdImportData = Cast<UUsdAssetImportData>( StaticMesh->AssetImportData );
-			if ( UsdImportData && UsdImportData->PrimPath != PrimPath.GetString() )
-#endif // WITH_EDITOR
+			TArray<UMaterialInterface*> ExistingAssignments;
+			for ( FStaticMaterial& StaticMaterial : StaticMesh->GetStaticMaterials() )
 			{
-				TArray<UMaterialInterface*> ExistingAssignments;
-				for ( FStaticMaterial& StaticMaterial : StaticMesh->GetStaticMaterials() )
-				{
-					ExistingAssignments.Add( StaticMaterial.MaterialInterface );
-				}
-
-				MeshTranslationImpl::SetMaterialOverrides(
-					GetPrim(),
-					ExistingAssignments,
-					*StaticMeshComponent,
-					*Context->AssetCache.Get(),
-					Context->Time,
-					Context->ObjectFlags,
-					Context->bAllowInterpretingLODs,
-					Context->RenderContext
-				);
+				ExistingAssignments.Add( StaticMaterial.MaterialInterface );
 			}
+
+			MeshTranslationImpl::SetMaterialOverrides(
+				GetPrim(),
+				ExistingAssignments,
+				*StaticMeshComponent,
+				*Context->AssetCache.Get(),
+				Context->Time,
+				Context->ObjectFlags,
+				Context->bAllowInterpretingLODs,
+				Context->RenderContext
+			);
 		}
 	}
 
diff --git a/Engine/Plugins/Importers/USDImporter/Source/USDStage/Private/USDStageActor.cpp b/Engine/Plugins/Importers/USDImporter/Source/USDStage/Private/USDStageActor.cpp
--- a/Engine/Plugins/Importers/USDImporter/Source/USDStage/Private/USDStageActor.cpp
+++ b/Engine/Plugins/Importers/USDImporter/Source/USDStage/Private/USDStageActor.cpp
@@ -8,6 +8,7 @@
 #include "USDClassesModule.h"
 #include "USDConversionUtils.h"
 #include "USDErrorUtils.h"
+#include "USDGeomMeshConversion.h"
 #include "USDGeomMeshTranslator.h"
 #include "USDGeomXformableTranslator.h"
 #include "USDLayerUtils.h"
@@ -865,6 +866,7 @@ void AUsdStageActor::OnUsdObjectsChanged( const UsdUtils::FObjectChangesByPath&
 	{
 		RefreshStageTask.EnterProgressFrame();
 
+		const FString& PrimPath = PrimChangedInfo.Key;
 		const bool bIsResync = PrimChangedInfo.Value;
 
 		if ( bIsResync && !bDeselected )
@@ -899,19 +901,58 @@ void AUsdStageActor::OnUsdObjectsChanged( const UsdUtils::FObjectChangesByPath&
 			return !PathToProcess.IsEmpty() && PathsProcessed.Contains( PathToProcess );
 		};
 
-		// Reload assets
+		auto UpdateComponents = [&]( const FString& InPrimPath, const bool bInResync )
 		{
-			UE::FSdfPath AssetsPrimPath = FUsdStageActorImpl::UnwindToNonCollapsedPrim( this, PrimChangedInfo.Key, FUsdSchemaTranslator::ECollapsingType::Assets );
+			UE::FSdfPath ComponentsPrimPath = FUsdStageActorImpl::UnwindToNonCollapsedPrim( this, InPrimPath, FUsdSchemaTranslator::ECollapsingType::Components );
 
-			TSet< FString >& RefreshedAssets = bIsResync ? ResyncedAssets : UpdatedAssets;
+			TSet< FString >& RefreshedComponents = bInResync ? ResyncedComponents : UpdatedComponents;
+
+			if ( !IsPathAlreadyProcessed( RefreshedComponents, ComponentsPrimPath.GetString() ) )
+			{
+				TSharedRef< FUsdSchemaTranslationContext > TranslationContext = FUsdStageActorImpl::CreateUsdSchemaTranslationContext( this, ComponentsPrimPath.GetString() );
+				UpdatePrim( ComponentsPrimPath, bInResync, *TranslationContext );
+				TranslationContext->CompleteTasks();
+
+				RefreshedComponents.Add( ComponentsPrimPath.GetString() );
+
+				if ( bInResync )
+				{
+					// Consider that the path has been updated in the case of a resync
+					UpdatedComponents.Add( ComponentsPrimPath.GetString() );
+				}
+			}
+		};
+
+		auto ReloadAssets = [&]( const FString& InPrimPath, const bool bInResync )
+		{
+			UE::FSdfPath AssetsPrimPath = FUsdStageActorImpl::UnwindToNonCollapsedPrim( this, InPrimPath, FUsdSchemaTranslator::ECollapsingType::Assets );
+
+			TSet< FString >& RefreshedAssets = bInResync ? ResyncedAssets : UpdatedAssets;
 
 			if ( !IsPathAlreadyProcessed( RefreshedAssets, AssetsPrimPath.GetString() ) )
 			{
 				TSharedRef< FUsdSchemaTranslationContext > TranslationContext = FUsdStageActorImpl::CreateUsdSchemaTranslationContext( this, AssetsPrimPath.GetString() );
 
-				if ( bIsResync )
+				if ( bInResync )
 				{
-					LoadAssets( *TranslationContext, GetOrLoadUsdStage().GetPrimAtPath( AssetsPrimPath ) );
+					UMaterialInterface* ExistingMaterial = Cast<UMaterialInterface>( AssetCache->GetAssetForPrim( AssetsPrimPath.GetString() ) );
+
+					UE::FUsdPrim PrimToResync = GetOrLoadUsdStage().GetPrimAtPath( AssetsPrimPath );
+					LoadAssets( *TranslationContext, PrimToResync );
+
+					UMaterialInterface* NewMaterial = Cast<UMaterialInterface>( AssetCache->GetAssetForPrim( AssetsPrimPath.GetString() ) );
+
+					// For UE-120185: If we recreated a material for a prim path we also need to update all components that were using it.
+					// This could be fleshed out further if other asset types require this refresh of "dependent components" but materials
+					// seem to be the only ones that do at the moment
+					if ( ExistingMaterial && ( ExistingMaterial != NewMaterial ) )
+					{
+						for ( const FString& MaterialUserPrim : UsdUtils::GetMaterialUsers( PrimToResync ) )
+						{
+							const bool bResyncComponent = true; // We need to force resync to reassign materials
+							UpdateComponents( MaterialUserPrim, bResyncComponent );
+						}
+					}
 
 					// Resyncing also includes "updating" the prim
 					UpdatedAssets.Add( AssetsPrimPath.GetString() );
@@ -923,29 +964,10 @@ void AUsdStageActor::OnUsdObjectsChanged( const UsdUtils::FObjectChangesByPath&
 
 				RefreshedAssets.Add( AssetsPrimPath.GetString() );
 			}
-		}
-
-		// Update components
-		{
-			UE::FSdfPath ComponentsPrimPath = FUsdStageActorImpl::UnwindToNonCollapsedPrim( this, PrimChangedInfo.Key, FUsdSchemaTranslator::ECollapsingType::Components );
-
-			TSet< FString >& RefreshedComponents = bIsResync ? ResyncedComponents : UpdatedComponents;
-
-			if ( !IsPathAlreadyProcessed( RefreshedComponents, ComponentsPrimPath.GetString() ) )
-			{
-				TSharedRef< FUsdSchemaTranslationContext > TranslationContext = FUsdStageActorImpl::CreateUsdSchemaTranslationContext( this, ComponentsPrimPath.GetString() );
-				UpdatePrim( ComponentsPrimPath, bIsResync, *TranslationContext );
-				TranslationContext->CompleteTasks();
-
-				RefreshedComponents.Add( ComponentsPrimPath.GetString() );
+		};
 
-				if ( bIsResync )
-				{
-					// Consider that the path has been updated in the case of a resync
-					UpdatedComponents.Add( ComponentsPrimPath.GetString() );
-				}
-			}
-		}
+		ReloadAssets( PrimPath, bIsResync );
+		UpdateComponents( PrimPath, bIsResync );
 
 		if ( HasAuthorityOverStage() )
 		{
diff --git a/Engine/Plugins/Importers/USDImporter/Source/USDUtilities/Private/USDGeomMeshConversion.cpp b/Engine/Plugins/Importers/USDImporter/Source/USDUtilities/Private/USDGeomMeshConversion.cpp
--- a/Engine/Plugins/Importers/USDImporter/Source/USDUtilities/Private/USDGeomMeshConversion.cpp
+++ b/Engine/Plugins/Importers/USDImporter/Source/USDUtilities/Private/USDGeomMeshConversion.cpp
@@ -37,6 +37,7 @@
 	#include "pxr/usd/sdf/types.h"
 	#include "pxr/usd/usd/editContext.h"
 	#include "pxr/usd/usd/prim.h"
+	#include "pxr/usd/usd/primRange.h"
 	#include "pxr/usd/usdGeom/mesh.h"
 	#include "pxr/usd/usdGeom/primvarsAPI.h"
 	#include "pxr/usd/usdGeom/subset.h"
@@ -1322,6 +1323,47 @@ UsdUtils::FUsdPrimMaterialAssignmentInfo UsdUtils::GetPrimMaterialAssignments( c
 	return Result;
 }
 
+TArray<FString> UsdUtils::GetMaterialUsers( const UE::FUsdPrim& MaterialPrim )
+{
+	TArray<FString> Result;
+
+	FScopedUsdAllocs Allocs;
+
+	pxr::UsdPrim UsdMaterialPrim{ MaterialPrim };
+	if ( !UsdMaterialPrim || !UsdMaterialPrim.IsA<pxr::UsdShadeMaterial>() )
+	{
+		return Result;
+	}
+
+	pxr::UsdStageRefPtr UsdStage = UsdMaterialPrim.GetStage();
+
+	pxr::UsdPrimRange PrimRange = pxr::UsdPrimRange::Stage( UsdStage, pxr::UsdTraverseInstanceProxies() );
+	for ( pxr::UsdPrimRange::iterator PrimRangeIt = PrimRange.begin(); PrimRangeIt != PrimRange.end(); ++PrimRangeIt )
+	{
+		pxr::UsdPrim Prim = *PrimRangeIt;
+
+		if ( !Prim.HasAPI<pxr::UsdShadeMaterialBindingAPI>() )
+		{
+			continue;
+		}
+
+		pxr::UsdShadeMaterialBindingAPI BindingAPI( Prim );
+		pxr::UsdShadeMaterial ShadeMaterial = BindingAPI.ComputeBoundMaterial();
+		if ( !ShadeMaterial )
+		{
+			continue;
+		}
+
+		pxr::UsdPrim ShadeMaterialPrim = ShadeMaterial.GetPrim();
+		if ( ShadeMaterialPrim == UsdMaterialPrim )
+		{
+			Result.Add( UsdToUnreal::ConvertPath( Prim.GetPrimPath() ) );
+		}
+	}
+
+	return Result;
+}
+
 bool UnrealToUsd::ConvertStaticMesh( const UStaticMesh* StaticMesh, pxr::UsdPrim& UsdPrim, const pxr::UsdTimeCode TimeCode, UE::FUsdStage* StageForMaterialAssignments, int32 LowestMeshLOD, int32 HighestMeshLOD )
 {
 	FScopedUsdAllocs UsdAllocs;
diff --git a/Engine/Plugins/Importers/USDImporter/Source/USDUtilities/Public/USDGeomMeshConversion.h b/Engine/Plugins/Importers/USDImporter/Source/USDUtilities/Public/USDGeomMeshConversion.h
--- a/Engine/Plugins/Importers/USDImporter/Source/USDUtilities/Public/USDGeomMeshConversion.h
+++ b/Engine/Plugins/Importers/USDImporter/Source/USDUtilities/Public/USDGeomMeshConversion.h
@@ -31,6 +31,10 @@ namespace UsdUtils
 	struct FDisplayColorMaterial;
 	struct FUsdPrimMaterialAssignmentInfo;
 }
+namespace UE
+{
+	class FUsdPrim;
+}
 
 namespace UsdToUnreal
 {
@@ -161,6 +165,12 @@ namespace UsdUtils
 	USDUTILITIES_API FUsdPrimMaterialAssignmentInfo GetPrimMaterialAssignments( const pxr::UsdPrim& UsdPrim, const pxr::UsdTimeCode TimeCode = pxr::UsdTimeCode::EarliestTime(),
 		bool bProvideMaterialIndices = true, const pxr::TfToken& RenderContext = pxr::UsdShadeTokens->universalRenderContext );
 
+	/**
+	 * Returns an array of prim paths to prims within MaterialPrim's stage that have a UsdShadeMaterialBindingAPI, and whose computed material binding points to MaterialPrim.
+	 * The user prim's schema is not even checked, so this list will naturally include regular UsdGeomMesh prims, UsdGeomSubset prims, and any other with the UsdShadeMaterialBindingAPI.
+	 */
+	USDUTILITIES_API TArray<FString> GetMaterialUsers( const UE::FUsdPrim& MaterialPrim );
+
 	/** Returns whether this UsdMesh can be interpreted as a LOD of a mesh with multiple LODs */
 	USDUTILITIES_API bool IsGeomMeshALOD( const pxr::UsdPrim& UsdMeshPrim );
 
