USD: LevelSequence exporter to USD.

Fix startTimeCode not being authored if it was set with the default value of 0;

diff --git a/Engine/Plugins/Importers/USDImporter/Source/USDExporter/Private/LevelExporterUSD.cpp b/Engine/Plugins/Importers/USDImporter/Source/USDExporter/Private/LevelExporterUSD.cpp
--- a/Engine/Plugins/Importers/USDImporter/Source/USDExporter/Private/LevelExporterUSD.cpp
+++ b/Engine/Plugins/Importers/USDImporter/Source/USDExporter/Private/LevelExporterUSD.cpp
@@ -135,7 +135,7 @@ bool ULevelExporterUSD::ExportBinary( UObject* Object, const TCHAR* Type, FArchi
 			// There is a dedicated "Export selected" option that sets this, so let's sync to it
 			if ( ExportTask )
 			{
-				Options->bSelectionOnly = ExportTask->bSelected;
+				Options->Inner.bSelectionOnly = ExportTask->bSelected;
 			}
 
 			const bool bIsImport = false;
@@ -154,9 +154,9 @@ bool ULevelExporterUSD::ExportBinary( UObject* Object, const TCHAR* Type, FArchi
 
 	// The more robust thing is to force-stream-in all levels that are currently invisible/unloaded but we check to export.
 	// Not only to force actors to spawn, but also to make sure that when we want to bake a landscape it is visible
-	TArray<ULevel*> StreamedInLevels = UE::LevelExporterUSD::Private::StreamInRequiredLevels( World, Options->LevelsToIgnore );
+	TArray<ULevel*> StreamedInLevels = UE::LevelExporterUSD::Private::StreamInRequiredLevels( World, Options->Inner.LevelsToIgnore );
 
-	if ( Options->bIgnoreSequencerAnimations )
+	if ( Options->Inner.bIgnoreSequencerAnimations )
 	{
 		for ( const TWeakPtr<ISequencer>& Sequencer : FLevelEditorSequencerIntegration::Get().GetSequencers() )
 		{
@@ -177,7 +177,7 @@ bool ULevelExporterUSD::ExportBinary( UObject* Object, const TCHAR* Type, FArchi
 	}
 	Options->CurrentTask = nullptr;
 
-	if ( Options->bIgnoreSequencerAnimations )
+	if ( Options->Inner.bIgnoreSequencerAnimations )
 	{
 		for ( const TWeakPtr<ISequencer>& Sequencer : FLevelEditorSequencerIntegration::Get().GetSequencers() )
 		{
diff --git a/Engine/Plugins/Importers/USDImporter/Source/USDExporter/Private/LevelExporterUSDOptionsCustomization.cpp b/Engine/Plugins/Importers/USDImporter/Source/USDExporter/Private/LevelExporterUSDOptionsCustomization.cpp
--- a/Engine/Plugins/Importers/USDImporter/Source/USDExporter/Private/LevelExporterUSDOptionsCustomization.cpp
+++ b/Engine/Plugins/Importers/USDImporter/Source/USDExporter/Private/LevelExporterUSDOptionsCustomization.cpp
@@ -3,6 +3,7 @@
 #include "LevelExporterUSDOptionsCustomization.h"
 
 #include "LevelExporterUSDOptions.h"
+#include "LevelSequenceExporterUSDOptions.h"
 
 #include "CoreMinimal.h"
 #include "DetailCategoryBuilder.h"
@@ -30,7 +31,7 @@ namespace UE
 				SLATE_BEGIN_ARGS(SLevelPickerRow) {}
 				SLATE_END_ARGS()
 
-				void Construct( const FArguments& InArgs, const TSharedRef<STableViewBase>& OwnerTableView, TWeakPtr<FString> InEntry, ULevelExporterUSDOptions* Options )
+				void Construct( const FArguments& InArgs, const TSharedRef<STableViewBase>& OwnerTableView, TWeakPtr<FString> InEntry, FLevelExporterUSDOptionsInner* Inner )
 				{
 					STableRow::Construct( STableRow::FArguments(), OwnerTableView );
 
@@ -49,28 +50,28 @@ namespace UE
 						.MaxWidth(20)
 						[
 							SNew( SCheckBox )
-							.IsChecked_Lambda( [ LevelName, Options ]()
+							.IsChecked_Lambda( [ LevelName, Inner ]()
 							{
-								if ( Options )
+								if ( Inner )
 								{
-									return Options->LevelsToIgnore.Contains( LevelName )
+									return Inner->LevelsToIgnore.Contains( LevelName )
 										? ECheckBoxState::Unchecked
 										: ECheckBoxState::Checked;
 								}
 
 								return ECheckBoxState::Undetermined;
 							})
-							.OnCheckStateChanged_Lambda( [ LevelName, Options ]( ECheckBoxState State )
+							.OnCheckStateChanged_Lambda( [ LevelName, Inner ]( ECheckBoxState State )
 							{
-								if ( Options )
+								if ( Inner )
 								{
 									if ( State == ECheckBoxState::Checked )
 									{
-										Options->LevelsToIgnore.Remove( LevelName );
+										Inner->LevelsToIgnore.Remove( LevelName );
 									}
 									else if ( State == ECheckBoxState::Unchecked )
 									{
-										Options->LevelsToIgnore.Add( LevelName );
+										Inner->LevelsToIgnore.Add( LevelName );
 									}
 								}
 							})
@@ -91,42 +92,42 @@ namespace UE
 			class SLevelPickerList : public SListView<TSharedRef<FString>>
 			{
 			public:
-				void Construct( const FArguments& InArgs, ULevelExporterUSDOptions* Options )
+				void Construct( const FArguments& InArgs, FLevelExporterUSDOptionsInner* Inner )
 				{
-					// We'll be writing directly to its LevelsToIgnore, so keep it alive while we're alive
-					OptionsPtr.Reset( Options );
-
-					Options->LevelsToIgnore.Reset();
-
-					if ( UWorld* EditorWorld = GEditor->GetEditorWorldContext().World() )
+					if ( Inner )
 					{
-						// Make sure all streamed levels are loaded so we can query their names and export them
-						const bool bForce = true;
-						EditorWorld->LoadSecondaryLevels( bForce );
+						Inner->LevelsToIgnore.Reset();
 
-						if ( ULevel* PersistentLevel = EditorWorld->PersistentLevel )
+						if ( UWorld* EditorWorld = GEditor->GetEditorWorldContext().World() )
 						{
-							const FString LevelName = TEXT( "Persistent Level" );
-							RootItems.Add( MakeShared< FString >( LevelName ) );
+							// Make sure all streamed levels are loaded so we can query their names and export them
+							const bool bForce = true;
+							EditorWorld->LoadSecondaryLevels( bForce );
 
-							if ( !PersistentLevel->bIsVisible )
+							if ( ULevel* PersistentLevel = EditorWorld->PersistentLevel )
 							{
-								Options->LevelsToIgnore.Add( LevelName );
+								const FString LevelName = TEXT( "Persistent Level" );
+								RootItems.Add( MakeShared< FString >( LevelName ) );
+
+								if ( !PersistentLevel->bIsVisible )
+								{
+									Inner->LevelsToIgnore.Add( LevelName );
+								}
 							}
-						}
 
-						for ( ULevelStreaming* StreamingLevel : EditorWorld->GetStreamingLevels() )
-						{
-							if ( StreamingLevel )
+							for ( ULevelStreaming* StreamingLevel : EditorWorld->GetStreamingLevels() )
 							{
-								if ( ULevel* Level = StreamingLevel->GetLoadedLevel() )
+								if ( StreamingLevel )
 								{
-									const FString LevelName = Level->GetTypedOuter<UWorld>()->GetName();
-									RootItems.Add( MakeShared< FString >( *LevelName ) );
-
-									if ( !Level->bIsVisible )
+									if ( ULevel* Level = StreamingLevel->GetLoadedLevel() )
 									{
-										Options->LevelsToIgnore.Add( LevelName );
+										const FString LevelName = Level->GetTypedOuter<UWorld>()->GetName();
+										RootItems.Add( MakeShared< FString >( *LevelName ) );
+
+										if ( !Level->bIsVisible )
+										{
+											Inner->LevelsToIgnore.Add( LevelName );
+										}
 									}
 								}
 							}
@@ -138,17 +139,16 @@ namespace UE
 						SListView::FArguments()
 						.ListItemsSource( &RootItems )
 						.SelectionMode( ESelectionMode::None )
-						.OnGenerateRow( this, &SLevelPickerList::OnGenerateRow, Options )
+						.OnGenerateRow( this, &SLevelPickerList::OnGenerateRow, Inner )
 					);
 				}
 
 			private:
-				TSharedRef< ITableRow > OnGenerateRow( TSharedRef<FString> InEntry, const TSharedRef<STableViewBase>& OwnerTable, ULevelExporterUSDOptions* Options ) const
+				TSharedRef< ITableRow > OnGenerateRow( TSharedRef<FString> InEntry, const TSharedRef<STableViewBase>& OwnerTable, FLevelExporterUSDOptionsInner* Inner ) const
 				{
-					return SNew( SLevelPickerRow, OwnerTable, InEntry, Options );
+					return SNew( SLevelPickerRow, OwnerTable, InEntry, Inner );
 				}
 
-				TStrongObjectPtr<ULevelExporterUSDOptions> OptionsPtr;
 				TArray< TSharedRef<FString> > RootItems;
 			};
 		}
@@ -169,16 +169,37 @@ void FLevelExporterUSDOptionsCustomization::CustomizeDetails(IDetailLayoutBuilde
 		return;
 	}
 
-	ULevelExporterUSDOptions* Options = Cast< ULevelExporterUSDOptions>( SelectedObjects[ 0 ].Get() );
-	if ( !Options )
+	TSharedPtr< LevelExporterUSDImpl::SLevelPickerList > PickerTree = nullptr;
+	TStrongObjectPtr<UObject> OptionsPtr;
+	FName LevelFilterPropName;
+	FName ExportSublayersPropName;
+	if ( ULevelExporterUSDOptions* Options = Cast< ULevelExporterUSDOptions>( SelectedObjects[ 0 ].Get() ) )
 	{
-		return;
+		OptionsPtr.Reset( Options );
+
+		PickerTree = SNew( LevelExporterUSDImpl::SLevelPickerList, &Options->Inner );
+		LevelFilterPropName = TEXT( "Inner.LevelsToIgnore" );
+		ExportSublayersPropName = TEXT( "Inner.bExportSublayers" );
 	}
+	else if ( ULevelSequenceExporterUsdOptions* LevelSequenceOptions = Cast< ULevelSequenceExporterUsdOptions>( SelectedObjects[ 0 ].Get() ) )
+	{
+		OptionsPtr.Reset( LevelSequenceOptions );
+
+		// For now there is no easy way of fetching the level to export from a ULevelSequence... we could potentially try to guess what it is
+		// by looking at the soft object paths, but even those aren't exposed, so here we just default to using the current level as the export level.
+		LevelSequenceOptions->Level = GWorld;
 
-	TSharedPtr< LevelExporterUSDImpl::SLevelPickerList > PickerTree = SNew( LevelExporterUSDImpl::SLevelPickerList, Options );
+		PickerTree = SNew( LevelExporterUSDImpl::SLevelPickerList, &LevelSequenceOptions->LevelExportOptions );
+		LevelFilterPropName = TEXT( "LevelExportOptions.LevelsToIgnore" );
+		ExportSublayersPropName = TEXT( "LevelExportOptions.bExportSublayers" );
+	}
+	else
+	{
+		return;
+	}
 
-	TSharedRef<IPropertyHandle> LevelFilterProp = DetailLayoutBuilder.GetProperty( GET_MEMBER_NAME_CHECKED( ULevelExporterUSDOptions, LevelsToIgnore ) );
-	TSharedRef<IPropertyHandle> ExportSublayersProp = DetailLayoutBuilder.GetProperty( GET_MEMBER_NAME_CHECKED( ULevelExporterUSDOptions, bExportSublayers ) );
+	TSharedRef<IPropertyHandle> LevelFilterProp = DetailLayoutBuilder.GetProperty( LevelFilterPropName );
+	TSharedRef<IPropertyHandle> ExportSublayersProp = DetailLayoutBuilder.GetProperty( ExportSublayersPropName );
 
 	// Touch these properties and categories to enforce this ordering
 	DetailLayoutBuilder.EditCategory( TEXT( "Stage options" ) );
diff --git a/Engine/Plugins/Importers/USDImporter/Source/USDExporter/Private/LevelSequenceExporterUSD.cpp b/Engine/Plugins/Importers/USDImporter/Source/USDExporter/Private/LevelSequenceExporterUSD.cpp
new file mode 100644
--- /dev/null
+++ b/Engine/Plugins/Importers/USDImporter/Source/USDExporter/Private/LevelSequenceExporterUSD.cpp
@@ -0,0 +1,844 @@
+// Copyright Epic Games, Inc. All Rights Reserved.
+
+#include "LevelSequenceExporterUSD.h"
+
+#include "LevelSequenceExporterUSDOptions.h"
+#include "UnrealUSDWrapper.h"
+#include "USDConversionUtils.h"
+#include "USDLayerUtils.h"
+#include "USDLog.h"
+#include "USDOptionsWindow.h"
+#include "USDPrimConversion.h"
+
+#include "UsdWrappers/SdfLayer.h"
+#include "UsdWrappers/SdfPath.h"
+#include "UsdWrappers/UsdAttribute.h"
+#include "UsdWrappers/UsdPrim.h"
+#include "UsdWrappers/UsdStage.h"
+
+#include "AssetExportTask.h"
+#include "Components/SkeletalMeshComponent.h"
+#include "ISequencer.h"
+#include "ISequencerModule.h"
+#include "LevelSequence.h"
+#include "Misc/LevelSequenceEditorSpawnRegister.h"
+#include "Modules/ModuleManager.h"
+#include "MovieSceneSequence.h"
+#include "MovieSceneSpawnRegister.h"
+#include "MovieSceneTimeHelpers.h"
+#include "MovieSceneTrack.h"
+#include "Sections/MovieSceneSubSection.h"
+#include "Tracks/MovieScene3DTransformTrack.h"
+#include "Tracks/MovieSceneBoolTrack.h"
+#include "Tracks/MovieSceneFloatTrack.h"
+#include "Tracks/MovieScenePropertyTrack.h"
+#include "Tracks/MovieSceneSkeletalAnimationTrack.h"
+#include "Tracks/MovieSceneSpawnTrack.h"
+#include "Tracks/MovieSceneSubTrack.h"
+#include "Tracks/MovieSceneVisibilityTrack.h"
+#include "UObject/UObjectGlobals.h"
+
+namespace UE
+{
+	namespace LevelSequenceExporterUSD
+	{
+		namespace Private
+		{
+			// Custom spawn register so that when DestroySpawnedObject is called while bDestroyingJustHides is true we
+			// actually just hide the objects, so that we can keep a live reference to components within the bakers.
+			// We're going to convert the spawnable tracks into visibility tracks when exporting to USD, which
+			// also works well with this approach
+			class FLevelSequenceHidingSpawnRegister : public FLevelSequenceEditorSpawnRegister
+			{
+			public:
+				bool bDestroyingJustHides = true;
+
+				virtual UObject* SpawnObject( FMovieSceneSpawnable& Spawnable, FMovieSceneSequenceIDRef TemplateID, IMovieScenePlayer& Player ) override
+				{
+					UObject* Object = ExistingSpawns.FindRef( Spawnable.GetGuid() );
+
+					if ( !Object )
+					{
+						Object = FLevelSequenceEditorSpawnRegister::SpawnObject( Spawnable, TemplateID, Player );
+						ExistingSpawns.Add( Spawnable.GetGuid(), Object );
+					}
+
+					if ( Object )
+					{
+						if ( USceneComponent* Component = Cast<USceneComponent>( Object ) )
+						{
+							const bool bNewVisibility = true;
+							Component->SetVisibility( bNewVisibility );
+						}
+						else if ( AActor* Actor = Cast<AActor>( Object ) )
+						{
+							const bool bIsHidden = false;
+							Actor->SetHidden( bIsHidden );
+						}
+					}
+
+					return Object;
+				}
+
+				virtual void DestroySpawnedObject( UObject& Object ) override
+				{
+					if ( bDestroyingJustHides )
+					{
+						if ( USceneComponent* Component = Cast<USceneComponent>( &Object ) )
+						{
+							const bool bNewVisibility = false;
+							Component->SetVisibility( bNewVisibility );
+						}
+						else if ( AActor* Actor = Cast<AActor>( &Object ) )
+						{
+							const bool bIsHidden = true;
+							Actor->SetHidden( bIsHidden );
+						}
+					}
+					else
+					{
+						FLevelSequenceEditorSpawnRegister::DestroySpawnedObject( Object );
+						if ( const FGuid* ExistingKey = ExistingSpawns.FindKey( &Object ) )
+						{
+							ExistingSpawns.Remove( *ExistingKey );
+						}
+					}
+				}
+
+				// FLevelSequenceHidingSpawnRegister is a bit of a hack and just hides it's spawned actors instead
+				// of deleting them (when we want it to do so). Unfortunately, the base FMovieSceneSpawnRegister part
+				// will still nevertheless clear it's Register entry for the spawnable when deleting (even if just hiding),
+				// and there's nothing we can do to prevent it. This means we can't call FindSpawnedObject and must use our
+				// own GetExistingSpawn and ExistingSpawns member
+				UObject* GetExistingSpawn( const FGuid& Guid )
+				{
+					return ExistingSpawns.FindRef( Guid );
+				}
+
+			private:
+				TMap<FGuid, UObject*> ExistingSpawns;
+			};
+
+			// Contain all of the baker lambda functions for a given component. Only one baker per baking type is allowed.
+			struct FCombinedComponentBakers
+			{
+				UnrealToUsd::EBakingType CombinedBakingType = UnrealToUsd::EBakingType::None;
+				TArray<TFunction<void( double )>> Bakers;
+			};
+
+			class FLevelSequenceExportContext
+			{
+			public:
+				FLevelSequenceExportContext( ULevelSequence& InSequence, TSharedRef<ISequencer> InSequencer, TSharedRef<FLevelSequenceHidingSpawnRegister> InSpawnRegister )
+					: RootSequence( InSequence )
+					, Sequencer( InSequencer )
+					, SpawnRegister( InSpawnRegister )
+				{
+				}
+
+				FLevelSequenceExportContext() = delete;
+				FLevelSequenceExportContext( const FLevelSequenceExportContext& Other ) = delete;
+				FLevelSequenceExportContext( const FLevelSequenceExportContext&& Other ) = delete;
+				FLevelSequenceExportContext& operator=( const FLevelSequenceExportContext& Other ) = delete;
+				FLevelSequenceExportContext& operator=( const FLevelSequenceExportContext&& Other ) = delete;
+
+			public:
+				// The actual content asset that is being exported
+				ULevelSequence& RootSequence;
+
+				ULevelSequenceExporterUsdOptions* ExportOptions;
+
+				// Our own read-only sequencer that we use to play the level sequences while we bake them out one frame at a time
+				TSharedRef<ISequencer> Sequencer;
+
+				// Object that manages spawned instances for FMovieScenePossessables
+				TSharedRef<FLevelSequenceHidingSpawnRegister> SpawnRegister;
+
+				// Used to keep track of which sequences we already baked
+				TMap<UMovieSceneSequence*, FString> ExportedMovieScenes;
+
+				// File paths that we already used during this export.
+				// Used so that we can prevent conflicts between files emitted for this export, but can still
+				// overwrite other files on disk
+				TSet<FString> UsedFilePaths;
+
+				// File path of the exported USD root layer, in case we also exported the level along with the level sequence
+				FString LevelFilePath;
+
+			public:
+				~FLevelSequenceExportContext()
+				{
+					SpawnRegister->bDestroyingJustHides = false;
+					SpawnRegister->CleanUp( *Sequencer );
+				}
+			};
+
+			// Returns a spawned UObject for a given spawnable from a MovieScene.
+			// This assumes it has already been spawned and is tracked by SpawnRegister.
+			UObject* LocateBoundObject( const UMovieScene* MovieScene, FMovieSceneSpawnable& Spawnable, const TSharedRef<FLevelSequenceHidingSpawnRegister>& SpawnRegister )
+			{
+				// Everything should have been spawned by now so just check our spawn register
+				return SpawnRegister->GetExistingSpawn( Spawnable.GetGuid() );
+			}
+
+			// Returns the UObject possessed by a possessable from a MovieScene
+			UObject* LocateBoundObject( const UMovieSceneSequence& MovieSceneSequence, UMovieScene* MovieScene, const FMovieScenePossessable& Possessable, const TSharedRef<FLevelSequenceHidingSpawnRegister>& SpawnRegister )
+			{
+				const FGuid& Guid = Possessable.GetGuid();
+				const FGuid& ParentGuid = Possessable.GetParent();
+
+				// If we have a parent guid, we must provide the object as a context because really the binding path
+				// will just contain the component name
+				UObject* ParentContext = nullptr;
+				if ( ParentGuid.IsValid() )
+				{
+					if ( FMovieScenePossessable* ParentPossessable = MovieScene->FindPossessable( ParentGuid ) )
+					{
+						ParentContext = LocateBoundObject( MovieSceneSequence, MovieScene, *ParentPossessable, SpawnRegister );
+					}
+					else if ( FMovieSceneSpawnable* ParentSpawnable = MovieScene->FindSpawnable( ParentGuid ) )
+					{
+						ParentContext = LocateBoundObject( MovieScene, *ParentSpawnable, SpawnRegister );
+					}
+				}
+
+				TArray<UObject*, TInlineAllocator<1>> Objects = MovieSceneSequence.LocateBoundObjects( Guid, ParentContext );
+				if ( Objects.Num() > 0 )
+				{
+					return Objects[ 0 ];
+				}
+
+				return nullptr;
+			}
+
+			// Returns false if this track has no keys and no skeletal animations, and true otherwise
+			bool IsTrackAnimated( const UMovieSceneTrack* Track )
+			{
+				for ( const UMovieSceneSection* Section : Track->GetAllSections() )
+				{
+					TOptional< TRange<FFrameNumber> > Range = Section->GetAutoSizeRange();
+					if ( Range.IsSet() && !Range.GetValue().IsEmpty() )
+					{
+						return true;
+					}
+				}
+
+				return false;
+			}
+
+			// Recursively go over Sequence and spawn all bound spawnables and then hide them
+			void PreSpawnSpawnables( FLevelSequenceExportContext& Context, UMovieSceneSequence& Sequence )
+			{
+				UMovieScene* MovieScene = Sequence.GetMovieScene();
+				if ( !MovieScene )
+				{
+					return;
+				}
+
+				// Spawn everything on this movie scene
+				int32 NumSpawnables = MovieScene->GetSpawnableCount();
+				for ( int32 Index = 0; Index < NumSpawnables; ++Index )
+				{
+					const FMovieSceneSpawnable& Spawnable = MovieScene->GetSpawnable( Index );
+					const FGuid& Guid = Spawnable.GetGuid();
+
+					UObject* SpawnedObject = StaticCastSharedRef<FMovieSceneSpawnRegister>( Context.SpawnRegister )->SpawnObject( Guid, *MovieScene, Context.Sequencer->GetFocusedTemplateID(), *Context.Sequencer );
+
+					// HACK: Remove the SequencerActor tag from the spawned actor because the USD level exporter
+					// uses these as an indicator that it shouldn't export a given actor, because the transient actors spawned by the
+					// sequencer and the AUsdStageActor contain them.
+					// We will manually delete all of those actors with our spawn register later and never actually use them
+					// for anything else though, so this should be fine
+					if ( AActor* SpawnedActor = Cast<AActor>( SpawnedObject ) )
+					{
+						SpawnedActor->Tags.Remove( TEXT( "SequencerActor" ) );
+					}
+				}
+
+				// Hide all our spawns again so we can pretend they haven't actually spawned yet
+				Context.SpawnRegister->bDestroyingJustHides = true;
+				Context.SpawnRegister->CleanUp( *Context.Sequencer );
+
+				// Recurse to subsequences
+				for ( const UMovieSceneTrack* MasterTrack : MovieScene->GetMasterTracks() )
+				{
+					const UMovieSceneSubTrack* SubTrack = Cast<const UMovieSceneSubTrack>( MasterTrack );
+					if ( !SubTrack )
+					{
+						continue;
+					}
+
+					for ( UMovieSceneSection* Section : SubTrack->GetAllSections() )
+					{
+						UMovieSceneSubSection* SubSection = Cast<UMovieSceneSubSection>( Section );
+						if ( !SubSection )
+						{
+							continue;
+						}
+
+						UMovieSceneSequence* SubSequence = SubSection->GetSequence();
+						if ( !SubSequence )
+						{
+							continue;
+						}
+
+						// Focus on this section before recursing so that every time we spawn a spawnable we can just query
+						// Sequencer->GetFocusedTemplateID() to fetch the FMovieSceneSequenceIDRef for the movie scene that
+						// actually owns the spawnable.
+						// We always reset the sequencer to a new root sequence before baking anyway
+						Context.Sequencer->FocusSequenceInstance( *SubSection );
+
+						PreSpawnSpawnables( Context, *SubSequence );
+					}
+				}
+			}
+
+			// Appends to InOutComponentBakers all of the component bakers for all components bound to MovieSceneSequence.
+			// In the process it will generate the output prims for each of these components, and keep track of them
+			// within the bakers themselves
+			void GenerateBakersForMovieScene(
+				FLevelSequenceExportContext& Context,
+				UMovieSceneSequence& MovieSceneSequence,
+				UE::FUsdStage& UsdStage,
+				TMap<USceneComponent*, FCombinedComponentBakers>& InOutComponentBakers
+			)
+			{
+				UMovieScene* MovieScene = MovieSceneSequence.GetMovieScene();
+				if ( !MovieScene )
+				{
+					return;
+				}
+
+				// Find all components (if actors, the root components of those actors) bound to the movie scene.
+				// We use components here because when exporting actors and components to USD we basically just ignore
+				// actors altogether and export the component attachment hierarchy instead
+				TArray<TPair<FGuid, USceneComponent*>> BoundComponents;
+
+				// Possessables
+				int32 NumPossessables = MovieScene->GetPossessableCount();
+				for ( int32 Index = 0; Index < NumPossessables; ++Index )
+				{
+					const FMovieScenePossessable& Possessable = MovieScene->GetPossessable( Index );
+					const FGuid& Guid = Possessable.GetGuid();
+
+					UObject* BoundObject = LocateBoundObject( MovieSceneSequence, MovieScene, Possessable, Context.SpawnRegister );
+
+					USceneComponent* BoundComponent = Cast<USceneComponent>( BoundObject );
+					if ( !BoundComponent )
+					{
+						if ( const AActor* Actor = Cast<AActor>( BoundObject ) )
+						{
+							BoundComponent = Actor->GetRootComponent();
+						}
+					}
+					if ( !BoundComponent )
+					{
+						continue;
+					}
+
+					BoundComponents.Add( TPair<FGuid, USceneComponent*>( Guid, BoundComponent ) );
+				}
+
+				// Spawnables
+				int32 NumSpawnables = MovieScene->GetSpawnableCount();
+				for ( int32 Index = 0; Index < NumSpawnables; ++Index )
+				{
+					FMovieSceneSpawnable& Spawnable = MovieScene->GetSpawnable( Index );
+					const FGuid& Guid = Spawnable.GetGuid();
+
+					UObject* BoundObject = Context.SpawnRegister->GetExistingSpawn( Guid );
+					if ( !BoundObject )
+					{
+						// This should never happen as we preemptively spawn everything:
+						// At this point all our spawns should be spawned, but invisible
+						UE_LOG( LogUsd, Warning, TEXT( "Failed to find spawned object for spawnable with Guid '%s'" ), *Guid.ToString() );
+						continue;
+					}
+
+					USceneComponent* BoundComponent = Cast<USceneComponent>( BoundObject );
+					if ( !BoundComponent )
+					{
+						if ( const AActor* Actor = Cast<AActor>( BoundObject ) )
+						{
+							BoundComponent = Actor->GetRootComponent();
+						}
+					}
+					if ( !BoundComponent )
+					{
+						continue;
+					}
+
+					BoundComponents.Add( TPair<FGuid, USceneComponent*>( Guid, BoundComponent ) );
+				}
+
+				// Generate bakers
+				for ( const TPair<FGuid, USceneComponent*>& Pair : BoundComponents )
+				{
+					const FGuid& Guid = Pair.Key;
+					USceneComponent* BoundComponent = Pair.Value;
+
+					FString PrimPath = UsdUtils::GetPrimPathForObject(
+						BoundComponent,
+						TEXT(""),
+						Context.ExportOptions && Context.ExportOptions->LevelExportOptions.bExportActorFolders
+					);
+					if ( PrimPath.IsEmpty() )
+					{
+						continue;
+					}
+
+					FString SchemaName = UsdUtils::GetSchemaNameForComponent( *BoundComponent );
+					if ( SchemaName.IsEmpty() )
+					{
+						continue;
+					}
+
+					// We will define a prim here so that we can apply schemas and use the shortcut CreateXAttribute functions,
+					// and not have to worry about attribute names and types. Later on we will convert these prims back into just 'overs' though
+					UE::FUsdPrim Prim = UsdStage.DefinePrim( UE::FSdfPath{ *PrimPath }, *SchemaName );
+					if ( !Prim )
+					{
+						continue;
+					}
+
+					if ( const FMovieSceneBinding* Binding = MovieScene->FindBinding( Guid ) )
+					{
+						for ( const UMovieSceneTrack* Track : Binding->GetTracks() )
+						{
+							if ( !IsTrackAnimated( Track ) )
+							{
+								continue;
+							}
+
+							UnrealToUsd::FComponentBaker Baker;
+
+							if ( const UMovieScenePropertyTrack* PropertyTrack = Cast<UMovieScenePropertyTrack>( Track ) )
+							{
+								const FString PropertyPath = PropertyTrack->GetPropertyPath().ToString();
+								UnrealToUsd::CreateComponentPropertyBaker( Prim, *BoundComponent, PropertyPath, Baker );
+							}
+							else if ( const UMovieSceneSpawnTrack* SpawnTrack = Cast<UMovieSceneSpawnTrack>( Track ) )
+							{
+								// Just handle spawnable tracks as if they're visibility tracks, and hide the prim when not "spawned"
+								// Remember that our spawn register just hides the spawnables when they're not spawned anyway, so this
+								// is essentially the same
+								const FString PropertyPath = TEXT( "bHidden" );
+								UnrealToUsd::CreateComponentPropertyBaker( Prim, *BoundComponent, PropertyPath, Baker);
+							}
+							else if ( const UMovieSceneSkeletalAnimationTrack* SkeletalTrack = Cast<UMovieSceneSkeletalAnimationTrack>( Track ) )
+							{
+								if ( USkeletalMeshComponent* SkeletalBoundComponent = Cast<USkeletalMeshComponent>( BoundComponent ) )
+								{
+									UE::FUsdPrim SkelAnimPrim = UsdStage.DefinePrim( UE::FSdfPath{ *PrimPath }.AppendChild( TEXT( "Anim" ) ), TEXT( "SkelAnimation" ) );
+									if ( !SkelAnimPrim )
+									{
+										UE_LOG( LogUsd, Warning, TEXT( "Failed to generate SkelAnimation prim when baking out SkelRoot '%s'" ), *PrimPath );
+										continue;
+									}
+									UnrealToUsd::CreateSkeletalAnimationBaker( Prim, SkelAnimPrim, *SkeletalBoundComponent, Baker );
+								}
+							}
+
+							// If we made a baker and we don't have one of this type for this component yet, add its lambda to the array
+							FCombinedComponentBakers& ExistingBakers = InOutComponentBakers.FindOrAdd( BoundComponent );
+							if ( Baker.BakerType != UnrealToUsd::EBakingType::None && !EnumHasAnyFlags( ExistingBakers.CombinedBakingType, Baker.BakerType ) )
+							{
+								ExistingBakers.Bakers.Add( Baker.BakerFunction );
+								ExistingBakers.CombinedBakingType |= Baker.BakerType;
+							}
+						}
+					}
+				}
+
+				// We always have to generate bakers for subsequences (even if we're not exporting separate files for subsequences) since we
+				// will bake the full combined, composed level sequence as the main output USD layer. This because USD doesn't allow any form
+				// of animation blending, and so composing individual USD layers exported for each subsequence with the same result as the sequencer
+				// is impossible... we have to settle for having each USD layer represent the full effect of it's level sequence.
+				// Also note that we can't share these bakers with our parent movie scenes in case we're a subsequence, unfortunately, because our
+				// bakers will contain lambdas that write directly to a given prim, and those prims are specific to each layer that we're exporting
+				// (e.g. our parent level sequence will export to different prims than this movie scene will)
+				for ( const UMovieSceneTrack* MasterTrack : MovieScene->GetMasterTracks() )
+				{
+					const UMovieSceneSubTrack* SubTrack = Cast<const UMovieSceneSubTrack>( MasterTrack );
+					if ( !SubTrack )
+					{
+						continue;
+					}
+
+					for ( const UMovieSceneSection* Section : SubTrack->GetAllSections() )
+					{
+						const UMovieSceneSubSection* SubSection = Cast<const UMovieSceneSubSection>( Section );
+						if ( !SubSection )
+						{
+							continue;
+						}
+
+						UMovieSceneSequence* SubSequence = SubSection->GetSequence();
+						if ( !SubSequence )
+						{
+							continue;
+						}
+
+						GenerateBakersForMovieScene( Context, *SubSequence, UsdStage, InOutComponentBakers );
+					}
+				}
+			}
+
+			// Steps through MovieSceneSequence frame by frame, invoking all baker lambdas every frame
+			void BakeMovieSceneSequence(
+				FLevelSequenceExportContext& Context,
+				UMovieSceneSequence& MovieSceneSequence,
+				UE::FUsdStage& UsdStage,
+				const TMap<USceneComponent*, FCombinedComponentBakers>& ComponentBakers
+			)
+			{
+				UMovieScene* MovieScene = MovieSceneSequence.GetMovieScene();
+				if ( !MovieScene )
+				{
+					return;
+				}
+
+				// Currently we bake the full composed result of each level sequence into a single layer,
+				// because USD can't compose individual layers in the same way (with blending and so on). So here
+				// we make sure that our sequencer exports each MovieSceneSequence as if it was a root, emitting the
+				// same result as if we had exported that subsequence's LevelSequence by itself
+				Context.Sequencer->ResetToNewRootSequence( MovieSceneSequence );
+
+				const TRange< FFrameNumber > PlaybackRange = MovieScene->GetPlaybackRange();
+				const FFrameRate Resolution = MovieScene->GetTickResolution();
+				const FFrameRate DisplayRate = MovieScene->GetDisplayRate();
+				const double StageTimeCodesPerSecond = UsdStage.GetTimeCodesPerSecond();
+				const FFrameRate StageFrameRate( StageTimeCodesPerSecond, 1 );
+
+				const FFrameTime Interval = FFrameRate::TransformTime( 1, DisplayRate, Resolution );
+				FFrameNumber StartFrame = UE::MovieScene::DiscreteInclusiveLower( PlaybackRange );
+				FFrameNumber EndFrame = UE::MovieScene::DiscreteExclusiveUpper( PlaybackRange );
+
+				if ( Context.ExportOptions && Context.ExportOptions->bOverrideExportRange )
+				{
+					StartFrame = FFrameRate::TransformTime( Context.ExportOptions->StartFrame, DisplayRate, Resolution ).FloorToFrame();
+					EndFrame = FFrameRate::TransformTime( Context.ExportOptions->EndFrame, DisplayRate, Resolution ).CeilToFrame();
+				}
+
+				FFrameTime StartFrameUETime = FFrameRate::Snap( StartFrame, Resolution, DisplayRate ).FloorToFrame();
+				double StartTimeCode = FFrameRate::TransformTime( StartFrameUETime, Resolution, StageFrameRate ).AsDecimal();
+				UsdStage.SetStartTimeCode( StartTimeCode );
+
+				FFrameTime EndFrameUETime = FFrameRate::Snap( EndFrame, Resolution, DisplayRate ).FloorToFrame();
+				double EndTimeCode = FFrameRate::TransformTime( EndFrameUETime, Resolution, StageFrameRate ).AsDecimal();
+				UsdStage.SetEndTimeCode( EndTimeCode );
+
+				for ( FFrameTime EvalTime = StartFrame; EvalTime <= EndFrame; EvalTime += Interval )
+				{
+					Context.Sequencer->SetLocalTimeDirectly( EvalTime );
+					Context.Sequencer->ForceEvaluate();
+
+					FFrameTime KeyTime = FFrameRate::Snap( EvalTime, Resolution, DisplayRate ).FloorToFrame();
+					double UsdTimeCode = FFrameRate::TransformTime( KeyTime, Resolution, StageFrameRate ).AsDecimal();
+
+					for ( const TPair<USceneComponent*, FCombinedComponentBakers>& Pair : ComponentBakers )
+					{
+						for ( const TFunction<void( double )>& Lambda : Pair.Value.Bakers )
+						{
+							if ( Lambda )
+							{
+								Lambda( UsdTimeCode );
+							}
+						}
+					}
+				}
+
+				// Convert all prims back to overs (so that this layer doesn't define anything on a stage that doesn't
+				// previously have it - it's only supposed to carry animation data)
+				// We need to do this in this way because when going from 'def' to 'over' we need to do it from leaf towards the
+				// root, as USD doesn't like a parent 'over' with a child 'def'
+				TFunction<void( FUsdPrim& )> RecursiveSetOver;
+				RecursiveSetOver = [&RecursiveSetOver]( FUsdPrim& Prim )
+				{
+					for ( FUsdPrim& Child : Prim.GetChildren() )
+					{
+						RecursiveSetOver( Child );
+					}
+
+					Prim.SetSpecifier( ESdfSpecifier::Over );
+				};
+				FUsdPrim Root = UsdStage.GetPseudoRoot();
+				for ( FUsdPrim& TopLevelPrim : Root.GetChildren() )
+				{
+					RecursiveSetOver( TopLevelPrim );
+				}
+			}
+
+			void ExportMovieSceneSequence(
+				FLevelSequenceExportContext& Context,
+				UMovieSceneSequence& MovieSceneSequence,
+				const FString& FilePath
+			)
+			{
+				if ( FilePath.IsEmpty() || Context.ExportedMovieScenes.Contains( &MovieSceneSequence ) )
+				{
+					return;
+				}
+
+				UMovieScene* MovieScene = MovieSceneSequence.GetMovieScene();
+				if ( !MovieScene )
+				{
+					return;
+				}
+
+				// Make sure we don't overwrite a file we just wrote *during this export*.
+				// Overwriting other files is OK, as we want to allow a "repeatedly export over the same files" workflow
+				FString UniqueFilePath = FilePath;
+				if ( Context.UsedFilePaths.Contains( UniqueFilePath ) )
+				{
+					int32 Suffix = 0;
+					do
+					{
+						UniqueFilePath = FString::Printf( TEXT( "%s_%d" ), *FilePath, Suffix++ );
+					} while ( Context.UsedFilePaths.Contains( UniqueFilePath ) );
+				}
+
+				UE::FUsdStage UsdStage = UnrealUSDWrapper::NewStage( *UniqueFilePath );
+				if ( !UsdStage )
+				{
+					return;
+				}
+
+				if ( Context.ExportOptions )
+				{
+					UsdUtils::SetUsdStageMetersPerUnit( UsdStage, Context.ExportOptions->StageOptions.MetersPerUnit );
+					UsdUtils::SetUsdStageUpAxis( UsdStage, Context.ExportOptions->StageOptions.UpAxis );
+					UsdStage.SetTimeCodesPerSecond( Context.ExportOptions->TimeCodesPerSecond );
+				}
+
+				UE::FUsdPrim RootPrim = UsdStage.OverridePrim( UE::FSdfPath( TEXT( "/Root" ) ) );
+				if ( !RootPrim )
+				{
+					return;
+				}
+
+				UsdStage.SetDefaultPrim( RootPrim );
+
+				TMap<USceneComponent*, FCombinedComponentBakers> Bakers;
+				GenerateBakersForMovieScene( Context, MovieSceneSequence, UsdStage, Bakers );
+
+				// Bake this MovieScene
+				// We bake each MovieScene individually instead of doing one large simultaneous bake because this way
+				// not only we avoid having to handle FMovieSceneSequenceTransforms when writing out the UsdTimeCodes,
+				// we guarantee we'll get the same result as if we exported each subsequence individually.
+				// They could have differed, for example, if we had a limited the playback range of a subsequence
+				BakeMovieSceneSequence( Context, MovieSceneSequence, UsdStage, Bakers );
+
+				const TRange< FFrameNumber > PlaybackRange = MovieScene->GetPlaybackRange();
+				const FFrameRate Resolution = MovieScene->GetTickResolution();
+				const FFrameRate DisplayRate = MovieScene->GetDisplayRate();
+				const double StageTimeCodesPerSecond = UsdStage.GetTimeCodesPerSecond();
+				const FFrameRate StageFrameRate( StageTimeCodesPerSecond, 1 );
+
+				if ( Context.ExportOptions )
+				{
+					if ( Context.ExportOptions->bExportSubsequencesAsLayers )
+					{
+						FString Directory;
+						FString FileName;
+						FString Extension;
+						FPaths::Split( UniqueFilePath, Directory, FileName, Extension );
+
+						for ( const UMovieSceneTrack* MasterTrack : MovieScene->GetMasterTracks() )
+						{
+							const UMovieSceneSubTrack* SubTrack = Cast<const UMovieSceneSubTrack>( MasterTrack );
+							if ( !SubTrack )
+							{
+								continue;
+							}
+
+							for ( const UMovieSceneSection* Section : SubTrack->GetAllSections() )
+							{
+								const UMovieSceneSubSection* SubSection = Cast<const UMovieSceneSubSection>( Section );
+								if ( !SubSection )
+								{
+									continue;
+								}
+
+								UMovieSceneSequence* SubSequence = SubSection->GetSequence();
+								if ( !SubSequence )
+								{
+									continue;
+								}
+
+								FString SubSequencePath = FPaths::Combine( Directory, FString::Printf( TEXT( "%s.%s" ), *SubSequence->GetName(), *Extension ) );
+
+								ExportMovieSceneSequence( Context, *SubSequence, SubSequencePath );
+
+								// For now we don't want to actually add the subsequence layers as sublayers since each exported level sequence
+								// contains the full baked result anyway, but this is how we'd do it:
+								/*
+								double Offset = 0.0;
+								if ( SubSection->HasStartFrame() )
+								{
+									FFrameNumber LowerBound = SubSection->GetTrueRange().GetLowerBoundValue();
+									FFrameTime LowerBoundTime{ FFrameRate::Snap( LowerBound, Resolution, DisplayRate ).FloorToFrame() };
+									Offset = FFrameRate::TransformTime( LowerBoundTime, Resolution, StageFrameRate ).AsDecimal();
+								}
+								float Scale = SubSection->Parameters.TimeScale;
+								const int32 Index = -1;
+								UsdUtils::InsertSubLayer( UsdStage.GetRootLayer(), *SubSequencePath, Index, Offset, Scale );
+								*/
+							}
+						}
+					}
+
+					// We can add the level as a sublayer to every exported subsequence, so that each can be opened individually and
+					// automatically load the level layer. It doesn't matter much if the parent stage has composed the level
+					// sublayer multiple times (in case we add subsequence layers as sublayers in the future), as the prims will
+					// just all override each other with the same data
+					if ( Context.ExportOptions->bUseExportedLevelAsSublayer && FPaths::FileExists( Context.LevelFilePath ) )
+					{
+						UsdUtils::InsertSubLayer( UsdStage.GetRootLayer(), *Context.LevelFilePath );
+					}
+				}
+
+				Context.ExportedMovieScenes.Add( &MovieSceneSequence, UniqueFilePath );
+				Context.UsedFilePaths.Add( UniqueFilePath );
+
+				UsdStage.GetRootLayer().Save();
+			}
+		}
+	}
+}
+
+ULevelSequenceExporterUsd::ULevelSequenceExporterUsd()
+{
+#if USE_USD_SDK
+	for ( const FString& Extension : UnrealUSDWrapper::GetAllSupportedFileFormats() )
+	{
+		// USDZ is not supported for writing for now
+		if ( Extension.Equals( TEXT( "usdz" ) ) )
+		{
+			continue;
+		}
+
+		FormatExtension.Add( Extension );
+		FormatDescription.Add( TEXT( "USD file" ) );
+	}
+	SupportedClass = ULevelSequence::StaticClass();
+	bText = false;
+#endif // #if USE_USD_SDK
+}
+
+bool ULevelSequenceExporterUsd::ExportBinary( UObject* Object, const TCHAR* Type, FArchive& Ar, FFeedbackContext* Warn, int32 FileIndex, uint32 PortFlags )
+{
+	namespace LevelSequenceExporterImpl = UE::LevelSequenceExporterUSD::Private;
+
+#if USE_USD_SDK
+	ULevelSequence* LevelSequence = Cast< ULevelSequence >( Object );
+	if ( !LevelSequence )
+	{
+		return false;
+	}
+
+	UMovieScene* MovieScene = LevelSequence->GetMovieScene();
+	if ( !MovieScene )
+	{
+		return false;
+	}
+
+	ULevelSequenceExporterUsdOptions* Options = nullptr;
+	if ( ExportTask )
+	{
+		Options = Cast<ULevelSequenceExporterUsdOptions>( ExportTask->Options );
+	}
+	if ( !Options && ( !ExportTask || !ExportTask->bAutomated ) )
+	{
+		Options = GetMutableDefault<ULevelSequenceExporterUsdOptions>();
+		if ( Options )
+		{
+			const bool bIsImport = false;
+			const bool bContinue = SUsdOptionsWindow::ShowOptions( *Options, bIsImport );
+			if ( !bContinue )
+			{
+				return false;
+			}
+		}
+	}
+
+	TSharedPtr<LevelSequenceExporterImpl::FLevelSequenceHidingSpawnRegister> SpawnRegister = MakeShared<LevelSequenceExporterImpl::FLevelSequenceHidingSpawnRegister>();
+	if ( !SpawnRegister.IsValid() )
+	{
+		return false;
+	}
+
+	FSequencerInitParams Params;
+	Params.RootSequence = LevelSequence;
+	Params.SpawnRegister = SpawnRegister;
+	Params.ViewParams.bReadOnly = true;
+	Params.bEditWithinLevelEditor = false;
+	Params.PlaybackContext = Options->Level.Get();
+
+	// Set to read only or else CreateSequencer will change the playback range of the moviescene without even calling Modify()
+	const bool bOldReadOnly = MovieScene->IsReadOnly();
+	const bool bNewReadOnly = true;
+	MovieScene->SetReadOnly( bNewReadOnly );
+	TSharedPtr<ISequencer> TempSequencer = FModuleManager::LoadModuleChecked<ISequencerModule>( "Sequencer" ).CreateSequencer( Params );
+	MovieScene->SetReadOnly( bOldReadOnly );
+	if ( !TempSequencer.IsValid() )
+	{
+		return false;
+	}
+
+	TempSequencer->EnterSilentMode();
+	TempSequencer->SetPlaybackStatus( EMovieScenePlayerStatus::Playing );
+
+	SpawnRegister->SetSequencer( TempSequencer );
+
+	LevelSequenceExporterImpl::FLevelSequenceExportContext Context{ *LevelSequence, TempSequencer.ToSharedRef(), SpawnRegister.ToSharedRef() };
+	Context.ExportOptions = Options;
+
+	// Spawn (but hide) all spawnables so that they will also show up on the level export if we need them to
+	LevelSequenceExporterImpl::PreSpawnSpawnables( Context, Context.RootSequence );
+
+	// Capture this first because when we launch UExporter::RunAssetExportTask the CurrentFileName will change
+	const FString TargetFileName = UExporter::CurrentFilename;
+
+	// Export level if we need to
+	if ( Options && Options->bExportLevel )
+	{
+		if ( UWorld* WorldToExport = Options->Level.Get() )
+		{
+			// Come up with a file path for the level
+			FString Directory;
+			FString Filename;
+			FString Extension;
+			FPaths::Split( TargetFileName, Directory, Filename, Extension );
+			Context.LevelFilePath = FPaths::Combine( Directory, FString::Printf( TEXT( "%s.%s" ), *GWorld->GetFName().ToString(), *Extension ) );
+
+			Context.UsedFilePaths.Add( Context.LevelFilePath );
+
+			ULevelExporterUSDOptions* LevelOptions = GetMutableDefault<ULevelExporterUSDOptions>();
+			LevelOptions->StageOptions = Options->StageOptions;
+			LevelOptions->Inner = Options->LevelExportOptions;
+
+			UAssetExportTask* LevelExportTask = NewObject<UAssetExportTask>();
+			FGCObjectScopeGuard ExportTaskGuard( LevelExportTask );
+			LevelExportTask->Object = WorldToExport;
+			LevelExportTask->Options = LevelOptions;
+			LevelExportTask->Exporter = nullptr;
+			LevelExportTask->Filename = Context.LevelFilePath;
+			LevelExportTask->bSelected = false;
+			LevelExportTask->bReplaceIdentical = true;
+			LevelExportTask->bPrompt = false;
+			LevelExportTask->bUseFileArchive = false;
+			LevelExportTask->bWriteEmptyFiles = false;
+			LevelExportTask->bAutomated = true; // Pretend this is an automated task so it doesn't pop the dialog
+
+			UExporter::RunAssetExportTask( LevelExportTask );
+		}
+	}
+
+	LevelSequenceExporterImpl::ExportMovieSceneSequence( Context, *LevelSequence, TargetFileName );
+
+	return true;
+#else
+	return false;
+#endif // #if USE_USD_SDK
+}
diff --git a/Engine/Plugins/Importers/USDImporter/Source/USDExporter/Private/USDConversionBlueprintLibrary.cpp b/Engine/Plugins/Importers/USDImporter/Source/USDExporter/Private/USDConversionBlueprintLibrary.cpp
--- a/Engine/Plugins/Importers/USDImporter/Source/USDExporter/Private/USDConversionBlueprintLibrary.cpp
+++ b/Engine/Plugins/Importers/USDImporter/Source/USDExporter/Private/USDConversionBlueprintLibrary.cpp
@@ -113,3 +113,18 @@ void UUsdConversionBlueprintLibrary::AddPayload( const FString& ReferencingStage
 	}
 #endif // USE_USD_SDK
 }
+
+FString UUsdConversionBlueprintLibrary::GetPrimPathForObject( const UObject* ActorOrComponent, const FString& ParentPrimPath, bool bUseActorFolders )
+{
+	return UsdUtils::GetPrimPathForObject( ActorOrComponent, ParentPrimPath, bUseActorFolders );
+}
+
+FString UUsdConversionBlueprintLibrary::GetSchemaNameForComponent( const USceneComponent* Component )
+{
+	if ( !Component )
+	{
+		return {};
+	}
+
+	return UsdUtils::GetSchemaNameForComponent( *Component );
+}
diff --git a/Engine/Plugins/Importers/USDImporter/Source/USDExporter/Private/USDExporterModule.cpp b/Engine/Plugins/Importers/USDImporter/Source/USDExporter/Private/USDExporterModule.cpp
--- a/Engine/Plugins/Importers/USDImporter/Source/USDExporter/Private/USDExporterModule.cpp
+++ b/Engine/Plugins/Importers/USDImporter/Source/USDExporter/Private/USDExporterModule.cpp
@@ -14,7 +14,10 @@ public:
 	virtual void StartupModule() override
 	{
 		FPropertyEditorModule& PropertyModule = FModuleManager::LoadModuleChecked< FPropertyEditorModule >( TEXT( "PropertyEditor" ) );
+
+		// We intentionally use the same customization for both of these
 		PropertyModule.RegisterCustomClassLayout( TEXT( "LevelExporterUSDOptions" ), FOnGetDetailCustomizationInstance::CreateStatic( &FLevelExporterUSDOptionsCustomization::MakeInstance ) );
+		PropertyModule.RegisterCustomClassLayout( TEXT( "LevelSequenceExporterUSDOptions" ), FOnGetDetailCustomizationInstance::CreateStatic( &FLevelExporterUSDOptionsCustomization::MakeInstance ) );
 	}
 
 	virtual void ShutdownModule() override
@@ -22,6 +25,7 @@ public:
 		if ( FPropertyEditorModule* PropertyModule = FModuleManager::GetModulePtr< FPropertyEditorModule >( TEXT( "PropertyEditor" ) ) )
 		{
 			PropertyModule->UnregisterCustomClassLayout( TEXT( "LevelExporterUSDOptions" ) );
+			PropertyModule->UnregisterCustomClassLayout( TEXT( "LevelSequenceExporterUSDOptions" ) );
 		}
 	}
 };
diff --git a/Engine/Plugins/Importers/USDImporter/Source/USDExporter/Public/LevelExporterUSDOptions.h b/Engine/Plugins/Importers/USDImporter/Source/USDExporter/Public/LevelExporterUSDOptions.h
--- a/Engine/Plugins/Importers/USDImporter/Source/USDExporter/Public/LevelExporterUSDOptions.h
+++ b/Engine/Plugins/Importers/USDImporter/Source/USDExporter/Public/LevelExporterUSDOptions.h
@@ -12,27 +12,11 @@
 
 #include "LevelExporterUSDOptions.generated.h"
 
-/**
- * Options for exporting levels to USD format.
- */
-UCLASS( Config = Editor, Blueprintable, HideCategories=Hidden )
-class USDEXPORTER_API ULevelExporterUSDOptions : public UObject
+USTRUCT( BlueprintType )
+struct USDEXPORTER_API FLevelExporterUSDOptionsInner
 {
 	GENERATED_BODY()
 
-public:
-	/** Basic options about the stage to export */
-	UPROPERTY( EditAnywhere, config, BlueprintReadWrite, Category = "Stage options", meta = ( ShowOnlyInnerProperties ) )
-	FUsdStageOptions StageOptions;
-
-	/** StartTimeCode to be used for all exported layers */
-    UPROPERTY( EditAnywhere, config, BlueprintReadWrite, Category = "Stage options" )
-    float StartTimeCode;
-
-	/** EndTimeCode to be used for all exported layers */
-    UPROPERTY( EditAnywhere, config, BlueprintReadWrite, Category = "Stage options" )
-    float EndTimeCode;
-
 	/** Whether to export only the selected actors, and assets used by them */
     UPROPERTY( EditAnywhere, config, BlueprintReadWrite, Category = "Export settings" )
     bool bSelectionOnly;
@@ -95,6 +79,32 @@ public:
 	/** Names of levels that should be ignored when collecting actors to export (e.g. "Persistent Level", "Level1", "MySubLevel", etc.) */
     UPROPERTY( EditAnywhere, config, BlueprintReadWrite, Category = "Sublayers" )
     TSet<FString> LevelsToIgnore;
+};
+
+/**
+ * Options for exporting levels to USD format.
+ */
+UCLASS( Config = Editor, Blueprintable, HideCategories=Hidden )
+class USDEXPORTER_API ULevelExporterUSDOptions : public UObject
+{
+	GENERATED_BODY()
+
+public:
+	/** Basic options about the stage to export */
+	UPROPERTY( EditAnywhere, config, BlueprintReadWrite, Category = "Stage options", meta = ( ShowOnlyInnerProperties ) )
+	FUsdStageOptions StageOptions;
+
+	/** StartTimeCode to be used for all exported layers */
+    UPROPERTY( EditAnywhere, config, BlueprintReadWrite, Category = "Stage options" )
+    float StartTimeCode;
+
+	/** EndTimeCode to be used for all exported layers */
+    UPROPERTY( EditAnywhere, config, BlueprintReadWrite, Category = "Stage options" )
+    float EndTimeCode;
+
+	/** Inner struct that actually contains most of the export options */
+    UPROPERTY( EditAnywhere, config, BlueprintReadWrite, Category = "Export settings", meta = ( ShowOnlyInnerProperties ) )
+    FLevelExporterUSDOptionsInner Inner;
 
 public:
 	// We temporarily stash our export task here as a way of passing our options down to
diff --git a/Engine/Plugins/Importers/USDImporter/Source/USDExporter/Public/LevelSequenceExporterUSD.h b/Engine/Plugins/Importers/USDImporter/Source/USDExporter/Public/LevelSequenceExporterUSD.h
new file mode 100644
--- /dev/null
+++ b/Engine/Plugins/Importers/USDImporter/Source/USDExporter/Public/LevelSequenceExporterUSD.h
@@ -0,0 +1,20 @@
+// Copyright Epic Games, Inc. All Rights Reserved.
+
+#pragma once
+
+#include "Exporters/Exporter.h"
+
+#include "LevelSequenceExporterUSD.generated.h"
+
+UCLASS()
+class ULevelSequenceExporterUsd : public UExporter
+{
+	GENERATED_BODY()
+
+public:
+	ULevelSequenceExporterUsd();
+
+	//~ Begin UExporter Interface
+	virtual bool ExportBinary( UObject* Object, const TCHAR* Type, FArchive& Ar, FFeedbackContext* Warn, int32 FileIndex = 0, uint32 PortFlags=0 ) override;
+	//~ End UExporter Interface
+};
diff --git a/Engine/Plugins/Importers/USDImporter/Source/USDExporter/Public/LevelSequenceExporterUSDOptions.h b/Engine/Plugins/Importers/USDImporter/Source/USDExporter/Public/LevelSequenceExporterUSDOptions.h
new file mode 100644
--- /dev/null
+++ b/Engine/Plugins/Importers/USDImporter/Source/USDExporter/Public/LevelSequenceExporterUSDOptions.h
@@ -0,0 +1,71 @@
+// Copyright Epic Games, Inc. All Rights Reserved.
+
+#pragma once
+
+#include "LevelExporterUSDOptions.h"
+#include "UnrealUSDWrapper.h"
+#include "USDStageOptions.h"
+
+#include "Engine/EngineTypes.h"
+
+#include "LevelSequenceExporterUSDOptions.generated.h"
+
+/**
+ * Options for level sequences to the USD format.
+ */
+UCLASS( Config = Editor, Blueprintable )
+class USDEXPORTER_API ULevelSequenceExporterUsdOptions : public UObject
+{
+	GENERATED_BODY()
+
+public:
+	/** Export options to use for the layer where the animation is emitted */
+	UPROPERTY( EditAnywhere, config, BlueprintReadWrite, Category = USDSettings, meta = ( ShowOnlyInnerProperties ) )
+	FUsdStageOptions StageOptions;
+
+	/**
+	 * Value to bake all generated USD layers with
+	 */
+	UPROPERTY( EditAnywhere, config, BlueprintReadWrite, Category = "Level Sequence" )
+	float TimeCodesPerSecond = 24.0f;
+
+	/**
+	 * If checked will cause StartFrame and EndFrame to be used as the frame range for the level sequence export.
+	 * If unchecked the existing playback range of each level sequence will be used instead.
+	 */
+	UPROPERTY( EditAnywhere, config, BlueprintReadWrite, Category = "Level Sequence" )
+	bool bOverrideExportRange = false;
+
+	/** Initial frame of the level sequence to bake out to USD (inclusive) */
+	UPROPERTY( EditAnywhere, config, BlueprintReadWrite, Category = "Level Sequence", meta = ( EditCondition = "bOverrideExportRange", EditConditionHides ) )
+	int32 StartFrame = 0;
+
+	/** Final frame of the level sequence to bake out to USD (inclusive) */
+	UPROPERTY( EditAnywhere, config, BlueprintReadWrite, Category = "Level Sequence", meta = ( EditCondition = "bOverrideExportRange", EditConditionHides ) )
+	int32 EndFrame = 0;
+
+	/**
+	 * The effect of subsequences is always included on the main exported layer, but if this option is true we will also
+	 * export individual sublayers for each subsequence, so that they can be used by itself if needed.
+	 */
+	UPROPERTY( EditAnywhere, config, BlueprintReadWrite, Category = "Level Sequence" )
+	bool bExportSubsequencesAsLayers = true;
+
+	/**
+	 * Exports the provided level alongside the layer that represents the exported Level Sequence
+	 */
+	UPROPERTY( EditAnywhere, config, BlueprintReadWrite, Category = "Level Sequence" )
+	bool bExportLevel = false;
+
+	UPROPERTY( EditAnywhere, BlueprintReadWrite, Category = "Level Export", meta = ( EditCondition = "bExportLevel", EditConditionHides ) )
+	TWeakObjectPtr<UWorld> Level = nullptr;
+
+	UPROPERTY( EditAnywhere, config, BlueprintReadWrite, Category = "Level Export", meta = ( EditCondition = "bExportLevel", EditConditionHides ) )
+	FLevelExporterUSDOptionsInner LevelExportOptions;
+
+	/**
+	 * If checked this will also add the exported level as a sublayer to the USD files emitted for all exported level sequences
+	 */
+	UPROPERTY( EditAnywhere, config, BlueprintReadWrite, Category = "Level Export", meta = ( EditCondition = "bExportLevel", EditConditionHides ) )
+	bool bUseExportedLevelAsSublayer = false;
+};
diff --git a/Engine/Plugins/Importers/USDImporter/Source/USDExporter/Public/USDConversionBlueprintLibrary.h b/Engine/Plugins/Importers/USDImporter/Source/USDExporter/Public/USDConversionBlueprintLibrary.h
--- a/Engine/Plugins/Importers/USDImporter/Source/USDExporter/Public/USDConversionBlueprintLibrary.h
+++ b/Engine/Plugins/Importers/USDImporter/Source/USDExporter/Public/USDConversionBlueprintLibrary.h
@@ -25,4 +25,10 @@ public:
 
 	UFUNCTION( BlueprintCallable, Category = "Layer utils" )
 	static void AddPayload( const FString& ReferencingStagePath, const FString& ReferencingPrimPath, const FString& TargetStagePath );
+
+	UFUNCTION( BlueprintCallable, Category = "Prim utils" )
+	static FString GetPrimPathForObject( const UObject* ActorOrComponent, const FString& ParentPrimPath = TEXT(""), bool bUseActorFolders = false );
+
+	UFUNCTION( BlueprintCallable, Category = "Prim utils" )
+	static FString GetSchemaNameForComponent( const USceneComponent* Component );
 };
diff --git a/Engine/Plugins/Importers/USDImporter/Source/USDExporter/USDExporter.Build.cs b/Engine/Plugins/Importers/USDImporter/Source/USDExporter/USDExporter.Build.cs
--- a/Engine/Plugins/Importers/USDImporter/Source/USDExporter/USDExporter.Build.cs
+++ b/Engine/Plugins/Importers/USDImporter/Source/USDExporter/USDExporter.Build.cs
@@ -32,12 +32,14 @@ namespace UnrealBuildTool.Rules
 					"JsonUtilities",
 					"Landscape",
 					"LevelSequence",
+					"LevelSequenceEditor", // For LevelSequenceEditorSpawnRegister, which we need to use when exporting level sequences with spawnables
 					"MaterialBaking", // So that we can use some of the export option properties
 					"MaterialUtilities",
 					"MeshDescription",
 					"MeshUtilities",
 					"MessageLog",
 					"MovieScene",
+					"MovieSceneTracks",
 					"PropertyEditor",
 					"PythonScriptPlugin",
 					"RawMesh",
diff --git a/Engine/Plugins/Importers/USDImporter/Source/USDUtilities/Private/USDConversionUtils.cpp b/Engine/Plugins/Importers/USDImporter/Source/USDUtilities/Private/USDConversionUtils.cpp
--- a/Engine/Plugins/Importers/USDImporter/Source/USDUtilities/Private/USDConversionUtils.cpp
+++ b/Engine/Plugins/Importers/USDImporter/Source/USDUtilities/Private/USDConversionUtils.cpp
@@ -34,6 +34,8 @@
 #include "Engine/StaticMesh.h"
 #include "Engine/Texture.h"
 #include "GeometryCache.h"
+#include "InstancedFoliageActor.h"
+#include "LandscapeProxy.h"
 
 #if WITH_EDITOR
 #include "ObjectTools.h"
@@ -361,6 +363,153 @@ UClass* UsdUtils::GetComponentTypeForPrim( const pxr::UsdPrim& Prim )
 	}
 }
 
+FString UsdUtils::GetSchemaNameForComponent( const USceneComponent& Component )
+{
+	AActor* OwnerActor = Component.GetOwner();
+	if ( OwnerActor->IsA<AInstancedFoliageActor>() )
+	{
+		return TEXT( "PointInstancer" );
+	}
+	else if ( OwnerActor->IsA<ALandscapeProxy>() )
+	{
+		return TEXT( "Mesh" );
+	}
+
+	if ( Component.IsA<USkinnedMeshComponent>() )
+	{
+		return TEXT( "SkelRoot" );
+	}
+	else if ( Component.IsA<UHierarchicalInstancedStaticMeshComponent>() )
+	{
+		// The original HISM component becomes just a regular Xform prim, so that we can handle
+		// its children correctly. We'll manually create a new child PointInstancer prim to it
+		// however, and convert the HISM data onto that prim.
+		return TEXT( "Xform" );
+	}
+	else if ( const UStaticMeshComponent* StaticMeshComponent = Cast<UStaticMeshComponent>( &Component ) )
+	{
+		UStaticMesh* Mesh = StaticMeshComponent->GetStaticMesh();
+		if ( Mesh && Mesh->GetNumLODs() > 1 )
+		{
+			// Don't export 'Mesh' if we're going to export LODs, as those will also be Mesh prims.
+			// We need at least an Xform schema though as this component may still have a transform of its own
+			return TEXT( "Xform" );
+		}
+		return TEXT( "Mesh" );
+	}
+	else if ( Component.IsA<UCineCameraComponent>() )
+	{
+		return TEXT( "Camera" );
+	}
+	else if ( Component.IsA<UDirectionalLightComponent>() )
+	{
+		return TEXT( "DistantLight" );
+	}
+	else if ( Component.IsA<URectLightComponent>() )
+	{
+		return TEXT( "RectLight" );
+	}
+	else if ( Component.IsA<UPointLightComponent>() )
+	{
+		return TEXT( "SphereLight" );
+	}
+	else if ( Component.IsA<USkyLightComponent>() )
+	{
+		return TEXT( "DomeLight" );
+	}
+
+	return TEXT( "Xform" );
+}
+
+FString UsdUtils::GetPrimPathForObject( const UObject* ActorOrComponent, const FString& ParentPrimPath, bool bUseActorFolders )
+{
+	if ( !ActorOrComponent )
+	{
+		return {};
+	}
+
+	// Get component and its owner actor
+	const USceneComponent* Component = Cast<const USceneComponent>( ActorOrComponent );
+	const AActor* Owner = nullptr;
+	if ( Component )
+	{
+		Owner = Component->GetOwner();
+	}
+	else
+	{
+		Owner = Cast<AActor>( ActorOrComponent );
+		if ( Owner )
+		{
+			Component = Owner->GetRootComponent();
+		}
+	}
+	if ( !Component || !Owner )
+	{
+		return {};
+	}
+
+	// Get component name. Use actor label if the component is its root component
+	FString Path;
+	if ( Component == Owner->GetRootComponent() )
+	{
+		Path = Owner->GetActorLabel();
+	}
+	else
+	{
+		Path = Component->GetName();
+	}
+	Path = UsdUtils::SanitizeUsdIdentifier( *Path );
+
+	// Get a clean folder path string if we have and need one
+	FString FolderPathString;
+	if ( bUseActorFolders && Component == Owner->GetRootComponent() )
+	{
+		const FName& FolderPath = Owner->GetFolderPath();
+		if ( !FolderPath.IsNone() )
+		{
+			FolderPathString = FolderPath.ToString();
+
+			TArray<FString> FolderSegments;
+			FolderPathString.ParseIntoArray( FolderSegments, TEXT( "/" ) );
+
+			for ( FString& Segment : FolderSegments )
+			{
+				Segment = UsdUtils::SanitizeUsdIdentifier( *Segment );
+			}
+
+			FolderPathString = FString::Join( FolderSegments, TEXT( "/" ) );
+
+			if ( !FolderPathString.IsEmpty() )
+			{
+				Path = FolderPathString / Path;
+			}
+		}
+	}
+
+	// Get parent prim path if we need to
+	if ( !ParentPrimPath.IsEmpty() )
+	{
+		Path = ParentPrimPath / Path;
+	}
+	else
+	{
+		FString FoundParentPath;
+
+		if ( USceneComponent* ParentComp = Component->GetAttachParent() )
+		{
+			FoundParentPath = GetPrimPathForObject( ParentComp, TEXT( "" ), bUseActorFolders );
+		}
+		else
+		{
+			FoundParentPath = TEXT( "/Root" );
+		}
+
+		Path = FoundParentPath / Path;
+	}
+
+	return Path;
+}
+
 TUsdStore< pxr::TfToken > UsdUtils::GetUVSetName( int32 UVChannelIndex )
 {
 	FScopedUnrealAllocs UnrealAllocs;
diff --git a/Engine/Plugins/Importers/USDImporter/Source/USDUtilities/Private/USDLayerUtils.cpp b/Engine/Plugins/Importers/USDImporter/Source/USDUtilities/Private/USDLayerUtils.cpp
--- a/Engine/Plugins/Importers/USDImporter/Source/USDUtilities/Private/USDLayerUtils.cpp
+++ b/Engine/Plugins/Importers/USDImporter/Source/USDUtilities/Private/USDLayerUtils.cpp
@@ -44,7 +44,7 @@ namespace UsdUtils
 	}
 }
 
-bool UsdUtils::InsertSubLayer( const pxr::SdfLayerRefPtr& ParentLayer, const TCHAR* SubLayerFile, int32 Index )
+bool UsdUtils::InsertSubLayer( const pxr::SdfLayerRefPtr& ParentLayer, const TCHAR* SubLayerFile, int32 Index, double OffsetTimeCodes, double TimeCodesScale )
 {
 	if ( !ParentLayer )
 	{
@@ -67,6 +67,16 @@ bool UsdUtils::InsertSubLayer( const pxr::SdfLayerRefPtr& ParentLayer, const TCH
 
 	ParentLayer->InsertSubLayerPath( UnrealToUsd::ConvertString( *RelativeSubLayerPath ).Get(), Index );
 
+	if ( !FMath::IsNearlyZero( OffsetTimeCodes ) || !FMath::IsNearlyEqual( TimeCodesScale, 1.0 ) )
+	{
+		if ( Index == -1 )
+		{
+			Index = ParentLayer->GetNumSubLayerPaths() - 1;
+		}
+
+		ParentLayer->SetSubLayerOffset( pxr::SdfLayerOffset{ OffsetTimeCodes, TimeCodesScale }, Index );
+	}
+
 	return true;
 }
 
diff --git a/Engine/Plugins/Importers/USDImporter/Source/USDUtilities/Private/USDPrimConversion.cpp b/Engine/Plugins/Importers/USDImporter/Source/USDUtilities/Private/USDPrimConversion.cpp
--- a/Engine/Plugins/Importers/USDImporter/Source/USDUtilities/Private/USDPrimConversion.cpp
+++ b/Engine/Plugins/Importers/USDImporter/Source/USDUtilities/Private/USDPrimConversion.cpp
@@ -5,17 +5,27 @@
 #include "UnrealUSDWrapper.h"
 #include "USDConversionUtils.h"
 #include "USDLayerUtils.h"
+#include "USDLightConversion.h"
+#include "USDSkeletalDataConversion.h"
 #include "USDTypesConversion.h"
 
+#include "UsdWrappers/UsdPrim.h"
+#include "UsdWrappers/UsdStage.h"
+
 #include "Channels/MovieSceneChannelProxy.h"
 #include "Channels/MovieSceneDoubleChannel.h"
 #include "CineCameraActor.h"
 #include "CineCameraComponent.h"
+#include "Components/DirectionalLightComponent.h"
 #include "Components/HierarchicalInstancedStaticMeshComponent.h"
 #include "Components/LightComponent.h"
 #include "Components/MeshComponent.h"
+#include "Components/PointLightComponent.h"
+#include "Components/RectLightComponent.h"
 #include "Components/SceneComponent.h"
+#include "Components/SkeletalMeshComponent.h"
 #include "Components/SkinnedMeshComponent.h"
+#include "Components/SpotLightComponent.h"
 #include "Engine/SkeletalMesh.h"
 #include "Engine/StaticMesh.h"
 #include "InstancedFoliageActor.h"
@@ -42,6 +52,7 @@
 #include "pxr/usd/usdGeom/xformable.h"
 #include "pxr/usd/usdGeom/xformCommonAPI.h"
 #include "pxr/usd/usdLux/light.h"
+#include "pxr/usd/usdSkel/animation.h"
 #include "pxr/usd/usdSkel/root.h"
 
 #include "USDIncludesEnd.h"
@@ -620,7 +631,7 @@ bool UnrealToUsd::ConvertHierarchicalInstancedStaticMeshComponent( const UHierar
 	return true;
 }
 
-bool UnrealToUsd::ConvertCameraComponent( const pxr::UsdStageRefPtr& Stage, const UCineCameraComponent* CameraComponent, pxr::UsdPrim& UsdPrim )
+bool UnrealToUsd::ConvertCameraComponent( const pxr::UsdStageRefPtr& Stage, const UCineCameraComponent* CameraComponent, pxr::UsdPrim& UsdPrim, double TimeCode )
 {
 	if ( !UsdPrim || !CameraComponent )
 	{
@@ -639,27 +650,27 @@ bool UnrealToUsd::ConvertCameraComponent( const pxr::UsdStageRefPtr& Stage, cons
 
 	if ( pxr::UsdAttribute Attr = GeomCamera.CreateFocalLengthAttr() )
 	{
-		Attr.Set<float>( UnrealToUsd::ConvertDistance( StageInfo, CameraComponent->CurrentFocalLength ) );
+		Attr.Set<float>( UnrealToUsd::ConvertDistance( StageInfo, CameraComponent->CurrentFocalLength ), TimeCode );
 	}
 
 	if ( pxr::UsdAttribute Attr = GeomCamera.CreateFocusDistanceAttr() )
 	{
-		Attr.Set<float>( UnrealToUsd::ConvertDistance( StageInfo, CameraComponent->FocusSettings.ManualFocusDistance ) );
+		Attr.Set<float>( UnrealToUsd::ConvertDistance( StageInfo, CameraComponent->FocusSettings.ManualFocusDistance ), TimeCode );
 	}
 
 	if ( pxr::UsdAttribute Attr = GeomCamera.CreateFStopAttr() )
 	{
-		Attr.Set<float>( CameraComponent->CurrentAperture );
+		Attr.Set<float>( CameraComponent->CurrentAperture, TimeCode );
 	}
 
 	if ( pxr::UsdAttribute Attr = GeomCamera.CreateHorizontalApertureAttr() )
 	{
-		Attr.Set<float>( UnrealToUsd::ConvertDistance( StageInfo, CameraComponent->Filmback.SensorWidth ) );
+		Attr.Set<float>( UnrealToUsd::ConvertDistance( StageInfo, CameraComponent->Filmback.SensorWidth ), TimeCode );
 	}
 
 	if ( pxr::UsdAttribute Attr = GeomCamera.CreateVerticalApertureAttr() )
 	{
-		Attr.Set<float>( UnrealToUsd::ConvertDistance( StageInfo, CameraComponent->Filmback.SensorHeight ) );
+		Attr.Set<float>( UnrealToUsd::ConvertDistance( StageInfo, CameraComponent->Filmback.SensorHeight ), TimeCode );
 	}
 
 	return true;
@@ -815,6 +826,356 @@ bool UnrealToUsd::ConvertInstancedFoliageActor( const AInstancedFoliageActor& Ac
 #endif // WITH_EDITOR
 }
 
+bool UnrealToUsd::CreateComponentPropertyBaker( UE::FUsdPrim& Prim, const USceneComponent& Component, const FString& PropertyPath, FComponentBaker& OutBaker )
+{
+	EBakingType BakerType = EBakingType::None;
+	TFunction<void(double)> BakerFunction;
+
+	FScopedUsdAllocs Allocs;
+
+	pxr::UsdPrim UsdPrim{ Prim };
+	pxr::UsdStageRefPtr UsdStage = UsdPrim.GetStage();
+	FUsdStageInfo StageInfo{ UsdStage };
+
+	// SceneComponent
+	{
+		if ( PropertyPath == TEXT( "Transform" ) )
+		{
+			pxr::UsdGeomXformable Xformable( UsdPrim );
+			if ( !Xformable )
+			{
+				return false;
+			}
+
+			Xformable.CreateXformOpOrderAttr();
+
+			// Clear existing transform data and leave just one Transform op there
+			pxr::UsdGeomXformOp TransformOp = Xformable.MakeMatrixXform();
+			if ( !TransformOp )
+			{
+				return false;
+			}
+
+			pxr::UsdAttribute Attr = TransformOp.GetAttr();
+			if ( !Attr )
+			{
+				return false;
+			}
+
+			Attr.Clear();
+
+			// Compensate different orientation for light or camera components
+			FTransform AdditionalRotation = FTransform::Identity;
+			if ( UsdPrim.IsA< pxr::UsdGeomCamera >() || UsdPrim.IsA< pxr::UsdLuxLight >() )
+			{
+				AdditionalRotation = FTransform( FRotator( 0.0f, 90.0f, 0.0f ) );
+
+				if ( StageInfo.UpAxis == EUsdUpAxis::ZAxis )
+				{
+					AdditionalRotation *= FTransform( FRotator( 90.0f, 0.0f, 0.0f ) );
+				}
+			}
+
+			BakerType = EBakingType::Transform;
+			BakerFunction = [&Component, AdditionalRotation, StageInfo, Attr]( double UsdTimeCode )
+			{
+				FScopedUsdAllocs Allocs;
+
+				FTransform FinalUETransform = AdditionalRotation * Component.GetRelativeTransform();
+				pxr::GfMatrix4d UsdTransform = UnrealToUsd::ConvertTransform( StageInfo, FinalUETransform );
+				Attr.Set< pxr::GfMatrix4d>( UsdTransform, UsdTimeCode );
+			};
+		}
+		// bHidden is for the actor, and bHiddenInGame is for a component
+		// A component is only visible when it's not hidden and its actor is not hidden
+		// A bHidden is just handled like a bHiddenInGame for the actor's root component
+		// Whenever we handle a bHiddenInGame, we always combine it with the actor's bHidden
+		else if ( PropertyPath == TEXT( "bHidden" ) || PropertyPath == TEXT( "bHiddenInGame" ) )
+		{
+			pxr::UsdGeomImageable Imageable( UsdPrim );
+			if ( !Imageable )
+			{
+				return false;
+			}
+
+			pxr::UsdAttribute Attr = Imageable.CreateVisibilityAttr();
+			Attr.Clear();
+
+			BakerType = EBakingType::Visibility;
+			BakerFunction = [&Component, Imageable]( double UsdTimeCode )
+			{
+				if ( Component.bHiddenInGame || Component.GetOwner()->IsHidden() )
+				{
+					Imageable.MakeInvisible( UsdTimeCode );
+				}
+				else
+				{
+					Imageable.MakeVisible( UsdTimeCode );
+				}
+			};
+		}
+	}
+
+	if ( const UCineCameraComponent* CameraComponent = Cast<UCineCameraComponent>( &Component ) )
+	{
+		static TSet<FString> RelevantProperties = {
+			TEXT( "CurrentFocalLength" ),
+			TEXT( "FocusSettings.ManualFocusDistance" ),
+			TEXT( "CurrentAperture" ),
+			TEXT( "Filmback.SensorWidth" ),
+			TEXT( "Filmback.SensorHeight" )
+		};
+
+		if ( RelevantProperties.Contains( PropertyPath ) )
+		{
+			BakerType = EBakingType::Camera;
+			BakerFunction = [UsdStage, CameraComponent, UsdPrim]( double UsdTimeCode ) mutable
+			{
+				UnrealToUsd::ConvertCameraComponent( UsdStage, CameraComponent, UsdPrim, UsdTimeCode );
+			};
+		}
+	}
+	else if ( const ULightComponentBase* LightComponentBase = Cast<ULightComponentBase>( &Component ) )
+	{
+		if ( const URectLightComponent* RectLightComponent = Cast<URectLightComponent>( LightComponentBase ) )
+		{
+			static TSet<FString> RelevantProperties = {
+				TEXT( "SourceHeight" ),
+				TEXT( "SourceWidth" ),
+				TEXT( "Temperature" ),
+				TEXT( "bUseTemperature" ),
+				TEXT( "LightColor" ),
+				TEXT( "Intensity" )
+			};
+
+			if ( RelevantProperties.Contains( PropertyPath ) )
+			{
+				BakerType = EBakingType::Light;
+				BakerFunction = [UsdStage, RectLightComponent, UsdPrim]( double UsdTimeCode ) mutable
+				{
+					UnrealToUsd::ConvertLightComponent( *RectLightComponent, UsdPrim, UsdTimeCode );
+					UnrealToUsd::ConvertRectLightComponent( *RectLightComponent, UsdPrim, UsdTimeCode );
+				};
+			}
+		}
+		else if ( const USpotLightComponent* SpotLightComponent = Cast<USpotLightComponent>( LightComponentBase ) )
+		{
+			static TSet<FString> RelevantProperties = {
+				TEXT( "OuterConeAngle" ),
+				TEXT( "InnerConeAngle" ),
+				TEXT( "Temperature" ),
+				TEXT( "bUseTemperature" ),
+				TEXT( "LightColor" ),
+				TEXT( "Intensity" )
+			};
+
+			if ( RelevantProperties.Contains( PropertyPath ) )
+			{
+				BakerType = EBakingType::Light;
+				BakerFunction = [UsdStage, SpotLightComponent, UsdPrim]( double UsdTimeCode ) mutable
+				{
+					UnrealToUsd::ConvertLightComponent( *SpotLightComponent, UsdPrim, UsdTimeCode );
+					UnrealToUsd::ConvertPointLightComponent( *SpotLightComponent, UsdPrim, UsdTimeCode );
+					UnrealToUsd::ConvertSpotLightComponent( *SpotLightComponent, UsdPrim, UsdTimeCode );
+				};
+			}
+		}
+		else if ( const UPointLightComponent* PointLightComponent = Cast<UPointLightComponent>( LightComponentBase ) )
+		{
+			static TSet<FString> RelevantProperties = {
+				TEXT( "SourceRadius" ),
+				TEXT( "Temperature" ),
+				TEXT( "bUseTemperature" ),
+				TEXT( "LightColor" ),
+				TEXT( "Intensity" )
+			};
+
+			if ( RelevantProperties.Contains( PropertyPath ) )
+			{
+				BakerType = EBakingType::Light;
+				BakerFunction = [UsdStage, PointLightComponent, UsdPrim]( double UsdTimeCode ) mutable
+				{
+					UnrealToUsd::ConvertLightComponent( *PointLightComponent, UsdPrim, UsdTimeCode );
+					UnrealToUsd::ConvertPointLightComponent( *PointLightComponent, UsdPrim, UsdTimeCode );
+				};
+			}
+		}
+		else if ( const UDirectionalLightComponent* DirectionalLightComponent = Cast<UDirectionalLightComponent>( LightComponentBase ) )
+		{
+			static TSet<FString> RelevantProperties = {
+				TEXT( "LightSourceAngle" ),
+				TEXT( "Temperature" ),
+				TEXT( "bUseTemperature" ),
+				TEXT( "LightColor" ),
+				TEXT( "Intensity" )
+			};
+
+			if ( RelevantProperties.Contains( PropertyPath ) )
+			{
+				BakerType = EBakingType::Light;
+				BakerFunction = [UsdStage, DirectionalLightComponent, UsdPrim]( double UsdTimeCode ) mutable
+				{
+					UnrealToUsd::ConvertLightComponent( *DirectionalLightComponent, UsdPrim, UsdTimeCode );
+					UnrealToUsd::ConvertDirectionalLightComponent( *DirectionalLightComponent, UsdPrim, UsdTimeCode );
+				};
+			}
+		}
+	}
+
+	if ( BakerType != EBakingType::None && BakerFunction)
+	{
+		OutBaker.BakerType = BakerType;
+		OutBaker.BakerFunction = BakerFunction;
+		return true;
+	}
+
+	return false;
+}
+
+bool UnrealToUsd::CreateSkeletalAnimationBaker( UE::FUsdPrim& SkelRoot, UE::FUsdPrim& SkelAnimation, USkeletalMeshComponent& Component, FComponentBaker& OutBaker )
+{
+	USkeletalMesh* SkeletalMesh = Component.SkeletalMesh;
+	if ( !SkeletalMesh )
+	{
+		return false;
+	}
+
+	FScopedUsdAllocs Allocs;
+
+	pxr::UsdSkelRoot UsdSkelRoot{ SkelRoot };
+	pxr::UsdSkelAnimation UsdSkelAnimation{ SkelAnimation };
+	if ( !SkelRoot || !SkelAnimation )
+	{
+		return false;
+	}
+
+	// Make sure that the skel root is using our animation
+	pxr::UsdPrim SkelRootPrim = UsdSkelRoot.GetPrim();
+	pxr::UsdPrim SkelAnimPrim = UsdSkelAnimation.GetPrim();
+	UsdUtils::BindAnimationSource( SkelRootPrim, SkelAnimPrim );
+
+	FUsdStageInfo StageInfo{ SkelRoot.GetStage() };
+
+	pxr::UsdAttribute JointsAttr			= UsdSkelAnimation.CreateJointsAttr();
+	pxr::UsdAttribute TranslationsAttr		= UsdSkelAnimation.CreateTranslationsAttr();
+	pxr::UsdAttribute RotationsAttr			= UsdSkelAnimation.CreateRotationsAttr();
+	pxr::UsdAttribute ScalesAttr			= UsdSkelAnimation.CreateScalesAttr();
+
+	// Joints
+	const FReferenceSkeleton& RefSkeleton = SkeletalMesh->GetRefSkeleton();
+	const int32 NumBones = RefSkeleton.GetRefBoneInfo().Num();
+	UnrealToUsd::ConvertJointsAttribute( RefSkeleton, JointsAttr );
+
+	// Build active morph targets array if it isn't setup already
+	TArray<UMorphTarget*>& MorphTargets = SkeletalMesh->GetMorphTargets();
+	if ( Component.ActiveMorphTargets.Num() != Component.MorphTargetWeights.Num() && MorphTargets.Num() != 0 )
+	{
+		for ( int32 MorphTargetIndex = 0; MorphTargetIndex < MorphTargets.Num(); ++MorphTargetIndex )
+		{
+			FActiveMorphTarget ActiveMorphTarget;
+			ActiveMorphTarget.MorphTarget = MorphTargets[ MorphTargetIndex ];
+			ActiveMorphTarget.WeightIndex = MorphTargetIndex;
+
+			Component.ActiveMorphTargets.Add( ActiveMorphTarget );
+		}
+	}
+
+	// Blend shape names
+	// Here we have to export UMorphTarget FNames in some order, then the weights in that same order. That is all.
+	// Those work out as "channels", and USD will resolve those to match the right thing on each mesh.
+	// We sort them in weight index order so that within the Baker we just write out weights in the order they are in.
+	pxr::UsdAttribute BlendShapeWeightsAttr;
+	pxr::UsdAttribute BlendShapesAttr;
+	const int32 NumMorphTargets = Component.MorphTargetWeights.Num();
+	if ( NumMorphTargets > 0 )
+	{
+		BlendShapeWeightsAttr = UsdSkelAnimation.CreateBlendShapeWeightsAttr();
+		BlendShapesAttr = UsdSkelAnimation.CreateBlendShapesAttr();
+
+		TArray<FActiveMorphTarget> SortedMorphTargets = Component.ActiveMorphTargets;
+		SortedMorphTargets.Sort([](const FActiveMorphTarget& Left, const FActiveMorphTarget& Right)
+		{
+			return Left.WeightIndex < Right.WeightIndex;
+		});
+
+		pxr::VtArray< pxr::TfToken > BlendShapeNames;
+		BlendShapeNames.reserve(SortedMorphTargets.Num());
+
+		for ( const FActiveMorphTarget& ActiveMorphTarget : SortedMorphTargets )
+		{
+			FString BlendShapeName;
+			if ( UMorphTarget* MorphTarget = ActiveMorphTarget.MorphTarget )
+			{
+				BlendShapeName = MorphTarget->GetFName().ToString();
+			}
+
+			BlendShapeNames.push_back( UnrealToUsd::ConvertToken( *BlendShapeName ).Get() );
+		}
+
+		BlendShapesAttr.Set( BlendShapeNames );
+	}
+
+	pxr::VtVec3fArray Translations;
+	pxr::VtQuatfArray Rotations;
+	pxr::VtVec3hArray Scales;
+	pxr::VtArray< float > BlendShapeWeights;
+
+	OutBaker.BakerType = EBakingType::Skeletal;
+	OutBaker.BakerFunction =
+		[&Component, StageInfo, Translations, Rotations, Scales, BlendShapeWeights, TranslationsAttr, RotationsAttr, ScalesAttr, BlendShapeWeightsAttr, NumBones, NumMorphTargets]
+		( double UsdTimeCode ) mutable
+		{
+			FScopedUsdAllocs InnerAllocs;
+
+			Translations.resize( NumBones );
+			Rotations.resize( NumBones );
+			Scales.resize( NumBones );
+
+			// This whole incantation is required or else the component will really not update until the next frame.
+			// Note: This will also cause the update of morph target weights.
+			Component.TickAnimation( 0.f, false );
+			Component.UpdateLODStatus();
+			Component.RefreshBoneTransforms();
+			Component.RefreshSlaveComponents();
+			Component.UpdateComponentToWorld();
+			Component.FinalizeBoneTransform();
+			Component.MarkRenderTransformDirty();
+			Component.MarkRenderDynamicDataDirty();
+
+			const TArray< FTransform >& LocalBoneTransforms = Component.GetBoneSpaceTransforms();
+			for ( int32 BoneIndex = 0; BoneIndex < NumBones; ++BoneIndex )
+			{
+				FTransform BoneTransform = LocalBoneTransforms[ BoneIndex ];
+				BoneTransform = UsdUtils::ConvertAxes( StageInfo.UpAxis == EUsdUpAxis::ZAxis, BoneTransform );
+
+				Translations[ BoneIndex ] = UnrealToUsd::ConvertVector( BoneTransform.GetTranslation() ) * ( 0.01f / StageInfo.MetersPerUnit );
+				Rotations[ BoneIndex ] = UnrealToUsd::ConvertQuat( BoneTransform.GetRotation() ).GetNormalized();
+				Scales[ BoneIndex ] = pxr::GfVec3h( UnrealToUsd::ConvertVector( BoneTransform.GetScale3D() ) );
+			}
+
+			if ( Translations.size() > 0 )
+			{
+				TranslationsAttr.Set( Translations, UsdTimeCode );
+				RotationsAttr.Set( Rotations, UsdTimeCode );
+				ScalesAttr.Set( Scales, UsdTimeCode );
+			}
+
+			if ( NumMorphTargets > 0 && BlendShapeWeightsAttr )
+			{
+				BlendShapeWeights.resize( NumMorphTargets );
+
+				for ( int32 MorphTargetIndex = 0; MorphTargetIndex < NumMorphTargets; ++MorphTargetIndex )
+				{
+					BlendShapeWeights[ MorphTargetIndex ] = Component.MorphTargetWeights[ MorphTargetIndex ];
+				}
+
+				BlendShapeWeightsAttr.Set( BlendShapeWeights, UsdTimeCode );
+			}
+		};
+
+	return true;
+}
+
 bool UnrealToUsd::ConvertXformable( const UMovieScene3DTransformTrack& MovieSceneTrack, pxr::UsdPrim& UsdPrim, const FMovieSceneSequenceTransform& SequenceTransform )
 {
 	if ( !UsdPrim )
diff --git a/Engine/Plugins/Importers/USDImporter/Source/USDUtilities/Private/USDSkeletalDataConversion.cpp b/Engine/Plugins/Importers/USDImporter/Source/USDUtilities/Private/USDSkeletalDataConversion.cpp
--- a/Engine/Plugins/Importers/USDImporter/Source/USDUtilities/Private/USDSkeletalDataConversion.cpp
+++ b/Engine/Plugins/Importers/USDImporter/Source/USDUtilities/Private/USDSkeletalDataConversion.cpp
@@ -1026,21 +1026,6 @@ namespace UnrealToUsdImpl
 			OutFullPaths[BoneIndex] = FString::Printf(TEXT("%s/%s"), *OutFullPaths[ BoneInfo.ParentIndex ], *SanitizedBoneName );
 		}
 	}
-
-	// Sets the JointsAttr value based on the bone paths of ReferenceSkeleton
-	void SetJoinsAttr( const FReferenceSkeleton& ReferenceSkeleton, pxr::UsdAttribute JointsAttr )
-	{
-		TArray<FString> FullBonePaths;
-		UnrealToUsdImpl::CreateFullBonePaths( ReferenceSkeleton.GetRefBoneInfo(), FullBonePaths );
-
-		pxr::VtArray<pxr::TfToken> Joints;
-		Joints.reserve( FullBonePaths.Num() );
-		for ( const FString& BonePath : FullBonePaths )
-		{
-			Joints.push_back( UnrealToUsd::ConvertToken( *BonePath ).Get() );
-		}
-		JointsAttr.Set( Joints );
-	}
 }
 
 bool UsdToUnreal::ConvertSkeleton(const pxr::UsdSkelSkeletonQuery& SkeletonQuery, FSkeletalMeshImportData& SkelMeshImportData)
@@ -2301,7 +2286,8 @@ bool UnrealToUsd::ConvertSkeleton( const FReferenceSkeleton& ReferenceSkeleton,
 
 	// Joints
 	{
-		UnrealToUsdImpl::SetJoinsAttr( ReferenceSkeleton, UsdSkeleton.CreateJointsAttr() );
+		pxr::UsdAttribute JointsAttr = UsdSkeleton.CreateJointsAttr();
+		UnrealToUsd::ConvertJointsAttribute( ReferenceSkeleton, JointsAttr );
 	}
 
 	pxr::VtArray<pxr::GfMatrix4d> LocalSpaceJointTransforms;
@@ -2340,6 +2326,27 @@ bool UnrealToUsd::ConvertSkeleton( const FReferenceSkeleton& ReferenceSkeleton,
 	return true;
 }
 
+bool UnrealToUsd::ConvertJointsAttribute( const FReferenceSkeleton& ReferenceSkeleton, pxr::UsdAttribute& JointsAttribute )
+{
+	if ( !JointsAttribute )
+	{
+		return false;
+	}
+
+	TArray<FString> FullBonePaths;
+	UnrealToUsdImpl::CreateFullBonePaths( ReferenceSkeleton.GetRefBoneInfo(), FullBonePaths );
+
+	pxr::VtArray<pxr::TfToken> Joints;
+	Joints.reserve( FullBonePaths.Num() );
+	for ( const FString& BonePath : FullBonePaths )
+	{
+		Joints.push_back( UnrealToUsd::ConvertToken( *BonePath ).Get() );
+	}
+
+	JointsAttribute.Set( Joints );
+	return true;
+}
+
 bool UnrealToUsd::ConvertSkeleton( const USkeleton* Skeleton, pxr::UsdSkelSkeleton& UsdSkeleton )
 {
 	if ( !Skeleton )
@@ -2682,7 +2689,8 @@ bool UnrealToUsd::ConvertAnimSequence( UAnimSequence* AnimSequence, pxr::UsdPrim
 
 	// Joints
 	{
-		UnrealToUsdImpl::SetJoinsAttr( RefSkeleton, UsdSkelAnim.CreateJointsAttr() );
+		pxr::UsdAttribute JointsAttr = UsdSkelAnim.CreateJointsAttr();
+		UnrealToUsd::ConvertJointsAttribute( RefSkeleton, JointsAttr );
 	}
 
 	// Translations, Rotations & Scales
diff --git a/Engine/Plugins/Importers/USDImporter/Source/USDUtilities/Public/USDConversionUtils.h b/Engine/Plugins/Importers/USDImporter/Source/USDUtilities/Public/USDConversionUtils.h
--- a/Engine/Plugins/Importers/USDImporter/Source/USDUtilities/Public/USDConversionUtils.h
+++ b/Engine/Plugins/Importers/USDImporter/Source/USDUtilities/Public/USDConversionUtils.h
@@ -37,6 +37,7 @@ PXR_NAMESPACE_CLOSE_SCOPE
 #endif // #if USE_USD_SDK
 
 class UUsdAssetImportData;
+class USceneComponent;
 enum class EUsdUpAxis : uint8;
 namespace UE
 {
@@ -93,9 +94,23 @@ namespace UsdUtils
 	USDUTILITIES_API bool HasCompositionArcs( const pxr::UsdPrim& Prim );
 
 	USDUTILITIES_API UClass* GetActorTypeForPrim( const pxr::UsdPrim& Prim );
-
 	USDUTILITIES_API UClass* GetComponentTypeForPrim( const pxr::UsdPrim& Prim );
 
+	/** Returns the USD schema name that should be used when exporting Component (e.g. "Xform", "Mesh", "Camera", etc.) */
+	USDUTILITIES_API FString GetSchemaNameForComponent( const USceneComponent& Component );
+
+	/**
+	 * Returns a prim path to use when exporting a given ActorOrComponent, assuming we're exporting according to the component attachment hierarchy.
+	 * e.g. We have a top-level actor with label "Actor" with an attached child actor "Child" that has a component "Comp1" attached to its root, and "Comp2"
+	 * attached to "Comp1". If Comp2 is provided to this function, it will return "/Root/Actor/Child/Comp1/Comp2".
+	 *
+	 * The actor label is always used in place of the root component name.
+	 *
+	 * ParentPrimPath is optional, and without it the function will recurse upwards and build the full path all the way to the root. Actor folders
+	 * are handled just like if they were actors/components
+	 */
+	USDUTILITIES_API FString GetPrimPathForObject( const UObject* ActorOrComponent, const FString& ParentPrimPath = TEXT(""), bool bUseActorFolders = false );
+
 	USDUTILITIES_API TUsdStore< pxr::TfToken > GetUVSetName( int32 UVChannelIndex );
 
 	/**
diff --git a/Engine/Plugins/Importers/USDImporter/Source/USDUtilities/Public/USDLayerUtils.h b/Engine/Plugins/Importers/USDImporter/Source/USDUtilities/Public/USDLayerUtils.h
--- a/Engine/Plugins/Importers/USDImporter/Source/USDUtilities/Public/USDLayerUtils.h
+++ b/Engine/Plugins/Importers/USDImporter/Source/USDUtilities/Public/USDLayerUtils.h
@@ -38,8 +38,16 @@ namespace UsdUtils
 		Save
 	};
 
-	/** Inserts the SubLayerFile path into ParentLayer as a sublayer */
-	USDUTILITIES_API bool InsertSubLayer( const pxr::SdfLayerRefPtr& ParentLayer, const TCHAR* SubLayerFile, int32 Index = -1 );
+	/**
+	 * Inserts the SubLayerFile path into ParentLayer as a sublayer
+	 * @param ParentLayer - Layer to receive the new sublayer entry
+	 * @param SubLayerFile - Absolute path to a usd file to use as sublayer
+	 * @param Index - Zero-based index on the list of ParentLayer's sublayers list to insert the new sublayer. -1 means "at the end" (the default)
+	 * @param OffsetTimeCodes - Offset in USD time codes to use for the sublayer reference
+	 * @param TimeCodesScale - Scale to use for the sublayer reference
+	 * @return Whether the sublayer reference was added successfully or not
+	 */
+	USDUTILITIES_API bool InsertSubLayer( const pxr::SdfLayerRefPtr& ParentLayer, const TCHAR* SubLayerFile, int32 Index = -1, double OffsetTimeCodes = 0.0, double TimeCodesScale = 1.0 );
 
 #if WITH_EDITOR
 	/** Opens a file dialog to open or save a USD file */
diff --git a/Engine/Plugins/Importers/USDImporter/Source/USDUtilities/Public/USDPrimConversion.h b/Engine/Plugins/Importers/USDImporter/Source/USDUtilities/Public/USDPrimConversion.h
--- a/Engine/Plugins/Importers/USDImporter/Source/USDUtilities/Public/USDPrimConversion.h
+++ b/Engine/Plugins/Importers/USDImporter/Source/USDUtilities/Public/USDPrimConversion.h
@@ -34,6 +34,7 @@ class UMeshComponent;
 class UMovieScene;
 class UMovieScene3DTransformTrack;
 class USceneComponent;
+class USkeletalMeshComponent;
 struct FFrameRate;
 struct FMovieSceneSequenceTransform;
 
@@ -58,7 +59,7 @@ namespace UnrealToUsd
 	USDUTILITIES_API bool ConvertSceneComponent( const pxr::UsdStageRefPtr& Stage, const USceneComponent* SceneComponent, pxr::UsdPrim& UsdPrim );
 	USDUTILITIES_API bool ConvertMeshComponent( const pxr::UsdStageRefPtr& Stage, const UMeshComponent* MeshComponent, pxr::UsdPrim& UsdPrim );
 	USDUTILITIES_API bool ConvertHierarchicalInstancedStaticMeshComponent( const UHierarchicalInstancedStaticMeshComponent* HISMComponent, pxr::UsdPrim& UsdPrim, double TimeCode = UsdUtils::GetDefaultTimeCode() );
-	USDUTILITIES_API bool ConvertCameraComponent( const pxr::UsdStageRefPtr& Stage, const UCineCameraComponent* CameraComponent, pxr::UsdPrim& UsdPrim );
+	USDUTILITIES_API bool ConvertCameraComponent( const pxr::UsdStageRefPtr& Stage, const UCineCameraComponent* CameraComponent, pxr::UsdPrim& UsdPrim, double TimeCode = UsdUtils::GetDefaultTimeCode() );
 
 	USDUTILITIES_API bool ConvertXformable( const FTransform& RelativeTransform, pxr::UsdPrim& UsdPrim, double TimeCode );
 
@@ -79,6 +80,47 @@ namespace UnrealToUsd
 	 * @param TimeCode			   TimeCode to write the attribute values at. Use UsdUtils::GetDefaultTimeCode() for the Default value.
 	 */
 	USDUTILITIES_API bool ConvertInstancedFoliageActor( const AInstancedFoliageActor& Actor, pxr::UsdPrim& UsdPrim, double TimeCode );
+
+	// We will defer to the common UnrealToUsd::ConvertXComponent functions when baking level sequences, which will already bake
+	// all of the properties of a component out at the same time.
+	// We use this enum to keep track of which type of baking was already being done for a component, so that if
+	// e.g. we have a track for camera aperture and another for camera focal length, we end up just baking that camera
+	// only once (as only once will be enough to handle all animated properties)
+	enum class USDUTILITIES_API EBakingType : uint8
+	{
+		None = 0,
+		Transform = 1,
+		Visibility = 2,
+		Camera = 4,
+		Light = 8,
+		Skeletal = 16,
+	};
+	ENUM_CLASS_FLAGS( EBakingType );
+
+	// Contains a lambda function responsible for baking a USceneComponent on the level into a prim on an FUsdStage
+	struct USDUTILITIES_API FComponentBaker
+	{
+		EBakingType BakerType;
+		TFunction<void( double UsdTimeCode )> BakerFunction;
+	};
+
+	/**
+	 * Creates a property baker responsible for baking Component's current state onto Prim according to its component type,
+	 * if PropertyPath describes a property that we map from UE to USD for that component type.
+	 *
+	 * Returns whether a baker was created or not.
+	 *
+	 * e.g. if we provide a UCineCameraComponent and PropertyPath == "FocusSettings.ManualFocusDistance" it will create a
+	 * EBakingType::Camera baker (that bakes all of the mappable UE cine camera properties at once) and return true.
+	 * If we provide PropertyPath == "FocusSettings.FocusSmoothingInterpSpeed" it will return false as that is not mapped to USD.
+	 */
+	USDUTILITIES_API bool CreateComponentPropertyBaker( UE::FUsdPrim& Prim, const USceneComponent& Component, const FString& PropertyPath, FComponentBaker& OutBaker );
+
+	/**
+	 * Creates a property baker responsible for baking the joint/blend shape state of Component as a SkelAnimation,
+	 * returning true if a EBakingType::Skeletal baker was created
+	 */
+	USDUTILITIES_API bool CreateSkeletalAnimationBaker( UE::FUsdPrim& SkelRoot, UE::FUsdPrim& SkelAnimation, USkeletalMeshComponent& Component, FComponentBaker& OutBaker );
 }
 
 #endif // #if USE_USD_SDK
diff --git a/Engine/Plugins/Importers/USDImporter/Source/USDUtilities/Public/USDSkeletalDataConversion.h b/Engine/Plugins/Importers/USDImporter/Source/USDUtilities/Public/USDSkeletalDataConversion.h
--- a/Engine/Plugins/Importers/USDImporter/Source/USDUtilities/Public/USDSkeletalDataConversion.h
+++ b/Engine/Plugins/Importers/USDImporter/Source/USDUtilities/Public/USDSkeletalDataConversion.h
@@ -186,6 +186,9 @@ namespace UnrealToUsd
 	USDUTILITIES_API bool ConvertSkeleton( const USkeleton* Skeleton, pxr::UsdSkelSkeleton& UsdSkeleton );
 	USDUTILITIES_API bool ConvertSkeleton( const FReferenceSkeleton& ReferenceSkeleton, pxr::UsdSkelSkeleton& UsdSkeleton );
 
+	// Fill out a UsdSkelAnimation's Joints attribute with data from ReferenceSkeleton, taking care to concatenate bone paths
+	USDUTILITIES_API bool ConvertJointsAttribute( const FReferenceSkeleton& ReferenceSkeleton, pxr::UsdAttribute& JointsAttribute );
+
 	/**
 	 * Converts SkeletalMesh, its skeleton and morph target data into the corresponding USD objects and populates SkelRoot with them, at time TimeCode
 	 * @param SkeletalMesh - Mesh with the source data. If it contains multiple LODs it will lead to the creation of LOD variant sets and variants within SkelRoot
diff --git a/Engine/Plugins/Importers/USDImporter/Source/USDUtilities/USDUtilities.Build.cs b/Engine/Plugins/Importers/USDImporter/Source/USDUtilities/USDUtilities.Build.cs
--- a/Engine/Plugins/Importers/USDImporter/Source/USDUtilities/USDUtilities.Build.cs
+++ b/Engine/Plugins/Importers/USDImporter/Source/USDUtilities/USDUtilities.Build.cs
@@ -29,6 +29,7 @@ namespace UnrealBuildTool.Rules
 					"Foliage",
 					"GeometryCache", // Just so that we can fetch its AssetImportData
 					"IntelTBB",
+					"Landscape", // So that GetSchemaNameForComponent knows what to do with landscape proxies
 					"MeshDescription",
 					"MovieScene",
 					"MovieSceneTracks",
diff --git a/Engine/Plugins/Importers/USDImporter/Source/UnrealUSDWrapper/Private/UsdWrappers/UsdPrim.cpp b/Engine/Plugins/Importers/USDImporter/Source/UnrealUSDWrapper/Private/UsdWrappers/UsdPrim.cpp
--- a/Engine/Plugins/Importers/USDImporter/Source/UnrealUSDWrapper/Private/UsdWrappers/UsdPrim.cpp
+++ b/Engine/Plugins/Importers/USDImporter/Source/UnrealUSDWrapper/Private/UsdWrappers/UsdPrim.cpp
@@ -145,6 +145,20 @@ namespace UE
 	}
 #endif // #if USE_USD_SDK
 
+	bool FUsdPrim::SetSpecifier( ESdfSpecifier Specifier )
+	{
+#if USE_USD_SDK
+		static_assert( ( int32 ) ESdfSpecifier::Def == ( int32 ) pxr::SdfSpecifierDef, "ESdfSpecifier enum doesn't match USD!" );
+		static_assert( ( int32 ) ESdfSpecifier::Over == ( int32 ) pxr::SdfSpecifierOver, "ESdfSpecifier enum doesn't match USD!" );
+		static_assert( ( int32 ) ESdfSpecifier::Class == ( int32 ) pxr::SdfSpecifierClass, "ESdfSpecifier enum doesn't match USD!" );
+		static_assert( ( int32 ) ESdfSpecifier::Num == ( int32 ) pxr::SdfNumSpecifiers, "ESdfSpecifier enum doesn't match USD!" );
+
+		return Impl->PxrUsdPrim.Get().SetSpecifier( static_cast< pxr::SdfSpecifier > ( Specifier ) );
+#else
+		return false;
+#endif // #if USE_USD_SDK
+	}
+
 	bool FUsdPrim::IsActive() const
 	{
 #if USE_USD_SDK
diff --git a/Engine/Plugins/Importers/USDImporter/Source/UnrealUSDWrapper/Private/UsdWrappers/UsdStage.cpp b/Engine/Plugins/Importers/USDImporter/Source/UnrealUSDWrapper/Private/UsdWrappers/UsdStage.cpp
--- a/Engine/Plugins/Importers/USDImporter/Source/UnrealUSDWrapper/Private/UsdWrappers/UsdStage.cpp
+++ b/Engine/Plugins/Importers/USDImporter/Source/UnrealUSDWrapper/Private/UsdWrappers/UsdStage.cpp
@@ -406,7 +406,7 @@ namespace UE
 #if USE_USD_SDK
 		if ( PtrType& Ptr = Impl->GetInner() )
 		{
-			if ( !FMath::IsNearlyEqual( TimeCode, Ptr->GetStartTimeCode() ) )
+			if ( !FMath::IsNearlyEqual( TimeCode, Ptr->GetStartTimeCode() ) || !Ptr->HasAuthoredTimeCodeRange() )
 			{
 				Ptr->SetStartTimeCode( TimeCode );
 			}
@@ -420,7 +420,7 @@ namespace UE
 #if USE_USD_SDK
 		if ( PtrType& Ptr = Impl->GetInner() )
 		{
-			if ( !FMath::IsNearlyEqual( TimeCode, Ptr->GetEndTimeCode() ) )
+			if ( !FMath::IsNearlyEqual( TimeCode, Ptr->GetEndTimeCode() ) || !Ptr->HasAuthoredTimeCodeRange() )
 			{
 				Ptr->SetEndTimeCode( TimeCode );
 			}
@@ -493,6 +493,19 @@ namespace UE
 #endif // #if USE_USD_SDK
 	}
 
+	template<typename PtrType>
+	FUsdPrim FUsdStageBase<PtrType>::OverridePrim( const FSdfPath& Path )
+	{
+#if USE_USD_SDK
+		if ( PtrType& Ptr = Impl->GetInner() )
+		{
+			return FUsdPrim( Ptr->OverridePrim( Path ) );
+		}
+#endif // #if USE_USD_SDK
+
+		return FUsdPrim();
+	}
+
 	template<typename PtrType>
 	FUsdPrim FUsdStageBase<PtrType>::DefinePrim( const FSdfPath & Path, const TCHAR * TypeName )
 	{
diff --git a/Engine/Plugins/Importers/USDImporter/Source/UnrealUSDWrapper/Public/UsdWrappers/UsdPrim.h b/Engine/Plugins/Importers/USDImporter/Source/UnrealUSDWrapper/Public/UsdWrappers/UsdPrim.h
--- a/Engine/Plugins/Importers/USDImporter/Source/UnrealUSDWrapper/Public/UsdWrappers/UsdPrim.h
+++ b/Engine/Plugins/Importers/USDImporter/Source/UnrealUSDWrapper/Public/UsdWrappers/UsdPrim.h
@@ -30,6 +30,15 @@ namespace UE
 		class FUsdPrimImpl;
 	}
 
+	/** Corresponds to pxr::SdfSpecifier, refer to the USD SDK documentation */
+	enum class ESdfSpecifier
+	{
+		Def,	// Defines a concrete prim
+		Over,	// Overrides an existing prim
+		Class,	// Defines an abstract prim
+		Num		// The number of specifiers
+	};
+
 	/**
 	 * Minimal pxr::UsdPrim wrapper for Unreal that can be used from no-rtti modules.
 	 */
@@ -64,6 +73,8 @@ namespace UE
 
 	// Wrapped pxr::UsdPrim functions, refer to the USD SDK documentation
 	public:
+		bool SetSpecifier( ESdfSpecifier Specifier );
+
 		bool IsActive() const;
 		bool SetActive( bool bActive );
 
diff --git a/Engine/Plugins/Importers/USDImporter/Source/UnrealUSDWrapper/Public/UsdWrappers/UsdStage.h b/Engine/Plugins/Importers/USDImporter/Source/UnrealUSDWrapper/Public/UsdWrappers/UsdStage.h
--- a/Engine/Plugins/Importers/USDImporter/Source/UnrealUSDWrapper/Public/UsdWrappers/UsdStage.h
+++ b/Engine/Plugins/Importers/USDImporter/Source/UnrealUSDWrapper/Public/UsdWrappers/UsdStage.h
@@ -89,6 +89,8 @@ namespace UE
 		void SetFramesPerSecond( double FramesPerSecond );
 
 		void SetDefaultPrim( const FUsdPrim& Prim );
+
+		FUsdPrim OverridePrim( const FSdfPath& Path );
 		FUsdPrim DefinePrim( const FSdfPath& Path, const TCHAR* TypeName = TEXT("") );
 		bool RemovePrim( const FSdfPath& Path );
 
