USD: Implemented support for material overrides on USD geometry cache components

diff --git a/Engine/Plugins/Importers/USDImporter/Source/USDUtilities/Private/USDPrimConversion.cpp b/Engine/Plugins/Importers/USDImporter/Source/USDUtilities/Private/USDPrimConversion.cpp
--- a/Engine/Plugins/Importers/USDImporter/Source/USDUtilities/Private/USDPrimConversion.cpp
+++ b/Engine/Plugins/Importers/USDImporter/Source/USDUtilities/Private/USDPrimConversion.cpp
@@ -28,6 +28,8 @@
 #include "Components/SpotLightComponent.h"
 #include "Engine/SkeletalMesh.h"
 #include "Engine/StaticMesh.h"
+#include "GeometryCache.h"
+#include "GeometryCacheComponent.h"
 #include "InstancedFoliageActor.h"
 #include "MovieSceneTimeHelpers.h"
 #include "Rendering/SkeletalMeshRenderData.h"
@@ -386,7 +388,27 @@ bool UnrealToUsd::ConvertMeshComponent( const pxr::UsdStageRefPtr& Stage, const
 
 	FScopedUsdAllocs Allocs;
 
-	if ( const UStaticMeshComponent* StaticMeshComponent = Cast<const UStaticMeshComponent>( MeshComponent ) )
+	if ( const UGeometryCacheComponent* GeometryCacheComponent = Cast<const UGeometryCacheComponent>( MeshComponent ) )
+	{
+		const TArray<UMaterialInterface*>& Overrides = MeshComponent->OverrideMaterials;
+		for ( int32 MatIndex = 0; MatIndex < Overrides.Num(); ++MatIndex )
+		{
+			UMaterialInterface* Override = Overrides[ MatIndex ];
+			if ( !Override )
+			{
+				continue;
+			}
+
+			pxr::SdfPath OverridePrimPath = UsdPrim.GetPath();
+
+			pxr::UsdPrim OverridePrim = Stage->OverridePrim( OverridePrimPath );
+			if ( pxr::UsdAttribute UnrealMaterialAttr = OverridePrim.CreateAttribute( UnrealIdentifiers::MaterialAssignment, pxr::SdfValueTypeNames->String ) )
+			{
+				UnrealMaterialAttr.Set( UnrealToUsd::ConvertString( *Override->GetPathName() ).Get() );
+			}
+		}
+	}
+	else if ( const UStaticMeshComponent* StaticMeshComponent = Cast<const UStaticMeshComponent>( MeshComponent ) )
 	{
 #if WITH_EDITOR
 		if ( UStaticMesh* Mesh = StaticMeshComponent->GetStaticMesh() )
