USD: Fix compile errors when USD SDK is not available and when packaging

diff --git a/Engine/Plugins/Importers/USDImporter/Source/USDExporter/Private/LevelSequenceExporterUSD.cpp b/Engine/Plugins/Importers/USDImporter/Source/USDExporter/Private/LevelSequenceExporterUSD.cpp
--- a/Engine/Plugins/Importers/USDImporter/Source/USDExporter/Private/LevelSequenceExporterUSD.cpp
+++ b/Engine/Plugins/Importers/USDImporter/Source/USDExporter/Private/LevelSequenceExporterUSD.cpp
@@ -44,6 +44,7 @@ namespace UE
 	{
 		namespace Private
 		{
+#if USE_USD_SDK
 			// Custom spawn register so that when DestroySpawnedObject is called while bDestroyingJustHides is true we
 			// actually just hide the objects, so that we can keep a live reference to components within the bakers.
 			// We're going to convert the spawnable tracks into visibility tracks when exporting to USD, which
@@ -703,6 +704,7 @@ namespace UE
 
 				UsdStage.GetRootLayer().Save();
 			}
+#endif // USE_USD_SDK
 		}
 	}
 }
diff --git a/Engine/Plugins/Importers/USDImporter/Source/USDExporter/Private/USDConversionBlueprintLibrary.cpp b/Engine/Plugins/Importers/USDImporter/Source/USDExporter/Private/USDConversionBlueprintLibrary.cpp
--- a/Engine/Plugins/Importers/USDImporter/Source/USDExporter/Private/USDConversionBlueprintLibrary.cpp
+++ b/Engine/Plugins/Importers/USDImporter/Source/USDExporter/Private/USDConversionBlueprintLibrary.cpp
@@ -116,15 +116,21 @@ void UUsdConversionBlueprintLibrary::AddPayload( const FString& ReferencingStage
 
 FString UUsdConversionBlueprintLibrary::GetPrimPathForObject( const UObject* ActorOrComponent, const FString& ParentPrimPath, bool bUseActorFolders )
 {
+#if USE_USD_SDK
 	return UsdUtils::GetPrimPathForObject( ActorOrComponent, ParentPrimPath, bUseActorFolders );
+#else
+	return {};
+#endif // USE_USD_SDK
 }
 
 FString UUsdConversionBlueprintLibrary::GetSchemaNameForComponent( const USceneComponent* Component )
 {
-	if ( !Component )
+#if USE_USD_SDK
+	if ( Component )
 	{
-		return {};
+		return UsdUtils::GetSchemaNameForComponent( *Component );
 	}
+#endif // USE_USD_SDK
 
-	return UsdUtils::GetSchemaNameForComponent( *Component );
+	return {};
 }
diff --git a/Engine/Plugins/Importers/USDImporter/Source/USDStage/Private/USDStageActor.cpp b/Engine/Plugins/Importers/USDImporter/Source/USDStage/Private/USDStageActor.cpp
--- a/Engine/Plugins/Importers/USDImporter/Source/USDStage/Private/USDStageActor.cpp
+++ b/Engine/Plugins/Importers/USDImporter/Source/USDStage/Private/USDStageActor.cpp
@@ -942,6 +942,7 @@ void AUsdStageActor::OnUsdObjectsChanged( const UsdUtils::FObjectChangesByPath&
 
 					UMaterialInterface* NewMaterial = Cast<UMaterialInterface>( AssetCache->GetAssetForPrim( AssetsPrimPath.GetString() ) );
 
+#if USE_USD_SDK
 					// For UE-120185: If we recreated a material for a prim path we also need to update all components that were using it.
 					// This could be fleshed out further if other asset types require this refresh of "dependent components" but materials
 					// seem to be the only ones that do at the moment
@@ -953,6 +954,7 @@ void AUsdStageActor::OnUsdObjectsChanged( const UsdUtils::FObjectChangesByPath&
 							UpdateComponents( MaterialUserPrim, bResyncComponent );
 						}
 					}
+#endif // USE_USD_SDK
 
 					// Resyncing also includes "updating" the prim
 					UpdatedAssets.Add( AssetsPrimPath.GetString() );
diff --git a/Engine/Plugins/Importers/USDImporter/Source/USDUtilities/Private/USDConversionUtils.cpp b/Engine/Plugins/Importers/USDImporter/Source/USDUtilities/Private/USDConversionUtils.cpp
--- a/Engine/Plugins/Importers/USDImporter/Source/USDUtilities/Private/USDConversionUtils.cpp
+++ b/Engine/Plugins/Importers/USDImporter/Source/USDUtilities/Private/USDConversionUtils.cpp
@@ -450,11 +450,13 @@ FString UsdUtils::GetPrimPathForObject( const UObject* ActorOrComponent, const F
 
 	// Get component name. Use actor label if the component is its root component
 	FString Path;
+#if WITH_EDITOR
 	if ( Component == Owner->GetRootComponent() )
 	{
 		Path = Owner->GetActorLabel();
 	}
 	else
+#endif // WITH_EDITOR
 	{
 		Path = Component->GetName();
 	}
@@ -462,6 +464,7 @@ FString UsdUtils::GetPrimPathForObject( const UObject* ActorOrComponent, const F
 
 	// Get a clean folder path string if we have and need one
 	FString FolderPathString;
+#if WITH_EDITOR
 	if ( bUseActorFolders && Component == Owner->GetRootComponent() )
 	{
 		const FName& FolderPath = Owner->GetFolderPath();
@@ -485,6 +488,7 @@ FString UsdUtils::GetPrimPathForObject( const UObject* ActorOrComponent, const F
 			}
 		}
 	}
+#endif // WITH_EDITOR
 
 	// Get parent prim path if we need to
 	if ( !ParentPrimPath.IsEmpty() )
diff --git a/Engine/Plugins/Importers/USDImporter/Source/USDUtilities/Private/USDPrimConversion.cpp b/Engine/Plugins/Importers/USDImporter/Source/USDUtilities/Private/USDPrimConversion.cpp
--- a/Engine/Plugins/Importers/USDImporter/Source/USDUtilities/Private/USDPrimConversion.cpp
+++ b/Engine/Plugins/Importers/USDImporter/Source/USDUtilities/Private/USDPrimConversion.cpp
@@ -1056,6 +1056,7 @@ bool UnrealToUsd::CreateComponentPropertyBaker( UE::FUsdPrim& Prim, const UScene
 
 bool UnrealToUsd::CreateSkeletalAnimationBaker( UE::FUsdPrim& SkelRoot, UE::FUsdPrim& SkelAnimation, USkeletalMeshComponent& Component, FComponentBaker& OutBaker )
 {
+#if WITH_EDITOR
 	USkeletalMesh* SkeletalMesh = Component.SkeletalMesh;
 	if ( !SkeletalMesh )
 	{
@@ -1196,6 +1197,9 @@ bool UnrealToUsd::CreateSkeletalAnimationBaker( UE::FUsdPrim& SkelRoot, UE::FUsd
 		};
 
 	return true;
+#else
+	return false;
+#endif // WITH_EDITOR
 }
 
 bool UnrealToUsd::ConvertXformable( const UMovieScene3DTransformTrack& MovieSceneTrack, pxr::UsdPrim& UsdPrim, const FMovieSceneSequenceTransform& SequenceTransform )
