USD: Merge fix for struct initialization from //UE5/Release-5.0 cl 17755044

diff --git a/Engine/Plugins/Importers/USDImporter/Source/USDExporter/Public/LevelExporterUSDOptions.h b/Engine/Plugins/Importers/USDImporter/Source/USDExporter/Public/LevelExporterUSDOptions.h
--- a/Engine/Plugins/Importers/USDImporter/Source/USDExporter/Public/LevelExporterUSDOptions.h
+++ b/Engine/Plugins/Importers/USDImporter/Source/USDExporter/Public/LevelExporterUSDOptions.h
@@ -19,11 +19,11 @@ struct USDEXPORTER_API FLevelExporterUSDOptionsInner
 
 	/** Whether to export only the selected actors, and assets used by them */
     UPROPERTY( EditAnywhere, config, BlueprintReadWrite, Category = "Export settings" )
-    bool bSelectionOnly;
+    bool bSelectionOnly = false;
 
 	/** Whether to bake UE materials and add material bindings to the baked assets */
     UPROPERTY( EditAnywhere, config, BlueprintReadWrite, Category = "Export settings" )
-    bool bBakeMaterials;
+    bool bBakeMaterials = false;
 
 	/** Resolution to use when baking materials into textures */
 	UPROPERTY( EditAnywhere, config, BlueprintReadWrite, Category = "Export settings", meta = ( EditCondition = "bBakeMaterials", ClampMin = "1" ) )
@@ -31,11 +31,11 @@ struct USDEXPORTER_API FLevelExporterUSDOptionsInner
 
 	/** Whether to remove the 'unrealMaterial' attribute after binding the corresponding baked material */
 	UPROPERTY( EditAnywhere, config, BlueprintReadWrite, Category = "Export settings", meta = ( EditCondition = "bBakeMaterials" ) )
-	bool bRemoveUnrealMaterials;
+	bool bRemoveUnrealMaterials = false;
 
 	/** If true, the actual static/skeletal mesh data is exported in "payload" files, and referenced via the payload composition arc */
 	UPROPERTY( EditAnywhere, config, BlueprintReadWrite, Category = "Export settings" )
-	bool bUsePayload;
+	bool bUsePayload = false;
 
 	/** USD format to use for exported payload files */
 	UPROPERTY( EditAnywhere, config, BlueprintReadWrite, Category = "Export settings", meta = ( EditCondition = "bUsePayload", GetOptions = GetUsdExtensions ) )
@@ -51,18 +51,18 @@ struct USDEXPORTER_API FLevelExporterUSDOptionsInner
 
 	/** Whether to use UE actor folders as empty prims */
     UPROPERTY( EditAnywhere, config, BlueprintReadWrite, Category = "Export settings" )
-    bool bExportActorFolders;
+    bool bExportActorFolders = false;
 
 	/** Lowest of the LOD indices to export landscapes with (use 0 for full resolution) */
 	UPROPERTY( EditAnywhere, config, BlueprintReadWrite, Category = "Export settings", meta = ( ClampMin = "0" ) )
-	int32 LowestLandscapeLOD;
+	int32 LowestLandscapeLOD = 0;
 
 	/**
 	 * Highest of the LOD indices to export landscapes with. Each value above 0 halves resolution.
 	 * The max value depends on the number of components and sections per component of each landscape, and may be clamped.
 	 */
 	UPROPERTY( EditAnywhere, config, BlueprintReadWrite, Category = "Export settings", meta = ( ClampMin = "0" ) )
-	int32 HighestLandscapeLOD;
+	int32 HighestLandscapeLOD = 0;
 
 	/** Resolution to use when baking landscape materials into textures  */
 	UPROPERTY( EditAnywhere, config, BlueprintReadWrite, Category = "Export settings", meta = ( ClampMin = "1" ) )
@@ -74,7 +74,7 @@ struct USDEXPORTER_API FLevelExporterUSDOptionsInner
 
 	/** If true, will export sub-levels as separate layers (referenced as sublayers). If false, will collapse all sub-levels in a single exported root layer */
     UPROPERTY( EditAnywhere, config, BlueprintReadWrite, Category = "Sublayers" )
-    bool bExportSublayers;
+    bool bExportSublayers = false;
 
 	/** Names of levels that should be ignored when collecting actors to export (e.g. "Persistent Level", "Level1", "MySubLevel", etc.) */
     UPROPERTY( EditAnywhere, config, BlueprintReadWrite, Category = "Sublayers" )
