USD: Fix camera compensation when writing back transforms to the USD stage;

Fix stage actor writing back to the stage when receiving levelsequence resync changes.

diff --git a/Engine/Plugins/Importers/USDImporter/Source/USDStage/Private/USDLevelSequenceHelper.cpp b/Engine/Plugins/Importers/USDImporter/Source/USDStage/Private/USDLevelSequenceHelper.cpp
--- a/Engine/Plugins/Importers/USDImporter/Source/USDStage/Private/USDLevelSequenceHelper.cpp
+++ b/Engine/Plugins/Importers/USDImporter/Source/USDStage/Private/USDLevelSequenceHelper.cpp
@@ -37,6 +37,7 @@
 
 #if WITH_EDITOR
 #include "Editor.h"
+#include "Editor/TransBuffer.h"
 #include "ILevelSequenceEditorToolkit.h"
 #include "ISequencer.h"
 #include "Subsystems/AssetEditorSubsystem.h"
@@ -230,18 +231,24 @@ private:
 
 // Changes handling
 public:
-	void StartMonitoringChanges() { bMonitorChanges = true; }
-	void StopMonitoringChanges() { bMonitorChanges = false; }
+	void StartMonitoringChanges() { MonitoringChangesWhenZero.Decrement(); }
+	void StopMonitoringChanges() { MonitoringChangesWhenZero.Increment(); }
+	bool IsMonitoringChanges() const { return MonitoringChangesWhenZero.GetValue() == 0; }
+
+	/**
+	 * Used as a fire-and-forget block that will prevent any levelsequence object (tracks, moviescene, sections, etc.) change from being written to the stage.
+	 * We unblock during HandleTransactionStateChanged.
+	 */
+	void BlockMonitoringChangesForThisTransaction();
 
 private:
 	void OnObjectTransacted(UObject* Object, const class FTransactionObjectEvent& Event);
+	void HandleTransactionStateChanged( const FTransactionContext& InTransactionContext, const ETransactionStateEventType InTransactionState );
 	void HandleMovieSceneChange(UMovieScene& MovieScene);
 	void HandleSubSectionChange(UMovieSceneSubSection& Section);
 	void HandleTransformTrackChange( const UMovieScene3DTransformTrack& TransformTrack, bool bIsMuteChange );
 	void HandleDisplayRateChange(const double DisplayRate);
 
-	bool bMonitorChanges; // Flag to handle or not changes done to the level sequence or one of its subobjects
-
 	FDelegateHandle OnObjectTransactedHandle;
 	FDelegateHandle OnStageEditTargetChangedHandle;
 
@@ -262,6 +269,13 @@ private:
 	TWeakObjectPtr<AUsdStageActor> StageActor;
 	FGuid StageActorBinding;
 
+	// Only when this is zero we write LevelSequence object (tracks, moviescene, sections, etc.) transactions back to the USD stage
+	FThreadSafeCounter MonitoringChangesWhenZero;
+
+	// When we call BlockMonitoringChangesForThisTransaction, we record the FGuid of the current transaction. We'll early out of all OnObjectTransacted calls for that transaction
+	// We keep a set here in order to remember all the blocked transactions as we're going through them
+	TSet<FGuid> BlockedTransactionGuids;
+
 	UE::FUsdStage UsdStage;
 };
 
@@ -272,10 +286,17 @@ const double FUsdLevelSequenceHelperImpl::EmptySubSectionRange = 10.0;
 
 FUsdLevelSequenceHelperImpl::FUsdLevelSequenceHelperImpl()
 	: MainLevelSequence( nullptr )
-	, bMonitorChanges( true )
 {
 #if WITH_EDITOR
 	OnObjectTransactedHandle = FCoreUObjectDelegates::OnObjectTransacted.AddRaw(this, &FUsdLevelSequenceHelperImpl::OnObjectTransacted);
+
+	if ( GEditor )
+	{
+		if ( UTransBuffer* Transactor = Cast<UTransBuffer>( GEditor->Trans ) )
+		{
+			Transactor->OnTransactionStateChanged().AddRaw( this, &FUsdLevelSequenceHelperImpl::HandleTransactionStateChanged );
+		}
+	}
 #endif // WITH_EDITOR
 }
 
@@ -290,6 +311,14 @@ FUsdLevelSequenceHelperImpl::~FUsdLevelSequenceHelperImpl()
 #if WITH_EDITOR
 	FCoreUObjectDelegates::OnObjectTransacted.Remove(OnObjectTransactedHandle);
 	OnObjectTransactedHandle.Reset();
+
+	if ( GEditor )
+	{
+		if ( UTransBuffer* Transactor = Cast<UTransBuffer>( GEditor->Trans ) )
+		{
+			Transactor->OnTransactionStateChanged().RemoveAll( this );
+		}
+	}
 #endif // WITH_EDITOR
 }
 
@@ -1252,9 +1281,28 @@ void FUsdLevelSequenceHelperImpl::UpdateMovieSceneTimeRanges( UMovieScene& Movie
 	MovieScene.SetDisplayRate( FFrameRate( FramesPerSecond, 1 ) );
 }
 
+void FUsdLevelSequenceHelperImpl::BlockMonitoringChangesForThisTransaction()
+{
+	if ( ITransaction* Trans = GUndo )
+	{
+		FTransactionContext Context = Trans->GetContext();
+
+		// We're already blocking this one, so ignore this so that we don't increment our counter too many times
+		if ( BlockedTransactionGuids.Contains( Context.TransactionId ) )
+		{
+			return;
+		}
+
+		BlockedTransactionGuids.Add( Context.TransactionId );
+	}
+
+	// Also block via the regular way in case we're receiving a notice due to a Python change and the user hasn't manually created a transaction
+	StopMonitoringChanges();
+}
+
 void FUsdLevelSequenceHelperImpl::OnObjectTransacted(UObject* Object, const class FTransactionObjectEvent& Event)
 {
-	if ( !MainLevelSequence || !bMonitorChanges || !Object || Object->IsPendingKill() || !UsdStage )
+	if ( !MainLevelSequence || !IsMonitoringChanges() || !Object || Object->IsPendingKill() || !UsdStage || BlockedTransactionGuids.Contains( Event.GetTransactionId() ) )
 	{
 		return;
 	}
@@ -1282,6 +1330,14 @@ void FUsdLevelSequenceHelperImpl::OnObjectTransacted(UObject* Object, const clas
 	}
 }
 
+void FUsdLevelSequenceHelperImpl::HandleTransactionStateChanged( const FTransactionContext& InTransactionContext, const ETransactionStateEventType InTransactionState )
+{
+	if ( InTransactionState == ETransactionStateEventType::TransactionFinalized && BlockedTransactionGuids.Contains( InTransactionContext.TransactionId ) )
+	{
+		StartMonitoringChanges();
+	}
+}
+
 double FUsdLevelSequenceHelperImpl::GetFramesPerSecond() const
 {
 	if ( !UsdStage )
@@ -1620,6 +1676,7 @@ public:
 
 	void StartMonitoringChanges() {}
 	void StopMonitoringChanges() {}
+	void BlockMonitoringChangesForThisTransaction() {}
 
 	ULevelSequence* GetMainLevelSequence() const { return nullptr; }
 	TArray< ULevelSequence* > GetSubSequences() const { return {}; }
@@ -1737,6 +1794,14 @@ void FUsdLevelSequenceHelper::StopMonitoringChanges()
 	}
 }
 
+void FUsdLevelSequenceHelper::BlockMonitoringChangesForThisTransaction()
+{
+	if ( UsdSequencerImpl.IsValid() )
+	{
+		UsdSequencerImpl->BlockMonitoringChangesForThisTransaction();
+	}
+}
+
 ULevelSequence* FUsdLevelSequenceHelper::GetMainLevelSequence() const
 {
 	if (UsdSequencerImpl.IsValid())
diff --git a/Engine/Plugins/Importers/USDImporter/Source/USDStage/Private/USDStageActor.cpp b/Engine/Plugins/Importers/USDImporter/Source/USDStage/Private/USDStageActor.cpp
--- a/Engine/Plugins/Importers/USDImporter/Source/USDStage/Private/USDStageActor.cpp
+++ b/Engine/Plugins/Importers/USDImporter/Source/USDStage/Private/USDStageActor.cpp
@@ -623,6 +623,10 @@ void AUsdStageActor::OnUsdObjectsChanged( const UsdUtils::FObjectChangesByPath&
 		}
 	}
 
+	// We may update our levelsequence objects (tracks, moviescene, sections, etc.) due to these changes. We definitely don't want to write anything
+	// back to USD when these objects change though
+	LevelSequenceHelper.BlockMonitoringChangesForThisTransaction();
+
 	// The most important thing here is to iterate in parent to child order, so build SortedPrimsChangedList
 	TArray< TPair< FString, bool > > SortedPrimsChangedList;
 	for ( const TPair<FString, TArray<UsdUtils::FObjectChangeNotice>>& InfoChange : InfoChanges )
@@ -838,6 +842,11 @@ void AUsdStageActor::ResumeMonitoringLevelSequence()
 	LevelSequenceHelper.StartMonitoringChanges();
 }
 
+void AUsdStageActor::BlockMonitoringLevelSequenceForThisTransaction()
+{
+	LevelSequenceHelper.BlockMonitoringChangesForThisTransaction();
+}
+
 UUsdPrimTwin* AUsdStageActor::GetOrCreatePrimTwin( const UE::FSdfPath& UsdPrimPath )
 {
 	const FString PrimPath = UsdPrimPath.GetString();
@@ -1405,7 +1414,7 @@ void AUsdStageActor::UnloadUsdStage()
 
 	// Stop listening because we'll discard LevelSequence assets, which may trigger transactions
 	// and could lead to stage changes
-	StopMonitoringLevelSequence();
+	BlockMonitoringLevelSequenceForThisTransaction();
 
 	if ( AssetCache )
 	{
diff --git a/Engine/Plugins/Importers/USDImporter/Source/USDStage/Public/USDLevelSequenceHelper.h b/Engine/Plugins/Importers/USDImporter/Source/USDStage/Public/USDLevelSequenceHelper.h
--- a/Engine/Plugins/Importers/USDImporter/Source/USDStage/Public/USDLevelSequenceHelper.h
+++ b/Engine/Plugins/Importers/USDImporter/Source/USDStage/Public/USDLevelSequenceHelper.h
@@ -55,6 +55,7 @@ public:
 	/** Blocks updating the level sequences & tracks from object changes. */
 	void StartMonitoringChanges();
 	void StopMonitoringChanges();
+	void BlockMonitoringChangesForThisTransaction();
 
 	ULevelSequence* GetMainLevelSequence() const;
 	TArray< ULevelSequence* > GetSubSequences() const;
diff --git a/Engine/Plugins/Importers/USDImporter/Source/USDStage/Public/USDStageActor.h b/Engine/Plugins/Importers/USDImporter/Source/USDStage/Public/USDStageActor.h
--- a/Engine/Plugins/Importers/USDImporter/Source/USDStage/Public/USDStageActor.h
+++ b/Engine/Plugins/Importers/USDImporter/Source/USDStage/Public/USDStageActor.h
@@ -229,6 +229,7 @@ public:
 	/** Prevents writing back data to the USD stage whenever our LevelSequences are modified */
 	USDSTAGE_API void StopMonitoringLevelSequence();
 	USDSTAGE_API void ResumeMonitoringLevelSequence();
+	USDSTAGE_API void BlockMonitoringLevelSequenceForThisTransaction();
 
 	UUsdPrimTwin* GetOrCreatePrimTwin( const UE::FSdfPath& UsdPrimPath );
 	UUsdPrimTwin* ExpandPrim( const UE::FUsdPrim& Prim, FUsdSchemaTranslationContext& TranslationContext );
diff --git a/Engine/Plugins/Importers/USDImporter/Source/USDStageEditor/Private/SUSDStage.cpp b/Engine/Plugins/Importers/USDImporter/Source/USDStageEditor/Private/SUSDStage.cpp
--- a/Engine/Plugins/Importers/USDImporter/Source/USDStageEditor/Private/SUSDStage.cpp
+++ b/Engine/Plugins/Importers/USDImporter/Source/USDStageEditor/Private/SUSDStage.cpp
@@ -756,35 +756,27 @@ void SUsdStage::OnPrimSelectionChanged( const TArray<FString>& PrimPaths )
 
 void SUsdStage::OpenStage( const TCHAR* FilePath )
 {
-	// This scope is important so that we can resume monitoring the level sequence after our scoped transaction is finished
-	{
-		// Create the transaction before calling UsdStageModule.GetUsdStageActor as that may create the actor, and we want
-		// the actor spawning to be part of the transaction
-		FScopedTransaction Transaction( FText::Format(
-			LOCTEXT( "OpenStageTransaction", "Open USD stage '{0}'" ),
-			FText::FromString( FilePath )
-		) );
-
-		if ( !ViewModel.UsdStageActor.IsValid() )
-		{
-			IUsdStageModule& UsdStageModule = FModuleManager::Get().LoadModuleChecked< IUsdStageModule >( "UsdStage" );
-			ViewModel.UsdStageActor = &UsdStageModule.GetUsdStageActor( GWorld );
-
-			SetupStageActorDelegates();
-		}
+	// Create the transaction before calling UsdStageModule.GetUsdStageActor as that may create the actor, and we want
+	// the actor spawning to be part of the transaction
+	FScopedTransaction Transaction( FText::Format(
+		LOCTEXT( "OpenStageTransaction", "Open USD stage '{0}'" ),
+		FText::FromString( FilePath )
+	) );
 
-		// Block writing level sequence changes back to the USD stage until we finished this transaction, because once we do
-		// the movie scene and tracks will all trigger OnObjectTransacted. We listen for those on FUsdLevelSequenceHelperImpl::OnObjectTransacted,
-		// and would otherwise end up writing all of the data we just loaded back to the USD stage
-		ViewModel.UsdStageActor->StopMonitoringLevelSequence();
+	if ( !ViewModel.UsdStageActor.IsValid() )
+	{
+		IUsdStageModule& UsdStageModule = FModuleManager::Get().LoadModuleChecked< IUsdStageModule >( "UsdStage" );
+		ViewModel.UsdStageActor = &UsdStageModule.GetUsdStageActor( GWorld );
 
-		ViewModel.OpenStage( FilePath );
+		SetupStageActorDelegates();
 	}
 
-	if ( AUsdStageActor* StageActor = ViewModel.UsdStageActor.Get() )
-	{
-		StageActor->ResumeMonitoringLevelSequence();
-	}
+	// Block writing level sequence changes back to the USD stage until we finished this transaction, because once we do
+	// the movie scene and tracks will all trigger OnObjectTransacted. We listen for those on FUsdLevelSequenceHelperImpl::OnObjectTransacted,
+	// and would otherwise end up writing all of the data we just loaded back to the USD stage
+	ViewModel.UsdStageActor->BlockMonitoringLevelSequenceForThisTransaction();
+
+	ViewModel.OpenStage( FilePath );
 }
 
 void SUsdStage::Refresh()
diff --git a/Engine/Plugins/Importers/USDImporter/Source/USDUtilities/Private/USDPrimConversion.cpp b/Engine/Plugins/Importers/USDImporter/Source/USDUtilities/Private/USDPrimConversion.cpp
--- a/Engine/Plugins/Importers/USDImporter/Source/USDUtilities/Private/USDPrimConversion.cpp
+++ b/Engine/Plugins/Importers/USDImporter/Source/USDUtilities/Private/USDPrimConversion.cpp
@@ -731,8 +731,9 @@ bool UnrealToUsd::ConvertXformable( const UMovieScene3DTransformTrack& MovieScen
 		return false;
 	}
 
-	const UMovieScene* MovieScene = MovieSceneTrack.GetTypedOuter< UMovieScene >();
+	const FUsdStageInfo StageInfo( UsdPrim.GetStage() );
 
+	const UMovieScene* MovieScene = MovieSceneTrack.GetTypedOuter< UMovieScene >();
 	if ( !MovieScene )
 	{
 		return false;
@@ -741,14 +742,12 @@ bool UnrealToUsd::ConvertXformable( const UMovieScene3DTransformTrack& MovieScen
 	FScopedUsdAllocs UsdAllocs;
 
 	pxr::UsdGeomXformable Xformable( UsdPrim );
-
 	if ( !Xformable )
 	{
 		return false;
 	}
 
 	UMovieScene3DTransformSection* TransformSection = Cast< UMovieScene3DTransformSection >( const_cast< UMovieScene3DTransformTrack& >( MovieSceneTrack ).FindSection( 0 ) );
-
 	if ( !TransformSection )
 	{
 		return false;
@@ -824,7 +823,6 @@ bool UnrealToUsd::ConvertXformable( const UMovieScene3DTransformTrack& MovieScen
 
 	bool bIsDataOutOfSync = false;
 	{
-		const FUsdStageInfo StageInfo( UsdPrim.GetStage() );
 		int32 ValueIndex = 0;
 
 		FFrameTime UsdStartTime = FFrameRate::TransformTime( PlaybackRange.GetLowerBoundValue(), Resolution, StageFrameRate );
@@ -875,6 +873,19 @@ bool UnrealToUsd::ConvertXformable( const UMovieScene3DTransformTrack& MovieScen
 
 		pxr::SdfChangeBlock ChangeBlock;
 
+		// Compensate different orientation for light or camera components
+		// TODO: Handle inverse compensation when the bound component is a child of a light/camera component
+		FTransform AdditionalRotation = FTransform::Identity;
+		if ( UsdPrim.IsA< pxr::UsdGeomCamera >() || UsdPrim.IsA< pxr::UsdLuxLight >() )
+		{
+			AdditionalRotation = FTransform( FRotator( 0.0f, 90.0f, 0.0f ) );
+
+			if ( StageInfo.UpAxis == EUsdUpAxis::ZAxis )
+			{
+				AdditionalRotation *= FTransform( FRotator( 90.0f, 0.0f, 0.0f ) );
+			}
+		}
+
 		int32 ValueIndex = 0;
 		for ( const TPair< FFrameNumber, float >& Value : LocationValuesX )
 		{
@@ -885,7 +896,7 @@ bool UnrealToUsd::ConvertXformable( const UMovieScene3DTransformTrack& MovieScen
 			FVector Scale( ScaleValuesX[ ValueIndex ].Value, ScaleValuesY[ ValueIndex ].Value, ScaleValuesZ[ ValueIndex ].Value );
 
 			FTransform Transform( Rotation, Location, Scale );
-			ConvertXformable( Transform, UsdPrim, UsdFrameTime.AsDecimal() );
+			ConvertXformable( AdditionalRotation* Transform, UsdPrim, UsdFrameTime.AsDecimal() );
 
 			++ValueIndex;
 		}
