USD: Manual merge to fix robomerge blockage when trying to merge cl 16311506 from UE4/Main.

diff --git a/Engine/Plugins/Importers/USDImporter/Source/USDClasses/Public/USDAssetCache.h b/Engine/Plugins/Importers/USDImporter/Source/USDClasses/Public/USDAssetCache.h
--- a/Engine/Plugins/Importers/USDImporter/Source/USDClasses/Public/USDAssetCache.h
+++ b/Engine/Plugins/Importers/USDImporter/Source/USDClasses/Public/USDAssetCache.h
@@ -30,6 +30,8 @@ public:
 	UObject* GetAssetForPrim( const FString& PrimPath ) const;
 	TMap< FString, UObject* > GetAssetPrimLinks() const { return PrimPathToAssets; }; // Can't return a reference as it wouldn't be thread-safe
 
+	bool IsAssetOwnedByCache( UObject* Asset ) const { return OwnedAssets.Contains( Asset ); }
+
 	int32 GetNumAssets() const { return TransientStorage.Num() + PersistentStorage.Num(); }
 	void Reset();
 
diff --git a/Engine/Plugins/Importers/USDImporter/Source/USDStageImporter/Private/USDStageAssetImportFactory.cpp b/Engine/Plugins/Importers/USDImporter/Source/USDStageImporter/Private/USDStageAssetImportFactory.cpp
--- a/Engine/Plugins/Importers/USDImporter/Source/USDStageImporter/Private/USDStageAssetImportFactory.cpp
+++ b/Engine/Plugins/Importers/USDImporter/Source/USDStageImporter/Private/USDStageAssetImportFactory.cpp
@@ -55,6 +55,11 @@ UObject* UUsdStageAssetImportFactory::FactoryCreateFile(UClass* InClass, UObject
 {
 	UObject* ImportedObject = nullptr;
 
+	if ( AssetImportTask && IsAutomatedImport() )
+	{
+		ImportContext.ImportOptions = Cast<UUsdStageImportOptions>( AssetImportTask->Options );
+	}
+
 	const FString InitialPackagePath = InParent ? InParent->GetName() : TEXT( "/Game/" );
 	const bool bIsReimport = false;
 	const bool bAllowActorImport = false;
diff --git a/Engine/Plugins/Importers/USDImporter/Source/USDStageImporter/Private/USDStageImportContext.cpp b/Engine/Plugins/Importers/USDImporter/Source/USDStageImporter/Private/USDStageImportContext.cpp
--- a/Engine/Plugins/Importers/USDImporter/Source/USDStageImporter/Private/USDStageImportContext.cpp
+++ b/Engine/Plugins/Importers/USDImporter/Source/USDStageImporter/Private/USDStageImportContext.cpp
@@ -57,7 +57,10 @@ bool FUsdStageImportContext::Init(const FString& InName, const FString& InFilePa
 			{
 				return false;
 			}
-			PackagePath = PickContentPathDlg->GetPath().ToString() + TEXT("/");
+			// e.g. "/Game/MyFolder/layername/"
+			// We inject the package path here because this is what the automated import task upstream code will do. This way the importer
+			// can always expect to receive /ContentPath/layername/
+			PackagePath = FString::Printf( TEXT( "%s/%s/" ), *PickContentPathDlg->GetPath().ToString(), *InName );
 		}
 
 		ImportOptions->EnableActorImport( bAllowActorImport );
diff --git a/Engine/Plugins/Importers/USDImporter/Source/USDStageImporter/Private/USDStageImportFactory.cpp b/Engine/Plugins/Importers/USDImporter/Source/USDStageImporter/Private/USDStageImportFactory.cpp
--- a/Engine/Plugins/Importers/USDImporter/Source/USDStageImporter/Private/USDStageImportFactory.cpp
+++ b/Engine/Plugins/Importers/USDImporter/Source/USDStageImporter/Private/USDStageImportFactory.cpp
@@ -7,6 +7,7 @@
 #include "USDStageImporterModule.h"
 #include "USDStageImportOptions.h"
 
+#include "AssetImportTask.h"
 #include "AssetRegistryModule.h"
 #include "AssetSelection.h"
 #include "Editor.h"
@@ -41,6 +42,11 @@ UObject* UUsdStageImportFactory::FactoryCreateFile(UClass* InClass, UObject* InP
 {
 	UObject* ImportedObject = nullptr;
 
+	if ( AssetImportTask && IsAutomatedImport() )
+	{
+		ImportContext.ImportOptions = Cast<UUsdStageImportOptions>( AssetImportTask->Options );
+	}
+
 #if USE_USD_SDK
 	const FString InitialPackagePath =InParent ? InParent->GetName() : TEXT( "/Game/" );
 	const bool bIsReimport = false;
diff --git a/Engine/Plugins/Importers/USDImporter/Source/USDStageImporter/Private/USDStageImportOptions.cpp b/Engine/Plugins/Importers/USDImporter/Source/USDStageImporter/Private/USDStageImportOptions.cpp
--- a/Engine/Plugins/Importers/USDImporter/Source/USDStageImporter/Private/USDStageImportOptions.cpp
+++ b/Engine/Plugins/Importers/USDImporter/Source/USDStageImporter/Private/USDStageImportOptions.cpp
@@ -17,7 +17,10 @@ UUsdStageImportOptions::UUsdStageImportOptions(const FObjectInitializer& ObjectI
 	bImportMaterials = true;
 
 	PurposesToImport = (int32) (EUsdPurpose::Default | EUsdPurpose::Proxy | EUsdPurpose::Render | EUsdPurpose::Guide);
+	IUsdSchemasModule& UsdSchemasModule = FModuleManager::Get().LoadModuleChecked< IUsdSchemasModule >( TEXT("USDSchemas") );
+	RenderContextToImport = UsdSchemasModule.GetRenderContextRegistry().GetUniversalRenderContext();
 	ImportTime = 0.0f;
+	bOverrideStageOptions = false;
 	StageOptions.MetersPerUnit = 0.01f;
 	StageOptions.UpAxis = EUsdUpAxis::ZAxis;
 
@@ -28,9 +31,6 @@ UUsdStageImportOptions::UUsdStageImportOptions(const FObjectInitializer& ObjectI
 	bPrimPathFolderStructure = false;
 	bCollapse = true;
 	bInterpretLODs = true;
-
-	IUsdSchemasModule& UsdSchemasModule = FModuleManager::Get().LoadModuleChecked< IUsdSchemasModule >( TEXT("USDSchemas") );
-	RenderContextToImport = UsdSchemasModule.GetRenderContextRegistry().GetUniversalRenderContext();
 }
 
 void UUsdStageImportOptions::EnableActorImport( bool bEnable )
diff --git a/Engine/Plugins/Importers/USDImporter/Source/USDStageImporter/Private/USDStageImporter.cpp b/Engine/Plugins/Importers/USDImporter/Source/USDStageImporter/Private/USDStageImporter.cpp
--- a/Engine/Plugins/Importers/USDImporter/Source/USDStageImporter/Private/USDStageImporter.cpp
+++ b/Engine/Plugins/Importers/USDImporter/Source/USDStageImporter/Private/USDStageImporter.cpp
@@ -674,7 +674,7 @@ namespace UsdStageImporterImpl
 			}
 
 			FString TargetAssetName = GetUserFriendlyName(Asset, UniqueAssetNames);
-			FString DestPackagePath = FPaths::Combine(ImportContext.PackagePath, ImportContext.ObjectName, AssetTypeFolder, TargetAssetName);
+			FString DestPackagePath = FPaths::Combine(ImportContext.PackagePath, AssetTypeFolder, TargetAssetName);
 			PublishAsset(ImportContext, Asset, DestPackagePath, ObjectsToRemap);
 		}
 
@@ -691,7 +691,7 @@ namespace UsdStageImporterImpl
 
 			for (ULevelSequence* LevelSequence : LevelSequences)
 			{
-				const FString DestPackagePath = FPaths::Combine(ImportContext.PackagePath, ImportContext.ObjectName, AssetTypeFolder, LevelSequence->GetName());
+				const FString DestPackagePath = FPaths::Combine(ImportContext.PackagePath, AssetTypeFolder, LevelSequence->GetName());
 				UObject* PublishedLevelSequence = PublishAsset(ImportContext, LevelSequence, DestPackagePath, ObjectsToRemap);
 			}
 		}
@@ -1059,15 +1059,27 @@ namespace UsdStageImporterImpl
 	 * Example: An UMaterialInstance is already in the cache, so when translating we just retrieve the existing asset --> The textures that it's using won't be retrieved or marked as "Used"
 	 * Example: An USkeletalMesh is already in the cache, so in the same way we would miss its USkeleton, materials and textures of those materials
 	 */
-	void CollectUsedAssetDependencies( const TSet<UObject*>& InPrimaryAssets, TSet<UObject*>& OutAssetsAndDependencies )
+	void CollectUsedAssetDependencies( UUsdAssetCache* AssetCache, TSet<UObject*>& OutAssetsAndDependencies )
 	{
+		if ( !AssetCache )
+		{
+			return;
+		}
+
+		const TSet<UObject*>& InPrimaryAssets = AssetCache->GetActiveAssets();
 		TArray<UObject*> AssetQueue = InPrimaryAssets.Array();
 		OutAssetsAndDependencies.Reserve( AssetQueue.Num() );
 
 		for ( int32 AssetIndex = 0; AssetIndex < AssetQueue.Num(); ++AssetIndex )
 		{
 			UObject* Asset = AssetQueue[ AssetIndex ];
-			OutAssetsAndDependencies.Add( Asset );
+
+			// Only add it as a dependency if it's owned by the asset cache, but still traverse it because
+			// we may be in some strange situation where the material shouldn't be in this list, but one of its used textures should
+			if ( AssetCache->IsAssetOwnedByCache( Asset ) )
+			{
+				OutAssetsAndDependencies.Add( Asset );
+			}
 
 			if ( UMaterial* Material = Cast<UMaterial>( Asset ) )
 			{
@@ -1079,7 +1091,10 @@ namespace UsdStageImporterImpl
 				OutAssetsAndDependencies.Reserve( OutAssetsAndDependencies.Num() + UsedTextures.Num() );
 				for ( UTexture* UsedTexture : UsedTextures )
 				{
-					OutAssetsAndDependencies.Add(UsedTexture);
+					if ( AssetCache->IsAssetOwnedByCache( UsedTexture ) )
+					{
+						OutAssetsAndDependencies.Add( UsedTexture );
+					}
 				}
 			}
 			else if ( UMaterialInstance* MaterialInstance = Cast<UMaterialInstance>( Asset ) )
@@ -1087,20 +1102,32 @@ namespace UsdStageImporterImpl
 				OutAssetsAndDependencies.Reserve( OutAssetsAndDependencies.Num() + MaterialInstance->TextureParameterValues.Num() );
 				for ( const FTextureParameterValue& TextureValue : MaterialInstance->TextureParameterValues )
 				{
-					if ( TextureValue.ParameterValue )
+					if ( UTexture* Texture = TextureValue.ParameterValue )
 					{
-						OutAssetsAndDependencies.Add( TextureValue.ParameterValue );
+						if ( AssetCache->IsAssetOwnedByCache( Texture ) )
+						{
+							OutAssetsAndDependencies.Add( Texture );
+						}
 					}
 				}
 			}
 			else if ( USkeletalMesh* SkeletalMesh = Cast<USkeletalMesh>( Asset ) )
 			{
-				OutAssetsAndDependencies.Add( SkeletalMesh->GetSkeleton() );
+				if ( USkeleton* Skeleton = SkeletalMesh->GetSkeleton() )
+				{
+					if ( AssetCache->IsAssetOwnedByCache( Skeleton ) )
+					{
+						OutAssetsAndDependencies.Add( Skeleton );
+					}
+				}
 
 				AssetQueue.Reserve( AssetQueue.Num() + SkeletalMesh->GetMaterials().Num() );
 				for ( const FSkeletalMaterial& SkeletalMaterial : SkeletalMesh->GetMaterials() )
 				{
-					AssetQueue.Add( SkeletalMaterial.MaterialInterface );
+					if ( UMaterialInterface* UsedMaterial = SkeletalMaterial.MaterialInterface )
+					{
+						AssetQueue.Add( UsedMaterial );
+					}
 				}
 			}
 			else if ( UStaticMesh* StaticMesh = Cast<UStaticMesh>( Asset ) )
@@ -1108,23 +1135,45 @@ namespace UsdStageImporterImpl
 				AssetQueue.Reserve( AssetQueue.Num() + StaticMesh->GetStaticMaterials().Num() );
 				for ( const FStaticMaterial& StaticMaterial : StaticMesh->GetStaticMaterials() )
 				{
-					AssetQueue.Add( StaticMaterial.MaterialInterface );
+					if ( UMaterialInterface* UsedMaterial = StaticMaterial.MaterialInterface )
+					{
+						AssetQueue.Add( UsedMaterial );
+					}
 				}
 			}
 			else if ( UGeometryCache* GeometryCache = Cast<UGeometryCache>( Asset ) )
 			{
-				AssetQueue.Append( GeometryCache->Materials );
+				for ( UMaterialInterface* UsedMaterial : GeometryCache->Materials )
+				{
+					if ( UsedMaterial )
+					{
+						AssetQueue.Add( UsedMaterial );
+					}
+				}
 			}
 			else if ( UAnimSequence* AnimSequence = Cast<UAnimSequence>( Asset ) )
 			{
-				AssetQueue.Add( AnimSequence->GetPreviewMesh() );
+				if ( USkeletalMesh* Mesh = AnimSequence->GetPreviewMesh() )
+				{
+					AssetQueue.Add( Mesh );
+				}
 
-				OutAssetsAndDependencies.Add( AnimSequence->GetSkeleton() );
+				if ( USkeleton* Skeleton = AnimSequence->GetSkeleton() )
+				{
+					if ( AssetCache->IsAssetOwnedByCache( Skeleton ) )
+					{
+						OutAssetsAndDependencies.Add( Skeleton );
+					}
+				}
 			}
 			else if ( UTexture* Texture = Cast<UTexture>( Asset ) )
 			{
 				// Do nothing. Textures have no additional dependencies
 			}
+			else if ( USkeleton* Skeleton = Cast<USkeleton>( Asset ) )
+			{
+				// Do nothing. Skeletons have no additional dependencies
+			}
 			else
 			{
 				UE_LOG( LogUsd, Warning, TEXT( "Unknown asset '%s' encountered when collecting used assets before USD import." ), Asset ? *Asset->GetName() : TEXT( "nullptr" ) );
@@ -1193,7 +1242,7 @@ void UUsdStageImporter::ImportFromFile(FUsdStageImportContext& ImportContext)
 	}
 	TranslationContext->CompleteTasks();
 
-	UsdStageImporterImpl::CollectUsedAssetDependencies( ImportContext.AssetCache->GetActiveAssets(), UsedAssetsAndDependencies );
+	UsdStageImporterImpl::CollectUsedAssetDependencies( ImportContext.AssetCache, UsedAssetsAndDependencies );
 	UsdStageImporterImpl::UpdateAssetImportData( UsedAssetsAndDependencies, ImportContext.FilePath, ImportContext.ImportOptions );
 	UsdStageImporterImpl::PublishAssets( ImportContext, UsedAssetsAndDependencies, ObjectsToRemap );
 	UsdStageImporterImpl::ResolveActorConflicts( ImportContext, ExistingSceneActor, ObjectsToRemap );
diff --git a/Engine/Plugins/Importers/USDImporter/Source/USDUtilities/Private/USDGeomMeshConversion.cpp b/Engine/Plugins/Importers/USDImporter/Source/USDUtilities/Private/USDGeomMeshConversion.cpp
--- a/Engine/Plugins/Importers/USDImporter/Source/USDUtilities/Private/USDGeomMeshConversion.cpp
+++ b/Engine/Plugins/Importers/USDImporter/Source/USDUtilities/Private/USDGeomMeshConversion.cpp
@@ -1019,7 +1019,7 @@ UsdUtils::FUsdPrimMaterialAssignmentInfo UsdUtils::GetPrimMaterialAssignments( c
 				UE_LOG( LogUsd, Warning, TEXT( "String array attribute 'unrealMaterials' is deprecated: Use the singular string 'unrealMaterial' attribute" ) );
 				if ( UEMaterials.size() > 1 )
 				{
-					UE_LOG( LogUsd, Warning, TEXT( "Found more than on Unreal material assigned to Mesh '%s'. The first material ('%s') will be chosen, and the rest ignored." ),
+					UE_LOG( LogUsd, Warning, TEXT( "Found more than one Unreal material assigned to Mesh '%s'. The first material ('%s') will be chosen, and the rest ignored." ),
 						*UsdToUnreal::ConvertPath( UsdPrim.GetPath() ), *ValidPackagePath );
 				}
 			}
@@ -1597,8 +1597,6 @@ bool UsdUtils::IsGeomMeshALOD( const pxr::UsdPrim& UsdMeshPrim )
 
 	FScopedUsdAllocs Allocs;
 
-	std::string MeshName = UsdMeshPrim.GetName();
-
 	pxr::UsdPrim ParentPrim = UsdMeshPrim.GetParent();
 	if ( !ParentPrim )
 	{
@@ -1613,15 +1611,14 @@ bool UsdUtils::IsGeomMeshALOD( const pxr::UsdPrim& UsdMeshPrim )
 		return false;
 	}
 
-	for ( const std::string& VariantName : VariantSets.GetVariantSet( LODString ).GetVariantNames() )
+	std::string Selection = VariantSets.GetVariantSet( LODString ).GetVariantSelection();
+	int32 LODIndex = UsdGeomMeshImpl::GetLODIndexFromName( Selection );
+	if ( LODIndex == INDEX_NONE )
 	{
-		if ( VariantName == MeshName )
-		{
-			return true;
-		}
+		return false;
 	}
 
-	return false;
+	return true;
 }
 
 int32 UsdUtils::GetNumberOfLODVariants( const pxr::UsdPrim& Prim )
@@ -1673,7 +1670,31 @@ bool UsdUtils::IterateLODMeshes( const pxr::UsdPrim& ParentPrim, TFunction<bool(
 
 		LODVariantSet.SetVariantSelection( LODVariantName );
 
-		pxr::UsdGeomMesh LODMesh = pxr::UsdGeomMesh(ParentPrim.GetChild( pxr::TfToken( LODVariantName ) ));
+		pxr::UsdGeomMesh LODMesh;
+		pxr::TfToken TargetChildNameToken{ LODVariantName };
+
+		// Search for our LOD child mesh
+		pxr::UsdPrimSiblingRange PrimRange = ParentPrim.GetChildren();
+		for ( pxr::UsdPrimSiblingRange::iterator PrimRangeIt = PrimRange.begin(); PrimRangeIt != PrimRange.end(); ++PrimRangeIt )
+		{
+			const pxr::UsdPrim& Child = *PrimRangeIt;
+			if ( pxr::UsdGeomMesh ChildMesh{ Child } )
+			{
+				if ( Child.GetName() == TargetChildNameToken )
+				{
+					LODMesh = ChildMesh;
+					// Don't break here so we can show warnings if the user has other prims here (that we may end up ignoring)
+					// USD doesn't allow name collisions anyway, so there won't be any other prim named TargetChildNameToken
+				}
+				else
+				{
+					UE_LOG(LogUsd, Warning, TEXT("Unexpected prim '%s' inside LOD variant '%s'. For automatic parsing of LODs, each LOD variant should contain only a single Mesh prim named the same as the variant!"),
+						*UsdToUnreal::ConvertPath( Child.GetPath() ),
+						*UsdToUnreal::ConvertString( LODVariantName )
+					);
+				}
+			}
+		}
 		if ( !LODMesh )
 		{
 			continue;
diff --git a/Engine/Plugins/Importers/USDImporter/Source/USDUtilities/Private/USDShadeConversion.cpp b/Engine/Plugins/Importers/USDImporter/Source/USDUtilities/Private/USDShadeConversion.cpp
--- a/Engine/Plugins/Importers/USDImporter/Source/USDUtilities/Private/USDShadeConversion.cpp
+++ b/Engine/Plugins/Importers/USDImporter/Source/USDUtilities/Private/USDShadeConversion.cpp
@@ -186,7 +186,7 @@ namespace UE
 			}
 
 			/** If ResolvedTexturePath points at a texture inside an usdz file, this will use USD to pull the asset from the file, and TextureFactory to import it directly from the binary buffer */
-			UTexture* ReadTextureFromUsdzArchiveEditor(const FString& ResolvedTexturePath, const FString& TextureExtension, UTextureFactory* TextureFactory )
+			UTexture* ReadTextureFromUsdzArchiveEditor(const FString& ResolvedTexturePath, const FString& TextureExtension, UTextureFactory* TextureFactory, UObject* Outer, EObjectFlags ObjectFlags )
 			{
 #if WITH_EDITOR
 				uint64 BufferSize = 0;
@@ -202,9 +202,9 @@ namespace UE
 
 				return Cast<UTexture>(TextureFactory->FactoryCreateBinary(
 					UTexture::StaticClass(),
-					GetTransientPackage(),
+					Outer,
 					NAME_None,
-					RF_Transient,
+					ObjectFlags,
 					nullptr,
 					*TextureExtension,
 					BufferStart,
@@ -947,28 +947,26 @@ namespace UE
 					const FString ResolvedTexturePath = UsdUtils::GetResolvedTexturePath( TextureAssetPathAttr );
 					if ( !ResolvedTexturePath.IsEmpty() )
 					{
-						// Try checking if the texture is inside an USDZ archive first, or else TextureFactory throws an error
+						EObjectFlags ObjectFlags = RF_Transactional;
+						if ( !Outer )
+						{
+							Outer = GetTransientPackage();
+						}
+
+						if ( Outer == GetTransientPackage() )
+						{
+							ObjectFlags = ObjectFlags | RF_Transient;
+						}
+
 						FString TextureExtension;
 						if ( IsInsideUsdzArchive( ResolvedTexturePath, TextureExtension ) )
 						{
 							// Always prefer using the TextureFactory if we can, as it may provide compression, which the runtime version never will
-							Texture = ReadTextureFromUsdzArchiveEditor( ResolvedTexturePath, TextureExtension, TextureFactory );
+							Texture = ReadTextureFromUsdzArchiveEditor( ResolvedTexturePath, TextureExtension, TextureFactory, Outer, ObjectFlags );
 						}
-
 						// Not inside an USDZ archive, just a regular texture
-						if ( !Texture )
+						else
 						{
-							EObjectFlags ObjectFlags = RF_Transactional;
-							if ( !Outer )
-							{
-								Outer = GetTransientPackage();
-							}
-
-							if ( Outer == GetTransientPackage() )
-							{
-								ObjectFlags = ObjectFlags | RF_Transient;
-							}
-
 							Texture = Cast< UTexture >( TextureFactory->ImportObject( UTexture::StaticClass(), Outer, NAME_None, ObjectFlags, ResolvedTexturePath, TEXT( "" ), bOutCancelled ) );
 						}
 
