USD: Fix handling of visibility attribute when undo/redoing

diff --git a/Engine/Plugins/Importers/USDImporter/Source/USDStage/Private/USDPrimTwin.cpp b/Engine/Plugins/Importers/USDImporter/Source/USDStage/Private/USDPrimTwin.cpp
--- a/Engine/Plugins/Importers/USDImporter/Source/USDStage/Private/USDPrimTwin.cpp
+++ b/Engine/Plugins/Importers/USDImporter/Source/USDStage/Private/USDPrimTwin.cpp
@@ -23,6 +23,8 @@ UUsdPrimTwin& UUsdPrimTwin::AddChild( const FString& InPrimPath )
 	ChildPrim = NewObject<UUsdPrimTwin>( this, NAME_None, RF_Transient | RF_Transactional | RF_Public );
 	ChildPrim->PrimPath = InPrimPath;
 
+	ChildPrim->Parent = this;
+
 	return *ChildPrim;
 }
 
@@ -36,6 +38,7 @@ void UUsdPrimTwin::RemoveChild( const TCHAR* InPrimPath )
 	{
 		if ( ChildIt->Value->PrimPath == InPrimPath )
 		{
+			ChildIt->Value->Parent.Reset();
 			ChildIt.RemoveCurrent();
 			break;
 		}
diff --git a/Engine/Plugins/Importers/USDImporter/Source/USDStage/Private/USDStageActor.cpp b/Engine/Plugins/Importers/USDImporter/Source/USDStage/Private/USDStageActor.cpp
--- a/Engine/Plugins/Importers/USDImporter/Source/USDStage/Private/USDStageActor.cpp
+++ b/Engine/Plugins/Importers/USDImporter/Source/USDStage/Private/USDStageActor.cpp
@@ -351,7 +351,7 @@ struct FUsdStageActorImpl
 			return false;
 		}
 
-		if ( !FilePath.IsEmpty() && !FPaths::IsRelative( FilePath ) && !FilePath.StartsWith( USD_IDENTIFIER_TOKEN ) )
+		if ( !FilePath.IsEmpty() && !FPaths::IsRelative( FilePath ) && !FilePath.StartsWith( UnrealIdentifiers::IdentifierPrefix ) )
 		{
 			FilePath = UsdUtils::MakePathRelativeToProjectDir( FilePath );
 		}
@@ -363,7 +363,7 @@ struct FUsdStageActorImpl
 			const FString LayerIdentifier = Layer.GetIdentifier();
 
 			// Something like "@identifier:anon:0000022F9E194D50:tmp.usda" if we're also pointing at an anonymous layer
-			if ( FilePath.RemoveFromStart( USD_IDENTIFIER_TOKEN ) )
+			if ( FilePath.RemoveFromStart( UnrealIdentifiers::IdentifierPrefix ) )
 			{
 				// Same anonymous layers
 				if ( FilePath == LayerIdentifier )
@@ -385,6 +385,128 @@ struct FUsdStageActorImpl
 
 		return false;
 	}
+
+	/**
+	 * Uses USD's MakeVisible to handle the visible/inherited update logic as it is a bit complex.
+	 * Will update a potentially large chunk of the component hierarchy to having/not the `invisible` component tag, as well as the
+	 * correct value of bVisible.
+	 * Note that bVisible corresponds to computed visibility, and the component tags correspond to individual prim-level visibilities
+	 */
+	static void MakeVisible( UUsdPrimTwin& UsdPrimTwin, UE::FUsdStage& Stage )
+	{
+		// Find the highest invisible prim parent: Nothing above this can possibly change with what we're doing
+		UUsdPrimTwin* Iter = &UsdPrimTwin;
+		UUsdPrimTwin* HighestInvisibleParent = nullptr;
+		while ( Iter )
+		{
+			if ( USceneComponent* Component = Iter->GetSceneComponent() )
+			{
+				if ( Component->ComponentTags.Contains( UnrealIdentifiers::Invisible ) )
+				{
+					HighestInvisibleParent = Iter;
+				}
+			}
+
+			Iter = Iter->GetParent();
+		}
+
+		// No parent (not even UsdPrimTwin's prim directly) was invisible, so we should already be visible and there's nothing to do
+		if ( !HighestInvisibleParent )
+		{
+			return;
+		}
+
+		UE::FUsdPrim Prim = Stage.GetPrimAtPath( UE::FSdfPath( *UsdPrimTwin.PrimPath ) );
+		if ( !Prim )
+		{
+			return;
+		}
+		UsdUtils::MakeVisible( Prim );
+
+		TFunction<void(UUsdPrimTwin&, bool)> RecursiveResyncVisibility;
+		RecursiveResyncVisibility = [ &Stage, &RecursiveResyncVisibility ]( UUsdPrimTwin& PrimTwin, bool bPrimHasInvisibleParent )
+		{
+			USceneComponent* Component = PrimTwin.GetSceneComponent();
+			if ( !Component )
+			{
+				return;
+			}
+
+			UE::FUsdPrim CurrentPrim = Stage.GetPrimAtPath( UE::FSdfPath( *PrimTwin.PrimPath ) );
+			if ( !CurrentPrim )
+			{
+				return;
+			}
+
+			const bool bPrimHasInheritedVisibility = UsdUtils::HasInheritedVisibility( CurrentPrim );
+			const bool bPrimIsVisible = bPrimHasInheritedVisibility && !bPrimHasInvisibleParent;
+
+			const bool bComponentHasInvisibleTag = Component->ComponentTags.Contains( UnrealIdentifiers::Invisible );
+			const bool bComponentIsVisible = Component->IsVisible();
+
+			const bool bTagIsCorrect = bComponentHasInvisibleTag == !bPrimHasInheritedVisibility;
+			const bool bComputedVisibilityIsCorrect = bPrimIsVisible == bComponentIsVisible;
+
+			// Stop recursing as this prim's or its children couldn't possibly need to be updated
+			if ( bTagIsCorrect && bComputedVisibilityIsCorrect )
+			{
+				return;
+			}
+
+			if ( !bTagIsCorrect )
+			{
+				if ( bPrimHasInheritedVisibility )
+				{
+					Component->ComponentTags.Remove( UnrealIdentifiers::Invisible );
+					Component->ComponentTags.AddUnique( UnrealIdentifiers::Inherited );
+				}
+				else
+				{
+					Component->ComponentTags.AddUnique( UnrealIdentifiers::Invisible );
+					Component->ComponentTags.Remove( UnrealIdentifiers::Inherited );
+				}
+			}
+
+			if ( !bComputedVisibilityIsCorrect )
+			{
+				const bool bPropagateToChildren = false;
+				Component->Modify();
+				Component->SetVisibility( bPrimIsVisible, bPropagateToChildren );
+			}
+
+			for ( const TPair<FString, UUsdPrimTwin*>& ChildPair : PrimTwin.GetChildren() )
+			{
+				if ( UUsdPrimTwin* ChildTwin = ChildPair.Value )
+				{
+					RecursiveResyncVisibility( *ChildTwin, !bPrimIsVisible );
+				}
+			}
+		};
+
+		const bool bHasInvisibleParent = false;
+		RecursiveResyncVisibility( *HighestInvisibleParent, bHasInvisibleParent );
+	}
+
+	/**
+	 * Sets this prim to 'invisible', and force all of the child components
+	 * to bVisible = false. Leave their individual prim-level visibilities intact though.
+	 * Note that bVisible corresponds to computed visibility, and the component tags correspond to individual prim-level visibilities
+	 */
+	static void MakeInvisible( UUsdPrimTwin& UsdPrimTwin )
+	{
+		USceneComponent* PrimSceneComponent = UsdPrimTwin.GetSceneComponent();
+		if ( !PrimSceneComponent )
+		{
+			return;
+		}
+
+		PrimSceneComponent->ComponentTags.AddUnique( UnrealIdentifiers::Invisible );
+		PrimSceneComponent->ComponentTags.Remove( UnrealIdentifiers::Inherited );
+
+		const bool bPropagateToChildren = true;
+		const bool bVisible = false;
+		PrimSceneComponent->SetVisibility( bVisible, bPropagateToChildren );
+	}
 };
 
 /**
@@ -1030,7 +1152,7 @@ void AUsdStageActor::SetRootLayer( const FString& RootFilePath )
 {
 	FString RelativeFilePath = RootFilePath;
 #if USE_USD_SDK
-	if ( !RelativeFilePath.IsEmpty() && !FPaths::IsRelative( RelativeFilePath ) && !RelativeFilePath.StartsWith( USD_IDENTIFIER_TOKEN ) )
+	if ( !RelativeFilePath.IsEmpty() && !FPaths::IsRelative( RelativeFilePath ) && !RelativeFilePath.StartsWith( UnrealIdentifiers::IdentifierPrefix ) )
 	{
 		RelativeFilePath = UsdUtils::MakePathRelativeToProjectDir( RootFilePath );
 	}
@@ -1168,7 +1290,7 @@ void AUsdStageActor::OpenUsdStage()
 	UsdUtils::StartMonitoringErrors();
 
 	FString AbsPath;
-	if ( !RootLayer.FilePath.StartsWith( USD_IDENTIFIER_TOKEN ) && FPaths::IsRelative( RootLayer.FilePath ) )
+	if ( !RootLayer.FilePath.StartsWith( UnrealIdentifiers::IdentifierPrefix ) && FPaths::IsRelative( RootLayer.FilePath ) )
 	{
 		// The RootLayer property is marked as RelativeToGameDir, and UsdUtils::BrowseUsdFile will also emit paths relative to the project's directory
 		FString ProjectDir = FPaths::ConvertRelativePathToFull( FPaths::ProjectDir() );
@@ -1905,10 +2027,31 @@ void AUsdStageActor::OnObjectPropertyChanged( UObject* ObjectBeingModified, FPro
 		{
 			if ( UsdStage )
 			{
+				// This block is important, as it not only prevents us from getting into infinite loops with the USD notices,
+				// but it also guarantees that if we have an object property change, the corresponding stage notice is not also
+				// independently saved to the transaction via the UUsdTransactor, which would be duplication
 				FScopedBlockNoticeListening BlockNotices( this );
 
 				UE::FUsdPrim UsdPrim = UsdStage.GetPrimAtPath( UE::FSdfPath( *PrimPath ) );
 
+				// We want to keep component visibilities in sync with USD, which uses inherited visibilities
+				// To accomplish that while blocking notices we must always propagate component visibility changes manually.
+				// This part is effectively the same as calling pxr::UsdGeomImageable::MakeVisible/Invisible.
+				if ( PropertyChangedEvent.GetPropertyName() == TEXT( "bVisible" ) )
+				{
+					PrimSceneComponent->Modify();
+
+					const bool bVisible = PrimSceneComponent->GetVisibleFlag();
+					if ( bVisible )
+					{
+						FUsdStageActorImpl::MakeVisible( *UsdPrimTwin, UsdStage );
+					}
+					else
+					{
+						FUsdStageActorImpl::MakeInvisible( *UsdPrimTwin );
+					}
+				}
+
 #if USE_USD_SDK
 				UnrealToUsd::ConvertSceneComponent( UsdStage, PrimSceneComponent, UsdPrim );
 
@@ -1959,14 +2102,6 @@ void AUsdStageActor::OnObjectPropertyChanged( UObject* ObjectBeingModified, FPro
 				}
 #endif // #if USE_USD_SDK
 
-				// We want to keep component visibilities in sync with USD, which uses inherited visibilities
-				// To accomplish that while blocking notices we must always propagate component visibility changes
-				if ( PropertyChangedEvent.GetPropertyName() == TEXT( "bVisible" ) )
-				{
-					PrimSceneComponent->Modify();
-					PrimSceneComponent->SetVisibility( PrimSceneComponent->GetVisibleFlag(), true );
-				}
-
 				// Update stage window in case any of our component changes trigger USD stage changes
 				this->OnPrimChanged.Broadcast( PrimPath, false );
 			}
diff --git a/Engine/Plugins/Importers/USDImporter/Source/USDStage/Public/USDPrimTwin.h b/Engine/Plugins/Importers/USDImporter/Source/USDStage/Public/USDPrimTwin.h
--- a/Engine/Plugins/Importers/USDImporter/Source/USDStage/Public/USDPrimTwin.h
+++ b/Engine/Plugins/Importers/USDImporter/Source/USDStage/Public/USDPrimTwin.h
@@ -24,6 +24,9 @@ class UUsdPrimTwin final : public UObject
 public:
 	UUsdPrimTwin& AddChild( const FString& InPrimPath );
 	void RemoveChild( const TCHAR* InPrimPath );
+	const TMap< FString, UUsdPrimTwin* >& GetChildren() const { return Children; }
+
+	UUsdPrimTwin* GetParent() const { return Parent.Get(); }
 
 	void Clear();
 
@@ -67,4 +70,7 @@ public:
 private:
 	UPROPERTY( Transient )
 	TMap< FString, UUsdPrimTwin* > Children;
+
+	UPROPERTY( Transient )
+	TWeakObjectPtr< UUsdPrimTwin > Parent;
 };
diff --git a/Engine/Plugins/Importers/USDImporter/Source/USDStageEditorViewModels/Private/USDStageViewModel.cpp b/Engine/Plugins/Importers/USDImporter/Source/USDStageEditorViewModels/Private/USDStageViewModel.cpp
--- a/Engine/Plugins/Importers/USDImporter/Source/USDStageEditorViewModels/Private/USDStageViewModel.cpp
+++ b/Engine/Plugins/Importers/USDImporter/Source/USDStageEditorViewModels/Private/USDStageViewModel.cpp
@@ -77,7 +77,7 @@ void FUsdStageViewModel::NewStage( const TCHAR* FilePath )
 		UE::FSdfLayer Layer = UsdStage.GetRootLayer();
 		if ( Layer )
 		{
-			StagePath = FString( USD_IDENTIFIER_TOKEN ) + Layer.GetIdentifier();
+			StagePath = FString( UnrealIdentifiers::IdentifierPrefix ) + Layer.GetIdentifier();
 		}
 	}
 
diff --git a/Engine/Plugins/Importers/USDImporter/Source/USDUtilities/Private/USDConversionUtils.cpp b/Engine/Plugins/Importers/USDImporter/Source/USDUtilities/Private/USDConversionUtils.cpp
--- a/Engine/Plugins/Importers/USDImporter/Source/USDUtilities/Private/USDConversionUtils.cpp
+++ b/Engine/Plugins/Importers/USDImporter/Source/USDUtilities/Private/USDConversionUtils.cpp
@@ -953,6 +953,76 @@ FString UsdUtils::SanitizeUsdIdentifier( const TCHAR* InIdentifier )
 	return InIdentifier;
 }
 
+void UsdUtils::MakeVisible( UE::FUsdPrim& Prim, double TimeCode )
+{
+#if USE_USD_SDK
+	FScopedUsdAllocs Allocs;
+
+	pxr::UsdPrim PxrUsdPrim{ Prim };
+	if ( pxr::UsdGeomImageable Imageable{ PxrUsdPrim } )
+	{
+		Imageable.MakeVisible( TimeCode );
+	}
+#endif // USE_USD_SDK
+}
+
+void UsdUtils::MakeInvisible( UE::FUsdPrim& Prim, double TimeCode )
+{
+#if USE_USD_SDK
+	FScopedUsdAllocs Allocs;
+
+	pxr::UsdPrim PxrUsdPrim{ Prim };
+	if ( pxr::UsdGeomImageable Imageable{ PxrUsdPrim } )
+	{
+		Imageable.MakeInvisible( TimeCode );
+	}
+#endif // USE_USD_SDK
+}
+
+bool UsdUtils::IsVisible( UE::FUsdPrim& Prim, double TimeCode )
+{
+#if USE_USD_SDK
+	FScopedUsdAllocs Allocs;
+
+	pxr::UsdPrim PxrUsdPrim{ Prim };
+	if ( pxr::UsdGeomImageable Imageable{ PxrUsdPrim } )
+	{
+		return Imageable.ComputeVisibility( TimeCode ) == pxr::UsdGeomTokens->inherited;
+	}
+
+	return true;
+#else
+	return false;
+#endif // USE_USD_SDK
+}
+
+bool UsdUtils::HasInheritedVisibility( UE::FUsdPrim& Prim, double TimeCode )
+{
+#if USE_USD_SDK
+	FScopedUsdAllocs Allocs;
+
+	pxr::UsdPrim PxrUsdPrim{ Prim };
+	if ( pxr::UsdGeomImageable Imageable{ PxrUsdPrim } )
+	{
+		if ( pxr::UsdAttribute VisibilityAttr = Imageable.GetVisibilityAttr() )
+		{
+			pxr::TfToken Visibility;
+			if ( !VisibilityAttr.Get<pxr::TfToken>( &Visibility, TimeCode ) )
+			{
+				return true;
+			}
+
+			return Visibility == pxr::UsdGeomTokens->inherited;
+		}
+	}
+
+	// If it doesn't have the attribute the default is for it to be 'inherited'
+	return true;
+#else
+	return false;
+#endif // USE_USD_SDK
+}
+
 #if USE_USD_SDK
 #undef LOCTEXT_NAMESPACE
 #endif
diff --git a/Engine/Plugins/Importers/USDImporter/Source/USDUtilities/Private/USDPrimConversion.cpp b/Engine/Plugins/Importers/USDImporter/Source/USDUtilities/Private/USDPrimConversion.cpp
--- a/Engine/Plugins/Importers/USDImporter/Source/USDUtilities/Private/USDPrimConversion.cpp
+++ b/Engine/Plugins/Importers/USDImporter/Source/USDUtilities/Private/USDPrimConversion.cpp
@@ -112,12 +112,33 @@ bool UsdToUnreal::ConvertXformable( const pxr::UsdStageRefPtr& Stage, const pxr:
 
 	SceneComponent.SetRelativeTransform( Transform );
 
-	// Visibility
-	const bool bIsHidden = ( Xformable.ComputeVisibility( EvalTime ) == pxr::UsdGeomTokens->invisible );
-
 	SceneComponent.Modify();
+
+	// Computed (effective) visibility
+	const bool bIsHidden = ( Xformable.ComputeVisibility( EvalTime ) == pxr::UsdGeomTokens->invisible );
 	SceneComponent.SetVisibility( !bIsHidden );
 
+	// Per-prim visibility
+	bool bIsInvisible = false; // Default to 'inherited'
+	if ( pxr::UsdAttribute VisibilityAttr = Xformable.GetVisibilityAttr() )
+	{
+		pxr::TfToken Value;
+		if ( VisibilityAttr.Get( &Value, EvalTime ) )
+		{
+			bIsInvisible = Value == pxr::UsdGeomTokens->invisible;
+		}
+	}
+	if ( bIsInvisible )
+	{
+		SceneComponent.ComponentTags.AddUnique( UnrealIdentifiers::Invisible );
+		SceneComponent.ComponentTags.Remove( UnrealIdentifiers::Inherited );
+	}
+	else
+	{
+		SceneComponent.ComponentTags.Remove( UnrealIdentifiers::Invisible );
+		SceneComponent.ComponentTags.AddUnique( UnrealIdentifiers::Inherited );
+	}
+
 	return true;
 }
 
@@ -321,15 +342,22 @@ bool UnrealToUsd::ConvertSceneComponent( const pxr::UsdStageRefPtr& Stage, const
 	// Transform
 	ConvertXformable( RelativeTransform, UsdPrim, UsdUtils::GetDefaultTimeCode() );
 
-	// Visibility
-	bool bVisible = SceneComponent->GetVisibleFlag();
-	if ( bVisible )
-	{
-		XForm.MakeVisible();
-	}
-	else
+	// Per-prim visibility
+	if ( pxr::UsdAttribute VisibilityAttr = XForm.CreateVisibilityAttr() )
 	{
-		XForm.MakeInvisible();
+		pxr::TfToken Value = pxr::UsdGeomTokens->inherited;
+
+		if ( SceneComponent->ComponentTags.Contains( UnrealIdentifiers::Invisible ) )
+		{
+			Value = pxr::UsdGeomTokens->invisible;
+		}
+		else if ( !SceneComponent->ComponentTags.Contains( UnrealIdentifiers::Inherited ) )
+		{
+			// We don't have visible nor inherited tags: We're probably exporting a pure UE component, so write out component visibility instead
+			Value = SceneComponent->IsVisible() ? pxr::UsdGeomTokens->inherited : pxr::UsdGeomTokens->invisible;
+		}
+
+		VisibilityAttr.Set<pxr::TfToken>( Value );
 	}
 
 	return true;
diff --git a/Engine/Plugins/Importers/USDImporter/Source/USDUtilities/Public/USDConversionUtils.h b/Engine/Plugins/Importers/USDImporter/Source/USDUtilities/Public/USDConversionUtils.h
--- a/Engine/Plugins/Importers/USDImporter/Source/USDUtilities/Public/USDConversionUtils.h
+++ b/Engine/Plugins/Importers/USDImporter/Source/USDUtilities/Public/USDConversionUtils.h
@@ -167,5 +167,15 @@ namespace UsdUtils
 	 * Note that this obviously doesn't check for a potential name collision.
 	 */
 	USDUTILITIES_API FString SanitizeUsdIdentifier( const TCHAR* InIdentifier );
+
+	/** Will call UsdGeomImageable::MakeVisible/MakeInvisible if Prim is a UsdGeomImageable */
+	USDUTILITIES_API void MakeVisible( UE::FUsdPrim& Prim, double TimeCode = UsdUtils::GetDefaultTimeCode() );
+	USDUTILITIES_API void MakeInvisible( UE::FUsdPrim& Prim, double TimeCode = UsdUtils::GetDefaultTimeCode() );
+
+	/** Returns if the ComputedVisibility for Prim says it should be visible */
+	USDUTILITIES_API bool IsVisible( UE::FUsdPrim& Prim, double TimeCode = UsdUtils::GetDefaultTimeCode() );
+
+	/** Returns whether Prim has visibility set to 'inherited' */
+	USDUTILITIES_API bool HasInheritedVisibility( UE::FUsdPrim& Prim, double TimeCode = UsdUtils::GetDefaultTimeCode() );
 }
 
diff --git a/Engine/Plugins/Importers/USDImporter/Source/UnrealUSDWrapper/Private/UnrealUSDWrapper.cpp b/Engine/Plugins/Importers/USDImporter/Source/UnrealUSDWrapper/Private/UnrealUSDWrapper.cpp
--- a/Engine/Plugins/Importers/USDImporter/Source/UnrealUSDWrapper/Private/UnrealUSDWrapper.cpp
+++ b/Engine/Plugins/Importers/USDImporter/Source/UnrealUSDWrapper/Private/UnrealUSDWrapper.cpp
@@ -856,6 +856,10 @@ EUsdGeomOrientation IUsdPrim::GetGeometryOrientation(const pxr::UsdGeomMesh& Mes
 }
 #endif // USE_USD_SDK
 
+const TCHAR* UnrealIdentifiers::Invisible = TEXT("invisible");
+const TCHAR* UnrealIdentifiers::Inherited = TEXT("inherited");
+const TCHAR* UnrealIdentifiers::IdentifierPrefix = TEXT("@identifier:");
+
 FUsdDelegates::FUsdImportDelegate FUsdDelegates::OnPreUsdImport;
 FUsdDelegates::FUsdImportDelegate FUsdDelegates::OnPostUsdImport;
 
@@ -983,7 +987,7 @@ UE::FUsdStage UnrealUSDWrapper::OpenStage( const TCHAR* Identifier, EUsdInitialL
 	{
 		Stage = pxr::UsdStage::Open( TCHAR_TO_ANSI( Identifier ), pxr::UsdStage::InitialLoadSet( InitialLoadSet ) );
 	}
-	else if ( IdentifierStr.RemoveFromStart( USD_IDENTIFIER_TOKEN ) )
+	else if ( IdentifierStr.RemoveFromStart( UnrealIdentifiers::IdentifierPrefix ) )
 	{
 		pxr::SdfLayerRefPtr RootLayer = pxr::SdfLayer::Find( TCHAR_TO_ANSI( *IdentifierStr ) );
 		if ( RootLayer )
diff --git a/Engine/Plugins/Importers/USDImporter/Source/UnrealUSDWrapper/Public/UnrealUSDWrapper.h b/Engine/Plugins/Importers/USDImporter/Source/UnrealUSDWrapper/Public/UnrealUSDWrapper.h
--- a/Engine/Plugins/Importers/USDImporter/Source/UnrealUSDWrapper/Public/UnrealUSDWrapper.h
+++ b/Engine/Plugins/Importers/USDImporter/Source/UnrealUSDWrapper/Public/UnrealUSDWrapper.h
@@ -19,8 +19,6 @@
 #include "USDIncludesEnd.h"
 #endif // #if USE_USD_SDK
 
-#define USD_IDENTIFIER_TOKEN TEXT("@identifier:")
-
 #if USE_USD_SDK
 PXR_NAMESPACE_OPEN_SCOPE
 	class GfMatrix4d;
@@ -315,8 +313,11 @@ namespace UnrealIdentifiers
 
 	// Token used to indicate that a material parsed from a material prim should use world space normals
 	extern UNREALUSDWRAPPER_API const pxr::TfToken WorldSpaceNormals;
-
 #endif // #if USE_USD_SDK
+
+	extern UNREALUSDWRAPPER_API const TCHAR* Invisible;
+	extern UNREALUSDWRAPPER_API const TCHAR* Inherited;
+	extern UNREALUSDWRAPPER_API const TCHAR* IdentifierPrefix;
 }
 
 struct UNREALUSDWRAPPER_API FUsdDelegates
