USD: Fix components being left in an invalid state after undoing a transaction where their actor is deleted

diff --git a/Engine/Plugins/Importers/USDImporter/Source/USDStage/Private/USDPrimTwin.cpp b/Engine/Plugins/Importers/USDImporter/Source/USDStage/Private/USDPrimTwin.cpp
--- a/Engine/Plugins/Importers/USDImporter/Source/USDStage/Private/USDPrimTwin.cpp
+++ b/Engine/Plugins/Importers/USDImporter/Source/USDStage/Private/USDPrimTwin.cpp
@@ -74,12 +74,33 @@ void UUsdPrimTwin::Clear()
 
 	if ( ActorToDestroy && !ActorToDestroy->IsA< AUsdStageActor >() && !ActorToDestroy->IsActorBeingDestroyed() && ActorToDestroy->GetWorld() )
 	{
+		// We have to manually Modify() all the actor's components because they're transient, so USceneComponent::DetachFromComponent
+		// won't automatically Modify them before detaching. If we don't do this they may be first recorded into the transaction in
+		// the detached state, so if that transaction is undone they'd be left detached
+		TArray<USceneComponent*> ChildComponents;
+		ActorToDestroy->GetComponents(ChildComponents);
+		for ( USceneComponent* Component : ChildComponents )
+		{
+			Component->Modify();
+		}
+
 		ActorToDestroy->Modify();
 		ActorToDestroy->GetWorld()->DestroyActor( ActorToDestroy );
 		SpawnedActor = nullptr;
 	}
 	else if ( SceneComponent.IsValid() && !SceneComponent->IsBeingDestroyed() )
 	{
+		// See comment above: USceneComponent::DetachFromComponent won't Modify our components since they're transient,
+		// so we need to do so manually
+		if ( USceneComponent* AttachParent = SceneComponent->GetAttachParent() )
+		{
+			AttachParent->Modify();
+		}
+		for ( USceneComponent* AttachChild : SceneComponent->GetAttachChildren() )
+		{
+			AttachChild->Modify();
+		}
+
 		SceneComponent->Modify();
 		SceneComponent->DestroyComponent();
 		SceneComponent = nullptr;
