USD: Fix skeletal mesh component dropping assigned skeletal mesh after undo and dragging of Time property.

diff --git a/Engine/Plugins/Importers/USDImporter/Source/USDClasses/Private/USDAssetCache.cpp b/Engine/Plugins/Importers/USDImporter/Source/USDClasses/Private/USDAssetCache.cpp
--- a/Engine/Plugins/Importers/USDImporter/Source/USDClasses/Private/USDAssetCache.cpp
+++ b/Engine/Plugins/Importers/USDImporter/Source/USDClasses/Private/USDAssetCache.cpp
@@ -157,16 +157,13 @@ void UUsdAssetCache::Reset()
 {
 	FScopeLock Lock( &CriticalSection );
 
+	Modify();
+
 	TransientStorage.Reset();
 	OwnedAssets.Reset();
 	PrimPathToAssets.Reset();
 	ActiveAssets.Reset();
-
-	if ( PersistentStorage.Num() > 0 )
-	{
-		Modify();
-		PersistentStorage.Reset();
-	}
+	PersistentStorage.Reset();
 }
 
 void UUsdAssetCache::MarkAssetsAsStale()
@@ -190,8 +187,6 @@ void UUsdAssetCache::Serialize( FArchive& Ar )
 	if ( Ar.GetPortFlags() & PPF_DuplicateForPIE )
 	{
 		Ar << TransientStorage;
-		Ar << OwnedAssets;
-		Ar << PrimPathToAssets;
 		Ar << ActiveAssets;
 	}
 }
diff --git a/Engine/Plugins/Importers/USDImporter/Source/USDClasses/Public/USDAssetCache.h b/Engine/Plugins/Importers/USDImporter/Source/USDClasses/Public/USDAssetCache.h
--- a/Engine/Plugins/Importers/USDImporter/Source/USDClasses/Public/USDAssetCache.h
+++ b/Engine/Plugins/Importers/USDImporter/Source/USDClasses/Public/USDAssetCache.h
@@ -52,19 +52,21 @@ public:
 	virtual void Serialize( FArchive& Ar ) override;
 
 private:
-	UPROPERTY(transient, VisibleAnywhere, Category="Assets")
+	UPROPERTY( Transient, VisibleAnywhere, Category = "Assets" )
 	TMap< FString, UObject* > TransientStorage;
 
-	UPROPERTY(VisibleAnywhere, Category="Assets")
+	UPROPERTY( VisibleAnywhere, Category = "Assets" )
 	TMap< FString, UObject* > PersistentStorage;
 
-	UPROPERTY(EditAnywhere, Category="Assets", AdvancedDisplay)
+	UPROPERTY( EditAnywhere, Category = "Assets", AdvancedDisplay )
 	bool bAllowPersistentStorage;
 
 	// Points to the assets in primary storage, used to quickly check if we own an asset
+	UPROPERTY()
 	TSet< UObject* > OwnedAssets;
 
 	// Keeps associations from prim paths to assets that we own in primary storage
+	UPROPERTY()
     TMap< FString, UObject* > PrimPathToAssets;
 
 	// Assets that were added/retrieved since the last call to MarkAssetsAsSlate();
