USD: Expose mesh and landscape bake resolution to level exporter options

diff --git a/Engine/Plugins/Importers/USDImporter/Source/USDExporter/Public/LevelExporterUSDOptions.h b/Engine/Plugins/Importers/USDImporter/Source/USDExporter/Public/LevelExporterUSDOptions.h
--- a/Engine/Plugins/Importers/USDImporter/Source/USDExporter/Public/LevelExporterUSDOptions.h
+++ b/Engine/Plugins/Importers/USDImporter/Source/USDExporter/Public/LevelExporterUSDOptions.h
@@ -40,6 +40,10 @@ public:
     UPROPERTY( EditAnywhere, config, BlueprintReadWrite, Category = "Export settings" )
     bool bBakeMaterials;
 
+	/** Resolution to use when baking materials into textures */
+	UPROPERTY( EditAnywhere, config, BlueprintReadWrite, Category = "Export settings", meta = ( EditCondition = "bBakeMaterials", ClampMin = "1" ) )
+	FIntPoint BakeResolution = FIntPoint( 512, 512 );
+
 	/** Whether to remove the 'unrealMaterial' attribute after binding the corresponding baked material */
 	UPROPERTY( EditAnywhere, config, BlueprintReadWrite, Category = "Export settings", meta = ( EditCondition = "bBakeMaterials" ) )
 	bool bRemoveUnrealMaterials;
@@ -63,6 +67,10 @@ public:
 	UPROPERTY( EditAnywhere, config, BlueprintReadWrite, Category = "Export settings", meta = ( ClampMin = "0" ) )
 	int32 HighestLandscapeLOD;
 
+	/** Resolution to use when baking landscape materials into textures  */
+	UPROPERTY( EditAnywhere, config, BlueprintReadWrite, Category = "Export settings", meta = ( ClampMin = "1" ) )
+	FIntPoint LandscapeBakeResolution = FIntPoint( 1024, 1024 );
+
 	/** If true, will export sub-levels as separate layers (referenced as sublayers). If false, will collapse all sub-levels in a single exported root layer */
     UPROPERTY( EditAnywhere, config, BlueprintReadWrite, Category = "Sublayers" )
     bool bExportSublayers;
