USD: Improve UX of LevelFilter for python scene exporter.

Force-stream in sublevels that were picked to be exported but aren't currently visible.

diff --git a/Engine/Plugins/Importers/USDImporter/Source/USDExporter/Private/LevelExporterUSD.cpp b/Engine/Plugins/Importers/USDImporter/Source/USDExporter/Private/LevelExporterUSD.cpp
--- a/Engine/Plugins/Importers/USDImporter/Source/USDExporter/Private/LevelExporterUSD.cpp
+++ b/Engine/Plugins/Importers/USDImporter/Source/USDExporter/Private/LevelExporterUSD.cpp
@@ -15,9 +15,82 @@
 #include "UsdWrappers/UsdStage.h"
 
 #include "AssetExportTask.h"
+#include "EditorLevelUtils.h"
+#include "Engine/LevelStreaming.h"
 #include "Engine/World.h"
 #include "IPythonScriptPlugin.h"
 
+namespace UE
+{
+	namespace LevelExporterUSD
+	{
+		namespace Private
+		{
+			/**
+			 * Fully streams in all levels whose names are not in LevelsToIgnore, returning a list of all levels that were streamed in.
+			 */
+			TArray<ULevel*> StreamInRequiredLevels( UWorld* World, const TSet<FString>& LevelsToIgnore )
+			{
+				TArray<ULevel*> Result;
+				if ( !World )
+				{
+					return Result;
+				}
+
+				// Make sure all streamed levels are loaded so we can query their visibility
+				const bool bForce = true;
+				World->LoadSecondaryLevels( bForce );
+
+				if ( ULevel* PersistentLevel = World->PersistentLevel )
+				{
+					const FString LevelName = TEXT( "Persistent Level" );
+					if ( !PersistentLevel->bIsVisible && !LevelsToIgnore.Contains( LevelName ) )
+					{
+						Result.Add( PersistentLevel );
+					}
+				}
+
+				for ( ULevelStreaming* StreamingLevel : World->GetStreamingLevels() )
+				{
+					if ( StreamingLevel )
+					{
+						if ( ULevel* Level = StreamingLevel->GetLoadedLevel() )
+						{
+							const FString LevelName = Level->GetTypedOuter<UWorld>()->GetName();
+							if ( !Level->bIsVisible && !LevelsToIgnore.Contains( LevelName ) )
+							{
+								Result.Add( Level );
+							}
+						}
+					}
+				}
+
+				TArray<bool> ShouldBeVisible;
+				ShouldBeVisible.SetNumUninitialized(Result.Num());
+				for ( bool& bVal : ShouldBeVisible )
+				{
+					bVal = true;
+				}
+
+				const bool bForceLayersVisible = true;
+				EditorLevelUtils::SetLevelsVisibility( Result, ShouldBeVisible, bForceLayersVisible, ELevelVisibilityDirtyMode::DontModify );
+
+				return Result;
+			}
+
+			/** Streams out LevelsToStreamOut from World */
+			void StreamOutLevels( const TArray<ULevel*>& LevelsToStreamOut )
+			{
+				TArray<bool> ShouldBeVisible;
+				ShouldBeVisible.SetNumZeroed( LevelsToStreamOut.Num() );
+
+				const bool bForceLayersVisible = false;
+				EditorLevelUtils::SetLevelsVisibility( LevelsToStreamOut, ShouldBeVisible, bForceLayersVisible, ELevelVisibilityDirtyMode::DontModify );
+			}
+		}
+	}
+}
+
 ULevelExporterUSD::ULevelExporterUSD()
 {
 #if USE_USD_SDK
@@ -77,6 +150,10 @@ bool ULevelExporterUSD::ExportBinary( UObject* Object, const TCHAR* Type, FArchi
 		return false;
 	}
 
+	// The more robust thing is to force-stream-in all levels that are currently invisible/unloaded but we check to export.
+	// Not only to force actors to spawn, but also to make sure that when we want to bake a landscape it is visible
+	TArray<ULevel*> StreamedInLevels = UE::LevelExporterUSD::Private::StreamInRequiredLevels( World, Options->LevelsToIgnore );
+
 	// Note how we don't explicitly pass the Options down to Python here: We stash our desired export options on the CDO, and
 	// those are read from Python by executing export_with_cdo_options().
 	Options->CurrentTask = ExportTask;
@@ -86,6 +163,9 @@ bool ULevelExporterUSD::ExportBinary( UObject* Object, const TCHAR* Type, FArchi
 	}
 	Options->CurrentTask = nullptr;
 
+	// Return the newly streamed in levels to their old visibilities
+	UE::LevelExporterUSD::Private::StreamOutLevels( StreamedInLevels );
+
 	return true;
 #else
 	return false;
diff --git a/Engine/Plugins/Importers/USDImporter/Source/USDExporter/Private/LevelExporterUSDOptionsCustomization.cpp b/Engine/Plugins/Importers/USDImporter/Source/USDExporter/Private/LevelExporterUSDOptionsCustomization.cpp
new file mode 100644
--- /dev/null
+++ b/Engine/Plugins/Importers/USDImporter/Source/USDExporter/Private/LevelExporterUSDOptionsCustomization.cpp
@@ -0,0 +1,211 @@
+// Copyright Epic Games, Inc. All Rights Reserved.
+
+#include "LevelExporterUSDOptionsCustomization.h"
+
+#include "LevelExporterUSDOptions.h"
+
+#include "CoreMinimal.h"
+#include "DetailCategoryBuilder.h"
+#include "DetailLayoutBuilder.h"
+#include "DetailWidgetRow.h"
+#include "Editor.h"
+#include "Engine/LevelStreaming.h"
+#include "Modules/ModuleManager.h"
+#include "UObject/StrongObjectPtr.h"
+#include "Widgets/Input/SCheckBox.h"
+#include "Widgets/Layout/SScrollBorder.h"
+#include "Widgets/Views/SListView.h"
+
+#define LOCTEXT_NAMESPACE "LevelExporterUSDOptionsCustomization"
+
+namespace UE
+{
+	namespace LevelExporterOptions
+	{
+		namespace Private
+		{
+			class SLevelPickerRow : public STableRow< TSharedRef<FString> >
+			{
+			public:
+				SLATE_BEGIN_ARGS(SLevelPickerRow) {}
+				SLATE_END_ARGS()
+
+				void Construct( const FArguments& InArgs, const TSharedRef<STableViewBase>& OwnerTableView, TWeakPtr<FString> InEntry, ULevelExporterUSDOptions* Options )
+				{
+					STableRow::Construct( STableRow::FArguments(), OwnerTableView );
+
+					FString LevelName;
+					if ( TSharedPtr<FString> PinnedEntry = InEntry.Pin() )
+					{
+						LevelName = *PinnedEntry;
+					}
+
+					SetRowContent(
+						SNew( SHorizontalBox )
+
+						+ SHorizontalBox::Slot()
+						.HAlign( HAlign_Left )
+						.VAlign( VAlign_Center )
+						.MaxWidth(20)
+						[
+							SNew( SCheckBox )
+							.IsChecked_Lambda( [ LevelName, Options ]()
+							{
+								if ( Options )
+								{
+									return Options->LevelsToIgnore.Contains( LevelName )
+										? ECheckBoxState::Unchecked
+										: ECheckBoxState::Checked;
+								}
+
+								return ECheckBoxState::Undetermined;
+							})
+							.OnCheckStateChanged_Lambda( [ LevelName, Options ]( ECheckBoxState State )
+							{
+								if ( Options )
+								{
+									if ( State == ECheckBoxState::Checked )
+									{
+										Options->LevelsToIgnore.Remove( LevelName );
+									}
+									else if ( State == ECheckBoxState::Unchecked )
+									{
+										Options->LevelsToIgnore.Add( LevelName );
+									}
+								}
+							})
+						]
+
+						+ SHorizontalBox::Slot()
+						.HAlign( HAlign_Left )
+						.VAlign( VAlign_Center )
+						[
+							SNew( STextBlock )
+							.Text( FText::FromString( LevelName ) )
+							.Font( FEditorStyle::GetFontStyle( TEXT( "PropertyWindow.NormalFont" ) ) )
+						]
+					);
+				}
+			};
+
+			class SLevelPickerList : public SListView<TSharedRef<FString>>
+			{
+			public:
+				void Construct( const FArguments& InArgs, ULevelExporterUSDOptions* Options )
+				{
+					// We'll be writing directly to its LevelsToIgnore, so keep it alive while we're alive
+					OptionsPtr.Reset( Options );
+
+					Options->LevelsToIgnore.Reset();
+
+					if ( UWorld* EditorWorld = GEditor->GetEditorWorldContext().World() )
+					{
+						// Make sure all streamed levels are loaded so we can query their names and export them
+						const bool bForce = true;
+						EditorWorld->LoadSecondaryLevels( bForce );
+
+						if ( ULevel* PersistentLevel = EditorWorld->PersistentLevel )
+						{
+							const FString LevelName = TEXT( "Persistent Level" );
+							RootItems.Add( MakeShared< FString >( LevelName ) );
+
+							if ( !PersistentLevel->bIsVisible )
+							{
+								Options->LevelsToIgnore.Add( LevelName );
+							}
+						}
+
+						for ( ULevelStreaming* StreamingLevel : EditorWorld->GetStreamingLevels() )
+						{
+							if ( StreamingLevel )
+							{
+								if ( ULevel* Level = StreamingLevel->GetLoadedLevel() )
+								{
+									const FString LevelName = Level->GetTypedOuter<UWorld>()->GetName();
+									RootItems.Add( MakeShared< FString >( *LevelName ) );
+
+									if ( !Level->bIsVisible )
+									{
+										Options->LevelsToIgnore.Add( LevelName );
+									}
+								}
+							}
+						}
+					}
+
+					SListView::Construct
+					(
+						SListView::FArguments()
+						.ListItemsSource( &RootItems )
+						.SelectionMode( ESelectionMode::None )
+						.OnGenerateRow( this, &SLevelPickerList::OnGenerateRow, Options )
+					);
+				}
+
+			private:
+				TSharedRef< ITableRow > OnGenerateRow( TSharedRef<FString> InEntry, const TSharedRef<STableViewBase>& OwnerTable, ULevelExporterUSDOptions* Options ) const
+				{
+					return SNew( SLevelPickerRow, OwnerTable, InEntry, Options );
+				}
+
+				TStrongObjectPtr<ULevelExporterUSDOptions> OptionsPtr;
+				TArray< TSharedRef<FString> > RootItems;
+			};
+		}
+	}
+}
+namespace LevelExporterUSDImpl = UE::LevelExporterOptions::Private;
+
+TSharedRef<IDetailCustomization> FLevelExporterUSDOptionsCustomization::MakeInstance()
+{
+	return MakeShared<FLevelExporterUSDOptionsCustomization>();
+}
+
+void FLevelExporterUSDOptionsCustomization::CustomizeDetails(IDetailLayoutBuilder& DetailLayoutBuilder)
+{
+	TArray<TWeakObjectPtr<UObject>> SelectedObjects = DetailLayoutBuilder.GetSelectedObjects();
+	if ( SelectedObjects.Num() != 1 )
+	{
+		return;
+	}
+
+	ULevelExporterUSDOptions* Options = Cast< ULevelExporterUSDOptions>( SelectedObjects[ 0 ].Get() );
+	if ( !Options )
+	{
+		return;
+	}
+
+	TSharedPtr< LevelExporterUSDImpl::SLevelPickerList > PickerTree = SNew( LevelExporterUSDImpl::SLevelPickerList, Options );
+
+	TSharedRef<IPropertyHandle> LevelFilterProp = DetailLayoutBuilder.GetProperty( GET_MEMBER_NAME_CHECKED( ULevelExporterUSDOptions, LevelsToIgnore ) );
+	TSharedRef<IPropertyHandle> ExportSublayersProp = DetailLayoutBuilder.GetProperty( GET_MEMBER_NAME_CHECKED( ULevelExporterUSDOptions, bExportSublayers ) );
+
+	// Touch these properties and categories to enforce this ordering
+	DetailLayoutBuilder.EditCategory( TEXT( "Stage options" ) );
+	DetailLayoutBuilder.EditCategory( TEXT( "Export settings" ) );
+	DetailLayoutBuilder.EditCategory( TEXT( "Sublayers" ) );
+	DetailLayoutBuilder.AddPropertyToCategory( ExportSublayersProp );
+
+	DetailLayoutBuilder.HideProperty( LevelFilterProp );
+	DetailLayoutBuilder.AddCustomRowToCategory( LevelFilterProp, LevelFilterProp->GetPropertyDisplayName() )
+	.NameContent()
+	[
+		SNew( STextBlock )
+		.Text( FText::FromString( TEXT( "Levels to export" ) ) )
+		.Font( FEditorStyle::GetFontStyle( TEXT( "PropertyWindow.NormalFont" ) ) )
+	]
+	.ValueContent()
+	.MinDesiredWidth(300)
+	[
+		SNew(SScrollBorder, PickerTree.ToSharedRef())
+		[
+			SNew(SBox) // Prevent the list from expanding freely
+			.MaxDesiredHeight(200)
+			[
+				PickerTree.ToSharedRef()
+			]
+		]
+	];
+}
+
+#undef LOCTEXT_NAMESPACE
diff --git a/Engine/Plugins/Importers/USDImporter/Source/USDExporter/Private/LevelExporterUSDOptionsCustomization.h b/Engine/Plugins/Importers/USDImporter/Source/USDExporter/Private/LevelExporterUSDOptionsCustomization.h
new file mode 100644
--- /dev/null
+++ b/Engine/Plugins/Importers/USDImporter/Source/USDExporter/Private/LevelExporterUSDOptionsCustomization.h
@@ -0,0 +1,18 @@
+// Copyright Epic Games, Inc. All Rights Reserved.
+
+#pragma once
+
+#include "CoreMinimal.h"
+#include "IDetailCustomization.h"
+
+class IDetailLayoutBuilder;
+
+class FLevelExporterUSDOptionsCustomization : public IDetailCustomization
+{
+public:
+	static TSharedRef<IDetailCustomization> MakeInstance();
+
+	// IDetailCustomization interface
+	virtual void CustomizeDetails( IDetailLayoutBuilder& DetailLayout ) override;
+	// End of IDetailCustomization interface
+};
diff --git a/Engine/Plugins/Importers/USDImporter/Source/USDExporter/Private/USDExporterModule.cpp b/Engine/Plugins/Importers/USDImporter/Source/USDExporter/Private/USDExporterModule.cpp
--- a/Engine/Plugins/Importers/USDImporter/Source/USDExporter/Private/USDExporterModule.cpp
+++ b/Engine/Plugins/Importers/USDImporter/Source/USDExporter/Private/USDExporterModule.cpp
@@ -2,10 +2,28 @@
 
 #include "USDExporterModule.h"
 
+#include "LevelExporterUSDOptionsCustomization.h"
 #include "USDMemory.h"
 
+#include "Modules/ModuleManager.h"
+#include "PropertyEditorModule.h"
+
 class FUsdExporterModule : public IUsdExporterModule
 {
+public:
+	virtual void StartupModule() override
+	{
+		FPropertyEditorModule& PropertyModule = FModuleManager::LoadModuleChecked< FPropertyEditorModule >( TEXT( "PropertyEditor" ) );
+		PropertyModule.RegisterCustomClassLayout( TEXT( "LevelExporterUSDOptions" ), FOnGetDetailCustomizationInstance::CreateStatic( &FLevelExporterUSDOptionsCustomization::MakeInstance ) );
+	}
+
+	virtual void ShutdownModule() override
+	{
+		if ( FPropertyEditorModule* PropertyModule = FModuleManager::GetModulePtr< FPropertyEditorModule >( TEXT( "PropertyEditor" ) ) )
+		{
+			PropertyModule->UnregisterCustomClassLayout( TEXT( "LevelExporterUSDOptions" ) );
+		}
+	}
 };
 
 IMPLEMENT_MODULE_USD( FUsdExporterModule, USDExporter );
diff --git a/Engine/Plugins/Importers/USDImporter/Source/USDExporter/Public/LevelExporterUSDOptions.h b/Engine/Plugins/Importers/USDImporter/Source/USDExporter/Public/LevelExporterUSDOptions.h
--- a/Engine/Plugins/Importers/USDImporter/Source/USDExporter/Public/LevelExporterUSDOptions.h
+++ b/Engine/Plugins/Importers/USDImporter/Source/USDExporter/Public/LevelExporterUSDOptions.h
@@ -75,9 +75,9 @@ public:
     UPROPERTY( EditAnywhere, config, BlueprintReadWrite, Category = "Sublayers" )
     bool bExportSublayers;
 
-	/** The exporter will only export actors if their level name is within this set. An empty set will allow all actors. */
+	/** Names of levels that should be ignored when collecting actors to export (e.g. "Persistent Level", "Level1", "MySubLevel", etc.) */
     UPROPERTY( EditAnywhere, config, BlueprintReadWrite, Category = "Sublayers" )
-    TSet<FString> LevelFilter;
+    TSet<FString> LevelsToIgnore;
 
 public:
 	// We temporarily stash our export task here as a way of passing our options down to
