USD: Fix exported static mesh displayColor primvar always using constant interpolation;

Fix incorrect color space conversion when importing vertex colors

diff --git a/Engine/Plugins/Importers/USDImporter/Source/USDSchemas/Private/USDGeomMeshTranslator.cpp b/Engine/Plugins/Importers/USDImporter/Source/USDSchemas/Private/USDGeomMeshTranslator.cpp
--- a/Engine/Plugins/Importers/USDImporter/Source/USDSchemas/Private/USDGeomMeshTranslator.cpp
+++ b/Engine/Plugins/Importers/USDImporter/Source/USDSchemas/Private/USDGeomMeshTranslator.cpp
@@ -760,7 +760,8 @@ namespace UsdGeomMeshTranslatorImpl
 						OutMeshData.TangentsZ.Add( Normal );
 						OutMeshData.TangentsX.Add( VertexInstanceTangents[ VertexInstanceID ] );
 
-						OutMeshData.Colors.Add( FLinearColor( VertexInstanceColors[ VertexInstanceID ] ).ToFColor( false ) );
+						const bool bSRGB = true;
+						OutMeshData.Colors.Add( FLinearColor( VertexInstanceColors[ VertexInstanceID ] ).ToFColor( bSRGB ) );
 
 						// Supporting only one UV channel
 						const int32 UVIndex = 0;
diff --git a/Engine/Plugins/Importers/USDImporter/Source/USDUtilities/Private/USDGeomMeshConversion.cpp b/Engine/Plugins/Importers/USDImporter/Source/USDUtilities/Private/USDGeomMeshConversion.cpp
--- a/Engine/Plugins/Importers/USDImporter/Source/USDUtilities/Private/USDGeomMeshConversion.cpp
+++ b/Engine/Plugins/Importers/USDImporter/Source/USDUtilities/Private/USDGeomMeshConversion.cpp
@@ -188,10 +188,10 @@ namespace UE
 					// Vertex colors
 					if ( LODRenderMesh.bHasColorVertexData )
 					{
-						pxr::UsdAttribute DisplayColorAttribute = UsdMesh.CreateDisplayColorAttr();
-						pxr::UsdAttribute DisplayOpacityAttribute = UsdMesh.CreateDisplayOpacityAttr();
+						pxr::UsdGeomPrimvar DisplayColorPrimvar = UsdMesh.CreateDisplayColorPrimvar( pxr::UsdGeomTokens->vertex );
+						pxr::UsdGeomPrimvar DisplayOpacityPrimvar = UsdMesh.CreateDisplayOpacityPrimvar( pxr::UsdGeomTokens->vertex );
 
-						if ( DisplayColorAttribute )
+						if ( DisplayColorPrimvar )
 						{
 							pxr::VtArray< pxr::GfVec3f > DisplayColors;
 							DisplayColors.reserve( VertexCount );
@@ -208,8 +208,8 @@ namespace UE
 								DisplayOpacities.push_back( Color[ 3 ] );
 							}
 
-							DisplayColorAttribute.Set( DisplayColors, TimeCode );
-							DisplayOpacityAttribute.Set( DisplayOpacities, TimeCode );
+							DisplayColorPrimvar.Set( DisplayColors, TimeCode );
+							DisplayOpacityPrimvar.Set( DisplayOpacities, TimeCode );
 						}
 					}
 				}
@@ -805,11 +805,6 @@ bool UsdToUnreal::ConvertGeomMesh( const pxr::UsdTyped& UsdSchema, FMeshDescript
 
 				// Vertex color
 				{
-					auto ConvertToLinear = []( const pxr::GfVec3f& UsdColor ) -> FLinearColor
-					{
-						return FLinearColor( FLinearColor( UsdToUnreal::ConvertColor( UsdColor ) ).ToFColor( false ) );
-					};
-
 					const int32 ValueIndex = UsdGeomMeshImpl::GetPrimValueIndex( ColorInterpolation, ControlPointIndex, CurrentVertexInstanceIndex, PolygonIndex );
 
 					GfVec3f UsdColor( 1.f, 1.f, 1.f );
@@ -819,7 +814,7 @@ bool UsdToUnreal::ConvertGeomMesh( const pxr::UsdTyped& UsdSchema, FMeshDescript
 						UsdColor = UsdColors[ ValueIndex ];
 					}
 
-					MeshDescriptionColors[ AddedVertexInstanceId ] = ConvertToLinear( UsdColor );
+					MeshDescriptionColors[ AddedVertexInstanceId ] = UsdToUnreal::ConvertColor( UsdColor );
 				}
 
 				// Vertex opacity
diff --git a/Engine/Plugins/Importers/USDImporter/Source/USDUtilities/Private/USDSkeletalDataConversion.cpp b/Engine/Plugins/Importers/USDImporter/Source/USDUtilities/Private/USDSkeletalDataConversion.cpp
--- a/Engine/Plugins/Importers/USDImporter/Source/USDUtilities/Private/USDSkeletalDataConversion.cpp
+++ b/Engine/Plugins/Importers/USDImporter/Source/USDUtilities/Private/USDSkeletalDataConversion.cpp
@@ -1241,11 +1241,6 @@ bool UsdToUnreal::ConvertSkinnedMesh(const pxr::UsdSkelSkinningQuery& SkinningQu
 			uint32 NumColors = UsdColors.size();
 			pxr::TfToken USDInterpType = ColorPrimvar.GetInterpolation();
 
-			auto ConvertToColor = []( const pxr::GfVec3f& UsdColor ) -> FColor
-			{
-				return FLinearColor( FLinearColor( UsdToUnreal::ConvertColor( UsdColor ) ).ToFColor( false ) ).ToFColor(true);
-			};
-
 			if ( USDInterpType == pxr::UsdGeomTokens->uniform )
 			{
 				NumExpectedColors = NumFaces;
@@ -1272,7 +1267,8 @@ bool UsdToUnreal::ConvertSkinnedMesh(const pxr::UsdSkelSkinningQuery& SkinningQu
 				Colors.Reserve( NumColors );
 				for ( uint32 Index = 0; Index < NumColors; ++Index )
 				{
-					Colors.Add( ConvertToColor( UsdColors[ Index ] ) );
+					const bool bSRGB = true;
+					Colors.Add( UsdToUnreal::ConvertColor( UsdColors[ Index ] ).ToFColor( bSRGB ) );
 				}
 
 				SkelMeshImportData.bHasVertexColors = true;
