USD: Fix cubemap resetting from skylights in some situations.

Fix PostEditChangeProperty forcing the stage to load even if it was already loaded;
Remove PostEditChangeProperty usage when SetRootLayer is applicable;
Show warning if cubemap is not found on the filesystem when setting it onto a prim.

diff --git a/Engine/Plugins/Importers/USDImporter/Source/USDStage/Private/USDStageActor.cpp b/Engine/Plugins/Importers/USDImporter/Source/USDStage/Private/USDStageActor.cpp
--- a/Engine/Plugins/Importers/USDImporter/Source/USDStage/Private/USDStageActor.cpp
+++ b/Engine/Plugins/Importers/USDImporter/Source/USDStage/Private/USDStageActor.cpp
@@ -1640,6 +1640,20 @@ void AUsdStageActor::ReloadAnimations()
 }
 
 #if WITH_EDITOR
+
+void AUsdStageActor::PostEditChangeProperty( FPropertyChangedEvent& PropertyChangedEvent )
+{
+	// For handling root layer changes via direct changes to properties we want to go through OnObjectPropertyChanged -> HandlePropertyChangedEvent ->
+	// -> SetRootLayer (which checks whether this stage is already opened or not) -> PostRegisterAllComponents.
+	// We need to intercept PostEditChangeProperty too because in the editor any call to PostEditChangeProperty can also *directly* trigger
+	// PostRegister/UnregisterAllComponents which would have sidestepped our checks in SetRootLayer.
+	// Note that any property change event would also end up calling our intended path via OnObjectPropertyChanged, this just prevents us from loading
+	// the same stage again if we don't need to.
+
+	bIsModifyingAProperty = true;
+	Super::PostEditChangeProperty( PropertyChangedEvent );
+}
+
 void AUsdStageActor::PostTransacted(const FTransactionObjectEvent& TransactionEvent)
 {
 	const TArray<FName>& ChangedProperties = TransactionEvent.GetChangedProperties();
@@ -1720,8 +1734,8 @@ void AUsdStageActor::PostTransacted(const FTransactionObjectEvent& TransactionEv
 void AUsdStageActor::PreEditChange( FProperty* PropertyThatWillChange )
 {
 	// If we're just editing some other actor property like Time or anything else, we will get
-	// PostRegister/Unregister calls in the editor due to AActor::PostEditChangeProperty. Here we
-	// determine in which cases we should ignore those PostRegister/Unregister calls by using the
+	// PostRegister/Unregister calls in the editor due to AActor::PostEditChangeProperty *and* AActor::PreEditChange.
+	// Here we determine in which cases we should ignore those PostRegister/Unregister calls by using the
 	// bIsModifyingAProperty flag
 	if ( !IsActorBeingDestroyed() )
 	{
@@ -1729,7 +1743,7 @@ void AUsdStageActor::PreEditChange( FProperty* PropertyThatWillChange )
 		{
 			// PreEditChange gets called for actor lifecycle functions too (like if the actor transacts on undo/redo).
 			// In those cases we will have nullptr PropertyThatWillChange, and we don't want to block our PostRegister/Unregister
-			// functions. We only care about blocking the calls triggered by AActor::PostEditChangeProperty
+			// functions. We only care about blocking the calls triggered by AActor::PostEditChangeProperty and AActor::PreEditChange
 			if ( PropertyThatWillChange )
 			{
 				bIsModifyingAProperty = true;
diff --git a/Engine/Plugins/Importers/USDImporter/Source/USDStage/Public/USDStageActor.h b/Engine/Plugins/Importers/USDImporter/Source/USDStage/Public/USDStageActor.h
--- a/Engine/Plugins/Importers/USDImporter/Source/USDStage/Public/USDStageActor.h
+++ b/Engine/Plugins/Importers/USDImporter/Source/USDStage/Public/USDStageActor.h
@@ -142,6 +142,7 @@ public:
 
 public:
 #if WITH_EDITOR
+	virtual void PostEditChangeProperty( FPropertyChangedEvent& PropertyChangedEvent ) override;
 	virtual void PostTransacted(const FTransactionObjectEvent& TransactionEvent) override;
 	virtual void PreEditChange( FProperty* PropertyThatWillChange ) override;
 	virtual void PreEditUndo() override;
diff --git a/Engine/Plugins/Importers/USDImporter/Source/USDStageEditorViewModels/Private/USDStageViewModel.cpp b/Engine/Plugins/Importers/USDImporter/Source/USDStageEditorViewModels/Private/USDStageViewModel.cpp
--- a/Engine/Plugins/Importers/USDImporter/Source/USDStageEditorViewModels/Private/USDStageViewModel.cpp
+++ b/Engine/Plugins/Importers/USDImporter/Source/USDStageEditorViewModels/Private/USDStageViewModel.cpp
@@ -111,10 +111,7 @@ void FUsdStageViewModel::OpenStage( const TCHAR* FilePath )
 	if ( AUsdStageActor* StageActor = UsdStageActor.Get() )
 	{
 		StageActor->Modify();
-
-		StageActor->RootLayer.FilePath = FilePath;
-		FPropertyChangedEvent RootLayerPropertyChangedEvent( FindFieldChecked< FProperty >( StageActor->GetClass(), FName("RootLayer") ) );
-		StageActor->PostEditChangeProperty( RootLayerPropertyChangedEvent );
+		StageActor->SetRootLayer( FilePath );
 	}
 	else
 	{
diff --git a/Engine/Plugins/Importers/USDImporter/Source/USDUtilities/Private/USDLightConversion.cpp b/Engine/Plugins/Importers/USDImporter/Source/USDUtilities/Private/USDLightConversion.cpp
--- a/Engine/Plugins/Importers/USDImporter/Source/USDUtilities/Private/USDLightConversion.cpp
+++ b/Engine/Plugins/Importers/USDImporter/Source/USDUtilities/Private/USDLightConversion.cpp
@@ -468,8 +468,16 @@ bool UnrealToUsd::ConvertSkyLightComponent( const USkyLightComponent& LightCompo
 			if ( UAssetImportData* AssetImportData = TextureCube->AssetImportData )
 			{
 				FString FilePath = AssetImportData->GetFirstFilename();
-				UsdUtils::MakePathRelativeToLayer( UE::FSdfLayer( Prim.GetStage()->GetEditTarget().GetLayer() ), FilePath );
+				if ( !FPaths::FileExists( FilePath ) )
+				{
+					UE_LOG(LogTemp, Warning, TEXT("Used '%s' as cubemap when converting SkyLightComponent '%s' onto prim '%s', but the cubemap does not exist on the filesystem!"),
+						*FilePath,
+						*LightComponent.GetPathName(),
+						*UsdToUnreal::ConvertPath(Prim.GetPath())
+					);
+				}
 
+				UsdUtils::MakePathRelativeToLayer( UE::FSdfLayer( Prim.GetStage()->GetEditTarget().GetLayer() ), FilePath );
 				Attr.Set<pxr::SdfAssetPath>( pxr::SdfAssetPath{ UnrealToUsd::ConvertString( *FilePath ).Get() }, TimeCode );
 			}
 		}
