USD: Fix unnecessary stage closing/opening when we're in a persistent level and it is hidden/shown.

Fix leaked references to stage when closing the stage without unloading it;
Fix double UI refresh when unhiding/hiding sublevel or deleting/spawning actors;
Fix leaked reference to stage when hiding in a sublevel with SUSDStageEditor not visible;
Tweak test_export_level so it doesn't freeze for 20 seconds;
Add two new tests for hiding/showing levels.

diff --git a/Engine/Plugins/Importers/USDImporter/Source/USDStage/Private/USDStageActor.cpp b/Engine/Plugins/Importers/USDImporter/Source/USDStage/Private/USDStageActor.cpp
--- a/Engine/Plugins/Importers/USDImporter/Source/USDStage/Private/USDStageActor.cpp
+++ b/Engine/Plugins/Importers/USDImporter/Source/USDStage/Private/USDStageActor.cpp
@@ -1324,6 +1324,8 @@ void AUsdStageActor::CloseUsdStage()
 {
 	FUsdStageActorImpl::DiscardStage( UsdStage, this );
 	UsdStage = UE::FUsdStage();
+	LevelSequenceHelper.Init( UE::FUsdStage() ); // Drop the helper's reference to the stage
+	OnStageChanged.Broadcast();
 }
 
 #if WITH_EDITOR
@@ -1559,7 +1561,6 @@ void AUsdStageActor::UnloadUsdStage()
 		LevelSequence = nullptr;
 	}
 	LevelSequenceHelper.Clear();
-	LevelSequenceHelper.Init( UE::FUsdStage() ); // Drop the helper's reference to the stage
 
 	if ( RootUsdTwin )
 	{
@@ -1670,26 +1671,28 @@ void AUsdStageActor::PostTransacted(const FTransactionObjectEvent& TransactionEv
 		{
 			OpenUsdStage();
 		}
-
-		OnStageChanged.Broadcast();
 	}
 
-	// If we're in a sublevel that is hidden, we'll respond to the generated PostUnregisterAllComponent call
-	// and unload our spawned actors/assets, so let's close/open the stage too
-	if ( ChangedProperties.Contains( GET_MEMBER_NAME_CHECKED( AActor, bHiddenEdLevel ) ) ||
-		 ChangedProperties.Contains( GET_MEMBER_NAME_CHECKED( AActor, bHiddenEdLayer ) ) ||
-		 ChangedProperties.Contains( GET_MEMBER_NAME_CHECKED( AActor, bHiddenEd ) ) )
+	// If we're in the persistent level don't do anything, because hiding/showing the persistent level doesn't
+	// cause actors to load/unload like it does if they're in sublevels
+	ULevel* CurrentLevel = GetLevel();
+	if ( CurrentLevel && !CurrentLevel->IsPersistentLevel() )
 	{
-		if ( IsHiddenEd() )
-		{
-			CloseUsdStage();
-		}
-		else
+		// If we're in a sublevel that is hidden, we'll respond to the generated PostUnregisterAllComponent call
+		// and unload our spawned actors/assets, so let's close/open the stage too
+		if ( ChangedProperties.Contains( GET_MEMBER_NAME_CHECKED( AActor, bHiddenEdLevel ) ) ||
+			 ChangedProperties.Contains( GET_MEMBER_NAME_CHECKED( AActor, bHiddenEdLayer ) ) ||
+			 ChangedProperties.Contains( GET_MEMBER_NAME_CHECKED( AActor, bHiddenEd ) ) )
 		{
-			OpenUsdStage();
+			if ( IsHiddenEd() )
+			{
+				CloseUsdStage();
+			}
+			else
+			{
+				OpenUsdStage();
+			}
 		}
-
-		OnStageChanged.Broadcast();
 	}
 
 	if (TransactionEvent.GetEventType() == ETransactionObjectEventType::UndoRedo)
@@ -1707,8 +1710,6 @@ void AUsdStageActor::PostTransacted(const FTransactionObjectEvent& TransactionEv
 			CloseUsdStage();
 			OpenUsdStage();
 			ReloadAnimations();
-
-			OnStageChanged.Broadcast();
 		}
 		else if (ChangedProperties.Contains(GET_MEMBER_NAME_CHECKED(AUsdStageActor, Time)))
 		{
@@ -1869,8 +1870,10 @@ void AUsdStageActor::PostRegisterAllComponents()
 	const bool bIsLevelHidden = !Level || ( !Level->bIsVisible && !Level->bIsAssociatingLevel );
 
 	// We get an inactive world when dragging a ULevel asset
+	// Prevent loading on bHiddenEdLevel because PostRegisterAllComponents gets called in the process of hiding our level, if we're in the persistent level.
+	// This is just hiding though, so we shouldn't actively load/unload anything
 	UWorld* World = GetWorld();
-	if ( IsTemplate() || bIsTransitioningIntoPIE || !World || World->WorldType == EWorldType::Inactive || bIsLevelHidden || bIsModifyingAProperty || bIsUndoRedoing )
+	if ( IsTemplate() || bIsTransitioningIntoPIE || !World || World->WorldType == EWorldType::Inactive || bIsLevelHidden || bIsModifyingAProperty || bIsUndoRedoing || bHiddenEdLevel )
 	{
 		return;
 	}
diff --git a/Engine/Plugins/Importers/USDImporter/Source/USDStageEditor/Private/SUSDStageTreeView.cpp b/Engine/Plugins/Importers/USDImporter/Source/USDStageEditor/Private/SUSDStageTreeView.cpp
--- a/Engine/Plugins/Importers/USDImporter/Source/USDStageEditor/Private/SUSDStageTreeView.cpp
+++ b/Engine/Plugins/Importers/USDImporter/Source/USDStageEditor/Private/SUSDStageTreeView.cpp
@@ -358,29 +358,34 @@ void SUsdStageTreeView::OnGetChildren( FUsdPrimViewModelRef InParent, TArray< FU
 
 void SUsdStageTreeView::Refresh( AUsdStageActor* InUsdStageActor )
 {
-	RootItems.Empty();
+	UE::FUsdStage OldStage = RootItems.Num() > 0 ? RootItems[0]->UsdStage : UE::FUsdStage();
+	UE::FUsdStage NewStage = InUsdStageActor ? static_cast< const AUsdStageActor* >( InUsdStageActor )->GetUsdStage() : UE::FUsdStage();
 
-	if (UsdStageActor.Get() != InUsdStageActor)
+	RootItems.Empty();
+	if ( UsdStageActor.Get() != InUsdStageActor || NewStage != OldStage )
 	{
+		// This is very important: Internally the tree will store FUsdPrimViewModelRef in its SparseItemInfos member if we have
+		// any member manually expanded/collapsed. These can prevent the FUsdPrimViewModels from being collected, and prevent the
+		// stage from being fully closed, so we must do this whenever the stage changes
+		ClearExpandedItems();
 		TreeItemExpansionStates.Reset();
 	}
 
 	UsdStageActor = InUsdStageActor;
-
 	if ( !UsdStageActor.IsValid() )
 	{
 		return;
 	}
 
-	if ( UE::FUsdStage UsdStage = const_cast< const AUsdStageActor* >( UsdStageActor.Get() )->GetUsdStage() )
+	if ( NewStage )
 	{
-		if ( UE::FUsdPrim RootPrim = UsdStage.GetPseudoRoot() )
+		if ( UE::FUsdPrim RootPrim = NewStage.GetPseudoRoot() )
 		{
-			RootItems.Add( MakeShared< FUsdPrimViewModel >( nullptr, UsdStage, RootPrim ) );
+			RootItems.Add( MakeShared< FUsdPrimViewModel >( nullptr, NewStage, RootPrim ) );
 		}
-	}
 
-	RestoreExpansionStates();
+		RestoreExpansionStates();
+	}
 }
 
 void SUsdStageTreeView::RefreshPrim( const FString& PrimPath, bool bResync )
diff --git a/Engine/Plugins/Importers/USDImporter/Source/USDStageEditorViewModels/Public/USDPrimAttributesViewModel.h b/Engine/Plugins/Importers/USDImporter/Source/USDStageEditorViewModels/Public/USDPrimAttributesViewModel.h
--- a/Engine/Plugins/Importers/USDImporter/Source/USDStageEditorViewModels/Public/USDPrimAttributesViewModel.h
+++ b/Engine/Plugins/Importers/USDImporter/Source/USDStageEditorViewModels/Public/USDPrimAttributesViewModel.h
@@ -21,8 +21,6 @@ public:
 	void SetAttributeValue( const UsdUtils::FConvertedVtValue& InValue );
 
 public:
-	UE::FUsdStage UsdStage;
-
 	FString Label;
 	UsdUtils::FConvertedVtValue Value;
 	FString ValueRole;
diff --git a/Engine/Plugins/Importers/USDImporter/Source/USDUtilities/Private/USDConversionUtils.cpp b/Engine/Plugins/Importers/USDImporter/Source/USDUtilities/Private/USDConversionUtils.cpp
--- a/Engine/Plugins/Importers/USDImporter/Source/USDUtilities/Private/USDConversionUtils.cpp
+++ b/Engine/Plugins/Importers/USDImporter/Source/USDUtilities/Private/USDConversionUtils.cpp
@@ -775,6 +775,11 @@ UUsdAssetImportData* UsdUtils::GetAssetImportData( UObject* Asset )
 void UsdUtils::AddReference( UE::FUsdPrim& Prim, const TCHAR* AbsoluteFilePath )
 {
 #if USE_USD_SDK
+	if ( !Prim || !AbsoluteFilePath )
+	{
+		return;
+	}
+
 	FScopedUsdAllocs UsdAllocs;
 
 	pxr::UsdPrim UsdPrim( Prim );
