USD: Fix build breakage for non-editor targets;

diff --git a/Engine/Plugins/Importers/USDImporter/Source/USDStage/Private/USDStageActor.cpp b/Engine/Plugins/Importers/USDImporter/Source/USDStage/Private/USDStageActor.cpp
--- a/Engine/Plugins/Importers/USDImporter/Source/USDStage/Private/USDStageActor.cpp
+++ b/Engine/Plugins/Importers/USDImporter/Source/USDStage/Private/USDStageActor.cpp
@@ -1857,7 +1857,8 @@ void AUsdStageActor::PostRegisterAllComponents()
 	Super::PostRegisterAllComponents();
 
 #if WITH_EDITOR
-	if ( bIsEditorPreviewActor )
+	// Prevent loading on bHiddenEdLevel because PostRegisterAllComponents gets called in the process of hiding our level, if we're in the persistent level.
+	if ( bIsEditorPreviewActor || bHiddenEdLevel )
 	{
 		return;
 	}
@@ -1870,10 +1871,9 @@ void AUsdStageActor::PostRegisterAllComponents()
 	const bool bIsLevelHidden = !Level || ( !Level->bIsVisible && !Level->bIsAssociatingLevel );
 
 	// We get an inactive world when dragging a ULevel asset
-	// Prevent loading on bHiddenEdLevel because PostRegisterAllComponents gets called in the process of hiding our level, if we're in the persistent level.
 	// This is just hiding though, so we shouldn't actively load/unload anything
 	UWorld* World = GetWorld();
-	if ( IsTemplate() || bIsTransitioningIntoPIE || !World || World->WorldType == EWorldType::Inactive || bIsLevelHidden || bIsModifyingAProperty || bIsUndoRedoing || bHiddenEdLevel )
+	if ( IsTemplate() || bIsTransitioningIntoPIE || !World || World->WorldType == EWorldType::Inactive || bIsLevelHidden || bIsModifyingAProperty || bIsUndoRedoing )
 	{
 		return;
 	}
