USD: Fix skeletal mesh being visible on viewport after exporting UAnimSequence asset to USD.

diff --git a/Engine/Plugins/Importers/USDImporter/Source/USDUtilities/Private/USDSkeletalDataConversion.cpp b/Engine/Plugins/Importers/USDImporter/Source/USDUtilities/Private/USDSkeletalDataConversion.cpp
--- a/Engine/Plugins/Importers/USDImporter/Source/USDUtilities/Private/USDSkeletalDataConversion.cpp
+++ b/Engine/Plugins/Importers/USDImporter/Source/USDUtilities/Private/USDSkeletalDataConversion.cpp
@@ -2694,6 +2694,9 @@ bool UnrealToUsd::ConvertAnimSequence( UAnimSequence* AnimSequence, pxr::UsdPrim
 			RotationsAttr.Set( Rotations, pxr::UsdTimeCode( TimeCode ) );
 			ScalesAttr.Set( Scales, pxr::UsdTimeCode( TimeCode ) );
 		}
+
+		// Actively delete it or else it will remain visible on the viewport
+		DebugSkelMeshComponent->DestroyComponent();
 	}
 
 	const int32 StageEndTimeCode = SkelAnimPrim.GetStage()->GetEndTimeCode();
