USD Stage Editor: Switched the rendercontext menu entries to radio buttons since only one can be selected at a time.

diff --git a/Engine/Plugins/Importers/USDImporter/Source/USDStageEditor/Private/SUSDStage.cpp b/Engine/Plugins/Importers/USDImporter/Source/USDStageEditor/Private/SUSDStage.cpp
--- a/Engine/Plugins/Importers/USDImporter/Source/USDStageEditor/Private/SUSDStage.cpp
+++ b/Engine/Plugins/Importers/USDImporter/Source/USDStageEditor/Private/SUSDStage.cpp
@@ -596,7 +596,7 @@ void SUsdStage::FillRenderContextSubMenu( FMenuBuilder& MenuBuilder )
 				})
 			),
 			NAME_None,
-			EUserInterfaceActionType::Check
+			EUserInterfaceActionType::RadioButton
 		);
 	};
 
