USD: Better handling of picked USD format.

Add a combobox for picking the payload USD format;
Properly use the main output file extension for other emitted files;
Fix inconsistent name sanitizing when exporting meshes to filepaths with parentheses;
Add some more defensive sanitizing for sublayer names and landscape actors names.

diff --git a/Engine/Plugins/Importers/USDImporter/Source/USDExporter/Private/AnimSequenceExporterUSD.cpp b/Engine/Plugins/Importers/USDImporter/Source/USDExporter/Private/AnimSequenceExporterUSD.cpp
--- a/Engine/Plugins/Importers/USDImporter/Source/USDExporter/Private/AnimSequenceExporterUSD.cpp
+++ b/Engine/Plugins/Importers/USDImporter/Source/USDExporter/Private/AnimSequenceExporterUSD.cpp
@@ -78,6 +78,8 @@ bool UAnimSequenceExporterUSD::ExportBinary( UObject* Object, const TCHAR* Type,
 		FString ExtensionPart;
 		FPaths::Split( PayloadFilename, PathPart, FilenamePart, ExtensionPart );
 
+		ExtensionPart = Options->PreviewMeshOptions.PayloadFormat;
+
 		PayloadFilename = FPaths::Combine( PathPart, FilenamePart + TEXT( "_payload." ) + ExtensionPart );
 	}
 
diff --git a/Engine/Plugins/Importers/USDImporter/Source/USDExporter/Private/SkeletalMeshExporterUSD.cpp b/Engine/Plugins/Importers/USDImporter/Source/USDExporter/Private/SkeletalMeshExporterUSD.cpp
--- a/Engine/Plugins/Importers/USDImporter/Source/USDExporter/Private/SkeletalMeshExporterUSD.cpp
+++ b/Engine/Plugins/Importers/USDImporter/Source/USDExporter/Private/SkeletalMeshExporterUSD.cpp
@@ -75,6 +75,8 @@ bool USkeletalMeshExporterUsd::ExportBinary( UObject* Object, const TCHAR* Type,
 		FString ExtensionPart;
 		FPaths::Split( PayloadFilename, PathPart, FilenamePart, ExtensionPart );
 
+		ExtensionPart = Options->Inner.PayloadFormat;
+
 		PayloadFilename = FPaths::Combine( PathPart, FilenamePart + TEXT( "_payload." ) + ExtensionPart );
 	}
 
diff --git a/Engine/Plugins/Importers/USDImporter/Source/USDExporter/Private/StaticMeshExporterUSD.cpp b/Engine/Plugins/Importers/USDImporter/Source/USDExporter/Private/StaticMeshExporterUSD.cpp
--- a/Engine/Plugins/Importers/USDImporter/Source/USDExporter/Private/StaticMeshExporterUSD.cpp
+++ b/Engine/Plugins/Importers/USDImporter/Source/USDExporter/Private/StaticMeshExporterUSD.cpp
@@ -85,6 +85,8 @@ bool UStaticMeshExporterUsd::ExportBinary( UObject* Object, const TCHAR* Type, F
 		FString ExtensionPart;
 		FPaths::Split( PayloadFilename, PathPart, FilenamePart, ExtensionPart );
 
+		ExtensionPart = Options->PayloadFormat;
+
 		PayloadFilename = FPaths::Combine( PathPart, FilenamePart + TEXT( "_payload." ) + ExtensionPart );
 	}
 
diff --git a/Engine/Plugins/Importers/USDImporter/Source/USDExporter/Public/LevelExporterUSDOptions.h b/Engine/Plugins/Importers/USDImporter/Source/USDExporter/Public/LevelExporterUSDOptions.h
--- a/Engine/Plugins/Importers/USDImporter/Source/USDExporter/Public/LevelExporterUSDOptions.h
+++ b/Engine/Plugins/Importers/USDImporter/Source/USDExporter/Public/LevelExporterUSDOptions.h
@@ -52,6 +52,10 @@ public:
 	UPROPERTY( EditAnywhere, config, BlueprintReadWrite, Category = "Export settings" )
 	bool bUsePayload;
 
+	/** USD format to use for exported payload files */
+	UPROPERTY( EditAnywhere, config, BlueprintReadWrite, Category = "Export settings", meta = ( EditCondition = "bUsePayload", GetOptions = GetUsdExtensions ) )
+	FString PayloadFormat;
+
 	/** Whether to use UE actor folders as empty prims */
     UPROPERTY( EditAnywhere, config, BlueprintReadWrite, Category = "Export settings" )
     bool bExportActorFolders;
@@ -84,4 +88,13 @@ public:
 	// the Python exporter, that does the actual level exporting
 	UPROPERTY( VisibleAnywhere, BlueprintReadOnly, Category = Hidden )
 	UAssetExportTask* CurrentTask;
+
+private:
+	UFUNCTION()
+	static TArray<FString> GetUsdExtensions()
+	{
+		TArray<FString> Extensions = UnrealUSDWrapper::GetAllSupportedFileFormats();
+		Extensions.Remove( TEXT( "usdz" ) );
+		return Extensions;
+	}
 };
diff --git a/Engine/Plugins/Importers/USDImporter/Source/USDExporter/Public/SkeletalMeshExporterUSDOptions.h b/Engine/Plugins/Importers/USDImporter/Source/USDExporter/Public/SkeletalMeshExporterUSDOptions.h
--- a/Engine/Plugins/Importers/USDImporter/Source/USDExporter/Public/SkeletalMeshExporterUSDOptions.h
+++ b/Engine/Plugins/Importers/USDImporter/Source/USDExporter/Public/SkeletalMeshExporterUSDOptions.h
@@ -2,10 +2,11 @@
 
 #pragma once
 
+#include "LevelExporterUSDOptions.h"
 #include "UnrealUSDWrapper.h"
+#include "USDStageOptions.h"
 
 #include "Engine/EngineTypes.h"
-#include "USDStageOptions.h"
 
 #include "SkeletalMeshExporterUSDOptions.generated.h"
 
@@ -24,6 +25,10 @@ struct USDEXPORTER_API FSkeletalMeshExporterUSDInnerOptions
 	/** If true, the mesh data is exported to yet another "payload" file, and referenced via a payload composition arc */
 	UPROPERTY( EditAnywhere, config, BlueprintReadWrite, Category = USDSettings )
 	bool bUsePayload;
+
+	/** USD format to use for exported payload files */
+	UPROPERTY( EditAnywhere, config, BlueprintReadWrite, Category = USDSettings, meta = ( EditCondition = "bUsePayload", GetOptions = "USDExporter.LevelExporterUSDOptions.GetUsdExtensions" ) )
+	FString PayloadFormat;
 };
 
 /**
diff --git a/Engine/Plugins/Importers/USDImporter/Source/USDExporter/Public/StaticMeshExporterUSDOptions.h b/Engine/Plugins/Importers/USDImporter/Source/USDExporter/Public/StaticMeshExporterUSDOptions.h
--- a/Engine/Plugins/Importers/USDImporter/Source/USDExporter/Public/StaticMeshExporterUSDOptions.h
+++ b/Engine/Plugins/Importers/USDImporter/Source/USDExporter/Public/StaticMeshExporterUSDOptions.h
@@ -2,6 +2,7 @@
 
 #pragma once
 
+#include "LevelExporterUSDOptions.h"
 #include "UnrealUSDWrapper.h"
 #include "USDStageOptions.h"
 
@@ -24,4 +25,8 @@ public:
 	/** If true, the mesh data is exported to yet another "payload" file, and referenced via a payload composition arc */
 	UPROPERTY( EditAnywhere, config, BlueprintReadWrite, Category = USDSettings )
 	bool bUsePayload;
+
+	/** USD format to use for exported payload files */
+	UPROPERTY( EditAnywhere, config, BlueprintReadWrite, Category = USDSettings, meta = ( EditCondition = "bUsePayload", GetOptions = "USDExporter.LevelExporterUSDOptions.GetUsdExtensions" ) )
+	FString PayloadFormat;
 };
