USD: Use new functionality of the material baking module to bake linear channels as linear

diff --git a/Engine/Plugins/Importers/USDImporter/Source/USDUtilities/Private/USDShadeConversion.cpp b/Engine/Plugins/Importers/USDImporter/Source/USDUtilities/Private/USDShadeConversion.cpp
--- a/Engine/Plugins/Importers/USDImporter/Source/USDUtilities/Private/USDShadeConversion.cpp
+++ b/Engine/Plugins/Importers/USDImporter/Source/USDUtilities/Private/USDShadeConversion.cpp
@@ -1104,8 +1104,14 @@ namespace UE
 
 				TArray<FBakeOutput> BakeOutputs;
 				IMaterialBakingModule& Module = FModuleManager::Get().LoadModuleChecked<IMaterialBakingModule>( "MaterialBaking" );
+				bool bLinearBake = true;
+				Module.SetLinearBake( bLinearBake );
 				Module.BakeMaterials( { &MatSet }, { &MeshSettings }, BakeOutputs );
 
+				// It's recommended to set this back to false as it's a global option
+				bLinearBake = false;
+				Module.SetLinearBake( bLinearBake );
+
 				if ( BakeOutputs.Num() < 1 )
 				{
 					return false;
