USD: add basic support for import of materials with diffuse driven by displayColor

While not nearly a complete solution for importing shading networks in USD that make
use of primvar reader shaders, this covers the common case of a primvar reader that
reads displayColor and uses that to drive the diffuseColor input of a UsdPreviewSurface
shader. We import displayColor in USD as vertex colors on UE meshes, so they become
readily available using the "Vertex Color" node in the UE material.

diff --git a/Engine/Plugins/Importers/USDImporter/Source/USDUtilities/Private/USDShadeConversion.cpp b/Engine/Plugins/Importers/USDImporter/Source/USDUtilities/Private/USDShadeConversion.cpp
--- a/Engine/Plugins/Importers/USDImporter/Source/USDUtilities/Private/USDShadeConversion.cpp
+++ b/Engine/Plugins/Importers/USDImporter/Source/USDUtilities/Private/USDShadeConversion.cpp
@@ -25,6 +25,7 @@
 #include "Materials/MaterialExpressionConstant.h"
 #include "Materials/MaterialExpressionConstant4Vector.h"
 #include "Materials/MaterialExpressionTextureCoordinate.h"
+#include "Materials/MaterialExpressionVertexColor.h"
 #include "Materials/MaterialInstanceConstant.h"
 #include "Materials/MaterialInstanceDynamic.h"
 #include "Misc/FileHelper.h"
@@ -451,21 +452,18 @@ namespace UE
 				}
 			}
 
-			enum class EParameterValueType
-			{
-				None = 0,
-				Float,
-				Vector,
-				Texture,
-			};
-
 			struct FTextureParameterValue
 			{
 				UTexture* Texture;
 				int32 UVIndex;
 				int32 OutputIndex = 0;
 			};
-			using FParameterValue = TVariant<float, FVector, FTextureParameterValue, bool>;
+			struct FPrimvarReaderParameterValue
+			{
+				std::string PrimvarName;
+				FVector FallbackValue;
+			};
+			using FParameterValue = TVariant<float, FVector, FTextureParameterValue, FPrimvarReaderParameterValue, bool>;
 
 			bool GetTextureParameterValue( pxr::UsdShadeInput& ShadeInput, TextureGroup LODGroup, FParameterValue& OutValue,
 				UMaterialInterface* Material = nullptr, UUsdAssetCache* TexturesCache = nullptr, TMap<FString, int32>* PrimvarToUVIndex = nullptr, bool bForceVirtualTextures = false )
@@ -484,19 +482,19 @@ namespace UE
 				{
 					pxr::UsdShadeInput FileInput;
 
-				    // UsdUVTexture: Get its file input
-				    if ( AttributeType == pxr::UsdShadeAttributeType::Output )
-				    {
-					    FileInput = Source.GetInput( UnrealIdentifiers::File );
-				    }
-				    // Check if we are being directly passed an asset
-				    else
-				    {
-					    FileInput = Source.GetInput( SourceName );
-				    }
-
-				    if ( FileInput && FileInput.GetTypeName() == pxr::SdfValueTypeNames->Asset ) // Check that FileInput is of type Asset
-				    {
+					// UsdUVTexture: Get its file input
+					if ( AttributeType == pxr::UsdShadeAttributeType::Output )
+					{
+						FileInput = Source.GetInput( UnrealIdentifiers::File );
+					}
+					// Check if we are being directly passed an asset
+					else
+					{
+						FileInput = Source.GetInput( SourceName );
+					}
+
+					if ( FileInput && FileInput.GetTypeName() == pxr::SdfValueTypeNames->Asset ) // Check that FileInput is of type Asset
+					{
 						const FString TexturePath = UsdUtils::GetResolvedTexturePath( FileInput.GetAttr() );
 						// We will assume the texture is valid, and show a warning if we fail to parse this later.
 						// Note that we don't even check that the file exists: If we have a texture bound to opacity then we
@@ -674,6 +672,78 @@ namespace UE
 				return true;
 			}
 
+			bool GetPrimvarReaderParameterValue( const pxr::UsdShadeInput& Input, const pxr::TfToken& PrimvarReaderShaderId, const FLinearColor& DefaultValue, FParameterValue& OutValue )
+			{
+				FScopedUsdAllocs Allocs;
+
+				if ( !Input )
+				{
+					return false;
+				}
+
+				const bool bShaderOutputsOnly = true;
+				const pxr::UsdShadeAttributeVector ValProdAttrs = Input.GetValueProducingAttributes( bShaderOutputsOnly );
+				for ( const pxr::UsdAttribute& ValProdAttr : ValProdAttrs )
+				{
+					const pxr::UsdShadeShader ValProdShader( ValProdAttr.GetPrim() );
+					if ( !ValProdShader )
+					{
+						continue;
+					}
+
+					pxr::TfToken ShaderId;
+					if ( !ValProdShader.GetShaderId( &ShaderId ) || ShaderId != PrimvarReaderShaderId )
+					{
+						continue;
+					}
+
+					const pxr::UsdShadeInput VarnameInput = ValProdShader.GetInput( UnrealIdentifiers::Varname );
+					if ( !VarnameInput )
+					{
+						continue;
+					}
+
+					// The schema for UsdPrimvarReader specifies that the "varname" input should be
+					// string-typed, but some assets might be set up token-typed, so we'll consider
+					// either type.
+					std::string PrimvarName;
+					if ( VarnameInput.GetTypeName() == pxr::SdfValueTypeNames->String )
+					{
+						if ( !VarnameInput.Get( &PrimvarName ) )
+						{
+							continue;
+						}
+					}
+					else if ( VarnameInput.GetTypeName() == pxr::SdfValueTypeNames->Token )
+					{
+						pxr::TfToken PrimvarNameToken;
+						if ( !VarnameInput.Get( &PrimvarNameToken ) )
+						{
+							continue;
+						}
+						PrimvarName = PrimvarNameToken.GetString();
+					}
+
+					if ( PrimvarName.empty() )
+					{
+						continue;
+					}
+
+					FLinearColor FallbackColor = DefaultValue;
+					const pxr::UsdShadeInput FallbackInput = ValProdShader.GetInput( UnrealIdentifiers::Fallback );
+					pxr::GfVec3f UsdFallbackColor;
+					if ( FallbackInput && FallbackInput.Get<pxr::GfVec3f>( &UsdFallbackColor ) )
+					{
+						FallbackColor = UsdToUnreal::ConvertColor( UsdFallbackColor );
+					}
+
+					OutValue.Set<FPrimvarReaderParameterValue>( FPrimvarReaderParameterValue{ PrimvarName, FVector(FallbackColor) } );
+					return true;
+				}
+
+				return false;
+			}
+
 			bool GetVec3ParameterValue( pxr::UsdShadeConnectableAPI& Connectable, const pxr::TfToken& InputName, const FLinearColor& DefaultValue, FParameterValue& OutValue,
 				bool bIsNormalMap = false, UMaterialInterface* Material = nullptr, UUsdAssetCache* TexturesCache = nullptr, TMap<FString, int32>* PrimvarToUVIndex = nullptr, bool bForceVirtualTextures = false )
 			{
@@ -693,7 +763,12 @@ namespace UE
 				{
 					if ( !GetTextureParameterValue( Input, bIsNormalMap ? TEXTUREGROUP_WorldNormalMap : TEXTUREGROUP_World, OutValue, Material, TexturesCache, PrimvarToUVIndex, bForceVirtualTextures ) )
 					{
-						return GetVec3ParameterValue( Source, SourceName, DefaultValue, OutValue, bIsNormalMap, Material, TexturesCache, PrimvarToUVIndex, bForceVirtualTextures );
+						// Check whether this input receives its value through a connection to a
+						// primvar reader shader.
+						if ( !GetPrimvarReaderParameterValue( Input, UnrealIdentifiers::UsdPrimvarReader_float3, DefaultValue, OutValue ) )
+						{
+							return GetVec3ParameterValue( Source, SourceName, DefaultValue, OutValue, bIsNormalMap, Material, TexturesCache, PrimvarToUVIndex, bForceVirtualTextures );
+						}
 					}
 				}
 				// No other node connected, so we must have some value
@@ -713,36 +788,36 @@ namespace UE
 			}
 
 			bool GetBoolParameterValue( pxr::UsdShadeConnectableAPI& Connectable, const pxr::TfToken& InputName, bool DefaultValue, FParameterValue& OutValue, UMaterialInterface* Material = nullptr, UUsdAssetCache* TexturesCache = nullptr, TMap<FString, int32>* PrimvarToUVIndex = nullptr)
-	        {
-		        FScopedUsdAllocs Allocs;
-
-		        pxr::UsdShadeInput Input = Connectable.GetInput( InputName );
-		        if ( !Input )
-		        {
-			        return false;
-		        }
-
-		        // If we have another shader/node connected
-		        pxr::UsdShadeConnectableAPI Source;
-		        pxr::TfToken SourceName;
-		        pxr::UsdShadeAttributeType AttributeType;
-		        if ( Input.GetConnectedSource( &Source, &SourceName, &AttributeType ) )
-		        {
-			        // Recurse because the attribute may just be pointing at some other attribute that has the data
-			        // (e.g. when shader input is just "hoisted" and connected to the parent material input)
-			        return GetBoolParameterValue( Source, SourceName, DefaultValue, OutValue, Material, TexturesCache, PrimvarToUVIndex );
-		        }
-		        // No other node connected, so we must have some value
-		        else
-		        {
-			        bool InputValue = DefaultValue;
-			        Input.Get< bool >( &InputValue );
-
-			        OutValue.Set< bool >( InputValue );
-		        }
-
-		        return true;
-	        }
+			{
+				FScopedUsdAllocs Allocs;
+
+				pxr::UsdShadeInput Input = Connectable.GetInput( InputName );
+				if ( !Input )
+				{
+					return false;
+				}
+
+				// If we have another shader/node connected
+				pxr::UsdShadeConnectableAPI Source;
+				pxr::TfToken SourceName;
+				pxr::UsdShadeAttributeType AttributeType;
+				if ( Input.GetConnectedSource( &Source, &SourceName, &AttributeType ) )
+				{
+					// Recurse because the attribute may just be pointing at some other attribute that has the data
+					// (e.g. when shader input is just "hoisted" and connected to the parent material input)
+					return GetBoolParameterValue( Source, SourceName, DefaultValue, OutValue, Material, TexturesCache, PrimvarToUVIndex );
+				}
+				// No other node connected, so we must have some value
+				else
+				{
+					bool InputValue = DefaultValue;
+					Input.Get< bool >( &InputValue );
+
+					OutValue.Set< bool >( InputValue );
+				}
+
+				return true;
+			}
 
 #if WITH_EDITOR
 			/**
@@ -781,6 +856,30 @@ namespace UE
 					return Expression;
 				}
 
+				UMaterialExpression* operator()( const FPrimvarReaderParameterValue& PrimvarReaderValue ) const
+				{
+					UMaterialExpression* Expression = nullptr;
+
+					// We currently only import the "displayColor" primvar in USD as vertex colors
+					// on UE meshes, so that's the only primvar we can support here.
+					if ( PrimvarReaderValue.PrimvarName == "displayColor" )
+					{
+						Expression = UMaterialEditingLibrary::CreateMaterialExpression( &Material, UMaterialExpressionVertexColor::StaticClass() );
+					}
+					else
+					{
+						Expression = UMaterialEditingLibrary::CreateMaterialExpression( &Material, UMaterialExpressionConstant4Vector::StaticClass() );
+
+						UMaterialExpressionConstant4Vector* FallbackExpression = Cast< UMaterialExpressionConstant4Vector >( Expression );
+						if ( FallbackExpression )
+						{
+							FallbackExpression->Constant = FLinearColor( PrimvarReaderValue.FallbackValue );
+						}
+					}
+
+					return Expression;
+				}
+
 			private:
 				UMaterial& Material;
 			};
@@ -918,6 +1017,11 @@ namespace UE
 					SetTextureParameterValue( Material, ParameterName, TextureValue.Texture );
 				}
 
+				void operator()( const FPrimvarReaderParameterValue& PrimvarReaderValue ) const
+				{
+					SetVectorParameterValue( Material, ParameterName, PrimvarReaderValue.FallbackValue );
+				}
+
 				void operator()( const bool BoolValue ) const
 				{
 					SetBoolParameterValue( Material, ParameterName, BoolValue );
@@ -952,6 +1056,18 @@ namespace UE
 					SetTextureParameterValue( Material, *FString::Printf( TEXT( "%sTexture" ), ParameterName ), TextureValue.Texture );
 					SetScalarParameterValue( Material, *FString::Printf( TEXT( "Use%sTexture" ), ParameterName ), 1.0f );
 				}
+
+				void operator()( const FPrimvarReaderParameterValue& PrimvarReaderValue ) const
+				{
+					FSetParameterValueVisitor::operator()( PrimvarReaderValue.FallbackValue );
+
+					// We currently only import the "displayColor" primvar in USD as vertex colors
+					// on UE meshes, so that's the only primvar we can support here.
+					if ( PrimvarReaderValue.PrimvarName == "displayColor" )
+					{
+						SetScalarParameterValue( Material, TEXT( "UseVertexColorForBaseColor" ), 1.0f );
+					}
+				}
 			};
 
 			void SetParameterValue( UMaterialInstance& Material, const TCHAR* ParameterName, const FParameterValue& ParameterValue, bool bForUsdPreviewSurface )
diff --git a/Engine/Plugins/Importers/USDImporter/Source/UnrealUSDWrapper/Private/UnrealUSDWrapper.cpp b/Engine/Plugins/Importers/USDImporter/Source/UnrealUSDWrapper/Private/UnrealUSDWrapper.cpp
--- a/Engine/Plugins/Importers/USDImporter/Source/UnrealUSDWrapper/Private/UnrealUSDWrapper.cpp
+++ b/Engine/Plugins/Importers/USDImporter/Source/UnrealUSDWrapper/Private/UnrealUSDWrapper.cpp
@@ -116,6 +116,7 @@ namespace UnrealIdentifiers
 
 	const TfToken UsdPreviewSurface = TfToken( "UsdPreviewSurface" );
 	const TfToken UsdPrimvarReader_float2 = TfToken( "UsdPrimvarReader_float2" );
+	const TfToken UsdPrimvarReader_float3 = TfToken( "UsdPrimvarReader_float3" );
 	const TfToken UsdUVTexture = TfToken( "UsdUVTexture" );
 
 	const TfToken WorldSpaceNormals = TfToken( "worldSpaceNormals" );
diff --git a/Engine/Plugins/Importers/USDImporter/Source/UnrealUSDWrapper/Public/UnrealUSDWrapper.h b/Engine/Plugins/Importers/USDImporter/Source/UnrealUSDWrapper/Public/UnrealUSDWrapper.h
--- a/Engine/Plugins/Importers/USDImporter/Source/UnrealUSDWrapper/Public/UnrealUSDWrapper.h
+++ b/Engine/Plugins/Importers/USDImporter/Source/UnrealUSDWrapper/Public/UnrealUSDWrapper.h
@@ -312,6 +312,7 @@ namespace UnrealIdentifiers
 	// Tokens copied from usdImaging, because at the moment it's all we need from it
 	extern UNREALUSDWRAPPER_API const pxr::TfToken UsdPreviewSurface;
 	extern UNREALUSDWRAPPER_API const pxr::TfToken UsdPrimvarReader_float2;
+	extern UNREALUSDWRAPPER_API const pxr::TfToken UsdPrimvarReader_float3;
 	extern UNREALUSDWRAPPER_API const pxr::TfToken UsdUVTexture;
 
 	// Token used to indicate that a material parsed from a material prim should use world space normals
