Fix for uninitialized struct members
#rnx

diff --git a/Engine/Plugins/Importers/USDImporter/Source/USDExporter/Public/USDAssetOptions.h b/Engine/Plugins/Importers/USDImporter/Source/USDExporter/Public/USDAssetOptions.h
--- a/Engine/Plugins/Importers/USDImporter/Source/USDExporter/Public/USDAssetOptions.h
+++ b/Engine/Plugins/Importers/USDImporter/Source/USDExporter/Public/USDAssetOptions.h
@@ -33,7 +33,7 @@ struct USDEXPORTER_API FUsdMeshAssetOptions
 
 	/** If true, the mesh data is exported to yet another "payload" file, and referenced via a payload composition arc */
 	UPROPERTY( EditAnywhere, config, BlueprintReadWrite, Category = "Mesh options" )
-	bool bUsePayload;
+	bool bUsePayload = false;
 
 	/** USD format to use for exported payload files */
 	UPROPERTY( EditAnywhere, config, BlueprintReadWrite, Category = "Mesh options", meta = ( EditCondition = bUsePayload, GetOptions = "USDExporter.LevelExporterUSDOptions.GetUsdExtensions" ) )
@@ -41,11 +41,11 @@ struct USDEXPORTER_API FUsdMeshAssetOptions
 
 	/** Whether to bake the mesh's assigned material and export these as separate UsdPreviewSurface assets */
 	UPROPERTY( EditAnywhere, config, BlueprintReadWrite, Category = "Material options" )
-	bool bBakeMaterials;
+	bool bBakeMaterials = false;
 
 	/** Whether to remove the 'unrealMaterial' attribute after binding the corresponding baked material */
 	UPROPERTY( EditAnywhere, config, BlueprintReadWrite, Category = "Material options", meta = ( EditCondition = bBakeMaterials ) )
-	bool bRemoveUnrealMaterials;
+	bool bRemoveUnrealMaterials = false;
 
 	UPROPERTY( EditAnywhere, config, BlueprintReadWrite, Category = "Material options", meta = ( EditCondition = bBakeMaterials ) )
 	FUsdMaterialBakingOptions MaterialBakingOptions;
