USD: Fix randomly failing USD tests when USD Stage Editor window is opened.

diff --git a/Engine/Plugins/Importers/USDImporter/Source/USDStageEditor/Private/SUSDStage.cpp b/Engine/Plugins/Importers/USDImporter/Source/USDStageEditor/Private/SUSDStage.cpp
--- a/Engine/Plugins/Importers/USDImporter/Source/USDStageEditor/Private/SUSDStage.cpp
+++ b/Engine/Plugins/Importers/USDImporter/Source/USDStageEditor/Private/SUSDStage.cpp
@@ -257,11 +257,16 @@ void SUsdStage::SetupStageActorDelegates()
 		OnActorDestroyedHandle = ViewModel.UsdStageActor->OnActorDestroyed.AddLambda(
 			[ this ]()
 			{
+				// Refresh widgets on game thread, but close the stage right away. In some contexts this is important, for example when
+				// running a Python script: If our USD Stage Editor is open and our script deletes an actor, it will trigger OnActorDestroyed.
+				// If we had this CloseStage() call inside the AsyncTask, it would only take place when the script has finished running
+				// (and control flow returned to the game thread) which can lead to some weird results (and break automated tests).
+				// We could get around this on the Python script's side by just yielding, but there may be other scenarios
+				ClearStageActorDelegates();
+				this->ViewModel.CloseStage();
+
 				AsyncTask( ENamedThreads::GameThread, [this]()
 				{
-					ClearStageActorDelegates();
-					this->ViewModel.CloseStage();
-
 					this->Refresh();
 				});
 			}
