USD: Fix material prim name collisions when baking materials on meshes with LODs

diff --git a/Engine/Plugins/Importers/USDImporter/Source/USDExporter/Public/USDConversionBlueprintContext.h b/Engine/Plugins/Importers/USDImporter/Source/USDExporter/Public/USDConversionBlueprintContext.h
--- a/Engine/Plugins/Importers/USDImporter/Source/USDExporter/Public/USDConversionBlueprintContext.h
+++ b/Engine/Plugins/Importers/USDImporter/Source/USDExporter/Public/USDConversionBlueprintContext.h
@@ -148,7 +148,7 @@ public:
 	 * @param bUsePayload - Should be True if the Stage was exported using payload files to store the actual Mesh prims. Also dictates minor
 	 *                      behaviors when authoring the material binding relationships.
 	 * @param bRemoveUnrealMaterials - Whether to remove the `unrealMaterial` attributes after replacing them with material bindings.
-	 *                                 Important because the `unrealMaterial` attributes will be used as a higher priority when determining material assignments
+	 *                                 Important because the `unrealMaterial` attributes will be used as a higher priority when determining material assignments on import
 	 */
 	UFUNCTION( BlueprintCallable, Category = "Conversion utils" )
 	void ReplaceUnrealMaterialsWithBaked( const FFilePath& LayerToAuthorIn, const TMap<FString, FString>& BakedMaterials, bool bIsAssetLayer, bool bUsePayload, bool bRemoveUnrealMaterials );
diff --git a/Engine/Plugins/Importers/USDImporter/Source/USDUtilities/Private/USDGeomMeshConversion.cpp b/Engine/Plugins/Importers/USDImporter/Source/USDUtilities/Private/USDGeomMeshConversion.cpp
--- a/Engine/Plugins/Importers/USDImporter/Source/USDUtilities/Private/USDGeomMeshConversion.cpp
+++ b/Engine/Plugins/Importers/USDImporter/Source/USDUtilities/Private/USDGeomMeshConversion.cpp
@@ -1900,7 +1900,7 @@ void UsdUtils::ReplaceUnrealMaterialsWithBaked(
 	// This handles Mesh prims as well as GeomSubset prims.
 	// Note how we receive the stage as an argument instead of capturing it from the outer scope:
 	// This ensures the inner function doesn't hold a reference to the stage
-	TFunction<void( pxr::UsdStageRefPtr StageToTraverse, pxr::UsdPrim Prim, TOptional<FMaterialScopePrim> MatPrimScope, TOptional<pxr::UsdVariantSet> OuterVariantSet )> TraverseForMaterialReplacement;
+	TFunction<void( pxr::UsdStageRefPtr StageToTraverse, pxr::UsdPrim Prim, TOptional<FMaterialScopePrim>& MatPrimScope, TOptional<pxr::UsdVariantSet> OuterVariantSet )> TraverseForMaterialReplacement;
 	TraverseForMaterialReplacement =
 		[
 			&TraverseForMaterialReplacement,
@@ -1915,7 +1915,7 @@ void UsdUtils::ReplaceUnrealMaterialsWithBaked(
 		(
 			pxr::UsdStageRefPtr StageToTraverse,
 			pxr::UsdPrim Prim,
-			TOptional<FMaterialScopePrim> MatPrimScope,
+			TOptional<FMaterialScopePrim>& MatPrimScope,
 			TOptional<pxr::UsdVariantSet> OuterVariantSet
 		)
 		{
@@ -1927,12 +1927,12 @@ void UsdUtils::ReplaceUnrealMaterialsWithBaked(
 				TOptional<std::string> OriginalSelection = VarSet.HasAuthoredVariantSelection() ? VarSet.GetVariantSelection() : TOptional<std::string>{};
 
 				// Prims within variant sets can't have relationships to prims outside the scope of the prim that
-				// contains the variant set itself.This means we'll need a new material scope prim if we're stepping
+				// contains the variant set itself. This means we'll need a new material scope prim if we're stepping
 				// into a variant within an asset layer, so that any material proxy prims we author are contained within it.
-				// Note that we only do this for asset layers : If we're parsing the root layer, any LOD variant sets we can step into
+				// Note that we only do this for asset layers: If we're parsing the root layer, any LOD variant sets we can step into
 				// are brought in via references to asset files, and we know that referenced subtree only has relationships to
-				// things within that same subtree( which will be entirely brought in to the root layer ).This means we can
-				// just keep inner_mat_prim_scope as Noneand default to using the layer's mat scope prim if we need one
+				// things within that same subtree ( which will be entirely brought in to the root layer ). This means we can
+				// just keep inner_mat_prim_scope as None and default to using the layer's mat scope prim if we need one
 				TOptional<FMaterialScopePrim> InnerMatPrimScope = bIsAssetLayer ? FMaterialScopePrim{ StageToTraverse, Prim } : TOptional<FMaterialScopePrim>{};
 
 				// Switch into each of the LOD variants the prim has, and recurse into the child prims
@@ -2044,13 +2044,13 @@ void UsdUtils::ReplaceUnrealMaterialsWithBaked(
 
 					FMaterialScopePrim* MatPrimScopePtr = nullptr;
 
-					// On-demand create a *single* material scope prim for the stage, if we're not inside a variant set
 					if ( MatPrimScope.IsSet() )
 					{
 						MatPrimScopePtr = &MatPrimScope.GetValue();
 					}
 					else
 					{
+						// On-demand create a *single* material scope prim for the stage, if we're not inside a variant set
 						if ( !StageMatScope.IsSet() )
 						{
 							// If a prim from a stage references another layer, USD's composition will effectively
@@ -2108,7 +2108,8 @@ void UsdUtils::ReplaceUnrealMaterialsWithBaked(
 		};
 
 	pxr::UsdPrim Root = Stage.GetPseudoRoot();
-	TraverseForMaterialReplacement( UsdStage, Root, {}, {} );
+	TOptional<FMaterialScopePrim> Empty;
+	TraverseForMaterialReplacement( UsdStage, Root, Empty, {} );
 }
 
 #undef LOCTEXT_NAMESPACE
