USD: Fix behavior when exporting level sequences with subsequences that have spawnables.

Fix spawnables showing as visible on the exported level;
Fix spawnables being left alive on the level when exporting sequences with subsequences that have spawnables;
Fix spawnables remaining visible on the level if the spawnable tracks have zero keys;
Fix spawnable tracks not writing out the first timeSample if it matches the default prim visibility;
Fix tracks with no keys and just altered default values not being exported due to being considered non-animated;

diff --git a/Engine/Plugins/Importers/USDImporter/Source/USDExporter/Private/LevelSequenceExporterUSD.cpp b/Engine/Plugins/Importers/USDImporter/Source/USDExporter/Private/LevelSequenceExporterUSD.cpp
--- a/Engine/Plugins/Importers/USDImporter/Source/USDExporter/Private/LevelSequenceExporterUSD.cpp
+++ b/Engine/Plugins/Importers/USDImporter/Source/USDExporter/Private/LevelSequenceExporterUSD.cpp
@@ -28,14 +28,10 @@
 #include "MovieSceneTimeHelpers.h"
 #include "MovieSceneTrack.h"
 #include "Sections/MovieSceneSubSection.h"
-#include "Tracks/MovieScene3DTransformTrack.h"
-#include "Tracks/MovieSceneBoolTrack.h"
-#include "Tracks/MovieSceneFloatTrack.h"
 #include "Tracks/MovieScenePropertyTrack.h"
 #include "Tracks/MovieSceneSkeletalAnimationTrack.h"
 #include "Tracks/MovieSceneSpawnTrack.h"
 #include "Tracks/MovieSceneSubTrack.h"
-#include "Tracks/MovieSceneVisibilityTrack.h"
 #include "UObject/UObjectGlobals.h"
 
 namespace UE
@@ -66,15 +62,24 @@ namespace UE
 
 					if ( Object )
 					{
-						if ( USceneComponent* Component = Cast<USceneComponent>( Object ) )
-						{
-							const bool bNewVisibility = true;
-							Component->SetVisibility( bNewVisibility );
-						}
-						else if ( AActor* Actor = Cast<AActor>( Object ) )
+						USceneComponent* Component = nullptr;
+
+						if ( AActor* Actor = Cast<AActor>( Object ) )
 						{
 							const bool bIsHidden = false;
 							Actor->SetHidden( bIsHidden );
+
+							Component = Actor->GetRootComponent();
+						}
+						else
+						{
+							Component = Cast<USceneComponent>( Object );
+						}
+
+						if ( Component )
+						{
+							const bool bNewVisibility = true;
+							Component->SetVisibility( bNewVisibility );
 						}
 					}
 
@@ -85,15 +90,27 @@ namespace UE
 				{
 					if ( bDestroyingJustHides )
 					{
-						if ( USceneComponent* Component = Cast<USceneComponent>( &Object ) )
-						{
-							const bool bNewVisibility = false;
-							Component->SetVisibility( bNewVisibility );
-						}
-						else if ( AActor* Actor = Cast<AActor>( &Object ) )
+						USceneComponent* Component = nullptr;
+
+						if ( AActor* Actor = Cast<AActor>( &Object ) )
 						{
 							const bool bIsHidden = true;
 							Actor->SetHidden( bIsHidden );
+
+							Component = Actor->GetRootComponent();
+						}
+						else
+						{
+							Component = Cast<USceneComponent>( &Object );
+						}
+
+						// Make sure we toggle both actor and component as the component is the one
+						// that is monitored by the level exporter, and we want our spawnables to be hidden
+						// by default
+						if ( Component )
+						{
+							const bool bNewVisibility = false;
+							Component->SetVisibility( bNewVisibility );
 						}
 					}
 					else
@@ -106,6 +123,32 @@ namespace UE
 					}
 				}
 
+				void DeleteSpawns( IMovieScenePlayer& Player )
+				{
+					bDestroyingJustHides = false;
+					CleanUp( Player );
+
+					// If we still have existing spawns it may be because our base class' Register member didn't contain an
+					// entry for a spawnable before we called CleanUp (check its implementation: It just iterates over that Register).
+					// This is expected in some scenarios because we're sort of abusing this inheritance: The base class expects
+					// CleanUp to delete the object already and so removes its Register entry, but we'll keep the instances
+					// alive when bDestroyingJustHides=true. Because of this we must explicitly clean up these "abandoned" spawns here,
+					// which resynchronizes us with Register:
+					TArray<UObject*> ObjectsToDelete;
+					ObjectsToDelete.Reserve( ExistingSpawns.Num() );
+					for ( const TPair<FGuid, UObject*>& Pair : ExistingSpawns )
+					{
+						if ( UObject* Spawn = Pair.Value )
+						{
+							ObjectsToDelete.Add( Spawn );
+						}
+					}
+					for ( UObject* Object : ObjectsToDelete )
+					{
+						DestroySpawnedObject( *Object );
+					}
+				}
+
 				// FLevelSequenceHidingSpawnRegister is a bit of a hack and just hides it's spawned actors instead
 				// of deleting them (when we want it to do so). Unfortunately, the base FMovieSceneSpawnRegister part
 				// will still nevertheless clear it's Register entry for the spawnable when deleting (even if just hiding),
@@ -169,8 +212,7 @@ namespace UE
 			public:
 				~FLevelSequenceExportContext()
 				{
-					SpawnRegister->bDestroyingJustHides = false;
-					SpawnRegister->CleanUp( *Sequencer );
+					SpawnRegister->DeleteSpawns( *Sequencer );
 				}
 			};
 
@@ -212,13 +254,18 @@ namespace UE
 				return nullptr;
 			}
 
-			// Returns false if this track has no keys and no skeletal animations, and true otherwise
 			bool IsTrackAnimated( const UMovieSceneTrack* Track )
 			{
 				for ( const UMovieSceneSection* Section : Track->GetAllSections() )
 				{
-					TOptional< TRange<FFrameNumber> > Range = Section->GetAutoSizeRange();
-					if ( Range.IsSet() && !Range.GetValue().IsEmpty() )
+					// We can't just check whether a section has range here because it may just have a bunch
+					// of channels with modified default values instead (and so no range or even keys, but can still affect the level).
+					// Sadly there's no way of telling if the default value in the channel has been modified or not (that is, whether it
+					// matches the unanimated value or not), so we'll just have to export any track with an active section for now.
+					// This will emit a bunch of unwanted tracks, but later on we may want to do a post-processing pass on our exported
+					// data to make sure we don't emit intermediary keys on linearly interpolated sections, and that would naturally fix
+					// this too.
+					if ( Section->IsActive() )
 					{
 						return true;
 					}
@@ -406,7 +453,9 @@ namespace UE
 					{
 						for ( const UMovieSceneTrack* Track : Binding->GetTracks() )
 						{
-							if ( !IsTrackAnimated( Track ) )
+							// Let even non-animated spawn tracks through because even if they have zero keyframes we'll still need
+							// to bake the visibility of the spawnables despawning when e.g. a subsequence shot is over
+							if ( !Track || ( !IsTrackAnimated( Track ) && !Track->IsA<UMovieSceneSpawnTrack>() ) )
 							{
 								continue;
 							}
diff --git a/Engine/Plugins/Importers/USDImporter/Source/USDExporter/Public/LevelSequenceExporterUSDOptions.h b/Engine/Plugins/Importers/USDImporter/Source/USDExporter/Public/LevelSequenceExporterUSDOptions.h
--- a/Engine/Plugins/Importers/USDImporter/Source/USDExporter/Public/LevelSequenceExporterUSDOptions.h
+++ b/Engine/Plugins/Importers/USDImporter/Source/USDExporter/Public/LevelSequenceExporterUSDOptions.h
@@ -46,7 +46,7 @@ public:
 
 	/**
 	 * The effect of subsequences is always included on the main exported layer, but if this option is true we will also
-	 * export individual sublayers for each subsequence, so that they can be used by itself if needed.
+	 * export individual sublayers for each subsequence, so that they can be used by themselves if needed.
 	 */
 	UPROPERTY( EditAnywhere, config, BlueprintReadWrite, Category = "Level Sequence" )
 	bool bExportSubsequencesAsLayers = true;
diff --git a/Engine/Plugins/Importers/USDImporter/Source/USDUtilities/Private/USDPrimConversion.cpp b/Engine/Plugins/Importers/USDImporter/Source/USDUtilities/Private/USDPrimConversion.cpp
--- a/Engine/Plugins/Importers/USDImporter/Source/USDUtilities/Private/USDPrimConversion.cpp
+++ b/Engine/Plugins/Importers/USDImporter/Source/USDUtilities/Private/USDPrimConversion.cpp
@@ -2449,10 +2449,32 @@ bool UnrealToUsd::CreateComponentPropertyBaker( UE::FUsdPrim& Prim, const UScene
 				if ( Component.bHiddenInGame || Component.GetOwner()->IsHidden() )
 				{
 					Imageable.MakeInvisible( UsdTimeCode );
+
+					if ( pxr::UsdAttribute Attr = Imageable.CreateVisibilityAttr() )
+					{
+						if ( !Attr.HasAuthoredValue() )
+						{
+							Attr.Set<pxr::TfToken>( pxr::UsdGeomTokens->invisible, UsdTimeCode );
+						}
+					}
 				}
 				else
 				{
 					Imageable.MakeVisible( UsdTimeCode );
+
+					// Imagine our visibility track has a single key that switches to hidden at frame 60.
+					// If our prim is visible by default, MakeVisible will author absolutely nothing, and we'll end up
+					// with a timeSamples that just has '60: "invisible"'. Weirdly enough, in USD that means the prim
+					// will be invisible throughout *the entire duration of the animation* though, which is not what we want.
+					// This check will ensure that if we're visible we should have a value here and not rely on the
+					// fallback value of 'visible', as that doesn't behave how we want.
+					if ( pxr::UsdAttribute Attr = Imageable.CreateVisibilityAttr() )
+					{
+						if ( !Attr.HasAuthoredValue() )
+						{
+							Attr.Set<pxr::TfToken>( pxr::UsdGeomTokens->inherited, UsdTimeCode );
+						}
+					}
 				}
 			};
 		}
