USD: Handle component visibility with bHiddenInGame instead of bVisible as the sequencer can only manipulate the former.

diff --git a/Engine/Plugins/Importers/USDImporter/Source/USDExporter/Private/LevelSequenceExporterUSD.cpp b/Engine/Plugins/Importers/USDImporter/Source/USDExporter/Private/LevelSequenceExporterUSD.cpp
--- a/Engine/Plugins/Importers/USDImporter/Source/USDExporter/Private/LevelSequenceExporterUSD.cpp
+++ b/Engine/Plugins/Importers/USDImporter/Source/USDExporter/Private/LevelSequenceExporterUSD.cpp
@@ -67,7 +67,7 @@ namespace UE
 						if ( AActor* Actor = Cast<AActor>( Object ) )
 						{
 							const bool bIsHidden = false;
-							Actor->SetHidden( bIsHidden );
+							Actor->SetActorHiddenInGame( bIsHidden );
 
 							Component = Actor->GetRootComponent();
 						}
@@ -78,8 +78,8 @@ namespace UE
 
 						if ( Component )
 						{
-							const bool bNewVisibility = true;
-							Component->SetVisibility( bNewVisibility );
+							const bool bNewHidden = false;
+							Component->SetHiddenInGame( bNewHidden );
 						}
 					}
 
@@ -95,7 +95,7 @@ namespace UE
 						if ( AActor* Actor = Cast<AActor>( &Object ) )
 						{
 							const bool bIsHidden = true;
-							Actor->SetHidden( bIsHidden );
+							Actor->SetActorHiddenInGame( bIsHidden );
 
 							Component = Actor->GetRootComponent();
 						}
@@ -109,8 +109,8 @@ namespace UE
 						// by default
 						if ( Component )
 						{
-							const bool bNewVisibility = false;
-							Component->SetVisibility( bNewVisibility );
+							const bool bNewHidden = true;
+							Component->SetHiddenInGame( bNewHidden );
 						}
 					}
 					else
diff --git a/Engine/Plugins/Importers/USDImporter/Source/USDSchemas/Private/USDGeomPointInstancerTranslator.cpp b/Engine/Plugins/Importers/USDImporter/Source/USDSchemas/Private/USDGeomPointInstancerTranslator.cpp
--- a/Engine/Plugins/Importers/USDImporter/Source/USDSchemas/Private/USDGeomPointInstancerTranslator.cpp
+++ b/Engine/Plugins/Importers/USDImporter/Source/USDSchemas/Private/USDGeomPointInstancerTranslator.cpp
@@ -216,7 +216,7 @@ void FUsdGeomPointInstancerTranslator::UpdateComponents( USceneComponent* PointI
 					// can mark the prototype parent prim as invisible, as an attempt to prevent them from being parsed directly as scene meshes.
 					// We will handle the visibility of all the components we spawn manually
 					bPrototypeRootVisible = UsdUtils::HasInheritedVisibility( PrototypePrim, Context->Time );
-					PrototypeXformComponent->SetVisibility( bPointInstancerVisible && bPrototypeRootVisible );
+					PrototypeXformComponent->SetHiddenInGame( !bPointInstancerVisible || !bPrototypeRootVisible );
 
 					PrototypeParentComponent = PrototypeXformComponent;
 				}
@@ -250,7 +250,7 @@ void FUsdGeomPointInstancerTranslator::UpdateComponents( USceneComponent* PointI
 							// Traverse up every time because we have no idea how far deep each UsdGeomMesh prim is
 							const bool bHasInvisibleParent = UsdUtils::HasInvisibleParent( PrototypeTargetPrim, PrototypePrim, Context->Time );
 							const bool bPrototypeMeshVisible = UsdUtils::HasInheritedVisibility( PrototypeTargetPrim, Context->Time );
-							HismComponent->SetVisibility( bPointInstancerVisible && bPrototypeRootVisible && bPrototypeMeshVisible && !bHasInvisibleParent );
+							HismComponent->SetHiddenInGame( !bPointInstancerVisible || !bPrototypeRootVisible || !bPrototypeMeshVisible || bHasInvisibleParent );
 
 							UUsdAssetCache& AssetCache = *Context->AssetCache.Get();
 							UStaticMesh* StaticMesh = UsdGeomPointInstancerTranslatorImpl::SetStaticMesh( PrototypeGeomMesh, *HismComponent, AssetCache );
diff --git a/Engine/Plugins/Importers/USDImporter/Source/USDStage/Private/USDStageActor.cpp b/Engine/Plugins/Importers/USDImporter/Source/USDStage/Private/USDStageActor.cpp
--- a/Engine/Plugins/Importers/USDImporter/Source/USDStage/Private/USDStageActor.cpp
+++ b/Engine/Plugins/Importers/USDImporter/Source/USDStage/Private/USDStageActor.cpp
@@ -474,7 +474,7 @@ struct FUsdStageActorImpl
 			const bool bPrimIsVisible = bPrimHasInheritedVisibility && !bPrimHasInvisibleParent;
 
 			const bool bComponentHasInvisibleTag = Component->ComponentTags.Contains( UnrealIdentifiers::Invisible );
-			const bool bComponentIsVisible = Component->IsVisible();
+			const bool bComponentIsVisible = !Component->bHiddenInGame;
 
 			const bool bTagIsCorrect = bComponentHasInvisibleTag == !bPrimHasInheritedVisibility;
 			const bool bComputedVisibilityIsCorrect = bPrimIsVisible == bComponentIsVisible;
@@ -503,7 +503,7 @@ struct FUsdStageActorImpl
 			{
 				const bool bPropagateToChildren = false;
 				Component->Modify();
-				Component->SetVisibility( bPrimIsVisible, bPropagateToChildren );
+				Component->SetHiddenInGame( !bPrimIsVisible, bPropagateToChildren );
 			}
 
 			for ( const TPair<FString, UUsdPrimTwin*>& ChildPair : PrimTwin.GetChildren() )
@@ -536,8 +536,8 @@ struct FUsdStageActorImpl
 		PrimSceneComponent->ComponentTags.Remove( UnrealIdentifiers::Inherited );
 
 		const bool bPropagateToChildren = true;
-		const bool bVisible = false;
-		PrimSceneComponent->SetVisibility( bVisible, bPropagateToChildren );
+		const bool bNewHidden = true;
+		PrimSceneComponent->SetHiddenInGame( bNewHidden, bPropagateToChildren );
 	}
 
 	static void SendAnalytics( AUsdStageActor* StageActor, double ElapsedSeconds, double NumberOfFrames, const FString& Extension )
@@ -2252,14 +2252,13 @@ void AUsdStageActor::OnObjectPropertyChanged( UObject* ObjectBeingModified, FPro
 				{
 					PrimSceneComponent->Modify();
 
-					const bool bVisible = PrimSceneComponent->GetVisibleFlag();
-					if ( bVisible )
+					if ( PrimSceneComponent->bHiddenInGame )
 					{
-						FUsdStageActorImpl::MakeVisible( *UsdPrimTwin, UsdStage );
+						FUsdStageActorImpl::MakeInvisible( *UsdPrimTwin );
 					}
 					else
 					{
-						FUsdStageActorImpl::MakeInvisible( *UsdPrimTwin );
+						FUsdStageActorImpl::MakeVisible( *UsdPrimTwin, UsdStage );
 					}
 				}
 
diff --git a/Engine/Plugins/Importers/USDImporter/Source/USDUtilities/Private/USDPrimConversion.cpp b/Engine/Plugins/Importers/USDImporter/Source/USDUtilities/Private/USDPrimConversion.cpp
--- a/Engine/Plugins/Importers/USDImporter/Source/USDUtilities/Private/USDPrimConversion.cpp
+++ b/Engine/Plugins/Importers/USDImporter/Source/USDUtilities/Private/USDPrimConversion.cpp
@@ -252,7 +252,7 @@ bool UsdToUnreal::ConvertXformable( const pxr::UsdStageRefPtr& Stage, const pxr:
 
 	// Computed (effective) visibility
 	const bool bIsHidden = ( Xformable.ComputeVisibility( EvalTime ) == pxr::UsdGeomTokens->invisible );
-	SceneComponent.SetVisibility( !bIsHidden );
+	SceneComponent.SetHiddenInGame( bIsHidden );
 
 	// Per-prim visibility
 	bool bIsInvisible = false; // Default to 'inherited'
@@ -1918,7 +1918,7 @@ bool UnrealToUsd::ConvertSceneComponent( const pxr::UsdStageRefPtr& Stage, const
 		else if ( !SceneComponent->ComponentTags.Contains( UnrealIdentifiers::Inherited ) )
 		{
 			// We don't have visible nor inherited tags: We're probably exporting a pure UE component, so write out component visibility instead
-			Value = SceneComponent->IsVisible() ? pxr::UsdGeomTokens->inherited : pxr::UsdGeomTokens->invisible;
+			Value = SceneComponent->bHiddenInGame ? pxr::UsdGeomTokens->invisible : pxr::UsdGeomTokens->inherited;
 		}
 
 		VisibilityAttr.Set<pxr::TfToken>( Value );
