USD: Fix level sequence export conflicting with opened sequencer.

diff --git a/Engine/Plugins/Importers/USDImporter/Source/USDExporter/Private/LevelSequenceExporterUSD.cpp b/Engine/Plugins/Importers/USDImporter/Source/USDExporter/Private/LevelSequenceExporterUSD.cpp
--- a/Engine/Plugins/Importers/USDImporter/Source/USDExporter/Private/LevelSequenceExporterUSD.cpp
+++ b/Engine/Plugins/Importers/USDImporter/Source/USDExporter/Private/LevelSequenceExporterUSD.cpp
@@ -18,8 +18,10 @@
 
 #include "AssetExportTask.h"
 #include "Components/SkeletalMeshComponent.h"
+#include "Editor.h"
 #include "ISequencer.h"
 #include "ISequencerModule.h"
+#include "LevelEditorSequencerIntegration.h"
 #include "LevelSequence.h"
 #include "Misc/LevelSequenceEditorSpawnRegister.h"
 #include "Modules/ModuleManager.h"
@@ -28,6 +30,7 @@
 #include "MovieSceneTimeHelpers.h"
 #include "MovieSceneTrack.h"
 #include "Sections/MovieSceneSubSection.h"
+#include "Subsystems/AssetEditorSubsystem.h"
 #include "Tracks/MovieScenePropertyTrack.h"
 #include "Tracks/MovieSceneSkeletalAnimationTrack.h"
 #include "Tracks/MovieSceneSpawnTrack.h"
@@ -906,11 +909,51 @@ bool ULevelSequenceExporterUsd::ExportBinary( UObject* Object, const TCHAR* Type
 		}
 	}
 
+	// Always close all opened sequencers since it doesn't look like we're supposed to have more than one opened at a time.
+	// Without this, the ResetToNewRootSequence call may actually evaluate our opened subsequence at the playhead position of
+	// *other* sequencers, for whatever reason. We also can't call RestorePreAnimatedState on the sequence we're exporting
+	// after that either, as it hasn't stored anything yet.
+	// Additionally, it seems the sequencer will also attempt to interpolate between all opened sequencers when evaluating a track,
+	// which could affect the sequence we're exporting.
+	TArray<UObject*> AssetsToReopenEditorsFor;
+	UAssetEditorSubsystem* AssetEditorSubsystem = nullptr;
+	if ( GEditor && GIsEditor )
+	{
+		AssetEditorSubsystem = GEditor->GetEditorSubsystem<UAssetEditorSubsystem>();
+		if ( AssetEditorSubsystem )
+		{
+			for ( const TWeakPtr<ISequencer>& Sequencer : FLevelEditorSequencerIntegration::Get().GetSequencers() )
+			{
+				if ( TSharedPtr<ISequencer> PinnedSequencer = Sequencer.Pin() )
+				{
+					if ( PinnedSequencer == TempSequencer )
+					{
+						continue;
+					}
+
+					ULevelSequence* OpenedSequence = Cast<ULevelSequence>( PinnedSequencer->GetRootMovieSceneSequence() );
+					if ( !OpenedSequence )
+					{
+						continue;
+					}
+
+					AssetsToReopenEditorsFor.Add( OpenedSequence );
+					AssetEditorSubsystem->CloseAllEditorsForAsset( OpenedSequence );
+				}
+			}
+		}
+	}
+
 	LevelSequenceExporterImpl::ExportMovieSceneSequence( Context, *LevelSequence, TargetFileName );
 
 	// Set this back to Stopped or else it will keep the editor viewport controls permanently hidden
 	TempSequencer->SetPlaybackStatus( EMovieScenePlayerStatus::Stopped );
 
+	if ( GEditor && GIsEditor && AssetEditorSubsystem )
+	{
+		AssetEditorSubsystem->OpenEditorForAssets( AssetsToReopenEditorsFor );
+	}
+
 	return true;
 #else
 	return false;
