USD: Fix file handles being left open after hashing textures.

diff --git a/Engine/Plugins/Importers/USDImporter/Source/USDUtilities/Private/USDShadeConversion.cpp b/Engine/Plugins/Importers/USDImporter/Source/USDUtilities/Private/USDShadeConversion.cpp
--- a/Engine/Plugins/Importers/USDImporter/Source/USDUtilities/Private/USDShadeConversion.cpp
+++ b/Engine/Plugins/Importers/USDImporter/Source/USDUtilities/Private/USDShadeConversion.cpp
@@ -327,7 +327,7 @@ namespace UE
 					}
 				}
 				// Copied from FMD5Hash::HashFileFromArchive as it doesn't expose its FMD5
-				else if ( FArchive* Ar = IFileManager::Get().CreateFileReader( *ResolvedTexturePath ) )
+				else if ( TUniquePtr<FArchive> Ar{ IFileManager::Get().CreateFileReader( *ResolvedTexturePath ) } )
 				{
 					TArray<uint8> LocalScratch;
 					LocalScratch.SetNumUninitialized( 1024 * 64 );
