USD: Fix assert when changing Nanite threshold while editing our generated LevelSequence.

diff --git a/Engine/Plugins/Importers/USDImporter/Source/USDStage/Private/USDStageActor.cpp b/Engine/Plugins/Importers/USDImporter/Source/USDStage/Private/USDStageActor.cpp
--- a/Engine/Plugins/Importers/USDImporter/Source/USDStage/Private/USDStageActor.cpp
+++ b/Engine/Plugins/Importers/USDImporter/Source/USDStage/Private/USDStageActor.cpp
@@ -1738,6 +1738,21 @@ void AUsdStageActor::ReloadAnimations()
 {
 	TRACE_CPUPROFILER_EVENT_SCOPE( AUsdStageActor::ReloadAnimations );
 
+	// If we're using some property editor that can trigger a stage reload (like the Nanite threshold spinbox),
+	// applying a value may trigger ReloadAnimations -> Can trigger asset editors to open/close/change focus ->
+	// -> Can trigger focus to drop from the property editors -> Can cause the values to be applied from the
+	// property editors when releasing focus -> Can trigger another call to ReloadAnimations.
+	// CloseAllEditorsForAsset in particular is problematic for this because it will destroy the asset editor
+	// (which is TSharedFromThis) and the reentrant call will try use AsShared() internally and assert, as it
+	// hasn't finished being destroyed.
+	// In that case we only want the outer call to change the level sequence, so a reentrant guard does what we need
+	static bool bIsReentrant = false;
+	if ( bIsReentrant )
+	{
+		return;
+	}
+	TGuardValue<bool> ReentrantGuard( bIsReentrant, true );
+
 	if ( !UsdStage )
 	{
 		return;
