USD: Fix animated static meshes sometimes being one frame behind the intended time code.

diff --git a/Engine/Plugins/Importers/USDImporter/Source/USDSchemas/Private/USDGeomMeshTranslator.cpp b/Engine/Plugins/Importers/USDImporter/Source/USDSchemas/Private/USDGeomMeshTranslator.cpp
--- a/Engine/Plugins/Importers/USDImporter/Source/USDSchemas/Private/USDGeomMeshTranslator.cpp
+++ b/Engine/Plugins/Importers/USDImporter/Source/USDSchemas/Private/USDGeomMeshTranslator.cpp
@@ -1328,7 +1328,10 @@ void FUsdGeomMeshTranslator::UpdateComponents( USceneComponent* SceneComponent )
 	)
 	{
 		// The assets might have changed since our attributes are animated
+		// Note that we must wait for these to complete as they make take a while and we want to
+		// reassign our new static meshes when we get to FUsdGeomXformableTranslator::UpdateComponents
 		CreateAssets();
+		Context->CompleteTasks();
 	}
 
 #if WITH_EDITOR
