USD: Use attribute timeSample ranges for geometry cache start and end frames

diff --git a/Engine/Plugins/Importers/USDImporter/Source/USDSchemas/Private/USDGeomMeshTranslator.cpp b/Engine/Plugins/Importers/USDImporter/Source/USDSchemas/Private/USDGeomMeshTranslator.cpp
--- a/Engine/Plugins/Importers/USDImporter/Source/USDSchemas/Private/USDGeomMeshTranslator.cpp
+++ b/Engine/Plugins/Importers/USDImporter/Source/USDSchemas/Private/USDGeomMeshTranslator.cpp
@@ -195,7 +195,7 @@ namespace UsdGeomMeshTranslatorImpl
 			{
 				const pxr::UsdAttribute& Attribute = Prim.GetAttribute( AttributeName );
 
-				if ( Attribute.ValueMightBeTimeVarying() )
+				if ( Attribute && Attribute.ValueMightBeTimeVarying() )
 				{
 					bHasAttributesTimeSamples = true;
 					break;
@@ -712,6 +712,51 @@ namespace UsdGeomMeshTranslatorImpl
 #if WITH_EDITOR
 	void GeometryCacheDataForMeshDescription( FGeometryCacheMeshData& OutMeshData, FMeshDescription& MeshDescription );
 
+	void GetGeometryCacheDataTimeCodeRange( const UE::FUsdStage& Stage, const FString& PrimPath, int32& OutStartFrame, int32& OutEndFrame )
+	{
+		if ( !Stage || PrimPath.IsEmpty() )
+		{
+			return;
+		}
+
+		FScopedUsdAllocs Allocs;
+
+		pxr::UsdPrim UsdPrim = pxr::UsdPrim{ Stage.GetPrimAtPath( UE::FSdfPath{ *PrimPath } ) };
+		if ( !UsdPrim )
+		{
+			return;
+		}
+
+		constexpr bool bIncludeInherited = false;
+		pxr::TfTokenVector GeomMeshAttributeNames = pxr::UsdGeomMesh::GetSchemaAttributeNames( bIncludeInherited );
+		pxr::TfTokenVector GeomPointBasedAttributeNames = pxr::UsdGeomPointBased::GetSchemaAttributeNames( bIncludeInherited );
+
+		GeomMeshAttributeNames.reserve( GeomMeshAttributeNames.size() + GeomPointBasedAttributeNames.size() );
+		GeomMeshAttributeNames.insert( GeomMeshAttributeNames.end(), GeomPointBasedAttributeNames.begin(), GeomPointBasedAttributeNames.end() );
+
+		TOptional<double> MinStartTimeCode;
+		TOptional<double> MaxEndTimeCode;
+
+		for ( const pxr::TfToken& AttributeName : GeomMeshAttributeNames )
+		{
+			if ( const pxr::UsdAttribute& Attribute = UsdPrim.GetAttribute( AttributeName ) )
+			{
+				std::vector<double> TimeSamples;
+				if ( Attribute.GetTimeSamples( &TimeSamples ) && TimeSamples.size() > 0 )
+				{
+					MinStartTimeCode = FMath::Min( MinStartTimeCode.Get( TNumericLimits<double>::Max() ), TimeSamples[ 0 ] );
+					MaxEndTimeCode = FMath::Max( MaxEndTimeCode.Get( TNumericLimits<double>::Lowest() ), TimeSamples[ TimeSamples.size() - 1 ] );
+				}
+			}
+		}
+
+		if ( MinStartTimeCode.IsSet() && MaxEndTimeCode.IsSet() )
+		{
+			OutStartFrame = FMath::FloorToInt( MinStartTimeCode.GetValue() );
+			OutEndFrame = FMath::CeilToInt( MaxEndTimeCode.GetValue() );
+		}
+	}
+
 	UGeometryCache* CreateGeometryCache( const FString& InPrimPath, TArray< FMeshDescription >& LODIndexToMeshDescription, TSharedRef< FUsdSchemaTranslationContext> Context, bool& bOutIsNew )
 	{
 		UGeometryCache* GeometryCache = nullptr;
@@ -747,6 +792,11 @@ namespace UsdGeomMeshTranslatorImpl
 			TMap< FString, TMap< FString, int32 > > Unused;
 			TMap< FString, TMap< FString, int32 > >* MaterialToPrimvarToUVIndex = Context->MaterialToPrimvarToUVIndex ? Context->MaterialToPrimvarToUVIndex : &Unused;
 
+			// Fetch the animated mesh start/end frame as they may be different from just the stage's start and end time codes
+			int32 StartFrame = FMath::FloorToInt( Stage.GetStartTimeCode() );
+			int32 EndFrame = FMath::CeilToInt( Stage.GetEndTimeCode() );
+			UsdGeomMeshTranslatorImpl::GetGeometryCacheDataTimeCodeRange( Stage, InPrimPath, StartFrame, EndFrame );
+
 			// Create and configure a new USDTrack to be added to the GeometryCache
 			UGeometryCacheTrackUsd* UsdTrack = NewObject< UGeometryCacheTrackUsd >( GeometryCache );
 			UsdTrack->Initialize(
@@ -754,8 +804,8 @@ namespace UsdGeomMeshTranslatorImpl
 				InPrimPath,
 				Context->RenderContext,
 				*MaterialToPrimvarToUVIndex,
-				Stage.GetStartTimeCode(),
-				Stage.GetEndTimeCode(),
+				StartFrame,
+				EndFrame,
 				[InPrimPath]( const TWeakObjectPtr<UGeometryCacheTrackUsd> TrackPtr, float Time, FGeometryCacheMeshData& OutMeshData )
 				{
 					UGeometryCacheTrackUsd* Track = TrackPtr.Get();
