USD: Upgrade prim visibility changes to Resyncs to give them the chance to generate new assets if needed

diff --git a/Engine/Plugins/Importers/USDImporter/Source/UnrealUSDWrapper/Private/USDListener.cpp b/Engine/Plugins/Importers/USDImporter/Source/UnrealUSDWrapper/Private/USDListener.cpp
--- a/Engine/Plugins/Importers/USDImporter/Source/UnrealUSDWrapper/Private/USDListener.cpp
+++ b/Engine/Plugins/Importers/USDImporter/Source/UnrealUSDWrapper/Private/USDListener.cpp
@@ -145,7 +145,6 @@ namespace UsdToUnreal
 		ConvertPathRange( InNotice.GetChangedInfoOnlyPaths(), OutInfoChanges );
 		ConvertPathRange( InNotice.GetResyncedPaths(), OutResyncChanges );
 
-		// Upgrade info changes with content reloads into resync changes
 		for ( TPair< FString, TArray<UsdUtils::FObjectChangeNotice> >& InfoPair : OutInfoChanges )
 		{
 			const FString& PrimPath = InfoPair.Key;
@@ -153,11 +152,34 @@ namespace UsdToUnreal
 
 			for ( TArray<UsdUtils::FObjectChangeNotice>::TIterator ChangeIt = InfoPair.Value.CreateIterator(); ChangeIt; ++ChangeIt )
 			{
+				bool bUpgrade = false;
+
+				// Upgrade info changes with content reloads into resync changes
 				if ( ChangeIt->Flags.bDidReloadContent )
+				{
+					bUpgrade = true;
+				}
+
+				// Upgrade visibility changes to resyncs because in case of mesh collapsing having one of the collapsed meshes go visible/invisible
+				// should cause the regeneration of the collapsed asset. This is a bit expensive, but the asset cache will be used so its not as if
+				// the mesh will be completely regenerated however
+				if ( !bUpgrade )
+				{
+					for ( UsdUtils::FAttributeChange& Change : ChangeIt->AttributeChanges )
+					{
+						if ( Change.PropertyName == ANSI_TO_TCHAR( pxr::UsdGeomTokens->visibility.GetString().c_str() ) )
+						{
+							bUpgrade = true;
+							break;
+						}
+					}
+				}
+
+				if ( bUpgrade )
 				{
 					if ( !AnalogueResyncChanges )
 					{
-						AnalogueResyncChanges = &OutResyncChanges.FindOrAdd(PrimPath);
+						AnalogueResyncChanges = &OutResyncChanges.FindOrAdd( PrimPath );
 					}
 
 					AnalogueResyncChanges->Add( *ChangeIt );
