USD: Fix issues caused by new 'unreal' render context material outputs.

Emit 'unreal' render context shader output when baking materials to USD;
Write all material overrides to the MaterialOverrides output layer;
Reuse Material prims with 'unreal' surface output when writing the baked UsdPreviewSurface universal output;
Properly implement removing the 'unreal' surface output when "Remove unreal materials" is checked.

diff --git a/Engine/Plugins/Importers/USDImporter/Source/USDExporter/Private/MaterialExporterUSD.cpp b/Engine/Plugins/Importers/USDImporter/Source/USDExporter/Private/MaterialExporterUSD.cpp
--- a/Engine/Plugins/Importers/USDImporter/Source/USDExporter/Private/MaterialExporterUSD.cpp
+++ b/Engine/Plugins/Importers/USDImporter/Source/USDExporter/Private/MaterialExporterUSD.cpp
@@ -98,6 +98,8 @@ bool UMaterialExporterUsd::ExportMaterial( const UMaterialInterface& Material, c
 
 	UsdStage.SetDefaultPrim( RootPrim );
 
+	UsdUtils::SetUnrealSurfaceOutput( RootPrim, Material.GetPathName() );
+
 	UnrealToUsd::ConvertMaterialToBakedSurface(
 		Material,
 		Options.Properties,
diff --git a/Engine/Plugins/Importers/USDImporter/Source/USDExporter/Private/USDConversionBlueprintContext.cpp b/Engine/Plugins/Importers/USDImporter/Source/USDExporter/Private/USDConversionBlueprintContext.cpp
--- a/Engine/Plugins/Importers/USDImporter/Source/USDExporter/Private/USDConversionBlueprintContext.cpp
+++ b/Engine/Plugins/Importers/USDImporter/Source/USDExporter/Private/USDConversionBlueprintContext.cpp
@@ -439,3 +439,29 @@ void UUsdConversionBlueprintContext::ReplaceUnrealMaterialsWithBaked( const FFil
 #endif // USE_USD_SDK
 }
 
+bool UUsdConversionBlueprintContext::RemoveUnrealSurfaceOutput( const FString& PrimPath, const FFilePath& LayerToAuthorIn )
+{
+#if USE_USD_SDK
+	if ( !Stage )
+	{
+		return false;
+	}
+
+	UE::FUsdPrim Prim = UnrealToUsdImpl::GetPrim( Stage, PrimPath );
+	if ( !Prim )
+	{
+		return false;
+	}
+
+	UE::FSdfLayer Layer;
+	if ( !LayerToAuthorIn.FilePath.IsEmpty() )
+	{
+		Layer = UE::FSdfLayer::FindOrOpen( *LayerToAuthorIn.FilePath );
+	}
+
+	return UsdUtils::RemoveUnrealSurfaceOutput( Prim, Layer );
+#else
+	return false;
+#endif // USE_USD_SDK
+}
+
diff --git a/Engine/Plugins/Importers/USDImporter/Source/USDExporter/Public/USDConversionBlueprintContext.h b/Engine/Plugins/Importers/USDImporter/Source/USDExporter/Public/USDConversionBlueprintContext.h
--- a/Engine/Plugins/Importers/USDImporter/Source/USDExporter/Public/USDConversionBlueprintContext.h
+++ b/Engine/Plugins/Importers/USDImporter/Source/USDExporter/Public/USDConversionBlueprintContext.h
@@ -152,4 +152,15 @@ public:
 	 */
 	UFUNCTION( BlueprintCallable, Category = "Conversion utils" )
 	void ReplaceUnrealMaterialsWithBaked( const FFilePath& LayerToAuthorIn, const TMap<FString, FString>& BakedMaterials, bool bIsAssetLayer, bool bUsePayload, bool bRemoveUnrealMaterials );
+
+	/**
+	 * Clears any opinions for the 'unreal' render context surface output of MaterialPrim within LayerToAuthorIn.
+	 * If LayerToAuthorIn is empty it will clear opinions from all layers of the stage's layer stack.
+	 *
+	 * @param Primpath - Path to the prim pxr::UsdShadeMaterial schema to update the 'unreal' surface output of (e.g. "/Root/MyCube/Red")
+	 * @param LayerToAuthorIn - Layer to clear the opinions in. Can be the empty string to clear opinions from all layers in the layer stack
+	 * @return Whether we successfully cleared the opinions or not
+	 */
+	UFUNCTION( BlueprintCallable, Category = "Conversion utils" )
+	bool RemoveUnrealSurfaceOutput( const FString& PrimPath, const FFilePath& LayerToAuthorIn );
 };
diff --git a/Engine/Plugins/Importers/USDImporter/Source/USDUtilities/Private/USDConversionUtils.cpp b/Engine/Plugins/Importers/USDImporter/Source/USDUtilities/Private/USDConversionUtils.cpp
--- a/Engine/Plugins/Importers/USDImporter/Source/USDUtilities/Private/USDConversionUtils.cpp
+++ b/Engine/Plugins/Importers/USDImporter/Source/USDUtilities/Private/USDConversionUtils.cpp
@@ -1405,13 +1405,19 @@ USDUTILITIES_API void UsdUtils::RemoveAllPrimSpecs( const UE::FUsdPrim& Prim, co
 
 	for ( const pxr::SdfPrimSpecHandle& Spec : UsdPrim.GetPrimStack() )
 	{
+		// For whatever reason sometimes there are invalid specs in the layer stack, so we need to be careful
+		if ( !Spec )
+		{
+			continue;
+		}
+
 		pxr::SdfPath SpecPath = Spec->GetPath();
 		if ( !SpecPath.IsPrimPath() )
 		{
 			continue;
 		}
 
-		if ( Spec->GetLayer() != UsdLayer )
+		if ( UsdLayer && Spec->GetLayer() != UsdLayer )
 		{
 			continue;
 		}
diff --git a/Engine/Plugins/Importers/USDImporter/Source/USDUtilities/Private/USDGeomMeshConversion.cpp b/Engine/Plugins/Importers/USDImporter/Source/USDUtilities/Private/USDGeomMeshConversion.cpp
--- a/Engine/Plugins/Importers/USDImporter/Source/USDUtilities/Private/USDGeomMeshConversion.cpp
+++ b/Engine/Plugins/Importers/USDImporter/Source/USDUtilities/Private/USDGeomMeshConversion.cpp
@@ -10,6 +10,7 @@
 #include "USDErrorUtils.h"
 #include "USDLog.h"
 #include "USDMemory.h"
+#include "USDShadeConversion.h"
 #include "USDTypesConversion.h"
 
 #include "UsdWrappers/SdfLayer.h"
@@ -1061,31 +1062,6 @@ UsdUtils::FUsdPrimMaterialAssignmentInfo UsdUtils::GetPrimMaterialAssignments( c
 		return {};
 	};
 
-	auto FetchMaterialFromUnrealRenderContextPrim = [ &RenderContext ]( const pxr::UsdPrim& UsdPrim )->TOptional<FString>
-	{
-		pxr::UsdShadeMaterialBindingAPI BindingAPI( UsdPrim );
-		pxr::UsdShadeMaterial ShadeMaterial = BindingAPI.ComputeBoundMaterial();
-		if ( !ShadeMaterial )
-		{
-			return {};
-		}
-
-		// Ignore this material if UsdToUnreal::ConvertMaterial would as well
-		pxr::UsdShadeShader SurfaceShader = ShadeMaterial.ComputeSurfaceSource( RenderContext );
-		if ( !SurfaceShader )
-		{
-			return {};
-		}
-
-		pxr::SdfAssetPath AssetPath;
-		if ( SurfaceShader.GetSourceAsset( &AssetPath, RenderContext ) )
-		{
-			return UsdToUnreal::ConvertString( AssetPath.GetAssetPath() );
-		}
-
-		return {};
-	};
-
 	auto FetchMaterialByComputingBoundMaterial = [ &RenderContext ]( const pxr::UsdPrim& UsdPrim ) -> TOptional<FString>
 	{
 		pxr::UsdShadeMaterialBindingAPI BindingAPI( UsdPrim );
@@ -1203,13 +1179,20 @@ UsdUtils::FUsdPrimMaterialAssignmentInfo UsdUtils::GetPrimMaterialAssignments( c
 	if ( RenderContext == UnrealIdentifiers::Unreal )
 	{
 		// Priority 1.1: unreal rendercontext material prim
-		if ( TOptional<FString> UnrealMaterial = FetchMaterialFromUnrealRenderContextPrim( UsdPrim ) )
+		pxr::UsdShadeMaterialBindingAPI BindingAPI( UsdPrim );
+		if ( BindingAPI )
 		{
-			FUsdPrimMaterialSlot& Slot = Result.Slots.Emplace_GetRef();
-			Slot.MaterialSource = UnrealMaterial.GetValue();
-			Slot.AssignmentType = UsdUtils::EPrimAssignmentType::UnrealMaterial;
+			if ( pxr::UsdShadeMaterial ShadeMaterial = BindingAPI.ComputeBoundMaterial() )
+			{
+				if ( TOptional<FString> UnrealMaterial = UsdUtils::GetUnrealSurfaceOutput( ShadeMaterial.GetPrim() ) )
+				{
+					FUsdPrimMaterialSlot& Slot = Result.Slots.Emplace_GetRef();
+					Slot.MaterialSource = UnrealMaterial.GetValue();
+					Slot.AssignmentType = UsdUtils::EPrimAssignmentType::UnrealMaterial;
 
-			return Result;
+					return Result;
+				}
+			}
 		}
 
 		// Priority 1.2: unrealMaterial attribute directly on the prim
@@ -1271,12 +1254,19 @@ UsdUtils::FUsdPrimMaterialAssignmentInfo UsdUtils::GetPrimMaterialAssignments( c
 				// Priority 4.1.1: Partition has an unreal rendercontext material prim binding
 				if ( !bHasAssignment )
 				{
-					if ( TOptional<FString> UnrealMaterial = FetchMaterialFromUnrealRenderContextPrim( GeomSubset.GetPrim() ) )
+					pxr::UsdShadeMaterialBindingAPI BindingAPI( GeomSubset.GetPrim() );
+					if ( BindingAPI )
 					{
-						FUsdPrimMaterialSlot& Slot = Result.Slots.Emplace_GetRef();
-						Slot.MaterialSource = UnrealMaterial.GetValue();
-						Slot.AssignmentType = UsdUtils::EPrimAssignmentType::UnrealMaterial;
-						bHasAssignment = true;
+						if ( pxr::UsdShadeMaterial ShadeMaterial = BindingAPI.ComputeBoundMaterial() )
+						{
+							if ( TOptional<FString> UnrealMaterial = UsdUtils::GetUnrealSurfaceOutput( ShadeMaterial.GetPrim() ) )
+							{
+								FUsdPrimMaterialSlot& Slot = Result.Slots.Emplace_GetRef();
+								Slot.MaterialSource = UnrealMaterial.GetValue();
+								Slot.AssignmentType = UsdUtils::EPrimAssignmentType::UnrealMaterial;
+								bHasAssignment = true;
+							}
+						}
 					}
 				}
 
@@ -1882,6 +1872,15 @@ void UsdUtils::ReplaceUnrealMaterialsWithBaked(
 		{
 			pxr::SdfPath Path = ParentPrim.GetPrimPath().AppendPath( pxr::SdfPath{ "Materials" } );
 			Prim = ScopeStage->DefinePrim( Path, UnrealToUsd::ConvertToken( TEXT( "Scope" ) ).Get() );
+
+			// Initialize our UsedPrimNames correctly, so we can guarantee we'll never have name collisions
+			if ( Prim )
+			{
+				for ( pxr::UsdPrim Child : Prim.GetFilteredChildren( pxr::UsdTraverseInstanceProxies( pxr::UsdPrimAllPrimsPredicate ) ) )
+				{
+					UsedPrimNames.Add( UsdToUnreal::ConvertToken( Child.GetName() ) );
+				}
+			}
 		}
 
 		pxr::UsdPrim Prim;
@@ -1968,22 +1967,49 @@ void UsdUtils::ReplaceUnrealMaterialsWithBaked(
 				}
 			}
 
-			std::string MaterialPathName;
-			pxr::UsdAttribute Attr = Prim.GetAttribute( UnrealIdentifiers::MaterialAssignment );
-			if ( !Attr || !Attr.Get<std::string>( &MaterialPathName ) )
+			std::string UnrealMaterialAttrAssetPath;
+			FString UnrealMaterialPrimAssetPath;
+
+			pxr::UsdAttribute UnrealMaterialAttr = Prim.GetAttribute( UnrealIdentifiers::MaterialAssignment );
+			pxr::UsdShadeMaterial UnrealMaterial;
+
+			pxr::UsdShadeMaterialBindingAPI MaterialBindingAPI{ Prim };
+			if ( MaterialBindingAPI )
+			{
+				if ( pxr::UsdShadeMaterial BoundMaterial = MaterialBindingAPI.ComputeBoundMaterial() )
+				{
+					UnrealMaterial = BoundMaterial;
+
+					TOptional<FString> ExistingUEAssetReference = UsdUtils::GetUnrealSurfaceOutput( UnrealMaterial.GetPrim() );
+					if ( ExistingUEAssetReference.IsSet() )
+					{
+						UnrealMaterialPrimAssetPath = MoveTemp( ExistingUEAssetReference.GetValue() );
+					}
+				}
+			}
+
+			if ( !UnrealMaterial && ( !UnrealMaterialAttr || !UnrealMaterialAttr.Get<std::string>( &UnrealMaterialAttrAssetPath ) ) )
 			{
 				return;
 			}
 
-			FString BakedFilename = BakedMaterials.FindRef( UsdToUnreal::ConvertString( MaterialPathName ) );
+			pxr::UsdPrim UnrealMaterialPrim = UnrealMaterial.GetPrim();
 
-			// If we have a valid unrealMaterial but just haven't baked it, leave unrealMaterials alone and abort
-			if ( MaterialPathName.size() > 0 && BakedFilename.IsEmpty() )
+			// Prioritize the Unreal material since import will do so too
+			FString UnrealMaterialAssetPath = UnrealMaterialPrimAssetPath.IsEmpty()
+				? UsdToUnreal::ConvertString( UnrealMaterialAttrAssetPath )
+				: UnrealMaterialPrimAssetPath;
+
+			FString BakedFilename = BakedMaterials.FindRef( UnrealMaterialAssetPath );
+
+			// If we have a valid UE asset but just haven't baked it, something went wrong: Just leave everything alone and abort
+			if ( !UnrealMaterialAssetPath.IsEmpty() && BakedFilename.IsEmpty() )
 			{
 				return;
 			}
 
-			pxr::SdfPath AttrPath = Attr.GetPath();
+			pxr::SdfPath UnrealMaterialAttrPath = UnrealMaterialAttr ? UnrealMaterialAttr.GetPath() : pxr::SdfPath{};
+			pxr::SdfPath UnrealMaterialPrimPath = UnrealMaterial ? UnrealMaterialPrim.GetPrimPath() : pxr::SdfPath{};
 
 			// Find out if we need to remove / author material bindings within an actual variant or outside of it, as an over.
 			// We don't do this when using payloads because our override prims aren't inside the actual LOD variants : They just
@@ -1994,25 +2020,31 @@ void UsdUtils::ReplaceUnrealMaterialsWithBaked(
 			if ( bAuthorInsideVariants )
 			{
 				pxr::UsdVariantSet& OuterVariantSetValue = OuterVariantSet.GetValue();
-
 				pxr::SdfPath VarPrimPath = OuterVariantSetValue.GetPrim().GetPath();
-				if ( AttrPath.HasPrefix( VarPrimPath ) )
+				pxr::SdfPath VarPrimPathWithVar = VarPrimPath.AppendVariantSelection( OuterVariantSetValue.GetName(), OuterVariantSetValue.GetVariantSelection() );
+
+				if ( UnrealMaterialAttrPath.HasPrefix( VarPrimPath ) )
 				{
 					// This builds a path like '/MyMesh{LOD=LOD0}LOD0.unrealMaterial',
 					// or '/MyMesh{LOD=LOD0}LOD0/Section1.unrealMaterial'.This is required because we'll query the layer
 					// for a spec path below, and this path must contain the variant selection in it, which the path returned
 					// from attr.GetPath() doesn't contain
-					pxr::SdfPath VarPrimPathWithVar = VarPrimPath.AppendVariantSelection( OuterVariantSetValue.GetName(), OuterVariantSetValue.GetVariantSelection() );
-					AttrPath = AttrPath.ReplacePrefix( VarPrimPath, VarPrimPathWithVar );
+					UnrealMaterialAttrPath = UnrealMaterialAttrPath.ReplacePrefix( VarPrimPath, VarPrimPathWithVar );
+				}
+
+				if ( UnrealMaterialPrimPath.HasPrefix( VarPrimPath ) )
+				{
+					UnrealMaterialPrimPath = UnrealMaterialPrimPath.ReplacePrefix( VarPrimPath, VarPrimPathWithVar );
 				}
 			}
 
-			// We always want to replace the actual attribute in whatever layer it was authored, and not just override with
-			// a stronger opinion.So search through all sublayers to find the ones with unrealMaterial attribute specs
+			// We always want to replace things in whatever layer they were authored, and not just override with
+			// a stronger opinion, so search through all sublayers to find the ones with the specs we are targeting
 			for ( pxr::SdfLayerHandle Layer : LayersToTraverse )
 			{
-				pxr::SdfAttributeSpecHandle AttrSpec = Layer->GetAttributeAtPath( AttrPath );
-				if ( !AttrSpec )
+				pxr::SdfAttributeSpecHandle UnrealMaterialAttrSpec = Layer->GetAttributeAtPath( UnrealMaterialAttrPath );
+				pxr::SdfPrimSpecHandle UnrealMaterialPrimSpec = Layer->GetPrimAtPath( UnrealMaterialPrimPath );
+				if ( !UnrealMaterialAttrSpec && !UnrealMaterialPrimSpec )
 				{
 					continue;
 				}
@@ -2026,12 +2058,20 @@ void UsdUtils::ReplaceUnrealMaterialsWithBaked(
 					{
 						VarContext.Emplace( OuterVariantSet.GetValue().GetVariantEditContext() );
 					}
-					Prim.RemoveProperty( UnrealIdentifiers::MaterialAssignment );
+
+					if ( UnrealMaterialAttrSpec )
+					{
+						Prim.RemoveProperty( UnrealIdentifiers::MaterialAssignment );
+					}
+
+					if ( UnrealMaterialPrimSpec )
+					{
+						UsdUtils::RemoveUnrealSurfaceOutput( UnrealMaterialPrim, UE::FSdfLayer{ Layer } );
+					}
 				}
 
-				// It was just an empty unrealMaterial attribute, so just cancel now as our baked_filename
-				// can't possibly be useful
-				if ( MaterialPathName.size() == 0 )
+				// It was just an empty UE asset path, so just cancel now as our BakedFilename can't possibly be useful
+				if ( UnrealMaterialAssetPath.IsEmpty() )
 				{
 					continue;
 				}
@@ -2042,67 +2082,117 @@ void UsdUtils::ReplaceUnrealMaterialsWithBaked(
 				{
 					pxr::UsdEditContext MatContext( StageToTraverse, pxr::SdfLayerRefPtr{ LayerToAuthorIn } );
 
-					FMaterialScopePrim* MatPrimScopePtr = nullptr;
-
-					if ( MatPrimScope.IsSet() )
+					// We are already referencing an unreal material prim: Let's just augment it with a reference to the baked
+					// material usd asset layer.
+					// Note how this will likely not be within MatPrimScope but instead will be a child of the Mesh/GeomSubset.
+					// This is fine, and in the future we'll likely exclusively do this since it will handle mesh-specific
+					// material baking much better, as it will allow even having separate bakes for each LOD
+					if ( UnrealMaterial && UnrealMaterialPrimSpec )
 					{
-						MatPrimScopePtr = &MatPrimScope.GetValue();
+						MatPrim = UnrealMaterialPrim;
+
+						bool bAlreadyHasReference = false;
+
+						// Make sure we don't reference it more than once. This shouldn't be needed since we'll only ever run into
+						// these unreal material prims once per Mesh/GeomSubset, but when creating MatScopePrims we can guarantee we
+						// add a reference only once by adding it along with the Material prim creation, so it would be nice to be able to
+						// guarantee it here as well
+						pxr::SdfReferencesProxy ReferencesProxy = UnrealMaterialPrimSpec->GetReferenceList();
+						for ( const pxr::SdfReference& UsdReference : ReferencesProxy.GetAddedOrExplicitItems() )
+						{
+							FString ReferencedFilePath = UsdToUnreal::ConvertString( UsdReference.GetAssetPath() );
+							FString LayerPath = UsdToUnreal::ConvertString( Layer->GetRealPath() );
+
+							if ( !LayerPath.IsEmpty() )
+							{
+								ReferencedFilePath = FPaths::ConvertRelativePathToFull( LayerPath, ReferencedFilePath );
+							}
+
+							if ( FPaths::IsSamePath( ReferencedFilePath, BakedFilename ) )
+							{
+								bAlreadyHasReference = true;
+								break;
+							}
+						}
+
+						if ( !bAlreadyHasReference )
+						{
+							UE::FUsdPrim UEMatPrim{ MatPrim };
+							UsdUtils::AddReference( UEMatPrim, *BakedFilename );
+						}
 					}
+					// Need a MatScopePrim authored somewhere within this layer
 					else
 					{
-						// On-demand create a *single* material scope prim for the stage, if we're not inside a variant set
-						if ( !StageMatScope.IsSet() )
+						FMaterialScopePrim* MatPrimScopePtr = nullptr;
+
+						if ( MatPrimScope.IsSet() )
 						{
-							// If a prim from a stage references another layer, USD's composition will effectively
-							// paste the default prim of the referenced layer over the referencing prim. Because of
-							// this, the subprims within the hierarchy of that default prim can't ever have
-							// relationships to other prims outside that of that same hierarchy, as those prims
-							// will not be present on the referencing stage at all. This is why we author our stage
-							// materials scope under the default prim, and not the pseudoroot
-							StageMatScope = FMaterialScopePrim{ StageToTraverse, StageToTraverse->GetDefaultPrim() };
+							MatPrimScopePtr = &MatPrimScope.GetValue();
+						}
+						else
+						{
+							// On-demand create a *single* material scope prim for the stage, if we're not inside a variant set
+							if ( !StageMatScope.IsSet() )
+							{
+								// If a prim from a stage references another layer, USD's composition will effectively
+								// paste the default prim of the referenced layer over the referencing prim. Because of
+								// this, the subprims within the hierarchy of that default prim can't ever have
+								// relationships to other prims outside that of that same hierarchy, as those prims
+								// will not be present on the referencing stage at all. This is why we author our stage
+								// materials scope under the default prim, and not the pseudoroot
+								StageMatScope = FMaterialScopePrim{ StageToTraverse, StageToTraverse->GetDefaultPrim() };
+							}
+							MatPrimScopePtr = &StageMatScope.GetValue();
 						}
-						MatPrimScopePtr = &StageMatScope.GetValue();
-					}
 
-					// This should never happen
-					if ( !ensure(MatPrimScopePtr) )
-					{
-						continue;
-					}
+						// This should never happen
+						if ( !ensure( MatPrimScopePtr ) )
+						{
+							continue;
+						}
 
-					// Fetch (or create) a material proxy prim for this material
-					if ( pxr::UsdPrim* FoundPrim = MatPrimScopePtr->BakedFileNameToMatPrim.Find( BakedFilename ) )
-					{
-						MatPrim = *FoundPrim;
-					}
-					else
-					{
-						std::string MaterialName = MaterialPathName.substr( MaterialPathName.find_last_of( "\\/" ) + 1 );
-						MaterialName = MaterialName.substr( 0, MaterialName.find_last_of( "." ) );
-						FString MatName = UsdToUnreal::ConvertString( pxr::TfMakeValidIdentifier( MaterialName ) );
-						FString MatPrimName = UsdUtils::GetUniqueName( MatName, MatPrimScopePtr->UsedPrimNames );
-						MatPrimScopePtr->UsedPrimNames.Add( MatPrimName );
-
-						MatPrim = StageToTraverse->DefinePrim(
-							MatPrimScopePtr->Prim.GetPath().AppendChild( UnrealToUsd::ConvertToken( *MatPrimName ).Get() ),
-							UnrealToUsd::ConvertToken( TEXT( "Material" ) ).Get()
-						);
-
-						UE::FUsdPrim UEMatPrim{ MatPrim };
-						UsdUtils::AddReference( UEMatPrim, *BakedFilename );
-						MatPrimScopePtr->BakedFileNameToMatPrim.Add( BakedFilename, MatPrim );
+						// We already have a material proxy prim for this UE material within MatPrimScope, so just reuse it
+						if ( pxr::UsdPrim* FoundPrim = MatPrimScopePtr->BakedFileNameToMatPrim.Find( BakedFilename ) )
+						{
+							MatPrim = *FoundPrim;
+						}
+						// Create a new material proxy prim for this UE material within MatPrimScope
+						else
+						{
+							FString MatName = FPaths::GetBaseFilename( UnrealMaterialAssetPath );
+							MatName = UsdToUnreal::ConvertString( pxr::TfMakeValidIdentifier( UnrealToUsd::ConvertString( *MatName ).Get() ) );
+							FString MatPrimName = UsdUtils::GetUniqueName( MatName, MatPrimScopePtr->UsedPrimNames );
+							MatPrimScopePtr->UsedPrimNames.Add( MatPrimName );
+
+							MatPrim = StageToTraverse->DefinePrim(
+								MatPrimScopePtr->Prim.GetPath().AppendChild( UnrealToUsd::ConvertToken( *MatPrimName ).Get() ),
+								UnrealToUsd::ConvertToken( TEXT( "Material" ) ).Get()
+							);
+
+							// We should only keep track and reuse the material proxy prims that we create within the MatPrimScope, not
+							// the ones we have appropriated from within Mesh/GeomSubset from being UnrealPrims
+							MatPrimScopePtr->BakedFileNameToMatPrim.Add( BakedFilename, MatPrim );
+
+							UE::FUsdPrim UEMatPrim{ MatPrim };
+							UsdUtils::AddReference( UEMatPrim, *BakedFilename );
+						}
 					}
 				}
 
-				// Reference the mat proxy prim as a material binding
+				// Make sure we have a binding to the material prim and the material binding API
+				if( pxr::UsdShadeMaterial MaterialToBind{ MatPrim } )
 				{
 					TOptional<pxr::UsdEditContext> VarContext;
 					if ( bAuthorInsideVariants )
 					{
 						VarContext.Emplace( OuterVariantSet.GetValue().GetVariantEditContext() );
 					}
-					pxr::UsdShadeMaterialBindingAPI ShadeAPI{ Prim };
-					ShadeAPI.Bind( pxr::UsdShadeMaterial{ MatPrim } );
+
+					if ( pxr::UsdShadeMaterialBindingAPI AppliedMaterialBindingAPI = pxr::UsdShadeMaterialBindingAPI::Apply( Prim ) )
+					{
+						AppliedMaterialBindingAPI.Bind( MaterialToBind );
+					}
 				}
 			}
 		};
diff --git a/Engine/Plugins/Importers/USDImporter/Source/USDUtilities/Private/USDPrimConversion.cpp b/Engine/Plugins/Importers/USDImporter/Source/USDUtilities/Private/USDPrimConversion.cpp
--- a/Engine/Plugins/Importers/USDImporter/Source/USDUtilities/Private/USDPrimConversion.cpp
+++ b/Engine/Plugins/Importers/USDImporter/Source/USDUtilities/Private/USDPrimConversion.cpp
@@ -7,8 +7,8 @@
 #include "USDLayerUtils.h"
 #include "USDLightConversion.h"
 #include "USDLog.h"
+#include "USDShadeConversion.h"
 #include "USDSkeletalDataConversion.h"
-#include "USDLog.h"
 #include "USDTypesConversion.h"
 
 #include "UsdWrappers/UsdAttribute.h"
@@ -63,6 +63,7 @@
 #include "pxr/usd/usdGeom/imageable.h"
 #include "pxr/usd/usdGeom/mesh.h"
 #include "pxr/usd/usdGeom/pointInstancer.h"
+#include "pxr/usd/usdGeom/scope.h"
 #include "pxr/usd/usdGeom/tokens.h"
 #include "pxr/usd/usdGeom/xform.h"
 #include "pxr/usd/usdGeom/xformable.h"
@@ -76,8 +77,8 @@
 #include "pxr/usd/usdLux/tokens.h"
 #include "pxr/usd/usdShade/connectableAPI.h"
 #include "pxr/usd/usdShade/material.h"
-#include "pxr/usd/usdShade/shader.h"
 #include "pxr/usd/usdShade/materialBindingAPI.h"
+#include "pxr/usd/usdShade/shader.h"
 #include "pxr/usd/usdShade/tokens.h"
 #include "pxr/usd/usdSkel/animation.h"
 #include "pxr/usd/usdSkel/root.h"
@@ -94,36 +95,50 @@ namespace UE
 			// it already has an 'unreal' render context output and the expected structure, or by creating a new Material prim
 			// that fulfills those requirements.
 			// Doesn't write to 'unrealMaterial' at all, as we intend on deprecating it in the future.
-			void AuthorMaterialOverride( pxr::UsdPrim& MeshPrim, const pxr::SdfAssetPath& UEMaterialAssetPath )
+			void AuthorMaterialOverride( pxr::UsdPrim& MeshPrim, const FString& UEMaterialAssetPath )
 			{
 				if ( !MeshPrim )
 				{
 					return;
 				}
 
+				FScopedUsdAllocs UsdAllocs;
+
 				pxr::UsdShadeMaterialBindingAPI BindingAPI = pxr::UsdShadeMaterialBindingAPI::Apply( MeshPrim );
 
 				// If this mesh prim already has a binding to a *child* material with the 'unreal' render context,
 				// just write our material there and early out
 				if ( pxr::UsdShadeMaterial ShadeMaterial = BindingAPI.ComputeBoundMaterial() )
 				{
-					// We need to try reusing these materials or else we'd write a new material prim every time we change
-					// the override in UE, but we also run the risk of modifying a material that is used by multiple prims
-					// (and here we just want to set the override for this Mesh prim). The compromise is to only reuse the
-					// material if it is a child of MeshPrim already, and always to author our material prims as children
-					std::string MaterialPath = ShadeMaterial.GetPrim().GetPath().GetString();
-					std::string MeshPrimPath = MeshPrim.GetPath().GetString();
-					if ( MaterialPath.rfind( MeshPrimPath, 0 ) == 0 )
+					// Only consider this material reusable if its within the same layer as the edit target, otherwise we'll
+					// prefer to author something else that can be fully defined on the edit target. This is handy when we're
+					// exporting a level, as we'll ensure we're making these prims on the MaterialOverrides layer
+					pxr::SdfLayerRefPtr EditTarget = MeshPrim.GetStage()->GetEditTarget().GetLayer();
+					pxr::SdfLayerRefPtr MaterialPrimLayer = static_cast< pxr::SdfLayerRefPtr >( UsdUtils::FindLayerForPrim( ShadeMaterial.GetPrim() ) );
+					if ( EditTarget == MaterialPrimLayer )
 					{
-						if ( pxr::UsdShadeShader ExistingShader = ShadeMaterial.ComputeSurfaceSource( UnrealIdentifiers::Unreal ) )
+						// We need to try reusing these materials or else we'd write a new material prim every time we change
+						// the override in UE, but we also run the risk of modifying a material that is used by multiple prims
+						// (and here we just want to set the override for this Mesh prim). The compromise is to only reuse the
+						// material if it is a child of MeshPrim already, and always to author our material prims as children
+						std::string MaterialPath = ShadeMaterial.GetPrim().GetPath().GetString();
+						std::string MeshPrimPath = MeshPrim.GetPath().GetString();
+						if ( MaterialPath.rfind( MeshPrimPath, 0 ) == 0 )
 						{
-							ExistingShader.SetSourceAsset( UEMaterialAssetPath, UnrealIdentifiers::Unreal );
-							return;
+							if ( pxr::UsdPrim MaterialPrim = ShadeMaterial.GetPrim() )
+							{
+								UsdUtils::SetUnrealSurfaceOutput( MaterialPrim, UEMaterialAssetPath );
+								return;
+							}
 						}
 					}
 				}
 
 				// Find a unique name for our child material prim
+				// Note how we'll always author these materials as children of the meshes themselves instead of emitting a common
+				// Material prim to use for multiple overrides: This because in the future we'll want to have a separate material
+				// bake for each mesh (to make sure we get vertex color effects, etc.), and so we'd have multiple baked .usda material
+				// asset layers for each UE material, and we'd want each mesh/section/LOD to refer to its own anyway
 				FString ChildMaterialName = TEXT( "UnrealMaterial" );
 				if ( pxr::UsdPrim ExistingPrim = MeshPrim.GetChild( UnrealToUsd::ConvertToken( *ChildMaterialName ).Get() ) )
 				{
@@ -142,38 +157,24 @@ namespace UE
 				pxr::UsdStageRefPtr Stage = MeshPrim.GetStage();
 				pxr::SdfPath MeshPath = MeshPrim.GetPath();
 				pxr::SdfPath MaterialPath = MeshPath.AppendChild( UnrealToUsd::ConvertToken( *ChildMaterialName ).Get() );
-				pxr::SdfPath ShaderPath = MaterialPath.AppendChild( UnrealToUsd::ConvertToken( TEXT( "UnrealShader" ) ).Get() );
 
 				pxr::UsdShadeMaterial ChildMaterial = pxr::UsdShadeMaterial::Define( Stage, MaterialPath );
 				if ( !ChildMaterial )
 				{
-					UE_LOG(LogUsd, Warning, TEXT("Failed to author material prim '%s' when trying to write '%s's material override '%s' to USD"),
+					UE_LOG( LogUsd, Warning, TEXT( "Failed to author material prim '%s' when trying to write '%s's material override '%s' to USD" ),
 						*UsdToUnreal::ConvertPath( MaterialPath ),
 						*UsdToUnreal::ConvertPath( MeshPrim.GetPath() ),
-						*UsdToUnreal::ConvertString( UEMaterialAssetPath.GetAssetPath() )
+						*UEMaterialAssetPath
 					);
 					return;
 				}
 
-				pxr::UsdShadeShader UnrealShader = pxr::UsdShadeShader::Define( Stage, ShaderPath );
-				if ( !UnrealShader )
+				if ( pxr::UsdPrim MaterialPrim = ChildMaterial.GetPrim() )
 				{
-					UE_LOG( LogUsd, Warning, TEXT( "Failed to author shader prim '%s' when trying to write '%s's material override '%s' to USD" ),
-						*UsdToUnreal::ConvertPath( ShaderPath ),
-						*UsdToUnreal::ConvertPath( MeshPrim.GetPath() ),
-						*UsdToUnreal::ConvertString( UEMaterialAssetPath.GetAssetPath() )
-					);
-					return;
-				}
-
-				UnrealShader.CreateImplementationSourceAttr( pxr::VtValue{ pxr::UsdShadeTokens->sourceAsset } );
-				UnrealShader.SetSourceAsset( UEMaterialAssetPath, UnrealIdentifiers::Unreal );
-				pxr::UsdShadeOutput ShaderOutput = UnrealShader.CreateOutput( UnrealToUsd::ConvertToken( TEXT( "out" ) ).Get(), pxr::SdfValueTypeNames->Token );
-
-				pxr::UsdShadeOutput MaterialOutput = ChildMaterial.CreateSurfaceOutput( UnrealIdentifiers::Unreal );
-				pxr::UsdShadeConnectableAPI::ConnectToSource( MaterialOutput, ShaderOutput );
+					UsdUtils::SetUnrealSurfaceOutput( MaterialPrim, UEMaterialAssetPath );
 
-				BindingAPI.Bind( ChildMaterial );
+					BindingAPI.Bind( ChildMaterial );
+				}
 			}
 		}
 	}
@@ -1950,8 +1951,7 @@ bool UnrealToUsd::ConvertMeshComponent( const pxr::UsdStageRefPtr& Stage, const
 			pxr::SdfPath OverridePrimPath = UsdPrim.GetPath();
 
 			pxr::UsdPrim MeshPrim = Stage->OverridePrim( OverridePrimPath );
-			std::string UEMaterialOverridePath = UnrealToUsd::ConvertString( *Override->GetPathName() ).Get();
-			UE::USDPrimConversionImpl::Private::AuthorMaterialOverride( MeshPrim, pxr::SdfAssetPath{ UEMaterialOverridePath } );
+			UE::USDPrimConversionImpl::Private::AuthorMaterialOverride( MeshPrim, Override->GetPathName() );
 		}
 	}
 	else if ( const UStaticMeshComponent* StaticMeshComponent = Cast<const UStaticMeshComponent>( MeshComponent ) )
@@ -1971,8 +1971,6 @@ bool UnrealToUsd::ConvertMeshComponent( const pxr::UsdStageRefPtr& Stage, const
 					continue;
 				}
 
-				std::string UEMaterialOverridePath = UnrealToUsd::ConvertString( *Override->GetPathName() ).Get();
-
 				for ( int32 LODIndex = 0; LODIndex < NumLODs; ++LODIndex )
 				{
 					int32 NumSections = Mesh->GetNumSections( LODIndex );
@@ -2014,7 +2012,7 @@ bool UnrealToUsd::ConvertMeshComponent( const pxr::UsdStageRefPtr& Stage, const
 						}
 
 						pxr::UsdPrim MeshPrim = Stage->OverridePrim( OverridePrimPath );
-						UE::USDPrimConversionImpl::Private::AuthorMaterialOverride( MeshPrim, pxr::SdfAssetPath{ UEMaterialOverridePath } );
+						UE::USDPrimConversionImpl::Private::AuthorMaterialOverride( MeshPrim, Override->GetPathName() );
 					}
 				}
 			}
@@ -2067,8 +2065,6 @@ bool UnrealToUsd::ConvertMeshComponent( const pxr::UsdStageRefPtr& Stage, const
 					continue;
 				}
 
-				std::string UEMaterialOverridePath = UnrealToUsd::ConvertString( *Override->GetPathName() ).Get();
-
 				for ( int32 LODIndex = 0; LODIndex < NumLODs; ++LODIndex )
 				{
 					if ( !LodRenderData.IsValidIndex( LODIndex ) )
@@ -2129,7 +2125,7 @@ bool UnrealToUsd::ConvertMeshComponent( const pxr::UsdStageRefPtr& Stage, const
 						}
 
 						pxr::UsdPrim MeshPrim = Stage->OverridePrim( OverridePrimPath );
-						UE::USDPrimConversionImpl::Private::AuthorMaterialOverride( MeshPrim, pxr::SdfAssetPath{ UEMaterialOverridePath } );
+						UE::USDPrimConversionImpl::Private::AuthorMaterialOverride( MeshPrim, Override->GetPathName() );
 					}
 				}
 			}
diff --git a/Engine/Plugins/Importers/USDImporter/Source/USDUtilities/Private/USDShadeConversion.cpp b/Engine/Plugins/Importers/USDImporter/Source/USDUtilities/Private/USDShadeConversion.cpp
--- a/Engine/Plugins/Importers/USDImporter/Source/USDUtilities/Private/USDShadeConversion.cpp
+++ b/Engine/Plugins/Importers/USDImporter/Source/USDUtilities/Private/USDShadeConversion.cpp
@@ -2302,6 +2302,102 @@ bool UsdUtils::MarkMaterialPrimWithWorldSpaceNormals( const UE::FUsdPrim& Materi
 	return true;
 }
 
+TOptional<FString> UsdUtils::GetUnrealSurfaceOutput( const pxr::UsdPrim& MaterialPrim )
+{
+	if ( !MaterialPrim )
+	{
+		return {};
+	}
+
+	FScopedUsdAllocs UsdAllocs;
+
+	pxr::UsdShadeMaterial ShadeMaterial{ MaterialPrim };
+	if ( !ShadeMaterial )
+	{
+		return {};
+	}
+
+	pxr::UsdShadeShader SurfaceShader = ShadeMaterial.ComputeSurfaceSource( UnrealIdentifiers::Unreal );
+	if ( !SurfaceShader )
+	{
+		return {};
+	}
+
+	pxr::SdfAssetPath AssetPath;
+	if ( SurfaceShader.GetSourceAsset( &AssetPath, UnrealIdentifiers::Unreal ) )
+	{
+		return UsdToUnreal::ConvertString( AssetPath.GetAssetPath() );
+	}
+
+	return {};
+}
+
+bool UsdUtils::SetUnrealSurfaceOutput( pxr::UsdPrim& MaterialPrim, const FString& UnrealMaterialPathName )
+{
+	if ( !MaterialPrim )
+	{
+		return false;
+	}
+
+	FScopedUsdAllocs UsdAllocs;
+
+	pxr::UsdShadeMaterial Material{ MaterialPrim };
+	if ( !Material )
+	{
+		return false;
+	}
+
+	pxr::UsdStageRefPtr Stage = MaterialPrim.GetStage();
+	pxr::SdfPath ShaderPath = MaterialPrim.GetPath().AppendChild( UnrealToUsd::ConvertToken( TEXT( "UnrealShader" ) ).Get() );
+
+	pxr::UsdShadeShader UnrealShader = pxr::UsdShadeShader::Define( Stage, ShaderPath );
+	if ( !UnrealShader )
+	{
+		return false;
+	}
+
+	UnrealShader.CreateImplementationSourceAttr( pxr::VtValue{ pxr::UsdShadeTokens->sourceAsset } );
+	UnrealShader.SetSourceAsset(
+		UnrealMaterialPathName.IsEmpty() ? pxr::SdfAssetPath{} : pxr::SdfAssetPath{ UnrealToUsd::ConvertString( *UnrealMaterialPathName ).Get() },
+		UnrealIdentifiers::Unreal
+	);
+	pxr::UsdShadeOutput ShaderOutput = UnrealShader.CreateOutput( UnrealToUsd::ConvertToken( TEXT( "out" ) ).Get(), pxr::SdfValueTypeNames->Token );
+
+	pxr::UsdShadeOutput MaterialOutput = Material.CreateSurfaceOutput( UnrealIdentifiers::Unreal );
+	pxr::UsdShadeConnectableAPI::ConnectToSource( MaterialOutput, ShaderOutput );
+
+	return true;
+}
+
+bool UsdUtils::RemoveUnrealSurfaceOutput( pxr::UsdPrim& MaterialPrim, const UE::FSdfLayer& LayerToAuthorIn )
+{
+	FScopedUsdAllocs UsdAllocs;
+
+	pxr::UsdShadeMaterial ShadeMaterial{ MaterialPrim };
+	pxr::UsdShadeConnectableAPI Connectable{ MaterialPrim };
+	if ( !ShadeMaterial || !Connectable )
+	{
+		return false;
+	}
+
+	if ( pxr::UsdShadeOutput MaterialOutput = ShadeMaterial.GetSurfaceOutput( UnrealIdentifiers::Unreal ) )
+	{
+		if ( pxr::UsdShadeShader SurfaceShader = ShadeMaterial.ComputeSurfaceSource( UnrealIdentifiers::Unreal ) )
+		{
+			// Fully remove the UnrealShader
+			UsdUtils::RemoveAllPrimSpecs( UE::FUsdPrim{ SurfaceShader.GetPrim() }, LayerToAuthorIn );
+		}
+
+		// Disconnect would author something like `token outputs:unreal:surface.connect = None`, which is not quite what we want:
+		// That would be an opinion to have it connected to nothing, but instead we just want to remove any opinion whatsoever,
+		// which is what ClearSource does. Note that these will still leave behind `token outputs:unreal:surface` lines,
+		// but those don't actually count as opinions apparently
+		pxr::UsdShadeConnectableAPI::ClearSource( MaterialOutput );
+	}
+
+	return true;
+}
+
 bool UsdUtils::IsMaterialTranslucent( const pxr::UsdShadeMaterial& UsdShadeMaterial )
 {
 	FScopedUsdAllocs UsdAllocs;
diff --git a/Engine/Plugins/Importers/USDImporter/Source/USDUtilities/Public/USDShadeConversion.h b/Engine/Plugins/Importers/USDImporter/Source/USDUtilities/Public/USDShadeConversion.h
--- a/Engine/Plugins/Importers/USDImporter/Source/USDUtilities/Public/USDShadeConversion.h
+++ b/Engine/Plugins/Importers/USDImporter/Source/USDUtilities/Public/USDShadeConversion.h
@@ -9,6 +9,8 @@
 #include "Materials/MaterialInterface.h"
 #include "UObject/WeakObjectPtr.h"
 
+#include "UsdWrappers/SdfLayer.h"
+
 #include "USDIncludesStart.h"
 	#include "pxr/usd/usd/timeCode.h"
 #include "USDIncludesEnd.h"
@@ -93,6 +95,29 @@ namespace UnrealToUsd
 
 namespace UsdUtils
 {
+	/** Returns a path to an UE asset (e.g. "/Game/Assets/Red.Red") if MaterialPrim has an 'unreal' render context surface output that points at one */
+	USDUTILITIES_API TOptional<FString> GetUnrealSurfaceOutput( const pxr::UsdPrim& MaterialPrim );
+
+	/**
+	 * Sets which UE material asset the 'unreal' render context surface output of MaterialPrim is pointing at (creating the surface output
+	 * on-demand if needed)
+	 *
+	 * @param MaterialPrim - pxr::UsdPrim with the pxr::UsdShadeMaterial schema to update the 'unreal' surface output of
+	 * @param UnrealMaterialPathName - Path to an UE UMaterialInterface asset (e.g. "/Game/Assets/Red.Red")
+	 * @return Whether we successfully set the surface output or not
+	 */
+	USDUTILITIES_API bool SetUnrealSurfaceOutput( pxr::UsdPrim& MaterialPrim, const FString& UnrealMaterialPathName );
+
+	/**
+	 * Clears any opinions for the 'unreal' render context surface output of MaterialPrim within LayerToAuthorIn.
+	 * If LayerToAuthorIn is an invalid layer (the default) it will clear opinions from all layers of the stage's layer stack.
+	 *
+	 * @param MaterialPrim - pxr::UsdPrim with the pxr::UsdShadeMaterial schema to update the 'unreal' surface output of
+	 * @param LayerToAuthorIn - Layer to clear the opinions in, or an invalid layer (e.g. UE::FSdfLayer{}, which is the default)
+	 * @return Whether we successfully cleared the opinions or not
+	 */
+	USDUTILITIES_API bool RemoveUnrealSurfaceOutput( pxr::UsdPrim& MaterialPrim, const UE::FSdfLayer& LayerToAuthorIn = UE::FSdfLayer{} );
+
 	/**
 	 * Returns whether the material needs to be rendered with the Translucent rendering mode.
 	 * This function exists because we need this information *before* we pick the right parent for a material instance and properly convert it.
