USD: Fix handling of reference/sublayer relative paths when the target USD file is within the project directory

diff --git a/Engine/Plugins/Importers/USDImporter/Source/USDStage/Private/USDStageActor.cpp b/Engine/Plugins/Importers/USDImporter/Source/USDStage/Private/USDStageActor.cpp
--- a/Engine/Plugins/Importers/USDImporter/Source/USDStage/Private/USDStageActor.cpp
+++ b/Engine/Plugins/Importers/USDImporter/Source/USDStage/Private/USDStageActor.cpp
@@ -1397,7 +1397,7 @@ void AUsdStageActor::OpenUsdStage()
 	FString AbsPath;
 	if ( !RootLayer.FilePath.StartsWith( UnrealIdentifiers::IdentifierPrefix ) && FPaths::IsRelative( RootLayer.FilePath ) )
 	{
-		// The RootLayer property is marked as RelativeToGameDir, and UsdUtils::BrowseUsdFile will also emit paths relative to the project's directory
+		// The RootLayer property is marked as RelativeToGameDir, and FUsdStageViewModel::OpenStage will also give us paths relative to the project's directory
 		FString ProjectDir = FPaths::ConvertRelativePathToFull( FPaths::ProjectDir() );
 		AbsPath = FPaths::ConvertRelativePathToFull( FPaths::Combine( ProjectDir, RootLayer.FilePath ) );
 	}
diff --git a/Engine/Plugins/Importers/USDImporter/Source/USDStageEditor/Private/SUSDStageTreeView.cpp b/Engine/Plugins/Importers/USDImporter/Source/USDStageEditor/Private/SUSDStageTreeView.cpp
--- a/Engine/Plugins/Importers/USDImporter/Source/USDStageEditor/Private/SUSDStageTreeView.cpp
+++ b/Engine/Plugins/Importers/USDImporter/Source/USDStageEditor/Private/SUSDStageTreeView.cpp
@@ -756,13 +756,11 @@ void SUsdStageTreeView::OnAddReference()
 		FText::FromString( PickedFile.GetValue() )
 	) );
 
-	const FString AbsoluteFilePath = FPaths::ConvertRelativePathToFull( PickedFile.GetValue() );
-
 	TArray< FUsdPrimViewModelRef > MySelectedItems = GetSelectedItems();
 
 	for ( FUsdPrimViewModelRef SelectedItem : MySelectedItems )
 	{
-		SelectedItem->AddReference( *AbsoluteFilePath );
+		SelectedItem->AddReference( *PickedFile.GetValue() );
 	}
 }
 
diff --git a/Engine/Plugins/Importers/USDImporter/Source/USDUtilities/Private/USDLayerUtils.cpp b/Engine/Plugins/Importers/USDImporter/Source/USDUtilities/Private/USDLayerUtils.cpp
--- a/Engine/Plugins/Importers/USDImporter/Source/USDUtilities/Private/USDLayerUtils.cpp
+++ b/Engine/Plugins/Importers/USDImporter/Source/USDUtilities/Private/USDLayerUtils.cpp
@@ -146,7 +146,8 @@ TOptional< FString > UsdUtils::BrowseUsdFile( EBrowseFileMode Mode, TSharedRef<
 
 	if ( OutFiles.Num() > 0 )
 	{
-		return MakePathRelativeToProjectDir( OutFiles[0] );
+		// Always make this an absolute path because it may try generating a relative path to the engine binary if it can
+		return FPaths::ConvertRelativePathToFull( OutFiles[ 0 ] );
 	}
 
 	return {};
diff --git a/Engine/Plugins/Importers/USDImporter/Source/USDUtilities/Public/USDLayerUtils.h b/Engine/Plugins/Importers/USDImporter/Source/USDUtilities/Public/USDLayerUtils.h
--- a/Engine/Plugins/Importers/USDImporter/Source/USDUtilities/Public/USDLayerUtils.h
+++ b/Engine/Plugins/Importers/USDImporter/Source/USDUtilities/Public/USDLayerUtils.h
@@ -52,7 +52,7 @@ namespace UsdUtils
 	USDUTILITIES_API bool InsertSubLayer( const pxr::SdfLayerRefPtr& ParentLayer, const TCHAR* SubLayerFile, int32 Index = -1, double OffsetTimeCodes = 0.0, double TimeCodesScale = 1.0 );
 
 #if WITH_EDITOR
-	/** Opens a file dialog to open or save a USD file */
+	/** Opens a file dialog to open or save a USD file. The returned file path will always be absolute */
 	USDUTILITIES_API TOptional< FString > BrowseUsdFile( EBrowseFileMode Mode, TSharedRef< const SWidget > OriginatingWidget );
 #endif // #if WITH_EDITOR
 
