USD: Fix failing to import temporary stages.

Fix import process closing stage before sending analytics about it;
Fix opened stage resetting its changes if importing the same stage into the level.

diff --git a/Engine/Plugins/Importers/USDImporter/Source/USDStageEditorViewModels/Private/USDStageViewModel.cpp b/Engine/Plugins/Importers/USDImporter/Source/USDStageEditorViewModels/Private/USDStageViewModel.cpp
--- a/Engine/Plugins/Importers/USDImporter/Source/USDStageEditorViewModels/Private/USDStageViewModel.cpp
+++ b/Engine/Plugins/Importers/USDImporter/Source/USDStageEditorViewModels/Private/USDStageViewModel.cpp
@@ -274,7 +274,13 @@ void FUsdStageViewModel::ImportStage()
 		ImportContext.bReadFromStageCache = true; // So that we import whatever the user has open right now, even if the file has changes
 
 		const FString RootPath = UsdStage.GetRootLayer().GetRealPath();
-		const FString StageName = FPaths::GetBaseFilename( RootPath );
+		FString StageName = FPaths::GetBaseFilename( RootPath );
+
+		// Provide a StageName when importing transient stages as this is used for the content folder name and actor label
+		if ( UsdStage.GetRootLayer().IsAnonymous() && RootPath.IsEmpty() )
+		{
+			StageName = TEXT("TransientStage");
+		}
 
 		const bool bIsAutomated = false;
 		if ( ImportContext.Init( StageName, RootPath, TEXT("/Game/"), RF_Public | RF_Transactional, bIsAutomated ) )
@@ -290,6 +296,9 @@ void FUsdStageViewModel::ImportStage()
 			ImportContext.TargetSceneActorAttachParent = StageActor->GetRootComponent()->GetAttachParent();
 			ImportContext.TargetSceneActorTargetTransform = StageActor->GetActorTransform();
 
+			// Pass the stage directly too in case we're importing a transient stage with no filepath
+			ImportContext.Stage = UsdStage;
+
 			UUsdStageImporter* USDImporter = IUsdStageImporterModule::Get().GetImporter();
 			USDImporter->ImportFromFile(ImportContext);
 
diff --git a/Engine/Plugins/Importers/USDImporter/Source/USDStageImporter/Private/USDStageImportContext.cpp b/Engine/Plugins/Importers/USDImporter/Source/USDStageImporter/Private/USDStageImportContext.cpp
--- a/Engine/Plugins/Importers/USDImporter/Source/USDStageImporter/Private/USDStageImportContext.cpp
+++ b/Engine/Plugins/Importers/USDImporter/Source/USDStageImporter/Private/USDStageImportContext.cpp
@@ -19,8 +19,16 @@ FUsdStageImportContext::FUsdStageImportContext()
 {
 	World = nullptr;
 	ImportOptions = NewObject< UUsdStageImportOptions >();
-	bReadFromStageCache = false;
-	bStageWasOriginallyOpen = false;
+
+	// Always reading from the stage cache is a good default because while we can have multiple instances of the
+	// same stage open, USD will only open a particular layer once. If we try importing without using the stage
+	// cache and the stage we want to import uses an existing open layer, we will forcibly reload
+	// that layer (check UnrealUSDWrapper::OpenStage), which would erase our previous changes to it and lead to modifications
+	// on existing open stages (e.g. we have cube.usda open with a stage actor and with some changes and we click
+	// file -> import into level and import cube.usda again)
+	bReadFromStageCache = true;
+
+	bStageWasOriginallyOpenInCache = false;
 	SceneActor = nullptr;
 	ImportedAsset = nullptr;
 	AssetCache = nullptr;
diff --git a/Engine/Plugins/Importers/USDImporter/Source/USDStageImporter/Private/USDStageImporter.cpp b/Engine/Plugins/Importers/USDImporter/Source/USDStageImporter/Private/USDStageImporter.cpp
--- a/Engine/Plugins/Importers/USDImporter/Source/USDStageImporter/Private/USDStageImporter.cpp
+++ b/Engine/Plugins/Importers/USDImporter/Source/USDStageImporter/Private/USDStageImporter.cpp
@@ -52,7 +52,7 @@
 
 namespace UsdStageImporterImpl
 {
-	UE::FUsdStage ReadUsdFile(FUsdStageImportContext& ImportContext)
+	void LoadStageFromFilePath(FUsdStageImportContext& ImportContext)
 	{
 		const FString FilePath = IFileManager::Get().ConvertToAbsolutePathForExternalAppForRead(*ImportContext.FilePath);
 
@@ -64,7 +64,7 @@ namespace UsdStageImporterImpl
 			FPaths::NormalizeFilename( RootPath );
 			if ( ImportContext.FilePath == RootPath )
 			{
-				ImportContext.bStageWasOriginallyOpen = true;
+				ImportContext.bStageWasOriginallyOpenInCache = true;
 				break;
 			}
 		}
@@ -74,11 +74,15 @@ namespace UsdStageImporterImpl
 		TArray<FString> ErrorStrings = UsdUtils::GetErrorsAndStopMonitoring();
 		FString Error = FString::Join(ErrorStrings, TEXT("\n"));
 
-		if (!Error.IsEmpty())
+		if (Error.IsEmpty())
 		{
+			ImportContext.Stage = Stage;
+		}
+		else
+		{
+			ImportContext.Stage = UE::FUsdStage();
 			FUsdLogManager::LogMessage( EMessageSeverity::Error, FText::Format( LOCTEXT( "CouldNotImportUSDFile", "Could not import USD file {0}\n {1}" ), FText::FromString( FilePath ), FText::FromString( Error ) ) );
 		}
-		return Stage;
 	}
 
 	FString FindValidPackagePath(const FString& InPackagePath)
@@ -221,7 +225,7 @@ namespace UsdStageImporterImpl
 	void SetupStageForImport( FUsdStageImportContext& ImportContext )
 	{
 #if USE_USD_SDK
-		if ( ImportContext.ImportOptions->bOverrideStageOptions && ImportContext.bStageWasOriginallyOpen )
+		if ( ImportContext.ImportOptions->bOverrideStageOptions )
 		{
 			ImportContext.OriginalMetersPerUnit = UsdUtils::GetUsdStageMetersPerUnit( ImportContext.Stage );
 			ImportContext.OriginalUpAxis = UsdUtils::GetUsdStageUpAxisAsEnum( ImportContext.Stage );
@@ -1030,17 +1034,20 @@ namespace UsdStageImporterImpl
 	{
 #if USE_USD_SDK
 		// Remove our imported stage from the stage cache if it wasn't in there to begin with
-		if (!ImportContext.bStageWasOriginallyOpen && ImportContext.bReadFromStageCache)
+		if (!ImportContext.bStageWasOriginallyOpenInCache && ImportContext.bReadFromStageCache)
 		{
 			UnrealUSDWrapper::EraseStageFromCache(ImportContext.Stage);
 		}
 
-		// Restore original meters per unit if the stage was already loaded
-		if ( ImportContext.ImportOptions->bOverrideStageOptions && ImportContext.bStageWasOriginallyOpen )
+		if ( ImportContext.ImportOptions->bOverrideStageOptions )
 		{
 			UsdUtils::SetUsdStageMetersPerUnit( ImportContext.Stage, ImportContext.OriginalMetersPerUnit );
 			UsdUtils::SetUsdStageUpAxis( ImportContext.Stage, ImportContext.OriginalUpAxis );
 		}
+
+		// Always discard the context's reference to the stage because it may be a persistent import context (like
+		// the non-static data member of UUsdStageImportFactory
+		ImportContext.Stage = UE::FUsdStage();
 #endif // #if USE_USD_SDK
 	}
 
@@ -1317,7 +1324,18 @@ namespace UsdStageImporterImpl
 			EventAttributes.Emplace( TEXT( "NumTextures" ), LexToString( NumTextures ) );
 			EventAttributes.Emplace( TEXT( "NumGeometryCaches" ), LexToString( NumGeometryCaches ) );
 
-			double NumberOfFrames = ImportContext.Stage.GetEndTimeCode() - ImportContext.Stage.GetStartTimeCode();
+			FString RootLayerIdentifier = ImportContext.FilePath;
+			double NumberOfFrames = 0;
+
+			if ( ImportContext.Stage )
+			{
+				NumberOfFrames = ImportContext.Stage.GetEndTimeCode() - ImportContext.Stage.GetStartTimeCode();
+
+				if ( RootLayerIdentifier.IsEmpty() )
+				{
+					RootLayerIdentifier = ImportContext.Stage.GetRootLayer().GetIdentifier();
+				}
+			}
 
 			IUsdClassesModule::SendAnalytics(
 				MoveTemp( EventAttributes ),
@@ -1325,7 +1343,7 @@ namespace UsdStageImporterImpl
 				ImportContext.bIsAutomated,
 				ElapsedSeconds,
 				NumberOfFrames,
-				FPaths::GetExtension( ImportContext.FilePath )
+				FPaths::GetExtension( RootLayerIdentifier )
 			);
 		}
 	}
@@ -1342,8 +1360,12 @@ void UUsdStageImporter::ImportFromFile(FUsdStageImportContext& ImportContext)
 
 	double StartTime = FPlatformTime::Cycles64();
 
-	ImportContext.Stage = UsdStageImporterImpl::ReadUsdFile(ImportContext);
-	if (!ImportContext.Stage)
+	if ( !ImportContext.Stage && !ImportContext.FilePath.IsEmpty() )
+	{
+		UsdStageImporterImpl::LoadStageFromFilePath( ImportContext );
+	}
+
+	if ( !ImportContext.Stage )
 	{
 		FUsdLogManager::LogMessage( EMessageSeverity::Error, LOCTEXT( "NoStageError", "Failed to open the USD Stage!" ) );
 		return;
@@ -1407,7 +1429,6 @@ void UUsdStageImporter::ImportFromFile(FUsdStageImportContext& ImportContext)
 	UsdStageImporterImpl::RemapSoftReferences( ImportContext, UsedAssetsAndDependencies, SoftObjectsToRemap );
 	UsdStageImporterImpl::Cleanup( ImportContext.SceneActor, ExistingSceneActor, ImportContext.ImportOptions->ExistingActorPolicy );
 	UsdStageImporterImpl::NotifyAssetRegistry( UsedAssetsAndDependencies );
-	UsdStageImporterImpl::CloseStageIfNeeded( ImportContext );
 
 	FUsdDelegates::OnPostUsdImport.Broadcast( ImportContext.FilePath );
 
@@ -1417,6 +1438,8 @@ void UUsdStageImporter::ImportFromFile(FUsdStageImportContext& ImportContext)
 		UsdStageImporterImpl::SendAnalytics( ImportContext, nullptr, TEXT("Import"), UsedAssetsAndDependencies, ElapsedSeconds);
 	}
 
+	UsdStageImporterImpl::CloseStageIfNeeded( ImportContext );
+
 #endif // #if USE_USD_SDK
 }
 
@@ -1428,14 +1451,17 @@ bool UUsdStageImporter::ReimportSingleAsset(FUsdStageImportContext& ImportContex
 #if USE_USD_SDK
 	double StartTime = FPlatformTime::Cycles64();
 
-	ImportContext.Stage = UsdStageImporterImpl::ReadUsdFile(ImportContext);
-	if (!ImportContext.Stage)
+	if ( !ImportContext.Stage && !ImportContext.FilePath.IsEmpty() )
+	{
+		UsdStageImporterImpl::LoadStageFromFilePath( ImportContext );
+	}
+
+	if ( !ImportContext.Stage )
 	{
 		FUsdLogManager::LogMessage( EMessageSeverity::Error, LOCTEXT( "NoStageError", "Failed to open the USD Stage!" ) );
 		return bSuccess;
 	}
 
-
 	FUsdDelegates::OnPreUsdImport.Broadcast(ImportContext.FilePath);
 
 	// We still need the scene actor to remap all other users of the mesh to the new reimported one. It's not critical if we fail though,
@@ -1512,7 +1538,6 @@ bool UUsdStageImporter::ReimportSingleAsset(FUsdStageImportContext& ImportContex
 
 	UsdStageImporterImpl::Cleanup( ImportContext.SceneActor, nullptr, ImportContext.ImportOptions->ExistingActorPolicy );
 	UsdStageImporterImpl::NotifyAssetRegistry( { ReimportedObject } );
-	UsdStageImporterImpl::CloseStageIfNeeded( ImportContext );
 
 	FUsdDelegates::OnPostUsdImport.Broadcast(ImportContext.FilePath);
 
@@ -1522,6 +1547,8 @@ bool UUsdStageImporter::ReimportSingleAsset(FUsdStageImportContext& ImportContex
 		UsdStageImporterImpl::SendAnalytics( ImportContext, ReimportedObject, TEXT( "Reimport" ), { ReimportedObject }, ElapsedSeconds );
 	}
 
+	UsdStageImporterImpl::CloseStageIfNeeded( ImportContext );
+
 #endif // #if USE_USD_SDK
 	return bSuccess;
 }
diff --git a/Engine/Plugins/Importers/USDImporter/Source/USDStageImporter/Public/USDStageImportContext.h b/Engine/Plugins/Importers/USDImporter/Source/USDStageImporter/Public/USDStageImportContext.h
--- a/Engine/Plugins/Importers/USDImporter/Source/USDStageImporter/Public/USDStageImportContext.h
+++ b/Engine/Plugins/Importers/USDImporter/Source/USDStageImporter/Public/USDStageImportContext.h
@@ -65,6 +65,7 @@ struct USDSTAGEIMPORTER_API FUsdStageImportContext
 	 */
 	TMap< FString, TMap< FString, int32 > > MaterialToPrimvarToUVIndex;
 
+	/** USD Stage to import */
 	UE::FUsdStage Stage;
 
 	/** Object flags to apply to newly imported objects */
@@ -80,7 +81,7 @@ struct USDSTAGEIMPORTER_API FUsdStageImportContext
 	bool bReadFromStageCache;
 
 	/** If we're reading from the stage cache and the stage was originally open, it will be left open when the import is completed */
-	bool bStageWasOriginallyOpen;
+	bool bStageWasOriginallyOpenInCache;
 
 	/** We modify the stage with our meters per unit import option on import. If the stage was already open, we use this to undo the changes after import */
 	float OriginalMetersPerUnit;
