USD: Allow controlling whether foliage is exported on actor's layer on the target component's layer

diff --git a/Engine/Plugins/Importers/USDImporter/Source/USDExporter/Private/USDConversionBlueprintContext.cpp b/Engine/Plugins/Importers/USDImporter/Source/USDExporter/Private/USDConversionBlueprintContext.cpp
--- a/Engine/Plugins/Importers/USDImporter/Source/USDExporter/Private/USDConversionBlueprintContext.cpp
+++ b/Engine/Plugins/Importers/USDImporter/Source/USDExporter/Private/USDConversionBlueprintContext.cpp
@@ -286,7 +286,7 @@ bool UUsdConversionBlueprintContext::ConvertCineCameraComponent( const UCineCame
 #endif // USE_USD_SDK
 }
 
-bool UUsdConversionBlueprintContext::ConvertInstancedFoliageActor( const AInstancedFoliageActor* Actor, const FString& PrimPath, float TimeCode )
+bool UUsdConversionBlueprintContext::ConvertInstancedFoliageActor( const AInstancedFoliageActor* Actor, const FString& PrimPath, ULevel* InstancesLevel, float TimeCode )
 {
 #if USE_USD_SDK
 	UE::FUsdPrim Prim = UnrealToUsdImpl::GetPrim( Stage, PrimPath );
@@ -295,7 +295,7 @@ bool UUsdConversionBlueprintContext::ConvertInstancedFoliageActor( const AInstan
 		return false;
 	}
 
-	return UnrealToUsd::ConvertInstancedFoliageActor( *Actor, Prim, TimeCode == FLT_MAX ? UsdUtils::GetDefaultTimeCode() : TimeCode );
+	return UnrealToUsd::ConvertInstancedFoliageActor( *Actor, Prim, TimeCode == FLT_MAX ? UsdUtils::GetDefaultTimeCode() : TimeCode, InstancesLevel );
 #else
 	return false;
 #endif // USE_USD_SDK
diff --git a/Engine/Plugins/Importers/USDImporter/Source/USDExporter/Private/USDExporterBlueprintLibrary.cpp b/Engine/Plugins/Importers/USDImporter/Source/USDExporter/Private/USDExporterBlueprintLibrary.cpp
--- a/Engine/Plugins/Importers/USDImporter/Source/USDExporter/Private/USDExporterBlueprintLibrary.cpp
+++ b/Engine/Plugins/Importers/USDImporter/Source/USDExporter/Private/USDExporterBlueprintLibrary.cpp
@@ -65,11 +65,6 @@ TArray<FTransform> UUsdExporterBlueprintLibrary::GetInstanceTransforms( AInstanc
 		return Result;
 	}
 
-	if ( InstancesLevel == nullptr )
-	{
-		InstancesLevel = Actor->GetLevel();
-	}
-
 	// Modified from AInstancedFoliageActor::GetInstancesForComponent to limit traversal only to our FoliageType
 
 	if ( const TUniqueObj<FFoliageInfo>* FoundInfo = Actor->GetFoliageInfos().Find( FoliageType ) )
@@ -81,7 +76,7 @@ TArray<FTransform> UUsdExporterBlueprintLibrary::GetInstanceTransforms( AInstanc
 		for ( const TPair<FFoliageInstanceBaseId, FFoliageInstanceBaseInfo>& FoliageInstancePair : Actor->InstanceBaseCache.InstanceBaseMap )
 		{
 			UActorComponent* Comp = FoliageInstancePair.Value.BasePtr.Get();
-			if ( !Comp || Comp->GetComponentLevel() != InstancesLevel )
+			if ( !Comp || ( InstancesLevel && ( Comp->GetComponentLevel() != InstancesLevel ) ) )
 			{
 				continue;
 			}
diff --git a/Engine/Plugins/Importers/USDImporter/Source/USDExporter/Public/LevelExporterUSDOptions.h b/Engine/Plugins/Importers/USDImporter/Source/USDExporter/Public/LevelExporterUSDOptions.h
--- a/Engine/Plugins/Importers/USDImporter/Source/USDExporter/Public/LevelExporterUSDOptions.h
+++ b/Engine/Plugins/Importers/USDImporter/Source/USDExporter/Public/LevelExporterUSDOptions.h
@@ -29,6 +29,14 @@ struct USDEXPORTER_API FLevelExporterUSDOptionsInner
 	UPROPERTY( EditAnywhere, config, BlueprintReadWrite, Category = "Export options" )
 	bool bIgnoreSequencerAnimations = false;
 
+	/**
+	 * By default foliage instances will be exported to the same layer as the component they were placed on in the editor.
+	 * Enable this to instead export the foliage instances to the same layer as the foliage actor they belong to.
+	 * This is useful if those foliage instances were placed with the "Place In Current Level" option.
+	 */
+	UPROPERTY( EditAnywhere, config, BlueprintReadWrite, Category = "Export options" )
+	bool bExportFoliageOnActorsLayer = false;
+
 	/** Where to place all the generated asset files */
 	UPROPERTY( EditAnywhere, config, BlueprintReadWrite, Category = "Asset options" )
 	FDirectoryPath AssetFolder;
diff --git a/Engine/Plugins/Importers/USDImporter/Source/USDExporter/Public/USDConversionBlueprintContext.h b/Engine/Plugins/Importers/USDImporter/Source/USDExporter/Public/USDConversionBlueprintContext.h
--- a/Engine/Plugins/Importers/USDImporter/Source/USDExporter/Public/USDConversionBlueprintContext.h
+++ b/Engine/Plugins/Importers/USDImporter/Source/USDExporter/Public/USDConversionBlueprintContext.h
@@ -129,7 +129,7 @@ public:
 	bool ConvertCineCameraComponent( const UCineCameraComponent* Component, const FString& PrimPath, float TimeCode = 3.402823466e+38F );
 
 	UFUNCTION( BlueprintCallable, Category = "Component conversion" )
-	bool ConvertInstancedFoliageActor( const AInstancedFoliageActor* Actor, const FString& PrimPath, float TimeCode = 3.402823466e+38F );
+	bool ConvertInstancedFoliageActor( const AInstancedFoliageActor* Actor, const FString& PrimPath, ULevel* InstancesLevel = nullptr, float TimeCode = 3.402823466e+38F );
 
 	UFUNCTION( BlueprintCallable, Category = "Component conversion" )
 	bool ConvertLandscapeProxyActorMesh( const ALandscapeProxy* Actor, const FString& PrimPath, int32 LowestLOD, int32 HighestLOD, float TimeCode = 3.402823466e+38F );
diff --git a/Engine/Plugins/Importers/USDImporter/Source/USDExporter/Public/USDExporterBlueprintLibrary.h b/Engine/Plugins/Importers/USDImporter/Source/USDExporter/Public/USDExporterBlueprintLibrary.h
--- a/Engine/Plugins/Importers/USDImporter/Source/USDExporter/Public/USDExporterBlueprintLibrary.h
+++ b/Engine/Plugins/Importers/USDImporter/Source/USDExporter/Public/USDExporterBlueprintLibrary.h
@@ -47,11 +47,11 @@ public:
 	static UObject* GetSource( UFoliageType* FoliageType );
 
 	/**
-	 * Returns the transforms of all instances of a particular UFoliageType on a given level (which defaults to the Actor's level if left nullptr).
+	 * Returns the transforms of all instances of a particular UFoliageType on a given level. If no level is provided all instances will be returned.
 	 * Use GetUsedFoliageTypes() to retrieve all foliage types managed by a particular actor.
 	 */
 	UFUNCTION( BlueprintCallable, meta = ( ScriptMethod ), Category = "USD Foliage Exporter" )
-	static TArray<FTransform> GetInstanceTransforms( AInstancedFoliageActor* Actor, UFoliageType* FoliageType, ULevel* InstancesLevel = nullptr);
+	static TArray<FTransform> GetInstanceTransforms( AInstancedFoliageActor* Actor, UFoliageType* FoliageType, ULevel* InstancesLevel = nullptr );
 
 	/** Defer to the USDClasses module to actually send analytics information */
 	UFUNCTION( BlueprintCallable, Category = "Analytics" )
diff --git a/Engine/Plugins/Importers/USDImporter/Source/USDUtilities/Private/USDPrimConversion.cpp b/Engine/Plugins/Importers/USDImporter/Source/USDUtilities/Private/USDPrimConversion.cpp
--- a/Engine/Plugins/Importers/USDImporter/Source/USDUtilities/Private/USDPrimConversion.cpp
+++ b/Engine/Plugins/Importers/USDImporter/Source/USDUtilities/Private/USDPrimConversion.cpp
@@ -2270,7 +2270,7 @@ bool UnrealToUsd::ConvertXformable( const FTransform& RelativeTransform, pxr::Us
 	return true;
 }
 
-bool UnrealToUsd::ConvertInstancedFoliageActor( const AInstancedFoliageActor& Actor, pxr::UsdPrim& UsdPrim, double TimeCode )
+bool UnrealToUsd::ConvertInstancedFoliageActor( const AInstancedFoliageActor& Actor, pxr::UsdPrim& UsdPrim, double TimeCode, ULevel* InstancesLevel )
 {
 #if WITH_EDITOR
 	using namespace pxr;
@@ -2294,42 +2294,48 @@ bool UnrealToUsd::ConvertInstancedFoliageActor( const AInstancedFoliageActor& Ac
 	int PrototypeIndex = 0;
 	for ( const TPair<UFoliageType*, TUniqueObj<FFoliageInfo>>& FoliagePair : Actor.GetFoliageInfos() )
 	{
-		const UFoliageType* FoliageType = FoliagePair.Key;
 		const FFoliageInfo& Info = FoliagePair.Value.Get();
 
-		for ( const TPair<FFoliageInstanceBaseId, TSet<int32>>& Pair : Info.ComponentHash )
+		for ( const TPair<FFoliageInstanceBaseId, FFoliageInstanceBaseInfo>& FoliageInstancePair : Actor.InstanceBaseCache.InstanceBaseMap )
 		{
-			const TSet<int32>& InstanceSet = Pair.Value;
-
-			const int32 NumInstances = InstanceSet.Num();
-			ProtoIndices.reserve( ProtoIndices.size() + NumInstances );
-			Positions.reserve( Positions.size() + NumInstances );
-			Orientations.reserve( Orientations.size() + NumInstances );
-			Scales.reserve( Scales.size() + NumInstances );
+			UActorComponent* Comp = FoliageInstancePair.Value.BasePtr.Get();
+			if ( !Comp || ( InstancesLevel && ( Comp->GetComponentLevel() != InstancesLevel ) ) )
+			{
+				continue;
+			}
 
-			for ( int32 InstanceIndex : InstanceSet )
+			if ( const TSet<int32>* InstanceSet = Info.ComponentHash.Find( FoliageInstancePair.Key ) )
 			{
-				const FFoliageInstancePlacementInfo* Instance = &Info.Instances[ InstanceIndex ];
+				const int32 NumInstances = InstanceSet->Num();
+				ProtoIndices.reserve( ProtoIndices.size() + NumInstances );
+				Positions.reserve( Positions.size() + NumInstances );
+				Orientations.reserve( Orientations.size() + NumInstances );
+				Scales.reserve( Scales.size() + NumInstances );
 
-				// Convert axes
-				FTransform UETransform{ Instance->Rotation, Instance->Location, Instance->DrawScale3D };
-				FTransform USDTransform = UsdUtils::ConvertAxes( StageInfo.UpAxis == EUsdUpAxis::ZAxis, UETransform );
+				for ( int32 InstanceIndex : (*InstanceSet) )
+				{
+					const FFoliageInstancePlacementInfo* Instance = &Info.Instances[ InstanceIndex ];
 
-				FVector Translation = USDTransform.GetTranslation();
-				FQuat Rotation = USDTransform.GetRotation();
-				FVector Scale = USDTransform.GetScale3D();
+					// Convert axes
+					FTransform UETransform{ Instance->Rotation, Instance->Location, Instance->DrawScale3D };
+					FTransform USDTransform = UsdUtils::ConvertAxes( StageInfo.UpAxis == EUsdUpAxis::ZAxis, UETransform );
 
-				// Compensate metersPerUnit
-				const float UEMetersPerUnit = 0.01f;
-				if ( !FMath::IsNearlyEqual( UEMetersPerUnit, StageInfo.MetersPerUnit ) )
-				{
-					Translation *= ( UEMetersPerUnit / StageInfo.MetersPerUnit );
-				}
+					FVector Translation = USDTransform.GetTranslation();
+					FQuat Rotation = USDTransform.GetRotation();
+					FVector Scale = USDTransform.GetScale3D();
 
-				ProtoIndices.push_back( PrototypeIndex );
-				Positions.push_back( GfVec3f( Translation.X, Translation.Y, Translation.Z ) );
-				Orientations.push_back( GfQuath( Rotation.W, Rotation.X, Rotation.Y, Rotation.Z ) );
-				Scales.push_back( GfVec3f( Scale.X, Scale.Y, Scale.Z ) );
+					// Compensate metersPerUnit
+					const float UEMetersPerUnit = 0.01f;
+					if ( !FMath::IsNearlyEqual( UEMetersPerUnit, StageInfo.MetersPerUnit ) )
+					{
+						Translation *= ( UEMetersPerUnit / StageInfo.MetersPerUnit );
+					}
+
+					ProtoIndices.push_back( PrototypeIndex );
+					Positions.push_back( GfVec3f( Translation.X, Translation.Y, Translation.Z ) );
+					Orientations.push_back( GfQuath( Rotation.W, Rotation.X, Rotation.Y, Rotation.Z ) );
+					Scales.push_back( GfVec3f( Scale.X, Scale.Y, Scale.Z ) );
+				}
 			}
 		}
 
diff --git a/Engine/Plugins/Importers/USDImporter/Source/USDUtilities/Public/USDPrimConversion.h b/Engine/Plugins/Importers/USDImporter/Source/USDUtilities/Public/USDPrimConversion.h
--- a/Engine/Plugins/Importers/USDImporter/Source/USDUtilities/Public/USDPrimConversion.h
+++ b/Engine/Plugins/Importers/USDImporter/Source/USDUtilities/Public/USDPrimConversion.h
@@ -137,8 +137,9 @@ namespace UnrealToUsd
 	 * @param Actor				   The actor to convert data from
 	 * @param Prim                 The pxr::UsdGeomPointInstancer to write to
 	 * @param TimeCode			   TimeCode to write the attribute values at. Use UsdUtils::GetDefaultTimeCode() for the Default value.
+	 * @param InstancesLevel	   If this is not nullptr, only foliage instances placed on components that belong to this level will be exported
 	 */
-	USDUTILITIES_API bool ConvertInstancedFoliageActor( const AInstancedFoliageActor& Actor, pxr::UsdPrim& UsdPrim, double TimeCode );
+	USDUTILITIES_API bool ConvertInstancedFoliageActor( const AInstancedFoliageActor& Actor, pxr::UsdPrim& UsdPrim, double TimeCode, ULevel* InstancesLevel = nullptr );
 
 	// We will defer to the common UnrealToUsd::ConvertXComponent functions when baking level sequences, which will already bake
 	// all of the properties of a component out at the same time.
