USD: Fix unexpected behavior when reloading a temporary stage.

Disable menu options that don't make sense (e.g. close a stage when none is open);
Prevent reloading of temporary stages.

diff --git a/Engine/Plugins/Importers/USDImporter/Source/USDStageEditor/Private/SUSDStage.cpp b/Engine/Plugins/Importers/USDImporter/Source/USDStageEditor/Private/SUSDStage.cpp
--- a/Engine/Plugins/Importers/USDImporter/Source/USDStageEditor/Private/SUSDStage.cpp
+++ b/Engine/Plugins/Importers/USDImporter/Source/USDStageEditor/Private/SUSDStage.cpp
@@ -465,7 +465,18 @@ void SUsdStage::FillFileMenu( FMenuBuilder& MenuBuilder )
 			FSlateIcon(),
 			FUIAction(
 				FExecuteAction::CreateSP( this, &SUsdStage::OnSave ),
-				FCanExecuteAction()
+				FCanExecuteAction::CreateLambda( [this]()
+				{
+					if ( const AUsdStageActor* StageActor = ViewModel.UsdStageActor.Get() )
+					{
+						if ( UE::FUsdStage Stage = StageActor->GetUsdStage() )
+						{
+							return true;
+						}
+					}
+
+					return false;
+				})
 			),
 			NAME_None,
 			EUserInterfaceActionType::Button
@@ -479,7 +490,21 @@ void SUsdStage::FillFileMenu( FMenuBuilder& MenuBuilder )
 			FSlateIcon(),
 			FUIAction(
 				FExecuteAction::CreateSP( this, &SUsdStage::OnReloadStage ),
-				FCanExecuteAction()
+				FCanExecuteAction::CreateLambda( [this]()
+				{
+					if ( const AUsdStageActor* StageActor = ViewModel.UsdStageActor.Get() )
+					{
+						if ( UE::FUsdStage Stage = StageActor->GetUsdStage() )
+						{
+							if ( !Stage.GetRootLayer().IsAnonymous() )
+							{
+								return true;
+							}
+						}
+					}
+
+					return false;
+				})
 			),
 			NAME_None,
 			EUserInterfaceActionType::Button
@@ -491,7 +516,18 @@ void SUsdStage::FillFileMenu( FMenuBuilder& MenuBuilder )
 			FSlateIcon(),
 			FUIAction(
 				FExecuteAction::CreateSP( this, &SUsdStage::OnResetStage ),
-				FCanExecuteAction()
+				FCanExecuteAction::CreateLambda( [this]()
+				{
+					if ( const AUsdStageActor* StageActor = ViewModel.UsdStageActor.Get() )
+					{
+						if ( UE::FUsdStage Stage = StageActor->GetUsdStage() )
+						{
+							return true;
+						}
+					}
+
+					return false;
+				})
 			),
 			NAME_None,
 			EUserInterfaceActionType::Button
@@ -505,7 +541,18 @@ void SUsdStage::FillFileMenu( FMenuBuilder& MenuBuilder )
 			FSlateIcon(),
 			FUIAction(
 				FExecuteAction::CreateSP( this, &SUsdStage::OnClose ),
-				FCanExecuteAction()
+				FCanExecuteAction::CreateLambda( [this]()
+				{
+					if ( const AUsdStageActor* StageActor = ViewModel.UsdStageActor.Get() )
+					{
+						if ( UE::FUsdStage Stage = StageActor->GetUsdStage() )
+						{
+							return true;
+						}
+					}
+
+					return false;
+				})
 			),
 			NAME_None,
 			EUserInterfaceActionType::Button
@@ -524,7 +571,18 @@ void SUsdStage::FillActionsMenu( FMenuBuilder& MenuBuilder )
 			FSlateIcon(),
 			FUIAction(
 				FExecuteAction::CreateSP( this, &SUsdStage::OnImport ),
-				FCanExecuteAction()
+				FCanExecuteAction::CreateLambda( [this]()
+				{
+					if ( const AUsdStageActor* StageActor = ViewModel.UsdStageActor.Get() )
+					{
+						if ( UE::FUsdStage Stage = StageActor->GetUsdStage() )
+						{
+							return true;
+						}
+					}
+
+					return false;
+				})
 			),
 			NAME_None,
 			EUserInterfaceActionType::Button
diff --git a/Engine/Plugins/Importers/USDImporter/Source/USDStageEditorViewModels/Private/USDLayersViewModel.cpp b/Engine/Plugins/Importers/USDImporter/Source/USDStageEditorViewModels/Private/USDLayersViewModel.cpp
--- a/Engine/Plugins/Importers/USDImporter/Source/USDStageEditorViewModels/Private/USDLayersViewModel.cpp
+++ b/Engine/Plugins/Importers/USDImporter/Source/USDStageEditorViewModels/Private/USDLayersViewModel.cpp
@@ -149,7 +149,7 @@ void FUsdLayerViewModel::RefreshData()
 	pxr::SdfLayerRefPtr UsdLayer = UE::FSdfLayer::FindOrOpen( *LayerIdentifier );
 	if ( UsdLayer )
 	{
-		LayerModel->bIsDirty = UsdLayer->IsDirty();
+		LayerModel->bIsDirty = ( UsdLayer->IsDirty() || UsdLayer->IsAnonymous() ) && !UsdUtils::IsSessionLayerWithinStage( UsdLayer, UsdStage );
 	}
 
 	const pxr::SdfLayerHandle& EditTargetLayer = UsdStageRef->GetEditTarget().GetLayer();
diff --git a/Engine/Plugins/Importers/USDImporter/Source/USDStageEditorViewModels/Private/USDStageViewModel.cpp b/Engine/Plugins/Importers/USDImporter/Source/USDStageEditorViewModels/Private/USDStageViewModel.cpp
--- a/Engine/Plugins/Importers/USDImporter/Source/USDStageEditorViewModels/Private/USDStageViewModel.cpp
+++ b/Engine/Plugins/Importers/USDImporter/Source/USDStageEditorViewModels/Private/USDStageViewModel.cpp
@@ -130,6 +130,13 @@ void FUsdStageViewModel::ReloadStage()
 	UE::FUsdStage Stage = UsdStageActor->GetOrLoadUsdStage();
 	pxr::UsdStageRefPtr UsdStage = pxr::UsdStageRefPtr( Stage );
 
+	// Can't reload from disk something that doesn't exist on disk yet
+	// (actually USD will let us do this but it seems to just clear the anonymous layers instead)
+	if ( Stage.GetRootLayer().IsAnonymous() )
+	{
+		return;
+	}
+
 	if ( UsdStage )
 	{
 		UsdUtils::StartMonitoringErrors();
diff --git a/Engine/Plugins/Importers/USDImporter/Source/USDUtilities/Private/USDLayerUtils.cpp b/Engine/Plugins/Importers/USDImporter/Source/USDUtilities/Private/USDLayerUtils.cpp
--- a/Engine/Plugins/Importers/USDImporter/Source/USDUtilities/Private/USDLayerUtils.cpp
+++ b/Engine/Plugins/Importers/USDImporter/Source/USDUtilities/Private/USDLayerUtils.cpp
@@ -636,6 +636,34 @@ UE::FSdfLayer UsdUtils::FindLayerForIdentifier( const TCHAR* Identifier, const U
 	return UE::FSdfLayer{};
 }
 
+bool UsdUtils::IsSessionLayerWithinStage( const pxr::SdfLayerRefPtr& Layer, const pxr::UsdStageRefPtr& Stage )
+{
+	if ( !Layer || !Stage )
+	{
+		return false;
+	}
+
+	pxr::SdfLayerRefPtr RootLayer = Stage->GetRootLayer();
+
+	const bool bIncludeSessionLayers = true;
+	for ( const pxr::SdfLayerHandle& ExistingLayer : Stage->GetLayerStack( bIncludeSessionLayers ) )
+	{
+		// All session layers come before the root layer within the layer stack
+		// Break before we compare with Layer because if Layer is the actual stage's RootLayer we want to return false
+		if ( ExistingLayer == RootLayer )
+		{
+			break;
+		}
+
+		if ( ExistingLayer == Layer )
+		{
+			return true;
+		}
+	}
+
+	return false;
+}
+
 #undef LOCTEXT_NAMESPACE
 
 #endif // #if USE_USD_SDK
diff --git a/Engine/Plugins/Importers/USDImporter/Source/USDUtilities/Public/USDLayerUtils.h b/Engine/Plugins/Importers/USDImporter/Source/USDUtilities/Public/USDLayerUtils.h
--- a/Engine/Plugins/Importers/USDImporter/Source/USDUtilities/Public/USDLayerUtils.h
+++ b/Engine/Plugins/Importers/USDImporter/Source/USDUtilities/Public/USDLayerUtils.h
@@ -100,6 +100,9 @@ namespace UsdUtils
 
 	/** Uses FindOrOpen to return the layer with the given identifier if possible. If the identifier is for an anonymous layer, it will search via display name instead */
 	USDUTILITIES_API UE::FSdfLayer FindLayerForIdentifier( const TCHAR* Identifier, const UE::FUsdStage& Stage );
+
+	/** Returns true if Layer is a session layer within Stage's layer stack */
+	USDUTILITIES_API bool IsSessionLayerWithinStage( const pxr::SdfLayerRefPtr& Layer, const pxr::UsdStageRefPtr& Stage );
 }
 
 #endif // #if USE_USD_SDK
