USD: Properly collapse Mesh child meshes, since the geom mesh translator claims that it can do that

diff --git a/Engine/Plugins/Importers/USDImporter/Source/USDSchemas/Private/USDGeomMeshTranslator.cpp b/Engine/Plugins/Importers/USDImporter/Source/USDSchemas/Private/USDGeomMeshTranslator.cpp
--- a/Engine/Plugins/Importers/USDImporter/Source/USDSchemas/Private/USDGeomMeshTranslator.cpp
+++ b/Engine/Plugins/Importers/USDImporter/Source/USDSchemas/Private/USDGeomMeshTranslator.cpp
@@ -504,7 +504,7 @@ namespace UsdGeomMeshTranslatorImpl
 	}
 
 	void LoadMeshDescriptions( const pxr::UsdTyped& UsdMesh, TArray<FMeshDescription>& OutLODIndexToMeshDescription, TArray<UsdUtils::FUsdPrimMaterialAssignmentInfo>& OutLODIndexToMaterialInfo,
-		const TMap< FString, TMap< FString, int32 > >& MaterialToPrimvarToUVIndex, const pxr::UsdTimeCode TimeCode, bool bInterpretLODs, const FName& RenderContext )
+		const TMap< FString, TMap< FString, int32 > >& MaterialToPrimvarToUVIndex, const pxr::UsdTimeCode TimeCode, bool bInterpretLODs, const FName& RenderContext, const EUsdPurpose PurposesToLoad )
 	{
 		if ( !UsdMesh )
 		{
@@ -528,10 +528,15 @@ namespace UsdGeomMeshTranslatorImpl
 			FMeshDescription TempMeshDescription;
 			UsdUtils::FUsdPrimMaterialAssignmentInfo TempMaterialInfo;
 
-			FStaticMeshAttributes StaticMeshAttributes( TempMeshDescription );
-			StaticMeshAttributes.Register();
-
-			bool bSuccess = UsdToUnreal::ConvertGeomMesh( UsdMesh, TempMeshDescription, TempMaterialInfo, FTransform::Identity, MaterialToPrimvarToUVIndex, TimeCode, RenderContextToken );
+			bool bSuccess = UsdToUnreal::ConvertGeomMeshHierarchy(
+				UsdMesh.GetPrim(),
+				TimeCode,
+				PurposesToLoad,
+				RenderContextToken,
+				MaterialToPrimvarToUVIndex,
+				TempMeshDescription,
+				TempMaterialInfo
+			);
 
 			if ( bSuccess )
 			{
@@ -792,6 +797,8 @@ namespace UsdGeomMeshTranslatorImpl
 			TMap< FString, TMap< FString, int32 > > Unused;
 			TMap< FString, TMap< FString, int32 > >* MaterialToPrimvarToUVIndex = Context->MaterialToPrimvarToUVIndex ? Context->MaterialToPrimvarToUVIndex : &Unused;
 
+			const EUsdPurpose PurposesToLoad = Context->PurposesToLoad;
+
 			// Fetch the animated mesh start/end frame as they may be different from just the stage's start and end time codes
 			int32 StartFrame = FMath::FloorToInt( Stage.GetStartTimeCode() );
 			int32 EndFrame = FMath::CeilToInt( Stage.GetEndTimeCode() );
@@ -806,7 +813,7 @@ namespace UsdGeomMeshTranslatorImpl
 				*MaterialToPrimvarToUVIndex,
 				StartFrame,
 				EndFrame,
-				[InPrimPath]( const TWeakObjectPtr<UGeometryCacheTrackUsd> TrackPtr, float Time, FGeometryCacheMeshData& OutMeshData )
+				[InPrimPath, PurposesToLoad]( const TWeakObjectPtr<UGeometryCacheTrackUsd> TrackPtr, float Time, FGeometryCacheMeshData& OutMeshData )
 				{
 					UGeometryCacheTrackUsd* Track = TrackPtr.Get();
 					if ( !Track )
@@ -847,7 +854,8 @@ namespace UsdGeomMeshTranslatorImpl
 						Track->MaterialToPrimvarToUVIndex,
 						pxr::UsdTimeCode( Time ),
 						bAllowInterpretingLODs,
-						Track->RenderContext
+						Track->RenderContext,
+						PurposesToLoad
 					);
 
 					// Convert the MeshDescription to MeshData
@@ -1177,7 +1185,8 @@ void FGeomMeshCreateAssetsTaskChain::SetupTasks()
 				*MaterialToPrimvarToUVIndex,
 				pxr::UsdTimeCode( Context->Time ),
 				Context->bAllowInterpretingLODs,
-				Context->RenderContext
+				Context->RenderContext,
+				Context->PurposesToLoad
 			);
 
 			// If we have at least one valid LOD, we should keep going
@@ -1231,7 +1240,8 @@ void FGeometryCacheCreateAssetsTaskChain::SetupTasks()
 				*MaterialToPrimvarToUVIndex,
 				pxr::UsdTimeCode( TimeCode ),
 				bAllowInterpretingLODs,
-				Context->RenderContext
+				Context->RenderContext,
+				Context->PurposesToLoad
 			);
 
 			// If we have at least one valid LOD, we should keep going
diff --git a/Engine/Plugins/Importers/USDImporter/Source/USDSchemas/Private/USDGeomXformableTranslator.cpp b/Engine/Plugins/Importers/USDImporter/Source/USDSchemas/Private/USDGeomXformableTranslator.cpp
--- a/Engine/Plugins/Importers/USDImporter/Source/USDSchemas/Private/USDGeomXformableTranslator.cpp
+++ b/Engine/Plugins/Importers/USDImporter/Source/USDSchemas/Private/USDGeomXformableTranslator.cpp
@@ -40,59 +40,6 @@
 	#include "pxr/usd/usdShade/materialBindingAPI.h"
 #include "USDIncludesEnd.h"
 
-namespace UsdGeomXformableTranslatorImpl
-{
-	void LoadMeshDescription( const pxr::UsdGeomXformable& UsdGeomXformable, const EUsdPurpose PurposesToLoad, const TMap< FString, TMap<FString, int32> >& MaterialToPrimvarToUVIndex, const pxr::UsdTimeCode TimeCode, FMeshDescription& OutMeshdescription, UsdUtils::FUsdPrimMaterialAssignmentInfo& OutMaterialAssignments )
-	{
-		TRACE_CPUPROFILER_EVENT_SCOPE( UsdGeomXformableTranslatorImpl::LoadMeshDescription );
-
-		FStaticMeshAttributes StaticMeshAttributes( OutMeshdescription );
-		StaticMeshAttributes.Register();
-
-		TFunction< void( FMeshDescription&, UsdUtils::FUsdPrimMaterialAssignmentInfo&, const pxr::UsdPrim&, const FTransform, const pxr::UsdTimeCode& ) > RecursiveChildCollapsing;
-		RecursiveChildCollapsing = [ &RecursiveChildCollapsing, &MaterialToPrimvarToUVIndex, PurposesToLoad ]( FMeshDescription& MeshDescription, UsdUtils::FUsdPrimMaterialAssignmentInfo& MaterialAssignments, const pxr::UsdPrim& ParentPrim, const FTransform CurrentTransform, const pxr::UsdTimeCode& TimeCode )
-		{
-			for ( const pxr::UsdPrim& ChildPrim : ParentPrim.GetFilteredChildren( pxr::UsdTraverseInstanceProxies() ) )
-			{
-				// Ignore meshes from disabled purposes
-				if ( !EnumHasAllFlags( PurposesToLoad, IUsdPrim::GetPurpose( ChildPrim ) ) )
-				{
-					continue;
-				}
-
-				// Ignore invisible prims
-				if ( pxr::UsdGeomImageable UsdGeomImageable = pxr::UsdGeomImageable( ChildPrim ) )
-				{
-					if ( UsdGeomImageable.ComputeVisibility() == pxr::UsdGeomTokens->invisible )
-					{
-						continue;
-					}
-				}
-
-				FTransform ChildTransform = CurrentTransform;
-
-				if ( pxr::UsdGeomXformable ChildXformable = pxr::UsdGeomXformable( ChildPrim ) )
-				{
-					FTransform LocalChildTransform;
-					UsdToUnreal::ConvertXformable( ChildPrim.GetStage(), ChildXformable, LocalChildTransform, TimeCode.GetValue() );
-
-					ChildTransform = LocalChildTransform * CurrentTransform;
-				}
-
-				if ( pxr::UsdGeomMesh ChildMesh = pxr::UsdGeomMesh( ChildPrim ) )
-				{
-					UsdToUnreal::ConvertGeomMesh( ChildMesh, MeshDescription, MaterialAssignments, ChildTransform, MaterialToPrimvarToUVIndex, TimeCode );
-				}
-
-				RecursiveChildCollapsing( MeshDescription, MaterialAssignments, ChildPrim, ChildTransform, TimeCode );
-			}
-		};
-
-		// Collapse the children
-		RecursiveChildCollapsing( OutMeshdescription, OutMaterialAssignments, UsdGeomXformable.GetPrim(), FTransform::Identity, TimeCode );
-	}
-}
-
 class FUsdGeomXformableCreateAssetsTaskChain : public FBuildStaticMeshTaskChain
 {
 public:
@@ -124,11 +71,18 @@ void FUsdGeomXformableCreateAssetsTaskChain::SetupTasks()
 			TMap< FString, TMap< FString, int32 > > Unused;
 			TMap< FString, TMap< FString, int32 > >* MaterialToPrimvarToUVIndex = Context->MaterialToPrimvarToUVIndex ? Context->MaterialToPrimvarToUVIndex : &Unused;
 
-			UsdGeomXformableTranslatorImpl::LoadMeshDescription(
-				pxr::UsdGeomXformable( GetPrim() ),
+			pxr::TfToken RenderContextToken = pxr::UsdShadeTokens->universalRenderContext;
+			if ( !Context->RenderContext.IsNone() )
+			{
+				RenderContextToken = UnrealToUsd::ConvertToken( *Context->RenderContext.ToString() ).Get();
+			}
+
+			UsdToUnreal::ConvertGeomMeshHierarchy(
+				GetPrim(),
+				pxr::UsdTimeCode( Context->Time ),
 				Context->PurposesToLoad,
+				RenderContextToken,
 				*MaterialToPrimvarToUVIndex,
-				pxr::UsdTimeCode( Context->Time ),
 				AddedMeshDescription,
 				AssignmentInfo
 			);
diff --git a/Engine/Plugins/Importers/USDImporter/Source/USDUtilities/Private/USDGeomMeshConversion.cpp b/Engine/Plugins/Importers/USDImporter/Source/USDUtilities/Private/USDGeomMeshConversion.cpp
--- a/Engine/Plugins/Importers/USDImporter/Source/USDUtilities/Private/USDGeomMeshConversion.cpp
+++ b/Engine/Plugins/Importers/USDImporter/Source/USDUtilities/Private/USDGeomMeshConversion.cpp
@@ -10,6 +10,7 @@
 #include "USDErrorUtils.h"
 #include "USDLog.h"
 #include "USDMemory.h"
+#include "USDPrimConversion.h"
 #include "USDShadeConversion.h"
 #include "USDTypesConversion.h"
 
@@ -498,6 +499,77 @@ namespace UE
 
 				return true;
 			}
+
+			bool RecursivelyCollapseChildMeshes(
+				const pxr::UsdPrim& Prim,
+				const FTransform& CurrentTransform,
+				const pxr::UsdTimeCode& TimeCode,
+				const EUsdPurpose PurposesToLoad,
+				const pxr::TfToken& RenderContext,
+				const TMap< FString, TMap<FString, int32> >& MaterialToPrimvarToUVIndex,
+				FMeshDescription& OutMeshDescription,
+				UsdUtils::FUsdPrimMaterialAssignmentInfo& OutMaterialAssignments,
+				bool bParseTransform
+			)
+			{
+				// Ignore meshes from disabled purposes
+				if ( !EnumHasAllFlags( PurposesToLoad, IUsdPrim::GetPurpose( Prim ) ) )
+				{
+					return true;
+				}
+
+				// Ignore invisible prims
+				if ( pxr::UsdGeomImageable UsdGeomImageable = pxr::UsdGeomImageable( Prim ) )
+				{
+					if ( UsdGeomImageable.ComputeVisibility() == pxr::UsdGeomTokens->invisible )
+					{
+						return true;
+					}
+				}
+
+				FTransform ChildTransform = CurrentTransform;
+
+				if ( bParseTransform )
+				{
+					if ( pxr::UsdGeomXformable Xformable = pxr::UsdGeomXformable( Prim ) )
+					{
+						FTransform LocalChildTransform;
+						UsdToUnreal::ConvertXformable( Prim.GetStage(), Xformable, LocalChildTransform, TimeCode.GetValue() );
+
+						ChildTransform = LocalChildTransform * CurrentTransform;
+					}
+				}
+
+				bool bSuccess = true;
+				if ( pxr::UsdGeomMesh Mesh = pxr::UsdGeomMesh( Prim ) )
+				{
+					bSuccess = UsdToUnreal::ConvertGeomMesh( Mesh, OutMeshDescription, OutMaterialAssignments, ChildTransform, MaterialToPrimvarToUVIndex, TimeCode, RenderContext );
+				}
+
+				for ( const pxr::UsdPrim& ChildPrim : Prim.GetFilteredChildren( pxr::UsdTraverseInstanceProxies() ) )
+				{
+					if ( !bSuccess )
+					{
+						break;
+					}
+
+					const bool bChildParseTransform = true;
+
+					bSuccess &= RecursivelyCollapseChildMeshes(
+						ChildPrim,
+						ChildTransform,
+						TimeCode,
+						PurposesToLoad,
+						RenderContext,
+						MaterialToPrimvarToUVIndex,
+						OutMeshDescription,
+						OutMaterialAssignments,
+						bChildParseTransform
+					);
+				}
+
+				return bSuccess;
+			}
 		}
 	}
 }
@@ -886,6 +958,36 @@ bool UsdToUnreal::ConvertGeomMesh( const pxr::UsdTyped& UsdSchema, FMeshDescript
 	return true;
 }
 
+bool UsdToUnreal::ConvertGeomMeshHierarchy(
+	const pxr::UsdPrim& Prim,
+	const pxr::UsdTimeCode& TimeCode,
+	const EUsdPurpose PurposesToLoad,
+	const pxr::TfToken& RenderContext,
+	const TMap< FString, TMap<FString, int32> >& MaterialToPrimvarToUVIndex,
+	FMeshDescription& OutMeshDescription,
+	UsdUtils::FUsdPrimMaterialAssignmentInfo& OutMaterialAssignments
+)
+{
+	FStaticMeshAttributes StaticMeshAttributes( OutMeshDescription );
+	StaticMeshAttributes.Register();
+
+	// Skip the first transform, because if we're parsing Prim's mesh we don't want to bake its own transform into the vertices,
+	// as that transform will be used on components instead
+	const bool bParseTransform = false;
+
+	return UE::UsdGeomMeshConversion::Private::RecursivelyCollapseChildMeshes(
+		Prim,
+		FTransform::Identity,
+		TimeCode,
+		PurposesToLoad,
+		RenderContext,
+		MaterialToPrimvarToUVIndex,
+		OutMeshDescription,
+		OutMaterialAssignments,
+		bParseTransform
+	);
+}
+
 bool UsdToUnreal::ConvertDisplayColor( const UsdUtils::FDisplayColorMaterial& DisplayColorDescription, UMaterialInstanceConstant& Material )
 {
 	FUsdLogManager::LogMessage( EMessageSeverity::Warning, LOCTEXT( "DeprecatedConvertDisplayColor", "Converting existing instances with UsdToUnreal::ConvertDisplayColor is deprecated in favor of just calling UsdUtils::CreateDisplayColorMaterialInstanceConstant instead, and may be removed in a future release." ) );
diff --git a/Engine/Plugins/Importers/USDImporter/Source/USDUtilities/Public/USDGeomMeshConversion.h b/Engine/Plugins/Importers/USDImporter/Source/USDUtilities/Public/USDGeomMeshConversion.h
--- a/Engine/Plugins/Importers/USDImporter/Source/USDUtilities/Public/USDGeomMeshConversion.h
+++ b/Engine/Plugins/Importers/USDImporter/Source/USDUtilities/Public/USDGeomMeshConversion.h
@@ -6,6 +6,7 @@
 
 #include "CoreMinimal.h"
 
+#include "UnrealUSDWrapper.h"
 #include "UsdWrappers/ForwardDeclarations.h"
 
 #include "USDIncludesStart.h"
@@ -58,6 +59,21 @@ namespace UsdToUnreal
 	USDUTILITIES_API bool ConvertGeomMesh( const pxr::UsdTyped& UsdSchema, FMeshDescription& MeshDescription, const FTransform& AdditionalTransform, const pxr::UsdTimeCode TimeCode = pxr::UsdTimeCode::EarliestTime() );
 	USDUTILITIES_API bool ConvertGeomMesh( const pxr::UsdTyped& UsdSchema, FMeshDescription& MeshDescription, const FTransform& AdditionalTransform, const TMap< FString, TMap< FString, int32 > >& MaterialToPrimvarsUVSetNames, const pxr::UsdTimeCode TimeCode = pxr::UsdTimeCode::EarliestTime() );
 
+	/**
+	 * Recursively traverses down Prim's subtree parsing all Mesh data (including Prim's), baking all transforms and putting the resulting
+	 * mesh data within OutMeshDescription and OutMaterialAssignments. It will not bake Prim's own transform into the Mesh data.
+	 *
+	 * @param Prim - Prim that will be converted (result will include its mesh and its children's)
+	 * @param TimeCode - USD timecode when the UsdGeomMeshes should be sampled for mesh data and converted
+	 * @param PurposesToLoad - Only prims with these purposes will be traversed
+	 * @param RenderContext - Render context to use when parsing Mesh materials
+	 * @param MaterialToPrimvarToUVIndex - Maps from a material prim path, to pairs indicating which primvar names are used as 'st' coordinates for this mesh, and which UVIndex materials will sample from (e.g. ["st0", 0], ["myUvSet2", 2], etc). This is used to pick which primvars will become UV sets.
+	 * @param OutMeshDescription - Output parameter that will be filled with the converted mesh data
+	 * @param OutMaterialAssignments - Output parameter that will be filled with the material data extracted from UsdSchema
+	 * @return Whether the conversion was successful or not.
+	 */
+	USDUTILITIES_API bool ConvertGeomMeshHierarchy( const pxr::UsdPrim& Prim, const pxr::UsdTimeCode& TimeCode, const EUsdPurpose PurposesToLoad, const pxr::TfToken& RenderContext, const TMap< FString, TMap<FString, int32> >& MaterialToPrimvarToUVIndex, FMeshDescription& OutMeshDescription, UsdUtils::FUsdPrimMaterialAssignmentInfo& OutMaterialAssignments );
+
 	/** Configure Material to become a vertex color/displayColor material, according to the given description */
 	USDUTILITIES_API bool ConvertDisplayColor( const UsdUtils::FDisplayColorMaterial& DisplayColorDescription, UMaterialInstanceConstant& Material );
 }
