USD: Fix camera transform compensation not handling cine camera components correctly anymore.

This is an early submit of this 4.27.X change into Release-5.0. When/if 4.27.X is opened, the original cl (swarm 17789608) will be submitted there directly.

diff --git a/Engine/Plugins/Importers/USDImporter/Source/USDSchemas/Private/USDGeomCameraTranslator.cpp b/Engine/Plugins/Importers/USDImporter/Source/USDSchemas/Private/USDGeomCameraTranslator.cpp
--- a/Engine/Plugins/Importers/USDImporter/Source/USDSchemas/Private/USDGeomCameraTranslator.cpp
+++ b/Engine/Plugins/Importers/USDImporter/Source/USDSchemas/Private/USDGeomCameraTranslator.cpp
@@ -31,7 +31,10 @@ USceneComponent* FUsdGeomCameraTranslator::CreateComponents()
 
 		if ( pxr::UsdPrim UsdPrim{ GetPrim() } )
 		{
-			if ( pxr::UsdPrim UsdParentPrim = UsdPrim.GetParent() )
+			pxr::UsdPrim UsdParentPrim = UsdPrim.GetParent();
+
+			// Check if we're a child of the pseudoroot here because if we are USD will emit a warning if we try getting attributes from it...
+			if ( UsdParentPrim && !UsdParentPrim.IsPseudoRoot() )
 			{
 				bool bIsCineCameraActorMainComponent = false;
 				if ( pxr::UsdAttribute Attr = UsdParentPrim.GetAttribute( UnrealToUsd::ConvertToken( TEXT( "unrealCameraPrimName" ) ).Get() ) )
diff --git a/Engine/Plugins/Importers/USDImporter/Source/USDUtilities/Private/USDPrimConversion.cpp b/Engine/Plugins/Importers/USDImporter/Source/USDUtilities/Private/USDPrimConversion.cpp
--- a/Engine/Plugins/Importers/USDImporter/Source/USDUtilities/Private/USDPrimConversion.cpp
+++ b/Engine/Plugins/Importers/USDImporter/Source/USDUtilities/Private/USDPrimConversion.cpp
@@ -1871,11 +1871,12 @@ bool UnrealToUsd::ConvertSceneComponent( const pxr::UsdStageRefPtr& Stage, const
 
 	FTransform RelativeTransform = SceneComponent->GetRelativeTransform();
 
-	// Compensate different orientation for light or camera components :
+	// Compensate different orientation for light or camera components:
 	// In USD cameras shoot towards local - Z, with + Y up.Lights also emit towards local - Z, with + Y up
 	// In UE cameras shoot towards local + X, with + Z up.Lights also emit towards local + X, with + Z up
-	if ( SceneComponent->IsA( UCineCameraComponent::StaticClass() ) ||
-		 SceneComponent->IsA( ULightComponent::StaticClass() ) )
+	// Note that this wouldn't have worked in case we collapsed light and camera components, but these always get their own
+	// actors, so we know that we don't have a single component that represents a large collapsed prim hierarchy
+	if ( UsdPrim.IsA<pxr::UsdGeomCamera>() || UsdPrim.IsA<pxr::UsdLuxLight>() )
 	{
 		FTransform AdditionalRotation = FTransform( FRotator( 0.0f, 90.f, 0.0f ) );
 
@@ -1888,10 +1889,9 @@ bool UnrealToUsd::ConvertSceneComponent( const pxr::UsdStageRefPtr& Stage, const
 	}
 
 	// Invert compensation applied to parent if it's a light or camera component
-	if ( USceneComponent* AttachParent = SceneComponent->GetAttachParent() )
+	if ( pxr::UsdPrim ParentPrim = UsdPrim.GetParent() )
 	{
-		if ( AttachParent->IsA( UCineCameraComponent::StaticClass() ) ||
-			 AttachParent->IsA( ULightComponent::StaticClass() ) )
+		if ( ParentPrim.IsA<pxr::UsdGeomCamera>() || ParentPrim.IsA<pxr::UsdLuxLight>() )
 		{
 			FTransform AdditionalRotation = FTransform( FRotator( 0.0f, 90.f, 0.0f ) );
 
@@ -3356,14 +3356,29 @@ bool UnrealToUsd::ConvertXformable( const UMovieScene3DTransformTrack& MovieScen
 		pxr::SdfChangeBlock ChangeBlock;
 
 		// Compensate different orientation for light or camera components
-		FTransform AdditionalRotation = FTransform::Identity;
+		FTransform CameraCompensation = FTransform::Identity;
 		if ( UsdPrim.IsA< pxr::UsdGeomCamera >() || UsdPrim.IsA< pxr::UsdLuxLight >() )
 		{
-			AdditionalRotation = FTransform( FRotator( 0.0f, 90.0f, 0.0f ) );
+			CameraCompensation = FTransform( FRotator( 0.0f, 90.0f, 0.0f ) );
 
 			if ( StageInfo.UpAxis == EUsdUpAxis::ZAxis )
 			{
-				AdditionalRotation *= FTransform( FRotator( 90.0f, 0.0f, 0.0f ) );
+				CameraCompensation *= FTransform( FRotator( 90.0f, 0.0f, 0.0f ) );
+			}
+		}
+
+		// Invert compensation applied to parent if it's a light or camera component
+		FTransform InverseCameraCompensation = FTransform::Identity;
+		if ( pxr::UsdPrim ParentPrim = UsdPrim.GetParent() )
+		{
+			if ( ParentPrim.IsA<pxr::UsdGeomCamera>() || ParentPrim.IsA<pxr::UsdLuxLight>() )
+			{
+				InverseCameraCompensation = FTransform( FRotator( 0.0f, 90.f, 0.0f ) );
+
+				if ( StageInfo.UpAxis == EUsdUpAxis::ZAxis )
+				{
+					InverseCameraCompensation *= FTransform( FRotator( 90.0f, 0.f, 0.0f ) );
+				}
 			}
 		}
 
@@ -3377,7 +3392,7 @@ bool UnrealToUsd::ConvertXformable( const UMovieScene3DTransformTrack& MovieScen
 			FVector Scale( ScaleValuesX[ ValueIndex ].Value, ScaleValuesY[ ValueIndex ].Value, ScaleValuesZ[ ValueIndex ].Value );
 
 			FTransform Transform( Rotation, Location, Scale );
-			ConvertXformable( AdditionalRotation* Transform, UsdPrim, UsdFrameTime.AsDecimal() );
+			ConvertXformable( CameraCompensation * Transform * InverseCameraCompensation.Inverse(), UsdPrim, UsdFrameTime.AsDecimal() );
 
 			++ValueIndex;
 		}
