USD: Also add inverse camera transform compensation to code path for baking a level sequence to USD.

This should have been in cl 18308238.

diff --git a/Engine/Plugins/Importers/USDImporter/Source/USDUtilities/Private/USDPrimConversion.cpp b/Engine/Plugins/Importers/USDImporter/Source/USDUtilities/Private/USDPrimConversion.cpp
--- a/Engine/Plugins/Importers/USDImporter/Source/USDUtilities/Private/USDPrimConversion.cpp
+++ b/Engine/Plugins/Importers/USDImporter/Source/USDUtilities/Private/USDPrimConversion.cpp
@@ -2420,6 +2420,23 @@ bool UnrealToUsd::CreateComponentPropertyBaker( UE::FUsdPrim& Prim, const UScene
 				}
 			}
 
+			// Invert compensation applied to parent if it's a light or camera component
+			if ( const USceneComponent* AttachParent = Component.GetAttachParent() )
+			{
+				if ( AttachParent->IsA( UCineCameraComponent::StaticClass() ) ||
+					AttachParent->IsA( ULightComponent::StaticClass() ) )
+				{
+					FTransform InverseCompensation = FTransform( FRotator( 0.0f, 90.f, 0.0f ) );
+
+					if ( UsdUtils::GetUsdStageUpAxis( UsdStage ) == pxr::UsdGeomTokens->z )
+					{
+						InverseCompensation *= FTransform( FRotator( 90.0f, 0.f, 0.0f ) );
+					}
+
+					AdditionalRotation = AdditionalRotation * InverseCompensation.Inverse();
+				}
+			}
+
 			BakerType = EBakingType::Transform;
 			BakerFunction = [&Component, AdditionalRotation, StageInfo, Attr]( double UsdTimeCode )
 			{
