USD: Fix material assignment being incorrect if polygons are not sorted in the same order as the GeomSubsets

diff --git a/Engine/Plugins/Importers/USDImporter/Source/USDUtilities/Private/USDGeomMeshConversion.cpp b/Engine/Plugins/Importers/USDImporter/Source/USDUtilities/Private/USDGeomMeshConversion.cpp
--- a/Engine/Plugins/Importers/USDImporter/Source/USDUtilities/Private/USDGeomMeshConversion.cpp
+++ b/Engine/Plugins/Importers/USDImporter/Source/USDUtilities/Private/USDGeomMeshConversion.cpp
@@ -628,6 +628,8 @@ bool UsdToUnreal::ConvertGeomMesh( const pxr::UsdTyped& UsdSchema, FMeshDescript
 
 	const double TimeCodeValue = TimeCode.GetValue();
 
+	const int32 MaterialIndexOffset = MaterialAssignments.Slots.Num();
+
 	// Material assignments
 	const bool bProvideMaterialIndices = true;
 	UsdUtils::FUsdPrimMaterialAssignmentInfo LocalInfo = UsdUtils::GetPrimMaterialAssignments( UsdPrim, TimeCode, bProvideMaterialIndices, RenderContext );
@@ -639,7 +641,6 @@ bool UsdToUnreal::ConvertGeomMesh( const pxr::UsdTyped& UsdSchema, FMeshDescript
 	const int32 VertexOffset = MeshDescription.Vertices().Num();
 	const int32 VertexInstanceOffset = MeshDescription.VertexInstances().Num();
 	const int32 PolygonOffset = MeshDescription.Polygons().Num();
-	const int32 MaterialIndexOffset = MaterialAssignments.Slots.Num();
 
 	FStaticMeshAttributes StaticMeshAttributes( MeshDescription );
 
@@ -675,6 +676,22 @@ bool UsdToUnreal::ConvertGeomMesh( const pxr::UsdTyped& UsdSchema, FMeshDescript
 		TArray<FVertexInstanceID> CornerInstanceIDs;
 		TArray<FVertexID> CornerVerticesIDs;
 		int32 CurrentVertexInstanceIndex = 0;
+		TPolygonGroupAttributesRef<FName> MaterialSlotNames = StaticMeshAttributes.GetPolygonGroupMaterialSlotNames();
+
+		// Material slots
+		// Note that we always create these in the order they show up in LocalInfo: If we created these on-demand when parsing polygons (like before)
+		// we could run into polygons in a different order than the material slots and end up with different material assignments.
+		// We could use the StaticMesh's SectionInfoMap to unswitch things, but that's not available at runtime so we better do this here
+		for ( int32 LocalMaterialIndex = 0; LocalMaterialIndex < LocalInfo.Slots.Num(); ++LocalMaterialIndex )
+		{
+			const int32 CombinedMaterialIndex = MaterialIndexOffset + LocalMaterialIndex;
+
+			FPolygonGroupID NewPolygonGroup = MeshDescription.CreatePolygonGroup();
+			PolygonGroupMapping.Add( CombinedMaterialIndex, NewPolygonGroup );
+
+			// This is important for runtime, where the material slots are matched to LOD sections based on their material slot name
+			MaterialSlotNames[ NewPolygonGroup ] = *LexToString( NewPolygonGroup.GetValue() );
+		}
 
 		bool bFlipThisGeometry = false;
 
@@ -816,7 +833,6 @@ bool UsdToUnreal::ConvertGeomMesh( const pxr::UsdTyped& UsdSchema, FMeshDescript
 			OpacityInterpolation = OpacityPrimvar.GetInterpolation();
 		}
 
-		TPolygonGroupAttributesRef<FName> MaterialSlotNames = StaticMeshAttributes.GetPolygonGroupMaterialSlotNames();
 		for ( int32 PolygonIndex = 0; PolygonIndex < FaceCounts.size(); ++PolygonIndex )
 		{
 			int32 PolygonVertexCount = FaceCounts[PolygonIndex];
@@ -924,15 +940,6 @@ bool UsdToUnreal::ConvertGeomMesh( const pxr::UsdTyped& UsdSchema, FMeshDescript
 
 			const int32 CombinedMaterialIndex = MaterialIndexOffset + LocalMaterialIndex;
 
-			if ( !PolygonGroupMapping.Contains( CombinedMaterialIndex ) )
-			{
-				FPolygonGroupID NewPolygonGroup = MeshDescription.CreatePolygonGroup();
-				PolygonGroupMapping.Add( CombinedMaterialIndex, NewPolygonGroup );
-
-				// This is important for runtime, where the material slots are matched to LOD sections based on their material slot name
-				MaterialSlotNames[ NewPolygonGroup ] = *LexToString( NewPolygonGroup.GetValue() );
-			}
-
 			if ( bFlipThisGeometry )
 			{
 				for (int32 i = 0; i < CornerInstanceIDs.Num() / 2; ++i)
