USD: Show layer identifiers when hovering over the layer tree view items

diff --git a/Engine/Plugins/Importers/USDImporter/Source/USDStageEditor/Private/SUSDLayersTreeView.cpp b/Engine/Plugins/Importers/USDImporter/Source/USDStageEditor/Private/SUSDLayersTreeView.cpp
--- a/Engine/Plugins/Importers/USDImporter/Source/USDStageEditor/Private/SUSDLayersTreeView.cpp
+++ b/Engine/Plugins/Importers/USDImporter/Source/USDStageEditor/Private/SUSDLayersTreeView.cpp
@@ -32,12 +32,22 @@ public:
 	virtual TSharedRef< SWidget > GenerateWidget( const TSharedPtr< IUsdTreeViewItem > InTreeItem, const TSharedPtr< ITableRow > TableRow ) override
 	{
 		FUsdLayerViewModelRef TreeItem = StaticCastSharedRef< FUsdLayerViewModel >( InTreeItem.ToSharedRef() );
+		TWeakPtr<FUsdLayerViewModel> TreeItemWeak = TreeItem;
 
 		return SNew( SBox )
 			.VAlign( VAlign_Center )
 			[
 				SNew(STextBlock)
 				.Text( TreeItem, &FUsdLayerViewModel::GetDisplayName )
+				.ToolTipText_Lambda( [TreeItemWeak]
+				{
+					if ( TSharedPtr<FUsdLayerViewModel> PinnedTreeItem = TreeItemWeak.Pin() )
+					{
+						return FText::FromString( PinnedTreeItem->LayerIdentifier );
+					}
+
+					return FText::GetEmpty();
+				})
 			];
 	}
 };
