USD: Fix log statement using LogTemp instead of LogUsd.

diff --git a/Engine/Plugins/Importers/USDImporter/Source/USDUtilities/Private/USDLightConversion.cpp b/Engine/Plugins/Importers/USDImporter/Source/USDUtilities/Private/USDLightConversion.cpp
--- a/Engine/Plugins/Importers/USDImporter/Source/USDUtilities/Private/USDLightConversion.cpp
+++ b/Engine/Plugins/Importers/USDImporter/Source/USDUtilities/Private/USDLightConversion.cpp
@@ -651,7 +651,7 @@ bool UnrealToUsd::ConvertSkyLightComponent( const USkyLightComponent& LightCompo
 				FString FilePath = AssetImportData->GetFirstFilename();
 				if ( !FPaths::FileExists( FilePath ) )
 				{
-					UE_LOG(LogTemp, Warning, TEXT("Used '%s' as cubemap when converting SkyLightComponent '%s' onto prim '%s', but the cubemap does not exist on the filesystem!"),
+					UE_LOG(LogUsd, Warning, TEXT("Used '%s' as cubemap when converting SkyLightComponent '%s' onto prim '%s', but the cubemap does not exist on the filesystem!"),
 						*FilePath,
 						*LightComponent.GetPathName(),
 						*UsdToUnreal::ConvertPath(Prim.GetPrimPath())
