USD: Tweak mesh collapsing logic to not use computed visibility.

Delete useless IsAnimated overload.

diff --git a/Engine/Plugins/Importers/USDImporter/Source/USDStage/Private/USDStageActor.cpp b/Engine/Plugins/Importers/USDImporter/Source/USDStage/Private/USDStageActor.cpp
--- a/Engine/Plugins/Importers/USDImporter/Source/USDStage/Private/USDStageActor.cpp
+++ b/Engine/Plugins/Importers/USDImporter/Source/USDStage/Private/USDStageActor.cpp
@@ -1134,6 +1134,7 @@ UUsdPrimTwin* AUsdStageActor::ExpandPrim( const UE::FUsdPrim& Prim, FUsdSchemaTr
 	}
 
 	// Update the prim animated status
+#if USE_USD_SDK
 	if ( UsdUtils::IsAnimated( Prim ) )
 	{
 		if ( !PrimsToAnimate.Contains( UsdPrimTwin->PrimPath ) )
@@ -1143,6 +1144,7 @@ UUsdPrimTwin* AUsdStageActor::ExpandPrim( const UE::FUsdPrim& Prim, FUsdSchemaTr
 		}
 	}
 	else
+#endif // USE_USD_SDK
 	{
 		if ( PrimsToAnimate.Contains( UsdPrimTwin->PrimPath ) )
 		{
diff --git a/Engine/Plugins/Importers/USDImporter/Source/USDStageImporter/Private/USDStageImporter.cpp b/Engine/Plugins/Importers/USDImporter/Source/USDStageImporter/Private/USDStageImporter.cpp
--- a/Engine/Plugins/Importers/USDImporter/Source/USDStageImporter/Private/USDStageImporter.cpp
+++ b/Engine/Plugins/Importers/USDImporter/Source/USDStageImporter/Private/USDStageImporter.cpp
@@ -339,10 +339,12 @@ namespace UsdStageImporterImpl
 				Component->RegisterComponent();
 			}
 
+#if USE_USD_SDK
 			if (UsdUtils::IsAnimated(Prim))
 			{
 				ImportAnimation(ImportContext, Prim, Component);
 			}
+#endif // USE_USD_SDK
 		}
 	}
 
diff --git a/Engine/Plugins/Importers/USDImporter/Source/USDUtilities/Private/USDConversionUtils.cpp b/Engine/Plugins/Importers/USDImporter/Source/USDUtilities/Private/USDConversionUtils.cpp
--- a/Engine/Plugins/Importers/USDImporter/Source/USDUtilities/Private/USDConversionUtils.cpp
+++ b/Engine/Plugins/Importers/USDImporter/Source/USDUtilities/Private/USDConversionUtils.cpp
@@ -906,15 +906,6 @@ FString UsdUtils::GetAssetPathFromPrimPath( const FString& RootContentPath, cons
 }
 #endif // #if USE_USD_SDK
 
-bool UsdUtils::IsAnimated( const UE::FUsdPrim& Prim )
-{
-#if USE_USD_SDK
-	return IsAnimated( static_cast< const pxr::UsdPrim& >( Prim ) );
-#else
-	return false;
-#endif // #if USE_USD_SDK
-}
-
 TArray< UE::FUsdPrim > UsdUtils::GetAllPrimsOfType( const UE::FUsdPrim& StartPrim, const TCHAR* SchemaName )
 {
 	return GetAllPrimsOfType( StartPrim, SchemaName, []( const UE::FUsdPrim& ) { return false; } );
diff --git a/Engine/Plugins/Importers/USDImporter/Source/USDUtilities/Private/USDGeomMeshConversion.cpp b/Engine/Plugins/Importers/USDImporter/Source/USDUtilities/Private/USDGeomMeshConversion.cpp
--- a/Engine/Plugins/Importers/USDImporter/Source/USDUtilities/Private/USDGeomMeshConversion.cpp
+++ b/Engine/Plugins/Importers/USDImporter/Source/USDUtilities/Private/USDGeomMeshConversion.cpp
@@ -509,7 +509,7 @@ namespace UE
 				const TMap< FString, TMap<FString, int32> >& MaterialToPrimvarToUVIndex,
 				FMeshDescription& OutMeshDescription,
 				UsdUtils::FUsdPrimMaterialAssignmentInfo& OutMaterialAssignments,
-				bool bParseTransform
+				bool bIsFirstPrim
 			)
 			{
 				// Ignore meshes from disabled purposes
@@ -518,19 +518,41 @@ namespace UE
 					return true;
 				}
 
-				// Ignore invisible prims
-				if ( pxr::UsdGeomImageable UsdGeomImageable = pxr::UsdGeomImageable( Prim ) )
+				FTransform ChildTransform = CurrentTransform;
+
+				if ( !bIsFirstPrim )
 				{
-					if ( UsdGeomImageable.ComputeVisibility() == pxr::UsdGeomTokens->invisible )
+					// Ignore invisible child meshes.
+					//
+					// We used to compute visibility here and flat out ignore any invisible meshes. However, it could be that this mesh is invisible
+					// due to the first prim (the parentmost prim of the recursive calls) being invisible. If the first is invisible but animated
+					// then its possible it will become visible later, so if the child meshes are all invisible due to that fact alone then we should still
+					// consider them. If the first is invisible but *not* animated then we should still consider it in the same way, because that's sort
+					// of what you'd expect by calling ConvertGeomMeshHierarchy: We shouldn't just return nothing if the prim happens to be invisible.
+					// Besides, it could be that first is invisible due to itself having a parent that is invisible but has visibility animations:
+					// In that case we'd also want to generate meshes even if first is effectively invisible, since those parents can become visible
+					// later as well. The only case left is if first is invisible due having parents that are invisible and not animated: Checking for
+					// this would involve checking visibility and animations of all of its parents though, which is probably a bit too much, and like
+					// in the case where first itself is invisible and not animated, the caller may still expect to receive a valid mesh even if the
+					// prim's parents are invisible.
+					//
+					// The only case in which we'll truly discard invisible submeshes now is if they're invisible *by themselves*. If we're collapsing
+					// them then we know they're not animated either, so they will basically never be visible at all, at any time code.
+					//
+					// Note that if we were to ever manually set any of these back to visible again via the editor, the visibility changes are
+					// now resyncs and we'll reparse this entire asset, which will give us the chance to add them back to the collapsed mesh.
+					if ( pxr::UsdGeomImageable UsdGeomImageable = pxr::UsdGeomImageable( Prim ) )
 					{
-						return true;
+						if ( pxr::UsdAttribute VisibilityAttr = UsdGeomImageable.GetVisibilityAttr() )
+						{
+							pxr::TfToken VisibilityToken;
+							if ( VisibilityAttr.Get( &VisibilityToken ) && VisibilityToken == pxr::UsdGeomTokens->invisible )
+							{
+								return true;
+							}
+						}
 					}
-				}
 
-				FTransform ChildTransform = CurrentTransform;
-
-				if ( bParseTransform )
-				{
 					if ( pxr::UsdGeomXformable Xformable = pxr::UsdGeomXformable( Prim ) )
 					{
 						FTransform LocalChildTransform;
@@ -553,7 +575,7 @@ namespace UE
 						break;
 					}
 
-					const bool bChildParseTransform = true;
+					const bool bChildIsFirstPrim = false;
 
 					bSuccess &= RecursivelyCollapseChildMeshes(
 						ChildPrim,
@@ -564,7 +586,7 @@ namespace UE
 						MaterialToPrimvarToUVIndex,
 						OutMeshDescription,
 						OutMaterialAssignments,
-						bChildParseTransform
+						bChildIsFirstPrim
 					);
 				}
 
@@ -980,7 +1002,7 @@ bool UsdToUnreal::ConvertGeomMeshHierarchy(
 
 	// Skip the first transform, because if we're parsing Prim's mesh we don't want to bake its own transform into the vertices,
 	// as that transform will be used on components instead
-	const bool bParseTransform = false;
+	const bool bIsFirstPrim = true;
 
 	return UE::UsdGeomMeshConversion::Private::RecursivelyCollapseChildMeshes(
 		Prim,
@@ -991,7 +1013,7 @@ bool UsdToUnreal::ConvertGeomMeshHierarchy(
 		MaterialToPrimvarToUVIndex,
 		OutMeshDescription,
 		OutMaterialAssignments,
-		bParseTransform
+		bIsFirstPrim
 	);
 }
 
diff --git a/Engine/Plugins/Importers/USDImporter/Source/USDUtilities/Public/USDConversionUtils.h b/Engine/Plugins/Importers/USDImporter/Source/USDUtilities/Public/USDConversionUtils.h
--- a/Engine/Plugins/Importers/USDImporter/Source/USDUtilities/Public/USDConversionUtils.h
+++ b/Engine/Plugins/Importers/USDImporter/Source/USDUtilities/Public/USDConversionUtils.h
@@ -162,8 +162,6 @@ namespace UsdUtils
 	USDUTILITIES_API TArray< UE::FUsdPrim > GetAllPrimsOfType( const UE::FUsdPrim& StartPrim, const TCHAR* SchemaName );
 	USDUTILITIES_API TArray< UE::FUsdPrim > GetAllPrimsOfType( const UE::FUsdPrim& StartPrim, const TCHAR* SchemaName, TFunction< bool( const UE::FUsdPrim& ) > PruneChildren, const TArray<const TCHAR*>& ExcludeSchemaNames = {} );
 
-	USDUTILITIES_API bool IsAnimated( const UE::FUsdPrim& Prim );
-
 	/** Returns the time code for non-timesampled values. Usually a quiet NaN. */
 	USDUTILITIES_API double GetDefaultTimeCode();
 
