USD: Add support for "Held" interpolation mode.

diff --git a/Engine/Plugins/Importers/USDImporter/Source/USDStage/Private/USDStageActor.cpp b/Engine/Plugins/Importers/USDImporter/Source/USDStage/Private/USDStageActor.cpp
--- a/Engine/Plugins/Importers/USDImporter/Source/USDStage/Private/USDStageActor.cpp
+++ b/Engine/Plugins/Importers/USDImporter/Source/USDStage/Private/USDStageActor.cpp
@@ -552,6 +552,7 @@ struct FUsdStageActorImpl
 			TArray<FAnalyticsEventAttribute> EventAttributes;
 
 			EventAttributes.Emplace( TEXT( "InitialLoadSet" ), LexToString( (uint8)StageActor->InitialLoadSet ) );
+			EventAttributes.Emplace( TEXT( "InterpolationType" ), LexToString( (uint8)StageActor->InterpolationType) );
 			EventAttributes.Emplace( TEXT( "KindsToCollapse" ), LexToString( StageActor->KindsToCollapse ) );
 			EventAttributes.Emplace( TEXT( "PurposesToLoad" ), LexToString( StageActor->PurposesToLoad ) );
 			EventAttributes.Emplace( TEXT( "NaniteTriangleThreshold" ), LexToString( StageActor->NaniteTriangleThreshold ) );
@@ -632,6 +633,7 @@ TMap<UBlueprint*, FDelegateHandle> FRecompilationTracker::RecompilingBlueprints;
 
 AUsdStageActor::AUsdStageActor()
 	: InitialLoadSet( EUsdInitialLoadSet::LoadAll )
+	, InterpolationType( EUsdInterpolationType::Linear )
 	, KindsToCollapse( ( int32 ) ( EUsdDefaultKind::Component | EUsdDefaultKind::Subcomponent ) )
 	, PurposesToLoad( (int32) EUsdPurpose::Proxy )
 	, NaniteTriangleThreshold( (uint64) 1000000 )
@@ -1252,6 +1254,14 @@ void AUsdStageActor::SetInitialLoadSet( EUsdInitialLoadSet NewLoadSet )
 	LoadUsdStage();
 }
 
+void AUsdStageActor::SetInterpolationType( EUsdInterpolationType NewType )
+{
+	Modify();
+
+	InterpolationType = NewType;
+	LoadUsdStage();
+}
+
 void AUsdStageActor::SetKindsToCollapse( int32 NewKindsToCollapse )
 {
 	Modify();
@@ -1413,6 +1423,8 @@ void AUsdStageActor::OpenUsdStage()
 	{
 		UsdStage.SetEditTarget( UsdStage.GetRootLayer() );
 
+		UsdStage.SetInterpolationType( InterpolationType );
+
 		UsdListener.Register( UsdStage );
 
 #if USE_USD_SDK
@@ -1608,6 +1620,8 @@ void AUsdStageActor::LoadUsdStage()
 		return;
 	}
 
+	UsdStage.SetInterpolationType( InterpolationType );
+
 	ReloadAnimations();
 
 	// Make sure our PrimsToAnimate and the LevelSequenceHelper are kept in sync, because we'll use PrimsToAnimate to
@@ -2370,6 +2384,10 @@ void AUsdStageActor::HandlePropertyChangedEvent( FPropertyChangedEvent& Property
 	{
 		SetInitialLoadSet( InitialLoadSet );
 	}
+	else if ( PropertyName == GET_MEMBER_NAME_CHECKED( AUsdStageActor, InterpolationType ) )
+	{
+		SetInterpolationType( InterpolationType );
+	}
 	else if ( PropertyName == GET_MEMBER_NAME_CHECKED( AUsdStageActor, KindsToCollapse ) )
 	{
 		SetKindsToCollapse( KindsToCollapse );
diff --git a/Engine/Plugins/Importers/USDImporter/Source/USDStage/Public/USDStageActor.h b/Engine/Plugins/Importers/USDImporter/Source/USDStage/Public/USDStageActor.h
--- a/Engine/Plugins/Importers/USDImporter/Source/USDStage/Public/USDStageActor.h
+++ b/Engine/Plugins/Importers/USDImporter/Source/USDStage/Public/USDStageActor.h
@@ -44,6 +44,9 @@ public:
 	UPROPERTY(EditAnywhere, BlueprintReadOnly, Category = "USD")
 	EUsdInitialLoadSet InitialLoadSet;
 
+	UPROPERTY( EditAnywhere, BlueprintReadOnly, Category = "USD" )
+	EUsdInterpolationType InterpolationType;
+
 	/**
 	 * Whether to try to combine individual assets and components of the same type on a kind-per-kind basis,
 	 * like multiple Mesh prims into a single Static Mesh
@@ -69,6 +72,9 @@ public:
 	UFUNCTION(BlueprintCallable, Category = "USD", meta = (CallInEditor = "true"))
 	USDSTAGE_API void SetInitialLoadSet( EUsdInitialLoadSet NewLoadSet );
 
+	UFUNCTION( BlueprintCallable, Category = "USD", meta = ( CallInEditor = "true" ) )
+	USDSTAGE_API void SetInterpolationType( EUsdInterpolationType NewType );
+
 	UFUNCTION( BlueprintCallable, Category = "USD", meta = ( CallInEditor = "true" ) )
 	USDSTAGE_API void SetKindsToCollapse( int32 NewKindsToCollapse );
 
diff --git a/Engine/Plugins/Importers/USDImporter/Source/USDStageEditor/Private/SUSDStage.cpp b/Engine/Plugins/Importers/USDImporter/Source/USDStageEditor/Private/SUSDStage.cpp
--- a/Engine/Plugins/Importers/USDImporter/Source/USDStageEditor/Private/SUSDStage.cpp
+++ b/Engine/Plugins/Importers/USDImporter/Source/USDStageEditor/Private/SUSDStage.cpp
@@ -614,6 +614,11 @@ void SUsdStage::FillOptionsMenu(FMenuBuilder& MenuBuilder)
 			LOCTEXT( "Collapsing", "Collapsing" ),
 			LOCTEXT( "Collapsing_ToolTip", "Whether to try to combine individual assets and components of the same type on a Kind-per-Kind basis, like multiple Mesh prims into a single Static Mesh" ),
 			FNewMenuDelegate::CreateSP( this, &SUsdStage::FillCollapsingSubMenu ) );
+
+		MenuBuilder.AddSubMenu(
+			LOCTEXT( "InterpolationType", "Interpolation type" ),
+			LOCTEXT( "InterpolationType_ToolTip", "Whether to interpolate between time samples linearly or with 'held' (i.e. constant) interpolation" ),
+			FNewMenuDelegate::CreateSP( this, &SUsdStage::FillInterpolationTypeSubMenu ) );
 	}
 	MenuBuilder.EndSection();
 
@@ -926,6 +931,77 @@ void SUsdStage::FillCollapsingSubMenu( FMenuBuilder& MenuBuilder )
 	);
 }
 
+void SUsdStage::FillInterpolationTypeSubMenu(FMenuBuilder& MenuBuilder)
+{
+	MenuBuilder.AddMenuEntry(
+		LOCTEXT("LinearType", "Linear"),
+		LOCTEXT("LinearType_ToolTip", "Attribute values are linearly interpolated between authored values"),
+		FSlateIcon(),
+		FUIAction(
+			FExecuteAction::CreateLambda([this]()
+			{
+				if(AUsdStageActor* StageActor = ViewModel.UsdStageActor.Get())
+				{
+					FScopedTransaction Transaction(FText::Format(
+						LOCTEXT("SetLinearInterpolationType", "Set USD stage actor '{0}' to linear interpolation"),
+						FText::FromString(StageActor->GetActorLabel())
+					));
+
+					StageActor->SetInterpolationType( EUsdInterpolationType::Linear );
+				}
+			}),
+			FCanExecuteAction::CreateLambda([this]()
+			{
+				return ViewModel.UsdStageActor.Get() != nullptr;
+			}),
+			FIsActionChecked::CreateLambda([this]()
+			{
+				if(AUsdStageActor* StageActor = ViewModel.UsdStageActor.Get())
+				{
+					return StageActor->InterpolationType == EUsdInterpolationType::Linear;
+				}
+				return false;
+			})
+		),
+		NAME_None,
+		EUserInterfaceActionType::RadioButton
+	);
+
+	MenuBuilder.AddMenuEntry(
+		LOCTEXT("HeldType", "Held"),
+		LOCTEXT("HeldType_ToolTip", "Attribute values are held constant between authored values. An attribute's value will be equal to the nearest preceding authored value. If there is no preceding authored value, the value will be equal to the nearest subsequent value."),
+		FSlateIcon(),
+		FUIAction(
+			FExecuteAction::CreateLambda([this]()
+			{
+				if(AUsdStageActor* StageActor = ViewModel.UsdStageActor.Get())
+				{
+					FScopedTransaction Transaction(FText::Format(
+						LOCTEXT("SetHeldInterpolationType", "Set USD stage actor '{0}' to held interpolation"),
+						FText::FromString(StageActor->GetActorLabel())
+					));
+
+					StageActor->SetInterpolationType( EUsdInterpolationType::Held );
+				}
+			}),
+			FCanExecuteAction::CreateLambda([this]()
+			{
+				return ViewModel.UsdStageActor.Get() != nullptr;
+			}),
+			FIsActionChecked::CreateLambda([this]()
+			{
+				if(AUsdStageActor* StageActor = ViewModel.UsdStageActor.Get())
+				{
+					return StageActor->InterpolationType == EUsdInterpolationType::Held;
+				}
+				return false;
+			})
+		),
+		NAME_None,
+		EUserInterfaceActionType::RadioButton
+	);
+}
+
 void SUsdStage::FillSelectionSubMenu( FMenuBuilder& MenuBuilder )
 {
 	MenuBuilder.AddMenuEntry(
diff --git a/Engine/Plugins/Importers/USDImporter/Source/USDStageEditor/Private/SUSDStage.h b/Engine/Plugins/Importers/USDImporter/Source/USDStageEditor/Private/SUSDStage.h
--- a/Engine/Plugins/Importers/USDImporter/Source/USDStageEditor/Private/SUSDStage.h
+++ b/Engine/Plugins/Importers/USDImporter/Source/USDStageEditor/Private/SUSDStage.h
@@ -46,6 +46,7 @@ protected:
 	void FillPurposesToLoadSubMenu( FMenuBuilder& MenuBuilder );
 	void FillRenderContextSubMenu( FMenuBuilder& MenuBuilder );
 	void FillCollapsingSubMenu( FMenuBuilder& MenuBuilder );
+	void FillInterpolationTypeSubMenu( FMenuBuilder& MenuBuilder );
 	void FillSelectionSubMenu( FMenuBuilder& MenuBuilder );
 	void FillNaniteThresholdSubMenu( FMenuBuilder& MenuBuilder );
 
diff --git a/Engine/Plugins/Importers/USDImporter/Source/UnrealUSDWrapper/Private/UsdWrappers/UsdStage.cpp b/Engine/Plugins/Importers/USDImporter/Source/UnrealUSDWrapper/Private/UsdWrappers/UsdStage.cpp
--- a/Engine/Plugins/Importers/USDImporter/Source/UnrealUSDWrapper/Private/UsdWrappers/UsdStage.cpp
+++ b/Engine/Plugins/Importers/USDImporter/Source/UnrealUSDWrapper/Private/UsdWrappers/UsdStage.cpp
@@ -482,6 +482,35 @@ namespace UE
 #endif
 	}
 
+	template<typename PtrType>
+	void FUsdStageBase<PtrType>::SetInterpolationType( EUsdInterpolationType InterpolationType )
+	{
+#if USE_USD_SDK
+		if ( PtrType& Ptr = Impl->GetInner() )
+		{
+			FScopedUsdAllocs Allocs;
+
+			Ptr->SetInterpolationType( InterpolationType == EUsdInterpolationType::Held ? pxr::UsdInterpolationTypeHeld : pxr::UsdInterpolationTypeLinear );
+		}
+#endif
+	}
+
+	template<typename PtrType>
+	EUsdInterpolationType FUsdStageBase<PtrType>::GetInterpolationType() const
+	{
+#if USE_USD_SDK
+		if ( PtrType& Ptr = Impl->GetInner() )
+		{
+			if ( Ptr->GetInterpolationType() == pxr::UsdInterpolationTypeHeld )
+			{
+				return EUsdInterpolationType::Held;
+			}
+		}
+#endif
+
+		return EUsdInterpolationType::Linear;
+	}
+
 	template<typename PtrType>
 	void FUsdStageBase<PtrType>::SetDefaultPrim( const FUsdPrim & Prim )
 	{
diff --git a/Engine/Plugins/Importers/USDImporter/Source/UnrealUSDWrapper/Public/UnrealUSDWrapper.h b/Engine/Plugins/Importers/USDImporter/Source/UnrealUSDWrapper/Public/UnrealUSDWrapper.h
--- a/Engine/Plugins/Importers/USDImporter/Source/UnrealUSDWrapper/Public/UnrealUSDWrapper.h
+++ b/Engine/Plugins/Importers/USDImporter/Source/UnrealUSDWrapper/Public/UnrealUSDWrapper.h
@@ -173,6 +173,13 @@ enum class EUsdInitialLoadSet : uint8
 	LoadNone
 };
 
+UENUM()
+enum class EUsdInterpolationType : uint8
+{
+	Held,
+	Linear
+};
+
 class IUnrealUSDWrapperModule : public IModuleInterface
 {
 };
diff --git a/Engine/Plugins/Importers/USDImporter/Source/UnrealUSDWrapper/Public/UsdWrappers/UsdStage.h b/Engine/Plugins/Importers/USDImporter/Source/UnrealUSDWrapper/Public/UsdWrappers/UsdStage.h
--- a/Engine/Plugins/Importers/USDImporter/Source/UnrealUSDWrapper/Public/UsdWrappers/UsdStage.h
+++ b/Engine/Plugins/Importers/USDImporter/Source/UnrealUSDWrapper/Public/UsdWrappers/UsdStage.h
@@ -4,6 +4,7 @@
 
 #include "Templates/UniquePtr.h"
 
+#include "UnrealUsdWrapper.h"
 #include "UsdWrappers/ForwardDeclarations.h"
 
 namespace UE
@@ -88,6 +89,9 @@ namespace UE
 		double GetFramesPerSecond() const;
 		void SetFramesPerSecond( double FramesPerSecond );
 
+		void SetInterpolationType( EUsdInterpolationType InterpolationType );
+		EUsdInterpolationType GetInterpolationType() const;
+
 		void SetDefaultPrim( const FUsdPrim& Prim );
 
 		FUsdPrim OverridePrim( const FSdfPath& Path );
