USD: Fix prim renaming also renaming referenced/payload prims within their own layers;

Fix removing prims sometimes leaving them on the stage due to having specs in other layers;
Update the layer's default prim when we rename it to something else.

diff --git a/Engine/Plugins/Importers/USDImporter/Source/USDStageEditor/Private/SUSDStageTreeView.cpp b/Engine/Plugins/Importers/USDImporter/Source/USDStageEditor/Private/SUSDStageTreeView.cpp
--- a/Engine/Plugins/Importers/USDImporter/Source/USDStageEditor/Private/SUSDStageTreeView.cpp
+++ b/Engine/Plugins/Importers/USDImporter/Source/USDStageEditor/Private/SUSDStageTreeView.cpp
@@ -730,7 +730,7 @@ void SUsdStageTreeView::OnRemovePrim()
 	{
 		UE::FUsdStage Stage = UsdStageActor->GetOrLoadUsdStage();
 
-		UsdUtils::RemoveAllPrimSpecs( SelectedItem->UsdPrim, Stage.GetEditTarget() );
+		UsdUtils::RemoveAllPrimSpecs( SelectedItem->UsdPrim );
 	}
 }
 
diff --git a/Engine/Plugins/Importers/USDImporter/Source/USDUtilities/Private/USDConversionUtils.cpp b/Engine/Plugins/Importers/USDImporter/Source/USDUtilities/Private/USDConversionUtils.cpp
--- a/Engine/Plugins/Importers/USDImporter/Source/USDUtilities/Private/USDConversionUtils.cpp
+++ b/Engine/Plugins/Importers/USDImporter/Source/USDUtilities/Private/USDConversionUtils.cpp
@@ -1098,6 +1098,13 @@ bool UsdUtils::RenamePrim( UE::FUsdPrim& Prim, const TCHAR* NewPrimName )
 	pxr::TfToken NewNameToken = UnrealToUsd::ConvertToken( NewPrimName ).Get();
 	pxr::SdfPath TargetPath = PxrUsdPrim.GetPrimPath().ReplaceName( NewNameToken );
 
+	std::unordered_set<std::string> LocalLayerIdentifiers;
+	const bool bIncludeSessionLayers = true;
+	for ( const pxr::SdfLayerHandle& Handle : PxrUsdStage->GetLayerStack( bIncludeSessionLayers ) )
+	{
+		LocalLayerIdentifiers.insert( Handle->GetIdentifier() );
+	}
+
 	std::vector<pxr::SdfPrimSpecHandle> SpecStack = PxrUsdPrim.GetPrimStack();
 	TArray<TPair<pxr::SdfLayerRefPtr, pxr::SdfBatchNamespaceEdit>> Edits;
 
@@ -1120,6 +1127,16 @@ bool UsdUtils::RenamePrim( UE::FUsdPrim& Prim, const TCHAR* NewPrimName )
 
 		pxr::SdfLayerRefPtr SpecLayer = Spec->GetLayer();
 
+		// We should only rename specs on layers that are in the stage's *local* layer stack (which will include root, sublayers and
+		// session layers). We shouldn't rename any spec that is created due to references/payloads to other layers, because if we do
+		// we'll end up renaming the prims within those layers too, which is not what we want: For reference/payloads it's as if
+		// we're just consuming the *contents* of the referenced prim, but we don't want to affect it. Another more drastic example:
+		// if we were to remove the referencer prim, we don't really want to delete the referenced prim within its layer
+		if ( LocalLayerIdentifiers.count( SpecLayer->GetIdentifier() ) == 0 )
+		{
+			continue;
+		}
+
 		pxr::SdfBatchNamespaceEdit BatchEdit;
 		BatchEdit.Add( pxr::SdfNamespaceEdit::Rename( SpecPath, NewNameToken ) );
 
@@ -1164,6 +1181,11 @@ bool UsdUtils::RenamePrim( UE::FUsdPrim& Prim, const TCHAR* NewPrimName )
 			const pxr::SdfLayerRefPtr& Layer = Pair.Key;
 			const pxr::SdfBatchNamespaceEdit& Edit = Pair.Value;
 
+			// Make sure that if the renamed prim is the layer's default prim, we also update that to match the
+			// prim's new name
+			pxr::UsdPrim ParentPrim = PxrUsdPrim.GetParent();
+			const bool bNeedToRenameDefaultPrim = ParentPrim && ParentPrim.IsPseudoRoot() && ( PxrUsdPrim.GetName() == Layer->GetDefaultPrim() );
+
 			if ( !Layer->Apply( Edit ) )
 			{
 				// This should not be happening since CanApply was true, so stop doing whatever it is we're doing
@@ -1175,6 +1197,11 @@ bool UsdUtils::RenamePrim( UE::FUsdPrim& Prim, const TCHAR* NewPrimName )
 
 				return false;
 			}
+
+			if ( bNeedToRenameDefaultPrim )
+			{
+				Layer->SetDefaultPrim( NewNameToken );
+			}
 		}
 	}
 
@@ -1394,6 +1421,20 @@ USDUTILITIES_API void UsdUtils::RemoveAllPrimSpecs( const UE::FUsdPrim& Prim, co
 	pxr::SdfLayerRefPtr UsdLayer{ Layer };
 	pxr::UsdStageRefPtr UsdStage = UsdPrim.GetStage();
 
+	std::unordered_set<std::string> LocalLayerIdentifiers;
+
+	// We'll want to remove specs from the entire stage. We need to be careful though to only remove specs from the
+	// local layer stack. If a prim within the stage has a reference/payload to another layer and we remove the
+	// referencer prim, we don't want to end up removing the referenced/payload prim within its own layer too.
+	if ( !UsdLayer )
+	{
+		const bool bIncludeSessionLayers = true;
+		for ( const pxr::SdfLayerHandle& Handle : UsdStage->GetLayerStack( bIncludeSessionLayers ) )
+		{
+			LocalLayerIdentifiers.insert( Handle->GetIdentifier() );
+		}
+	}
+
 	for ( const pxr::SdfPrimSpecHandle& Spec : UsdPrim.GetPrimStack() )
 	{
 		// For whatever reason sometimes there are invalid specs in the layer stack, so we need to be careful
@@ -1408,12 +1449,19 @@ USDUTILITIES_API void UsdUtils::RemoveAllPrimSpecs( const UE::FUsdPrim& Prim, co
 			continue;
 		}
 
-		if ( UsdLayer && Spec->GetLayer() != UsdLayer )
+		pxr::SdfLayerRefPtr SpecLayer = Spec->GetLayer();
+		if ( UsdLayer && SpecLayer != UsdLayer )
+		{
+			continue;
+		}
+
+		if ( !UsdLayer && LocalLayerIdentifiers.count( SpecLayer->GetIdentifier() ) == 0 )
 		{
 			continue;
 		}
 
-		UE_LOG( LogUsd, Log, TEXT( "Removing prim spec '%s' from edit target '%s'" ), *UsdToUnreal::ConvertPath( SpecPath ), *UsdToUnreal::ConvertString( UsdLayer->GetIdentifier() ) );
+		UE_LOG( LogUsd, Log, TEXT( "Removing prim spec '%s' from layer '%s'" ), *UsdToUnreal::ConvertPath( SpecPath ), *UsdToUnreal::ConvertString( SpecLayer->GetIdentifier() ) );
+		pxr::UsdEditContext Context( UsdStage, SpecLayer );
 		UsdStage->RemovePrim( SpecPath );
 	}
 
diff --git a/Engine/Plugins/Importers/USDImporter/Source/USDUtilities/Public/USDConversionUtils.h b/Engine/Plugins/Importers/USDImporter/Source/USDUtilities/Public/USDConversionUtils.h
--- a/Engine/Plugins/Importers/USDImporter/Source/USDUtilities/Public/USDConversionUtils.h
+++ b/Engine/Plugins/Importers/USDImporter/Source/USDUtilities/Public/USDConversionUtils.h
@@ -9,6 +9,7 @@
 #include "UnrealUSDWrapper.h"
 #include "USDMemory.h"
 #include "UsdWrappers/ForwardDeclarations.h"
+#include "UsdWrappers/SdfLayer.h"
 
 #if USE_USD_SDK
 #include "USDIncludesStart.h"
@@ -245,7 +246,11 @@ namespace UsdUtils
 	 * meaning the prim may still be left on the stage. Note that it's even possible to have both of those specs at the same time:
 	 * for example when we have a prim inside a variant set, but outside of it we have overrides to the same prim. This function
 	 * will remove both.
+	 *
+	 * @param Prim - Prim to remove
+	 * @param Layer - Layer to remove prim specs from. This can be left with the invalid layer (default) in order to remove all
+	 *				  specs from the entire stage's local layer stack.
 	 */
-	USDUTILITIES_API void RemoveAllPrimSpecs( const UE::FUsdPrim& Prim, const UE::FSdfLayer& Layer );
+	USDUTILITIES_API void RemoveAllPrimSpecs( const UE::FUsdPrim& Prim, const UE::FSdfLayer& Layer = UE::FSdfLayer{} );
 }
 
