USD: Support checking for fallback values when failing to retrieve a texture from a shader.

diff --git a/Engine/Plugins/Importers/USDImporter/Source/USDUtilities/Private/USDShadeConversion.cpp b/Engine/Plugins/Importers/USDImporter/Source/USDUtilities/Private/USDShadeConversion.cpp
--- a/Engine/Plugins/Importers/USDImporter/Source/USDUtilities/Private/USDShadeConversion.cpp
+++ b/Engine/Plugins/Importers/USDImporter/Source/USDUtilities/Private/USDShadeConversion.cpp
@@ -656,6 +656,17 @@ namespace UE
 				{
 					if ( !GetTextureParameterValue( Input, TEXTUREGROUP_WorldSpecular, OutValue, Material, TexturesCache, PrimvarToUVIndex, bForceVirtualTextures ) )
 					{
+						// Check if we have a fallback input that we can use instead, since we don't have a valid texture value
+						if ( const pxr::UsdShadeInput FallbackInput = Source.GetInput( UnrealIdentifiers::Fallback ) )
+						{
+							float UsdFallbackFloat;
+							if ( FallbackInput.Get<float>( &UsdFallbackFloat ) )
+							{
+								OutValue.Set<float>( UsdFallbackFloat );
+								return true;
+							}
+						}
+
 						// Recurse because the attribute may just be pointing at some other attribute that has the data
 						// (e.g. when shader input is just "hoisted" and connected to the parent material input)
 						return GetFloatParameterValue( Source, SourceName, DefaultValue, OutValue, Material, TexturesCache, PrimvarToUVIndex, bForceVirtualTextures );
@@ -764,6 +775,31 @@ namespace UE
 				{
 					if ( !GetTextureParameterValue( Input, bIsNormalMap ? TEXTUREGROUP_WorldNormalMap : TEXTUREGROUP_World, OutValue, Material, TexturesCache, PrimvarToUVIndex, bForceVirtualTextures ) )
 					{
+						// Check if we have a fallback input that we can use instead, since we don't have a valid texture value
+						if ( const pxr::UsdShadeInput FallbackInput = Source.GetInput( UnrealIdentifiers::Fallback ) )
+						{
+							pxr::GfVec3f UsdFallbackVec3;
+							if ( FallbackInput.Get<pxr::GfVec3f>( &UsdFallbackVec3 ) )
+							{
+								OutValue.Set<FVector>( UsdToUnreal::ConvertVector( UsdFallbackVec3 ) );
+								return true;
+							}
+
+							pxr::GfVec4f UsdFallbackVec4;
+							if ( FallbackInput.Get<pxr::GfVec4f>( &UsdFallbackVec4 ) )
+							{
+								if ( !FMath::IsNearlyEqual( UsdFallbackVec4[ 3 ], 1.0f ) )
+								{
+									UE_LOG( LogUsd, Warning, TEXT( "Ignoring alpha value from fallback GfVec4f [%f, %f, %f, %f] used for Shader '%s'" ),
+										UsdFallbackVec4[ 0 ], UsdFallbackVec4[ 1 ], UsdFallbackVec4[ 2 ], UsdFallbackVec4[ 3 ], *UsdToUnreal::ConvertPath( Source.GetPrim().GetPrimPath() )
+									);
+								}
+
+								OutValue.Set<FVector>( FVector{ UsdFallbackVec4[ 0 ], UsdFallbackVec4[ 1 ], UsdFallbackVec4[ 2 ] } );
+								return true;
+							}
+						}
+
 						// Check whether this input receives its value through a connection to a
 						// primvar reader shader.
 						if ( !GetPrimvarReaderParameterValue( Input, UnrealIdentifiers::UsdPrimvarReader_float3, DefaultValue, OutValue ) )
