USD: Fix non-portable include filename.

diff --git a/Engine/Plugins/Importers/USDImporter/Source/UnrealUSDWrapper/Public/UsdWrappers/UsdStage.h b/Engine/Plugins/Importers/USDImporter/Source/UnrealUSDWrapper/Public/UsdWrappers/UsdStage.h
--- a/Engine/Plugins/Importers/USDImporter/Source/UnrealUSDWrapper/Public/UsdWrappers/UsdStage.h
+++ b/Engine/Plugins/Importers/USDImporter/Source/UnrealUSDWrapper/Public/UsdWrappers/UsdStage.h
@@ -4,7 +4,7 @@
 
 #include "Templates/UniquePtr.h"
 
-#include "UnrealUsdWrapper.h"
+#include "UnrealUSDWrapper.h"
 #include "UsdWrappers/ForwardDeclarations.h"
 
 namespace UE
