USD: Add overload to GetUVSetPrimvars that allows the caller to provide the material assignment info

diff --git a/Engine/Plugins/Importers/USDImporter/Source/USDUtilities/Private/USDConversionUtils.cpp b/Engine/Plugins/Importers/USDImporter/Source/USDUtilities/Private/USDConversionUtils.cpp
--- a/Engine/Plugins/Importers/USDImporter/Source/USDUtilities/Private/USDConversionUtils.cpp
+++ b/Engine/Plugins/Importers/USDImporter/Source/USDUtilities/Private/USDConversionUtils.cpp
@@ -578,13 +578,25 @@ TArray< TUsdStore< pxr::UsdGeomPrimvar > > UsdUtils::GetUVSetPrimvars( const pxr
 		return {};
 	}
 
+	const bool bProvideMaterialIndices = false;
+	UsdUtils::FUsdPrimMaterialAssignmentInfo Info = UsdUtils::GetPrimMaterialAssignments( UsdMesh.GetPrim(), pxr::UsdTimeCode( 0.0 ), bProvideMaterialIndices, RenderContext );
+	return UsdUtils::GetUVSetPrimvars( UsdMesh, MaterialToPrimvarsUVSetNames, Info );
+}
+
+TArray< TUsdStore< pxr::UsdGeomPrimvar > > UsdUtils::GetUVSetPrimvars( const pxr::UsdGeomMesh& UsdMesh, const TMap< FString, TMap< FString, int32 > >& MaterialToPrimvarsUVSetNames, const UsdUtils::FUsdPrimMaterialAssignmentInfo& UsdMeshMaterialAssignmentInfo )
+{
+	if ( !UsdMesh )
+	{
+		return {};
+	}
+
 	FScopedUsdAllocs Allocs;
 
 	// Collect all primvars that could be used as UV sets
 	TMap<FString, pxr::UsdGeomPrimvar> PrimvarsByName;
 	TMap<int32, TArray<pxr::UsdGeomPrimvar>> UsablePrimvarsByUVIndex;
 	pxr::UsdGeomPrimvarsAPI PrimvarsAPI{ UsdMesh };
-	for (const pxr::UsdGeomPrimvar& Primvar : PrimvarsAPI.GetPrimvars() )
+	for ( const pxr::UsdGeomPrimvar& Primvar : PrimvarsAPI.GetPrimvars() )
 	{
 		if ( !Primvar || !Primvar.HasValue() )
 		{
@@ -609,8 +621,7 @@ TArray< TUsdStore< pxr::UsdGeomPrimvar > > UsdUtils::GetUVSetPrimvars( const pxr
 	TMap<int32, TArray<pxr::UsdGeomPrimvar>> PrimvarsUsedByAssignedMaterialsPerUVIndex;
 	{
 		const bool bProvideMaterialIndices = false;
-		UsdUtils::FUsdPrimMaterialAssignmentInfo Info = UsdUtils::GetPrimMaterialAssignments( UsdMesh.GetPrim(), pxr::UsdTimeCode( 0.0 ), bProvideMaterialIndices, RenderContext );
-		for ( const FUsdPrimMaterialSlot& Slot : Info.Slots )
+		for ( const FUsdPrimMaterialSlot& Slot : UsdMeshMaterialAssignmentInfo.Slots )
 		{
 			if ( Slot.AssignmentType == EPrimAssignmentType::MaterialPrim )
 			{
diff --git a/Engine/Plugins/Importers/USDImporter/Source/USDUtilities/Private/USDGeomMeshConversion.cpp b/Engine/Plugins/Importers/USDImporter/Source/USDUtilities/Private/USDGeomMeshConversion.cpp
--- a/Engine/Plugins/Importers/USDImporter/Source/USDUtilities/Private/USDGeomMeshConversion.cpp
+++ b/Engine/Plugins/Importers/USDImporter/Source/USDUtilities/Private/USDGeomMeshConversion.cpp
@@ -767,7 +767,7 @@ bool UsdToUnreal::ConvertGeomMesh( const pxr::UsdTyped& UsdSchema, FMeshDescript
 
 		TArray< FUVSet > UVSets;
 
-		TArray< TUsdStore< UsdGeomPrimvar > > PrimvarsByUVIndex = UsdUtils::GetUVSetPrimvars( UsdMesh, MaterialToPrimvarsUVSetNames, RenderContext );
+		TArray< TUsdStore< UsdGeomPrimvar > > PrimvarsByUVIndex = UsdUtils::GetUVSetPrimvars( UsdMesh, MaterialToPrimvarsUVSetNames, LocalInfo );
 
 		int32 HighestAddedUVChannel = 0;
 		for ( int32 UVChannelIndex = 0; UVChannelIndex < PrimvarsByUVIndex.Num(); ++UVChannelIndex )
diff --git a/Engine/Plugins/Importers/USDImporter/Source/USDUtilities/Public/USDConversionUtils.h b/Engine/Plugins/Importers/USDImporter/Source/USDUtilities/Public/USDConversionUtils.h
--- a/Engine/Plugins/Importers/USDImporter/Source/USDUtilities/Public/USDConversionUtils.h
+++ b/Engine/Plugins/Importers/USDImporter/Source/USDUtilities/Public/USDConversionUtils.h
@@ -47,6 +47,10 @@ namespace UE
 	class FUsdPrim;
 	class FSdfPath;
 }
+namespace UsdUtils
+{
+	struct FUsdPrimMaterialAssignmentInfo;
+}
 
 namespace UsdUtils
 {
@@ -140,10 +144,12 @@ namespace UsdUtils
 	 * This overload will only return primvars with 'texcoord2f' role.	 *
 	 * @param UsdMesh - Mesh that contains primvars that can be used as texture coordinates.
 	 * @param MaterialToPrimvarsUVSetNames - Maps from a material prim path, to pairs indicating which primvar names are used as 'st' coordinates, and which UVIndex the imported material will sample from (e.g. ["st0", 0], ["myUvSet2", 2], etc). These are supposed to be the materials used by the mesh, and we do this because it helps identify which primvars are valid/used as texture coordinates, as the user may have these named as 'myUvSet2' and still expect it to work
+	 * @param UsdMeshMaterialAssignmentInfo - Result of calling GetPrimMaterialAssignments on UsdMesh's Prim. This can be provided or will be retrieved on-demand using RenderContext
 	 * @param RenderContext - Render context to use when traversing through material shaders looking for used primvars
 	 * @return Array where each index gives the primvar that should be used for that UV index
 	 */
 	USDUTILITIES_API TArray< TUsdStore< pxr::UsdGeomPrimvar > > GetUVSetPrimvars( const pxr::UsdGeomMesh& UsdMesh, const TMap< FString, TMap< FString, int32 > >& MaterialToPrimvarsUVSetNames, const pxr::TfToken& RenderContext = pxr::UsdShadeTokens->universalRenderContext );
+	USDUTILITIES_API TArray< TUsdStore< pxr::UsdGeomPrimvar > > GetUVSetPrimvars( const pxr::UsdGeomMesh& UsdMesh, const TMap< FString, TMap< FString, int32 > >& MaterialToPrimvarsUVSetNames, const UsdUtils::FUsdPrimMaterialAssignmentInfo& UsdMeshMaterialAssignmentInfo );
 
 	USDUTILITIES_API bool IsAnimated( const pxr::UsdPrim& Prim );
 
