Fixed GeometryCacheUSDStream compile error on targets that don't USE_USD_SDK.

diff --git a/Engine/Plugins/Importers/USDImporter/Source/GeometryCacheUSD/Private/GeometryCacheUSDStream.cpp b/Engine/Plugins/Importers/USDImporter/Source/GeometryCacheUSD/Private/GeometryCacheUSDStream.cpp
--- a/Engine/Plugins/Importers/USDImporter/Source/GeometryCacheUSD/Private/GeometryCacheUSDStream.cpp
+++ b/Engine/Plugins/Importers/USDImporter/Source/GeometryCacheUSD/Private/GeometryCacheUSDStream.cpp
@@ -39,12 +39,14 @@ private:
 public:
 	static FString GetUsdStreamDDCKey(const UE::FUsdStage& Stage, const FString& PrimPath, int32 FrameIndex)
 	{
+#ifdef USE_USD_SDK
 		FString PrimHash = UsdUtils::HashGeomMeshPrim(Stage, PrimPath, FrameIndex);
 
 		if (!PrimHash.IsEmpty())
 		{
 			return BuildDerivedDataKey(PrimHash);
 		}
+#endif
 		return {};
 	}
 };
@@ -90,6 +92,7 @@ void FGeometryCacheUsdStream::GetMeshData(int32 FrameIndex, int32 ConcurrencyInd
 {
 	TRACE_CPUPROFILER_EVENT_SCOPE(FGeometryCacheUsdStream::GetMeshData);
 
+#ifdef USE_USD_SDK
 	// Get the mesh data straight from the Alembic file or from the DDC if it's already cached
 	if (GUsdStreamCacheInDDC)
 	{
@@ -117,6 +120,7 @@ void FGeometryCacheUsdStream::GetMeshData(int32 FrameIndex, int32 ConcurrencyInd
 		// If the key is empty, it means the prim was invalid or not a mesh, so don't do anything
 	}
 	else
+#endif
 	{
 		// Synchronously load the requested frame data
 		ReadFunc(UsdTrack, FrameIndex, OutMeshData);
