USD: Export DisplayRate as FramesPerSecond when exporting LevelSequences to USD.

diff --git a/Engine/Plugins/Importers/USDImporter/Source/USDExporter/Private/LevelSequenceExporterUSD.cpp b/Engine/Plugins/Importers/USDImporter/Source/USDExporter/Private/LevelSequenceExporterUSD.cpp
--- a/Engine/Plugins/Importers/USDImporter/Source/USDExporter/Private/LevelSequenceExporterUSD.cpp
+++ b/Engine/Plugins/Importers/USDImporter/Source/USDExporter/Private/LevelSequenceExporterUSD.cpp
@@ -664,6 +664,10 @@ namespace UE
 					UsdStage.SetTimeCodesPerSecond( Context.ExportOptions->TimeCodesPerSecond );
 				}
 
+				// Set this so that if we open this exported sequence back up in UE the Sequencer will start up showing
+				// the same DisplayRate as the original exported sequence
+				UsdStage.SetFramesPerSecond( MovieScene->GetDisplayRate().AsDecimal() );
+
 				UE::FUsdPrim RootPrim = UsdStage.OverridePrim( UE::FSdfPath( TEXT( "/Root" ) ) );
 				if ( !RootPrim )
 				{
