USD: Fix strange behavior when adding the same layer as sublayer to the root and session layers.

Fix layers never saving if they were sublayers of session layers;
Fix layers not being shown as dirty if they were sublayers to session layers.

diff --git a/Engine/Plugins/Importers/USDImporter/Source/USDStageEditorViewModels/Private/USDLayersViewModel.cpp b/Engine/Plugins/Importers/USDImporter/Source/USDStageEditorViewModels/Private/USDLayersViewModel.cpp
--- a/Engine/Plugins/Importers/USDImporter/Source/USDStageEditorViewModels/Private/USDLayersViewModel.cpp
+++ b/Engine/Plugins/Importers/USDImporter/Source/USDStageEditorViewModels/Private/USDLayersViewModel.cpp
@@ -149,7 +149,7 @@ void FUsdLayerViewModel::RefreshData()
 	pxr::SdfLayerRefPtr UsdLayer = UE::FSdfLayer::FindOrOpen( *LayerIdentifier );
 	if ( UsdLayer )
 	{
-		LayerModel->bIsDirty = ( UsdLayer->IsDirty() || UsdLayer->IsAnonymous() ) && !UsdUtils::IsSessionLayerWithinStage( UsdLayer, UsdStage );
+		LayerModel->bIsDirty = UsdLayer->IsDirty() && !UsdLayer->IsAnonymous();
 	}
 
 	const pxr::SdfLayerHandle& EditTargetLayer = UsdStageRef->GetEditTarget().GetLayer();
diff --git a/Engine/Plugins/Importers/USDImporter/Source/USDStageEditorViewModels/Private/USDStageViewModel.cpp b/Engine/Plugins/Importers/USDImporter/Source/USDStageEditorViewModels/Private/USDStageViewModel.cpp
--- a/Engine/Plugins/Importers/USDImporter/Source/USDStageEditorViewModels/Private/USDStageViewModel.cpp
+++ b/Engine/Plugins/Importers/USDImporter/Source/USDStageEditorViewModels/Private/USDStageViewModel.cpp
@@ -213,7 +213,17 @@ void FUsdStageViewModel::SaveStage()
 
 			UsdUtils::StartMonitoringErrors();
 
-			pxr::UsdStageRefPtr( UsdStage )->Save();
+			// Save layers manually instead of calling UsdStage::Save(). This is roughly the same implementation anyway, except
+			// that UsdStage::Save() will ignore a layer if it also happens to be added as a sublayer to a session layer, and
+			// we want to ensure we always save dirty layers when we hit SaveStage
+			for ( const pxr::SdfLayerHandle& Handle : pxr::UsdStageRefPtr{ UsdStage }->GetUsedLayers() )
+			{
+				if ( !Handle->IsAnonymous() && Handle->IsDirty() )
+				{
+					Handle->Save();
+				}
+			}
+
 			UsdViewModelImpl::SaveUEStateLayer( UsdStage );
 
 			UsdUtils::ShowErrorsAndStopMonitoring(LOCTEXT("USDSaveError", "Failed to save current USD Stage!\nCheck the Output Log for details."));
