USD: Fix name collisions when exporting spawnables with the same name

diff --git a/Engine/Plugins/Importers/USDImporter/Source/USDExporter/Private/LevelSequenceExporterUSD.cpp b/Engine/Plugins/Importers/USDImporter/Source/USDExporter/Private/LevelSequenceExporterUSD.cpp
--- a/Engine/Plugins/Importers/USDImporter/Source/USDExporter/Private/LevelSequenceExporterUSD.cpp
+++ b/Engine/Plugins/Importers/USDImporter/Source/USDExporter/Private/LevelSequenceExporterUSD.cpp
@@ -17,7 +17,9 @@
 #include "UsdWrappers/UsdStage.h"
 
 #include "AssetExportTask.h"
+#include "Compilation/MovieSceneCompiledDataManager.h"
 #include "Components/SkeletalMeshComponent.h"
+#include "Evaluation/MovieSceneSequenceHierarchy.h"
 #include "Editor.h"
 #include "ISequencer.h"
 #include "ISequencerModule.h"
@@ -55,12 +57,116 @@ namespace UE
 
 				virtual UObject* SpawnObject( FMovieSceneSpawnable& Spawnable, FMovieSceneSequenceIDRef TemplateID, IMovieScenePlayer& Player ) override
 				{
-					UObject* Object = ExistingSpawns.FindRef( Spawnable.GetGuid() );
+					ISequencer* Sequencer = static_cast<ISequencer*>(&Player);
+					UMovieSceneSequence* RootSequence = Sequencer->GetRootMovieSceneSequence();
+					if ( !RootSequence )
+					{
+						return nullptr;
+					}
+
+					UObject* Object = nullptr;
+
+					const FGuid& Guid = Spawnable.GetGuid();
+
+					TArray<UObject*>& ExistingInstancesForGuid = SpawnableInstances.FindOrAdd( Guid );
+
+					TMap<FMovieSceneSequenceID, TMap<FGuid, int32>>& SequenceInstanceToSpawnableIndices = RootSequenceToSpawnableInstanceIndices.FindOrAdd( RootSequence );
+					TMap<FGuid, int32>& SpawnableIndices = SequenceInstanceToSpawnableIndices.FindOrAdd( TemplateID );
+
+					// Already have an instance of this spawnable for this movie scene sequence instance
+					if ( int32* ExistingIndex = SpawnableIndices.Find( Guid ) )
+					{
+						Object = ExistingInstancesForGuid[ *ExistingIndex ];
+
+						UE_LOG( LogUsd, VeryVerbose, TEXT( "Already spawned '%s' (%0x) for RootSequence '%s', TemplateID '%u', Guid '%s' (index %d)" ),
+							*Object->GetPathName(),
+							Object,
+							*RootSequence->GetPathName(),
+							TemplateID.GetInternalValue(),
+							*Guid.ToString(),
+							*ExistingIndex
+						);
+					}
 
+					// We don't have an instance of the spawnable spawned for this exact movie sequence ID, but try to see if we can
+					// reuse any of the existing spawns for it.
+					// We have to reuse these because the exported level will just contain prims that correspond to the spawnables spawned for
+					// the top-level level sequence. If we don't reuse these, recursive exports of the subsequences would attempt to spawn
+					// their own spawnables (as we always need an instance for each movie scene sequence ID), which wouldn't correspond to
+					// any prim that was exported on the level as the actor names would get sanitized away from any name used by the top-level export
+					if ( !Object && ExistingInstancesForGuid.Num() > 0 )
+					{
+						TArray<bool> UsedIndices;
+						UsedIndices.SetNumZeroed( ExistingInstancesForGuid.Num() );
+
+						for ( const TPair<FMovieSceneSequenceID, TMap<FGuid, int32>>& Pair : SequenceInstanceToSpawnableIndices )
+						{
+							if ( const int32* UsedIndex = Pair.Value.Find( Guid ) )
+							{
+								UsedIndices[ *UsedIndex ] = true;
+							}
+						}
+
+						int32 IndexToReuse = INDEX_NONE;
+						for ( int32 Index = 0; Index < UsedIndices.Num(); ++Index )
+						{
+							if ( !UsedIndices[ Index ] )
+							{
+								IndexToReuse = Index;
+								break;
+							}
+						}
+
+						if ( IndexToReuse != INDEX_NONE )
+						{
+							Object = ExistingInstancesForGuid[ IndexToReuse ];
+							SpawnableIndices.Add( Guid, IndexToReuse );
+
+							UE_LOG( LogUsd, VeryVerbose, TEXT( "Reusing '%s' (%0x) for RootSequence '%s', TemplateID '%u', Guid '%s' (index %d)" ),
+								*Object->GetPathName(),
+								Object,
+								*RootSequence->GetPathName(),
+								TemplateID.GetInternalValue(),
+								*Guid.ToString(),
+								IndexToReuse
+							);
+						}
+					}
+
+					// Don't even have anything we can reuse: We need to spawn a brand new instance of this spawnable
 					if ( !Object )
 					{
 						Object = FLevelSequenceEditorSpawnRegister::SpawnObject( Spawnable, TemplateID, Player );
-						ExistingSpawns.Add( Spawnable.GetGuid(), Object );
+						UE_LOG( LogUsd, VeryVerbose, TEXT( "Spawning '%s' (%0x) for RootSequence '%s', TemplateID '%u', Guid '%s'" ),
+							*Object->GetPathName(),
+							Object,
+							*RootSequence->GetPathName(),
+							TemplateID.GetInternalValue(),
+							*Guid.ToString()
+						);
+
+						if ( AActor* SpawnedActor = Cast<AActor>( Object ) )
+						{
+							// HACK: Remove the SequencerActor tag from the spawned actor because the USD level exporter
+							// uses these as an indicator that it shouldn't export a given actor, because the transient actors spawned by the
+							// sequencer and the AUsdStageActor contain them.
+							// We will manually delete all of those actors with our spawn register later and never actually use them
+							// for anything else though, so this should be fine
+							SpawnedActor->Tags.Remove( TEXT( "SequencerActor" ) );
+
+							// Rename the spawn to a unique name or else in case of name collisions they will overwrite each other when writing animation data.
+							// The level exporter will rename actors to unique prims by itself though.
+							FString NewLabel = UsdUtils::GetUniqueName( SpawnedActor->GetActorLabel(), UsedActorLabels );
+							if ( NewLabel != SpawnedActor->GetActorLabel() )
+							{
+								const bool bMarkDirty = false;
+								SpawnedActor->SetActorLabel( NewLabel, bMarkDirty );
+							}
+							UsedActorLabels.Add( NewLabel );
+						}
+
+						ExistingInstancesForGuid.Add( Object );
+						SpawnableIndices.Add( Guid, ExistingInstancesForGuid.Num() - 1 );
 					}
 
 					if ( Object )
@@ -118,11 +224,46 @@ namespace UE
 					}
 					else
 					{
-						FLevelSequenceEditorSpawnRegister::DestroySpawnedObject( Object );
-						if ( const FGuid* ExistingKey = ExistingSpawns.FindKey( &Object ) )
+						// We shouldn't need to do this because we only ever fully delete when we're cleaning up,
+						// and by then we'll delete all of these maps anyway
+						for ( TPair<FGuid, TArray<UObject*>>& Pair : SpawnableInstances )
 						{
-							ExistingSpawns.Remove( *ExistingKey );
+							const FGuid& Guid = Pair.Key;
+							TArray<UObject*>& InstancesForGuid = Pair.Value;
+
+							int32 IndexToDelete = INDEX_NONE;
+							for ( int32 Index = 0; Index < InstancesForGuid.Num(); ++Index )
+							{
+								if ( InstancesForGuid[ Index ] == &Object )
+								{
+									IndexToDelete = Index;
+									break;
+								}
+							}
+
+							if ( IndexToDelete != INDEX_NONE )
+							{
+								for ( TPair<const UMovieSceneSequence*, TMap<FMovieSceneSequenceID, TMap<FGuid, int32>>>& RootSequencePair : RootSequenceToSpawnableInstanceIndices )
+								{
+									for ( TPair<FMovieSceneSequenceID, TMap<FGuid, int32>>& SequenceIDPair : RootSequencePair.Value )
+									{
+										TMap<FGuid, int32>& GuidToInstance = SequenceIDPair.Value;
+										if ( int32* InstanceIndex = GuidToInstance.Find( Guid ) )
+										{
+											if ( *InstanceIndex == IndexToDelete )
+											{
+												GuidToInstance.Remove( Guid );
+											}
+										}
+									}
+								}
+
+								InstancesForGuid.RemoveAt( IndexToDelete );
+								break;
+							}
 						}
+
+						FLevelSequenceEditorSpawnRegister::DestroySpawnedObject( Object );
 					}
 				}
 
@@ -138,17 +279,16 @@ namespace UE
 					// alive when bDestroyingJustHides=true. Because of this we must explicitly clean up these "abandoned" spawns here,
 					// which resynchronizes us with Register:
 					TArray<UObject*> ObjectsToDelete;
-					ObjectsToDelete.Reserve( ExistingSpawns.Num() );
-					for ( const TPair<FGuid, UObject*>& Pair : ExistingSpawns )
+					for ( TPair<FGuid, TArray<UObject*>>& Pair : SpawnableInstances )
 					{
-						if ( UObject* Spawn = Pair.Value )
-						{
-							ObjectsToDelete.Add( Spawn );
-						}
+						ObjectsToDelete.Append( Pair.Value );
 					}
 					for ( UObject* Object : ObjectsToDelete )
 					{
-						DestroySpawnedObject( *Object );
+						if ( Object )
+						{
+							DestroySpawnedObject( *Object );
+						}
 					}
 				}
 
@@ -156,14 +296,41 @@ namespace UE
 				// of deleting them (when we want it to do so). Unfortunately, the base FMovieSceneSpawnRegister part
 				// will still nevertheless clear it's Register entry for the spawnable when deleting (even if just hiding),
 				// and there's nothing we can do to prevent it. This means we can't call FindSpawnedObject and must use our
-				// own GetExistingSpawn and ExistingSpawns member
-				UObject* GetExistingSpawn( const FGuid& Guid )
+				// own GetExistingSpawn and data members
+				UObject* GetExistingSpawn( const UMovieSceneSequence& RootSequence, FMovieSceneSequenceID SequenceID, const FGuid& SpawnableGuid )
 				{
-					return ExistingSpawns.FindRef( Guid );
+					int32 SpawnableIndex = INDEX_NONE;
+					if ( TMap<FMovieSceneSequenceID, TMap<FGuid, int32>>* SequenceIDsToSpawns = RootSequenceToSpawnableInstanceIndices.Find( &RootSequence ) )
+					{
+						if ( TMap<FGuid, int32>* Spawns = SequenceIDsToSpawns->Find( SequenceID ) )
+						{
+							if ( int32* Index = Spawns->Find( SpawnableGuid ) )
+							{
+								SpawnableIndex = *Index;
+							}
+						}
+					}
+
+					if ( SpawnableIndex != INDEX_NONE )
+					{
+						if ( TArray<UObject*>* Instances = SpawnableInstances.Find( SpawnableGuid ) )
+						{
+							return ( *Instances )[ SpawnableIndex ];
+						}
+					}
+
+					return nullptr;
 				}
 
 			private:
-				TMap<FGuid, UObject*> ExistingSpawns;
+				// Ensures all of our new spawns have unique names
+				TSet<FString> UsedActorLabels;
+
+				// Tracks all instances we created for a given spawnable guid
+				TMap<FGuid, TArray<UObject*>> SpawnableInstances;
+
+				// Tracks the indices into SpawnableInstances for each spawnable guid, used by each sequence ID, in the hierarchy of each root sequence
+				TMap<const UMovieSceneSequence*, TMap<FMovieSceneSequenceID, TMap<FGuid, int32>>> RootSequenceToSpawnableInstanceIndices;
 			};
 
 			// Contain all of the baker lambda functions for a given component. Only one baker per baking type is allowed.
@@ -219,17 +386,16 @@ namespace UE
 				}
 			};
 
-			// Returns a spawned UObject for a given spawnable from a MovieScene.
-			// This assumes it has already been spawned and is tracked by SpawnRegister.
-			UObject* LocateBoundObject( const UMovieScene* MovieScene, FMovieSceneSpawnable& Spawnable, const TSharedRef<FLevelSequenceHidingSpawnRegister>& SpawnRegister )
-			{
-				// Everything should have been spawned by now so just check our spawn register
-				return SpawnRegister->GetExistingSpawn( Spawnable.GetGuid() );
-			}
-
 			// Returns the UObject possessed by a possessable from a MovieScene
-			UObject* LocateBoundObject( const UMovieSceneSequence& MovieSceneSequence, UMovieScene* MovieScene, const FMovieScenePossessable& Possessable, const TSharedRef<FLevelSequenceHidingSpawnRegister>& SpawnRegister )
+			UObject* LocateBoundObject(
+				const UMovieSceneSequence& RootMovieSceneSequence,
+				const UMovieSceneSequence& CurrentMovieSceneSequence,
+				FMovieSceneSequenceIDRef TemplateID,
+				const FMovieScenePossessable& Possessable,
+				const TSharedRef<FLevelSequenceHidingSpawnRegister>& SpawnRegister
+			)
 			{
+				UMovieScene* MovieScene = CurrentMovieSceneSequence.GetMovieScene();
 				const FGuid& Guid = Possessable.GetGuid();
 				const FGuid& ParentGuid = Possessable.GetParent();
 
@@ -240,15 +406,15 @@ namespace UE
 				{
 					if ( FMovieScenePossessable* ParentPossessable = MovieScene->FindPossessable( ParentGuid ) )
 					{
-						ParentContext = LocateBoundObject( MovieSceneSequence, MovieScene, *ParentPossessable, SpawnRegister );
+						ParentContext = LocateBoundObject( RootMovieSceneSequence, CurrentMovieSceneSequence, TemplateID, *ParentPossessable, SpawnRegister );
 					}
 					else if ( FMovieSceneSpawnable* ParentSpawnable = MovieScene->FindSpawnable( ParentGuid ) )
 					{
-						ParentContext = LocateBoundObject( MovieScene, *ParentSpawnable, SpawnRegister );
+						ParentContext = SpawnRegister->GetExistingSpawn( RootMovieSceneSequence, TemplateID, ParentSpawnable->GetGuid() );
 					}
 				}
 
-				TArray<UObject*, TInlineAllocator<1>> Objects = MovieSceneSequence.LocateBoundObjects( Guid, ParentContext );
+				TArray<UObject*, TInlineAllocator<1>> Objects = CurrentMovieSceneSequence.LocateBoundObjects( Guid, ParentContext );
 				if ( Objects.Num() > 0 )
 				{
 					return Objects[ 0 ];
@@ -277,71 +443,79 @@ namespace UE
 				return false;
 			}
 
-			// Recursively go over Sequence and spawn all bound spawnables and then hide them
-			void PreSpawnSpawnables( FLevelSequenceExportContext& Context, UMovieSceneSequence& Sequence )
+			TMap< UMovieSceneSequence*, TArray<FMovieSceneSequenceID> > GetSequenceHierarchyInstances( UMovieSceneSequence& Sequence, ISequencer& Sequencer )
 			{
-				UMovieScene* MovieScene = Sequence.GetMovieScene();
-				if ( !MovieScene )
-				{
-					return;
-				}
+				TMap< UMovieSceneSequence*, TArray<FMovieSceneSequenceID> > SequenceInstances;
 
-				// Spawn everything on this movie scene
-				int32 NumSpawnables = MovieScene->GetSpawnableCount();
-				for ( int32 Index = 0; Index < NumSpawnables; ++Index )
-				{
-					const FMovieSceneSpawnable& Spawnable = MovieScene->GetSpawnable( Index );
-					const FGuid& Guid = Spawnable.GetGuid();
+				FMovieSceneSequenceHierarchy SequenceHierarchyCache;
+				UMovieSceneCompiledDataManager::CompileHierarchy( &Sequence, &SequenceHierarchyCache, EMovieSceneServerClientMask::All );
 
-					UObject* SpawnedObject = StaticCastSharedRef<FMovieSceneSpawnRegister>( Context.SpawnRegister )->SpawnObject( Guid, *MovieScene, Context.Sequencer->GetFocusedTemplateID(), *Context.Sequencer );
-
-					// HACK: Remove the SequencerActor tag from the spawned actor because the USD level exporter
-					// uses these as an indicator that it shouldn't export a given actor, because the transient actors spawned by the
-					// sequencer and the AUsdStageActor contain them.
-					// We will manually delete all of those actors with our spawn register later and never actually use them
-					// for anything else though, so this should be fine
-					if ( AActor* SpawnedActor = Cast<AActor>( SpawnedObject ) )
+				SequenceInstances.FindOrAdd( &Sequence ).Add( Sequencer.GetRootTemplateID() );
+				for ( const TTuple< FMovieSceneSequenceID, FMovieSceneSubSequenceData >& Pair : SequenceHierarchyCache.AllSubSequenceData() )
+				{
+					if ( ULevelSequence* SubSequence = Cast<ULevelSequence>( Pair.Value.GetSequence() ) )
 					{
-						SpawnedActor->Tags.Remove( TEXT( "SequencerActor" ) );
+						SequenceInstances.FindOrAdd( SubSequence ).Add( Pair.Key );
 					}
 				}
 
-				// Hide all our spawns again so we can pretend they haven't actually spawned yet
-				Context.SpawnRegister->bDestroyingJustHides = true;
-				Context.SpawnRegister->CleanUp( *Context.Sequencer );
+				return SequenceInstances;
+			}
 
-				// Recurse to subsequences
-				for ( const UMovieSceneTrack* MasterTrack : MovieScene->GetMasterTracks() )
+			// Spawn and hide all spawnables for Sequence hierarchy
+			TMap< UMovieSceneSequence*, TArray<FMovieSceneSequenceID> > PreSpawnSpawnables( FLevelSequenceExportContext& Context, UMovieSceneSequence& RootSequence )
+			{
+				TMap< UMovieSceneSequence*, TArray<FMovieSceneSequenceID> > SequenceInstances = GetSequenceHierarchyInstances( RootSequence, Context.Sequencer.Get() );
+
+				Context.Sequencer->ResetToNewRootSequence( RootSequence );
+
+				// Spawn everything for this Sequence hierarchy
+				for ( TPair< UMovieSceneSequence*, TArray<FMovieSceneSequenceID> >& SequenceInstancePair : SequenceInstances )
 				{
-					const UMovieSceneSubTrack* SubTrack = Cast<const UMovieSceneSubTrack>( MasterTrack );
-					if ( !SubTrack )
+					const UMovieSceneSequence* Sequence = SequenceInstancePair.Key;
+					if ( !Sequence )
 					{
 						continue;
 					}
 
-					for ( UMovieSceneSection* Section : SubTrack->GetAllSections() )
+					UMovieScene* MovieScene = Sequence->GetMovieScene();
+					if ( !MovieScene )
 					{
-						UMovieSceneSubSection* SubSection = Cast<UMovieSceneSubSection>( Section );
-						if ( !SubSection )
+						continue;
+					}
+
+					// Spawn everything for this instance
+					for ( FMovieSceneSequenceID SequenceInstance : SequenceInstancePair.Value )
+					{
+						int32 NumSpawnables = MovieScene->GetSpawnableCount();
+						for ( int32 Index = 0; Index < NumSpawnables; ++Index )
 						{
-							continue;
+							const FMovieSceneSpawnable& Spawnable = MovieScene->GetSpawnable( Index );
+							const FGuid& Guid = Spawnable.GetGuid();
+
+							StaticCastSharedRef<FMovieSceneSpawnRegister>( Context.SpawnRegister )->SpawnObject( Guid, *MovieScene, SequenceInstance, *Context.Sequencer );
 						}
+					}
+				}
 
-						UMovieSceneSequence* SubSequence = SubSection->GetSequence();
-						if ( !SubSequence )
+				// Recurse into subsequences if we'll export them as sublayers. This because when exporting each, we'll reset our sequencer
+				// to have each subsequence as the root level sequence again, which means that the sequence ids will all change. We need to
+				// prepare our SpawnRegister to be able to reuse the previous spawns for these new instance ids
+				if ( Context.ExportOptions->bExportSubsequencesAsLayers )
+				{
+					for ( TPair< UMovieSceneSequence*, TArray<FMovieSceneSequenceID> >& SequenceInstancePair : SequenceInstances )
+					{
+						UMovieSceneSequence* Sequence = SequenceInstancePair.Key;
+						if ( !Sequence || Sequence == &RootSequence )
 						{
 							continue;
 						}
 
-						// Focus on this section before recursing so that every time we spawn a spawnable we can just query
-						// Sequencer->GetFocusedTemplateID() to fetch the FMovieSceneSequenceIDRef for the movie scene that
-						// actually owns the spawnable.
-						// We always reset the sequencer to a new root sequence before baking anyway
-						Context.Sequencer->FocusSequenceInstance( *SubSection );
-
-						PreSpawnSpawnables( Context, *SubSequence );
+						PreSpawnSpawnables( Context, *Sequence );
 					}
 				}
+
+				return SequenceInstances;
 			}
 
 			// Appends to InOutComponentBakers all of the component bakers for all components bound to MovieSceneSequence.
@@ -350,6 +524,7 @@ namespace UE
 			void GenerateBakersForMovieScene(
 				FLevelSequenceExportContext& Context,
 				UMovieSceneSequence& MovieSceneSequence,
+				const TMap< UMovieSceneSequence*, TArray<FMovieSceneSequenceID> >& SequenceInstances,
 				UE::FUsdStage& UsdStage,
 				TMap<USceneComponent*, FCombinedComponentBakers>& InOutComponentBakers
 			)
@@ -362,71 +537,93 @@ namespace UE
 
 				// Find all components (if actors, the root components of those actors) bound to the movie scene.
 				// We use components here because when exporting actors and components to USD we basically just ignore
-				// actors altogether and export the component attachment hierarchy instead
-				TArray<TPair<FGuid, USceneComponent*>> BoundComponents;
+				// actors altogether and export the component attachment hierarchy instead.
+				// Index from component to FGuid because we may have multiple spawned actors/components for a given spawnable Guid
+				TMap<USceneComponent*, FGuid> BoundComponents;
 
-				// Possessables
-				int32 NumPossessables = MovieScene->GetPossessableCount();
-				for ( int32 Index = 0; Index < NumPossessables; ++Index )
+				const TArray<FMovieSceneSequenceID>* InstancesOfThisSequence = SequenceInstances.Find( &MovieSceneSequence );
+				if ( !InstancesOfThisSequence )
 				{
-					const FMovieScenePossessable& Possessable = MovieScene->GetPossessable( Index );
-					const FGuid& Guid = Possessable.GetGuid();
+					return;
+				}
 
-					UObject* BoundObject = LocateBoundObject( MovieSceneSequence, MovieScene, Possessable, Context.SpawnRegister );
+				UMovieSceneSequence* RootSequence = Context.Sequencer->GetRootMovieSceneSequence();
+				if ( !RootSequence )
+				{
+					return;
+				}
 
-					USceneComponent* BoundComponent = Cast<USceneComponent>( BoundObject );
-					if ( !BoundComponent )
+				for ( FMovieSceneSequenceID SequenceInstance : *InstancesOfThisSequence )
+				{
+					// Possessables
+					int32 NumPossessables = MovieScene->GetPossessableCount();
+					for ( int32 Index = 0; Index < NumPossessables; ++Index )
 					{
-						if ( const AActor* Actor = Cast<AActor>( BoundObject ) )
+						const FMovieScenePossessable& Possessable = MovieScene->GetPossessable( Index );
+						const FGuid& Guid = Possessable.GetGuid();
+
+						UObject* BoundObject = LocateBoundObject(
+							*RootSequence,
+							MovieSceneSequence,
+							SequenceInstance,
+							Possessable,
+							Context.SpawnRegister
+						);
+
+						USceneComponent* BoundComponent = Cast<USceneComponent>( BoundObject );
+						if ( !BoundComponent )
 						{
-							BoundComponent = Actor->GetRootComponent();
+							if ( const AActor* Actor = Cast<AActor>( BoundObject ) )
+							{
+								BoundComponent = Actor->GetRootComponent();
+							}
+						}
+						if ( !BoundComponent )
+						{
+							continue;
 						}
-					}
-					if ( !BoundComponent )
-					{
-						continue;
-					}
-
-					BoundComponents.Add( TPair<FGuid, USceneComponent*>( Guid, BoundComponent ) );
-				}
-
-				// Spawnables
-				int32 NumSpawnables = MovieScene->GetSpawnableCount();
-				for ( int32 Index = 0; Index < NumSpawnables; ++Index )
-				{
-					FMovieSceneSpawnable& Spawnable = MovieScene->GetSpawnable( Index );
-					const FGuid& Guid = Spawnable.GetGuid();
 
-					UObject* BoundObject = Context.SpawnRegister->GetExistingSpawn( Guid );
-					if ( !BoundObject )
-					{
-						// This should never happen as we preemptively spawn everything:
-						// At this point all our spawns should be spawned, but invisible
-						UE_LOG( LogUsd, Warning, TEXT( "Failed to find spawned object for spawnable with Guid '%s'" ), *Guid.ToString() );
-						continue;
+						BoundComponents.Add( BoundComponent, Guid );
 					}
 
-					USceneComponent* BoundComponent = Cast<USceneComponent>( BoundObject );
-					if ( !BoundComponent )
+					// Spawnables
+					int32 NumSpawnables = MovieScene->GetSpawnableCount();
+					for ( int32 Index = 0; Index < NumSpawnables; ++Index )
 					{
-						if ( const AActor* Actor = Cast<AActor>( BoundObject ) )
+						FMovieSceneSpawnable& Spawnable = MovieScene->GetSpawnable( Index );
+						const FGuid& Guid = Spawnable.GetGuid();
+
+						UObject* BoundObject = Context.SpawnRegister->GetExistingSpawn( *RootSequence, SequenceInstance, Guid );
+						if ( !BoundObject )
 						{
-							BoundComponent = Actor->GetRootComponent();
+							// This should never happen as we preemptively spawn everything:
+							// At this point all our spawns should be spawned, but invisible
+							UE_LOG( LogUsd, Warning, TEXT( "Failed to find spawned object for spawnable with Guid '%s'" ), *Guid.ToString() );
+							continue;
+						}
+
+						USceneComponent* BoundComponent = Cast<USceneComponent>( BoundObject );
+						if ( !BoundComponent )
+						{
+							if ( const AActor* Actor = Cast<AActor>( BoundObject ) )
+							{
+								BoundComponent = Actor->GetRootComponent();
+							}
+						}
+						if ( !BoundComponent )
+						{
+							continue;
 						}
-					}
-					if ( !BoundComponent )
-					{
-						continue;
-					}
 
-					BoundComponents.Add( TPair<FGuid, USceneComponent*>( Guid, BoundComponent ) );
+						BoundComponents.Add( BoundComponent, Guid );
+					}
 				}
 
 				// Generate bakers
-				for ( const TPair<FGuid, USceneComponent*>& Pair : BoundComponents )
+				for ( const TPair<USceneComponent*, FGuid>& Pair : BoundComponents )
 				{
-					const FGuid& Guid = Pair.Key;
-					USceneComponent* BoundComponent = Pair.Value;
+					USceneComponent* BoundComponent = Pair.Key;
+					const FGuid& Guid = Pair.Value;
 
 					FString PrimPath = UsdUtils::GetPrimPathForObject(
 						BoundComponent,
@@ -532,7 +729,7 @@ namespace UE
 							continue;
 						}
 
-						GenerateBakersForMovieScene( Context, *SubSequence, UsdStage, InOutComponentBakers );
+						GenerateBakersForMovieScene( Context, *SubSequence, SequenceInstances, UsdStage, InOutComponentBakers );
 					}
 				}
 			}
@@ -551,12 +748,6 @@ namespace UE
 					return;
 				}
 
-				// Currently we bake the full composed result of each level sequence into a single layer,
-				// because USD can't compose individual layers in the same way (with blending and so on). So here
-				// we make sure that our sequencer exports each MovieSceneSequence as if it was a root, emitting the
-				// same result as if we had exported that subsequence's LevelSequence by itself
-				Context.Sequencer->ResetToNewRootSequence( MovieSceneSequence );
-
 				const TRange< FFrameNumber > PlaybackRange = MovieScene->GetPlaybackRange();
 				const FFrameRate Resolution = MovieScene->GetTickResolution();
 				const FFrameRate DisplayRate = MovieScene->GetDisplayRate();
@@ -641,15 +832,7 @@ namespace UE
 
 				// Make sure we don't overwrite a file we just wrote *during this export*.
 				// Overwriting other files is OK, as we want to allow a "repeatedly export over the same files" workflow
-				FString UniqueFilePath = FilePath;
-				if ( Context.UsedFilePaths.Contains( UniqueFilePath ) )
-				{
-					int32 Suffix = 0;
-					do
-					{
-						UniqueFilePath = FString::Printf( TEXT( "%s_%d" ), *FilePath, Suffix++ );
-					} while ( Context.UsedFilePaths.Contains( UniqueFilePath ) );
-				}
+				FString UniqueFilePath = UsdUtils::GetUniqueName( FilePath, Context.UsedFilePaths );
 
 				UE::FUsdStage UsdStage = UnrealUSDWrapper::NewStage( *UniqueFilePath );
 				if ( !UsdStage )
@@ -676,8 +859,16 @@ namespace UE
 
 				UsdStage.SetDefaultPrim( RootPrim );
 
+				// Currently we bake the full composed result of each level sequence into a single layer,
+				// because USD can't compose individual layers in the same way (with blending and so on). So here
+				// we make sure that our sequencer exports each MovieSceneSequence as if it was a root, emitting the
+				// same result as if we had exported that subsequence's LevelSequence by itself
+				Context.Sequencer->ResetToNewRootSequence( MovieSceneSequence );
+
+				TMap< UMovieSceneSequence*, TArray<FMovieSceneSequenceID> > SequenceInstances = GetSequenceHierarchyInstances( MovieSceneSequence, Context.Sequencer.Get() );
+
 				TMap<USceneComponent*, FCombinedComponentBakers> Bakers;
-				GenerateBakersForMovieScene( Context, MovieSceneSequence, UsdStage, Bakers );
+				GenerateBakersForMovieScene( Context, MovieSceneSequence, SequenceInstances, UsdStage, Bakers );
 
 				// Bake this MovieScene
 				// We bake each MovieScene individually instead of doing one large simultaneous bake because this way
@@ -872,8 +1063,16 @@ bool ULevelSequenceExporterUsd::ExportBinary( UObject* Object, const TCHAR* Type
 	LevelSequenceExporterImpl::FLevelSequenceExportContext Context{ *LevelSequence, TempSequencer.ToSharedRef(), SpawnRegister.ToSharedRef() };
 	Context.ExportOptions = Options;
 
-	// Spawn (but hide) all spawnables so that they will also show up on the level export if we need them to
-	LevelSequenceExporterImpl::PreSpawnSpawnables( Context, Context.RootSequence );
+	// Spawn (but hide) all spawnables so that they will also show up on the level export if we need them to.
+	// We have to traverse the template IDs when spawning spawnables, because we'll want to force each individual spawnable of each
+	// FMovieSceneSequenceID to spawn a separate object, so that they can become separate prims. Without doing this, if we used the same
+	// subsequence with spawnables multiple times within a parent sequence we'd only get one prim out, as the FMovieSceneSpawnable objects
+	// would be the exact same between all instances of the child sequence (same FGuid)
+	LevelSequenceExporterImpl::PreSpawnSpawnables( Context, *LevelSequence );
+
+	// Hide all our spawns again so we can pretend they haven't actually spawned yet
+	Context.SpawnRegister->bDestroyingJustHides = true;
+	Context.SpawnRegister->CleanUp( *Context.Sequencer );
 
 	// Capture this first because when we launch UExporter::RunAssetExportTask the CurrentFileName will change
 	const FString TargetFileName = UExporter::CurrentFilename;
