Fixed GeometryCacheUSDStream compile error on targets that don't USE_USD_SDK.

diff --git a/Engine/Plugins/Importers/USDImporter/Source/GeometryCacheUSD/Private/GeometryCacheUSDStream.cpp b/Engine/Plugins/Importers/USDImporter/Source/GeometryCacheUSD/Private/GeometryCacheUSDStream.cpp
--- a/Engine/Plugins/Importers/USDImporter/Source/GeometryCacheUSD/Private/GeometryCacheUSDStream.cpp
+++ b/Engine/Plugins/Importers/USDImporter/Source/GeometryCacheUSD/Private/GeometryCacheUSDStream.cpp
@@ -39,7 +39,7 @@ private:
 public:
 	static FString GetUsdStreamDDCKey(const UE::FUsdStage& Stage, const FString& PrimPath, int32 FrameIndex)
 	{
-#ifdef USE_USD_SDK
+#if USE_USD_SDK
 		FString PrimHash = UsdUtils::HashGeomMeshPrim(Stage, PrimPath, FrameIndex);
 
 		if (!PrimHash.IsEmpty())
@@ -92,7 +92,7 @@ void FGeometryCacheUsdStream::GetMeshData(int32 FrameIndex, int32 ConcurrencyInd
 {
 	TRACE_CPUPROFILER_EVENT_SCOPE(FGeometryCacheUsdStream::GetMeshData);
 
-#ifdef USE_USD_SDK
+#if USE_USD_SDK
 	// Get the mesh data straight from the Alembic file or from the DDC if it's already cached
 	if (GUsdStreamCacheInDDC)
 	{
