USD: Fix geometry cache tracks holding the stage open indefinitely while they remain in the undo buffer.

diff --git a/Engine/Plugins/Importers/USDImporter/Source/GeometryCacheUSD/GeometryCacheUSD.build.cs b/Engine/Plugins/Importers/USDImporter/Source/GeometryCacheUSD/GeometryCacheUSD.build.cs
--- a/Engine/Plugins/Importers/USDImporter/Source/GeometryCacheUSD/GeometryCacheUSD.build.cs
+++ b/Engine/Plugins/Importers/USDImporter/Source/GeometryCacheUSD/GeometryCacheUSD.build.cs
@@ -21,6 +21,7 @@ public class GeometryCacheUSD : ModuleRules
 		PrivateDependencyModuleNames.AddRange(
 			new string[] {
 				"UnrealUSDWrapper",
+				"USDClasses",
 				"USDUtilities"
             }
         );
diff --git a/Engine/Plugins/Importers/USDImporter/Source/GeometryCacheUSD/Private/GeometryCacheTrackUSD.cpp b/Engine/Plugins/Importers/USDImporter/Source/GeometryCacheUSD/Private/GeometryCacheTrackUSD.cpp
--- a/Engine/Plugins/Importers/USDImporter/Source/GeometryCacheUSD/Private/GeometryCacheTrackUSD.cpp
+++ b/Engine/Plugins/Importers/USDImporter/Source/GeometryCacheUSD/Private/GeometryCacheTrackUSD.cpp
@@ -2,6 +2,7 @@
 
 #include "GeometryCacheTrackUSD.h"
 
+#include "USDLog.h"
 #include "UsdWrappers/SdfLayer.h"
 #include "UsdWrappers/UsdStage.h"
 
@@ -16,7 +17,7 @@ UGeometryCacheTrackUsd::UGeometryCacheTrackUsd()
 
 void UGeometryCacheTrackUsd::BeginDestroy()
 {
-	CurrentStage = UE::FUsdStage();
+	UnloadUsdStage();
 
 	IGeometryCacheStreamer::Get().UnregisterTrack( this );
 	UsdStream.Reset();
@@ -117,6 +118,48 @@ bool UGeometryCacheTrackUsd::GetMeshData(int32 SampleIndex, FGeometryCacheMeshDa
 	return false;
 }
 
+bool UGeometryCacheTrackUsd::LoadUsdStage()
+{
+	if ( CurrentStagePinned )
+	{
+		// Already loaded
+		return true;
+	}
+
+	if ( CurrentStageWeak )
+	{
+		// Upgrade our weak pointer if its not invalid already
+		CurrentStagePinned = CurrentStageWeak;
+		return true;
+	}
+	else if ( !StageRootLayerPath.IsEmpty() )
+	{
+		UE_LOG( LogUsd, Warning, TEXT( "UGeometryCacheTrackUsd is reopening the stage '%s' to stream in frames for the geometry cache generated for prim '%s'" ), *StageRootLayerPath, *PrimPath );
+
+		// Reopen the stage. If our weak pointer is no longer valid then nothing cared about keeping that
+		// stage alive anyway, so it's likely not a problem if we start reading frames from the reopened stage
+		// and abandon any previous in-memory changes we had, if any.
+		// Keep in mind currently it's effectively impossible to get in here at all anyway, as we'll only
+		// have UGeometryCacheTrackUsd streaming stuff if a stage actor caused it to stream more frames,
+		// and in that case that stage actor would have kept our stage opened.
+		// Not using the stage cache here because there's currently no way of knowing when to erase it from the cache
+		// after we're done with it, and the UGeometryCacheTrackUsd shouldn't have authority to blindly just remove it
+		// from the cache as the user may have place it there intentionally
+		const bool bUseStageCache = false;
+		CurrentStagePinned = UnrealUSDWrapper::OpenStage( *StageRootLayerPath, EUsdInitialLoadSet::LoadAll, bUseStageCache );
+		CurrentStageWeak = CurrentStagePinned;
+		return true;
+	}
+
+	UE_LOG( LogUsd, Warning, TEXT( "UGeometryCacheTrackUsd track failed to access USD stage to stream requested frames" ) );
+	return false;
+}
+
+void UGeometryCacheTrackUsd::UnloadUsdStage()
+{
+	CurrentStagePinned = UE::FUsdStage();
+}
+
 void UGeometryCacheTrackUsd::Initialize(
 	const UE::FUsdStage& InStage,
 	const FString& InPrimPath,
@@ -127,8 +170,9 @@ void UGeometryCacheTrackUsd::Initialize(
 	FReadUsdMeshFunction InReadFunc
 )
 {
-	CurrentStage = InStage;
-	StageRootLayerPath = CurrentStage ? CurrentStage.GetRootLayer().GetRealPath() : FString();
+	CurrentStagePinned = InStage;
+	CurrentStageWeak = CurrentStagePinned;
+	StageRootLayerPath = CurrentStagePinned ? CurrentStagePinned.GetRootLayer().GetRealPath() : FString();
 
 	PrimPath = InPrimPath;
 	RenderContext = InRenderContext;
@@ -150,4 +194,4 @@ void UGeometryCacheTrackUsd::UpdateTime(float Time, bool bLooping)
 		int32 FrameIndex = FindSampleIndexFromTime(Time, bLooping);
 		UsdStream->UpdateCurrentFrameIndex(FrameIndex);
 	}
-}
\ No newline at end of file
+}
diff --git a/Engine/Plugins/Importers/USDImporter/Source/GeometryCacheUSD/Private/GeometryCacheUSDStream.cpp b/Engine/Plugins/Importers/USDImporter/Source/GeometryCacheUSD/Private/GeometryCacheUSDStream.cpp
--- a/Engine/Plugins/Importers/USDImporter/Source/GeometryCacheUSD/Private/GeometryCacheUSDStream.cpp
+++ b/Engine/Plugins/Importers/USDImporter/Source/GeometryCacheUSD/Private/GeometryCacheUSDStream.cpp
@@ -56,8 +56,8 @@ FGeometryCacheUsdStream::FGeometryCacheUsdStream(UGeometryCacheTrackUsd* InUsdTr
 	kUsdReadConcurrency,
 	FGeometryCacheStreamDetails{
 		InUsdTrack->GetEndFrameIndex() - InUsdTrack->GetStartFrameIndex(),
-		float((InUsdTrack->GetEndFrameIndex() - InUsdTrack->GetStartFrameIndex()) / InUsdTrack->CurrentStage.GetFramesPerSecond()),
-		float(1.0f / InUsdTrack->CurrentStage.GetFramesPerSecond()),
+		float((InUsdTrack->GetEndFrameIndex() - InUsdTrack->GetStartFrameIndex()) / InUsdTrack->CurrentStagePinned.GetFramesPerSecond()),
+		float(1.0f / InUsdTrack->CurrentStagePinned.GetFramesPerSecond()),
 		InUsdTrack->GetStartFrameIndex(),
 		InUsdTrack->GetEndFrameIndex()})
 , UsdTrack(InUsdTrack)
@@ -88,16 +88,34 @@ bool FGeometryCacheUsdStream::GetFrameData(int32 FrameIndex, FGeometryCacheMeshD
 	return true;
 }
 
+void FGeometryCacheUsdStream::UpdateRequestStatus( TArray<int32>& OutFramesCompleted )
+{
+	FGeometryCacheStreamBase::UpdateRequestStatus( OutFramesCompleted );
+
+	// We're fully done fetching what we need from USD for now, we can drop the track's strong stage reference so that the stage
+	// can close if needed
+	if ( FramesNeeded.Num() == 0 && FramesRequested.Num() == 0 && UsdTrack && UsdTrack->CurrentStagePinned )
+	{
+		UsdTrack->UnloadUsdStage();
+	}
+}
+
 void FGeometryCacheUsdStream::GetMeshData(int32 FrameIndex, int32 ConcurrencyIndex, FGeometryCacheMeshData& OutMeshData)
 {
 	TRACE_CPUPROFILER_EVENT_SCOPE(FGeometryCacheUsdStream::GetMeshData);
 
+	// Do this before calling the ReadFunc as FUsdStreamDDCUtils::GetUsdStreamDDCKey also needs a valid stage to work with
+	if ( !UsdTrack->LoadUsdStage() )
+	{
+		return;
+	}
+
 #if USE_USD_SDK
 	// Get the mesh data straight from the Alembic file or from the DDC if it's already cached
 	if (GUsdStreamCacheInDDC)
 	{
 		const FString& UsdPrimPath = UsdTrack->PrimPath;
-		const FString DerivedDataKey = FUsdStreamDDCUtils::GetUsdStreamDDCKey(UsdTrack->CurrentStage, UsdPrimPath, FrameIndex);
+		const FString DerivedDataKey = FUsdStreamDDCUtils::GetUsdStreamDDCKey(UsdTrack->CurrentStagePinned, UsdPrimPath, FrameIndex);
 
 		if (!DerivedDataKey.IsEmpty())
 		{
diff --git a/Engine/Plugins/Importers/USDImporter/Source/GeometryCacheUSD/Public/GeometryCacheTrackUSD.h b/Engine/Plugins/Importers/USDImporter/Source/GeometryCacheUSD/Public/GeometryCacheTrackUSD.h
--- a/Engine/Plugins/Importers/USDImporter/Source/GeometryCacheUSD/Public/GeometryCacheTrackUSD.h
+++ b/Engine/Plugins/Importers/USDImporter/Source/GeometryCacheUSD/Public/GeometryCacheTrackUSD.h
@@ -48,6 +48,12 @@ public:
 
 	bool GetMeshData(int32 SampleIndex, FGeometryCacheMeshData& OutMeshData);
 
+	// Upgrades our CurrentStageWeak into CurrentStagePinned, or re-opens the stage if its stale. Returns whether the stage was successfully opened or not
+	bool LoadUsdStage();
+
+	// Discards our CurrentStagePinned to release the stage
+	void UnloadUsdStage();
+
 public:
 	FName RenderContext;
 	int32 StartFrameIndex;
@@ -56,7 +62,8 @@ public:
 
 	FString PrimPath;
 
-	UE::FUsdStage CurrentStage;
+	UE::FUsdStage CurrentStagePinned;
+	UE::FUsdStageWeak CurrentStageWeak;
 	FString StageRootLayerPath;
 
 private:
diff --git a/Engine/Plugins/Importers/USDImporter/Source/GeometryCacheUSD/Public/GeometryCacheUSDStream.h b/Engine/Plugins/Importers/USDImporter/Source/GeometryCacheUSD/Public/GeometryCacheUSDStream.h
--- a/Engine/Plugins/Importers/USDImporter/Source/GeometryCacheUSD/Public/GeometryCacheUSDStream.h
+++ b/Engine/Plugins/Importers/USDImporter/Source/GeometryCacheUSD/Public/GeometryCacheUSDStream.h
@@ -14,6 +14,7 @@ public:
 
 	//~ Begin IGeometryCacheStream Interface
 	virtual bool GetFrameData(int32 FrameIndex, FGeometryCacheMeshData& OutMeshData) override;
+	virtual void UpdateRequestStatus( TArray<int32>& OutFramesCompleted ) override;
 
 protected:
 	virtual void GetMeshData(int32 FrameIndex, int32 ConcurrencyIndex, FGeometryCacheMeshData& OutMeshData) override;
diff --git a/Engine/Plugins/Importers/USDImporter/Source/USDSchemas/Private/USDGeomMeshTranslator.cpp b/Engine/Plugins/Importers/USDImporter/Source/USDSchemas/Private/USDGeomMeshTranslator.cpp
--- a/Engine/Plugins/Importers/USDImporter/Source/USDSchemas/Private/USDGeomMeshTranslator.cpp
+++ b/Engine/Plugins/Importers/USDImporter/Source/USDSchemas/Private/USDGeomMeshTranslator.cpp
@@ -821,21 +821,12 @@ namespace UsdGeomMeshTranslatorImpl
 						return;
 					}
 
-					if ( !Track->CurrentStage )
-					{
-						if ( !Track->StageRootLayerPath.IsEmpty() )
-						{
-							const bool bUseStageCache = true;
-							Track->CurrentStage = UnrealUSDWrapper::OpenStage( *Track->StageRootLayerPath, EUsdInitialLoadSet::LoadAll, bUseStageCache );
-							UE_LOG( LogUsd, Warning, TEXT( "UGeometryCacheTrackUsd is reopening the stage '%s' to stream in frames for the geometry cache generated for prim '%s'" ), *Track->StageRootLayerPath, *InPrimPath );
-						}
-					}
-					if ( !Track->CurrentStage )
+					if ( !Track->CurrentStagePinned )
 					{
 						return;
 					}
 
-					UE::FUsdPrim Prim = Track->CurrentStage.GetPrimAtPath( UE::FSdfPath{ *Track->PrimPath } );
+					UE::FUsdPrim Prim = Track->CurrentStagePinned.GetPrimAtPath( UE::FSdfPath{ *Track->PrimPath } );
 					if ( !Prim )
 					{
 						return;
