USD: Fix off-by-one behavior with USD geometry caches due to EndTimeIndex not being one past last frame.

diff --git a/Engine/Plugins/Importers/USDImporter/Source/USDSchemas/Private/USDGeomMeshTranslator.cpp b/Engine/Plugins/Importers/USDImporter/Source/USDSchemas/Private/USDGeomMeshTranslator.cpp
--- a/Engine/Plugins/Importers/USDImporter/Source/USDSchemas/Private/USDGeomMeshTranslator.cpp
+++ b/Engine/Plugins/Importers/USDImporter/Source/USDSchemas/Private/USDGeomMeshTranslator.cpp
@@ -804,6 +804,9 @@ namespace UsdGeomMeshTranslatorImpl
 			int32 EndFrame = FMath::CeilToInt( Stage.GetEndTimeCode() );
 			UsdGeomMeshTranslatorImpl::GetGeometryCacheDataTimeCodeRange( Stage, InPrimPath, StartFrame, EndFrame );
 
+			// The GeometryCache module expects the end frame to be one past the last animation frame
+			EndFrame += 1;
+
 			// Create and configure a new USDTrack to be added to the GeometryCache
 			UGeometryCacheTrackUsd* UsdTrack = NewObject< UGeometryCacheTrackUsd >( GeometryCache );
 			UsdTrack->Initialize(
