USD: Fix not being able to save usda/usdc files on the Mac.

diff --git a/Engine/Plugins/Importers/USDImporter/Source/USDUtilities/Private/USDLayerUtils.cpp b/Engine/Plugins/Importers/USDImporter/Source/USDUtilities/Private/USDLayerUtils.cpp
--- a/Engine/Plugins/Importers/USDImporter/Source/USDUtilities/Private/USDLayerUtils.cpp
+++ b/Engine/Plugins/Importers/USDImporter/Source/USDUtilities/Private/USDLayerUtils.cpp
@@ -131,7 +131,9 @@ TOptional< FString > UsdUtils::BrowseUsdFile( EBrowseFileMode Mode, TSharedRef<
 					continue;
 				}
 
-				FileTypes += FString::Printf( TEXT( "Universal Scene Description file|*.%s|" ), *SupportedExtension );
+				// The '(*.%s)' on the actual name (before the '|') is not optional: We need the name part to be different for each format
+				// or else the options will overwrite each other on the Mac
+				FileTypes += FString::Printf( TEXT( "Universal Scene Description file (*.%s)|*.%s|" ), *SupportedExtension, *SupportedExtension );
 			}
 
 			if ( !DesktopPlatform->SaveFileDialog( ParentWindowHandle, LOCTEXT( "ChooseFile", "Choose file").ToString(), TEXT(""), TEXT(""), *FileTypes, EFileDialogFlags::None, OutFiles ) )
