USD: Fix component visibility not being written out to USD due to bVisible to bHiddenInGame change.

diff --git a/Engine/Plugins/Importers/USDImporter/Source/USDStage/Private/USDStageActor.cpp b/Engine/Plugins/Importers/USDImporter/Source/USDStage/Private/USDStageActor.cpp
--- a/Engine/Plugins/Importers/USDImporter/Source/USDStage/Private/USDStageActor.cpp
+++ b/Engine/Plugins/Importers/USDImporter/Source/USDStage/Private/USDStageActor.cpp
@@ -421,8 +421,8 @@ struct FUsdStageActorImpl
 	/**
 	 * Uses USD's MakeVisible to handle the visible/inherited update logic as it is a bit complex.
 	 * Will update a potentially large chunk of the component hierarchy to having/not the `invisible` component tag, as well as the
-	 * correct value of bVisible.
-	 * Note that bVisible corresponds to computed visibility, and the component tags correspond to individual prim-level visibilities
+	 * correct value of bHiddenInGame.
+	 * Note that bHiddenInGame corresponds to computed visibility, and the component tags correspond to individual prim-level visibilities
 	 */
 	static void MakeVisible( UUsdPrimTwin& UsdPrimTwin, UE::FUsdStage& Stage )
 	{
@@ -521,8 +521,8 @@ struct FUsdStageActorImpl
 
 	/**
 	 * Sets this prim to 'invisible', and force all of the child components
-	 * to bVisible = false. Leave their individual prim-level visibilities intact though.
-	 * Note that bVisible corresponds to computed visibility, and the component tags correspond to individual prim-level visibilities
+	 * to bHiddenInGame = false. Leave their individual prim-level visibilities intact though.
+	 * Note that bHiddenInGame corresponds to computed visibility, and the component tags correspond to individual prim-level visibilities
 	 */
 	static void MakeInvisible( UUsdPrimTwin& UsdPrimTwin )
 	{
@@ -2279,7 +2279,7 @@ void AUsdStageActor::OnObjectPropertyChanged( UObject* ObjectBeingModified, FPro
 				// We want to keep component visibilities in sync with USD, which uses inherited visibilities
 				// To accomplish that while blocking notices we must always propagate component visibility changes manually.
 				// This part is effectively the same as calling pxr::UsdGeomImageable::MakeVisible/Invisible.
-				if ( PropertyChangedEvent.GetPropertyName() == TEXT( "bVisible" ) )
+				if ( PropertyChangedEvent.GetPropertyName() == TEXT( "bHiddenInGame" ) )
 				{
 					PrimSceneComponent->Modify();
 
