USD: Fix SpotLight intensity being incorrect on import.

diff --git a/Engine/Plugins/Importers/USDImporter/Source/USDUtilities/Private/USDLightConversion.cpp b/Engine/Plugins/Importers/USDImporter/Source/USDUtilities/Private/USDLightConversion.cpp
--- a/Engine/Plugins/Importers/USDImporter/Source/USDUtilities/Private/USDLightConversion.cpp
+++ b/Engine/Plugins/Importers/USDImporter/Source/USDUtilities/Private/USDLightConversion.cpp
@@ -225,7 +225,21 @@ bool UsdToUnreal::ConvertSphereLight( const pxr::UsdPrim& Prim, UPointLightCompo
 	const float UsdExposure = UsdUtils::GetUsdValue< float >( SphereLight.GetExposureAttr(), UsdTimeCode );
 	const float UsdRadius = UsdUtils::GetUsdValue< float >( SphereLight.GetRadiusAttr(), UsdTimeCode );
 
-	LightComponent.Intensity = UsdToUnreal::ConvertSphereLightIntensityAttr( UsdIntensity, UsdExposure, UsdRadius, StageInfo );
+	if ( pxr::UsdLuxShapingAPI ShapingAPI{ Prim } )
+	{
+		const float UsdConeAngle = UsdUtils::GetUsdValue< float >( ShapingAPI.GetShapingConeAngleAttr(), UsdTimeCode );
+		const float UsdConeSoftness = UsdUtils::GetUsdValue< float >( ShapingAPI.GetShapingConeSoftnessAttr(), UsdTimeCode );
+
+		float InnerConeAngle = 0.0f;
+		const float OuterConeAngle = UsdToUnreal::ConvertConeAngleSoftnessAttr( UsdConeAngle, UsdConeSoftness, InnerConeAngle );
+
+		LightComponent.Intensity = UsdToUnreal::ConvertLuxShapingAPIIntensityAttr( UsdIntensity, UsdExposure, UsdRadius, UsdConeAngle, UsdConeSoftness, StageInfo );
+	}
+	else
+	{
+		LightComponent.Intensity = UsdToUnreal::ConvertSphereLightIntensityAttr( UsdIntensity, UsdExposure, UsdRadius, StageInfo );
+	}
+
 	LightComponent.IntensityUnits = ELightUnits::Lumens;
 	LightComponent.SourceRadius = UsdToUnreal::ConvertDistance( StageInfo, UsdRadius );
 
@@ -293,25 +307,17 @@ bool UsdToUnreal::ConvertLuxShapingAPI( const pxr::UsdPrim& Prim, USpotLightComp
 	}
 
 	pxr::UsdLuxShapingAPI ShapingAPI{ Prim };
-	pxr::UsdLuxSphereLight SphereLight{ Prim };
-	if ( !ShapingAPI || !SphereLight )
+	if ( !ShapingAPI )
 	{
 		return false;
 	}
 
-	const FUsdStageInfo StageInfo( Prim.GetStage() );
-
-	const float UsdIntensity = UsdUtils::GetUsdValue< float >( SphereLight.GetIntensityAttr(), UsdTimeCode );
-	const float UsdExposure = UsdUtils::GetUsdValue< float >( SphereLight.GetExposureAttr(), UsdTimeCode );
-	const float UsdRadius = UsdUtils::GetUsdValue< float >( SphereLight.GetRadiusAttr(), UsdTimeCode );
 	const float UsdConeAngle = UsdUtils::GetUsdValue< float >( ShapingAPI.GetShapingConeAngleAttr(), UsdTimeCode );
 	const float UsdConeSoftness = UsdUtils::GetUsdValue< float >( ShapingAPI.GetShapingConeSoftnessAttr(), UsdTimeCode );
 
 	float InnerConeAngle = 0.0f;
 	const float OuterConeAngle = UsdToUnreal::ConvertConeAngleSoftnessAttr( UsdConeAngle, UsdConeSoftness, InnerConeAngle );
 
-	LightComponent.Intensity = UsdToUnreal::ConvertLuxShapingAPIIntensityAttr( UsdIntensity, UsdExposure, UsdRadius, UsdConeAngle, UsdConeSoftness, StageInfo );
-	LightComponent.IntensityUnits = ELightUnits::Lumens;
 	LightComponent.SetInnerConeAngle( InnerConeAngle );
 	LightComponent.SetOuterConeAngle( OuterConeAngle );
 
