USD: Fix imported level sequences always ending up with "_usda" suffix.

diff --git a/Engine/Plugins/Importers/USDImporter/Source/USDStage/Private/USDLevelSequenceHelper.cpp b/Engine/Plugins/Importers/USDImporter/Source/USDStage/Private/USDLevelSequenceHelper.cpp
--- a/Engine/Plugins/Importers/USDImporter/Source/USDStage/Private/USDLevelSequenceHelper.cpp
+++ b/Engine/Plugins/Importers/USDImporter/Source/USDStage/Private/USDLevelSequenceHelper.cpp
@@ -707,8 +707,8 @@ ULevelSequence* FUsdLevelSequenceHelperImpl::FindOrAddSequenceForLayer( const UE
 		// importer will make these names unique later if needed
 		const bool bIsImporting = StageActor.IsExplicitlyNull();
 		FName UniqueSequenceName = bIsImporting
-			? *UsdLevelSequenceHelperImpl::SanitizeObjectName( SequenceDisplayName )
-			: MakeUniqueObjectName( GetTransientPackage(), ULevelSequence::StaticClass(), *UsdLevelSequenceHelperImpl::SanitizeObjectName( SequenceDisplayName ) );
+			? *UsdLevelSequenceHelperImpl::SanitizeObjectName( FPaths::GetBaseFilename( SequenceDisplayName ) )
+			: MakeUniqueObjectName( GetTransientPackage(), ULevelSequence::StaticClass(), *UsdLevelSequenceHelperImpl::SanitizeObjectName( FPaths::GetBaseFilename( SequenceDisplayName ) ) );
 
 		Sequence = NewObject< ULevelSequence >( GetTransientPackage(), UniqueSequenceName, FUsdLevelSequenceHelperImpl::DefaultObjFlags );
 		Sequence->Initialize();
