USD: Try checking for a primvar reader with higher priority than using a shader's fallback value.

diff --git a/Engine/Plugins/Importers/USDImporter/Source/USDUtilities/Private/USDShadeConversion.cpp b/Engine/Plugins/Importers/USDImporter/Source/USDUtilities/Private/USDShadeConversion.cpp
--- a/Engine/Plugins/Importers/USDImporter/Source/USDUtilities/Private/USDShadeConversion.cpp
+++ b/Engine/Plugins/Importers/USDImporter/Source/USDUtilities/Private/USDShadeConversion.cpp
@@ -775,6 +775,13 @@ namespace UE
 				{
 					if ( !GetTextureParameterValue( Input, bIsNormalMap ? TEXTUREGROUP_WorldNormalMap : TEXTUREGROUP_World, OutValue, Material, TexturesCache, PrimvarToUVIndex, bForceVirtualTextures ) )
 					{
+						// Check whether this input receives its value through a connection to a
+						// primvar reader shader.
+						if ( GetPrimvarReaderParameterValue( Input, UnrealIdentifiers::UsdPrimvarReader_float3, DefaultValue, OutValue ) )
+						{
+							return true;
+						}
+
 						// Check if we have a fallback input that we can use instead, since we don't have a valid texture value
 						if ( const pxr::UsdShadeInput FallbackInput = Source.GetInput( UnrealIdentifiers::Fallback ) )
 						{
@@ -800,12 +807,8 @@ namespace UE
 							}
 						}
 
-						// Check whether this input receives its value through a connection to a
-						// primvar reader shader.
-						if ( !GetPrimvarReaderParameterValue( Input, UnrealIdentifiers::UsdPrimvarReader_float3, DefaultValue, OutValue ) )
-						{
-							return GetVec3ParameterValue( Source, SourceName, DefaultValue, OutValue, bIsNormalMap, Material, TexturesCache, PrimvarToUVIndex, bForceVirtualTextures );
-						}
+						// This shader doesn't have anything: Traverse into the input connectable itself
+						return GetVec3ParameterValue( Source, SourceName, DefaultValue, OutValue, bIsNormalMap, Material, TexturesCache, PrimvarToUVIndex, bForceVirtualTextures );
 					}
 				}
 				// No other node connected, so we must have some value
