USD: Partial revert of cl 18308126 to prevent meshes from being parsed multiple times.

diff --git a/Engine/Plugins/Importers/USDImporter/Source/USDSchemas/Private/USDGeomMeshTranslator.cpp b/Engine/Plugins/Importers/USDImporter/Source/USDSchemas/Private/USDGeomMeshTranslator.cpp
--- a/Engine/Plugins/Importers/USDImporter/Source/USDSchemas/Private/USDGeomMeshTranslator.cpp
+++ b/Engine/Plugins/Importers/USDImporter/Source/USDSchemas/Private/USDGeomMeshTranslator.cpp
@@ -504,7 +504,7 @@ namespace UsdGeomMeshTranslatorImpl
 	}
 
 	void LoadMeshDescriptions( const pxr::UsdTyped& UsdMesh, TArray<FMeshDescription>& OutLODIndexToMeshDescription, TArray<UsdUtils::FUsdPrimMaterialAssignmentInfo>& OutLODIndexToMaterialInfo,
-		const TMap< FString, TMap< FString, int32 > >& MaterialToPrimvarToUVIndex, const pxr::UsdTimeCode TimeCode, bool bInterpretLODs, const FName& RenderContext, const EUsdPurpose PurposesToLoad )
+		const TMap< FString, TMap< FString, int32 > >& MaterialToPrimvarToUVIndex, const pxr::UsdTimeCode TimeCode, bool bInterpretLODs, const FName& RenderContext )
 	{
 		if ( !UsdMesh )
 		{
@@ -528,15 +528,10 @@ namespace UsdGeomMeshTranslatorImpl
 			FMeshDescription TempMeshDescription;
 			UsdUtils::FUsdPrimMaterialAssignmentInfo TempMaterialInfo;
 
-			bool bSuccess = UsdToUnreal::ConvertGeomMeshHierarchy(
-				UsdMesh.GetPrim(),
-				TimeCode,
-				PurposesToLoad,
-				RenderContextToken,
-				MaterialToPrimvarToUVIndex,
-				TempMeshDescription,
-				TempMaterialInfo
-			);
+			FStaticMeshAttributes StaticMeshAttributes( TempMeshDescription );
+			StaticMeshAttributes.Register();
+
+			bool bSuccess = UsdToUnreal::ConvertGeomMesh( UsdMesh, TempMeshDescription, TempMaterialInfo, FTransform::Identity, MaterialToPrimvarToUVIndex, TimeCode, RenderContextToken );
 
 			if ( bSuccess )
 			{
@@ -797,8 +792,6 @@ namespace UsdGeomMeshTranslatorImpl
 			TMap< FString, TMap< FString, int32 > > Unused;
 			TMap< FString, TMap< FString, int32 > >* MaterialToPrimvarToUVIndex = Context->MaterialToPrimvarToUVIndex ? Context->MaterialToPrimvarToUVIndex : &Unused;
 
-			const EUsdPurpose PurposesToLoad = Context->PurposesToLoad;
-
 			// Fetch the animated mesh start/end frame as they may be different from just the stage's start and end time codes
 			int32 StartFrame = FMath::FloorToInt( Stage.GetStartTimeCode() );
 			int32 EndFrame = FMath::CeilToInt( Stage.GetEndTimeCode() );
@@ -816,7 +809,7 @@ namespace UsdGeomMeshTranslatorImpl
 				*MaterialToPrimvarToUVIndex,
 				StartFrame,
 				EndFrame,
-				[InPrimPath, PurposesToLoad]( const TWeakObjectPtr<UGeometryCacheTrackUsd> TrackPtr, float Time, FGeometryCacheMeshData& OutMeshData )
+				[InPrimPath]( const TWeakObjectPtr<UGeometryCacheTrackUsd> TrackPtr, float Time, FGeometryCacheMeshData& OutMeshData )
 				{
 					UGeometryCacheTrackUsd* Track = TrackPtr.Get();
 					if ( !Track )
@@ -848,8 +841,7 @@ namespace UsdGeomMeshTranslatorImpl
 						Track->MaterialToPrimvarToUVIndex,
 						pxr::UsdTimeCode( Time ),
 						bAllowInterpretingLODs,
-						Track->RenderContext,
-						PurposesToLoad
+						Track->RenderContext
 					);
 
 					// Convert the MeshDescription to MeshData
@@ -1179,8 +1171,7 @@ void FGeomMeshCreateAssetsTaskChain::SetupTasks()
 				*MaterialToPrimvarToUVIndex,
 				pxr::UsdTimeCode( Context->Time ),
 				Context->bAllowInterpretingLODs,
-				Context->RenderContext,
-				Context->PurposesToLoad
+				Context->RenderContext
 			);
 
 			// If we have at least one valid LOD, we should keep going
@@ -1234,8 +1225,7 @@ void FGeometryCacheCreateAssetsTaskChain::SetupTasks()
 				*MaterialToPrimvarToUVIndex,
 				pxr::UsdTimeCode( TimeCode ),
 				bAllowInterpretingLODs,
-				Context->RenderContext,
-				Context->PurposesToLoad
+				Context->RenderContext
 			);
 
 			// If we have at least one valid LOD, we should keep going
@@ -1428,6 +1418,18 @@ void FUsdGeomMeshTranslator::UpdateComponents( USceneComponent* SceneComponent )
 	Super::UpdateComponents( SceneComponent );
 }
 
+bool FUsdGeomMeshTranslator::CollapsesChildren( ECollapsingType CollapsingType ) const
+{
+	// We can't claim we collapse anything here since we'll just parse the mesh for this prim and that's it,
+	// otherwise the translation context wouldn't spawn translators for our child prims.
+	// Another approach would be to actually recursively collapse our child mesh prims, but that leads to a few
+	// issues. For example this translator could end up globbing a child Mesh prim, while the translation context
+	// could simultaneously spawn other translators that could also end up accounting for that same mesh.
+	// Generally Gprims shouldn't be nested into each other anyway (see https://graphics.pixar.com/usd/release/glossary.html#usdglossary-gprim)
+	// so it's likely best to just not collapse anything here.
+	return false;
+}
+
 bool FUsdGeomMeshTranslator::CanBeCollapsed( ECollapsingType CollapsingType ) const
 {
 	UE::FUsdPrim Prim = GetPrim();
diff --git a/Engine/Plugins/Importers/USDImporter/Source/USDSchemas/Public/USDGeomMeshTranslator.h b/Engine/Plugins/Importers/USDImporter/Source/USDSchemas/Public/USDGeomMeshTranslator.h
--- a/Engine/Plugins/Importers/USDImporter/Source/USDSchemas/Public/USDGeomMeshTranslator.h
+++ b/Engine/Plugins/Importers/USDImporter/Source/USDSchemas/Public/USDGeomMeshTranslator.h
@@ -75,6 +75,7 @@ public:
 	virtual USceneComponent* CreateComponents() override;
 	virtual void UpdateComponents( USceneComponent* SceneComponent ) override;
 
+	virtual bool CollapsesChildren( ECollapsingType CollapsingType ) const override;
 	virtual bool CanBeCollapsed( ECollapsingType CollapsingType ) const override;
 
 };
