USD: Fix crash when importing exported landscapes in some scenarios.

diff --git a/Engine/Plugins/Importers/USDImporter/Source/USDSchemas/Private/USDSchemaTranslator.cpp b/Engine/Plugins/Importers/USDImporter/Source/USDSchemas/Private/USDSchemaTranslator.cpp
--- a/Engine/Plugins/Importers/USDImporter/Source/USDSchemas/Private/USDSchemaTranslator.cpp
+++ b/Engine/Plugins/Importers/USDImporter/Source/USDSchemas/Private/USDSchemaTranslator.cpp
@@ -159,8 +159,10 @@ void FUsdSchemaTranslationContext::CompleteTasks()
 
 	FScopedSlowTask SlowTask( TranslatorTasks.Num(), LOCTEXT( "TasksProgress", "Executing USD Schema tasks" ) );
 
-	// Some tasks need to not only be executed serially, but also with the guarantee that no other concurrent tasks are running
-	TArray< TSharedPtr< FUsdSchemaTranslatorTaskChain >> PendingTaskChains;
+	// These are just pointers into the TranslatorTasks items that are task chains that are in the ESchemaTranslationStatus::Pending state.
+	// We'll use this below to execute tasks until all items in TranslatorTasks are pending, and then we switch to/from an ExclusiveSync pass.
+	// This because we need to ensure that all ExclusiveSync tasks are run in isolation.
+	TSet< FUsdSchemaTranslatorTaskChain* > PendingTaskChains;
 
 	// Note that this first pass is for tasks that allow concurrent execution (so *not* exclusive sync tasks). If this is ever changed,
 	// we would also need to change the behavior of StartIfAsync to delay until the proper async pass, and not start right away
@@ -184,7 +186,7 @@ void FUsdSchemaTranslationContext::CompleteTasks()
 				}
 				else if ( TaskChainStatus == ESchemaTranslationStatus::Pending )
 				{
-					PendingTaskChains.Add(TaskChain);
+					PendingTaskChains.Add( TaskChain.Get() );
 				}
 			}
 		}
