USD: Actually fully prevent the publishing of meshes/materials if those are unwanted given the import options;

Remove old edit condition between level sequences and animsequences;
Always set material overrides on components, which allows for some useful replace-import workflows.

diff --git a/Engine/Plugins/Importers/USDImporter/Source/USDStageImporter/Private/USDStageImporter.cpp b/Engine/Plugins/Importers/USDImporter/Source/USDStageImporter/Private/USDStageImporter.cpp
--- a/Engine/Plugins/Importers/USDImporter/Source/USDStageImporter/Private/USDStageImporter.cpp
+++ b/Engine/Plugins/Importers/USDImporter/Source/USDStageImporter/Private/USDStageImporter.cpp
@@ -24,6 +24,7 @@
 #include "AssetRegistryModule.h"
 #include "AssetToolsModule.h"
 #include "ComponentRecreateRenderStateContext.h"
+#include "Components/SkeletalMeshComponent.h"
 #include "Dialogs/DlgPickPath.h"
 #include "Editor.h"
 #include "Engine/SkeletalMesh.h"
@@ -1342,6 +1343,76 @@ namespace UsdStageImporterImpl
 			);
 		}
 	}
+
+	// Removes from AssetsToImport assets that are unwanted according to our import options, and adds entries to
+	// ObjectsToRemap and SoftObjectsToRemap that remaps them to nullptr.
+	// This function is needed because it's not enough to e.g. just prevent new meshes from being imported from
+	// UsdStageImporterImpl::ImportMeshes, because we may want to reuse meshes we already got from the asset cache.
+	// Additionally, we'll want to remap even our components away from pointing to these assets
+	void PruneUnwantedAssets( FUsdStageImportContext& ImportContext, TSet<UObject*>& AssetsToImport, TMap<UObject*, UObject*>& ObjectsToRemap, TMap<FSoftObjectPath, FSoftObjectPath>& SoftObjectsToRemap )
+	{
+		const bool bImportSkeletalAnimations = ImportContext.ImportOptions->bImportGeometry && ImportContext.ImportOptions->bImportSkeletalAnimations;
+
+		for ( TSet<UObject*>::TIterator It( AssetsToImport ); It; ++It )
+		{
+			UObject* Asset = *It;
+
+			if ( !Asset )
+			{
+				It.RemoveCurrent();
+				continue;
+			}
+
+			if (
+				( !ImportContext.ImportOptions->bImportGeometry && ( Asset->IsA<UStaticMesh>() || Asset->IsA<USkeletalMesh>() || Asset->IsA<USkeleton>() || Asset->IsA<UGeometryCache>() ) ) ||
+				( !bImportSkeletalAnimations && ( Asset->IsA<UAnimSequence>() ) ) ||
+				( !ImportContext.ImportOptions->bImportLevelSequences && ( Asset->IsA<ULevelSequence>() ) ) ||
+				( !ImportContext.ImportOptions->bImportMaterials && ( Asset->IsA<UMaterialInterface>() || Asset->IsA<UTexture>() ) )
+			)
+			{
+				ObjectsToRemap.Add( Asset, nullptr );
+				SoftObjectsToRemap.Add( Asset, nullptr );
+				It.RemoveCurrent();
+			}
+		}
+	}
+
+	// We need to recreate the render state for some mesh component types in case we changed the materials that are assigned to them.
+	// Also, skeletal mesh components need to be manually ticked, or else they may be showing an animated state of an animation that
+	// we chose not to import, and wouldn't update otherwise until manually ticked by the user (or after save/reload), which may look
+	// like a bug
+	void RefreshComponents( AActor* RootSceneActor )
+	{
+		if ( !RootSceneActor )
+		{
+			return;
+		}
+
+		TArray<USceneComponent*> Components;
+		const bool bIncludeAllDescendants = true;
+		RootSceneActor->GetRootComponent()->GetChildrenComponents( bIncludeAllDescendants, Components );
+
+		for ( USceneComponent* Component : Components )
+		{
+			if ( USkeletalMeshComponent* SkeletalMeshComponent = Cast<USkeletalMeshComponent>( Component ) )
+			{
+				if ( SkeletalMeshComponent->AnimationData.AnimToPlay == nullptr )
+				{
+					SkeletalMeshComponent->TickAnimation( 0.f, false );
+					SkeletalMeshComponent->RefreshBoneTransforms();
+					SkeletalMeshComponent->RefreshSlaveComponents();
+					SkeletalMeshComponent->UpdateComponentToWorld();
+					SkeletalMeshComponent->FinalizeBoneTransform();
+					SkeletalMeshComponent->MarkRenderTransformDirty();
+					SkeletalMeshComponent->MarkRenderDynamicDataDirty();
+				}
+
+				// It does need us to manually set this to dirty regardless or else it won't update in case we changed material
+				// assignments
+				SkeletalMeshComponent->MarkRenderStateDirty();
+			}
+		}
+	}
 }
 
 void UUsdStageImporter::ImportFromFile(FUsdStageImportContext& ImportContext)
@@ -1393,20 +1464,20 @@ void UUsdStageImporter::ImportFromFile(FUsdStageImportContext& ImportContext)
 	ImportContext.AssetCache->MarkAssetsAsStale();
 	ImportContext.LevelSequenceHelper.SetAssetCache( ImportContext.AssetCache );
 
-	// Shotgun approach to recreate all render states because we may want to reimport/delete/reassing a material/static/skeletalmesh while it is currently being drawn
+	// Shotgun approach to recreate all render states because we may want to reimport/delete/reassign a material/static/skeletalmesh while it is currently being drawn
 	FGlobalComponentRecreateRenderStateContext RecreateRenderStateContext;
 
 	TSharedRef<FUsdSchemaTranslationContext> TranslationContext = MakeShared<FUsdSchemaTranslationContext>( ImportContext.Stage, *ImportContext.AssetCache );
 	TranslationContext->Level = ImportContext.World->GetCurrentLevel();
 	TranslationContext->ObjectFlags = ImportContext.ImportObjectFlags;
 	TranslationContext->Time = static_cast< float >( UsdUtils::GetDefaultTimeCode() );
-	TranslationContext->PurposesToLoad = (EUsdPurpose) ImportContext.ImportOptions->PurposesToImport;
+	TranslationContext->PurposesToLoad = ( EUsdPurpose ) ImportContext.ImportOptions->PurposesToImport;
 	TranslationContext->NaniteTriangleThreshold = ImportContext.ImportOptions->NaniteTriangleThreshold;
 	TranslationContext->RenderContext = ImportContext.ImportOptions->RenderContextToImport;
 	TranslationContext->ParentComponent = ImportContext.SceneActor ? ImportContext.SceneActor->GetRootComponent() : nullptr;
 	TranslationContext->KindsToCollapse = ( EUsdDefaultKind ) ImportContext.ImportOptions->KindsToCollapse;
 	TranslationContext->bAllowInterpretingLODs = ImportContext.ImportOptions->bInterpretLODs;
-	TranslationContext->bAllowParsingSkeletalAnimations = ImportContext.ImportOptions->bImportSkeletalAnimations;
+	TranslationContext->bAllowParsingSkeletalAnimations = ImportContext.ImportOptions->bImportGeometry && ImportContext.ImportOptions->bImportSkeletalAnimations;
 	TranslationContext->MaterialToPrimvarToUVIndex = &ImportContext.MaterialToPrimvarToUVIndex;
 	TranslationContext->BlendShapesByPath = &BlendShapesByPath;
 	{
@@ -1417,6 +1488,7 @@ void UUsdStageImporter::ImportFromFile(FUsdStageImportContext& ImportContext)
 	TranslationContext->CompleteTasks();
 
 	UsdStageImporterImpl::CollectUsedAssetDependencies( ImportContext, UsedAssetsAndDependencies );
+	UsdStageImporterImpl::PruneUnwantedAssets( ImportContext, UsedAssetsAndDependencies, ObjectsToRemap, SoftObjectsToRemap );
 	UsdStageImporterImpl::UpdateAssetImportData( UsedAssetsAndDependencies, ImportContext.FilePath, ImportContext.ImportOptions );
 	UsdStageImporterImpl::PublishAssets( ImportContext, UsedAssetsAndDependencies, ObjectsToRemap, SoftObjectsToRemap );
 	UsdStageImporterImpl::ResolveActorConflicts( ImportContext, ExistingSceneActor, ObjectsToRemap, SoftObjectsToRemap );
@@ -1424,6 +1496,7 @@ void UUsdStageImporter::ImportFromFile(FUsdStageImportContext& ImportContext)
 	UsdStageImporterImpl::RemapSoftReferences( ImportContext, UsedAssetsAndDependencies, SoftObjectsToRemap );
 	UsdStageImporterImpl::Cleanup( ImportContext.SceneActor, ExistingSceneActor, ImportContext.ImportOptions->ExistingActorPolicy );
 	UsdStageImporterImpl::NotifyAssetRegistry( UsedAssetsAndDependencies );
+	UsdStageImporterImpl::RefreshComponents( ImportContext.SceneActor );
 
 	FUsdDelegates::OnPostUsdImport.Broadcast( ImportContext.FilePath );
 
@@ -1487,7 +1560,7 @@ bool UUsdStageImporter::ReimportSingleAsset(FUsdStageImportContext& ImportContex
 	TranslationContext->NaniteTriangleThreshold = ImportContext.ImportOptions->NaniteTriangleThreshold;
 	TranslationContext->KindsToCollapse = ( EUsdDefaultKind ) ImportContext.ImportOptions->KindsToCollapse;
 	TranslationContext->bAllowInterpretingLODs = ImportContext.ImportOptions->bInterpretLODs;
-	TranslationContext->bAllowParsingSkeletalAnimations = ImportContext.ImportOptions->bImportSkeletalAnimations;
+	TranslationContext->bAllowParsingSkeletalAnimations = ImportContext.ImportOptions->bImportGeometry && ImportContext.ImportOptions->bImportSkeletalAnimations;
 	TranslationContext->MaterialToPrimvarToUVIndex = &ImportContext.MaterialToPrimvarToUVIndex;
 	TranslationContext->BlendShapesByPath = &BlendShapesByPath;
 	{
@@ -1533,6 +1606,7 @@ bool UUsdStageImporter::ReimportSingleAsset(FUsdStageImportContext& ImportContex
 
 	UsdStageImporterImpl::Cleanup( ImportContext.SceneActor, nullptr, ImportContext.ImportOptions->ExistingActorPolicy );
 	UsdStageImporterImpl::NotifyAssetRegistry( { ReimportedObject } );
+	UsdStageImporterImpl::RefreshComponents( ImportContext.SceneActor );
 
 	FUsdDelegates::OnPostUsdImport.Broadcast(ImportContext.FilePath);
 
diff --git a/Engine/Plugins/Importers/USDImporter/Source/USDStageImporter/Public/USDStageImportOptions.h b/Engine/Plugins/Importers/USDImporter/Source/USDStageImporter/Public/USDStageImportOptions.h
--- a/Engine/Plugins/Importers/USDImporter/Source/USDStageImporter/Public/USDStageImportOptions.h
+++ b/Engine/Plugins/Importers/USDImporter/Source/USDStageImporter/Public/USDStageImportOptions.h
@@ -57,7 +57,7 @@ public:
 	UPROPERTY(BlueprintReadWrite, config, EditAnywhere, Category = "DataToImport", meta = (EditCondition = bImportGeometry, DisplayName = "Skeletal Animations"))
 	bool bImportSkeletalAnimations;
 
-	UPROPERTY( BlueprintReadWrite, config, EditAnywhere, Category = "DataToImport", meta = ( EditCondition = bImportSkeletalAnimations, DisplayName = "LevelSequences" ) )
+	UPROPERTY( BlueprintReadWrite, config, EditAnywhere, Category = "DataToImport", meta = (DisplayName = "LevelSequences" ) )
 	bool bImportLevelSequences;
 
 	UPROPERTY(BlueprintReadWrite, config, EditAnywhere, Category = "DataToImport", meta = (DisplayName = "Materials & Textures"))
