USD: Fix LOD meshes always being imported with the "SM_LOD0" name.

diff --git a/Engine/Plugins/Importers/USDImporter/Source/USDSchemas/Private/USDGeomMeshTranslator.cpp b/Engine/Plugins/Importers/USDImporter/Source/USDSchemas/Private/USDGeomMeshTranslator.cpp
--- a/Engine/Plugins/Importers/USDImporter/Source/USDSchemas/Private/USDGeomMeshTranslator.cpp
+++ b/Engine/Plugins/Importers/USDImporter/Source/USDSchemas/Private/USDGeomMeshTranslator.cpp
@@ -541,7 +541,7 @@ namespace UsdGeomMeshTranslatorImpl
 		}
 	}
 
-	UStaticMesh* CreateStaticMesh( TArray<FMeshDescription>& LODIndexToMeshDescription, FUsdSchemaTranslationContext& Context, const FString& PrimPath, const bool bShouldEnableNanite, bool& bOutIsNew )
+	UStaticMesh* CreateStaticMesh( TArray<FMeshDescription>& LODIndexToMeshDescription, FUsdSchemaTranslationContext& Context, const FString& MeshName, const bool bShouldEnableNanite, bool& bOutIsNew )
 	{
 		UStaticMesh* StaticMesh = nullptr;
 
@@ -588,7 +588,7 @@ namespace UsdGeomMeshTranslatorImpl
 		{
 			bOutIsNew = true;
 
-			FName AssetName = MakeUniqueObjectName( GetTransientPackage(), UStaticMesh::StaticClass(), *FPaths::GetBaseFilename( PrimPath ) );
+			FName AssetName = MakeUniqueObjectName( GetTransientPackage(), UStaticMesh::StaticClass(), *FPaths::GetBaseFilename( MeshName ) );
 			StaticMesh = NewObject< UStaticMesh >( GetTransientPackage(), AssetName, Context.ObjectFlags | EObjectFlags::RF_Public );
 
 #if WITH_EDITOR
@@ -1008,10 +1008,25 @@ void FBuildStaticMeshTaskChain::SetupTasks()
 			FModuleManager::LoadModuleChecked< IMeshBuilderModule >( TEXT("MeshBuilder") );
 #endif // WITH_EDITOR
 
-			bool bIsNew = true;
 			const FString PrimPathString = PrimPath.GetString();
+
+			// It's useful to have the LOD Mesh prims be named "LOD0", "LOD1", etc. within the LOD variants so that we
+			// can easily tell which Mesh is actually meant to be the LOD mesh (in case there are more Meshes in each
+			// variant or other Meshes outside of the variant), but it's not ideal to have all the generated assets end
+			// up imported as "SM_LOD0_22", "SM_LOD0_23", etc. So here we fetch the parent prim name in case we're a LOD
+			FString MeshName;
+			if ( Context->bAllowInterpretingLODs && UsdUtils::IsGeomMeshALOD( GetPrim() ) )
+			{
+				MeshName = PrimPath.GetParentPath().GetString();
+			}
+			else
+			{
+				MeshName = PrimPath.GetString();
+			}
+
+			bool bIsNew = true;
 			const bool bShouldEnableNanite = UsdGeomMeshTranslatorImpl::ShouldEnableNanite( LODIndexToMeshDescription, LODIndexToMaterialInfo, *Context, GetPrim() );
-			StaticMesh = UsdGeomMeshTranslatorImpl::CreateStaticMesh( LODIndexToMeshDescription, *Context, PrimPathString, bShouldEnableNanite, bIsNew );
+			StaticMesh = UsdGeomMeshTranslatorImpl::CreateStaticMesh( LODIndexToMeshDescription, *Context, MeshName, bShouldEnableNanite, bIsNew );
 
 			if ( StaticMesh )
 			{
diff --git a/Engine/Plugins/Importers/USDImporter/Source/USDStageImporter/Private/USDStageImporter.cpp b/Engine/Plugins/Importers/USDImporter/Source/USDStageImporter/Private/USDStageImporter.cpp
--- a/Engine/Plugins/Importers/USDImporter/Source/USDStageImporter/Private/USDStageImporter.cpp
+++ b/Engine/Plugins/Importers/USDImporter/Source/USDStageImporter/Private/USDStageImporter.cpp
@@ -381,6 +381,21 @@ namespace UsdStageImporterImpl
 			if (UUsdAssetImportData* AssetImportData = Cast<UUsdAssetImportData>(Mesh->AssetImportData))
 			{
 				AssetPath = AssetImportData->PrimPath;
+
+				// If we have multiple LODs here we must have parsed the LOD variant set pattern. If our prims were named with the LOD
+				// pattern, go from e.g. '/Root/MyMesh/LOD0' to '/Root/MyMesh', or else every single LOD mesh will be named "SM_LOD0_X".
+				// We'll actually check though because if the user set a custom name for his prim other than LOD0 then we'll keep that
+				if ( Mesh->GetNumLODs() > 1 )
+				{
+					FString PrimName = FPaths::GetBaseFilename( AssetPath );
+					if ( PrimName.RemoveFromStart( TEXT( "LOD" ), ESearchCase::CaseSensitive ) )
+					{
+						if ( PrimName.IsNumeric() )
+						{
+							AssetPath = FPaths::GetPath( AssetPath );
+						}
+					}
+				}
 			}
 		}
 		else if (USkeletalMesh* SkMesh = Cast<USkeletalMesh>(Asset))
