USD: Fetch references from all specs when building the references panel.

diff --git a/Engine/Plugins/Importers/USDImporter/Source/USDStageEditorViewModels/Private/USDReferencesViewModel.cpp b/Engine/Plugins/Importers/USDImporter/Source/USDStageEditorViewModels/Private/USDReferencesViewModel.cpp
--- a/Engine/Plugins/Importers/USDImporter/Source/USDStageEditorViewModels/Private/USDReferencesViewModel.cpp
+++ b/Engine/Plugins/Importers/USDImporter/Source/USDStageEditorViewModels/Private/USDReferencesViewModel.cpp
@@ -14,6 +14,7 @@
 #include "USDIncludesStart.h"
 	#include "pxr/usd/sdf/reference.h"
 	#include "pxr/usd/usd/prim.h"
+	#include "pxr/usd/usd/primCompositionQuery.h"
 #include "USDIncludesEnd.h"
 
 #endif // #if USE_USD_SDK
@@ -32,19 +33,22 @@ void FUsdReferencesViewModel::UpdateReferences( const UE::FUsdStageWeak& UsdStag
 
 	pxr::UsdPrim Prim( UsdStage.GetPrimAtPath( UE::FSdfPath( PrimPath ) ) );
 
-	// Retrieve the strongest opinion prim spec, hopefully it's enough to get all references
-	pxr::SdfPrimSpecHandle PrimSpec = Prim.GetPrimStack().size() > 0 ? Prim.GetPrimStack()[0] : pxr::SdfPrimSpecHandle();
+	pxr::UsdPrimCompositionQuery PrimCompositionQuery = pxr::UsdPrimCompositionQuery::GetDirectReferences( Prim );
 
-	if ( PrimSpec )
+	for ( const pxr::UsdPrimCompositionQueryArc& CompositionArc : PrimCompositionQuery.GetCompositionArcs() )
 	{
-		pxr::SdfReferencesProxy ReferencesProxy = PrimSpec->GetReferenceList();
-
-		for ( const pxr::SdfReference& UsdReference : ReferencesProxy.GetAddedOrExplicitItems() )
+		if ( CompositionArc.GetArcType() == pxr::PcpArcTypeReference )
 		{
-			FUsdReference Reference;
-			Reference.AssetPath = UsdToUnreal::ConvertString( UsdReference.GetAssetPath() );
+			pxr::SdfReferenceEditorProxy ReferenceEditor;
+			pxr::SdfReference UsdReference;
+
+			if ( CompositionArc.GetIntroducingListEditor( &ReferenceEditor, &UsdReference ) )
+			{
+				FUsdReference Reference;
+				Reference.AssetPath = UsdToUnreal::ConvertString( UsdReference.GetAssetPath() );
 
-			References.Add( MakeSharedUnreal< FUsdReference >( MoveTemp( Reference ) ) );
+				References.Add( MakeSharedUnreal< FUsdReference >( MoveTemp( Reference ) ) );
+			}
 		}
 	}
 #endif // #if USE_USD_SDK
