USD: Fix inconsistent warning when exporting levels and baking materials.

diff --git a/Engine/Plugins/Importers/USDImporter/Source/USDUtilities/Private/USDGeomMeshConversion.cpp b/Engine/Plugins/Importers/USDImporter/Source/USDUtilities/Private/USDGeomMeshConversion.cpp
--- a/Engine/Plugins/Importers/USDImporter/Source/USDUtilities/Private/USDGeomMeshConversion.cpp
+++ b/Engine/Plugins/Importers/USDImporter/Source/USDUtilities/Private/USDGeomMeshConversion.cpp
@@ -2099,6 +2099,13 @@ void UsdUtils::ReplaceUnrealMaterialsWithBaked(
 				}
 			}
 
+			// Don't try fetching attributes from the pseudoroot as we'll obviously never have a material binding here
+			// and we may get some USD warnings
+			if ( Prim.IsPseudoRoot() )
+			{
+				return;
+			}
+
 			std::string UnrealMaterialAttrAssetPath;
 			FString UnrealMaterialPrimAssetPath;
 
