USD: Fix visibility track changes not becoming timeSamples if writing to a sublayer in some scenarios.

diff --git a/Engine/Plugins/Importers/USDImporter/Source/USDUtilities/Private/USDPrimConversion.cpp b/Engine/Plugins/Importers/USDImporter/Source/USDUtilities/Private/USDPrimConversion.cpp
--- a/Engine/Plugins/Importers/USDImporter/Source/USDUtilities/Private/USDPrimConversion.cpp
+++ b/Engine/Plugins/Importers/USDImporter/Source/USDUtilities/Private/USDPrimConversion.cpp
@@ -2841,14 +2841,19 @@ UnrealToUsd::FPropertyTrackWriter UnrealToUsd::CreatePropertyTrackWriter( const
 				Attr = Imageable.CreateVisibilityAttr();
 				if ( Attr )
 				{
-					Result.BoolWriter = [Imageable]( bool UEValue, double UsdTimeCode )
+					Result.BoolWriter = [Imageable, Attr]( bool UEValue, double UsdTimeCode )
 					{
 						if ( UEValue )
 						{
+							// We have to do both here as MakeVisible will ensure we also flip any parent prims,
+							// and setting the attribute will ensure we write a timeSample. Otherwise if MakeVisible
+							// finds that the prim should already be visible due to a stronger opinion, it won't write anything
+							Attr.Set<pxr::TfToken>( pxr::UsdGeomTokens->inherited, UsdTimeCode );
 							Imageable.MakeVisible( UsdTimeCode );
 						}
 						else
 						{
+							Attr.Set<pxr::TfToken>( pxr::UsdGeomTokens->invisible, UsdTimeCode );
 							Imageable.MakeInvisible( UsdTimeCode );
 						}
 					};
