USD: Restrict export LOD range sliders to MAX_MESH_LOD_COUNT-1.

diff --git a/Engine/Plugins/Importers/USDImporter/Source/USDExporter/Private/USDExporterModule.cpp b/Engine/Plugins/Importers/USDImporter/Source/USDExporter/Private/USDExporterModule.cpp
--- a/Engine/Plugins/Importers/USDImporter/Source/USDExporter/Private/USDExporterModule.cpp
+++ b/Engine/Plugins/Importers/USDImporter/Source/USDExporter/Private/USDExporterModule.cpp
@@ -3,6 +3,7 @@
 #include "USDExporterModule.h"
 
 #include "LevelExporterUSDOptionsCustomization.h"
+#include "USDAssetOptions.h"
 #include "USDMemory.h"
 
 #include "Modules/ModuleManager.h"
@@ -18,6 +19,21 @@ public:
 		// We intentionally use the same customization for both of these
 		PropertyModule.RegisterCustomClassLayout( TEXT( "LevelExporterUSDOptions" ), FOnGetDetailCustomizationInstance::CreateStatic( &FLevelExporterUSDOptionsCustomization::MakeInstance ) );
 		PropertyModule.RegisterCustomClassLayout( TEXT( "LevelSequenceExporterUSDOptions" ), FOnGetDetailCustomizationInstance::CreateStatic( &FLevelExporterUSDOptionsCustomization::MakeInstance ) );
+
+		// Modify the static mesh LOD range to have the proper define value as the maximum.
+		// We have to do this the hard way here because we can't use the define within the meta tag itself
+		if ( UScriptStruct* ScriptStruct = FUsdMeshAssetOptions::StaticStruct() )
+		{
+			if ( FProperty* LowestMeshLODProperty = ScriptStruct->FindPropertyByName( GET_MEMBER_NAME_CHECKED( FUsdMeshAssetOptions, LowestMeshLOD ) ) )
+			{
+				LowestMeshLODProperty->SetMetaData( TEXT( "ClampMax" ), LexToString( MAX_MESH_LOD_COUNT - 1 ) );
+			}
+
+			if ( FProperty* HighestMeshLODProperty = ScriptStruct->FindPropertyByName( GET_MEMBER_NAME_CHECKED( FUsdMeshAssetOptions, HighestMeshLOD ) ) )
+			{
+				HighestMeshLODProperty->SetMetaData( TEXT( "ClampMax" ), LexToString( MAX_MESH_LOD_COUNT - 1 ) );
+			}
+		}
 	}
 
 	virtual void ShutdownModule() override
