USD: Change component mobility before trying to set transforms during PIE, or else we get a warning.

diff --git a/Engine/Plugins/Importers/USDImporter/Source/USDSchemas/Private/USDGeomXformableTranslator.cpp b/Engine/Plugins/Importers/USDImporter/Source/USDSchemas/Private/USDGeomXformableTranslator.cpp
--- a/Engine/Plugins/Importers/USDImporter/Source/USDSchemas/Private/USDGeomXformableTranslator.cpp
+++ b/Engine/Plugins/Importers/USDImporter/Source/USDSchemas/Private/USDGeomXformableTranslator.cpp
@@ -350,6 +350,14 @@ void FUsdGeomXformableTranslator::UpdateComponents( USceneComponent* SceneCompon
 	{
 		SceneComponent->Modify();
 
+		// UsdToUnreal::ConvertXformable will set a new transform, which will emit warnings during PIE/Runtime if the component
+		// has Static mobility, so here we unregister, set the new transform value, and reregister below
+		const bool bStaticMobility = SceneComponent->Mobility == EComponentMobility::Static;
+		if ( bStaticMobility )
+		{
+			SceneComponent->UnregisterComponent();
+		}
+
 		UsdToUnreal::ConvertXformable( Context->Stage, pxr::UsdGeomXformable( GetPrim() ), *SceneComponent, Context->Time );
 
 		// If the user modified a mesh parameter (e.g. vertex color), the hash will be different and it will become a separate asset
@@ -377,6 +385,14 @@ void FUsdGeomXformableTranslator::UpdateComponents( USceneComponent* SceneCompon
 				StaticMeshComponent->RegisterComponent();
 			}
 		}
+
+		// Note how we should only register if we unregistered ourselves: If we did this every time we would
+		// register too early during the process of duplicating into PIE, and that would prevent a future RegisterComponent
+		// call from naturally creating the required render state
+		if ( bStaticMobility && !SceneComponent->IsRegistered() )
+		{
+			SceneComponent->RegisterComponent();
+		}
 	}
 }
 
