USD: Fix "F" not focusing on some spawned actors and components.

diff --git a/Engine/Plugins/Importers/USDImporter/Source/USDStageEditor/Private/SUSDStage.cpp b/Engine/Plugins/Importers/USDImporter/Source/USDStageEditor/Private/SUSDStage.cpp
--- a/Engine/Plugins/Importers/USDImporter/Source/USDStageEditor/Private/SUSDStage.cpp
+++ b/Engine/Plugins/Importers/USDImporter/Source/USDStageEditor/Private/SUSDStage.cpp
@@ -65,12 +65,21 @@ namespace SUSDStageImpl
 		}
 
 		TSet<AActor*> ActorsToSelect;
-		for ( USceneComponent* ComponentToSelect : ComponentsToSelect )
+		for ( TSet<USceneComponent*>::TIterator It(ComponentsToSelect); It; ++It )
 		{
-			if ( AActor* Owner = ComponentToSelect->GetOwner() )
+			if ( AActor* Owner = (*It)->GetOwner() )
 			{
 				// We always need the parent actor selected to select a component
 				ActorsToSelect.Add( Owner );
+
+				// If we're going to select a root component, select the actor instead. This is useful
+				// because if we later press "F" to focus on the prim, having the actor selected will use the entire
+				// actor's bounding box and focus on something. If all we have selected is an empty scene component
+				// however, the camera won't focus on anything
+				if ( *It == Owner->GetRootComponent() )
+				{
+					It.RemoveCurrent();
+				}
 			}
 		}
 
@@ -107,6 +116,7 @@ void SUsdStage::Construct( const FArguments& InArgs )
 	ViewModel.UsdStageActor = UsdStageModule.FindUsdStageActor( GWorld );
 
 	bUpdatingViewportSelection = false;
+	bUpdatingPrimSelection = false;
 
 	UE::FUsdStage UsdStage;
 
@@ -1158,6 +1168,11 @@ void SUsdStage::OnImport()
 
 void SUsdStage::OnPrimSelectionChanged( const TArray<FString>& PrimPaths )
 {
+	if ( bUpdatingPrimSelection )
+	{
+		return;
+	}
+
 	AUsdStageActor* StageActor = ViewModel.UsdStageActor.Get();
 	if ( !StageActor )
 	{
@@ -1271,6 +1286,8 @@ void SUsdStage::OnViewportSelectionChanged( UObject* NewSelection )
 		return;
 	}
 
+	TGuardValue<bool> SelectionLoopGuard( bUpdatingPrimSelection, true );
+
 	TArray<USceneComponent*> SelectedComponents;
 	{
 		Selection->GetSelectedObjects<USceneComponent>( SelectedComponents );
diff --git a/Engine/Plugins/Importers/USDImporter/Source/USDStageEditor/Private/SUSDStage.h b/Engine/Plugins/Importers/USDImporter/Source/USDStageEditor/Private/SUSDStage.h
--- a/Engine/Plugins/Importers/USDImporter/Source/USDStageEditor/Private/SUSDStage.h
+++ b/Engine/Plugins/Importers/USDImporter/Source/USDStageEditor/Private/SUSDStage.h
@@ -99,6 +99,9 @@ protected:
 	// True while we're in the middle of setting the viewport selection from the prim selection
 	bool bUpdatingViewportSelection;
 
+	// True while we're in the middle of updating the prim selection from the viewport selection
+	bool bUpdatingPrimSelection;
+
 	// We keep this menu alive instead of recreating it every time because it doesn't expose
 	// setting/getting/responding to the picked world, which resets every time it is reconstructed (see UE-127879)
 	// By keeping it alive it will keep its own state and we just do a full refresh when needed
