USD: Fix crash when switching to a level after exporting it;

Fix crash when cancelling out of export options and switching to the level that was going to be exported.

diff --git a/Engine/Plugins/Importers/USDImporter/Source/USDExporter/Private/LevelExporterUSDOptionsCustomization.cpp b/Engine/Plugins/Importers/USDImporter/Source/USDExporter/Private/LevelExporterUSDOptionsCustomization.cpp
--- a/Engine/Plugins/Importers/USDImporter/Source/USDExporter/Private/LevelExporterUSDOptionsCustomization.cpp
+++ b/Engine/Plugins/Importers/USDImporter/Source/USDExporter/Private/LevelExporterUSDOptionsCustomization.cpp
@@ -177,8 +177,10 @@ void FLevelExporterUSDOptionsCustomization::CustomizeDetails(IDetailLayoutBuilde
 	{
 		OptionsPtr.Reset( Options );
 
-		UWorld* World = Options->CurrentTask
-			? Cast<UWorld>( Options->CurrentTask->Object )
+		UAssetExportTask* Task = Options->CurrentTask.Get();
+
+		UWorld* World = Task
+			? Cast<UWorld>( Task->Object )
 			: GEditor->GetEditorWorldContext().World();
 
 		PickerTree = SNew( LevelExporterUSDImpl::SLevelPickerList, &Options->Inner, World );
diff --git a/Engine/Plugins/Importers/USDImporter/Source/USDExporter/Private/USDConversionBlueprintLibrary.cpp b/Engine/Plugins/Importers/USDImporter/Source/USDExporter/Private/USDConversionBlueprintLibrary.cpp
--- a/Engine/Plugins/Importers/USDImporter/Source/USDExporter/Private/USDConversionBlueprintLibrary.cpp
+++ b/Engine/Plugins/Importers/USDImporter/Source/USDExporter/Private/USDConversionBlueprintLibrary.cpp
@@ -26,10 +26,7 @@ namespace UE
 	{
 		namespace Private
 		{
-			// I haven't managed to get nested sublevels like this to even load in the engine at all,
-			// but UWorldPartitionConvertCommandlet::GatherAndPrepareSubLevelsToConvert implies that it's possible,
-			// so we're doing it recursively just in case too
-			void StreamInLevelRecursively( ULevel* Level, const TSet<FString>& LevelsToIgnore )
+			void StreamInLevels( ULevel* Level, const TSet<FString>& LevelsToIgnore )
 			{
 				UWorld* InnerWorld = Level->GetTypedOuter<UWorld>();
 				if ( !InnerWorld || InnerWorld->GetStreamingLevels().Num() == 0 )
@@ -128,17 +125,6 @@ namespace UE
 				// Synchronously show levels right now
 				InnerWorld->FlushLevelStreaming( EFlushLevelStreamingType::Visibility );
 
-				for ( ULevelStreaming* StreamingLevel : InnerWorld->GetStreamingLevels() )
-				{
-					if ( StreamingLevel->IsLevelLoaded() )
-					{
-						if ( ULevel* StreamedLevel = StreamingLevel->GetLoadedLevel() )
-						{
-							StreamInLevelRecursively( StreamedLevel, LevelsToIgnore );
-						}
-					}
-				}
-
 				if ( bCreatedContext )
 				{
 					GEngine->DestroyWorldContext( InnerWorld );
@@ -157,7 +143,7 @@ void UUsdConversionBlueprintLibrary::StreamInRequiredLevels( UWorld* World, cons
 
 	if ( ULevel* PersistentLevel = World->PersistentLevel )
 	{
-		UE::UsdConversionBlueprintLibraryImpl::Private::StreamInLevelRecursively( PersistentLevel, LevelsToIgnore );
+		UE::UsdConversionBlueprintLibraryImpl::Private::StreamInLevels( PersistentLevel, LevelsToIgnore );
 	}
 }
 
diff --git a/Engine/Plugins/Importers/USDImporter/Source/USDExporter/Public/LevelExporterUSDOptions.h b/Engine/Plugins/Importers/USDImporter/Source/USDExporter/Public/LevelExporterUSDOptions.h
--- a/Engine/Plugins/Importers/USDImporter/Source/USDExporter/Public/LevelExporterUSDOptions.h
+++ b/Engine/Plugins/Importers/USDImporter/Source/USDExporter/Public/LevelExporterUSDOptions.h
@@ -96,9 +96,11 @@ public:
 
 public:
 	// We temporarily stash our export task here as a way of passing our options down to
-	// the Python exporter, that does the actual level exporting
+	// the Python exporter, that does the actual level exporting.
+	// This is weak because we often use the CDO of this class directly, and we never want to
+	// permanently hold on to a particular export task
 	UPROPERTY( VisibleAnywhere, BlueprintReadOnly, Category = Hidden )
-	UAssetExportTask* CurrentTask;
+	TWeakObjectPtr<UAssetExportTask> CurrentTask;
 
 private:
 	UFUNCTION()
