USD: Fix GetUniqueName generating names like "Prim_0_3" in some scenarios;

Remove duplicate unique name generation function from stage importer.

diff --git a/Engine/Plugins/Importers/USDImporter/Source/USDStageImporter/Private/USDStageImporter.cpp b/Engine/Plugins/Importers/USDImporter/Source/USDStageImporter/Private/USDStageImporter.cpp
--- a/Engine/Plugins/Importers/USDImporter/Source/USDStageImporter/Private/USDStageImporter.cpp
+++ b/Engine/Plugins/Importers/USDImporter/Source/USDStageImporter/Private/USDStageImporter.cpp
@@ -111,49 +111,6 @@ namespace UsdStageImporterImpl
 		return Suffix == 1 ? InPackagePath : InPackagePath + TEXT("_") + LexToString(Suffix - 1);
 	}
 
-	/**
-	 * Removes any numbered suffix, followed by any number of underscores (e.g. Asset_2, Asset__23231 or Asset94 become 'Asset'), making
-	 * sure the string is kept at least one character long.
-	*/
-	void RemoveNumberedSuffix( FString& Prefix )
-	{
-		if ( !Prefix.IsNumeric() )
-		{
-			FString LastChar = Prefix.Right( 1 );
-			while ( LastChar.IsNumeric() )
-			{
-				const bool bAllowShrinking = false;
-				Prefix.LeftChopInline( 1, bAllowShrinking );
-				LastChar = Prefix.Right( 1 );
-			}
-			Prefix.Shrink();
-		}
-
-		while ( Prefix.Len() > 1 && Prefix.Right( 1 ) == TEXT( "_" ) )
-		{
-			Prefix.RemoveFromEnd( TEXT( "_" ) );
-		}
-	}
-
-	FString GetUniqueName(FString Prefix, TSet<FString>& UniqueNames)
-	{
-		if (!UniqueNames.Contains(Prefix))
-		{
-			return Prefix;
-		}
-
-		RemoveNumberedSuffix(Prefix);
-
-		int32 Suffix = 0;
-		FString Result;
-		do
-		{
-			Result = FString::Printf(TEXT("%s_%d"), *Prefix, Suffix++);
-		} while (UniqueNames.Contains(Result));
-
-		return Result;
-	}
-
 	void SetupSceneActor(FUsdStageImportContext& ImportContext)
 	{
 		if ( !ImportContext.ImportOptions->bImportActors )
@@ -463,7 +420,7 @@ namespace UsdStageImporterImpl
 		// We don't care if our assets overwrite something in the final destination package (that conflict will be
 		// handled according to EReplaceAssetPolicy). But we do want these assets to have unique names amongst themselves
 		// or else they will overwrite each other when publishing
-		FinalName = GetUniqueName( ObjectTools::SanitizeObjectName( FinalName ), UniqueAssetNames );
+		FinalName = UsdUtils::GetUniqueName( ObjectTools::SanitizeObjectName( FinalName ), UniqueAssetNames );
 		UniqueAssetNames.Add(FinalName);
 
 		return FinalName;
diff --git a/Engine/Plugins/Importers/USDImporter/Source/USDUtilities/Private/USDConversionUtils.cpp b/Engine/Plugins/Importers/USDImporter/Source/USDUtilities/Private/USDConversionUtils.cpp
--- a/Engine/Plugins/Importers/USDImporter/Source/USDUtilities/Private/USDConversionUtils.cpp
+++ b/Engine/Plugins/Importers/USDImporter/Source/USDUtilities/Private/USDConversionUtils.cpp
@@ -1248,6 +1248,29 @@ bool UsdUtils::RenamePrim( UE::FUsdPrim& Prim, const TCHAR* NewPrimName )
 #endif // #if USE_USD_SDK
 }
 
+bool UsdUtils::RemoveNumberedSuffix( FString& Prefix )
+{
+	if ( Prefix.IsNumeric() )
+	{
+		return false;
+	}
+
+	bool bRemoved = false;
+
+	FString LastChar = Prefix.Right( 1 );
+	while ( ( LastChar.IsNumeric() || LastChar == TEXT( "_" ) ) && Prefix.Len() > 1 )
+	{
+		const bool bAllowShrinking = false;
+		Prefix.LeftChopInline( 1, bAllowShrinking );
+		LastChar = Prefix.Right( 1 );
+
+		bRemoved = true;
+	}
+	Prefix.Shrink();
+
+	return bRemoved;
+}
+
 FString UsdUtils::GetUniqueName( FString Name, const TSet<FString>& UsedNames )
 {
 	if ( !UsedNames.Contains( Name ) )
@@ -1255,6 +1278,14 @@ FString UsdUtils::GetUniqueName( FString Name, const TSet<FString>& UsedNames )
 		return Name;
 	}
 
+	const bool bRemoved = RemoveNumberedSuffix( Name );
+
+	// Its possible that removing the suffix made it into a unique name already
+	if ( bRemoved && !UsedNames.Contains( Name ) )
+	{
+		return Name;
+	}
+
 	int32 Suffix = 0;
 	FString Result;
 	do
diff --git a/Engine/Plugins/Importers/USDImporter/Source/USDUtilities/Public/USDConversionUtils.h b/Engine/Plugins/Importers/USDImporter/Source/USDUtilities/Public/USDConversionUtils.h
--- a/Engine/Plugins/Importers/USDImporter/Source/USDUtilities/Public/USDConversionUtils.h
+++ b/Engine/Plugins/Importers/USDImporter/Source/USDUtilities/Public/USDConversionUtils.h
@@ -195,6 +195,12 @@ namespace UsdUtils
 	 */
 	USDUTILITIES_API bool RenamePrim( UE::FUsdPrim& Prim, const TCHAR* NewPrimName );
 
+	/**
+	 * Removes any numbered suffix, followed by any number of underscores (e.g. Asset_2, Asset__232_31 or Asset94 all become 'Asset'), making
+	 * sure the string is kept at least one character long. Returns true if it removed anything.
+	 */
+	USDUTILITIES_API bool RemoveNumberedSuffix( FString& Prefix );
+
 	/**
 	 * Appends numbered suffixes to Name until the result is not contained in UsedNames, and returns it.
 	 * Does not add the result to UsedNames before returning (as it is const).
