USD: Allow saving as any USD format by default.

diff --git a/Engine/Plugins/Importers/USDImporter/Source/USDUtilities/Private/USDLayerUtils.cpp b/Engine/Plugins/Importers/USDImporter/Source/USDUtilities/Private/USDLayerUtils.cpp
--- a/Engine/Plugins/Importers/USDImporter/Source/USDUtilities/Private/USDLayerUtils.cpp
+++ b/Engine/Plugins/Importers/USDImporter/Source/USDUtilities/Private/USDLayerUtils.cpp
@@ -107,13 +107,25 @@ TOptional< FString > UsdUtils::BrowseUsdFile( EBrowseFileMode Mode, TSharedRef<
 		return {};
 	}
 
+	if ( Mode == EBrowseFileMode::Save )
+	{
+		// USD 21.08 doesn't yet support saving to USDZ, so instead of allowing this option and leading to an error we'll just hide it
+		SupportedExtensions.Remove( TEXT( "usdz" ) );
+	}
+
+	FString JoinedExtensions = FString::Join( SupportedExtensions, TEXT( ";*." ) ); // Combine "usd" and "usda" into "usd; *.usda"
+	FString FileTypes = FString::Printf( TEXT( "Universal Scene Description files (*.%s)|*.%s|" ), *JoinedExtensions, *JoinedExtensions );
+	for ( const FString& SupportedExtension : SupportedExtensions )
+	{
+		// The '(*.%s)' on the actual name (before the '|') is not optional: We need the name part to be different for each format
+		// or else the options will overwrite each other on the Mac
+		FileTypes += FString::Printf( TEXT( "Universal Scene Description file (*.%s)|*.%s|" ), *SupportedExtension, *SupportedExtension );
+	}
+
 	switch ( Mode )
 	{
 		case EBrowseFileMode::Open :
 		{
-			FString JoinedExtensions = FString::Join( SupportedExtensions, TEXT( ";*." ) ); // Combine "usd" and "usda" into "usd; *.usda"
-			FString FileTypes = FString::Printf( TEXT( "Universal Scene Description files|*.%s" ), *JoinedExtensions, *JoinedExtensions );
-
 			if ( !DesktopPlatform->OpenFileDialog( ParentWindowHandle, LOCTEXT( "ChooseFile", "Choose file").ToString(), TEXT(""), TEXT(""), *FileTypes, EFileDialogFlags::None, OutFiles ) )
 			{
 				return {};
@@ -122,20 +134,6 @@ TOptional< FString > UsdUtils::BrowseUsdFile( EBrowseFileMode Mode, TSharedRef<
 		}
 		case EBrowseFileMode::Save :
 		{
-			FString FileTypes;
-			for ( const FString& SupportedExtension : SupportedExtensions )
-			{
-				// USD 21.08 doesn't yet support saving to USDZ, so instead of allowing this option and leading to an error we'll just hide it
-				if ( SupportedExtension == TEXT( "usdz" ) )
-				{
-					continue;
-				}
-
-				// The '(*.%s)' on the actual name (before the '|') is not optional: We need the name part to be different for each format
-				// or else the options will overwrite each other on the Mac
-				FileTypes += FString::Printf( TEXT( "Universal Scene Description file (*.%s)|*.%s|" ), *SupportedExtension, *SupportedExtension );
-			}
-
 			if ( !DesktopPlatform->SaveFileDialog( ParentWindowHandle, LOCTEXT( "ChooseFile", "Choose file").ToString(), TEXT(""), TEXT(""), *FileTypes, EFileDialogFlags::None, OutFiles ) )
 			{
 				return {};
