USD: Fix crash when creating materials while a level is being associated.

diff --git a/Engine/Plugins/Importers/USDImporter/Source/USDSchemas/Private/Custom/MDLUSDShadeMaterialTranslator.cpp b/Engine/Plugins/Importers/USDImporter/Source/USDSchemas/Private/Custom/MDLUSDShadeMaterialTranslator.cpp
--- a/Engine/Plugins/Importers/USDImporter/Source/USDSchemas/Private/Custom/MDLUSDShadeMaterialTranslator.cpp
+++ b/Engine/Plugins/Importers/USDImporter/Source/USDSchemas/Private/Custom/MDLUSDShadeMaterialTranslator.cpp
@@ -180,7 +180,16 @@ void FMdlUsdShadeMaterialTranslator::CreateAssets()
 
 				UsdToUnreal::ConvertShadeInputsToParameters( ShadeMaterial, *MdlMaterialInstance, Context->AssetCache.Get(), *Context->RenderContext.ToString() );
 
-				FMaterialUpdateContext UpdateContext( FMaterialUpdateContext::EOptions::Default, GMaxRHIShaderPlatform );
+				// We can't blindly recreate all component render states when a level is being added, because we may end up first creating
+				// render states for some components, and UWorld::AddToWorld calls FScene::AddPrimitive which expects the component to not have
+				// primitives yet
+				FMaterialUpdateContext::EOptions::Type Options = FMaterialUpdateContext::EOptions::Default;
+				if ( Context->Level->bIsAssociatingLevel )
+				{
+					Options = ( FMaterialUpdateContext::EOptions::Type ) ( Options & ~FMaterialUpdateContext::EOptions::RecreateRenderStates );
+				}
+
+				FMaterialUpdateContext UpdateContext( Options, GMaxRHIShaderPlatform );
 				UpdateContext.AddMaterialInstance( MdlMaterialInstance );
 				MdlMaterialInstance->PreEditChange( nullptr );
 				MdlMaterialInstance->PostEditChange();
diff --git a/Engine/Plugins/Importers/USDImporter/Source/USDSchemas/Private/USDShadeMaterialTranslator.cpp b/Engine/Plugins/Importers/USDImporter/Source/USDSchemas/Private/USDShadeMaterialTranslator.cpp
--- a/Engine/Plugins/Importers/USDImporter/Source/USDSchemas/Private/USDShadeMaterialTranslator.cpp
+++ b/Engine/Plugins/Importers/USDImporter/Source/USDSchemas/Private/USDShadeMaterialTranslator.cpp
@@ -80,7 +80,16 @@ void FUsdShadeMaterialTranslator::CreateAssets()
 
 					if ( UsdToUnreal::ConvertMaterial( ShadeMaterial, *NewMaterial, Context->AssetCache.Get(), PrimvarToUVIndex, *Context->RenderContext.ToString() ) )
 					{
-						FMaterialUpdateContext UpdateContext( FMaterialUpdateContext::EOptions::Default, GMaxRHIShaderPlatform );
+						// We can't blindly recreate all component render states when a level is being added, because we may end up first creating
+						// render states for some components, and UWorld::AddToWorld calls FScene::AddPrimitive which expects the component to not have
+						// primitives yet
+						FMaterialUpdateContext::EOptions::Type Options = FMaterialUpdateContext::EOptions::Default;
+						if ( Context->Level->bIsAssociatingLevel )
+						{
+							Options = ( FMaterialUpdateContext::EOptions::Type ) ( Options & ~FMaterialUpdateContext::EOptions::RecreateRenderStates );
+						}
+
+						FMaterialUpdateContext UpdateContext( Options, GMaxRHIShaderPlatform );
 						UpdateContext.AddMaterialInstance( NewMaterial );
 						NewMaterial->PreEditChange( nullptr );
 						NewMaterial->PostEditChange();
