USD:
* Reduce USD log spamming in particular with the USD GeometryCache and USD mesh Nanite triangle threshold.

diff --git a/Engine/Plugins/Importers/USDImporter/Source/USDSchemas/Private/USDGeomMeshTranslator.cpp b/Engine/Plugins/Importers/USDImporter/Source/USDSchemas/Private/USDGeomMeshTranslator.cpp
--- a/Engine/Plugins/Importers/USDImporter/Source/USDSchemas/Private/USDGeomMeshTranslator.cpp
+++ b/Engine/Plugins/Importers/USDImporter/Source/USDSchemas/Private/USDGeomMeshTranslator.cpp
@@ -122,7 +122,7 @@ namespace UsdGeomMeshTranslatorImpl
 			const int32 NumTriangles = LODIndexToMeshDescription[ 0 ].Triangles().Num();
 			if ( NumTriangles >= Context.NaniteTriangleThreshold )
 			{
-				UE_LOG( LogUsd, Log, TEXT( "Trying to enable Nanite for mesh generated for prim '%s' as it has '%d' triangles, and the threshold is '%d'" ),
+				UE_LOG( LogUsd, Verbose, TEXT( "Trying to enable Nanite for mesh generated for prim '%s' as it has '%d' triangles, and the threshold is '%d'" ),
 					*Prim.GetPrimPath().GetString(),
 					NumTriangles,
 					Context.NaniteTriangleThreshold
@@ -130,7 +130,7 @@ namespace UsdGeomMeshTranslatorImpl
 			}
 			else
 			{
-				UE_LOG( LogUsd, Log, TEXT( "Not enabling Nanite for mesh generated for prim '%s' as it has '%d' triangles, and the threshold is '%d'" ),
+				UE_LOG( LogUsd, Verbose, TEXT( "Not enabling Nanite for mesh generated for prim '%s' as it has '%d' triangles, and the threshold is '%d'" ),
 					*Prim.GetPrimPath().GetString(),
 					NumTriangles,
 					Context.NaniteTriangleThreshold
diff --git a/Engine/Plugins/Importers/USDImporter/Source/USDUtilities/Private/USDErrorUtils.cpp b/Engine/Plugins/Importers/USDImporter/Source/USDUtilities/Private/USDErrorUtils.cpp
--- a/Engine/Plugins/Importers/USDImporter/Source/USDUtilities/Private/USDErrorUtils.cpp
+++ b/Engine/Plugins/Importers/USDImporter/Source/USDUtilities/Private/USDErrorUtils.cpp
@@ -168,6 +168,8 @@ TOptional< UE::Internal::FUsdMessageLog > FUsdLogManager::MessageLog;
 int32 FUsdLogManager::MessageLogRefCount;
 FCriticalSection FUsdLogManager::MessageLogLock;
 
+TSet<FString> LoggedMessages;
+
 void FUsdLogManager::LogMessage( EMessageSeverity::Type Severity, const FText& Message )
 {
 	LogMessage( FTokenizedMessage::Create( Severity, Message ) );
@@ -177,17 +179,24 @@ void FUsdLogManager::LogMessage( const TSharedRef< FTokenizedMessage >& Message
 {
 	bool bMessageProcessed = false;
 
-#if WITH_EDITOR
-	if ( MessageLog )
 	{
-		FScopeLock Lock( &MessageLogLock );
-		if ( MessageLog ) // Make sure it's still valid
+		FScopeLock Lock(&MessageLogLock);
+		const FString& Str = Message->ToText().ToString();
+		if (LoggedMessages.Contains(Str))
+		{
+			return;
+		}
+
+		LoggedMessages.Add(Str);
+
+#if WITH_EDITOR
+		if (MessageLog)
 		{
-			MessageLog->Push( Message );
+			MessageLog->Push(Message);
 			bMessageProcessed = true;
 		}
-	}
 #endif
+	}
 
 	if ( !bMessageProcessed )
 	{
@@ -210,6 +219,8 @@ void FUsdLogManager::EnableMessageLog()
 {
 	FScopeLock Lock( &MessageLogLock );
 
+	LoggedMessages.Reset();
+
 	if ( ++MessageLogRefCount == 1 )
 	{
 		MessageLog.Emplace();
