USD: Fix crash caused by trying to get the Sequencer to display a deleted LevelSequence

diff --git a/Engine/Plugins/Importers/USDImporter/Source/USDStage/Private/USDStageActor.cpp b/Engine/Plugins/Importers/USDImporter/Source/USDStage/Private/USDStageActor.cpp
--- a/Engine/Plugins/Importers/USDImporter/Source/USDStage/Private/USDStageActor.cpp
+++ b/Engine/Plugins/Importers/USDImporter/Source/USDStage/Private/USDStageActor.cpp
@@ -1929,6 +1929,50 @@ void AUsdStageActor::PreEditUndo()
 
 void AUsdStageActor::HandleTransactionStateChanged( const FTransactionContext& InTransactionContext, const ETransactionStateEventType InTransactionState )
 {
+	// Hack for solving UE-127253
+	// When we Reload (or open a new stage), we call ReloadAnimations which will close the Sequencer (if opened), recreate our LevelSequence, and get the Sequencer
+	// to show that one instead. If we undo the Reload, that new LevelSequence will be deleted and the Sequencer will be left open trying to display it,
+	// which leads to crashes. Here we try detecting for that case and close/reopen the sequencer to show the correct one.
+#if WITH_EDITOR
+	if ( GIsEditor && GEditor && ( InTransactionState == ETransactionStateEventType::UndoRedoStarted || InTransactionState == ETransactionStateEventType::UndoRedoFinalized ) )
+	{
+		if ( UTransactor* Trans = GEditor->Trans )
+		{
+			static TSet<AUsdStageActor*> ActorsThatClosedTheSequencer;
+
+			int32 CurrentTransactionIndex = Trans->FindTransactionIndex( InTransactionContext.TransactionId );
+			const FTransaction* Transaction = Trans->GetTransaction( CurrentTransactionIndex );
+
+			if ( Transaction )
+			{
+				TArray<UObject*> TransactionObjects;
+				Transaction->GetTransactionObjects( TransactionObjects );
+
+				if ( TransactionObjects.Contains( this ) )
+				{
+					if ( InTransactionState == ETransactionStateEventType::UndoRedoStarted )
+					{
+						const bool bLevelSequenceEditorWasOpened = GEditor->GetEditorSubsystem<UAssetEditorSubsystem>()->CloseAllEditorsForAsset( LevelSequence ) > 0;
+						if ( bLevelSequenceEditorWasOpened )
+						{
+							ActorsThatClosedTheSequencer.Add( this );
+						}
+					}
+
+					if ( InTransactionState == ETransactionStateEventType::UndoRedoFinalized )
+					{
+						if ( ActorsThatClosedTheSequencer.Contains( this ) )
+						{
+							GEditor->GetEditorSubsystem<UAssetEditorSubsystem>()->OpenEditorForAsset( LevelSequence );
+							ActorsThatClosedTheSequencer.Remove( this );
+						}
+					}
+				}
+			}
+		}
+	}
+#endif // WITH_EDITOR
+
 	if ( InTransactionState == ETransactionStateEventType::TransactionFinalized ||
 		 InTransactionState == ETransactionStateEventType::UndoRedoFinalized ||
 		 InTransactionState == ETransactionStateEventType::TransactionCanceled )
