USD: Don't send NaN to GeometryCache objects when the time to evaluate is the default time code.

diff --git a/Engine/Plugins/Importers/USDImporter/Source/USDSchemas/Private/USDGeomMeshTranslator.cpp b/Engine/Plugins/Importers/USDImporter/Source/USDSchemas/Private/USDGeomMeshTranslator.cpp
--- a/Engine/Plugins/Importers/USDImporter/Source/USDSchemas/Private/USDGeomMeshTranslator.cpp
+++ b/Engine/Plugins/Importers/USDImporter/Source/USDSchemas/Private/USDGeomMeshTranslator.cpp
@@ -1413,12 +1413,22 @@ void FUsdGeomMeshTranslator::UpdateComponents( USceneComponent* SceneComponent )
 			GeometryCacheUsdComponent->GeometryCache = GeometryCache;
 		}
 
+		float TimeCode = Context->Time;
+		if ( FMath::IsNaN( TimeCode ) )
+		{
+			int32 StartFrame = FMath::FloorToInt( Context->Stage.GetStartTimeCode() );
+			int32 EndFrame = FMath::CeilToInt( Context->Stage.GetEndTimeCode() );
+			UsdGeomMeshTranslatorImpl::GetGeometryCacheDataTimeCodeRange( Context->Stage, PrimPath.GetString(), StartFrame, EndFrame );
+
+			TimeCode = static_cast< float >( StartFrame );
+		}
+
 		// This is the main call responsible for animating the geometry cache.
 		// It needs to happen after setting the geometry cache and before registering, because we must force the
 		// geometry cache to register itself at Context->Time so that it will synchronously load that frame right away.
 		// Otherwise the geometry cache will start at t=0 regardless of Context->Time
 		GeometryCacheUsdComponent->SetManualTick( true );
-		GeometryCacheUsdComponent->TickAtThisTime( Context->Time, true, false, true );
+		GeometryCacheUsdComponent->TickAtThisTime( TimeCode, true, false, true );
 
 		// Note how we should only register if our geometry cache changed: If we did this every time we would
 		// register too early during the process of duplicating into PIE, and that would prevent a future RegisterComponent
