USD: Fix not being able to write to cine camera component property tracks;

diff --git a/Engine/Plugins/Importers/USDImporter/Source/USDStage/Private/USDLevelSequenceHelper.cpp b/Engine/Plugins/Importers/USDImporter/Source/USDStage/Private/USDLevelSequenceHelper.cpp
--- a/Engine/Plugins/Importers/USDImporter/Source/USDStage/Private/USDLevelSequenceHelper.cpp
+++ b/Engine/Plugins/Importers/USDImporter/Source/USDStage/Private/USDLevelSequenceHelper.cpp
@@ -18,6 +18,7 @@
 #include "UsdWrappers/SdfLayer.h"
 #include "UsdWrappers/UsdAttribute.h"
 #include "UsdWrappers/UsdGeomXformable.h"
+#include "UsdWrappers/UsdPrim.h"
 #include "UsdWrappers/UsdStage.h"
 #include "UsdWrappers/UsdTyped.h"
 
@@ -147,6 +148,39 @@ namespace UsdLevelSequenceHelperImpl
 
 		return nullptr;
 	}
+
+	// Returns the UObject that is bound to the track. Will only consider possessables (and ignore spawnables)
+	// since we don't currently have any workflow where an opened USD stage would interact with UE spawnables.
+	UObject* LocateBoundObject( const UMovieSceneSequence& MovieSceneSequence, const FMovieScenePossessable& Possessable )
+	{
+		UMovieScene* MovieScene = MovieSceneSequence.GetMovieScene();
+		if ( !MovieScene )
+		{
+			return nullptr;
+		}
+
+		const FGuid& Guid = Possessable.GetGuid();
+		const FGuid& ParentGuid = Possessable.GetParent();
+
+		// If we have a parent guid, we must provide the object as a context because really the binding path
+		// will just contain the component name
+		UObject* ParentContext = nullptr;
+		if ( ParentGuid.IsValid() )
+		{
+			if ( FMovieScenePossessable* ParentPossessable = MovieScene->FindPossessable( ParentGuid ) )
+			{
+				ParentContext = LocateBoundObject( MovieSceneSequence, *ParentPossessable );
+			}
+		}
+
+		TArray<UObject*, TInlineAllocator<1>> Objects = MovieSceneSequence.LocateBoundObjects( Guid, ParentContext );
+		if ( Objects.Num() > 0 )
+		{
+			return Objects[ 0 ];
+		}
+
+		return nullptr;
+	}
 }
 
 class FUsdLevelSequenceHelperImpl : private FGCObject
@@ -2151,160 +2185,194 @@ void FUsdLevelSequenceHelperImpl::HandleTrackChange( const UMovieSceneTrack& Tra
 		return;
 	}
 
-	const UMovieScene* MovieScene = Sequence->GetMovieScene();
+	UMovieScene* MovieScene = Sequence->GetMovieScene();
 	if ( !MovieScene )
 	{
 		return;
 	}
 
 	FGuid PossessableGuid;
-	MovieScene->FindTrackBinding( Track, PossessableGuid );
+	bool bFound = MovieScene->FindTrackBinding( Track, PossessableGuid );
+	if ( !bFound )
+	{
+		return;
+	}
 
-	TArray< UObject*, TInlineAllocator< 1 > > BoundObjects = Cast< UMovieSceneSequence >( Sequence )->LocateBoundObjects( PossessableGuid, nullptr ); // TODO: Binding sur actor?
+	FMovieScenePossessable* Possessable = MovieScene->FindPossessable( PossessableGuid );
+	if ( !Possessable )
+	{
+		return;
+	}
 
-	for ( UObject* BoundObject : BoundObjects )
+	UObject* BoundObject = UsdLevelSequenceHelperImpl::LocateBoundObject( *Sequence, *Possessable );
+	if( !BoundObject )
 	{
-		USceneComponent* BoundSceneComponent = nullptr;
-		if ( AActor* BoundActor = Cast< AActor >( BoundObject ) )
-		{
-			BoundSceneComponent = BoundActor->GetRootComponent();
-		}
-		else
-		{
-			BoundSceneComponent = Cast< USceneComponent >( BoundObject );
-		}
-		if ( !BoundSceneComponent )
-		{
-			continue;
-		}
+		return;
+	}
 
-		// For CineCameraActors the CineCameraComponent is bound, but that is not the root component of the actor, or the component
-		// that is tracked by the prim twins
-		if ( UCineCameraComponent* CineCameraComponent = Cast<UCineCameraComponent>( BoundSceneComponent ) )
-		{
-			BoundSceneComponent = CineCameraComponent->GetAttachParent();
-		}
+	USceneComponent* BoundSceneComponent = nullptr;
+	if ( AActor* BoundActor = Cast< AActor >( BoundObject ) )
+	{
+		BoundSceneComponent = BoundActor->GetRootComponent();
+	}
+	else
+	{
+		BoundSceneComponent = Cast< USceneComponent >( BoundObject );
+	}
 
-		if ( UUsdPrimTwin* PrimTwin = StageActor->RootUsdTwin->Find( BoundSceneComponent ) )
-		{
-			FScopedBlockNoticeListening BlockNotices( StageActor.Get() );
-			UE::FUsdPrim UsdPrim = UsdStage.GetPrimAtPath( UE::FSdfPath( *PrimTwin->PrimPath ) );
+	if ( !BoundSceneComponent )
+	{
+		return;
+	}
 
-			SceneComponentsBindings.Emplace( PrimTwin ) = TPair< ULevelSequence*, FGuid >( Sequence, PossessableGuid ); // Make sure we track this binding
+	UUsdPrimTwin* PrimTwin = StageActor->RootUsdTwin->Find( BoundSceneComponent );
+
+	// If we exported/created this Camera prim ourselves, we'll have a decomposed parent Xform and a child Camera prim (to mirror
+	// the ACineCameraActor structure), and we should have created prim twins for both when opening this stage.
+	// If this USD layer is not authored by us, it may just be a standalone Camera prim: In this scenario the created PrimTwin
+	// will be pointing at the parent USceneComponent of the spawned ACineCameraActor, and we wouldn't find anything when searching
+	// for the camera component directly, so try again
+	if ( !PrimTwin && BoundSceneComponent->IsA<UCineCameraComponent>() )
+	{
+		if ( const UMovieScenePropertyTrack* PropertyTrack = Cast<const UMovieScenePropertyTrack>( &Track ) )
+		{
+			const FName& PropertyPath = PropertyTrack->GetPropertyPath();
 
-			if ( bIsMuteChange )
+			// In the scenario where we're trying to make non-decomposed Camera prims work, we only ever want to write out
+			// actual camera properties from the CameraComponent to the Camera prim. We won't write its USceneComponent
+			// properties, as we will use the ones from the ACineCameraActor's parent USceneComponent instead
+			if ( PropertyPath == UnrealIdentifiers::CurrentFocalLengthPropertyName ||
+				 PropertyPath == UnrealIdentifiers::ManualFocusDistancePropertyName ||
+				 PropertyPath == UnrealIdentifiers::ManualFocusDistancePropertyName ||
+				 PropertyPath == UnrealIdentifiers::CurrentAperturePropertyName ||
+				 PropertyPath == UnrealIdentifiers::SensorWidthPropertyName ||
+				 PropertyPath == UnrealIdentifiers::SensorHeightPropertyName
+			)
 			{
-				if ( const UMovieScenePropertyTrack* PropertyTrack = Cast<const UMovieScenePropertyTrack>( &Track ) )
-				{
-					const FName& PropertyPath = PropertyTrack->GetPropertyPath();
+				PrimTwin = StageActor->RootUsdTwin->Find( BoundSceneComponent->GetAttachParent() );
+			}
+		}
+	}
 
-					TArray<UE::FUsdAttribute> Attrs = UnrealToUsd::GetAttributesForProperty( UsdPrim, PropertyPath );
-					if ( Attrs.Num() > 0 )
-					{
-						// Only mute/unmute the first (i.e. main) attribute: If we mute the intensity track we don't want to also mute the
-						// rect width track if it has one
-						UE::FUsdAttribute& Attr = Attrs[0];
+	if ( PrimTwin )
+	{
+		FScopedBlockNoticeListening BlockNotices( StageActor.Get() );
+		UE::FUsdPrim UsdPrim = UsdStage.GetPrimAtPath( UE::FSdfPath( *PrimTwin->PrimPath ) );
 
-						bool bAllSectionsMuted = true;
-						for ( const UMovieSceneSection* Section : Track.GetAllSections() ) // There's no const version of "FindSection"
-						{
-							bAllSectionsMuted &= !Section->IsActive();
-						}
+		SceneComponentsBindings.Emplace( PrimTwin ) = TPair< ULevelSequence*, FGuid >( Sequence, PossessableGuid ); // Make sure we track this binding
 
-						if ( Track.IsEvalDisabled() || bAllSectionsMuted )
-						{
-							UsdUtils::MuteAttribute( Attr, UsdStage );
-						}
-						else
-						{
-							UsdUtils::UnmuteAttribute( Attr, UsdStage );
-						}
+		if ( bIsMuteChange )
+		{
+			if ( const UMovieScenePropertyTrack* PropertyTrack = Cast<const UMovieScenePropertyTrack>( &Track ) )
+			{
+				const FName& PropertyPath = PropertyTrack->GetPropertyPath();
 
-						// The attribute may have an effect on the stage, so animate it right away
-						StageActor->OnTimeChanged.Broadcast();
-					}
-				}
-				else if ( const UMovieSceneSkeletalAnimationTrack* SkeletalTrack = Cast<const UMovieSceneSkeletalAnimationTrack>( &Track ) )
+				TArray<UE::FUsdAttribute> Attrs = UnrealToUsd::GetAttributesForProperty( UsdPrim, PropertyPath );
+				if ( Attrs.Num() > 0 )
 				{
+					// Only mute/unmute the first (i.e. main) attribute: If we mute the intensity track we don't want to also mute the
+					// rect width track if it has one
+					UE::FUsdAttribute& Attr = Attrs[0];
+
 					bool bAllSectionsMuted = true;
-					for ( const UMovieSceneSection* Section : SkeletalTrack->GetAllSections() ) // There's no const version of "FindSection"
+					for ( const UMovieSceneSection* Section : Track.GetAllSections() ) // There's no const version of "FindSection"
 					{
 						bAllSectionsMuted &= !Section->IsActive();
 					}
 
-					if ( UE::FUsdPrim SkelAnimationPrim = UsdUtils::FindAnimationSource( UsdPrim ) )
+					if ( Track.IsEvalDisabled() || bAllSectionsMuted )
 					{
-						UE::FUsdAttribute TranslationsAttr = SkelAnimationPrim.GetAttribute( TEXT( "translations" ) );
-						UE::FUsdAttribute RotationsAttr = SkelAnimationPrim.GetAttribute( TEXT( "rotations" ) );
-						UE::FUsdAttribute ScalesAttr = SkelAnimationPrim.GetAttribute( TEXT( "scales" ) );
-						UE::FUsdAttribute BlendShapeWeightsAttr = SkelAnimationPrim.GetAttribute( TEXT( "blendShapeWeights" ) );
-
-						if ( Track.IsEvalDisabled() || bAllSectionsMuted )
-						{
-							UsdUtils::MuteAttribute( TranslationsAttr, UsdStage );
-							UsdUtils::MuteAttribute( RotationsAttr, UsdStage );
-							UsdUtils::MuteAttribute( ScalesAttr, UsdStage );
-							UsdUtils::MuteAttribute( BlendShapeWeightsAttr, UsdStage );
-						}
-						else
-						{
-							UsdUtils::UnmuteAttribute( TranslationsAttr, UsdStage );
-							UsdUtils::UnmuteAttribute( RotationsAttr, UsdStage );
-							UsdUtils::UnmuteAttribute( ScalesAttr, UsdStage );
-							UsdUtils::UnmuteAttribute( BlendShapeWeightsAttr, UsdStage );
-						}
-
-						// The attribute may have an effect on the stage, so animate it right away
-						StageActor->OnTimeChanged.Broadcast();
+						UsdUtils::MuteAttribute( Attr, UsdStage );
+					}
+					else
+					{
+						UsdUtils::UnmuteAttribute( Attr, UsdStage );
 					}
+
+					// The attribute may have an effect on the stage, so animate it right away
+					StageActor->OnTimeChanged.Broadcast();
 				}
 			}
-			else
+			else if ( const UMovieSceneSkeletalAnimationTrack* SkeletalTrack = Cast<const UMovieSceneSkeletalAnimationTrack>( &Track ) )
 			{
-				FMovieSceneSequenceTransform SequenceTransform;
-
-				if ( const FMovieSceneSequenceID* SequenceID = SequencesID.Find( Sequence ) )
+				bool bAllSectionsMuted = true;
+				for ( const UMovieSceneSection* Section : SkeletalTrack->GetAllSections() ) // There's no const version of "FindSection"
 				{
-					if ( FMovieSceneSubSequenceData* SubSequenceData = SequenceHierarchyCache.FindSubData( *SequenceID ) )
-					{
-						SequenceTransform = SubSequenceData->RootToSequenceTransform;
-					}
+					bAllSectionsMuted &= !Section->IsActive();
 				}
 
-				// Right now we don't write out changes to SkeletalAnimation tracks, and only property tracks... the UAnimSequence
-				// asset can't be modified all that much in UE anyway. Later on we may want to enable writing it out anyway though,
-				// and pick up on changes to the section offset or play rate and bake out the UAnimSequence again
-				if ( const UMovieScenePropertyTrack* PropertyTrack = Cast<const UMovieScenePropertyTrack>( &Track ) )
+				if ( UE::FUsdPrim SkelAnimationPrim = UsdUtils::FindAnimationSource( UsdPrim ) )
 				{
-					TSet<FName> PropertyPathsToRefresh;
-					UnrealToUsd::FPropertyTrackWriter Writer = UnrealToUsd::CreatePropertyTrackWriter( *BoundSceneComponent, *PropertyTrack, UsdPrim, PropertyPathsToRefresh );
+					UE::FUsdAttribute TranslationsAttr = SkelAnimationPrim.GetAttribute( TEXT( "translations" ) );
+					UE::FUsdAttribute RotationsAttr = SkelAnimationPrim.GetAttribute( TEXT( "rotations" ) );
+					UE::FUsdAttribute ScalesAttr = SkelAnimationPrim.GetAttribute( TEXT( "scales" ) );
+					UE::FUsdAttribute BlendShapeWeightsAttr = SkelAnimationPrim.GetAttribute( TEXT( "blendShapeWeights" ) );
 
-					if ( const UMovieSceneFloatTrack* FloatTrack = Cast<const UMovieSceneFloatTrack>( &Track ) )
-					{
-						UnrealToUsd::ConvertFloatTrack( *FloatTrack, SequenceTransform, Writer.FloatWriter, UsdPrim );
-					}
-					else if ( const UMovieSceneBoolTrack* BoolTrack = Cast<const UMovieSceneBoolTrack>( &Track ) )
-					{
-						UnrealToUsd::ConvertBoolTrack( *BoolTrack, SequenceTransform, Writer.BoolWriter, UsdPrim );
-					}
-					else if ( const UMovieSceneColorTrack* ColorTrack = Cast<const UMovieSceneColorTrack>( &Track ) )
+					if ( Track.IsEvalDisabled() || bAllSectionsMuted )
 					{
-						UnrealToUsd::ConvertColorTrack( *ColorTrack, SequenceTransform, Writer.ColorWriter, UsdPrim );
+						UsdUtils::MuteAttribute( TranslationsAttr, UsdStage );
+						UsdUtils::MuteAttribute( RotationsAttr, UsdStage );
+						UsdUtils::MuteAttribute( ScalesAttr, UsdStage );
+						UsdUtils::MuteAttribute( BlendShapeWeightsAttr, UsdStage );
 					}
-					else if ( const UMovieScene3DTransformTrack* TransformTrack = Cast<const UMovieScene3DTransformTrack>( &Track ) )
+					else
 					{
-						UnrealToUsd::Convert3DTransformTrack( *TransformTrack, SequenceTransform, Writer.TransformWriter, UsdPrim );
+						UsdUtils::UnmuteAttribute( TranslationsAttr, UsdStage );
+						UsdUtils::UnmuteAttribute( RotationsAttr, UsdStage );
+						UsdUtils::UnmuteAttribute( ScalesAttr, UsdStage );
+						UsdUtils::UnmuteAttribute( BlendShapeWeightsAttr, UsdStage );
 					}
 
-					// Refresh tracks that needed to be updated in USD (e.g. we wrote out a new keyframe to a RectLight's width -> that
-					// should become a new keyframe on our intensity track, because we use the RectLight's width for calculating intensity in UE).
-					if ( PropertyPathsToRefresh.Num() > 0 )
-					{
-						// For now only our light tracks can request a refresh like this, so we don't even need to check what the refresh
-						// is about: Just resync the light tracks
-						AddLightTracks( *PrimTwin, UsdPrim, PropertyPathsToRefresh );
-						RefreshSequencer();
-					}
+					// The attribute may have an effect on the stage, so animate it right away
+					StageActor->OnTimeChanged.Broadcast();
+				}
+			}
+		}
+		else
+		{
+			FMovieSceneSequenceTransform SequenceTransform;
+
+			if ( const FMovieSceneSequenceID* SequenceID = SequencesID.Find( Sequence ) )
+			{
+				if ( FMovieSceneSubSequenceData* SubSequenceData = SequenceHierarchyCache.FindSubData( *SequenceID ) )
+				{
+					SequenceTransform = SubSequenceData->RootToSequenceTransform;
+				}
+			}
+
+			// Right now we don't write out changes to SkeletalAnimation tracks, and only property tracks... the UAnimSequence
+			// asset can't be modified all that much in UE anyway. Later on we may want to enable writing it out anyway though,
+			// and pick up on changes to the section offset or play rate and bake out the UAnimSequence again
+			if ( const UMovieScenePropertyTrack* PropertyTrack = Cast<const UMovieScenePropertyTrack>( &Track ) )
+			{
+				TSet<FName> PropertyPathsToRefresh;
+				UnrealToUsd::FPropertyTrackWriter Writer = UnrealToUsd::CreatePropertyTrackWriter( *BoundSceneComponent, *PropertyTrack, UsdPrim, PropertyPathsToRefresh );
+
+				if ( const UMovieSceneFloatTrack* FloatTrack = Cast<const UMovieSceneFloatTrack>( &Track ) )
+				{
+					UnrealToUsd::ConvertFloatTrack( *FloatTrack, SequenceTransform, Writer.FloatWriter, UsdPrim );
+				}
+				else if ( const UMovieSceneBoolTrack* BoolTrack = Cast<const UMovieSceneBoolTrack>( &Track ) )
+				{
+					UnrealToUsd::ConvertBoolTrack( *BoolTrack, SequenceTransform, Writer.BoolWriter, UsdPrim );
+				}
+				else if ( const UMovieSceneColorTrack* ColorTrack = Cast<const UMovieSceneColorTrack>( &Track ) )
+				{
+					UnrealToUsd::ConvertColorTrack( *ColorTrack, SequenceTransform, Writer.ColorWriter, UsdPrim );
+				}
+				else if ( const UMovieScene3DTransformTrack* TransformTrack = Cast<const UMovieScene3DTransformTrack>( &Track ) )
+				{
+					UnrealToUsd::Convert3DTransformTrack( *TransformTrack, SequenceTransform, Writer.TransformWriter, UsdPrim );
+				}
+
+				// Refresh tracks that needed to be updated in USD (e.g. we wrote out a new keyframe to a RectLight's width -> that
+				// should become a new keyframe on our intensity track, because we use the RectLight's width for calculating intensity in UE).
+				if ( PropertyPathsToRefresh.Num() > 0 )
+				{
+					// For now only our light tracks can request a refresh like this, so we don't even need to check what the refresh
+					// is about: Just resync the light tracks
+					AddLightTracks( *PrimTwin, UsdPrim, PropertyPathsToRefresh );
+					RefreshSequencer();
 				}
 			}
 		}
diff --git a/Engine/Plugins/Importers/USDImporter/Source/USDUtilities/Private/USDPrimConversion.cpp b/Engine/Plugins/Importers/USDImporter/Source/USDUtilities/Private/USDPrimConversion.cpp
--- a/Engine/Plugins/Importers/USDImporter/Source/USDUtilities/Private/USDPrimConversion.cpp
+++ b/Engine/Plugins/Importers/USDImporter/Source/USDUtilities/Private/USDPrimConversion.cpp
@@ -2765,10 +2765,13 @@ bool UnrealToUsd::CreateSkeletalAnimationBaker( UE::FUsdPrim& SkelRoot, UE::FUsd
 
 UnrealToUsd::FPropertyTrackWriter UnrealToUsd::CreatePropertyTrackWriter( const USceneComponent& Component, const UMovieScenePropertyTrack& Track, UE::FUsdPrim& Prim, TSet<FName>& OutPropertyPathsToRefresh )
 {
-	const FName& PropertyPath = Track.GetPropertyPath();
-
 	UnrealToUsd::FPropertyTrackWriter Result;
 
+	if ( !Prim )
+	{
+		return Result;
+	}
+
 	FScopedUsdAllocs Allocs;
 
 	pxr::UsdPrim UsdPrim{ Prim };
@@ -2779,6 +2782,8 @@ UnrealToUsd::FPropertyTrackWriter UnrealToUsd::CreatePropertyTrackWriter( const
 
 	pxr::UsdAttribute Attr;
 
+	const FName& PropertyPath = Track.GetPropertyPath();
+
 	// SceneComponent
 	{
 		if ( PropertyPath == UnrealIdentifiers::TransformPropertyName )
