USD: Preserve prim and component selections through operations that trigger refreshes (like stage reloads, changing kinds to collapse, etc.);

Fix DeselectActorsAndComponents never actually deselecting actors;
Prevent some unnecessary refreshes and events when selecting prims and actors/components with selection sync enabled.

diff --git a/Engine/Plugins/Importers/USDImporter/Source/USDStage/Private/USDStageActor.cpp b/Engine/Plugins/Importers/USDImporter/Source/USDStage/Private/USDStageActor.cpp
--- a/Engine/Plugins/Importers/USDImporter/Source/USDStage/Private/USDStageActor.cpp
+++ b/Engine/Plugins/Importers/USDImporter/Source/USDStage/Private/USDStageActor.cpp
@@ -150,6 +150,12 @@ struct FUsdStageActorImpl
 				if ( USceneComponent* ReferencedComponent = PrimTwin.SceneComponent.Get() )
 				{
 					ComponentsToDeselect.Add( ReferencedComponent );
+
+					AActor* Owner = ReferencedComponent->GetOwner();
+					if ( Owner && Owner->GetRootComponent() == ReferencedComponent )
+					{
+						ActorsToDeselect.Add( Owner );
+					}
 				}
 			}, bRecursive );
 
diff --git a/Engine/Plugins/Importers/USDImporter/Source/USDStageEditor/Private/SUSDStage.cpp b/Engine/Plugins/Importers/USDImporter/Source/USDStageEditor/Private/SUSDStage.cpp
--- a/Engine/Plugins/Importers/USDImporter/Source/USDStageEditor/Private/SUSDStage.cpp
+++ b/Engine/Plugins/Importers/USDImporter/Source/USDStageEditor/Private/SUSDStage.cpp
@@ -65,9 +65,9 @@ namespace SUSDStageImpl
 		}
 
 		TSet<AActor*> ActorsToSelect;
-		for ( TSet<USceneComponent*>::TIterator It(ComponentsToSelect); It; ++It )
+		for ( TSet<USceneComponent*>::TIterator It( ComponentsToSelect ); It; ++It )
 		{
-			if ( AActor* Owner = (*It)->GetOwner() )
+			if ( AActor* Owner = ( *It )->GetOwner() )
 			{
 				// We always need the parent actor selected to select a component
 				ActorsToSelect.Add( Owner );
@@ -89,11 +89,19 @@ namespace SUSDStageImpl
 			return;
 		}
 
+		// Don't use GEditor->SelectNone() as that will affect *every type of selection in the editor*,
+		// including even some UI menus, brushes, etc.
+		if ( USelection* SelectedActors = GEditor->GetSelectedActors() )
+		{
+			SelectedActors->DeselectAll( AActor::StaticClass() );
+		}
+		if ( USelection* SelectedComponents = GEditor->GetSelectedComponents() )
+		{
+			SelectedComponents->DeselectAll( UActorComponent::StaticClass() );
+		}
+
 		const bool bSelected = true;
 		const bool bNotifySelectionChanged = true;
-		const bool bDeselectBSPSurfs = true;
-		GEditor->SelectNone( bNotifySelectionChanged, bDeselectBSPSurfs );
-
 		for ( AActor* Actor : ActorsToSelect )
 		{
 			GEditor->SelectActor( Actor, bSelected, bNotifySelectionChanged );
@@ -237,6 +245,13 @@ void SUsdStage::SetupStageActorDelegates()
 					{
 						this->UsdPrimInfoWidget->SetPrimPath( ViewModel.UsdStageActor->GetOrLoadUsdStage(), *PrimPath );
 					}
+
+					// If we resynced our selected prim or our ancestor and have selection sync enabled, try to refresh it so that we're
+					// still selecting the same actor/component that corresponds to the prim we have currently selected
+					if ( bResync && SelectedPrimPath.StartsWith( PrimPath ) )
+					{
+						OnPrimSelectionChanged( { SelectedPrimPath } );
+					}
 				});
 			}
 		);
@@ -741,6 +756,9 @@ void SUsdStage::FillPurposesToLoadSubMenu(FMenuBuilder& MenuBuilder)
 							FText::FromString(StageActor->GetActorLabel())
 						));
 
+						// c.f. comment in SUsdStage::FillCollapsingSubMenu
+						TGuardValue<bool> MaintainSelectionGuard( bUpdatingViewportSelection, true );
+
 						StageActor->Modify();
 						StageActor->PurposesToLoad = (int32)((EUsdPurpose)StageActor->PurposesToLoad ^ Purpose);
 
@@ -797,6 +815,9 @@ void SUsdStage::FillRenderContextSubMenu( FMenuBuilder& MenuBuilder )
 							FText::FromString(StageActor->GetActorLabel())
 						));
 
+						// c.f. comment in SUsdStage::FillCollapsingSubMenu
+						TGuardValue<bool> MaintainSelectionGuard( bUpdatingViewportSelection, true );
+
 						StageActor->Modify();
 						StageActor->RenderContext = RenderContext;
 
@@ -851,6 +872,14 @@ void SUsdStage::FillCollapsingSubMenu( FMenuBuilder& MenuBuilder )
 							FText::FromString( StageActor->GetActorLabel() )
 						) );
 
+						// When we change our kinds to collapse, we'll end up loading the stage again and refreshing our selection.
+						// We'll take care to try to re-select the same prims within SUsdStage::Refresh after the refresh is complete,
+						// but if selection sync is on we'll also later on try updating our prim selection to match our component selection.
+						// Not only is this "selection spam" is very visible on the USD Stage Editor, if our originally selected component
+						// was just collapsed away, we'd end up updating our prim selection to point to the parent prim instead
+						// (the collapsing root), which is not ideal.
+						TGuardValue<bool> MaintainSelectionGuard( bUpdatingViewportSelection, true );
+
 						int32 NewKindsToCollapse = ( int32 ) ( ( EUsdDefaultKind ) StageActor->KindsToCollapse ^ Kind );
 						StageActor->SetKindsToCollapse( NewKindsToCollapse );
 					}
@@ -957,6 +986,9 @@ void SUsdStage::FillInterpolationTypeSubMenu(FMenuBuilder& MenuBuilder)
 						FText::FromString(StageActor->GetActorLabel())
 					));
 
+					// c.f. comment in SUsdStage::FillCollapsingSubMenu
+					TGuardValue<bool> MaintainSelectionGuard( bUpdatingViewportSelection, true );
+
 					StageActor->SetInterpolationType( EUsdInterpolationType::Linear );
 				}
 			}),
@@ -991,6 +1023,9 @@ void SUsdStage::FillInterpolationTypeSubMenu(FMenuBuilder& MenuBuilder)
 						FText::FromString(StageActor->GetActorLabel())
 					));
 
+					// c.f. comment in SUsdStage::FillCollapsingSubMenu
+					TGuardValue<bool> MaintainSelectionGuard( bUpdatingViewportSelection, true );
+
 					StageActor->SetInterpolationType( EUsdInterpolationType::Held );
 				}
 			}),
@@ -1245,6 +1280,11 @@ void SUsdStage::Refresh()
 	if (UsdStageTreeView)
 	{
 		UsdStageTreeView->Refresh( StageActor );
+
+		// Refresh will generate brand new RootItems, so let's immediately select the new item
+		// that corresponds to the prim we're supposed to be selecting
+		UsdStageTreeView->SelectPrims( { SelectedPrimPath } );
+
 		UsdStageTreeView->RequestTreeRefresh();
 	}
 }
@@ -1343,6 +1383,9 @@ void SUsdStage::OnNaniteTriangleThresholdValueCommitted( int32 InValue, ETextCom
 		FText::FromString( StageActor->GetActorLabel() )
 	) );
 
+	// c.f. comment in SUsdStage::FillCollapsingSubMenu
+	TGuardValue<bool> MaintainSelectionGuard( bUpdatingViewportSelection, true );
+
 	StageActor->SetNaniteTriangleThreshold( InValue );
 	CurrentNaniteThreshold = InValue;
 }
diff --git a/Engine/Plugins/Importers/USDImporter/Source/USDStageEditor/Private/SUSDStageTreeView.cpp b/Engine/Plugins/Importers/USDImporter/Source/USDStageEditor/Private/SUSDStageTreeView.cpp
--- a/Engine/Plugins/Importers/USDImporter/Source/USDStageEditor/Private/SUSDStageTreeView.cpp
+++ b/Engine/Plugins/Importers/USDImporter/Source/USDStageEditor/Private/SUSDStageTreeView.cpp
@@ -479,8 +479,6 @@ FUsdPrimViewModelPtr SUsdStageTreeView::GetItemFromPrimPath( const FString& Prim
 
 void SUsdStageTreeView::SelectPrims( const TArray<FString>& PrimPaths )
 {
-	ClearSelection();
-
 	TArray< FUsdPrimViewModelRef > ItemsToSelect;
 	ItemsToSelect.Reserve( PrimPaths.Num() );
 
@@ -494,10 +492,19 @@ void SUsdStageTreeView::SelectPrims( const TArray<FString>& PrimPaths )
 
 	if ( ItemsToSelect.Num() > 0 )
 	{
+		// Clear selection without emitting events, as we'll emit new events with SetItemSelection
+		// anyway. This prevents a UI blink as OnPrimSelectionChanged would otherwise fire for
+		// ClearSelection() and then again right away for SetItemSelection()
+		Private_ClearSelection();
+
 		const bool bSelected = true;
 		SetItemSelection( ItemsToSelect, bSelected );
 		ScrollItemIntoView( ItemsToSelect.Last() );
 	}
+	else
+	{
+		ClearSelection();
+	}
 }
 
 TArray<FString> SUsdStageTreeView::GetSelectedPrims()
