USD: Make actor picker menu a bit wider as searching on it also display the actor types.

diff --git a/Engine/Plugins/Importers/USDImporter/Source/USDStageEditor/Private/SUSDStage.cpp b/Engine/Plugins/Importers/USDImporter/Source/USDStageEditor/Private/SUSDStage.cpp
--- a/Engine/Plugins/Importers/USDImporter/Source/USDStageEditor/Private/SUSDStage.cpp
+++ b/Engine/Plugins/Importers/USDImporter/Source/USDStageEditor/Private/SUSDStage.cpp
@@ -446,7 +446,7 @@ TSharedRef< SWidget > SUsdStage::MakeActorPickerMenuContent()
 
 		return SNew(SBox)
 			.Padding( FMargin( 1 ) )    // Add a small margin or else we'll get dark gray on dark gray which can look a bit confusing
-			.MinDesiredWidth( 300.0f )  // Force a min width or else the tree view item text will run up right to the very edge pixel of the menu
+			.MinDesiredWidth( 400.0f )  // Force a min width or else the tree view item text will run up right to the very edge pixel of the menu
 			.HAlign( HAlign_Fill )
 			[
 				ActorPickerMenu.ToSharedRef()
