USD: Fix not being able to write transform timeSamples to USD in sublayers if a stronger layer has a transform opinion.

diff --git a/Engine/Plugins/Importers/USDImporter/Source/USDUtilities/Private/USDPrimConversion.cpp b/Engine/Plugins/Importers/USDImporter/Source/USDUtilities/Private/USDPrimConversion.cpp
--- a/Engine/Plugins/Importers/USDImporter/Source/USDUtilities/Private/USDPrimConversion.cpp
+++ b/Engine/Plugins/Importers/USDImporter/Source/USDUtilities/Private/USDPrimConversion.cpp
@@ -176,6 +176,58 @@ namespace UE
 					BindingAPI.Bind( ChildMaterial );
 				}
 			}
+
+			// On the current edit target, will set the Xformable's op order to a single "xformOp:transform",
+			// create the corresponding attribute, and return the op
+			pxr::UsdGeomXformOp ForceMatrixXform( pxr::UsdGeomXformable& Xformable )
+			{
+				FScopedUsdAllocs Allocs;
+
+				// Note: We don't use Xformable.MakeMatrixXform() here because while it can clear the
+				// xform op order on the current edit target just fine, it will later try to AddTransformOp(),
+				// which calls AddXformOp. Internally, it will read the *composed* prim and if it finds that it already
+				// has an op of that type it will early out and not author anything. This means that if our stage
+				// has a strong opinion for an e.g. "xformOp:transform" already on the layer stack, it's not possible
+				// to author that same op on a weaker layer. We want to do this here, to ensure this prim's transform
+				// works as expected even if this weaker layer is used standalone, so we must do the analogous ourselves
+
+				// References: private constructor for UsdGeomXformOp that can receive a UsdPrim and UsdGeomXformable::AddXformOp
+
+				// Clear the existing xform op order for this prim on this layer
+				Xformable.ClearXformOpOrder();
+
+				// Find details about the transform attribute related to the default transform type xform op
+				pxr::TfToken TransformAttrName = pxr::UsdGeomXformOp::GetOpName( pxr::UsdGeomXformOp::TypeTransform );
+				const pxr::SdfValueTypeName& TransformAttrTypeName = pxr::UsdGeomXformOp::GetValueTypeName( pxr::UsdGeomXformOp::TypeTransform, pxr::UsdGeomXformOp::PrecisionDouble );
+				if ( TransformAttrName.IsEmpty() || !TransformAttrTypeName )
+				{
+					return {};
+				}
+
+				// Create the transform attribute that would match the default transform type xform op
+				const bool bCustom = false;
+				pxr::UsdPrim UsdPrim = Xformable.GetPrim();
+				pxr::UsdAttribute TransformAttr = UsdPrim.CreateAttribute( TransformAttrName, TransformAttrTypeName, bCustom );
+				if ( !TransformAttr )
+				{
+					return {};
+				}
+
+				// Now that the attribute is created, use it to create the corresponding pxr::UsdGeomXformOp
+				const bool bIsInverseOp = false;
+				pxr::UsdGeomXformOp NewOp{ TransformAttr, bIsInverseOp };
+				if ( !NewOp )
+				{
+					return {};
+				}
+
+				// Store the Op name on an array that will be our new op order value
+				pxr::VtTokenArray NewOps;
+				NewOps.push_back( NewOp.GetOpName() );
+				Xformable.CreateXformOpOrderAttr().Set( NewOps );
+
+				return NewOp;
+			}
 		}
 	}
 }
@@ -2244,27 +2296,9 @@ bool UnrealToUsd::ConvertXformable( const FTransform& RelativeTransform, pxr::Us
 	bool bResetXFormStack = false;
 	XForm.GetLocalTransformation( &UsdMatrix, &bResetXFormStack, UsdTimeCode );
 
-	bool bFoundTransformOp = false;
-	std::vector< pxr::UsdGeomXformOp > XFormOps = XForm.GetOrderedXformOps( &bResetXFormStack );
-	for ( const pxr::UsdGeomXformOp& XFormOp : XFormOps )
-	{
-		// Found transform op, trying to set its value
-		if ( XFormOp.GetOpType() == pxr::UsdGeomXformOp::TypeTransform )
-		{
-			bFoundTransformOp = true;
-			XFormOp.Set( UsdTransform, UsdTimeCode );
-			break;
-		}
-	}
-
-	// If transformOp is not found, make a new one
-	if ( !bFoundTransformOp )
+	if ( pxr::UsdGeomXformOp MatrixXform = UE::USDPrimConversionImpl::Private::ForceMatrixXform( XForm ) )
 	{
-		pxr::UsdGeomXformOp MatrixXform = XForm.MakeMatrixXform();
-		if ( MatrixXform )
-		{
-			MatrixXform.Set( UsdTransform, UsdTimeCode );
-		}
+		MatrixXform.Set( UsdTransform, UsdTimeCode );
 	}
 
 	return true;
@@ -2394,7 +2428,7 @@ bool UnrealToUsd::CreateComponentPropertyBaker( UE::FUsdPrim& Prim, const UScene
 			Xformable.CreateXformOpOrderAttr();
 
 			// Clear existing transform data and leave just one Transform op there
-			pxr::UsdGeomXformOp TransformOp = Xformable.MakeMatrixXform();
+			pxr::UsdGeomXformOp TransformOp = UE::USDPrimConversionImpl::Private::ForceMatrixXform( Xformable );
 			if ( !TransformOp )
 			{
 				return false;
@@ -2792,8 +2826,7 @@ UnrealToUsd::FPropertyTrackWriter UnrealToUsd::CreatePropertyTrackWriter( const
 			{
 				Xformable.CreateXformOpOrderAttr();
 
-				// Clear existing transform data and leave just one Transform op there
-				if ( pxr::UsdGeomXformOp TransformOp = Xformable.MakeMatrixXform() )
+				if ( pxr::UsdGeomXformOp TransformOp = UE::USDPrimConversionImpl::Private::ForceMatrixXform( Xformable ) )
 				{
 					Attr = TransformOp.GetAttr();
 
@@ -3375,7 +3408,7 @@ bool UnrealToUsd::ConvertXformable( const UMovieScene3DTransformTrack& MovieScen
 
 	if ( bIsDataOutOfSync )
 	{
-		if ( pxr::UsdGeomXformOp TransformOp = Xformable.MakeMatrixXform() )
+		if ( pxr::UsdGeomXformOp TransformOp = UE::USDPrimConversionImpl::Private::ForceMatrixXform( Xformable ) )
 		{
 			TransformOp.GetAttr().Clear(); // Clear existing transform data
 		}
