USD: Don't overwrite the stage actor's Time value with StartTimeCode when first loading from disk.

diff --git a/Engine/Plugins/Importers/USDImporter/Source/USDStage/Private/USDStageActor.cpp b/Engine/Plugins/Importers/USDImporter/Source/USDStage/Private/USDStageActor.cpp
--- a/Engine/Plugins/Importers/USDImporter/Source/USDStage/Private/USDStageActor.cpp
+++ b/Engine/Plugins/Importers/USDImporter/Source/USDStage/Private/USDStageActor.cpp
@@ -1646,7 +1646,8 @@ void AUsdStageActor::LoadUsdStage()
 
 	TranslationContext->CompleteTasks();
 
-	if ( UsdStage.GetRootLayer() )
+	// Keep our old Time value if we're loading the stage during initialization, so that we can save/load Time values
+	if ( UsdStage.GetRootLayer() && IsActorInitialized() )
 	{
 		SetTime( UsdStage.GetRootLayer().GetStartTimeCode() );
 
