USD: Fix not being able to see if the currently selected layer is the edit target or not.

diff --git a/Engine/Plugins/Importers/USDImporter/Source/USDStageEditor/Private/SUSDLayersTreeView.cpp b/Engine/Plugins/Importers/USDImporter/Source/USDStageEditor/Private/SUSDLayersTreeView.cpp
--- a/Engine/Plugins/Importers/USDImporter/Source/USDStageEditor/Private/SUSDLayersTreeView.cpp
+++ b/Engine/Plugins/Importers/USDImporter/Source/USDStageEditor/Private/SUSDLayersTreeView.cpp
@@ -177,9 +177,9 @@ class FUsdLayerEditColumn : public FUsdTreeViewColumn, public TSharedFromThis< F
 public:
 	const FSlateBrush* GetCheckedImage( const FUsdLayerViewModelRef InTreeItem ) const
 	{
-		return InTreeItem->LayerModel->bIsEditTarget ?
-			&FEditorStyle::Get().GetWidgetStyle< FCheckBoxStyle >( "Checkbox" ).CheckedImage :
-			nullptr;
+		return InTreeItem->LayerModel->bIsEditTarget
+			? FUsdStageEditorStyle::Get()->GetBrush( "UsdStageEditor.CheckBoxImage" )
+			: nullptr;
 	}
 
 	virtual TSharedRef< SWidget > GenerateWidget( const TSharedPtr< IUsdTreeViewItem > InTreeItem, const TSharedPtr< ITableRow > TableRow ) override
@@ -188,8 +188,7 @@ public:
 
 		TSharedRef< SWidget > Item =
 			SNew(SImage)
-				.Image( this, &FUsdLayerEditColumn::GetCheckedImage, TreeItem )
-				.ColorAndOpacity( FSlateColor::UseForeground() );
+				.Image( this, &FUsdLayerEditColumn::GetCheckedImage, TreeItem );
 
 		float ItemSize = FUsdStageEditorStyle::Get()->GetFloat( "UsdStageEditor.ListItemHeight" );
 
diff --git a/Engine/Plugins/Importers/USDImporter/Source/USDStageEditor/Private/SUSDStageEditorStyle.cpp b/Engine/Plugins/Importers/USDImporter/Source/USDStageEditor/Private/SUSDStageEditorStyle.cpp
--- a/Engine/Plugins/Importers/USDImporter/Source/USDStageEditor/Private/SUSDStageEditorStyle.cpp
+++ b/Engine/Plugins/Importers/USDImporter/Source/USDStageEditor/Private/SUSDStageEditorStyle.cpp
@@ -36,6 +36,7 @@ void FUsdStageEditorStyle::Initialize()
 {
 	// Const icon sizes
 	const FVector2D Icon8x8(8.0f, 8.0f);
+	const FVector2D Icon14x14(14.0f, 14.0f);
 	const FVector2D Icon16x16(16.0f, 16.0f);
 	const FVector2D Icon20x20(20.0f, 20.0f);
 	const FVector2D Icon40x40(40.0f, 40.0f);
@@ -60,6 +61,12 @@ void FUsdStageEditorStyle::Initialize()
 		/** Highlighted color of Prim labels that have composition arcs */
 		StyleSet->Set( "UsdStageEditor.HighlightPrimCompositionArcColor", FLinearColor( 0.2f, 0.08f, 0.0f, 1.0f ) );
 
+		/**
+		 * Checkmark that is always white. We can't use the editor style one in some places (which by default is styled blue)
+		 * because the highlight color is also the same blue, so the checkbox would be invisible
+		 */
+		StyleSet->Set( "UsdStageEditor.CheckBoxImage", new IMAGE_BRUSH( "Common/Check", Icon14x14, FLinearColor::White ) );
+
 		/**
 		 * Button without a background, border or foreground color that doesn't move when pressed.
 		 * No foreground color is important to let the automatic slate foreground color propagate down.
