USD: Fix skeletal mesh components not animating when going into PIE.

diff --git a/Engine/Plugins/Importers/USDImporter/Source/USDStage/Private/USDStageActor.cpp b/Engine/Plugins/Importers/USDImporter/Source/USDStage/Private/USDStageActor.cpp
--- a/Engine/Plugins/Importers/USDImporter/Source/USDStage/Private/USDStageActor.cpp
+++ b/Engine/Plugins/Importers/USDImporter/Source/USDStage/Private/USDStageActor.cpp
@@ -1502,6 +1502,12 @@ void AUsdStageActor::OnPostPIEStarted( bool bIsSimulating )
 			Prop->SetPropertyFlags( CPF_Transient );
 		}
 	}
+
+	// Setup for the very first frame when we duplicate into PIE, or else we will display skeletal mesh components on their
+	// StartTimeCode state. We have to do this here (after duplicating) as we need the calls to FUsdSkelRootTranslator::UpdateComponents
+	// to actually animate the components, and they will only be able to do anything after they have been registered (which
+	// needs to be done by the engine when going into PIE)
+	AnimatePrims();
 }
 
 void AUsdStageActor::OnObjectsReplaced( const TMap<UObject*, UObject*>& ObjectReplacementMap )
@@ -1994,11 +2000,9 @@ void AUsdStageActor::PostDuplicate( bool bDuplicateForPIE )
 {
 	Super::PostDuplicate( bDuplicateForPIE );
 
-	// Setup for the very first frame when we duplicate into PIE, or else we will just show a T-pose
 	if ( bDuplicateForPIE )
 	{
 		OpenUsdStage();
-		AnimatePrims();
 	}
 }
 
