USD: Update Sequencer DisplayRate if the stage's framesPerSecond value is updated;

Fix stage actor responding to USD events emitted by changes to Sequencer objects.

diff --git a/Engine/Plugins/Importers/USDImporter/Source/USDStage/Private/USDLevelSequenceHelper.cpp b/Engine/Plugins/Importers/USDImporter/Source/USDStage/Private/USDLevelSequenceHelper.cpp
--- a/Engine/Plugins/Importers/USDImporter/Source/USDStage/Private/USDLevelSequenceHelper.cpp
+++ b/Engine/Plugins/Importers/USDImporter/Source/USDStage/Private/USDLevelSequenceHelper.cpp
@@ -13,6 +13,7 @@
 #include "USDPrimConversion.h"
 #include "USDStageActor.h"
 #include "USDTypesConversion.h"
+#include "USDValueConversion.h"
 
 #include "UsdWrappers/SdfChangeBlock.h"
 #include "UsdWrappers/SdfLayer.h"
@@ -335,6 +336,7 @@ public:
 
 private:
 	void OnObjectTransacted(UObject* Object, const class FTransactionObjectEvent& Event);
+	void OnUsdObjectsChanged( const UsdUtils::FObjectChangesByPath& InfoChanges, const UsdUtils::FObjectChangesByPath& ResyncChanges );
 	void HandleTransactionStateChanged( const FTransactionContext& InTransactionContext, const ETransactionStateEventType InTransactionState );
 	void HandleMovieSceneChange(UMovieScene& MovieScene);
 	void HandleSubSectionChange(UMovieSceneSubSection& Section);
@@ -343,6 +345,7 @@ private:
 
 	FDelegateHandle OnObjectTransactedHandle;
 	FDelegateHandle OnStageEditTargetChangedHandle;
+	FDelegateHandle OnUsdObjectsChangedHandle;
 
 // Readonly handling
 private:
@@ -398,6 +401,7 @@ FUsdLevelSequenceHelperImpl::~FUsdLevelSequenceHelperImpl()
 	if ( StageActor.IsValid() )
 	{
 		StageActor->GetUsdListener().GetOnStageEditTargetChanged().Remove(OnStageEditTargetChangedHandle);
+		StageActor->GetUsdListener().GetOnObjectsChanged().Remove( OnUsdObjectsChangedHandle );
 		OnStageEditTargetChangedHandle.Reset();
 	}
 
@@ -569,6 +573,8 @@ void FUsdLevelSequenceHelperImpl::BindToUsdStageActor( AUsdStageActor* InStageAc
 			UpdateMovieSceneReadonlyFlags();
 		});
 
+	OnUsdObjectsChangedHandle = StageActor->GetUsdListener().GetOnObjectsChanged().AddRaw( this, &FUsdLevelSequenceHelperImpl::OnUsdObjectsChanged );
+
 	// Bind stage actor
 	StageActorBinding = MainLevelSequence->GetMovieScene()->AddPossessable(
 #if WITH_EDITOR
@@ -603,6 +609,7 @@ void FUsdLevelSequenceHelperImpl::UnbindFromUsdStageActor()
 	if ( StageActor.IsValid() )
 	{
 		StageActor->GetUsdListener().GetOnStageEditTargetChanged().Remove( OnStageEditTargetChangedHandle );
+		StageActor->GetUsdListener().GetOnObjectsChanged().Remove( OnUsdObjectsChangedHandle );
 		StageActor.Reset();
 	}
 
@@ -1953,6 +1960,49 @@ void FUsdLevelSequenceHelperImpl::OnObjectTransacted(UObject* Object, const clas
 	}
 }
 
+void FUsdLevelSequenceHelperImpl::OnUsdObjectsChanged( const UsdUtils::FObjectChangesByPath& InfoChanges, const UsdUtils::FObjectChangesByPath& ResyncChanges )
+{
+	AUsdStageActor* StageActorPtr = StageActor.Get();
+	if ( !StageActorPtr || !StageActorPtr->IsListeningToUsdNotices() )
+	{
+		return;
+	}
+
+	FScopedBlockMonitoringChangesForTransaction BlockMonitoring{ *this };
+
+	for ( const TPair<FString, TArray<UsdUtils::FObjectChangeNotice>>& InfoChange : InfoChanges )
+	{
+		const FString& PrimPath = InfoChange.Key;
+		if ( PrimPath == TEXT( "/" ) )
+		{
+			for ( const UsdUtils::FObjectChangeNotice& ObjectChange : InfoChange.Value )
+			{
+				for ( const UsdUtils::FAttributeChange& AttributeChange : ObjectChange.AttributeChanges )
+				{
+					if ( AttributeChange.PropertyName == TEXT( "framesPerSecond" ) && MainLevelSequence )
+					{
+						UsdUtils::FConvertedVtValue ConvertedValue;
+						if ( UsdToUnreal::ConvertValue( AttributeChange.NewValue, ConvertedValue ) )
+						{
+							if ( ConvertedValue.Entries.Num() == 1 && ConvertedValue.Entries[ 0 ].Num() == 1 )
+							{
+								if ( double* NewValue = ConvertedValue.Entries[ 0 ][ 0 ].TryGet<double>() )
+								{
+									if ( UMovieScene* MovieScene = MainLevelSequence->GetMovieScene() )
+									{
+										MovieScene->Modify();
+										MovieScene->SetDisplayRate( FFrameRate( *NewValue, 1 ) );
+									}
+								}
+							}
+						}
+					}
+				}
+			}
+		}
+	}
+}
+
 void FUsdLevelSequenceHelperImpl::HandleTransactionStateChanged( const FTransactionContext& InTransactionContext, const ETransactionStateEventType InTransactionState )
 {
 	if ( InTransactionState == ETransactionStateEventType::TransactionFinalized && BlockedTransactionGuids.Contains( InTransactionContext.TransactionId ) )
@@ -2020,8 +2070,8 @@ void FUsdLevelSequenceHelperImpl::HandleMovieSceneChange( UMovieScene& MovieScen
 	const FFrameTime StartTime = FFrameRate::TransformTime(UE::MovieScene::DiscreteInclusiveLower( PlaybackRange ).Value, MovieScene.GetTickResolution(), LayerTimeCodesPerSecond );
 	const FFrameTime EndTime = FFrameRate::TransformTime(UE::MovieScene::DiscreteExclusiveUpper( PlaybackRange ).Value, MovieScene.GetTickResolution(), LayerTimeCodesPerSecond );
 
-	UE::FSdfChangeBlock ChangeBlock;
 	FScopedBlockNoticeListening BlockNotices( StageActor.Get() );
+	UE::FSdfChangeBlock ChangeBlock;
 	if ( !FMath::IsNearlyEqual( DisplayRate.AsDecimal(), GetFramesPerSecond() ) )
 	{
 		UsdStage.SetFramesPerSecond( DisplayRate.AsDecimal() );
@@ -2660,13 +2710,18 @@ TArray< ULevelSequence* > FUsdLevelSequenceHelper::GetSubSequences() const
 }
 
 FScopedBlockMonitoringChangesForTransaction::FScopedBlockMonitoringChangesForTransaction( FUsdLevelSequenceHelper& InHelper )
-	: Helper( InHelper )
+	: FScopedBlockMonitoringChangesForTransaction( *InHelper.UsdSequencerImpl.Get() )
+{
+}
+
+FScopedBlockMonitoringChangesForTransaction::FScopedBlockMonitoringChangesForTransaction( FUsdLevelSequenceHelperImpl& InHelperImpl )
+	: HelperImpl( InHelperImpl )
 {
 	// If we're transacting we can just call this and the helper will unblock itself once the transaction is finished, because
 	// we need to make sure the unblocking happens after any call to OnObjectTransacted.
 	if ( GUndo )
 	{
-		Helper.BlockMonitoringChangesForThisTransaction();
+		HelperImpl.BlockMonitoringChangesForThisTransaction();
 	}
 	// If we're not in a transaction we still need to block this (can also happen e.g. if a Python change triggers a stage notice),
 	// but since we don't have to worry about the OnObjectTransacted calls we can just use this RAII object here to wrap over
@@ -2674,7 +2729,7 @@ FScopedBlockMonitoringChangesForTransaction::FScopedBlockMonitoringChangesForTra
 	else
 	{
 		bStoppedMonitoringChanges = true;
-		Helper.StopMonitoringChanges();
+		HelperImpl.StopMonitoringChanges();
 	}
 }
 
@@ -2682,6 +2737,6 @@ FScopedBlockMonitoringChangesForTransaction::~FScopedBlockMonitoringChangesForTr
 {
 	if ( bStoppedMonitoringChanges )
 	{
-		Helper.StartMonitoringChanges();
+		HelperImpl.StartMonitoringChanges();
 	}
 }
\ No newline at end of file
diff --git a/Engine/Plugins/Importers/USDImporter/Source/USDStage/Public/USDLevelSequenceHelper.h b/Engine/Plugins/Importers/USDImporter/Source/USDStage/Public/USDLevelSequenceHelper.h
--- a/Engine/Plugins/Importers/USDImporter/Source/USDStage/Public/USDLevelSequenceHelper.h
+++ b/Engine/Plugins/Importers/USDImporter/Source/USDStage/Public/USDLevelSequenceHelper.h
@@ -7,6 +7,7 @@
 #include "UsdWrappers/ForwardDeclarations.h"
 
 class AUsdStageActor;
+class FScopedBlockMonitoringChangesForTransaction;
 class FUsdLevelSequenceHelperImpl;
 class ULevelSequence;
 class UUsdAssetCache;
@@ -68,6 +69,7 @@ public:
 	TArray< ULevelSequence* > GetSubSequences() const;
 
 private:
+	friend class FScopedBlockMonitoringChangesForTransaction;
 	TUniquePtr<FUsdLevelSequenceHelperImpl> UsdSequencerImpl;
 };
 
@@ -75,6 +77,7 @@ class USDSTAGE_API FScopedBlockMonitoringChangesForTransaction final
 {
 public:
 	explicit FScopedBlockMonitoringChangesForTransaction( FUsdLevelSequenceHelper& InHelper );
+	explicit FScopedBlockMonitoringChangesForTransaction( FUsdLevelSequenceHelperImpl& InHelperImpl );
 	~FScopedBlockMonitoringChangesForTransaction();
 
 	FScopedBlockMonitoringChangesForTransaction() = delete;
@@ -84,6 +87,6 @@ public:
 	FScopedBlockMonitoringChangesForTransaction& operator=( FScopedBlockMonitoringChangesForTransaction&& ) = delete;
 
 private:
-	FUsdLevelSequenceHelper& Helper;
+	FUsdLevelSequenceHelperImpl& HelperImpl;
 	bool bStoppedMonitoringChanges = false;
 };
