USD: Tweak tooltips to try to explain why the universal render context shows up as "None" in UE dialogs.

diff --git a/Engine/Plugins/Importers/USDImporter/Source/USDStage/Private/USDStageActorCustomization.cpp b/Engine/Plugins/Importers/USDImporter/Source/USDStage/Private/USDStageActorCustomization.cpp
new file mode 100644
--- /dev/null
+++ b/Engine/Plugins/Importers/USDImporter/Source/USDStage/Private/USDStageActorCustomization.cpp
@@ -0,0 +1,169 @@
+// Copyright Epic Games, Inc. All Rights Reserved.
+
+#if WITH_EDITOR
+
+#include "USDStageActorCustomization.h"
+
+#include "USDSchemasModule.h"
+#include "USDSchemaTranslator.h"
+#include "USDStageActor.h"
+
+#include "DetailCategoryBuilder.h"
+#include "DetailLayoutBuilder.h"
+#include "DetailWidgetRow.h"
+#include "Modules/ModuleManager.h"
+#include "ScopedTransaction.h"
+#include "UObject/Object.h"
+#include "Widgets/DeclarativeSyntaxSupport.h"
+
+#define LOCTEXT_NAMESPACE "UsdStageActorCustomization"
+
+FUsdStageActorCustomization::FUsdStageActorCustomization()
+{
+}
+
+TSharedRef<IDetailCustomization> FUsdStageActorCustomization::MakeInstance()
+{
+	return MakeShared<FUsdStageActorCustomization>();
+}
+
+void FUsdStageActorCustomization::CustomizeDetails( IDetailLayoutBuilder& DetailLayoutBuilder )
+{
+	TArray<TWeakObjectPtr<UObject>> SelectedObjects = DetailLayoutBuilder.GetSelectedObjects();
+	if ( SelectedObjects.Num() != 1 )
+	{
+		return;
+	}
+
+	TWeakObjectPtr<UObject> SelectedObject = SelectedObjects[ 0 ];
+	if ( !SelectedObject.IsValid() )
+	{
+		return;
+	}
+
+	CurrentActor = Cast<AUsdStageActor>( SelectedObject.Get() );
+	if ( !CurrentActor )
+	{
+		return;
+	}
+
+	CurrentActor->OnStageChanged.AddSP( this, &FUsdStageActorCustomization::ForceRefreshDetails );
+
+	IUsdSchemasModule& UsdSchemasModule = FModuleManager::Get().LoadModuleChecked< IUsdSchemasModule >( TEXT( "USDSchemas" ) );
+
+	ComboBoxItems.Reset();
+	TSharedPtr<FString> InitiallySelectedContext;
+	for ( const FName& Context : UsdSchemasModule.GetRenderContextRegistry().GetRenderContexts() )
+	{
+		TSharedPtr<FString> ContextStr;
+		if ( Context == NAME_None )
+		{
+			ContextStr = MakeShared<FString>( TEXT( "universal" ) );
+		}
+		else
+		{
+			ContextStr = MakeShared<FString>( Context.ToString() );
+		}
+
+		if ( Context == CurrentActor->RenderContext )
+		{
+			InitiallySelectedContext = ContextStr;
+		}
+
+		ComboBoxItems.Add( ContextStr );
+	}
+
+	IDetailCategoryBuilder& CatBuilder = DetailLayoutBuilder.EditCategory( TEXT( "USD" ) );
+
+	if ( TSharedPtr<IPropertyHandle> RenderContextProperty = DetailLayoutBuilder.GetProperty( GET_MEMBER_NAME_CHECKED( AUsdStageActor, RenderContext ) ) )
+	{
+		DetailLayoutBuilder.HideProperty( RenderContextProperty );
+
+		CatBuilder.AddCustomRow( FText::FromString( TEXT( "RenderContextCustomization" ) ) )
+		.NameContent()
+		[
+			SNew( STextBlock )
+			.Text( FText::FromString( TEXT( "Render Context" ) ) )
+			.Font( FEditorStyle::GetFontStyle( TEXT( "PropertyWindow.NormalFont" ) ) )
+			.ToolTipText( RenderContextProperty->GetToolTipText() )
+		]
+		.ValueContent()
+		[
+			SAssignNew( ComboBox, SComboBox<TSharedPtr<FString>> )
+			.OptionsSource( &ComboBoxItems )
+			.InitiallySelectedItem( InitiallySelectedContext )
+			.OnSelectionChanged( this, &FUsdStageActorCustomization::OnComboBoxSelectionChanged )
+			.OnGenerateWidget_Lambda( []( TSharedPtr<FString> Item )
+			{
+				return SNew( STextBlock )
+					.Text( Item.IsValid() ? FText::FromString( *Item ) : FText::GetEmpty() )
+					.Font( FEditorStyle::GetFontStyle( TEXT( "PropertyWindow.NormalFont" ) ) );
+			} )
+			.Content()
+			[
+				SNew( STextBlock )
+				.Text( this, &FUsdStageActorCustomization::GetComboBoxSelectedOptionText )
+				.Font( FEditorStyle::GetFontStyle( TEXT( "PropertyWindow.NormalFont" ) ) )
+			]
+		];
+	}
+
+	// Add/remove properties so that they retain their usual order
+	if ( TSharedPtr<IPropertyHandle> TimeProperty = DetailLayoutBuilder.GetProperty( TEXT( "Time" ) ) )
+	{
+		CatBuilder.AddProperty( TimeProperty );
+	}
+	if ( TSharedPtr<IPropertyHandle> LevelSequenceProperty = DetailLayoutBuilder.GetProperty( TEXT( "LevelSequence" ) ) )
+	{
+		CatBuilder.AddProperty( LevelSequenceProperty );
+	}
+}
+
+void FUsdStageActorCustomization::CustomizeDetails( const TSharedPtr<IDetailLayoutBuilder>& DetailBuilder )
+{
+	DetailBuilderWeakPtr = DetailBuilder;
+	CustomizeDetails( *DetailBuilder );
+}
+
+void FUsdStageActorCustomization::OnComboBoxSelectionChanged( TSharedPtr<FString> NewContext, ESelectInfo::Type SelectType )
+{
+	if ( CurrentActor == nullptr || !NewContext.IsValid() )
+	{
+		return;
+	}
+
+	FScopedTransaction Transaction(FText::Format(
+		LOCTEXT("RenderContextChangedTransaction", "Changed the UsdStageActor {0}'s RenderContext to '{1}'"),
+		FText::FromString( CurrentActor->GetActorLabel() ),
+		FText::FromString( NewContext.IsValid() ? *NewContext : TEXT( "None" ) ) )
+	);
+
+	FName NewContextName = ( *NewContext ) == TEXT( "universal" ) ? NAME_None : FName( **NewContext );
+
+	CurrentActor->SetRenderContext( NewContextName );
+}
+
+FText FUsdStageActorCustomization::GetComboBoxSelectedOptionText() const
+{
+	TSharedPtr<FString> SelectedItem = ComboBox->GetSelectedItem();
+	if ( SelectedItem.IsValid() )
+	{
+		return FText::FromString( *SelectedItem );
+	}
+
+	return FText::GetEmpty();
+}
+
+void FUsdStageActorCustomization::ForceRefreshDetails()
+{
+	// Raw because we don't want to keep alive the details builder when calling the force refresh details
+	IDetailLayoutBuilder* DetailLayoutBuilder = DetailBuilderWeakPtr.Pin().Get();
+	if ( DetailLayoutBuilder )
+	{
+		DetailLayoutBuilder->ForceRefreshDetails();
+	}
+}
+
+#undef LOCTEXT_NAMESPACE
+
+#endif // WITH_EDITOR
\ No newline at end of file
diff --git a/Engine/Plugins/Importers/USDImporter/Source/USDStage/Private/USDStageActorCustomization.h b/Engine/Plugins/Importers/USDImporter/Source/USDStage/Private/USDStageActorCustomization.h
new file mode 100644
--- /dev/null
+++ b/Engine/Plugins/Importers/USDImporter/Source/USDStage/Private/USDStageActorCustomization.h
@@ -0,0 +1,37 @@
+// Copyright Epic Games, Inc. All Rights Reserved.
+
+#pragma once
+
+#if WITH_EDITOR
+
+#include "CoreMinimal.h"
+#include "IDetailCustomization.h"
+#include "Widgets/Input/SComboBox.h"
+
+class AUsdStageActor;
+class IDetailLayoutBuilder;
+
+class FUsdStageActorCustomization : public IDetailCustomization
+{
+public:
+	FUsdStageActorCustomization();
+	static TSharedRef<IDetailCustomization> MakeInstance();
+
+	// IDetailCustomization interface
+	virtual void CustomizeDetails( IDetailLayoutBuilder& DetailLayout ) override;
+	virtual void CustomizeDetails( const TSharedPtr<IDetailLayoutBuilder>& DetailBuilder ) override;
+	// End of IDetailCustomization interface
+
+private:
+	void OnComboBoxOpening();
+	void OnComboBoxSelectionChanged( TSharedPtr<FString> NewContext, ESelectInfo::Type SelectType );
+	FText GetComboBoxSelectedOptionText() const;
+	void ForceRefreshDetails();
+
+	AUsdStageActor* CurrentActor;
+	TWeakPtr<IDetailLayoutBuilder> DetailBuilderWeakPtr;
+	TSharedPtr<SComboBox<TSharedPtr<FString>>> ComboBox;
+	TArray<TSharedPtr<FString>> ComboBoxItems;
+};
+
+#endif // WITH_EDITOR
\ No newline at end of file
diff --git a/Engine/Plugins/Importers/USDImporter/Source/USDStage/Private/USDStageModule.cpp b/Engine/Plugins/Importers/USDImporter/Source/USDStage/Private/USDStageModule.cpp
--- a/Engine/Plugins/Importers/USDImporter/Source/USDStage/Private/USDStageModule.cpp
+++ b/Engine/Plugins/Importers/USDImporter/Source/USDStage/Private/USDStageModule.cpp
@@ -3,6 +3,7 @@
 #include "USDStageModule.h"
 
 #include "USDStageActor.h"
+#include "USDStageActorCustomization.h"
 
 #include "Modules/ModuleManager.h"
 #include "EngineUtils.h"
@@ -50,6 +51,9 @@ public:
 				}
 			}
 		}
+
+		FPropertyEditorModule& PropertyModule = FModuleManager::LoadModuleChecked<FPropertyEditorModule>( TEXT( "PropertyEditor" ) );
+		PropertyModule.RegisterCustomClassLayout( TEXT( "UsdStageActor" ), FOnGetDetailCustomizationInstance::CreateStatic( &FUsdStageActorCustomization::MakeInstance ) );
 #endif // WITH_EDITOR
 	}
 
@@ -74,6 +78,9 @@ public:
 				}
 			}
 		}
+
+		FPropertyEditorModule& PropertyModule = FModuleManager::LoadModuleChecked< FPropertyEditorModule >( TEXT( "PropertyEditor" ) );
+		PropertyModule.UnregisterCustomClassLayout( TEXT( "UsdStageActor" ) );
 #endif // WITH_EDITOR
 	}
 
diff --git a/Engine/Plugins/Importers/USDImporter/Source/USDStage/Public/USDStageActor.h b/Engine/Plugins/Importers/USDImporter/Source/USDStage/Public/USDStageActor.h
--- a/Engine/Plugins/Importers/USDImporter/Source/USDStage/Public/USDStageActor.h
+++ b/Engine/Plugins/Importers/USDImporter/Source/USDStage/Public/USDStageActor.h
@@ -62,6 +62,7 @@ public:
 	UPROPERTY( EditAnywhere, BlueprintReadOnly, Category = "USD", meta = ( NoSpinbox = "true", UIMin = "0", ClampMin = "0" ))
 	int32 NaniteTriangleThreshold;
 
+	/** Specifies which set of shaders to use when parsing USD materials, in addition to the universal render context. */
 	UPROPERTY(EditAnywhere, BlueprintReadOnly, Category = "USD")
 	FName RenderContext;
 
diff --git a/Engine/Plugins/Importers/USDImporter/Source/USDStage/Public/USDStageActorCustomization.h b/Engine/Plugins/Importers/USDImporter/Source/USDStage/Public/USDStageActorCustomization.h
new file mode 100644
--- /dev/null
+++ b/Engine/Plugins/Importers/USDImporter/Source/USDStage/Public/USDStageActorCustomization.h
@@ -0,0 +1,33 @@
+// Copyright Epic Games, Inc. All Rights Reserved.
+
+#pragma once
+
+#include "CoreMinimal.h"
+#include "IDetailCustomization.h"
+#include "Widgets/Input/SComboBox.h"
+
+class AUsdStageActor;
+class IDetailLayoutBuilder;
+
+class FUsdStageActorCustomization : public IDetailCustomization
+{
+public:
+	FUsdStageActorCustomization();
+	static TSharedRef<IDetailCustomization> MakeInstance();
+
+	// IDetailCustomization interface
+	virtual void CustomizeDetails( IDetailLayoutBuilder& DetailLayout ) override;
+	virtual void CustomizeDetails( const TSharedPtr<IDetailLayoutBuilder>& DetailBuilder ) override;
+	// End of IDetailCustomization interface
+
+private:
+	void OnComboBoxOpening();
+	void OnComboBoxSelectionChanged( TSharedPtr<FString> NewContext, ESelectInfo::Type SelectType );
+	FText GetComboBoxSelectedOptionText() const;
+	void ForceRefreshDetails();
+
+	AUsdStageActor* CurrentActor;
+	TWeakPtr<IDetailLayoutBuilder> DetailBuilderWeakPtr;
+	TSharedPtr<SComboBox<TSharedPtr<FString>>> ComboBox;
+	TArray<TSharedPtr<FString>> ComboBoxItems;
+};
diff --git a/Engine/Plugins/Importers/USDImporter/Source/USDStage/USDStage.Build.cs b/Engine/Plugins/Importers/USDImporter/Source/USDStage/USDStage.Build.cs
--- a/Engine/Plugins/Importers/USDImporter/Source/USDStage/USDStage.Build.cs
+++ b/Engine/Plugins/Importers/USDImporter/Source/USDStage/USDStage.Build.cs
@@ -20,9 +20,9 @@ namespace UnrealBuildTool.Rules
 					"MeshDescription",
 					"MovieScene",
 					"MovieSceneTracks",
+					"Projects", // So that we can use the plugin manager to find out our content dir and cook the master materials
 					"Slate",
 					"SlateCore",
-					"Projects", // So that we can use the plugin manager to find out our content dir and cook the master materials
 					"StaticMeshDescription",
 					"UnrealUSDWrapper",
 					"USDClasses",
@@ -37,7 +37,10 @@ namespace UnrealBuildTool.Rules
 					new string[]
 					{
 						"DeveloperToolSettings",
+						"EditorStyle", // For the font style on the stage actor customization
+						"InputCore", // For keyboard control on the widget in the stage actor customization
 						"LevelSequenceEditor",
+						"PropertyEditor", // For the stage actor's details customization
 						"Sequencer",
 						"UnrealEd",
 					}
diff --git a/Engine/Plugins/Importers/USDImporter/Source/USDStageImporter/Private/USDStageImportOptionsCustomization.cpp b/Engine/Plugins/Importers/USDImporter/Source/USDStageImporter/Private/USDStageImportOptionsCustomization.cpp
new file mode 100644
--- /dev/null
+++ b/Engine/Plugins/Importers/USDImporter/Source/USDStageImporter/Private/USDStageImportOptionsCustomization.cpp
@@ -0,0 +1,145 @@
+// Copyright Epic Games, Inc. All Rights Reserved.
+
+#include "USDStageImportOptionsCustomization.h"
+
+#include "USDSchemasModule.h"
+#include "USDSchemaTranslator.h"
+#include "USDStageImportOptions.h"
+
+#include "DetailCategoryBuilder.h"
+#include "DetailLayoutBuilder.h"
+#include "DetailWidgetRow.h"
+#include "Modules/ModuleManager.h"
+#include "UObject/Object.h"
+#include "Widgets/DeclarativeSyntaxSupport.h"
+
+#define LOCTEXT_NAMESPACE "UsdStageImportOptionsCustomization"
+
+FUsdStageImportOptionsCustomization::FUsdStageImportOptionsCustomization()
+{
+}
+
+TSharedRef<IDetailCustomization> FUsdStageImportOptionsCustomization::MakeInstance()
+{
+	return MakeShared<FUsdStageImportOptionsCustomization>();
+}
+
+void FUsdStageImportOptionsCustomization::CustomizeDetails(IDetailLayoutBuilder& DetailLayoutBuilder)
+{
+	TArray<TWeakObjectPtr<UObject>> SelectedObjects = DetailLayoutBuilder.GetSelectedObjects();
+	if ( SelectedObjects.Num() != 1 )
+	{
+		return;
+	}
+
+	TWeakObjectPtr<UObject> SelectedObject = SelectedObjects[ 0 ];
+	if ( !SelectedObject.IsValid() )
+	{
+		return;
+	}
+
+	CurrentOptions = Cast<UUsdStageImportOptions>( SelectedObject.Get() );
+	if ( !CurrentOptions )
+	{
+		return;
+	}
+
+	IUsdSchemasModule& UsdSchemasModule = FModuleManager::Get().LoadModuleChecked< IUsdSchemasModule >( TEXT( "USDSchemas" ) );
+
+	ComboBoxItems.Reset();
+	TSharedPtr<FString> InitiallySelectedContext;
+	for ( const FName& Context : UsdSchemasModule.GetRenderContextRegistry().GetRenderContexts() )
+	{
+		TSharedPtr<FString> ContextStr;
+		if ( Context == NAME_None )
+		{
+			ContextStr = MakeShared<FString>( TEXT( "universal" ) );
+		}
+		else
+		{
+			ContextStr = MakeShared<FString>( Context.ToString() );
+		}
+
+		if ( Context == CurrentOptions->RenderContextToImport )
+		{
+			InitiallySelectedContext = ContextStr;
+		}
+
+		ComboBoxItems.Add( ContextStr );
+	}
+
+	IDetailCategoryBuilder& CatBuilder = DetailLayoutBuilder.EditCategory( TEXT( "USD options" ) );
+
+	if ( TSharedPtr<IPropertyHandle> RenderContextProperty = DetailLayoutBuilder.GetProperty( GET_MEMBER_NAME_CHECKED( UUsdStageImportOptions, RenderContextToImport ) ) )
+	{
+		DetailLayoutBuilder.HideProperty( RenderContextProperty );
+
+		CatBuilder.AddCustomRow( FText::FromString( TEXT( "RenderContextCustomization" ) ) )
+		.NameContent()
+		[
+			SNew( STextBlock )
+			.Text( FText::FromString( TEXT( "Render Context to Import" ) ) )
+			.Font( FEditorStyle::GetFontStyle( TEXT( "PropertyWindow.NormalFont" ) ) )
+			.ToolTipText( RenderContextProperty->GetToolTipText() )
+		]
+		.ValueContent()
+		[
+			SAssignNew( ComboBox, SComboBox<TSharedPtr<FString>> )
+			.OptionsSource( &ComboBoxItems )
+			.InitiallySelectedItem( InitiallySelectedContext )
+			.OnSelectionChanged( this, &FUsdStageImportOptionsCustomization::OnComboBoxSelectionChanged )
+			.OnGenerateWidget_Lambda( []( TSharedPtr<FString> Item )
+			{
+				return SNew( STextBlock )
+					.Text( Item.IsValid() ? FText::FromString( *Item ) : FText::GetEmpty() )
+					.Font( FEditorStyle::GetFontStyle( TEXT( "PropertyWindow.NormalFont" ) ) );
+			} )
+			.Content()
+			[
+				SNew( STextBlock )
+				.Text( this, &FUsdStageImportOptionsCustomization::GetComboBoxSelectedOptionText )
+				.Font( FEditorStyle::GetFontStyle( TEXT( "PropertyWindow.NormalFont" ) ) )
+			]
+		];
+	}
+
+	// Add/remove properties so that they retain their usual order
+	if ( TSharedPtr<IPropertyHandle> OverrideStageOptionsProperty = DetailLayoutBuilder.GetProperty( GET_MEMBER_NAME_CHECKED( UUsdStageImportOptions, bOverrideStageOptions ) ) )
+	{
+		CatBuilder.AddProperty( OverrideStageOptionsProperty );
+	}
+	if ( TSharedPtr<IPropertyHandle> StageOptionsProperty = DetailLayoutBuilder.GetProperty( GET_MEMBER_NAME_CHECKED( UUsdStageImportOptions, StageOptions ) ) )
+	{
+		CatBuilder.AddProperty( StageOptionsProperty );
+	}
+}
+
+void FUsdStageImportOptionsCustomization::CustomizeDetails( const TSharedPtr<IDetailLayoutBuilder>& DetailBuilder )
+{
+	CustomizeDetails( *DetailBuilder );
+}
+
+void FUsdStageImportOptionsCustomization::OnComboBoxSelectionChanged( TSharedPtr<FString> NewContext, ESelectInfo::Type SelectType )
+{
+	if ( CurrentOptions == nullptr || !NewContext.IsValid() )
+	{
+		return;
+	}
+
+	FName NewContextName = ( *NewContext ) == TEXT( "universal" ) ? NAME_None : FName( **NewContext );
+
+	CurrentOptions->RenderContextToImport = NewContextName;
+}
+
+FText FUsdStageImportOptionsCustomization::GetComboBoxSelectedOptionText() const
+{
+	TSharedPtr<FString> SelectedItem = ComboBox->GetSelectedItem();
+	if ( SelectedItem.IsValid() )
+	{
+		return FText::FromString( *SelectedItem );
+	}
+
+	return FText::GetEmpty();
+}
+
+#undef LOCTEXT_NAMESPACE
diff --git a/Engine/Plugins/Importers/USDImporter/Source/USDStageImporter/Private/USDStageImportOptionsCustomization.h b/Engine/Plugins/Importers/USDImporter/Source/USDStageImporter/Private/USDStageImportOptionsCustomization.h
new file mode 100644
--- /dev/null
+++ b/Engine/Plugins/Importers/USDImporter/Source/USDStageImporter/Private/USDStageImportOptionsCustomization.h
@@ -0,0 +1,36 @@
+// Copyright Epic Games, Inc. All Rights Reserved.
+
+#pragma once
+
+#if WITH_EDITOR
+
+#include "CoreMinimal.h"
+#include "IDetailCustomization.h"
+#include "Widgets/Input/SComboBox.h"
+
+class UUsdStageImportOptions;
+class IDetailLayoutBuilder;
+
+class FUsdStageImportOptionsCustomization : public IDetailCustomization
+{
+public:
+	FUsdStageImportOptionsCustomization();
+	static TSharedRef<IDetailCustomization> MakeInstance();
+
+	// IDetailCustomization interface
+	virtual void CustomizeDetails(IDetailLayoutBuilder& DetailLayout) override;
+	virtual void CustomizeDetails(const TSharedPtr<IDetailLayoutBuilder>& DetailBuilder) override;
+	// End of IDetailCustomization interface
+
+private:
+	void OnComboBoxOpening();
+	void OnComboBoxSelectionChanged( TSharedPtr<FString> NewContext, ESelectInfo::Type SelectType );
+	FText GetComboBoxSelectedOptionText() const;
+	void ForceRefreshDetails();
+
+	UUsdStageImportOptions* CurrentOptions;
+	TSharedPtr<SComboBox<TSharedPtr<FString>>> ComboBox;
+	TArray<TSharedPtr<FString>> ComboBoxItems;
+};
+
+#endif // WITH_EDITOR
\ No newline at end of file
diff --git a/Engine/Plugins/Importers/USDImporter/Source/USDStageImporter/Private/USDStageImporterModule.cpp b/Engine/Plugins/Importers/USDImporter/Source/USDStageImporter/Private/USDStageImporterModule.cpp
--- a/Engine/Plugins/Importers/USDImporter/Source/USDStageImporter/Private/USDStageImporterModule.cpp
+++ b/Engine/Plugins/Importers/USDImporter/Source/USDStageImporter/Private/USDStageImporterModule.cpp
@@ -4,9 +4,11 @@
 
 #include "UnrealUSDWrapper.h"
 #include "USDStageImporter.h"
+#include "USDStageImportOptionsCustomization.h"
 
 #include "Misc/Paths.h"
 #include "Modules/ModuleManager.h"
+#include "PropertyEditorModule.h"
 #include "Templates/UniquePtr.h"
 
 #define LOCTEXT_NAMESPACE "UsdStageImporterModule"
@@ -19,12 +21,18 @@ public:
 #if USE_USD_SDK
 		IUnrealUSDWrapperModule& UnrealUSDWrapperModule = FModuleManager::Get().LoadModuleChecked< IUnrealUSDWrapperModule >(TEXT("UnrealUSDWrapper"));
 
+		FPropertyEditorModule& PropertyModule = FModuleManager::LoadModuleChecked<FPropertyEditorModule>( TEXT( "PropertyEditor" ) );
+		PropertyModule.RegisterCustomClassLayout( TEXT( "UsdStageImportOptions" ), FOnGetDetailCustomizationInstance::CreateStatic( &FUsdStageImportOptionsCustomization::MakeInstance ) );
+
 		USDStageImporter = MakeUnique<UUsdStageImporter>();
 #endif // #if USE_USD_SDK
 	}
 
 	virtual void ShutdownModule() override
 	{
+		FPropertyEditorModule& PropertyModule = FModuleManager::LoadModuleChecked< FPropertyEditorModule >( TEXT( "PropertyEditor" ) );
+		PropertyModule.UnregisterCustomClassLayout( TEXT( "UsdStageImportOptions" ) );
+
 		USDStageImporter.Reset();
 	}
 
diff --git a/Engine/Plugins/Importers/USDImporter/Source/USDStageImporter/Public/USDStageImportOptions.h b/Engine/Plugins/Importers/USDImporter/Source/USDStageImporter/Public/USDStageImportOptions.h
--- a/Engine/Plugins/Importers/USDImporter/Source/USDStageImporter/Public/USDStageImportOptions.h
+++ b/Engine/Plugins/Importers/USDImporter/Source/USDStageImporter/Public/USDStageImportOptions.h
@@ -73,7 +73,7 @@ public:
 	UPROPERTY(BlueprintReadWrite, config, EditAnywhere, Category = "USD options", meta = ( NoSpinbox = "true", UIMin = "0", ClampMin = "0" ) )
 	int32 NaniteTriangleThreshold;
 
-	/** Specifies which set of shaders to use, defaults to universal. */
+	/** Specifies which set of shaders to use when parsing USD materials, in addition to the universal render context. */
 	UPROPERTY(BlueprintReadWrite, config, EditAnywhere, Category= "USD options")
 	FName RenderContextToImport;
 
diff --git a/Engine/Plugins/Importers/USDImporter/Source/USDStageImporter/Public/USDStageImportOptionsCustomization.h b/Engine/Plugins/Importers/USDImporter/Source/USDStageImporter/Public/USDStageImportOptionsCustomization.h
new file mode 100644
--- /dev/null
+++ b/Engine/Plugins/Importers/USDImporter/Source/USDStageImporter/Public/USDStageImportOptionsCustomization.h
@@ -0,0 +1,32 @@
+// Copyright Epic Games, Inc. All Rights Reserved.
+
+#pragma once
+
+#include "CoreMinimal.h"
+#include "IDetailCustomization.h"
+#include "Widgets/Input/SComboBox.h"
+
+class UUsdStageImportOptions;
+class IDetailLayoutBuilder;
+
+class FUsdStageImportOptionsCustomization : public IDetailCustomization
+{
+public:
+	FUsdStageImportOptionsCustomization();
+	static TSharedRef<IDetailCustomization> MakeInstance();
+
+	// IDetailCustomization interface
+	virtual void CustomizeDetails(IDetailLayoutBuilder& DetailLayout) override;
+	virtual void CustomizeDetails(const TSharedPtr<IDetailLayoutBuilder>& DetailBuilder) override;
+	// End of IDetailCustomization interface
+
+private:
+	void OnComboBoxOpening();
+	void OnComboBoxSelectionChanged( TSharedPtr<FString> NewContext, ESelectInfo::Type SelectType );
+	FText GetComboBoxSelectedOptionText() const;
+	void ForceRefreshDetails();
+
+	UUsdStageImportOptions* CurrentOptions;
+	TSharedPtr<SComboBox<TSharedPtr<FString>>> ComboBox;
+	TArray<TSharedPtr<FString>> ComboBoxItems;
+};
diff --git a/Engine/Plugins/Importers/USDImporter/Source/USDStageImporter/USDStageImporter.Build.cs b/Engine/Plugins/Importers/USDImporter/Source/USDStageImporter/USDStageImporter.Build.cs
--- a/Engine/Plugins/Importers/USDImporter/Source/USDStageImporter/USDStageImporter.Build.cs
+++ b/Engine/Plugins/Importers/USDImporter/Source/USDStageImporter/USDStageImporter.Build.cs
@@ -15,6 +15,7 @@ namespace UnrealBuildTool.Rules
 					"Core",
 					"CoreUObject",
 					"EditorFramework",
+					"EditorStyle", // For the font style on the stage actor customization
 					"Engine",
 					"GeometryCache",
 					"InputCore",
@@ -23,14 +24,15 @@ namespace UnrealBuildTool.Rules
 					"MainFrame",
 					"MessageLog",
 					"MovieScene",
+					"PropertyEditor", // For the import options's details customization
 					"RenderCore", // So that we can release resources of reimported meshes
 					"Slate",
 					"SlateCore",
 					"UnrealEd",
 					"UnrealUSDWrapper",
 					"USDClasses",
-					"USDStage",
 					"USDSchemas",
+					"USDStage",
 					"USDUtilities",
 				}
 			);
