USD: Add buttons for different types of Save As.

This recovers an old implementation for 4.26 from swarm 14356418 and updates it to 5.0.

diff --git a/Engine/Plugins/Importers/USDImporter/Source/USDStageEditor/Private/SUSDLayersTreeView.cpp b/Engine/Plugins/Importers/USDImporter/Source/USDStageEditor/Private/SUSDLayersTreeView.cpp
--- a/Engine/Plugins/Importers/USDImporter/Source/USDStageEditor/Private/SUSDLayersTreeView.cpp
+++ b/Engine/Plugins/Importers/USDImporter/Source/USDStageEditor/Private/SUSDLayersTreeView.cpp
@@ -13,9 +13,11 @@
 #include "UsdWrappers/SdfLayer.h"
 #include "UsdWrappers/UsdStage.h"
 
+#include "DesktopPlatformModule.h"
 #include "EditorStyleSet.h"
 #include "Engine/World.h"
 #include "Framework/MultiBox/MultiBoxBuilder.h"
+#include "IDesktopPlatform.h"
 #include "Modules/ModuleManager.h"
 #include "Styling/SlateTypes.h"
 #include "Widgets/Images/SImage.h"
@@ -26,6 +28,37 @@
 
 #define LOCTEXT_NAMESPACE "SUSDLayersTreeView"
 
+namespace UE::USDLayersTreeViewImpl::Private
+{
+	void ExportLayerToPath( const UE::FSdfLayer& LayerToExport, const FString& TargetPath )
+	{
+		if ( !LayerToExport )
+		{
+			return;
+		}
+
+		// Clone the layer so that we don't modify the currently opened stage when we do the remapping below
+		UE::FSdfLayer OutputLayer = UE::FSdfLayer::CreateNew( *TargetPath );
+		OutputLayer.TransferContent( LayerToExport );
+
+		// Update references to assets (e.g. textures) so that they're absolute and also work from the new file
+		UsdUtils::ConvertAssetRelativePathsToAbsolute( OutputLayer, LayerToExport );
+
+		// Convert layer references to absolute paths so that it still works at its target location
+		FString LayerPath = LayerToExport.GetRealPath();
+		FPaths::NormalizeFilename( LayerPath );
+		TSet<FString> ExternalReferences = OutputLayer.GetExternalReferences();
+		for ( const FString& Ref : ExternalReferences )
+		{
+			FString AbsRef = FPaths::ConvertRelativePathToFull( FPaths::GetPath( LayerPath ), Ref ); // Relative to the original file
+			OutputLayer.UpdateExternalReference( *Ref, *AbsRef );
+		}
+
+		bool bForce = true;
+		OutputLayer.Save( bForce );
+	}
+}
+
 class FUsdLayerNameColumn : public FUsdTreeViewColumn
 {
 public:
@@ -296,6 +329,18 @@ TSharedPtr< SWidget > SUsdLayersTreeView::ConstructLayerContextMenu()
 			NAME_None,
 			EUserInterfaceActionType::Button
 		);
+
+		LayerOptions.AddMenuEntry(
+			LOCTEXT( "ExportLayer", "Export" ),
+			LOCTEXT( "Export_ToolTip", "Export the selected layers, having the exported layers reference the original stage's layers" ),
+			FSlateIcon(),
+			FUIAction(
+				FExecuteAction::CreateSP( this, &SUsdLayersTreeView::OnExportSelectedLayers ),
+				FCanExecuteAction()
+			),
+			NAME_None,
+			EUserInterfaceActionType::Button
+		);
 	}
 	LayerOptions.EndSection();
 
@@ -375,6 +420,76 @@ void SUsdLayersTreeView::OnEditSelectedLayer()
 	}
 }
 
+void SUsdLayersTreeView::OnExportSelectedLayers() const
+{
+	TArray< FUsdLayerViewModelRef > MySelectedItems = GetSelectedItems();
+
+	TArray<UE::FSdfLayer> LayersToExport;
+	LayersToExport.Reserve( MySelectedItems.Num() );
+
+	for ( FUsdLayerViewModelRef SelectedItem : MySelectedItems )
+	{
+		UE::FSdfLayer SelectedLayer = SelectedItem->GetLayer();
+		if ( !SelectedLayer )
+		{
+			continue;
+		}
+
+		LayersToExport.Add( SelectedLayer );
+	}
+
+	// Single layer -> Allow picking the target layer filename
+	if ( LayersToExport.Num() == 1 )
+	{
+		TOptional< FString > UsdFilePath = UsdUtils::BrowseUsdFile( UsdUtils::EBrowseFileMode::Save, AsShared() );
+		if ( !UsdFilePath.IsSet() )
+		{
+			return;
+		}
+
+		UE::USDLayersTreeViewImpl::Private::ExportLayerToPath( LayersToExport[ 0 ], UsdFilePath.GetValue() );
+	}
+	// Multiple layers -> Pick folder and export them with the same name
+	if ( LayersToExport.Num() > 1 )
+	{
+		IDesktopPlatform* DesktopPlatform = FDesktopPlatformModule::Get();
+		if ( !DesktopPlatform )
+		{
+			return;
+		}
+
+		TSharedPtr< SWindow > ParentWindow = FSlateApplication::Get().FindWidgetWindow( AsShared() );
+		void* ParentWindowHandle = ( ParentWindow.IsValid() && ParentWindow->GetNativeWindow().IsValid() )
+			? ParentWindow->GetNativeWindow()->GetOSWindowHandle()
+			: nullptr;
+
+		FString TargetFolderPath;
+		if ( !DesktopPlatform->OpenDirectoryDialog( ParentWindowHandle, LOCTEXT( "ChooseFolder", "Choose output folder" ).ToString(), TEXT( "" ), TargetFolderPath ) )
+		{
+			return;
+		}
+		TargetFolderPath = FPaths::ConvertRelativePathToFull( TargetFolderPath );
+
+		if ( FPaths::DirectoryExists( TargetFolderPath ) )
+		{
+			for ( const UE::FSdfLayer& LayerToExport : LayersToExport )
+			{
+				FString TargetFileName = FPaths::GetCleanFilename( LayerToExport.GetRealPath() );
+				FString FullPath = FPaths::Combine( TargetFolderPath, TargetFileName );
+				FString FinalFullPath = FullPath;
+
+				uint32 Suffix = 0;
+				while ( FPaths::FileExists( FinalFullPath ) )
+				{
+					FinalFullPath = FString::Printf( TEXT( "%s_%u" ), *FullPath, Suffix++ );
+				}
+
+				UE::USDLayersTreeViewImpl::Private::ExportLayerToPath( LayerToExport, FinalFullPath );
+			}
+		}
+	}
+}
+
 bool SUsdLayersTreeView::CanAddSubLayer() const
 {
 	return GetSelectedItems().Num() > 0;
diff --git a/Engine/Plugins/Importers/USDImporter/Source/USDStageEditor/Private/SUSDLayersTreeView.h b/Engine/Plugins/Importers/USDImporter/Source/USDStageEditor/Private/SUSDLayersTreeView.h
--- a/Engine/Plugins/Importers/USDImporter/Source/USDStageEditor/Private/SUSDLayersTreeView.h
+++ b/Engine/Plugins/Importers/USDImporter/Source/USDStageEditor/Private/SUSDLayersTreeView.h
@@ -33,6 +33,8 @@ private:
 	bool CanEditSelectedLayer() const;
 	void OnEditSelectedLayer();
 
+	void OnExportSelectedLayers() const;
+
 	bool CanAddSubLayer() const;
 	void OnAddSubLayer();
 	void OnNewSubLayer();
diff --git a/Engine/Plugins/Importers/USDImporter/Source/USDStageEditor/Private/SUSDStage.cpp b/Engine/Plugins/Importers/USDImporter/Source/USDStageEditor/Private/SUSDStage.cpp
--- a/Engine/Plugins/Importers/USDImporter/Source/USDStageEditor/Private/SUSDStage.cpp
+++ b/Engine/Plugins/Importers/USDImporter/Source/USDStageEditor/Private/SUSDStage.cpp
@@ -6,6 +6,7 @@
 #include "SUSDPrimInfo.h"
 #include "SUSDStageTreeView.h"
 #include "UnrealUSDWrapper.h"
+#include "USDConversionUtils.h"
 #include "USDErrorUtils.h"
 #include "USDLayerUtils.h"
 #include "USDSchemasModule.h"
@@ -24,6 +25,7 @@
 
 #include "ActorTreeItem.h"
 #include "Async/Async.h"
+#include "DesktopPlatformModule.h"
 #include "Dialogs/DlgPickPath.h"
 #include "EditorStyleSet.h"
 #include "Engine/Selection.h"
@@ -48,6 +50,65 @@ namespace SUSDStageConstants
 
 namespace SUSDStageImpl
 {
+#if PLATFORM_LINUX
+	struct FCaseSensitiveStringSetFuncs : BaseKeyFuncs<FString, FString>
+	{
+		static FORCEINLINE const FString& GetSetKey( const FString& Element )
+		{
+			return Element;
+		}
+		static FORCEINLINE bool Matches( const FString& A, const FString& B )
+		{
+			return A.Equals( B, ESearchCase::CaseSensitive );
+		}
+		static FORCEINLINE uint32 GetKeyHash( const FString& Key )
+		{
+			return FCrc::StrCrc32<TCHAR>( *Key );
+		}
+	};
+#endif
+
+	/**
+	 * Makes sure that AllLayers includes all the external references of any of its members.
+	 * Example: Receives AllLayers [a.usda], that references b.usda and c.usda, while c.usda also references d.usda -> Modifies AllLayers to be [a.usda, b.usda, c.usda, d.usda]
+	 */
+	void AppendAllExternalReferences( TArray<UE::FSdfLayer>& AllLayers )
+	{
+		// Consider paths in a case-sensitive way for linux
+#if PLATFORM_LINUX
+		TSet<FString, FCaseSensitiveStringSetFuncs> LoadedLayers;
+#else
+		TSet<FString> LoadedLayers;
+#endif
+
+		LoadedLayers.Reserve( AllLayers.Num() );
+		for ( const UE::FSdfLayer& Layer : AllLayers )
+		{
+			FString LayerPath = Layer.GetRealPath();
+			FPaths::NormalizeFilename( LayerPath );
+			LoadedLayers.Add( LayerPath );
+		}
+
+		for ( int32 Index = 0; Index < AllLayers.Num(); ++Index )
+		{
+			UE::FSdfLayer& Layer = AllLayers[ Index ];
+
+			TArray< UE::FSdfLayer > NewLayers;
+			for ( const FString& Reference : Layer.GetExternalReferences() )
+			{
+				FString AbsoluteReference = Layer.ComputeAbsolutePath( Reference );
+				FPaths::NormalizeFilename( AbsoluteReference );
+
+				if ( !LoadedLayers.Contains( AbsoluteReference ) )
+				{
+					NewLayers.Add( UE::FSdfLayer::FindOrOpen( *AbsoluteReference ) );
+					LoadedLayers.Add( AbsoluteReference );
+				}
+			}
+			AllLayers.Append( NewLayers );
+		}
+	}
+
 	void SelectGeneratedComponentsAndActors( AUsdStageActor* StageActor, const TArray<FString>& PrimPaths )
 	{
 		if ( !StageActor )
@@ -507,8 +568,16 @@ void SUsdStage::FillFileMenu( FMenuBuilder& MenuBuilder )
 			EUserInterfaceActionType::Button
 		);
 
-		MenuBuilder.AddSeparator();
+		MenuBuilder.AddSubMenu(
+			LOCTEXT( "Export", "Export" ),
+			FText::GetEmpty(),
+			FNewMenuDelegate::CreateSP( this, &SUsdStage::FillExportSubMenu )
+		);
+	}
+	MenuBuilder.EndSection();
 
+	MenuBuilder.BeginSection( "Reload", LOCTEXT( "Reload", "Reload" ) );
+	{
 		MenuBuilder.AddMenuEntry(
 			LOCTEXT("Reload", "Reload"),
 			LOCTEXT("Reload_ToolTip", "Reloads the stage from disk, keeping aspects of the session intact"),
@@ -557,9 +626,11 @@ void SUsdStage::FillFileMenu( FMenuBuilder& MenuBuilder )
 			NAME_None,
 			EUserInterfaceActionType::Button
 		);
+	}
+	MenuBuilder.EndSection();
 
-		MenuBuilder.AddSeparator();
-
+	MenuBuilder.BeginSection( "Close", LOCTEXT( "Close", "Close" ) );
+	{
 		MenuBuilder.AddMenuEntry(
 			LOCTEXT("Close", "Close"),
 			LOCTEXT("Close_ToolTip", "Closes the opened stage"),
@@ -665,6 +736,55 @@ void SUsdStage::FillOptionsMenu(FMenuBuilder& MenuBuilder)
 	MenuBuilder.EndSection();
 }
 
+void SUsdStage::FillExportSubMenu( FMenuBuilder& MenuBuilder )
+{
+	MenuBuilder.AddMenuEntry(
+		LOCTEXT( "ExportAll", "Export all layers..." ),
+		LOCTEXT( "ExportAll_ToolTip", "Exports copies of all file-based layers in the stage's layer stack to a new folder" ),
+		FSlateIcon(),
+		FUIAction(
+			FExecuteAction::CreateSP( this, &SUsdStage::OnExportAll ),
+			FCanExecuteAction::CreateLambda( [this]()
+			{
+				if ( const AUsdStageActor* StageActor = ViewModel.UsdStageActor.Get() )
+				{
+					if ( UE::FUsdStage Stage = StageActor->GetUsdStage() )
+					{
+						return true;
+					}
+				}
+
+				return false;
+			})
+		),
+		NAME_None,
+		EUserInterfaceActionType::Button
+	);
+
+	MenuBuilder.AddMenuEntry(
+		LOCTEXT( "ExportFlattened", "Export flattened stage..." ),
+		LOCTEXT( "ExportFlattened_ToolTip", "Flattens the current stage to a single USD layer and exports it as a new USD file" ),
+		FSlateIcon(),
+		FUIAction(
+			FExecuteAction::CreateSP( this, &SUsdStage::OnExportFlattened ),
+			FCanExecuteAction::CreateLambda( [this]()
+			{
+				if ( const AUsdStageActor* StageActor = ViewModel.UsdStageActor.Get() )
+				{
+					if ( UE::FUsdStage Stage = StageActor->GetUsdStage() )
+					{
+						return true;
+					}
+				}
+
+				return false;
+			})
+		),
+		NAME_None,
+		EUserInterfaceActionType::Button
+	);
+}
+
 void SUsdStage::FillPayloadsSubMenu(FMenuBuilder& MenuBuilder)
 {
 	MenuBuilder.AddMenuEntry(
@@ -1166,6 +1286,152 @@ void SUsdStage::OnSave()
 	}
 }
 
+void SUsdStage::OnExportAll()
+{
+	const AUsdStageActor* StageActor = ViewModel.UsdStageActor.Get();
+	if ( !StageActor )
+	{
+		return;
+	}
+
+	UE::FUsdStage UsdStage = StageActor->GetUsdStage();
+	if ( !UsdStage )
+	{
+		return;
+	}
+
+	const bool bIncludeClipLayers = false;
+	TArray<UE::FSdfLayer> LayerStack = UsdStage.GetUsedLayers( bIncludeClipLayers );
+	if ( LayerStack.Num() < 1 )
+	{
+		return;
+	}
+
+	IDesktopPlatform* DesktopPlatform = FDesktopPlatformModule::Get();
+	if ( !DesktopPlatform )
+	{
+		return;
+	}
+
+	TSharedPtr< SWindow > ParentWindow = FSlateApplication::Get().FindWidgetWindow( AsShared() );
+	void* ParentWindowHandle = ( ParentWindow.IsValid() && ParentWindow->GetNativeWindow().IsValid() )
+		? ParentWindow->GetNativeWindow()->GetOSWindowHandle()
+		: nullptr;
+
+	FString TargetFolderPath;
+	if ( !DesktopPlatform->OpenDirectoryDialog( ParentWindowHandle, LOCTEXT( "ChooseFolder", "Choose output folder" ).ToString(), TEXT( "" ), TargetFolderPath ) )
+	{
+		return;
+	}
+	TargetFolderPath = FPaths::ConvertRelativePathToFull( TargetFolderPath );
+
+	// Manually load any layer that is referenced by any of the stage's layers, but not currently loaded.
+	// This can happen if e.g. an unselected variant appends a reference or has a payload.
+	// We will need to actually load these as opposed to just copy-paste the files as we may need to update references/paths
+	SUSDStageImpl::AppendAllExternalReferences( LayerStack );
+
+	// Discard session layers
+	for ( int32 Index = LayerStack.Num() - 1; Index >= 0; --Index )
+	{
+		UE::FSdfLayer& Layer = LayerStack[ Index ];
+		if ( Layer.IsAnonymous() )
+		{
+			const int32 Count = 1;
+			const bool bAllowShrinking = false;
+			LayerStack.RemoveAt( Index, Count, bAllowShrinking );
+		}
+	}
+
+#if PLATFORM_LINUX
+	TMap<FString, FString, FDefaultSetAllocator, UsdUtils::FCaseSensitiveStringMapFuncs< FString > > OldPathToSavedPath;
+#else
+	TMap<FString, FString> OldPathToSavedPath;
+#endif
+
+	// If the stage has layers referencing each other, we will need to remap those to point exclusively at the newly saved files
+	TSet<FString> SavedPaths;
+	for ( const UE::FSdfLayer& Layer : LayerStack )
+	{
+		FString LayerPath = Layer.GetRealPath();
+		FPaths::NormalizeFilename( LayerPath );
+
+		FString TargetPath = FPaths::Combine( TargetFolderPath, Layer.GetDisplayName() );
+		FPaths::NormalizeFilename( TargetPath );
+
+		// Filename collision (should be rare, but possible given that we're discarding the folder structure)
+		if ( SavedPaths.Contains( TargetPath ) )
+		{
+			FString TargetPathNoExt = FPaths::SetExtension( TargetPath, TEXT( "" ) );
+			FString Ext = FPaths::GetExtension( TargetPath );
+			int32 Suffix = 0;
+
+			do
+			{
+				TargetPath = FString::Printf( TEXT( "%s_%d.%s" ), *TargetPathNoExt, Suffix++, *Ext );
+			} while ( SavedPaths.Contains( TargetPath ) );
+		}
+
+		OldPathToSavedPath.Add( LayerPath, TargetPath );
+		SavedPaths.Add( TargetPath );
+	}
+
+	// Actually save the output layers
+	for ( UE::FSdfLayer& Layer : LayerStack )
+	{
+		FString LayerPath = Layer.GetRealPath();
+		FPaths::NormalizeFilename( LayerPath );
+
+		if ( FString* TargetLayerPath = OldPathToSavedPath.Find( LayerPath ) )
+		{
+			// Clone the layer so that we don't modify the currently opened stage
+			UE::FSdfLayer OutputLayer = UE::FSdfLayer::CreateNew( **TargetLayerPath );
+			OutputLayer.TransferContent( Layer );
+
+			// Update references to assets (e.g. textures) so that they're absolute and also work from the new file
+			UsdUtils::ConvertAssetRelativePathsToAbsolute( OutputLayer, Layer );
+
+			// Remap references to layers so that they point at the other newly saved files. Note that SUSDStageImpl::AppendAllExternalReferences
+			// will have collected all external references already, so OldPathToSavedPath should have entries for all references we'll find
+			for ( const FString& Ref : OutputLayer.GetExternalReferences() )
+			{
+				FString AbsRef = FPaths::ConvertRelativePathToFull( FPaths::GetPath( LayerPath ), Ref ); // Relative to the original file
+				FPaths::NormalizeFilename( AbsRef );
+
+				if ( FString* SavedReference = OldPathToSavedPath.Find( AbsRef ) )
+				{
+					OutputLayer.UpdateExternalReference( *Ref, **SavedReference );
+				}
+			}
+
+			bool bForce = true;
+			OutputLayer.Save( bForce );
+		}
+	}
+}
+
+void SUsdStage::OnExportFlattened()
+{
+	const AUsdStageActor* StageActor = ViewModel.UsdStageActor.Get();
+	if ( !StageActor )
+	{
+		return;
+	}
+
+	UE::FUsdStage UsdStage = StageActor->GetUsdStage();
+	if ( !UsdStage )
+	{
+		return;
+	}
+
+	TOptional< FString > UsdFilePath = UsdUtils::BrowseUsdFile( UsdUtils::EBrowseFileMode::Save, AsShared() );
+	if ( !UsdFilePath.IsSet() )
+	{
+		return;
+	}
+
+	UsdStage.Export( *( UsdFilePath.GetValue() ) );
+}
+
 void SUsdStage::OnReloadStage()
 {
 	FScopedTransaction Transaction( LOCTEXT( "ReloadTransaction", "Reload USD stage" ) );
diff --git a/Engine/Plugins/Importers/USDImporter/Source/USDStageEditor/Private/SUSDStage.h b/Engine/Plugins/Importers/USDImporter/Source/USDStageEditor/Private/SUSDStage.h
--- a/Engine/Plugins/Importers/USDImporter/Source/USDStageEditor/Private/SUSDStage.h
+++ b/Engine/Plugins/Importers/USDImporter/Source/USDStageEditor/Private/SUSDStage.h
@@ -42,6 +42,7 @@ protected:
 	void FillFileMenu( FMenuBuilder& MenuBuilder );
 	void FillActionsMenu( FMenuBuilder& MenuBuilder );
 	void FillOptionsMenu( FMenuBuilder& MenuBuilder );
+	void FillExportSubMenu( FMenuBuilder& MenuBuilder );
 	void FillPayloadsSubMenu( FMenuBuilder& MenuBuilder );
 	void FillPurposesToLoadSubMenu( FMenuBuilder& MenuBuilder );
 	void FillRenderContextSubMenu( FMenuBuilder& MenuBuilder );
@@ -53,6 +54,8 @@ protected:
 	void OnNew();
 	void OnOpen();
 	void OnSave();
+	void OnExportAll();
+	void OnExportFlattened();
 	void OnReloadStage();
 	void OnResetStage();
 	void OnClose();
diff --git a/Engine/Plugins/Importers/USDImporter/Source/USDStageEditor/USDStageEditor.Build.cs b/Engine/Plugins/Importers/USDImporter/Source/USDStageEditor/USDStageEditor.Build.cs
--- a/Engine/Plugins/Importers/USDImporter/Source/USDStageEditor/USDStageEditor.Build.cs
+++ b/Engine/Plugins/Importers/USDImporter/Source/USDStageEditor/USDStageEditor.Build.cs
@@ -14,6 +14,7 @@ namespace UnrealBuildTool.Rules
 				{
 					"Core",
 					"CoreUObject",
+					"DesktopPlatform",
 					"DesktopWidgets",
 					"EditorFramework",
 					"EditorStyle",
diff --git a/Engine/Plugins/Importers/USDImporter/Source/USDUtilities/Private/USDLayerUtils.cpp b/Engine/Plugins/Importers/USDImporter/Source/USDUtilities/Private/USDLayerUtils.cpp
--- a/Engine/Plugins/Importers/USDImporter/Source/USDUtilities/Private/USDLayerUtils.cpp
+++ b/Engine/Plugins/Importers/USDImporter/Source/USDUtilities/Private/USDLayerUtils.cpp
@@ -25,6 +25,7 @@
 
 #include "USDIncludesStart.h"
 	#include "pxr/usd/pcp/layerStack.h"
+	#include "pxr/usd/sdf/attributeSpec.h"
 	#include "pxr/usd/sdf/fileFormat.h"
 	#include "pxr/usd/sdf/layer.h"
 	#include "pxr/usd/sdf/layerUtils.h"
@@ -38,6 +39,39 @@
 
 #define LOCTEXT_NAMESPACE "USDLayerUtils"
 
+namespace UE::UsdLayerUtilsImpl::Private
+{
+	/**
+		* Adapted from flattenUtils.cpp::_FixAssetPaths, except that we only handle actual AssetPaths here as layer/prim paths
+		* will be remapped via Layer.UpdateExternalReference.
+		* Returns whether anything was remapped
+		*/
+	bool FixAssetPaths( const pxr::SdfLayerHandle& SourceLayer, pxr::VtValue* Value )
+	{
+		if ( Value->IsHolding<pxr::SdfAssetPath>() )
+		{
+			pxr::SdfAssetPath AssetPath;
+			Value->Swap( AssetPath );
+			AssetPath = pxr::SdfAssetPath( SourceLayer->ComputeAbsolutePath( AssetPath.GetAssetPath() ) );
+			Value->Swap( AssetPath );
+			return true;
+		}
+		else if ( Value->IsHolding<pxr::VtArray<pxr::SdfAssetPath>>() )
+		{
+			pxr::VtArray<pxr::SdfAssetPath> PathArray;
+			Value->Swap( PathArray );
+			for ( pxr::SdfAssetPath& AssetPath : PathArray )
+			{
+				AssetPath = pxr::SdfAssetPath( SourceLayer->ComputeAbsolutePath( AssetPath.GetAssetPath() ) );
+			}
+			Value->Swap( PathArray );
+			return true;
+		}
+
+		return false;
+	}
+}
+
 namespace UsdUtils
 {
 	std::string GetUESessionStateLayerDisplayName( const pxr::UsdStageRefPtr& Stage )
@@ -664,6 +698,36 @@ bool UsdUtils::IsSessionLayerWithinStage( const pxr::SdfLayerRefPtr& Layer, cons
 	return false;
 }
 
+void UsdUtils::ConvertAssetRelativePathsToAbsolute( UE::FSdfLayer& LayerToConvert, const UE::FSdfLayer& AnchorLayer )
+{
+	FScopedUsdAllocs Allocs;
+
+	pxr::SdfLayerRefPtr UsdLayer{ LayerToConvert };
+
+	UsdLayer->Traverse(
+		pxr::SdfPath::AbsoluteRootPath(),
+		[ &UsdLayer, &AnchorLayer ]( const pxr::SdfPath& Path )
+		{
+			pxr::SdfSpecType SpecType = UsdLayer->GetSpecType( Path );
+			if ( SpecType != pxr::SdfSpecTypeAttribute )
+			{
+				return;
+			}
+
+			pxr::VtValue LayerValue;
+			if ( !UsdLayer->HasField( Path, pxr::SdfFieldKeys->Default, &LayerValue ) )
+			{
+				return;
+			}
+
+			if ( UE::UsdLayerUtilsImpl::Private::FixAssetPaths( pxr::SdfLayerRefPtr{ AnchorLayer }, &LayerValue ) )
+			{
+				UsdLayer->SetField( Path, pxr::SdfFieldKeys->Default, LayerValue );
+			}
+		}
+	);
+}
+
 #undef LOCTEXT_NAMESPACE
 
 #endif // #if USE_USD_SDK
diff --git a/Engine/Plugins/Importers/USDImporter/Source/USDUtilities/Public/USDLayerUtils.h b/Engine/Plugins/Importers/USDImporter/Source/USDUtilities/Public/USDLayerUtils.h
--- a/Engine/Plugins/Importers/USDImporter/Source/USDUtilities/Public/USDLayerUtils.h
+++ b/Engine/Plugins/Importers/USDImporter/Source/USDUtilities/Public/USDLayerUtils.h
@@ -103,6 +103,12 @@ namespace UsdUtils
 
 	/** Returns true if Layer is a session layer within Stage's layer stack */
 	USDUTILITIES_API bool IsSessionLayerWithinStage( const pxr::SdfLayerRefPtr& Layer, const pxr::UsdStageRefPtr& Stage );
+
+	/**
+	 * Finds all fields that have asset paths as values (e.g. texture references) in LayerToConvert and (if they're relative paths)
+	 * updates them to be absolute with respect to AnchorLayer. LayerToConvert and AnchorLayer can be the same layer, but they don't have to
+	 */
+	USDUTILITIES_API void ConvertAssetRelativePathsToAbsolute( UE::FSdfLayer& LayerToConvert, const UE::FSdfLayer& AnchorLayer );
 }
 
 #endif // #if USE_USD_SDK
diff --git a/Engine/Plugins/Importers/USDImporter/Source/UnrealUSDWrapper/Private/UsdWrappers/SdfLayer.cpp b/Engine/Plugins/Importers/USDImporter/Source/UnrealUSDWrapper/Private/UsdWrappers/SdfLayer.cpp
--- a/Engine/Plugins/Importers/USDImporter/Source/UnrealUSDWrapper/Private/UsdWrappers/SdfLayer.cpp
+++ b/Engine/Plugins/Importers/USDImporter/Source/UnrealUSDWrapper/Private/UsdWrappers/SdfLayer.cpp
@@ -227,6 +227,17 @@ namespace UE
 	}
 #endif // #if USE_USD_SDK
 
+	template<typename PtrType>
+	void FSdfLayerBase<PtrType>::TransferContent( const FSdfLayer& SourceLayer )
+	{
+#if USE_USD_SDK
+		if ( const PtrType& Ptr = Impl->GetInner() )
+		{
+			Ptr->TransferContent( pxr::SdfLayerRefPtr{ SourceLayer } );
+		}
+#endif // #if USE_USD_SDK
+	}
+
 	template<typename PtrType>
 	FSdfLayer FSdfLayerBase<PtrType>::FindOrOpen( const TCHAR* Identifier )
 	{
@@ -237,6 +248,16 @@ namespace UE
 #endif // #if USE_USD_SDK
 	}
 
+	template<typename PtrType>
+	FSdfLayer FSdfLayerBase<PtrType>::CreateNew( const TCHAR* Identifier )
+	{
+#if USE_USD_SDK
+		return FSdfLayer( pxr::SdfLayer::CreateNew( TCHAR_TO_ANSI( Identifier ) ) );
+#else
+		return FSdfLayer();
+#endif // #if USE_USD_SDK
+	}
+
 	template<typename PtrType>
 	bool FSdfLayerBase<PtrType>::Save( bool bForce ) const
 	{
@@ -250,6 +271,41 @@ namespace UE
 		return false;
 	}
 
+	template<typename PtrType>
+	TSet<FString> FSdfLayerBase<PtrType>::GetExternalReferences() const
+	{
+		TSet<FString> Result;
+#if USE_USD_SDK
+		if ( const PtrType& Ptr = Impl->GetInner() )
+		{
+			FScopedUsdAllocs UsdAllocs;
+			std::set<std::string> ExternalReferences = Ptr->GetExternalReferences();
+
+			Result.Reserve( ExternalReferences.size() );
+
+			for ( const std::string& Reference : ExternalReferences )
+			{
+				Result.Add( ANSI_TO_TCHAR( Reference.c_str() ) );
+			}
+		}
+#endif // #if USE_USD_SDK
+		return Result;
+	}
+
+	template<typename PtrType>
+	bool FSdfLayerBase<PtrType>::UpdateExternalReference( const TCHAR* OldReferencePath, const TCHAR* NewReferencePath )
+	{
+#if USE_USD_SDK
+		if ( const PtrType& Ptr = Impl->GetInner() )
+		{
+			FScopedUsdAllocs UsdAllocs;
+			return Ptr->UpdateExternalReference( TCHAR_TO_ANSI( OldReferencePath ), TCHAR_TO_ANSI( NewReferencePath ) );
+		}
+#endif // #if USE_USD_SDK
+
+		return false;
+	}
+
 	template<typename PtrType>
 	FString FSdfLayerBase<PtrType>::GetRealPath() const
 	{
@@ -292,6 +348,20 @@ namespace UE
 		return FString();
 	}
 
+	template<typename PtrType>
+	FString FSdfLayerBase<PtrType>::ComputeAbsolutePath( const FString& AssetPath ) const
+	{
+#if USE_USD_SDK
+		if ( const PtrType& Ptr = Impl->GetInner() )
+		{
+			FScopedUsdAllocs UsdAllocs;
+			return FString( ANSI_TO_TCHAR( Ptr->ComputeAbsolutePath( TCHAR_TO_ANSI( *AssetPath ) ).c_str() ) );
+		}
+#endif // #if USE_USD_SDK
+
+		return FString();
+	}
+
 	template<typename PtrType>
 	bool FSdfLayerBase<PtrType>::HasStartTimeCode() const
 	{
diff --git a/Engine/Plugins/Importers/USDImporter/Source/UnrealUSDWrapper/Private/UsdWrappers/UsdStage.cpp b/Engine/Plugins/Importers/USDImporter/Source/UnrealUSDWrapper/Private/UsdWrappers/UsdStage.cpp
--- a/Engine/Plugins/Importers/USDImporter/Source/UnrealUSDWrapper/Private/UsdWrappers/UsdStage.cpp
+++ b/Engine/Plugins/Importers/USDImporter/Source/UnrealUSDWrapper/Private/UsdWrappers/UsdStage.cpp
@@ -212,6 +212,27 @@ namespace UE
 		return FSdfLayer();
 	}
 
+	template<typename PtrType>
+	bool FUsdStageBase<PtrType>::Export( const TCHAR* FileName, bool bAddSourceFileComment, const TMap<FString, FString>& FileFormatArguments ) const
+	{
+#if USE_USD_SDK
+		if ( PtrType& Ptr = Impl->GetInner() )
+		{
+			FScopedUsdAllocs Allocs;
+
+			std::map<std::string, std::string> UsdFileFormatArguments;
+			for ( const TPair<FString, FString>& Pair : FileFormatArguments )
+			{
+				UsdFileFormatArguments.insert( { TCHAR_TO_ANSI( *Pair.Key ), TCHAR_TO_ANSI( *Pair.Value ) } );
+			}
+
+			return Ptr->Export( TCHAR_TO_ANSI( FileName ), bAddSourceFileComment, UsdFileFormatArguments );
+		}
+#endif // #if USE_USD_SDK
+
+		return false;
+	}
+
 	template<typename PtrType>
 	FSdfLayer FUsdStageBase<PtrType>::GetSessionLayer() const
 	{
@@ -306,6 +327,50 @@ namespace UE
 #endif // #if USE_USD_SDK
 	}
 
+	template<typename PtrType>
+	TArray<UE::FSdfLayer> FUsdStageBase<PtrType>::GetLayerStack( bool bIncludeSessionLayers /*= true */ ) const
+	{
+		TArray<UE::FSdfLayer> Result;
+#if USE_USD_SDK
+		if ( PtrType& Ptr = Impl->GetInner() )
+		{
+			FScopedUsdAllocs UsdAllocs;
+
+			pxr::SdfLayerHandleVector LayerHandleVector = Ptr->GetLayerStack( bIncludeSessionLayers );
+
+			Result.Reserve( LayerHandleVector.size() );
+
+			for ( pxr::SdfLayerHandle Layer : LayerHandleVector )
+			{
+				Result.Add( UE::FSdfLayer{ Layer } );
+			}
+		}
+#endif // #if USE_USD_SDK
+		return Result;
+	}
+
+	template<typename PtrType>
+	TArray<UE::FSdfLayer> FUsdStageBase<PtrType>::GetUsedLayers( bool bIncludeClipLayers /*= true */ ) const
+	{
+		TArray<UE::FSdfLayer> Result;
+#if USE_USD_SDK
+		if ( PtrType& Ptr = Impl->GetInner() )
+		{
+			FScopedUsdAllocs UsdAllocs;
+
+			pxr::SdfLayerHandleVector LayerHandleVector = Ptr->GetUsedLayers( bIncludeClipLayers );
+
+			Result.Reserve( LayerHandleVector.size() );
+
+			for ( pxr::SdfLayerHandle Layer : LayerHandleVector )
+			{
+				Result.Add( UE::FSdfLayer{ Layer } );
+			}
+		}
+#endif // #if USE_USD_SDK
+		return Result;
+	}
+
 	template<typename PtrType>
 	FSdfLayer FUsdStageBase<PtrType>::GetEditTarget() const
 	{
diff --git a/Engine/Plugins/Importers/USDImporter/Source/UnrealUSDWrapper/Public/UsdWrappers/SdfLayer.h b/Engine/Plugins/Importers/USDImporter/Source/UnrealUSDWrapper/Public/UsdWrappers/SdfLayer.h
--- a/Engine/Plugins/Importers/USDImporter/Source/UnrealUSDWrapper/Public/UsdWrappers/SdfLayer.h
+++ b/Engine/Plugins/Importers/USDImporter/Source/UnrealUSDWrapper/Public/UsdWrappers/SdfLayer.h
@@ -80,14 +80,22 @@ namespace UE
 
 	// Wrapped pxr::SdfLayer functions, refer to the USD SDK documentation
 	public:
+		void TransferContent( const FSdfLayer& SourceLayer );
+
 		static FSdfLayer FindOrOpen( const TCHAR* Identifier );
+		static FSdfLayer CreateNew( const TCHAR* Identifier );
 
 		bool Save( bool bForce = false ) const;
 
+		TSet<FString> GetExternalReferences() const;
+		bool UpdateExternalReference( const TCHAR* OldReferencePath, const TCHAR* NewReferencePath );
+
 		FString GetRealPath() const;
 		FString GetIdentifier() const;
 		FString GetDisplayName() const;
 
+		FString ComputeAbsolutePath(const FString& AssetPath) const;
+
 		bool IsEmpty() const;
 		bool IsAnonymous() const;
 
diff --git a/Engine/Plugins/Importers/USDImporter/Source/UnrealUSDWrapper/Public/UsdWrappers/UsdStage.h b/Engine/Plugins/Importers/USDImporter/Source/UnrealUSDWrapper/Public/UsdWrappers/UsdStage.h
--- a/Engine/Plugins/Importers/USDImporter/Source/UnrealUSDWrapper/Public/UsdWrappers/UsdStage.h
+++ b/Engine/Plugins/Importers/USDImporter/Source/UnrealUSDWrapper/Public/UsdWrappers/UsdStage.h
@@ -63,6 +63,12 @@ namespace UE
 
 	// Wrapped pxr::UsdStage functions, refer to the USD SDK documentation
 	public:
+		/**
+		 * Saves a flattened copy of the stage to the given path (e.g. "C:/Folder/FlattenedStage.usda"). Will use the corresponding file writer depending on FilePath extension.
+		 * Will not alter the current stage.
+		 */
+		bool Export( const TCHAR* FileName, bool bAddSourceFileComment = true, const TMap<FString, FString>& FileFormatArguments = {} ) const;
+
 		FSdfLayer GetRootLayer() const;
 		FSdfLayer GetSessionLayer() const;
 		bool HasLocalLayer( const FSdfLayer& Layer ) const;
@@ -71,6 +77,9 @@ namespace UE
 		FUsdPrim GetDefaultPrim() const;
 		FUsdPrim GetPrimAtPath( const FSdfPath& Path ) const;
 
+		TArray<FSdfLayer> GetLayerStack( bool bIncludeSessionLayers = true ) const;
+		TArray<FSdfLayer> GetUsedLayers( bool bIncludeClipLayers = true ) const;
+
 		bool IsEditTargetValid() const;
 		void SetEditTarget( const FSdfLayer& Layer );
 		FSdfLayer GetEditTarget() const;
