USDImporter: updates to adapt to UsdLuxLightAPI and Ar 2.0 changes in USD 21.11 and later

This is prep work in advance of upgrading the engine's version of USD
from 21.08 to the yet-to-be-released 22.03.

In USD 21.11, UsdLuxLight was removed and replaced with UsdLuxLightAPI,
which is an API schema instead of a typed one.

USD 21.11 also saw the default version of Ar move to 2.0 instead of 1.0.

The changes here should allow the USDImporter plugin for UE to be compiled with
versions both before and after the UsdLux transition, and with either version
of Ar active. However, assets or tests that make use of lights may still need
updating. Those updates will be made in follow-up changes.

diff --git a/Engine/Plugins/Importers/USDImporter/Source/USDSchemas/Private/USDLuxLightTranslator.cpp b/Engine/Plugins/Importers/USDImporter/Source/USDSchemas/Private/USDLuxLightTranslator.cpp
--- a/Engine/Plugins/Importers/USDImporter/Source/USDSchemas/Private/USDLuxLightTranslator.cpp
+++ b/Engine/Plugins/Importers/USDImporter/Source/USDSchemas/Private/USDLuxLightTranslator.cpp
@@ -22,6 +22,11 @@
 	#include "pxr/usd/usdLux/diskLight.h"
 	#include "pxr/usd/usdLux/distantLight.h"
 	#include "pxr/usd/usdLux/domeLight.h"
+#if defined(HAS_USDLUX_LIGHTAPI)
+	#include "pxr/usd/usdLux/lightAPI.h"
+#else
+	#include "pxr/usd/usdLux/light.h"
+#endif // #if defined(HAS_USDLUX_LIGHTAPI)
 	#include "pxr/usd/usdLux/rectLight.h"
 	#include "pxr/usd/usdLux/shapingAPI.h"
 	#include "pxr/usd/usdLux/sphereLight.h"
@@ -47,9 +52,14 @@ void FUsdLuxLightTranslator::UpdateComponents( USceneComponent* SceneComponent )
 	FScopedUsdAllocs UsdAllocs;
 
 	pxr::UsdPrim Prim = GetPrim();
-	pxr::UsdLuxLight UsdLight{ Prim };
 
-	if ( !UsdLight )
+#if defined(HAS_USDLUX_LIGHTAPI)
+	const pxr::UsdLuxLightAPI LightAPI( Prim );
+#else
+	const pxr::UsdLuxLight LightAPI( Prim );
+#endif // #if defined(HAS_USDLUX_LIGHTAPI)
+
+	if ( !LightAPI )
 	{
 		return;
 	}
diff --git a/Engine/Plugins/Importers/USDImporter/Source/USDStageEditor/Private/SUSDLayersTreeView.cpp b/Engine/Plugins/Importers/USDImporter/Source/USDStageEditor/Private/SUSDLayersTreeView.cpp
--- a/Engine/Plugins/Importers/USDImporter/Source/USDStageEditor/Private/SUSDLayersTreeView.cpp
+++ b/Engine/Plugins/Importers/USDImporter/Source/USDStageEditor/Private/SUSDLayersTreeView.cpp
@@ -47,11 +47,20 @@ namespace UE::USDLayersTreeViewImpl::Private
 		// Convert layer references to absolute paths so that it still works at its target location
 		FString LayerPath = LayerToExport.GetRealPath();
 		FPaths::NormalizeFilename( LayerPath );
-		TSet<FString> ExternalReferences = OutputLayer.GetExternalReferences();
-		for ( const FString& Ref : ExternalReferences )
+		const TSet<FString> AssetDependencies =
+#if defined(PXR_VERSION) && PXR_VERSION >= 2111
+			OutputLayer.GetCompositionAssetDependencies();
+#else
+			OutputLayer.GetExternalReferences();
+#endif
+		for ( const FString& Ref : AssetDependencies )
 		{
 			FString AbsRef = FPaths::ConvertRelativePathToFull( FPaths::GetPath( LayerPath ), Ref ); // Relative to the original file
+#if defined(PXR_VERSION) && PXR_VERSION >= 2111
+			OutputLayer.UpdateCompositionAssetDependency( *Ref, *AbsRef );
+#else
 			OutputLayer.UpdateExternalReference( *Ref, *AbsRef );
+#endif
 		}
 
 		bool bForce = true;
diff --git a/Engine/Plugins/Importers/USDImporter/Source/USDStageEditor/Private/SUSDStage.cpp b/Engine/Plugins/Importers/USDImporter/Source/USDStageEditor/Private/SUSDStage.cpp
--- a/Engine/Plugins/Importers/USDImporter/Source/USDStageEditor/Private/SUSDStage.cpp
+++ b/Engine/Plugins/Importers/USDImporter/Source/USDStageEditor/Private/SUSDStage.cpp
@@ -94,9 +94,13 @@ namespace SUSDStageImpl
 			UE::FSdfLayer& Layer = AllLayers[ Index ];
 
 			TArray< UE::FSdfLayer > NewLayers;
-			for ( const FString& Reference : Layer.GetExternalReferences() )
+#if defined(PXR_VERSION) && PXR_VERSION >= 2111
+			for ( const FString& AssetDependency : Layer.GetCompositionAssetDependencies() )
+#else
+			for ( const FString& AssetDependency : Layer.GetExternalReferences() )
+#endif
 			{
-				FString AbsoluteReference = Layer.ComputeAbsolutePath( Reference );
+				FString AbsoluteReference = Layer.ComputeAbsolutePath( AssetDependency );
 				FPaths::NormalizeFilename( AbsoluteReference );
 
 				if ( !LoadedLayers.Contains( AbsoluteReference ) )
@@ -1392,14 +1396,22 @@ void SUsdStage::OnExportAll()
 
 			// Remap references to layers so that they point at the other newly saved files. Note that SUSDStageImpl::AppendAllExternalReferences
 			// will have collected all external references already, so OldPathToSavedPath should have entries for all references we'll find
-			for ( const FString& Ref : OutputLayer.GetExternalReferences() )
+#if defined(PXR_VERSION) && PXR_VERSION >= 2111
+			for ( const FString& AssetDependency : OutputLayer.GetCompositionAssetDependencies() )
+#else
+			for ( const FString& AssetDependency : OutputLayer.GetExternalReferences() )
+#endif
 			{
-				FString AbsRef = FPaths::ConvertRelativePathToFull( FPaths::GetPath( LayerPath ), Ref ); // Relative to the original file
+				FString AbsRef = FPaths::ConvertRelativePathToFull( FPaths::GetPath( LayerPath ), AssetDependency ); // Relative to the original file
 				FPaths::NormalizeFilename( AbsRef );
 
 				if ( FString* SavedReference = OldPathToSavedPath.Find( AbsRef ) )
 				{
-					OutputLayer.UpdateExternalReference( *Ref, **SavedReference );
+#if defined(PXR_VERSION) && PXR_VERSION >= 2111
+					OutputLayer.UpdateCompositionAssetDependency( *AssetDependency, **SavedReference );
+#else
+					OutputLayer.UpdateExternalReference( *AssetDependency, **SavedReference );
+#endif
 				}
 			}
 
diff --git a/Engine/Plugins/Importers/USDImporter/Source/USDUtilities/Private/USDLayerUtils.cpp b/Engine/Plugins/Importers/USDImporter/Source/USDUtilities/Private/USDLayerUtils.cpp
--- a/Engine/Plugins/Importers/USDImporter/Source/USDUtilities/Private/USDLayerUtils.cpp
+++ b/Engine/Plugins/Importers/USDImporter/Source/USDUtilities/Private/USDLayerUtils.cpp
@@ -43,7 +43,7 @@ namespace UE::UsdLayerUtilsImpl::Private
 {
 	/**
 		* Adapted from flattenUtils.cpp::_FixAssetPaths, except that we only handle actual AssetPaths here as layer/prim paths
-		* will be remapped via Layer.UpdateExternalReference.
+		* will be remapped via Layer.UpdateCompositionAssetDependency().
 		* Returns whether anything was remapped
 		*/
 	bool FixAssetPaths( const pxr::SdfLayerHandle& SourceLayer, pxr::VtValue* Value )
diff --git a/Engine/Plugins/Importers/USDImporter/Source/USDUtilities/Private/USDLightConversion.cpp b/Engine/Plugins/Importers/USDImporter/Source/USDUtilities/Private/USDLightConversion.cpp
--- a/Engine/Plugins/Importers/USDImporter/Source/USDUtilities/Private/USDLightConversion.cpp
+++ b/Engine/Plugins/Importers/USDImporter/Source/USDUtilities/Private/USDLightConversion.cpp
@@ -29,6 +29,11 @@
 	#include "pxr/usd/usdLux/diskLight.h"
 	#include "pxr/usd/usdLux/distantLight.h"
 	#include "pxr/usd/usdLux/domeLight.h"
+#if defined(HAS_USDLUX_LIGHTAPI)
+	#include "pxr/usd/usdLux/lightAPI.h"
+#else
+	#include "pxr/usd/usdLux/light.h"
+#endif // #if defined(HAS_USDLUX_LIGHTAPI)
 	#include "pxr/usd/usdLux/rectLight.h"
 	#include "pxr/usd/usdLux/shapingAPI.h"
 	#include "pxr/usd/usdLux/sphereLight.h"
@@ -83,10 +88,17 @@ namespace LightConversionImpl
 	}
 }
 
+#if defined(HAS_USDLUX_LIGHTAPI)
+bool UsdToUnreal::ConvertLight( const pxr::UsdLuxLightAPI& LightAPI, ULightComponentBase& LightComponentBase, double TimeCode )
+{
+	return UsdToUnreal::ConvertLight( LightAPI.GetPrim(), LightComponentBase, TimeCode );
+}
+#else
 bool UsdToUnreal::ConvertLight( const pxr::UsdLuxLight& Light, ULightComponentBase& LightComponentBase, double TimeCode )
 {
 	return UsdToUnreal::ConvertLight( Light.GetPrim(), LightComponentBase, TimeCode );
 }
+#endif // #if defined(HAS_USDLUX_LIGHTAPI)
 
 bool UsdToUnreal::ConvertDistantLight( const pxr::UsdLuxDistantLight& DistantLight, UDirectionalLightComponent& LightComponent, double TimeCode )
 {
@@ -122,15 +134,20 @@ bool UsdToUnreal::ConvertLight( const pxr::UsdPrim& Prim, ULightComponentBase& L
 {
 	FScopedUsdAllocs Allocs;
 
-	pxr::UsdLuxLight Light{ Prim };
-	if ( !Light )
+#if defined(HAS_USDLUX_LIGHTAPI)
+	const pxr::UsdLuxLightAPI LightAPI( Prim );
+#else
+	const pxr::UsdLuxLight LightAPI( Prim );
+#endif // #if defined(HAS_USDLUX_LIGHTAPI)
+
+	if ( !LightAPI )
 	{
 		return false;
 	}
 
-	const float UsdIntensity = UsdUtils::GetUsdValue< float >( Light.GetIntensityAttr(), UsdTimeCode );
-	const float UsdExposure = UsdUtils::GetUsdValue< float >( Light.GetExposureAttr(), UsdTimeCode );
-	const pxr::GfVec3f UsdColor = UsdUtils::GetUsdValue< pxr::GfVec3f >( Light.GetColorAttr(), UsdTimeCode );
+	const float UsdIntensity = UsdUtils::GetUsdValue< float >( LightAPI.GetIntensityAttr(), UsdTimeCode );
+	const float UsdExposure = UsdUtils::GetUsdValue< float >( LightAPI.GetExposureAttr(), UsdTimeCode );
+	const pxr::GfVec3f UsdColor = UsdUtils::GetUsdValue< pxr::GfVec3f >( LightAPI.GetColorAttr(), UsdTimeCode );
 
 	const bool bSRGB = true;
 	LightComponentBase.LightColor = UsdToUnreal::ConvertColor( UsdColor ).ToFColor( bSRGB );
@@ -138,8 +155,8 @@ bool UsdToUnreal::ConvertLight( const pxr::UsdPrim& Prim, ULightComponentBase& L
 
 	if ( ULightComponent* LightComponent = Cast< ULightComponent >( &LightComponentBase ) )
 	{
-		LightComponent->bUseTemperature = UsdUtils::GetUsdValue< bool >( Light.GetEnableColorTemperatureAttr(), UsdTimeCode );
-		LightComponent->Temperature = UsdUtils::GetUsdValue< float >( Light.GetColorTemperatureAttr(), UsdTimeCode );
+		LightComponent->bUseTemperature = UsdUtils::GetUsdValue< bool >( LightAPI.GetEnableColorTemperatureAttr(), UsdTimeCode );
+		LightComponent->Temperature = UsdUtils::GetUsdValue< float >( LightAPI.GetColorTemperatureAttr(), UsdTimeCode );
 	}
 
 	return true;
@@ -326,7 +343,6 @@ bool UsdToUnreal::ConvertLuxShapingAPI( const pxr::UsdPrim& Prim, USpotLightComp
 
 float UsdToUnreal::ConvertLightIntensityAttr( float UsdIntensity, float UsdExposure )
 {
-	// This is equivalent to UsdLuxLight::ComputeBaseEmission() without the color term
 	return UsdIntensity * FMath::Exp2( UsdExposure );
 }
 
@@ -404,38 +420,43 @@ bool UnrealToUsd::ConvertLightComponent( const ULightComponentBase& LightCompone
 		return false;
 	}
 
-	pxr::UsdLuxLight Light{ Prim };
-	if ( !Light )
+#if defined(HAS_USDLUX_LIGHTAPI)
+	pxr::UsdLuxLightAPI LightAPI( Prim );
+#else
+	pxr::UsdLuxLight LightAPI( Prim );
+#endif // #if defined(HAS_USDLUX_LIGHTAPI)
+
+	if ( !LightAPI )
 	{
 		return false;
 	}
 
-	if ( pxr::UsdAttribute Attr = Light.CreateIntensityAttr() )
+	if ( pxr::UsdAttribute Attr = LightAPI.CreateIntensityAttr() )
 	{
 		Attr.Set<float>( LightComponent.Intensity, UsdTimeCode );
 	}
 
 	// When converting into UE we multiply intensity and exposure together, so when writing back we just
 	// put everything in intensity. USD also multiplies those two together, meaning it should end up the same
-	if ( pxr::UsdAttribute Attr = Light.CreateExposureAttr() )
+	if ( pxr::UsdAttribute Attr = LightAPI.CreateExposureAttr() )
 	{
 		Attr.Set<float>( 0.0f, UsdTimeCode );
 	}
 
 	if ( const ULightComponent* DerivedLightComponent = Cast<ULightComponent>( &LightComponent ) )
 	{
-		if ( pxr::UsdAttribute Attr = Light.CreateEnableColorTemperatureAttr() )
+		if ( pxr::UsdAttribute Attr = LightAPI.CreateEnableColorTemperatureAttr() )
 		{
 			Attr.Set<bool>( DerivedLightComponent->bUseTemperature, UsdTimeCode );
 		}
 
-		if ( pxr::UsdAttribute Attr = Light.CreateColorTemperatureAttr() )
+		if ( pxr::UsdAttribute Attr = LightAPI.CreateColorTemperatureAttr() )
 		{
 			Attr.Set<float>( DerivedLightComponent->Temperature, UsdTimeCode );
 		}
 	}
 
-	if ( pxr::UsdAttribute Attr = Light.CreateColorAttr() )
+	if ( pxr::UsdAttribute Attr = LightAPI.CreateColorAttr() )
 	{
 		pxr::GfVec4f LinearColor = UnrealToUsd::ConvertColor( LightComponent.LightColor );
 		Attr.Set<pxr::GfVec3f>( pxr::GfVec3f( LinearColor[0], LinearColor[1], LinearColor[2] ), UsdTimeCode );
@@ -479,8 +500,13 @@ bool UnrealToUsd::ConvertRectLightComponent( const URectLightComponent& LightCom
 		return false;
 	}
 
-	pxr::UsdLuxLight BaseLight{ Prim };
-	if ( !BaseLight )
+#if defined(HAS_USDLUX_LIGHTAPI)
+	pxr::UsdLuxLightAPI LightAPI( Prim );
+#else
+	pxr::UsdLuxLight LightAPI( Prim );
+#endif // #if defined(HAS_USDLUX_LIGHTAPI)
+
+	if ( !LightAPI )
 	{
 		return false;
 	}
@@ -521,7 +547,7 @@ bool UnrealToUsd::ConvertRectLightComponent( const URectLightComponent& LightCom
 	}
 
 	// Common for both
-	if ( pxr::UsdAttribute Attr = BaseLight.CreateIntensityAttr() )
+	if ( pxr::UsdAttribute Attr = LightAPI.CreateIntensityAttr() )
 	{
 		float FinalIntensityNits = 1.0f;
 		float OldIntensity = UsdUtils::GetUsdValue< float >( Attr, UsdTimeCode );
diff --git a/Engine/Plugins/Importers/USDImporter/Source/USDUtilities/Private/USDPrimConversion.cpp b/Engine/Plugins/Importers/USDImporter/Source/USDUtilities/Private/USDPrimConversion.cpp
--- a/Engine/Plugins/Importers/USDImporter/Source/USDUtilities/Private/USDPrimConversion.cpp
+++ b/Engine/Plugins/Importers/USDImporter/Source/USDUtilities/Private/USDPrimConversion.cpp
@@ -70,7 +70,11 @@
 #include "pxr/usd/usdGeom/xformCommonAPI.h"
 #include "pxr/usd/usdLux/diskLight.h"
 #include "pxr/usd/usdLux/distantLight.h"
+#if defined(HAS_USDLUX_LIGHTAPI)
+#include "pxr/usd/usdLux/lightAPI.h"
+#else
 #include "pxr/usd/usdLux/light.h"
+#endif // #if defined(HAS_USDLUX_LIGHTAPI)
 #include "pxr/usd/usdLux/rectLight.h"
 #include "pxr/usd/usdLux/shapingAPI.h"
 #include "pxr/usd/usdLux/sphereLight.h"
@@ -253,7 +257,12 @@ bool UsdToUnreal::ConvertXformable( const pxr::UsdStageRefPtr& Stage, const pxr:
 	// Extra rotation to match different camera facing direction convention
 	// Note: The camera space is always Y-up, yes, but this is not what this is: This is the camera's transform wrt the stage,
 	// which follows the stage up axis
-	if ( Xformable.GetPrim().IsA< pxr::UsdGeomCamera >() || Xformable.GetPrim().IsA< pxr::UsdLuxLight >() )
+	if ( Xformable.GetPrim().IsA< pxr::UsdGeomCamera >() ||
+#if defined(HAS_USDLUX_LIGHTAPI)
+			Xformable.GetPrim().HasAPI< pxr::UsdLuxLightAPI >() )
+#else
+			Xformable.GetPrim().IsA< pxr::UsdLuxLight >() )
+#endif // #if defined(HAS_USDLUX_LIGHTAPI)
 	{
 		if ( StageInfo.UpAxis == EUsdUpAxis::YAxis )
 		{
@@ -267,7 +276,12 @@ bool UsdToUnreal::ConvertXformable( const pxr::UsdStageRefPtr& Stage, const pxr:
 	// Invert the compensation applied to our parents, in case they're a camera or a light
 	if ( pxr::UsdPrim Parent = Xformable.GetPrim().GetParent() )
 	{
-		if ( Parent.IsA< pxr::UsdGeomCamera >() || Parent.IsA< pxr::UsdLuxLight >() )
+		if ( Parent.IsA< pxr::UsdGeomCamera >() ||
+#if defined(HAS_USDLUX_LIGHTAPI)
+				Parent.HasAPI< pxr::UsdLuxLightAPI >() )
+#else
+				Parent.IsA< pxr::UsdLuxLight >() )
+#endif // #if defined(HAS_USDLUX_LIGHTAPI)
 		{
 			if ( StageInfo.UpAxis == EUsdUpAxis::YAxis )
 			{
@@ -908,11 +922,15 @@ UsdToUnreal::FPropertyTrackReader UsdToUnreal::CreatePropertyTrackReader( const
 			}
 		}
 	}
-	else if ( pxr::UsdLuxLight Light{ Prim } )
+#if defined(HAS_USDLUX_LIGHTAPI)
+	else if ( const pxr::UsdLuxLightAPI LightAPI{ Prim } )
+#else
+	else if ( const pxr::UsdLuxLight LightAPI{ Prim } )
+#endif // #if defined(HAS_USDLUX_LIGHTAPI)
 	{
 		if ( PropertyPath == UnrealIdentifiers::LightColorPropertyName )
 		{
-			if ( pxr::UsdAttribute Attr = Light.GetColorAttr() )
+			if ( pxr::UsdAttribute Attr = LightAPI.GetColorAttr() )
 			{
 				pxr::GfVec3f UsdDefault;
 				Attr.Get<pxr::GfVec3f>( &UsdDefault );
@@ -935,7 +953,7 @@ UsdToUnreal::FPropertyTrackReader UsdToUnreal::CreatePropertyTrackReader( const
 		}
 		else if ( PropertyPath == UnrealIdentifiers::UseTemperaturePropertyName )
 		{
-			if ( pxr::UsdAttribute Attr = Light.GetEnableColorTemperatureAttr() )
+			if ( pxr::UsdAttribute Attr = LightAPI.GetEnableColorTemperatureAttr() )
 			{
 				bool Default;
 				Attr.Get<bool>( &Default );
@@ -951,7 +969,7 @@ UsdToUnreal::FPropertyTrackReader UsdToUnreal::CreatePropertyTrackReader( const
 		}
 		else if ( PropertyPath == UnrealIdentifiers::TemperaturePropertyName )
 		{
-			if ( pxr::UsdAttribute Attr = Light.GetColorTemperatureAttr() )
+			if ( pxr::UsdAttribute Attr = LightAPI.GetColorTemperatureAttr() )
 			{
 				float Default;
 				Attr.Get<float>( &Default );
@@ -1928,7 +1946,12 @@ bool UnrealToUsd::ConvertSceneComponent( const pxr::UsdStageRefPtr& Stage, const
 	// In UE cameras shoot towards local + X, with + Z up.Lights also emit towards local + X, with + Z up
 	// Note that this wouldn't have worked in case we collapsed light and camera components, but these always get their own
 	// actors, so we know that we don't have a single component that represents a large collapsed prim hierarchy
-	if ( UsdPrim.IsA<pxr::UsdGeomCamera>() || UsdPrim.IsA<pxr::UsdLuxLight>() )
+	if ( UsdPrim.IsA<pxr::UsdGeomCamera>() ||
+#if defined(HAS_USDLUX_LIGHTAPI)
+			UsdPrim.HasAPI<pxr::UsdLuxLightAPI>() )
+#else
+			UsdPrim.IsA<pxr::UsdLuxLight>() )
+#endif // #if defined(HAS_USDLUX_LIGHTAPI)
 	{
 		FTransform AdditionalRotation = FTransform( FRotator( 0.0f, 90.f, 0.0f ) );
 
@@ -1943,7 +1966,12 @@ bool UnrealToUsd::ConvertSceneComponent( const pxr::UsdStageRefPtr& Stage, const
 	// Invert compensation applied to parent if it's a light or camera component
 	if ( pxr::UsdPrim ParentPrim = UsdPrim.GetParent() )
 	{
-		if ( ParentPrim.IsA<pxr::UsdGeomCamera>() || ParentPrim.IsA<pxr::UsdLuxLight>() )
+		if ( ParentPrim.IsA<pxr::UsdGeomCamera>() ||
+#if defined(HAS_USDLUX_LIGHTAPI)
+			ParentPrim.HasAPI<pxr::UsdLuxLightAPI>() )
+#else
+			ParentPrim.IsA<pxr::UsdLuxLight>() )
+#endif // #if defined(HAS_USDLUX_LIGHTAPI)
 		{
 			FTransform AdditionalRotation = FTransform( FRotator( 0.0f, 90.f, 0.0f ) );
 
@@ -2444,7 +2472,13 @@ bool UnrealToUsd::CreateComponentPropertyBaker( UE::FUsdPrim& Prim, const UScene
 
 			// Compensate different orientation for light or camera components
 			FTransform AdditionalRotation = FTransform::Identity;
-			if ( UsdPrim.IsA< pxr::UsdGeomCamera >() || UsdPrim.IsA< pxr::UsdLuxLight >() )
+			if ( UsdPrim.IsA< pxr::UsdGeomCamera >() ||
+#if defined(HAS_USDLUX_LIGHTAPI)
+					UsdPrim.HasAPI< pxr::UsdLuxLightAPI >() )
+#else
+					UsdPrim.IsA< pxr::UsdLuxLight >() )
+#endif // #if defined(HAS_USDLUX_LIGHTAPI)
+				
 			{
 				AdditionalRotation = FTransform( FRotator( 0.0f, 90.0f, 0.0f ) );
 
@@ -2832,7 +2866,12 @@ UnrealToUsd::FPropertyTrackWriter UnrealToUsd::CreatePropertyTrackWriter( const
 
 					// Compensate different orientation for light or camera components
 					FTransform AdditionalRotation = FTransform::Identity;
-					if ( UsdPrim.IsA< pxr::UsdGeomCamera >() || UsdPrim.IsA< pxr::UsdLuxLight >() )
+					if ( UsdPrim.IsA< pxr::UsdGeomCamera >() ||
+#if defined(HAS_USDLUX_LIGHTAPI)
+							UsdPrim.HasAPI< pxr::UsdLuxLightAPI >() )
+#else
+							UsdPrim.IsA< pxr::UsdLuxLight >() )
+#endif // #if defined(HAS_USDLUX_LIGHTAPI)
 					{
 						AdditionalRotation = FTransform( FRotator( 0.0f, 90.0f, 0.0f ) );
 
@@ -2944,11 +2983,15 @@ UnrealToUsd::FPropertyTrackWriter UnrealToUsd::CreatePropertyTrackWriter( const
 			}
 		}
 	}
-	else if ( pxr::UsdLuxLight Light{ Prim } )
+#if defined(HAS_USDLUX_LIGHTAPI)
+	else if ( pxr::UsdLuxLightAPI LightAPI{ Prim } )
+#else
+	else if ( pxr::UsdLuxLight LightAPI{ Prim } )
+#endif // #if defined(HAS_USDLUX_LIGHTAPI)
 	{
 		if ( PropertyPath == UnrealIdentifiers::LightColorPropertyName )
 		{
-			Attr = Light.GetColorAttr();
+			Attr = LightAPI.GetColorAttr();
 			if ( Attr )
 			{
 				Result.ColorWriter = [Attr]( const FLinearColor& UEValue, double UsdTimeCode )
@@ -2960,7 +3003,7 @@ UnrealToUsd::FPropertyTrackWriter UnrealToUsd::CreatePropertyTrackWriter( const
 		}
 		else if ( PropertyPath == UnrealIdentifiers::UseTemperaturePropertyName )
 		{
-			Attr = Light.GetEnableColorTemperatureAttr();
+			Attr = LightAPI.GetEnableColorTemperatureAttr();
 			if ( Attr )
 			{
 				Result.BoolWriter = [Attr]( bool UEValue, double UsdTimeCode )
@@ -2971,7 +3014,7 @@ UnrealToUsd::FPropertyTrackWriter UnrealToUsd::CreatePropertyTrackWriter( const
 		}
 		else if ( PropertyPath == UnrealIdentifiers::TemperaturePropertyName )
 		{
-			Attr = Light.GetColorTemperatureAttr();
+			Attr = LightAPI.GetColorTemperatureAttr();
 			if ( Attr )
 			{
 				Result.FloatWriter = [Attr]( float UEValue, double UsdTimeCode )
@@ -3417,7 +3460,12 @@ bool UnrealToUsd::ConvertXformable( const UMovieScene3DTransformTrack& MovieScen
 
 		// Compensate different orientation for light or camera components
 		FTransform CameraCompensation = FTransform::Identity;
-		if ( UsdPrim.IsA< pxr::UsdGeomCamera >() || UsdPrim.IsA< pxr::UsdLuxLight >() )
+		if ( UsdPrim.IsA< pxr::UsdGeomCamera >() ||
+#if defined(HAS_USDLUX_LIGHTAPI)
+				UsdPrim.HasAPI< pxr::UsdLuxLightAPI >() )
+#else
+				UsdPrim.IsA< pxr::UsdLuxLight >() )
+#endif // #if defined(HAS_USDLUX_LIGHTAPI)
 		{
 			CameraCompensation = FTransform( FRotator( 0.0f, 90.0f, 0.0f ) );
 
@@ -3431,7 +3479,12 @@ bool UnrealToUsd::ConvertXformable( const UMovieScene3DTransformTrack& MovieScen
 		FTransform InverseCameraCompensation = FTransform::Identity;
 		if ( pxr::UsdPrim ParentPrim = UsdPrim.GetParent() )
 		{
-			if ( ParentPrim.IsA<pxr::UsdGeomCamera>() || ParentPrim.IsA<pxr::UsdLuxLight>() )
+			if ( ParentPrim.IsA<pxr::UsdGeomCamera>() ||
+#if defined(HAS_USDLUX_LIGHTAPI)
+					ParentPrim.HasAPI<pxr::UsdLuxLightAPI>() )
+#else
+					ParentPrim.IsA<pxr::UsdLuxLight>() )
+#endif // #if defined(HAS_USDLUX_LIGHTAPI)
 			{
 				InverseCameraCompensation = FTransform( FRotator( 0.0f, 90.f, 0.0f ) );
 
diff --git a/Engine/Plugins/Importers/USDImporter/Source/USDUtilities/Private/USDShadeConversion.cpp b/Engine/Plugins/Importers/USDImporter/Source/USDUtilities/Private/USDShadeConversion.cpp
--- a/Engine/Plugins/Importers/USDImporter/Source/USDUtilities/Private/USDShadeConversion.cpp
+++ b/Engine/Plugins/Importers/USDImporter/Source/USDUtilities/Private/USDShadeConversion.cpp
@@ -44,7 +44,9 @@
 #endif // WITH_EDITOR
 
 #include "USDIncludesStart.h"
+	#include "pxr/usd/ar/ar.h"
 	#include "pxr/usd/ar/asset.h"
+	#include "pxr/usd/ar/resolvedPath.h"
 	#include "pxr/usd/ar/resolver.h"
 	#include "pxr/usd/ar/resolverScopedCache.h"
 	#include "pxr/usd/sdf/layer.h"
@@ -133,15 +135,19 @@ namespace UE
 					? UsdToUnreal::ConvertString( pxr::SdfComputeAssetPathRelativeToLayer( LayerHandle, UnrealToUsd::ConvertString( *AssetPathToResolve ).Get() ) )
 					: AssetPathToResolve;
 
-				std::string ResolvedPathUsd = Resolver.Resolve( UnrealToUsd::ConvertString( *RelativePathToResolve ).Get().c_str() );
+				std::string AssetIdentifier = Resolver.Resolve( UnrealToUsd::ConvertString( *RelativePathToResolve ).Get().c_str() );
 
 				// Don't normalize an empty path as the result will be "."
-				if ( ResolvedPathUsd.size() > 0 )
+				if ( AssetIdentifier.size() > 0 )
 				{
-					ResolvedPathUsd = Resolver.ComputeNormalizedPath( ResolvedPathUsd );
+#if AR_VERSION >= 2
+					AssetIdentifier = Resolver.CreateIdentifier( AssetIdentifier );
+#else
+					AssetIdentifier = Resolver.ComputeNormalizedPath( AssetIdentifier );
+#endif // #if AR_VERSION >= 2
 				}
 
-				FString ResolvedAssetPath = UsdToUnreal::ConvertString( ResolvedPathUsd );
+				FString ResolvedAssetPath = UsdToUnreal::ConvertString( AssetIdentifier );
 
 				return ResolvedAssetPath;
 			}
@@ -179,7 +185,12 @@ namespace UE
 			TUsdStore<std::shared_ptr<const char>> ReadTextureBufferFromUsdzArchive( const FString& ResolvedTexturePath, uint64& OutBufferSize )
 			{
 				pxr::ArResolver& Resolver = pxr::ArGetResolver();
-				std::shared_ptr<pxr::ArAsset> Asset = Resolver.OpenAsset( UnrealToUsd::ConvertString( *ResolvedTexturePath ).Get() );
+				std::shared_ptr<pxr::ArAsset> Asset = Resolver.OpenAsset(
+#if AR_VERSION >= 2
+					pxr::ArResolvedPath( UnrealToUsd::ConvertString( *ResolvedTexturePath ).Get() ) );
+#else
+					UnrealToUsd::ConvertString( *ResolvedTexturePath ).Get() );
+#endif // #if AR_VERSION >= 2
 
 				TUsdStore<std::shared_ptr<const char>> Buffer;
 
@@ -2185,14 +2196,18 @@ FString UsdUtils::GetResolvedTexturePath( const pxr::UsdAttribute& TextureAssetP
 
 	pxr::ArResolver& Resolver = pxr::ArGetResolver();
 
-	std::string UsdResolvedPath = TextureAssetPath.GetResolvedPath();
+	std::string AssetIdentifier = TextureAssetPath.GetResolvedPath();
 	// Don't normalize an empty path as the result will be "."
-	if ( UsdResolvedPath.size() > 0 )
+	if ( AssetIdentifier.size() > 0 )
 	{
-		UsdResolvedPath = Resolver.ComputeNormalizedPath( UsdResolvedPath );
+#if AR_VERSION >= 2
+		AssetIdentifier = Resolver.CreateIdentifier( AssetIdentifier );
+#else
+		AssetIdentifier = Resolver.ComputeNormalizedPath( AssetIdentifier );
+#endif // #if AR_VERSION >= 2
 	}
 
-	FString ResolvedTexturePath = UsdToUnreal::ConvertString( UsdResolvedPath );
+	FString ResolvedTexturePath = UsdToUnreal::ConvertString( AssetIdentifier );
 	FPaths::NormalizeFilename( ResolvedTexturePath );
 
 	if ( ResolvedTexturePath.IsEmpty() )
diff --git a/Engine/Plugins/Importers/USDImporter/Source/USDUtilities/Public/USDLightConversion.h b/Engine/Plugins/Importers/USDImporter/Source/USDUtilities/Public/USDLightConversion.h
--- a/Engine/Plugins/Importers/USDImporter/Source/USDUtilities/Public/USDLightConversion.h
+++ b/Engine/Plugins/Importers/USDImporter/Source/USDUtilities/Public/USDLightConversion.h
@@ -17,7 +17,11 @@ PXR_NAMESPACE_OPEN_SCOPE
 class UsdLuxDiskLight;
 class UsdLuxDistantLight;
 class UsdLuxDomeLight;
+#if defined(HAS_USDLUX_LIGHTAPI)
+class UsdLuxLightAPI;
+#else
 class UsdLuxLight;
+#endif // #if defined(HAS_USDLUX_LIGHTAPI)
 class UsdLuxRectLight;
 class UsdLuxShapingAPI;
 class UsdLuxSphereLight;
@@ -40,11 +44,11 @@ enum class ELightUnits : uint8;
 struct FUsdStageInfo;
 
 /**
- * Converts UsdLuxLight attributes to the corresponding ULightComponent.
+ * Converts UsdLux light attributes to the corresponding ULightComponent.
  *
- * Corresponding UsdLuxLight schema to Unreal component:
+ * Corresponding UsdLux light schema to Unreal component:
  *
- *	UsdLuxLight			->	ULightComponent
+ *	UsdLuxLight[API]	->	ULightComponent
  *	UsdLuxDistantLight	->	UDirectionalLightComponent
  *	UsdLuxRectLight		->	URectLightComponent
  *	UsdLuxDiskLight		->	URectLightComponent
@@ -54,8 +58,13 @@ struct FUsdStageInfo;
  */
 namespace UsdToUnreal
 {
+#if defined(HAS_USDLUX_LIGHTAPI)
+	UE_DEPRECATED( 5.0, "Prefer the overload that receives a pxr::UsdPrim" )
+	USDUTILITIES_API bool ConvertLight( const pxr::UsdLuxLightAPI& LightAPI, ULightComponentBase& LightComponentBase, double TimeCode );
+#else
 	UE_DEPRECATED( 5.0, "Prefer the overload that receives a pxr::UsdPrim" )
 	USDUTILITIES_API bool ConvertLight( const pxr::UsdLuxLight& Light, ULightComponentBase& LightComponentBase, double TimeCode );
+#endif // #if defined(HAS_USDLUX_LIGHTAPI)
 	UE_DEPRECATED( 5.0, "Prefer the overload that receives a pxr::UsdPrim" )
 	USDUTILITIES_API bool ConvertDistantLight( const pxr::UsdLuxDistantLight& DistantLight, UDirectionalLightComponent& LightComponent, double TimeCode );
 	UE_DEPRECATED( 5.0, "Prefer the overload that receives a pxr::UsdPrim" )
diff --git a/Engine/Plugins/Importers/USDImporter/Source/UnrealUSDWrapper/Private/UnrealUSDWrapper.cpp b/Engine/Plugins/Importers/USDImporter/Source/UnrealUSDWrapper/Private/UnrealUSDWrapper.cpp
--- a/Engine/Plugins/Importers/USDImporter/Source/UnrealUSDWrapper/Private/UnrealUSDWrapper.cpp
+++ b/Engine/Plugins/Importers/USDImporter/Source/UnrealUSDWrapper/Private/UnrealUSDWrapper.cpp
@@ -58,7 +58,11 @@
 #include "pxr/usd/usdGeom/modelAPI.h"
 #include "pxr/usd/usdGeom/tokens.h"
 #include "pxr/usd/usdGeom/xformCommonAPI.h"
+#if defined(HAS_USDLUX_LIGHTAPI)
+#include "pxr/usd/usdLux/lightAPI.h"
+#else
 #include "pxr/usd/usdLux/light.h"
+#endif // #if defined(HAS_USDLUX_LIGHTAPI)
 #include "pxr/usd/usdShade/materialBindingAPI.h"
 #include "pxr/usd/usdShade/tokens.h"
 #include "pxr/usd/usdUtils/stageCache.h"
@@ -180,10 +184,18 @@ bool IsHolding(const VtValue& Value)
 // Only used to make sure we link against UsdLux so that it's available to Python on Linux
 float GetLightIntensity(const pxr::UsdPrim& Prim)
 {
-	pxr::UsdLuxLight UsdLuxLight( Prim );
+	float Value = 1.0f;
+
+#if defined(HAS_USDLUX_LIGHTAPI)
+	const pxr::UsdLuxLightAPI LightAPI(Prim);
+#else
+	const pxr::UsdLuxLight LightAPI(Prim);
+#endif // #if defined(HAS_USDLUX_LIGHTAPI)
 
-	float Value;
-	UsdLuxLight.GetIntensityAttr().Get( &Value );
+	if (LightAPI)
+	{
+		LightAPI.GetIntensityAttr().Get(&Value);
+	}
 
 	return Value;
 }
diff --git a/Engine/Plugins/Importers/USDImporter/Source/UnrealUSDWrapper/Private/UsdWrappers/SdfLayer.cpp b/Engine/Plugins/Importers/USDImporter/Source/UnrealUSDWrapper/Private/UsdWrappers/SdfLayer.cpp
--- a/Engine/Plugins/Importers/USDImporter/Source/UnrealUSDWrapper/Private/UsdWrappers/SdfLayer.cpp
+++ b/Engine/Plugins/Importers/USDImporter/Source/UnrealUSDWrapper/Private/UsdWrappers/SdfLayer.cpp
@@ -306,6 +306,45 @@ namespace UE
 		return false;
 	}
 
+#if defined(PXR_VERSION) && PXR_VERSION >= 2111
+	template<typename PtrType>
+	TSet<FString> FSdfLayerBase<PtrType>::GetCompositionAssetDependencies() const
+	{
+		TSet<FString> Results;
+#if USE_USD_SDK
+		if ( const PtrType& Ptr = Impl->GetInner() )
+		{
+			FScopedUsdAllocs UsdAllocs;
+			const std::set<std::string> AssetDependencies = Ptr->GetCompositionAssetDependencies();
+
+			Results.Reserve( AssetDependencies.size() );
+			for ( const std::string& AssetDependency : AssetDependencies )
+			{
+				Results.Add( ANSI_TO_TCHAR( AssetDependency.c_str() ) );
+			}
+		}
+#endif // #if USE_USD_SDK
+		return Results;
+	}
+
+	template<typename PtrType>
+	bool FSdfLayerBase<PtrType>::UpdateCompositionAssetDependency(
+			const TCHAR* OldAssetPath,
+			const TCHAR* NewAssetPath)
+	{
+#if USE_USD_SDK
+		if ( const PtrType& Ptr = Impl->GetInner() )
+		{
+			FScopedUsdAllocs UsdAllocs;
+			return Ptr->UpdateCompositionAssetDependency(
+				TCHAR_TO_ANSI( OldAssetPath ),
+				(NewAssetPath != nullptr) ? TCHAR_TO_ANSI( NewAssetPath ) : std::string() );
+		}
+#endif // #if USE_USD_SDK
+		return false;
+	}
+#endif // #if defined(PXR_VERSION) && PXR_VERSION >= 2111
+
 	template<typename PtrType>
 	FString FSdfLayerBase<PtrType>::GetRealPath() const
 	{
diff --git a/Engine/Plugins/Importers/USDImporter/Source/UnrealUSDWrapper/Public/USDIncludesStart.h b/Engine/Plugins/Importers/USDImporter/Source/UnrealUSDWrapper/Public/USDIncludesStart.h
--- a/Engine/Plugins/Importers/USDImporter/Source/UnrealUSDWrapper/Public/USDIncludesStart.h
+++ b/Engine/Plugins/Importers/USDImporter/Source/UnrealUSDWrapper/Public/USDIncludesStart.h
@@ -63,4 +63,10 @@ THIRD_PARTY_INCLUDES_START
 
 #include "pxr/pxr.h"
 
+// In USD 21.11, the UsdLuxLight schema was removed and replaced with the
+// single-apply API schema UsdLuxLightAPI.
+#if defined(PXR_VERSION) && PXR_VERSION >= 2111
+#define HAS_USDLUX_LIGHTAPI
+#endif
+
 #endif // #if USE_USD_SDK
diff --git a/Engine/Plugins/Importers/USDImporter/Source/UnrealUSDWrapper/Public/UsdWrappers/SdfLayer.h b/Engine/Plugins/Importers/USDImporter/Source/UnrealUSDWrapper/Public/UsdWrappers/SdfLayer.h
--- a/Engine/Plugins/Importers/USDImporter/Source/UnrealUSDWrapper/Public/UsdWrappers/SdfLayer.h
+++ b/Engine/Plugins/Importers/USDImporter/Source/UnrealUSDWrapper/Public/UsdWrappers/SdfLayer.h
@@ -90,6 +90,14 @@ namespace UE
 		TSet<FString> GetExternalReferences() const;
 		bool UpdateExternalReference( const TCHAR* OldReferencePath, const TCHAR* NewReferencePath );
 
+#if defined(PXR_VERSION) && PXR_VERSION >= 2111
+		TSet<FString> GetCompositionAssetDependencies() const;
+
+		bool UpdateCompositionAssetDependency(
+			const TCHAR* OldAssetPath,
+			const TCHAR* NewAssetPath = nullptr);
+#endif // #if defined(PXR_VERSION) && PXR_VERSION >= 2111
+
 		FString GetRealPath() const;
 		FString GetIdentifier() const;
 		FString GetDisplayName() const;
