USDImporter: restrict engine asset factories for USD to "native" USD file formats only

This sets up the asset factories for USD to only consider "usd", "usda", "usdc", and
"usdz" files, even when additional file formats (such as "abc") are registered for USD
through its plugin system. This ensures that the factories for USD do not inadvertently
shadow the factories for other engine-native importers, such as the ones for Alembic
or MaterialX.

diff --git a/Engine/Plugins/Importers/USDImporter/Source/USDStageImporter/Private/USDStageAssetImportFactory.cpp b/Engine/Plugins/Importers/USDImporter/Source/USDStageImporter/Private/USDStageAssetImportFactory.cpp
--- a/Engine/Plugins/Importers/USDImporter/Source/USDStageImporter/Private/USDStageAssetImportFactory.cpp
+++ b/Engine/Plugins/Importers/USDImporter/Source/USDStageImporter/Private/USDStageAssetImportFactory.cpp
@@ -35,7 +35,7 @@ UUsdStageAssetImportFactory::UUsdStageAssetImportFactory(const FObjectInitialize
 	bEditorImport = true;
 	bText = false;
 
-	for ( const FString& Extension : UnrealUSDWrapper::GetAllSupportedFileFormats() )
+	for ( const FString& Extension : UnrealUSDWrapper::GetNativeFileFormats() )
 	{
 		Formats.Add(FString::Printf(TEXT("%s; Universal Scene Description files"), *Extension));
 	}
diff --git a/Engine/Plugins/Importers/USDImporter/Source/USDStageImporter/Private/USDStageImportFactory.cpp b/Engine/Plugins/Importers/USDImporter/Source/USDStageImporter/Private/USDStageImportFactory.cpp
--- a/Engine/Plugins/Importers/USDImporter/Source/USDStageImporter/Private/USDStageImportFactory.cpp
+++ b/Engine/Plugins/Importers/USDImporter/Source/USDStageImporter/Private/USDStageImportFactory.cpp
@@ -32,7 +32,7 @@ UUsdStageImportFactory::UUsdStageImportFactory(const FObjectInitializer& ObjectI
 	bEditorImport = true;
 	bText = false;
 
-	for ( const FString& Extension : UnrealUSDWrapper::GetAllSupportedFileFormats() )
+	for ( const FString& Extension : UnrealUSDWrapper::GetNativeFileFormats() )
 	{
 		Formats.Add( FString::Printf( TEXT( "%s; Universal Scene Description files" ), *Extension ) );
 	}
diff --git a/Engine/Plugins/Importers/USDImporter/Source/UnrealUSDWrapper/Private/UnrealUSDWrapper.cpp b/Engine/Plugins/Importers/USDImporter/Source/UnrealUSDWrapper/Private/UnrealUSDWrapper.cpp
--- a/Engine/Plugins/Importers/USDImporter/Source/UnrealUSDWrapper/Private/UnrealUSDWrapper.cpp
+++ b/Engine/Plugins/Importers/USDImporter/Source/UnrealUSDWrapper/Private/UnrealUSDWrapper.cpp
@@ -52,6 +52,9 @@
 #include "pxr/usd/usd/stage.h"
 #include "pxr/usd/usd/stageCacheContext.h"
 #include "pxr/usd/usd/usdFileFormat.h"
+#include "pxr/usd/usd/usdaFileFormat.h"
+#include "pxr/usd/usd/usdcFileFormat.h"
+#include "pxr/usd/usd/usdzFileFormat.h"
 #include "pxr/usd/usd/variantSets.h"
 #include "pxr/usd/usdGeom/mesh.h"
 #include "pxr/usd/usdGeom/metrics.h"
@@ -860,6 +863,45 @@ TArray<FString> UnrealUSDWrapper::GetAllSupportedFileFormats()
 	return Result;
 }
 
+TArray<FString> UnrealUSDWrapper::GetNativeFileFormats()
+{
+	TArray<FString> Result;
+
+#if USE_USD_SDK
+	FScopedUsdAllocs Allocs;
+
+	const pxr::TfTokenVector NativeFormatIds{
+		pxr::UsdUsdFileFormatTokens->Id,
+		pxr::UsdUsdaFileFormatTokens->Id,
+		pxr::UsdUsdcFileFormatTokens->Id,
+		pxr::UsdUsdzFileFormatTokens->Id
+	};
+
+	std::set<std::string> FileExtensions;
+
+	for ( const pxr::TfToken& FormatId : NativeFormatIds )
+	{
+		const pxr::SdfFileFormatConstPtr Format = pxr::SdfFileFormat::FindById( FormatId );
+		if ( !Format )
+		{
+			continue;
+		}
+
+		for ( const std::string& FileExtension : Format->GetFileExtensions() )
+		{
+			FileExtensions.insert( FileExtension );
+		}
+	}
+
+	for ( const std::string& FileExtension : FileExtensions )
+	{
+		Result.Emplace( ANSI_TO_TCHAR( FileExtension.c_str() ) );
+	}
+#endif // #if USE_USD_SDK
+
+	return Result;
+}
+
 UE::FUsdStage UnrealUSDWrapper::OpenStage( const TCHAR* Identifier, EUsdInitialLoadSet InitialLoadSet, bool bUseStageCache )
 {
 	if ( !Identifier || FCString::Strlen( Identifier ) == 0 )
diff --git a/Engine/Plugins/Importers/USDImporter/Source/UnrealUSDWrapper/Public/UnrealUSDWrapper.h b/Engine/Plugins/Importers/USDImporter/Source/UnrealUSDWrapper/Public/UnrealUSDWrapper.h
--- a/Engine/Plugins/Importers/USDImporter/Source/UnrealUSDWrapper/Public/UnrealUSDWrapper.h
+++ b/Engine/Plugins/Importers/USDImporter/Source/UnrealUSDWrapper/Public/UnrealUSDWrapper.h
@@ -191,9 +191,23 @@ public:
 	UNREALUSDWRAPPER_API static double GetDefaultTimeCode();
 #endif  // #if USE_USD_SDK
 
-	/** Returns the file extensions the USD SDK supports reading from (e.g. ["usd", "usda", "usdc", etc.]) */
+	/**
+	 * Returns the file extensions of all file formats supported by USD.
+	 * 
+	 * These include the extensions for the file formats built into USD as
+	 * well as those for any other file formats introduced through the USD
+	 * plugin system (e.g. "abc" for the Alembic file format of the usdAbc
+	 * plugin).
+	 */
 	UNREALUSDWRAPPER_API static TArray<FString> GetAllSupportedFileFormats();
 
+	/**
+	 * Returns the file extensions that are native to USD.
+	 * 
+	 * These are the extensions for the file formats built into USD (i.e. "usd", "usda", "usdc", and "usdz").
+	 */
+	UNREALUSDWRAPPER_API static TArray<FString> GetNativeFileFormats();
+
 	/**
 	 * Opens a file as a root layer of an USD stage, and returns that stage.
 	 * @param Identifier - Path to a file that the USD SDK can open (or the identifier of a root layer), which will become the root layer of the new stage
