Enable Insight tracing in USD plugin

diff --git a/Engine/Plugins/Importers/USDImporter/Source/USDClasses/USDClasses.Build.cs b/Engine/Plugins/Importers/USDImporter/Source/USDClasses/USDClasses.Build.cs
--- a/Engine/Plugins/Importers/USDImporter/Source/USDClasses/USDClasses.Build.cs
+++ b/Engine/Plugins/Importers/USDImporter/Source/USDClasses/USDClasses.Build.cs
@@ -29,6 +29,11 @@ namespace UnrealBuildTool.Rules
 					"Json", // To read/write plugInfo.json files from UnrealUSDWrapper.cpp
 				}
 			);
+
+			if (Target.Configuration != UnrealTargetConfiguration.Shipping)
+			{
+				PublicDefinitions.Add("UE_TRACE_ENABLED=1");
+			}
 		}
 	}
 }
diff --git a/Engine/Plugins/Importers/USDImporter/Source/USDExporter/USDExporter.Build.cs b/Engine/Plugins/Importers/USDImporter/Source/USDExporter/USDExporter.Build.cs
--- a/Engine/Plugins/Importers/USDImporter/Source/USDExporter/USDExporter.Build.cs
+++ b/Engine/Plugins/Importers/USDImporter/Source/USDExporter/USDExporter.Build.cs
@@ -55,6 +55,11 @@ namespace UnrealBuildTool.Rules
 					"USDUtilities",
 				}
 			);
+
+			if (Target.Configuration != UnrealTargetConfiguration.Shipping)
+			{
+				PublicDefinitions.Add("UE_TRACE_ENABLED=1");
+			}
 		}
 	}
 }
diff --git a/Engine/Plugins/Importers/USDImporter/Source/USDSchemas/USDSchemas.Build.cs b/Engine/Plugins/Importers/USDImporter/Source/USDSchemas/USDSchemas.Build.cs
--- a/Engine/Plugins/Importers/USDImporter/Source/USDSchemas/USDSchemas.Build.cs
+++ b/Engine/Plugins/Importers/USDImporter/Source/USDSchemas/USDSchemas.Build.cs
@@ -46,6 +46,11 @@ namespace UnrealBuildTool.Rules
 					}
 				);
 			}
+
+			if (Target.Configuration != UnrealTargetConfiguration.Shipping)
+			{
+				PublicDefinitions.Add("UE_TRACE_ENABLED=1");
+			}
 		}
 	}
 }
diff --git a/Engine/Plugins/Importers/USDImporter/Source/USDStage/USDStage.Build.cs b/Engine/Plugins/Importers/USDImporter/Source/USDStage/USDStage.Build.cs
--- a/Engine/Plugins/Importers/USDImporter/Source/USDStage/USDStage.Build.cs
+++ b/Engine/Plugins/Importers/USDImporter/Source/USDStage/USDStage.Build.cs
@@ -46,6 +46,11 @@ namespace UnrealBuildTool.Rules
 					}
 				);
 			}
+
+			if (Target.Configuration != UnrealTargetConfiguration.Shipping)
+			{
+				PublicDefinitions.Add("UE_TRACE_ENABLED=1");
+			}
 		}
 	}
 }
diff --git a/Engine/Plugins/Importers/USDImporter/Source/USDStageEditor/USDStageEditor.Build.cs b/Engine/Plugins/Importers/USDImporter/Source/USDStageEditor/USDStageEditor.Build.cs
--- a/Engine/Plugins/Importers/USDImporter/Source/USDStageEditor/USDStageEditor.Build.cs
+++ b/Engine/Plugins/Importers/USDImporter/Source/USDStageEditor/USDStageEditor.Build.cs
@@ -35,6 +35,11 @@ namespace UnrealBuildTool.Rules
 					"WorkspaceMenuStructure",
 				}
 			);
+
+			if (Target.Configuration != UnrealTargetConfiguration.Shipping)
+			{
+				PublicDefinitions.Add("UE_TRACE_ENABLED=1");
+			}
 		}
 	}
 }
diff --git a/Engine/Plugins/Importers/USDImporter/Source/USDStageImporter/USDStageImporter.Build.cs b/Engine/Plugins/Importers/USDImporter/Source/USDStageImporter/USDStageImporter.Build.cs
--- a/Engine/Plugins/Importers/USDImporter/Source/USDStageImporter/USDStageImporter.Build.cs
+++ b/Engine/Plugins/Importers/USDImporter/Source/USDStageImporter/USDStageImporter.Build.cs
@@ -36,6 +36,11 @@ namespace UnrealBuildTool.Rules
 					"USDUtilities",
 				}
 			);
+
+			if (Target.Configuration != UnrealTargetConfiguration.Shipping)
+			{
+				PublicDefinitions.Add("UE_TRACE_ENABLED=1");
+			}
 		}
 	}
 }
diff --git a/Engine/Plugins/Importers/USDImporter/Source/USDUtilities/USDUtilities.Build.cs b/Engine/Plugins/Importers/USDImporter/Source/USDUtilities/USDUtilities.Build.cs
--- a/Engine/Plugins/Importers/USDImporter/Source/USDUtilities/USDUtilities.Build.cs
+++ b/Engine/Plugins/Importers/USDImporter/Source/USDUtilities/USDUtilities.Build.cs
@@ -57,6 +57,11 @@ namespace UnrealBuildTool.Rules
 					}
 				);
 			}
+
+			if (Target.Configuration != UnrealTargetConfiguration.Shipping)
+			{
+				PublicDefinitions.Add("UE_TRACE_ENABLED=1");
+			}
 		}
 	}
 }
