USD: Work around changes on blueprint compilation events.

Fix crash when recompiling blueprint that derives from the stage actor and has no RootLayer set.

If you recompile a blueprint, there is some mechanism that will create new instances of it post-compilation and replace the old instances with the new ones. This is a bit problematic for the AUsdStageActor because due to other reasons we can't have the prim twins be subobjects of the actor itself (see the comment within AUsdStageActor::AUsdStageActor for more details). This means we must specify how to replace an AUsdStageActor with another (AUsdStageActor::OnObjectsReplaced) for this reinstantiation step, and duplicate our old prim twins over to the new instance ourselves, and also means we shouldn't reload the stage when reinstantiating the stage actor (as that would also want to spawn its own prim twins).

In order to know not to load/unload the stage during this reinstantiation process, we have this FRecompilationTracker class within USDStageActor.cpp that tracks when a blueprint is recompiled or not via a couple of events. It is used within AUsdStageActor::PostRegisterAllComponents (and PostUnregisterAllComponents) to prevent the stage actor from loading or unloading the stage during the BP reinstantiation process.

With cl 18911890, a change was made so that inside BlueprintRecompilationManager.cpp, FBlueprintCompilationManagerImpl::FlushCompilationQueueImpl, the test way at the bottom for CompilerData.ShouldCompileClassFunctions() (line 1577, search for the "Make sure no junk in bytecode" comment) now turns false, when it used to turn true (this because the JobType changed to "SkeletonOnly"). This means it will never emit the BroadcastCompiled event for the BP if it is being loaded, even though it does still emit the OnBlueprintPreCompile still. This mismatch in the events caused the FRecompilationTracker to permanently "lose sync" with the state of the BP after it was loaded, and it would mean that depending on how you loaded the stage actor-derived blueprint, deleting it didn't actually unload the stage and despawn actors.

Looking at the cl description that seems like it could have been an intentional change to do that, so I just added a little workaround on our end.

Also, there was a missing nullptr check in AUsdStageActor::OnObjectsReplaced that caused a crash if you recompiled a blueprint that derived from the stage actor and it had no RootUsdTwin yet.

diff --git a/Engine/Plugins/Importers/USDImporter/Source/USDStage/Private/USDStageActor.cpp b/Engine/Plugins/Importers/USDImporter/Source/USDStage/Private/USDStageActor.cpp
--- a/Engine/Plugins/Importers/USDImporter/Source/USDStage/Private/USDStageActor.cpp
+++ b/Engine/Plugins/Importers/USDImporter/Source/USDStage/Private/USDStageActor.cpp
@@ -603,7 +603,9 @@ struct FRecompilationTracker
 
 	static void OnCompilationStarted( UBlueprint* BP )
 	{
+		// We don't care if a BP is compiling on first load: It only matters to use if we're compiling one that already has loaded instances on the level
 		if ( !BP ||
+			 BP->bIsRegeneratingOnLoad ||
 			 !BP->GeneratedClass ||
 			 !BP->GeneratedClass->IsChildOf( AUsdStageActor::StaticClass() ) ||
 			 RecompilingBlueprints.Contains( BP ) )
@@ -1534,7 +1536,7 @@ void AUsdStageActor::OnObjectsReplaced( const TMap<UObject*, UObject*>& ObjectRe
 		// twins from the old actor, which is what the reinstantiation process should have done instead anyway. Note that only
 		// later will the components and actors being pointed to by this duplicated prim twin be moved to the PIE world, so those
 		// references would be updated correctly.
-		if ( RootUsdTwin->GetOuter() == this )
+		if ( RootUsdTwin && RootUsdTwin->GetOuter() == this )
 		{
 			NewActor->RootUsdTwin = DuplicateObject( RootUsdTwin, NewActor );
 		}
