USD:
* Temp workaround for nested point instancers. They are skipped for now since they are not properly supported.

#r bDaniel.Coelho

diff --git a/Engine/Plugins/Importers/USDImporter/Source/USDSchemas/Private/USDGeomPointInstancerTranslator.cpp b/Engine/Plugins/Importers/USDImporter/Source/USDSchemas/Private/USDGeomPointInstancerTranslator.cpp
--- a/Engine/Plugins/Importers/USDImporter/Source/USDSchemas/Private/USDGeomPointInstancerTranslator.cpp
+++ b/Engine/Plugins/Importers/USDImporter/Source/USDSchemas/Private/USDGeomPointInstancerTranslator.cpp
@@ -224,10 +224,17 @@ void FUsdGeomPointInstancerTranslator::UpdateComponents( USceneComponent* PointI
 				ParentComponentGuard2.Emplace(Context->ParentComponent, PrototypeParentComponent );
 			}
 
+			// #ueent_todo: Fix support nested point instancers. For now, just skip them.
+			TArray< TUsdStore< pxr::UsdPrim > > ChildPointInstancerPrims = UsdUtils::GetAllPrimsOfType(PrototypeUsdPrim, pxr::TfType::Find< pxr::UsdGeomPointInstancer >());
+			if ( ChildPointInstancerPrims.Num() > 0 )
+			{
+				continue;
+			}
+
 			TArray< TUsdStore< pxr::UsdPrim > > ChildGeomMeshPrims = UsdUtils::GetAllPrimsOfType( PrototypeUsdPrim, pxr::TfType::Find< pxr::UsdGeomMesh >() );
 
 			for ( const TUsdStore< pxr::UsdPrim >& PrototypeGeomMeshPrim : ChildGeomMeshPrims )
-			{
+				{
 				TUsdStore< pxr::SdfPath > PrototypeTargetPrimPath = UsdGeomPointInstancerTranslatorImpl::UnwindToNonCollapsedPrim( FUsdSchemaTranslator::ECollapsingType::Assets, PrototypeGeomMeshPrim, Context );
 
 				const FString UEPrototypeTargetPrimPath = UsdToUnreal::ConvertPath( PrototypeTargetPrimPath.Get() );
