USD: Expose maximum number of vertices to collapse as a cvar.

diff --git a/Engine/Plugins/Importers/USDImporter/Source/USDSchemas/Private/USDGeomXformableTranslator.cpp b/Engine/Plugins/Importers/USDImporter/Source/USDSchemas/Private/USDGeomXformableTranslator.cpp
--- a/Engine/Plugins/Importers/USDImporter/Source/USDSchemas/Private/USDGeomXformableTranslator.cpp
+++ b/Engine/Plugins/Importers/USDImporter/Source/USDSchemas/Private/USDGeomXformableTranslator.cpp
@@ -40,6 +40,12 @@
 	#include "pxr/usd/usdShade/materialBindingAPI.h"
 #include "USDIncludesEnd.h"
 
+static int32 GMaxNumVerticesCollapsedMesh = 5000000;
+static FAutoConsoleVariableRef CVarMaxNumVerticesCollapsedMesh(
+	TEXT( "USD.MaxNumVerticesCollapsedMesh" ),
+	GMaxNumVerticesCollapsedMesh,
+	TEXT( "Maximum number of vertices that a combined Mesh can have for us to collapse it into a single StaticMesh" ) );
+
 class FUsdGeomXformableCreateAssetsTaskChain : public FBuildStaticMeshTaskChain
 {
 public:
@@ -484,7 +490,6 @@ bool FUsdGeomXformableTranslator::CollapsesChildren( ECollapsingType CollapsingT
 		}
 		else
 		{
-			const int32 MaxVertices = 500000;
 			int32 NumMaxExpectedMaterialSlots = 0;
 			int32 NumVertices = 0;
 			for ( const TUsdStore< pxr::UsdPrim >& ChildPrim : ChildGeomMeshes )
@@ -498,7 +503,7 @@ bool FUsdGeomXformableTranslator::CollapsesChildren( ECollapsingType CollapsingT
 
 					NumVertices += PointsArray.size();
 
-					if ( NumVertices > MaxVertices )
+					if ( NumVertices > GMaxNumVerticesCollapsedMesh )
 					{
 						bCollapsesChildren = false;
 						break;
