USD: Allow controlling whether we combine identical materials on static meshes or not.

diff --git a/Engine/Plugins/Importers/USDImporter/Source/USDSchemas/Private/USDGeomMeshTranslator.cpp b/Engine/Plugins/Importers/USDImporter/Source/USDSchemas/Private/USDGeomMeshTranslator.cpp
--- a/Engine/Plugins/Importers/USDImporter/Source/USDSchemas/Private/USDGeomMeshTranslator.cpp
+++ b/Engine/Plugins/Importers/USDImporter/Source/USDSchemas/Private/USDGeomMeshTranslator.cpp
@@ -451,7 +451,7 @@ namespace UsdGeomMeshTranslatorImpl
 	// Note that these other LODs will be hidden in other variants, and won't show up on traversal unless we actively switch the variants (which we do here).
 	// We use a separate function for this because there is a very specific set of conditions where we successfully can do this, and we
 	// want to fall back to just parsing UsdMesh as a simple single-LOD mesh if we fail.
-	bool TryLoadingMultipleLODs( const pxr::UsdTyped& UsdMesh, TArray<FMeshDescription>& OutLODIndexToMeshDescription, TArray<UsdUtils::FUsdPrimMaterialAssignmentInfo>& OutLODIndexToMaterialInfo, const TMap< FString, TMap< FString, int32 > >& InMaterialToPrimvarToUVIndex, const pxr::UsdTimeCode InTimeCode, const pxr::TfToken& RenderContext )
+	bool TryLoadingMultipleLODs( const pxr::UsdTyped& UsdMesh, TArray<FMeshDescription>& OutLODIndexToMeshDescription, TArray<UsdUtils::FUsdPrimMaterialAssignmentInfo>& OutLODIndexToMaterialInfo, const TMap< FString, TMap< FString, int32 > >& InMaterialToPrimvarToUVIndex, const pxr::UsdTimeCode InTimeCode, const pxr::TfToken& RenderContext, bool bCombineIdenticalMaterialSlots )
 	{
 		FScopedUsdAllocs Allocs;
 
@@ -478,7 +478,7 @@ namespace UsdGeomMeshTranslatorImpl
 			FStaticMeshAttributes StaticMeshAttributes( TempMeshDescription );
 			StaticMeshAttributes.Register();
 
-			bool bSuccess = UsdToUnreal::ConvertGeomMesh( LODMesh, TempMeshDescription, TempMaterialInfo, FTransform::Identity, InMaterialToPrimvarToUVIndex, InTimeCode, RenderContext );
+			bool bSuccess = UsdToUnreal::ConvertGeomMesh( LODMesh, TempMeshDescription, TempMaterialInfo, FTransform::Identity, InMaterialToPrimvarToUVIndex, InTimeCode, RenderContext, bCombineIdenticalMaterialSlots );
 			if ( bSuccess )
 			{
 				LODIndexToMeshDescriptionMap.Add( LODIndex, MoveTemp( TempMeshDescription ) );
@@ -504,7 +504,7 @@ namespace UsdGeomMeshTranslatorImpl
 	}
 
 	void LoadMeshDescriptions( const pxr::UsdTyped& UsdMesh, TArray<FMeshDescription>& OutLODIndexToMeshDescription, TArray<UsdUtils::FUsdPrimMaterialAssignmentInfo>& OutLODIndexToMaterialInfo,
-		const TMap< FString, TMap< FString, int32 > >& MaterialToPrimvarToUVIndex, const pxr::UsdTimeCode TimeCode, bool bInterpretLODs, const FName& RenderContext )
+		const TMap< FString, TMap< FString, int32 > >& MaterialToPrimvarToUVIndex, const pxr::UsdTimeCode TimeCode, bool bInterpretLODs, const FName& RenderContext, bool bCombineIdenticalMaterialSlots = true )
 	{
 		if ( !UsdMesh )
 		{
@@ -520,7 +520,7 @@ namespace UsdGeomMeshTranslatorImpl
 		bool bInterpretedLODs = false;
 		if ( bInterpretLODs )
 		{
-			bInterpretedLODs = TryLoadingMultipleLODs( UsdMesh, OutLODIndexToMeshDescription, OutLODIndexToMaterialInfo, MaterialToPrimvarToUVIndex, TimeCode, RenderContextToken );
+			bInterpretedLODs = TryLoadingMultipleLODs( UsdMesh, OutLODIndexToMeshDescription, OutLODIndexToMaterialInfo, MaterialToPrimvarToUVIndex, TimeCode, RenderContextToken, bCombineIdenticalMaterialSlots );
 		}
 
 		if ( !bInterpretedLODs )
@@ -531,7 +531,7 @@ namespace UsdGeomMeshTranslatorImpl
 			FStaticMeshAttributes StaticMeshAttributes( TempMeshDescription );
 			StaticMeshAttributes.Register();
 
-			bool bSuccess = UsdToUnreal::ConvertGeomMesh( UsdMesh, TempMeshDescription, TempMaterialInfo, FTransform::Identity, MaterialToPrimvarToUVIndex, TimeCode, RenderContextToken );
+			bool bSuccess = UsdToUnreal::ConvertGeomMesh( UsdMesh, TempMeshDescription, TempMaterialInfo, FTransform::Identity, MaterialToPrimvarToUVIndex, TimeCode, RenderContextToken, bCombineIdenticalMaterialSlots );
 
 			if ( bSuccess )
 			{
@@ -1186,7 +1186,8 @@ void FGeomMeshCreateAssetsTaskChain::SetupTasks()
 				*MaterialToPrimvarToUVIndex,
 				pxr::UsdTimeCode( Context->Time ),
 				Context->bAllowInterpretingLODs,
-				Context->RenderContext
+				Context->RenderContext,
+				Context->bMergeIdenticalMaterialSlots
 			);
 
 			// If we have at least one valid LOD, we should keep going
diff --git a/Engine/Plugins/Importers/USDImporter/Source/USDSchemas/Private/USDGeomXformableTranslator.cpp b/Engine/Plugins/Importers/USDImporter/Source/USDSchemas/Private/USDGeomXformableTranslator.cpp
--- a/Engine/Plugins/Importers/USDImporter/Source/USDSchemas/Private/USDGeomXformableTranslator.cpp
+++ b/Engine/Plugins/Importers/USDImporter/Source/USDSchemas/Private/USDGeomXformableTranslator.cpp
@@ -90,7 +90,8 @@ void FUsdGeomXformableCreateAssetsTaskChain::SetupTasks()
 				RenderContextToken,
 				*MaterialToPrimvarToUVIndex,
 				AddedMeshDescription,
-				AssignmentInfo
+				AssignmentInfo,
+				Context->bMergeIdenticalMaterialSlots
 			);
 
 			return !AddedMeshDescription.IsEmpty();
@@ -490,8 +491,17 @@ bool FUsdGeomXformableTranslator::CollapsesChildren( ECollapsingType CollapsingT
 		}
 		else
 		{
-			int32 NumMaxExpectedMaterialSlots = 0;
+			pxr::TfToken RenderContextToken = pxr::UsdShadeTokens->universalRenderContext;
+			if ( !Context->RenderContext.IsNone() )
+			{
+				RenderContextToken = UnrealToUsd::ConvertToken( *Context->RenderContext.ToString() ).Get();
+			}
+
 			int32 NumVertices = 0;
+
+			TSet<UsdUtils::FUsdPrimMaterialSlot> CombinedSlots;
+			int32 NumMaxExpectedMaterialSlots = 0;
+
 			for ( const TUsdStore< pxr::UsdPrim >& ChildPrim : ChildGeomMeshes )
 			{
 				pxr::UsdGeomMesh ChildGeomMesh( ChildPrim.Get() );
@@ -524,8 +534,18 @@ bool FUsdGeomXformableTranslator::CollapsesChildren( ECollapsingType CollapsingT
 				// Don't collapse children if the child meshes have Nanite override opinions but the combined mesh would lead to over 64 material slots.
 				if ( bChildrenWantNanite )
 				{
-					std::vector<pxr::UsdGeomSubset> GeomSubsets = pxr::UsdShadeMaterialBindingAPI( ChildPrim.Get() ).GetMaterialBindSubsets();
-					NumMaxExpectedMaterialSlots += FMath::Max<int32>( 1, GeomSubsets.size() + 1 ); // +1 because we may create an additional slot if it's not properly partitioned
+					if ( Context->bMergeIdenticalMaterialSlots )
+					{
+						const bool bProvideMaterialIndices = false;
+						UsdUtils::FUsdPrimMaterialAssignmentInfo LocalInfo = UsdUtils::GetPrimMaterialAssignments( ChildPrim.Get(), Context->Time, bProvideMaterialIndices, RenderContextToken );
+						CombinedSlots.Append( LocalInfo.Slots );
+						NumMaxExpectedMaterialSlots = CombinedSlots.Num();
+					}
+					else
+					{
+						std::vector<pxr::UsdGeomSubset> GeomSubsets = pxr::UsdShadeMaterialBindingAPI( ChildPrim.Get() ).GetMaterialBindSubsets();
+						NumMaxExpectedMaterialSlots += FMath::Max<int32>( 1, GeomSubsets.size() + 1 ); // +1 because we may create an additional slot if it's not properly partitioned
+					}
 
 					const int32 MaxNumSections = 64; // There is no define for this, but it's checked for on NaniteBuilder.cpp, FBuilderModule::Build
 					if ( NumMaxExpectedMaterialSlots > MaxNumSections )
diff --git a/Engine/Plugins/Importers/USDImporter/Source/USDSchemas/Public/USDSchemaTranslator.h b/Engine/Plugins/Importers/USDImporter/Source/USDSchemas/Public/USDSchemaTranslator.h
--- a/Engine/Plugins/Importers/USDImporter/Source/USDSchemas/Public/USDSchemaTranslator.h
+++ b/Engine/Plugins/Importers/USDImporter/Source/USDSchemas/Public/USDSchemaTranslator.h
@@ -169,6 +169,12 @@ struct USDSCHEMAS_API FUsdSchemaTranslationContext : public TSharedFromThis< FUs
 	 */
 	EUsdDefaultKind KindsToCollapse = EUsdDefaultKind::Component | EUsdDefaultKind::Subcomponent;
 
+	/**
+	 * If enabled, when multiple mesh prims are collapsed into a single static mesh, identical material slots are merged into one slot.
+	 * Otherwise, material slots are simply appended to the list.
+	 */
+	bool bMergeIdenticalMaterialSlots = true;
+
 	/**
 	 * If true, prims with a "LOD" variant set, and "LOD0", "LOD1", etc. variants containing each
 	 * a prim can be parsed into a single UStaticMesh asset with multiple LODs
diff --git a/Engine/Plugins/Importers/USDImporter/Source/USDStage/Private/USDStageActor.cpp b/Engine/Plugins/Importers/USDImporter/Source/USDStage/Private/USDStageActor.cpp
--- a/Engine/Plugins/Importers/USDImporter/Source/USDStage/Private/USDStageActor.cpp
+++ b/Engine/Plugins/Importers/USDImporter/Source/USDStage/Private/USDStageActor.cpp
@@ -103,6 +103,7 @@ struct FUsdStageActorImpl
 		TranslationContext->bAllowParsingSkeletalAnimations = true;
 
 		TranslationContext->KindsToCollapse = (EUsdDefaultKind) StageActor->KindsToCollapse;
+		TranslationContext->bMergeIdenticalMaterialSlots = StageActor->bMergeIdenticalMaterialSlots;
 
 		UE::FSdfPath UsdPrimPath( *PrimPath );
 		UUsdPrimTwin* ParentUsdPrimTwin = StageActor->GetRootPrimTwin()->Find( UsdPrimPath.GetParentPath().GetString() );
@@ -560,6 +561,7 @@ struct FUsdStageActorImpl
 			EventAttributes.Emplace( TEXT( "InitialLoadSet" ), LexToString( (uint8)StageActor->InitialLoadSet ) );
 			EventAttributes.Emplace( TEXT( "InterpolationType" ), LexToString( (uint8)StageActor->InterpolationType) );
 			EventAttributes.Emplace( TEXT( "KindsToCollapse" ), LexToString( StageActor->KindsToCollapse ) );
+			EventAttributes.Emplace( TEXT( "MergeIdenticalMaterialSlots" ), LexToString( StageActor->bMergeIdenticalMaterialSlots ) );
 			EventAttributes.Emplace( TEXT( "PurposesToLoad" ), LexToString( StageActor->PurposesToLoad ) );
 			EventAttributes.Emplace( TEXT( "NaniteTriangleThreshold" ), LexToString( StageActor->NaniteTriangleThreshold ) );
 			EventAttributes.Emplace( TEXT( "RenderContext" ), StageActor->RenderContext.ToString() );
@@ -643,6 +645,7 @@ AUsdStageActor::AUsdStageActor()
 	: InitialLoadSet( EUsdInitialLoadSet::LoadAll )
 	, InterpolationType( EUsdInterpolationType::Linear )
 	, KindsToCollapse( ( int32 ) ( EUsdDefaultKind::Component | EUsdDefaultKind::Subcomponent ) )
+	, bMergeIdenticalMaterialSlots( true )
 	, PurposesToLoad( (int32) EUsdPurpose::Proxy )
 	, NaniteTriangleThreshold( (uint64) 1000000 )
 	, Time( 0.0f )
@@ -1293,6 +1296,14 @@ void AUsdStageActor::SetKindsToCollapse( int32 NewKindsToCollapse )
 	LoadUsdStage();
 }
 
+void AUsdStageActor::SetMergeIdenticalMaterialSlots( bool bMerge )
+{
+	Modify();
+
+	bMergeIdenticalMaterialSlots = bMerge;
+	LoadUsdStage();
+}
+
 void AUsdStageActor::SetPurposesToLoad( int32 NewPurposesToLoad )
 {
 	Modify();
@@ -2449,6 +2460,10 @@ void AUsdStageActor::HandlePropertyChangedEvent( FPropertyChangedEvent& Property
 	{
 		SetKindsToCollapse( KindsToCollapse );
 	}
+	else if ( PropertyName == GET_MEMBER_NAME_CHECKED( AUsdStageActor, bMergeIdenticalMaterialSlots ) )
+	{
+		SetMergeIdenticalMaterialSlots( bMergeIdenticalMaterialSlots );
+	}
 	else if ( PropertyName == GET_MEMBER_NAME_CHECKED( AUsdStageActor, PurposesToLoad ) )
 	{
 		SetPurposesToLoad( PurposesToLoad );
diff --git a/Engine/Plugins/Importers/USDImporter/Source/USDStage/Public/USDStageActor.h b/Engine/Plugins/Importers/USDImporter/Source/USDStage/Public/USDStageActor.h
--- a/Engine/Plugins/Importers/USDImporter/Source/USDStage/Public/USDStageActor.h
+++ b/Engine/Plugins/Importers/USDImporter/Source/USDStage/Public/USDStageActor.h
@@ -54,6 +54,13 @@ public:
 	UPROPERTY(EditAnywhere, BlueprintReadOnly, Category = "USD", meta = (Bitmask, BitmaskEnum=EUsdDefaultKind))
 	int32 KindsToCollapse;
 
+	/**
+	 * If enabled, when multiple mesh prims are collapsed into a single static mesh, identical material slots are merged into one slot.
+	 * Otherwise, materials slots are simply appended to the list.
+	 */
+	UPROPERTY(EditAnywhere, BlueprintReadOnly, Category = "USD")
+	bool bMergeIdenticalMaterialSlots;
+
 	/* Only load prims with these specific purposes from the USD file */
 	UPROPERTY(EditAnywhere, BlueprintReadOnly, Category = "USD", meta = (Bitmask, BitmaskEnum=EUsdPurpose))
 	int32 PurposesToLoad;
@@ -79,6 +86,9 @@ public:
 	UFUNCTION( BlueprintCallable, Category = "USD", meta = ( CallInEditor = "true" ) )
 	USDSTAGE_API void SetKindsToCollapse( int32 NewKindsToCollapse );
 
+	UFUNCTION( BlueprintCallable, Category = "USD", meta = ( CallInEditor = "true" ) )
+	USDSTAGE_API void SetMergeIdenticalMaterialSlots( bool bMerge );
+
 	UFUNCTION(BlueprintCallable, Category = "USD", meta = (CallInEditor = "true"))
 	USDSTAGE_API void SetPurposesToLoad( int32 NewPurposesToLoad );
 
diff --git a/Engine/Plugins/Importers/USDImporter/Source/USDStageEditor/Private/SUSDStage.cpp b/Engine/Plugins/Importers/USDImporter/Source/USDStageEditor/Private/SUSDStage.cpp
--- a/Engine/Plugins/Importers/USDImporter/Source/USDStageEditor/Private/SUSDStage.cpp
+++ b/Engine/Plugins/Importers/USDImporter/Source/USDStageEditor/Private/SUSDStage.cpp
@@ -979,6 +979,41 @@ void SUsdStage::FillRenderContextSubMenu( FMenuBuilder& MenuBuilder )
 
 void SUsdStage::FillCollapsingSubMenu( FMenuBuilder& MenuBuilder )
 {
+	MenuBuilder.AddMenuEntry(
+		LOCTEXT( "MergeIdenticalMaterialSlots", "Merge identical material slots" ),
+		LOCTEXT( "MergeIdenticalMaterialSlots", "If enabled, when multiple mesh prims are collapsed into a single static mesh, identical material slots are merged into one slot.\nOtherwise, material slots are simply appended to the list." ),
+		FSlateIcon(),
+		FUIAction(
+			FExecuteAction::CreateLambda( [this]()
+			{
+				if ( AUsdStageActor* StageActor = ViewModel.UsdStageActor.Get() )
+				{
+					FScopedTransaction Transaction( FText::Format(
+						LOCTEXT( "ToggleCollapsingTransaction", "Toggle bMergeIdenticalMaterialSlots on USD stage actor '{1}'" ),
+						FText::FromString( StageActor->GetActorLabel() )
+					) );
+
+					TGuardValue<bool> MaintainSelectionGuard( bUpdatingViewportSelection, true );
+					StageActor->SetMergeIdenticalMaterialSlots( !StageActor->bMergeIdenticalMaterialSlots );
+				}
+			}),
+			FCanExecuteAction::CreateLambda( [this]()
+			{
+				return ViewModel.UsdStageActor.Get() != nullptr;
+			}),
+			FIsActionChecked::CreateLambda( [this]()
+			{
+				if ( AUsdStageActor* StageActor = ViewModel.UsdStageActor.Get() )
+				{
+					return StageActor->bMergeIdenticalMaterialSlots;
+				}
+				return false;
+			})
+		),
+		NAME_None,
+		EUserInterfaceActionType::ToggleButton
+	);
+
 	auto AddKindToCollapseEntry = [&]( const EUsdDefaultKind Kind, const FText& Text, FCanExecuteAction CanExecuteAction )
 	{
 		MenuBuilder.AddMenuEntry(
@@ -1023,6 +1058,8 @@ void SUsdStage::FillCollapsingSubMenu( FMenuBuilder& MenuBuilder )
 		);
 	};
 
+	MenuBuilder.BeginSection( "Kinds to collapse", LOCTEXT( "KindsToCollapse", "Kinds to collapse" ) );
+
 	AddKindToCollapseEntry(
 		EUsdDefaultKind::Model,
 		LOCTEXT( "ModelKind", "Model" ),
@@ -1092,6 +1129,8 @@ void SUsdStage::FillCollapsingSubMenu( FMenuBuilder& MenuBuilder )
 			return ViewModel.UsdStageActor.Get() != nullptr;
 		})
 	);
+
+	MenuBuilder.EndSection();
 }
 
 void SUsdStage::FillInterpolationTypeSubMenu(FMenuBuilder& MenuBuilder)
diff --git a/Engine/Plugins/Importers/USDImporter/Source/USDStageImporter/Private/USDStageImportOptions.cpp b/Engine/Plugins/Importers/USDImporter/Source/USDStageImporter/Private/USDStageImportOptions.cpp
--- a/Engine/Plugins/Importers/USDImporter/Source/USDStageImporter/Private/USDStageImportOptions.cpp
+++ b/Engine/Plugins/Importers/USDImporter/Source/USDStageImporter/Private/USDStageImportOptions.cpp
@@ -31,6 +31,7 @@ UUsdStageImportOptions::UUsdStageImportOptions(const FObjectInitializer& ObjectI
 
 	bPrimPathFolderStructure = false;
 	KindsToCollapse = ( int32 ) ( EUsdDefaultKind::Component | EUsdDefaultKind::Subcomponent );
+	bMergeIdenticalMaterialSlots = true;
 	bInterpretLODs = true;
 }
 
diff --git a/Engine/Plugins/Importers/USDImporter/Source/USDStageImporter/Private/USDStageImporter.cpp b/Engine/Plugins/Importers/USDImporter/Source/USDStageImporter/Private/USDStageImporter.cpp
--- a/Engine/Plugins/Importers/USDImporter/Source/USDStageImporter/Private/USDStageImporter.cpp
+++ b/Engine/Plugins/Importers/USDImporter/Source/USDStageImporter/Private/USDStageImporter.cpp
@@ -1238,6 +1238,7 @@ namespace UsdStageImporterImpl
 				EventAttributes.Emplace( TEXT( "ReplaceAssetPolicy" ), LexToString( (uint8)ImportContext.ImportOptions->ExistingAssetPolicy ) );
 				EventAttributes.Emplace( TEXT( "PrimPathFolderStructure" ), LexToString( ImportContext.ImportOptions->bPrimPathFolderStructure ) );
 				EventAttributes.Emplace( TEXT( "KindsToCollapse" ), LexToString( ImportContext.ImportOptions->KindsToCollapse ) );
+				EventAttributes.Emplace( TEXT( "MergeIdenticalMaterialSlots" ), LexToString( ImportContext.ImportOptions->bMergeIdenticalMaterialSlots ) );
 				EventAttributes.Emplace( TEXT( "InterpretLODs" ), LexToString( ImportContext.ImportOptions->bInterpretLODs ) );
 			}
 
@@ -1448,6 +1449,7 @@ void UUsdStageImporter::ImportFromFile(FUsdStageImportContext& ImportContext)
 	TranslationContext->RenderContext = ImportContext.ImportOptions->RenderContextToImport;
 	TranslationContext->ParentComponent = ImportContext.SceneActor ? ImportContext.SceneActor->GetRootComponent() : nullptr;
 	TranslationContext->KindsToCollapse = ( EUsdDefaultKind ) ImportContext.ImportOptions->KindsToCollapse;
+	TranslationContext->bMergeIdenticalMaterialSlots = ImportContext.ImportOptions->bMergeIdenticalMaterialSlots;
 	TranslationContext->bAllowInterpretingLODs = ImportContext.ImportOptions->bInterpretLODs;
 	TranslationContext->bAllowParsingSkeletalAnimations = ImportContext.ImportOptions->bImportGeometry && ImportContext.ImportOptions->bImportSkeletalAnimations;
 	TranslationContext->MaterialToPrimvarToUVIndex = &ImportContext.MaterialToPrimvarToUVIndex;
@@ -1531,6 +1533,7 @@ bool UUsdStageImporter::ReimportSingleAsset(FUsdStageImportContext& ImportContex
 	TranslationContext->PurposesToLoad = (EUsdPurpose) ImportContext.ImportOptions->PurposesToImport;
 	TranslationContext->NaniteTriangleThreshold = ImportContext.ImportOptions->NaniteTriangleThreshold;
 	TranslationContext->KindsToCollapse = ( EUsdDefaultKind ) ImportContext.ImportOptions->KindsToCollapse;
+	TranslationContext->bMergeIdenticalMaterialSlots = ImportContext.ImportOptions->bMergeIdenticalMaterialSlots;
 	TranslationContext->bAllowInterpretingLODs = ImportContext.ImportOptions->bInterpretLODs;
 	TranslationContext->bAllowParsingSkeletalAnimations = ImportContext.ImportOptions->bImportGeometry && ImportContext.ImportOptions->bImportSkeletalAnimations;
 	TranslationContext->MaterialToPrimvarToUVIndex = &ImportContext.MaterialToPrimvarToUVIndex;
diff --git a/Engine/Plugins/Importers/USDImporter/Source/USDStageImporter/Public/USDStageImportOptions.h b/Engine/Plugins/Importers/USDImporter/Source/USDStageImporter/Public/USDStageImportOptions.h
--- a/Engine/Plugins/Importers/USDImporter/Source/USDStageImporter/Public/USDStageImportOptions.h
+++ b/Engine/Plugins/Importers/USDImporter/Source/USDStageImporter/Public/USDStageImportOptions.h
@@ -117,6 +117,13 @@ public:
 	UPROPERTY( BlueprintReadWrite, config, EditAnywhere, Category = "Processing", meta = ( Bitmask, BitmaskEnum = EUsdDefaultKind ) )
 	int32 KindsToCollapse;
 
+	/**
+	 * If enabled, when multiple mesh prims are collapsed into a single static mesh, identical material slots are merged into one slot.
+	 * Otherwise, material slots are simply appended to the list.
+	 */
+	UPROPERTY(BlueprintReadWrite, config, EditAnywhere, Category = "Processing")
+	bool bMergeIdenticalMaterialSlots;
+
 	/** When true, if a prim has a "LOD" variant set with variants named "LOD0", "LOD1", etc. where each contains a UsdGeomMesh, the importer will attempt to parse the meshes as separate LODs of a single UStaticMesh. When false, only the selected variant will be parsed as LOD0 of the UStaticMesh.  */
 	UPROPERTY(BlueprintReadWrite, config, EditAnywhere, Category="Processing", meta=(DisplayName="Interpret LOD variant sets", EditCondition=bImportGeometry) )
 	bool bInterpretLODs;
diff --git a/Engine/Plugins/Importers/USDImporter/Source/USDUtilities/Private/USDGeomMeshConversion.cpp b/Engine/Plugins/Importers/USDImporter/Source/USDUtilities/Private/USDGeomMeshConversion.cpp
--- a/Engine/Plugins/Importers/USDImporter/Source/USDUtilities/Private/USDGeomMeshConversion.cpp
+++ b/Engine/Plugins/Importers/USDImporter/Source/USDUtilities/Private/USDGeomMeshConversion.cpp
@@ -510,7 +510,8 @@ namespace UE
 				const TMap< FString, TMap<FString, int32> >& MaterialToPrimvarToUVIndex,
 				FMeshDescription& OutMeshDescription,
 				UsdUtils::FUsdPrimMaterialAssignmentInfo& OutMaterialAssignments,
-				bool bIsFirstPrim
+				bool bIsFirstPrim,
+				const bool bCombineIdenticalMaterialSlots
 			)
 			{
 				// Ignore meshes from disabled purposes
@@ -566,7 +567,7 @@ namespace UE
 				bool bSuccess = true;
 				if ( pxr::UsdGeomMesh Mesh = pxr::UsdGeomMesh( Prim ) )
 				{
-					bSuccess = UsdToUnreal::ConvertGeomMesh( Mesh, OutMeshDescription, OutMaterialAssignments, ChildTransform, MaterialToPrimvarToUVIndex, TimeCode, RenderContext );
+					bSuccess = UsdToUnreal::ConvertGeomMesh( Mesh, OutMeshDescription, OutMaterialAssignments, ChildTransform, MaterialToPrimvarToUVIndex, TimeCode, RenderContext, bCombineIdenticalMaterialSlots );
 				}
 
 				for ( const pxr::UsdPrim& ChildPrim : Prim.GetFilteredChildren( pxr::UsdTraverseInstanceProxies() ) )
@@ -587,7 +588,8 @@ namespace UE
 						MaterialToPrimvarToUVIndex,
 						OutMeshDescription,
 						OutMaterialAssignments,
-						bChildIsFirstPrim
+						bChildIsFirstPrim,
+						bCombineIdenticalMaterialSlots
 					);
 				}
 
@@ -631,7 +633,7 @@ bool UsdToUnreal::ConvertGeomMesh( const pxr::UsdTyped& UsdSchema, FMeshDescript
 }
 
 bool UsdToUnreal::ConvertGeomMesh( const pxr::UsdTyped& UsdSchema, FMeshDescription& MeshDescription, UsdUtils::FUsdPrimMaterialAssignmentInfo& MaterialAssignments, const FTransform& AdditionalTransform,
-	const TMap< FString, TMap< FString, int32 > >& MaterialToPrimvarsUVSetNames, const pxr::UsdTimeCode TimeCode, const pxr::TfToken& RenderContext )
+	const TMap< FString, TMap< FString, int32 > >& MaterialToPrimvarsUVSetNames, const pxr::UsdTimeCode TimeCode, const pxr::TfToken& RenderContext, bool bCombineIdenticalMaterialSlots )
 {
 	TRACE_CPUPROFILER_EVENT_SCOPE( UsdToUnreal::ConvertGeomMesh );
 
@@ -656,11 +658,47 @@ bool UsdToUnreal::ConvertGeomMesh( const pxr::UsdTyped& UsdSchema, FMeshDescript
 	// Material assignments
 	const bool bProvideMaterialIndices = true;
 	UsdUtils::FUsdPrimMaterialAssignmentInfo LocalInfo = UsdUtils::GetPrimMaterialAssignments( UsdPrim, TimeCode, bProvideMaterialIndices, RenderContext );
-	MaterialAssignments.Slots.Append( LocalInfo.Slots ); // We always want to keep individual slots, even when collapsing
-
 	TArray< UsdUtils::FUsdPrimMaterialSlot >& LocalMaterialSlots = LocalInfo.Slots;
 	TArray< int32 >& FaceMaterialIndices = LocalInfo.MaterialIndices;
 
+	// Position 3 in this has the value 6 --> Local material slot #3 is actually the combined material slot #6
+	TArray<int32> LocalToCombinedMaterialSlotIndices;
+	LocalToCombinedMaterialSlotIndices.SetNumZeroed( LocalInfo.Slots.Num() );
+
+	if ( bCombineIdenticalMaterialSlots )
+	{
+		// Build a map of our existing slots since we can hash the entire slot, and our incoming mesh may have an arbitrary number of new slots
+		TMap< UsdUtils::FUsdPrimMaterialSlot, int32 > CombinedMaterialSlotsToIndex;
+		for ( int32 Index = 0; Index < MaterialAssignments.Slots.Num(); ++Index )
+		{
+			const UsdUtils::FUsdPrimMaterialSlot& Slot = MaterialAssignments.Slots[ Index ];
+			CombinedMaterialSlotsToIndex.Add( Slot, Index );
+		}
+
+		for ( int32 LocalIndex = 0; LocalIndex < LocalInfo.Slots.Num(); ++LocalIndex )
+		{
+			const UsdUtils::FUsdPrimMaterialSlot& LocalSlot = LocalInfo.Slots[ LocalIndex ];
+			if ( int32* ExistingCombinedIndex = CombinedMaterialSlotsToIndex.Find( LocalSlot ) )
+			{
+				LocalToCombinedMaterialSlotIndices[ LocalIndex ] = *ExistingCombinedIndex;
+			}
+			else
+			{
+				MaterialAssignments.Slots.Add( LocalSlot );
+				LocalToCombinedMaterialSlotIndices[ LocalIndex ] = MaterialAssignments.Slots.Num() - 1;
+			}
+		}
+	}
+	else
+	{
+		// Just append our new local material slots at the end of MaterialAssignments
+		MaterialAssignments.Slots.Append( LocalInfo.Slots );
+		for ( int32 LocalIndex = 0; LocalIndex < LocalInfo.Slots.Num(); ++LocalIndex )
+		{
+			LocalToCombinedMaterialSlotIndices[ LocalIndex ] = LocalIndex + MaterialIndexOffset;
+		}
+	}
+
 	const int32 VertexOffset = MeshDescription.Vertices().Num();
 	const int32 VertexInstanceOffset = MeshDescription.VertexInstances().Num();
 	const int32 PolygonOffset = MeshDescription.Polygons().Num();
@@ -701,19 +739,32 @@ bool UsdToUnreal::ConvertGeomMesh( const pxr::UsdTyped& UsdSchema, FMeshDescript
 		int32 CurrentVertexInstanceIndex = 0;
 		TPolygonGroupAttributesRef<FName> MaterialSlotNames = StaticMeshAttributes.GetPolygonGroupMaterialSlotNames();
 
+		// If we're going to share existing material slots, we'll need to share existing PolygonGroups in the mesh description,
+		// so we need to traverse it and prefill PolygonGroupMapping
+		if( bCombineIdenticalMaterialSlots )
+		{
+			for ( FPolygonGroupID PolygonGroupID : MeshDescription.PolygonGroups().GetElementIDs() )
+			{
+				// We always create our polygon groups in order with our combined material slots, so its easy to reconstruct this mapping here
+				PolygonGroupMapping.Add( PolygonGroupID.GetValue(), PolygonGroupID );
+			}
+		}
+
 		// Material slots
 		// Note that we always create these in the order they show up in LocalInfo: If we created these on-demand when parsing polygons (like before)
 		// we could run into polygons in a different order than the material slots and end up with different material assignments.
 		// We could use the StaticMesh's SectionInfoMap to unswitch things, but that's not available at runtime so we better do this here
 		for ( int32 LocalMaterialIndex = 0; LocalMaterialIndex < LocalInfo.Slots.Num(); ++LocalMaterialIndex )
 		{
-			const int32 CombinedMaterialIndex = MaterialIndexOffset + LocalMaterialIndex;
-
-			FPolygonGroupID NewPolygonGroup = MeshDescription.CreatePolygonGroup();
-			PolygonGroupMapping.Add( CombinedMaterialIndex, NewPolygonGroup );
+			const int32 CombinedMaterialIndex = LocalToCombinedMaterialSlotIndices[ LocalMaterialIndex ];
+			if ( !PolygonGroupMapping.Contains( CombinedMaterialIndex ) )
+			{
+				FPolygonGroupID NewPolygonGroup = MeshDescription.CreatePolygonGroup();
+				PolygonGroupMapping.Add( CombinedMaterialIndex, NewPolygonGroup );
 
-			// This is important for runtime, where the material slots are matched to LOD sections based on their material slot name
-			MaterialSlotNames[ NewPolygonGroup ] = *LexToString( NewPolygonGroup.GetValue() );
+				// This is important for runtime, where the material slots are matched to LOD sections based on their material slot name
+				MaterialSlotNames[ NewPolygonGroup ] = *LexToString( NewPolygonGroup.GetValue() );
+			}
 		}
 
 		bool bFlipThisGeometry = false;
@@ -961,7 +1012,7 @@ bool UsdToUnreal::ConvertGeomMesh( const pxr::UsdTyped& UsdSchema, FMeshDescript
 				}
 			}
 
-			const int32 CombinedMaterialIndex = MaterialIndexOffset + LocalMaterialIndex;
+			const int32 CombinedMaterialIndex = LocalToCombinedMaterialSlotIndices[ LocalMaterialIndex ];
 
 			if ( bFlipThisGeometry )
 			{
@@ -995,7 +1046,8 @@ bool UsdToUnreal::ConvertGeomMeshHierarchy(
 	const pxr::TfToken& RenderContext,
 	const TMap< FString, TMap<FString, int32> >& MaterialToPrimvarToUVIndex,
 	FMeshDescription& OutMeshDescription,
-	UsdUtils::FUsdPrimMaterialAssignmentInfo& OutMaterialAssignments
+	UsdUtils::FUsdPrimMaterialAssignmentInfo& OutMaterialAssignments,
+	bool bCombineIdenticalMaterialSlots
 )
 {
 	FStaticMeshAttributes StaticMeshAttributes( OutMeshDescription );
@@ -1014,7 +1066,8 @@ bool UsdToUnreal::ConvertGeomMeshHierarchy(
 		MaterialToPrimvarToUVIndex,
 		OutMeshDescription,
 		OutMaterialAssignments,
-		bIsFirstPrim
+		bIsFirstPrim,
+		bCombineIdenticalMaterialSlots
 	);
 }
 
diff --git a/Engine/Plugins/Importers/USDImporter/Source/USDUtilities/Public/USDGeomMeshConversion.h b/Engine/Plugins/Importers/USDImporter/Source/USDUtilities/Public/USDGeomMeshConversion.h
--- a/Engine/Plugins/Importers/USDImporter/Source/USDUtilities/Public/USDGeomMeshConversion.h
+++ b/Engine/Plugins/Importers/USDImporter/Source/USDUtilities/Public/USDGeomMeshConversion.h
@@ -48,11 +48,13 @@ namespace UsdToUnreal
 	 * @param AdditionalTransform - Transform (in UnrealEditor space) to bake into the vertex positions and normals of the converted mesh
 	 * @param MaterialToPrimvarsUVSetNames - Maps from a material prim path, to pairs indicating which primvar names are used as 'st' coordinates for this mesh, and which UVIndex materials will sample from (e.g. ["st0", 0], ["myUvSet2", 2], etc). This is used to pick which primvars will become UV sets.
 	 * @param TimeCode - USD timecode when the UsdGeomMesh should be sampled for mesh data and converted
+	 * @param RenderContext - The render context to use when parsing material assignments from the Mesh prim
+	 * @param bCombineIdenticalMaterialSlots - Whether to try reusing material slots (both local and the ones already in MaterialAssignments) when converting the material assignments from the mesh
 	 * @return Whether the conversion was successful or not.
 	 */
 	USDUTILITIES_API bool ConvertGeomMesh( const pxr::UsdTyped& UsdSchema, FMeshDescription& MeshDescription, UsdUtils::FUsdPrimMaterialAssignmentInfo& MaterialAssignments, const pxr::UsdTimeCode TimeCode = pxr::UsdTimeCode::EarliestTime(), const pxr::TfToken& RenderContext = pxr::UsdShadeTokens->universalRenderContext );
 	USDUTILITIES_API bool ConvertGeomMesh( const pxr::UsdTyped& UsdSchema, FMeshDescription& MeshDescription, UsdUtils::FUsdPrimMaterialAssignmentInfo& MaterialAssignments, const FTransform& AdditionalTransform, const pxr::UsdTimeCode TimeCode = pxr::UsdTimeCode::EarliestTime(), const pxr::TfToken& RenderContext = pxr::UsdShadeTokens->universalRenderContext );
-	USDUTILITIES_API bool ConvertGeomMesh( const pxr::UsdTyped& UsdSchema, FMeshDescription& MeshDescription, UsdUtils::FUsdPrimMaterialAssignmentInfo& MaterialAssignments, const FTransform& AdditionalTransform, const TMap< FString, TMap< FString, int32 > >& MaterialToPrimvarsUVSetNames, const pxr::UsdTimeCode TimeCode = pxr::UsdTimeCode::EarliestTime(), const pxr::TfToken& RenderContext = pxr::UsdShadeTokens->universalRenderContext );
+	USDUTILITIES_API bool ConvertGeomMesh( const pxr::UsdTyped& UsdSchema, FMeshDescription& MeshDescription, UsdUtils::FUsdPrimMaterialAssignmentInfo& MaterialAssignments, const FTransform& AdditionalTransform, const TMap< FString, TMap< FString, int32 > >& MaterialToPrimvarsUVSetNames, const pxr::UsdTimeCode TimeCode = pxr::UsdTimeCode::EarliestTime(), const pxr::TfToken& RenderContext = pxr::UsdShadeTokens->universalRenderContext, bool bCombineIdenticalMaterialSlots = true );
 
 	/** DEPRECATED and will not convert material information. Use the three first signatures above this one */
 	USDUTILITIES_API bool ConvertGeomMesh( const pxr::UsdTyped& UsdSchema, FMeshDescription& MeshDescription, const pxr::UsdTimeCode TimeCode = pxr::UsdTimeCode::EarliestTime() );
@@ -70,9 +72,10 @@ namespace UsdToUnreal
 	 * @param MaterialToPrimvarToUVIndex - Maps from a material prim path, to pairs indicating which primvar names are used as 'st' coordinates for this mesh, and which UVIndex materials will sample from (e.g. ["st0", 0], ["myUvSet2", 2], etc). This is used to pick which primvars will become UV sets.
 	 * @param OutMeshDescription - Output parameter that will be filled with the converted mesh data
 	 * @param OutMaterialAssignments - Output parameter that will be filled with the material data extracted from UsdSchema
+	 * @param bCombineIdenticalMaterialSlots - Whether to try reusing material slots (both local and the ones already in MaterialAssignments) when converting the material assignments from the mesh
 	 * @return Whether the conversion was successful or not.
 	 */
-	USDUTILITIES_API bool ConvertGeomMeshHierarchy( const pxr::UsdPrim& Prim, const pxr::UsdTimeCode& TimeCode, const EUsdPurpose PurposesToLoad, const pxr::TfToken& RenderContext, const TMap< FString, TMap<FString, int32> >& MaterialToPrimvarToUVIndex, FMeshDescription& OutMeshDescription, UsdUtils::FUsdPrimMaterialAssignmentInfo& OutMaterialAssignments );
+	USDUTILITIES_API bool ConvertGeomMeshHierarchy( const pxr::UsdPrim& Prim, const pxr::UsdTimeCode& TimeCode, const EUsdPurpose PurposesToLoad, const pxr::TfToken& RenderContext, const TMap< FString, TMap<FString, int32> >& MaterialToPrimvarToUVIndex, FMeshDescription& OutMeshDescription, UsdUtils::FUsdPrimMaterialAssignmentInfo& OutMaterialAssignments, bool bCombineIdenticalMaterialSlots = true );
 
 	/** Configure Material to become a vertex color/displayColor material, according to the given description */
 	USDUTILITIES_API bool ConvertDisplayColor( const UsdUtils::FDisplayColorMaterial& DisplayColorDescription, UMaterialInstanceConstant& Material );
