Fixed crash when starting a game compiled with the USDRuntime plugin enabled
  - root cause: The USDRuntime plugin requires ansi malloc (FORCE_ANSI_ALLOCATOR=1) which does not work with the poison proxy.
  - solution: Disable the poison proxy when compiling the game => UE_USE_MALLOC_FILL_BYTES=0

diff --git a/Engine/Plugins/Importers/USDImporter/Source/UnrealUSDWrapper/UnrealUSDWrapper.Build.cs b/Engine/Plugins/Importers/USDImporter/Source/UnrealUSDWrapper/UnrealUSDWrapper.Build.cs
--- a/Engine/Plugins/Importers/USDImporter/Source/UnrealUSDWrapper/UnrealUSDWrapper.Build.cs
+++ b/Engine/Plugins/Importers/USDImporter/Source/UnrealUSDWrapper/UnrealUSDWrapper.Build.cs
@@ -179,7 +179,9 @@ namespace UnrealBuildTool.Rules
 			// on each module with versions that use either the ANSI (so USD-compatible) allocators or the UE allocators (ModuleBoilerplate.h) when appropriate.
 			// In a monolithic build we can't do that, as the primary game module will already define overrides for operator new and delete with
 			// the standard UE allocators: Since we can only have one operator new/delete override on the entire monolithic executable, we can't define our own overrides.
-			// The only way around it is by forcing the ansi allocator in your project's target file (YourProject/Source/YourProject.Target.cs) file like this:
+			// Additionally, the ANSI allocator does not work properly with FMallocPoisonProxy. Consequently, FMallocPoisonProxy has to be disabled.
+			// The only way around it is by forcing the ansi allocator and disabling FMallocPoisonProxy in your project's target file
+			// (YourProject/Source/YourProject.Target.cs) file like this:
 			//
 			//		public class YourProject : TargetRules
 			//		{
@@ -187,14 +189,17 @@ namespace UnrealBuildTool.Rules
 			//			{
 			//				...
 			//				GlobalDefinitions.Add("FORCE_ANSI_ALLOCATOR=1");
+			//				GlobalDefinitions.Add("UE_USE_MALLOC_FILL_BYTES=0");
+			//				...
 			//			}
 			//		}
 			//
-			// This will force the entire built executable to use the ANSI C allocators for everything (by disabling the UE overrides in ModuleBoilerplate.h), and so UE and USD allocations will be compatible.
+			// This will force the entire built executable to use the ANSI C allocators for everything (by disabling the UE overrides in ModuleBoilerplate.h) while
+			// FMallocPoisonProxy is disabled, and so UE and USD allocations will be compatible.
 			// Note that by that point everything will be using the USD-compatible ANSI allocators anyway, so our overrides in USDMemory.h are also disabled, as they're unnecessary.
 			// Also note that we're forced to use dynamic linking for monolithic targets mainly because static linking the USD libraries disables support for user USD plugins, and secondly
 			// because those static libraries would need to be linked with the --whole-archive argument, and there is currently no standard way of doing that in UE.
-			if (bEnableUsdSdk && Target.LinkType == TargetLinkType.Monolithic && !Target.GlobalDefinitions.Contains("FORCE_ANSI_ALLOCATOR=1"))
+			if (bEnableUsdSdk && Target.LinkType == TargetLinkType.Monolithic && !Target.GlobalDefinitions.Contains("FORCE_ANSI_ALLOCATOR=1") && !Target.GlobalDefinitions.Contains("UE_USE_MALLOC_FILL_BYTES=0"))
 			{
 				PublicDefinitions.Add("USD_FORCE_DISABLED=1");
 				bEnableUsdSdk = false;
