USD: Fix invalid foliage instances not being exported to USD.

diff --git a/Engine/Plugins/Importers/USDImporter/Source/USDUtilities/Private/USDPrimConversion.cpp b/Engine/Plugins/Importers/USDImporter/Source/USDUtilities/Private/USDPrimConversion.cpp
--- a/Engine/Plugins/Importers/USDImporter/Source/USDUtilities/Private/USDPrimConversion.cpp
+++ b/Engine/Plugins/Importers/USDImporter/Source/USDUtilities/Private/USDPrimConversion.cpp
@@ -2332,6 +2332,63 @@ bool UnrealToUsd::ConvertXformable( const FTransform& RelativeTransform, pxr::Us
 	return true;
 }
 
+#if WITH_EDITOR
+namespace UE
+{
+	namespace USDPrimConversionImpl
+	{
+		namespace Private
+		{
+			void ConvertFoliageInstances(
+				const FFoliageInfo& Info,
+				const TSet<int32>& UEInstances,
+				const FUsdStageInfo& StageInfo,
+				int PrototypeIndex,
+				pxr::VtArray<int>& ProtoIndices,
+				pxr::VtArray<pxr::GfVec3f>& Positions,
+				pxr::VtArray<pxr::GfQuath>& Orientations,
+				pxr::VtArray<pxr::GfVec3f>& Scales
+			)
+			{
+				FScopedUsdAllocs Allocs;
+
+				const int32 NumInstances = UEInstances.Num();
+
+				ProtoIndices.reserve( ProtoIndices.size() + NumInstances );
+				Positions.reserve( Positions.size() + NumInstances );
+				Orientations.reserve( Orientations.size() + NumInstances );
+				Scales.reserve( Scales.size() + NumInstances );
+
+				for ( int32 InstanceIndex : UEInstances )
+				{
+					const FFoliageInstancePlacementInfo* Instance = &Info.Instances[ InstanceIndex ];
+
+					// Convert axes
+					FTransform UETransform{ Instance->Rotation, (FVector)Instance->Location, (FVector)Instance->DrawScale3D };
+					FTransform USDTransform = UsdUtils::ConvertAxes( StageInfo.UpAxis == EUsdUpAxis::ZAxis, UETransform );
+
+					FVector Translation = USDTransform.GetTranslation();
+					FQuat Rotation = USDTransform.GetRotation();
+					FVector Scale = USDTransform.GetScale3D();
+
+					// Compensate metersPerUnit
+					const float UEMetersPerUnit = 0.01f;
+					if ( !FMath::IsNearlyEqual( UEMetersPerUnit, StageInfo.MetersPerUnit ) )
+					{
+						Translation *= ( UEMetersPerUnit / StageInfo.MetersPerUnit );
+					}
+
+					ProtoIndices.push_back( PrototypeIndex );
+					Positions.push_back( pxr::GfVec3f( Translation.X, Translation.Y, Translation.Z ) );
+					Orientations.push_back( pxr::GfQuath( Rotation.W, Rotation.X, Rotation.Y, Rotation.Z ) );
+					Scales.push_back( pxr::GfVec3f( Scale.X, Scale.Y, Scale.Z ) );
+				}
+			}
+		}
+	}
+}
+#endif // WITH_EDITOR
+
 bool UnrealToUsd::ConvertInstancedFoliageActor( const AInstancedFoliageActor& Actor, pxr::UsdPrim& UsdPrim, double TimeCode, ULevel* InstancesLevel )
 {
 #if WITH_EDITOR
@@ -2353,51 +2410,48 @@ bool UnrealToUsd::ConvertInstancedFoliageActor( const AInstancedFoliageActor& Ac
 	VtArray<GfQuath> Orientations;
 	VtArray<GfVec3f> Scales;
 
+	TSet<FFoliageInstanceBaseId> HandledComponents;
+
 	int PrototypeIndex = 0;
 	for ( const TPair<UFoliageType*, TUniqueObj<FFoliageInfo>>& FoliagePair : Actor.GetFoliageInfos() )
 	{
+		const UFoliageType* FoliageType = FoliagePair.Key;
 		const FFoliageInfo& Info = FoliagePair.Value.Get();
 
+		// Traverse valid foliage instances: Those that are being tracked to belonging to a particular component
 		for ( const TPair<FFoliageInstanceBaseId, FFoliageInstanceBaseInfo>& FoliageInstancePair : Actor.InstanceBaseCache.InstanceBaseMap )
 		{
+			const FFoliageInstanceBaseId& ComponentId = FoliageInstancePair.Key;
+			HandledComponents.Add( ComponentId );
+
 			UActorComponent* Comp = FoliageInstancePair.Value.BasePtr.Get();
 			if ( !Comp || ( InstancesLevel && ( Comp->GetComponentLevel() != InstancesLevel ) ) )
 			{
 				continue;
 			}
 
-			if ( const TSet<int32>* InstanceSet = Info.ComponentHash.Find( FoliageInstancePair.Key ) )
+			if ( const TSet<int32>* InstanceSet = Info.ComponentHash.Find( ComponentId ) )
 			{
-				const int32 NumInstances = InstanceSet->Num();
-				ProtoIndices.reserve( ProtoIndices.size() + NumInstances );
-				Positions.reserve( Positions.size() + NumInstances );
-				Orientations.reserve( Orientations.size() + NumInstances );
-				Scales.reserve( Scales.size() + NumInstances );
+				UE::USDPrimConversionImpl::Private::ConvertFoliageInstances( Info, *InstanceSet, StageInfo, PrototypeIndex, ProtoIndices, Positions, Orientations, Scales );
+			}
+		}
 
-				for ( int32 InstanceIndex : (*InstanceSet) )
+		// Do another pass to grab invalid foliage instances (not assigned to any particular component)
+		// Only export these when we're not given a particular level to export, or if that level is the actor's level (essentially
+		// pretending the invalid instances belong to the actor's level). This mostly helps prevent it from exporting the invalid instances
+		// multiple times in case we're calling this function repeatedly for each individual sublevel
+		if ( !InstancesLevel || InstancesLevel == Actor.GetLevel() )
+		{
+			for ( const TPair<FFoliageInstanceBaseId, TSet<int32>>& Pair : Info.ComponentHash )
+			{
+				const FFoliageInstanceBaseId& ComponentId = Pair.Key;
+				if ( HandledComponents.Contains( ComponentId ) )
 				{
-					const FFoliageInstancePlacementInfo* Instance = &Info.Instances[ InstanceIndex ];
-
-					// Convert axes
-					FTransform UETransform{ Instance->Rotation, (FVector)Instance->Location, (FVector)Instance->DrawScale3D };
-					FTransform USDTransform = UsdUtils::ConvertAxes( StageInfo.UpAxis == EUsdUpAxis::ZAxis, UETransform );
-
-					FVector Translation = USDTransform.GetTranslation();
-					FQuat Rotation = USDTransform.GetRotation();
-					FVector Scale = USDTransform.GetScale3D();
-
-					// Compensate metersPerUnit
-					const float UEMetersPerUnit = 0.01f;
-					if ( !FMath::IsNearlyEqual( UEMetersPerUnit, StageInfo.MetersPerUnit ) )
-					{
-						Translation *= ( UEMetersPerUnit / StageInfo.MetersPerUnit );
-					}
-
-					ProtoIndices.push_back( PrototypeIndex );
-					Positions.push_back( GfVec3f( Translation.X, Translation.Y, Translation.Z ) );
-					Orientations.push_back( GfQuath( Rotation.W, Rotation.X, Rotation.Y, Rotation.Z ) );
-					Scales.push_back( GfVec3f( Scale.X, Scale.Y, Scale.Z ) );
+					continue;
 				}
+
+				const TSet<int32>& InstanceSet = Pair.Value;
+				UE::USDPrimConversionImpl::Private::ConvertFoliageInstances( Info, InstanceSet, StageInfo, PrototypeIndex, ProtoIndices, Positions, Orientations, Scales );
 			}
 		}
 
@@ -2478,7 +2532,7 @@ bool UnrealToUsd::CreateComponentPropertyBaker( UE::FUsdPrim& Prim, const UScene
 #else
 					UsdPrim.IsA< pxr::UsdLuxLight >() )
 #endif // #if defined(HAS_USDLUX_LIGHTAPI)
-				
+
 			{
 				AdditionalRotation = FTransform( FRotator( 0.0f, 90.0f, 0.0f ) );
 
