Optimize USD import for huge scene (i.e. Moana)
  - Parallelize point instancer evaluation from USD
  - Fix race-condition in USDMemory
  - Add some Insights tag and slowtask scope for better progress updates
  - Move SystemFree to async task because HeapFree is tremendously slow on Windows 10

diff --git a/Engine/Plugins/Importers/USDImporter/Source/USDSchemas/Private/USDGeomPointInstancerTranslator.cpp b/Engine/Plugins/Importers/USDImporter/Source/USDSchemas/Private/USDGeomPointInstancerTranslator.cpp
--- a/Engine/Plugins/Importers/USDImporter/Source/USDSchemas/Private/USDGeomPointInstancerTranslator.cpp
+++ b/Engine/Plugins/Importers/USDImporter/Source/USDSchemas/Private/USDGeomPointInstancerTranslator.cpp
@@ -18,8 +18,10 @@
 
 #if USE_USD_SDK
 
+#include "Async/Async.h"
 #include "Components/HierarchicalInstancedStaticMeshComponent.h"
 #include "Engine/StaticMesh.h"
+#include "Misc/ScopedSlowTask.h"
 #include "Modules/ModuleManager.h"
 
 #include "USDIncludesStart.h"
@@ -31,10 +33,14 @@
 	#include "pxr/usd/usdGeom/xformable.h"
 #include "USDIncludesEnd.h"
 
+#define LOCTEXT_NAMESPACE "USDGeomPointInstancer"
+
 namespace UsdGeomPointInstancerTranslatorImpl
 {
-	bool ConvertGeomPointInstancer( const pxr::UsdStageRefPtr& Stage, const pxr::UsdGeomPointInstancer& PointInstancer, const int32 ProtoIndex, UHierarchicalInstancedStaticMeshComponent& HismComponent, pxr::UsdTimeCode EvalTime )
+	bool GetPointInstancerTransforms( const pxr::UsdStageRefPtr& Stage, const pxr::UsdGeomPointInstancer& PointInstancer, const int32 ProtoIndex, pxr::UsdTimeCode EvalTime, TArray<FTransform>& OutInstanceTransforms )
 	{
+		TRACE_CPUPROFILER_EVENT_SCOPE( GetPointInstancerTransforms );
+
 		FScopedUsdAllocs UsdAllocs;
 
 		pxr::VtArray< int > ProtoIndices = UsdUtils::GetUsdValue< pxr::VtArray< int > >( PointInstancer.GetProtoIndicesAttr(), EvalTime );
@@ -55,23 +61,31 @@ namespace UsdGeomPointInstancerTranslatorImpl
 
 		FScopedUnrealAllocs UnrealAllocs;
 
+		OutInstanceTransforms.Reset( UsdInstanceTransforms.size() );
+
 		for ( pxr::GfMatrix4d& UsdMatrix : UsdInstanceTransforms )
 		{
 			if ( ProtoIndices[ Index ] == ProtoIndex )
 			{
-				FTransform InstanceTransform = UsdToUnreal::ConvertMatrix( StageInfo, UsdMatrix );
-
-				HismComponent.AddInstance( InstanceTransform );
+				OutInstanceTransforms.Add( UsdToUnreal::ConvertMatrix( StageInfo, UsdMatrix ) );
 			}
 
 			++Index;
 		}
 
-		HismComponent.BuildTreeIfOutdated( true, true );
-
 		return true;
 	}
 
+	void ApplyPointInstanceTransforms( UInstancedStaticMeshComponent* Component, TArray<FTransform>& InstanceTransforms )
+	{
+		if ( Component )
+		{
+			TRACE_CPUPROFILER_EVENT_SCOPE(ApplyPointInstanceTransforms);
+
+			Component->AddInstances( InstanceTransforms, false );
+		}
+	}
+
 	UStaticMesh* SetStaticMesh( const pxr::UsdGeomMesh& UsdMesh, UStaticMeshComponent& MeshComponent, const UUsdAssetCache& AssetCache )
 	{
 		FScopedUnrealAllocs UnrealAllocs;
@@ -129,6 +143,8 @@ void FUsdGeomPointInstancerTranslator::UpdateComponents( USceneComponent* PointI
 		return;
 	}
 
+	TRACE_CPUPROFILER_EVENT_SCOPE( FUsdGeomPointInstancerTranslator::UpdateComponents );
+
 	PointInstancerRootComponent->Modify();
 
 	FScopedUsdAllocs UsdAllocs;
@@ -189,8 +205,13 @@ void FUsdGeomPointInstancerTranslator::UpdateComponents( USceneComponent* PointI
 	{
 		TGuardValue< USceneComponent* > ParentComponentGuard( Context->ParentComponent, PointInstancerRootComponent );
 
+		TArray<TFuture<TTuple<UHierarchicalInstancedStaticMeshComponent*, TArray<FTransform>>>> Tasks;
+
+		FScopedSlowTask PrototypesPathsSlowTask( (float)PrototypesPaths.size() * 3, LOCTEXT("GeomPointUpdateComponents", "Updating Point Components") );
 		for ( int32 PrototypeIndex = 0; PrototypeIndex < PrototypesPaths.size(); ++PrototypeIndex)
 		{
+			PrototypesPathsSlowTask.EnterProgressFrame();
+
 			pxr::UsdPrim PrototypeUsdPrim = Prim.GetStage()->GetPrimAtPath( PrototypesPaths[PrototypeIndex] );
 			if ( !PrototypeUsdPrim )
 			{
@@ -232,10 +253,22 @@ void FUsdGeomPointInstancerTranslator::UpdateComponents( USceneComponent* PointI
 			}
 
 			TArray< TUsdStore< pxr::UsdPrim > > ChildGeomMeshPrims = UsdUtils::GetAllPrimsOfType( PrototypeUsdPrim, pxr::TfType::Find< pxr::UsdGeomMesh >() );
-
-			for ( const TUsdStore< pxr::UsdPrim >& PrototypeGeomMeshPrim : ChildGeomMeshPrims )
+			TArray< TUsdStore< pxr::SdfPath > > Unwound;
+			Unwound.AddZeroed( ChildGeomMeshPrims.Num() );
+			PrototypesPathsSlowTask.EnterProgressFrame();
+			ParallelFor(
+				TEXT("Unwinding"),
+				ChildGeomMeshPrims.Num(), 1,
+				[&](int32 Index)
 				{
-				TUsdStore< pxr::SdfPath > PrototypeTargetPrimPath = UsdGeomPointInstancerTranslatorImpl::UnwindToNonCollapsedPrim( FUsdSchemaTranslator::ECollapsingType::Assets, PrototypeGeomMeshPrim, Context );
+					const TUsdStore< pxr::UsdPrim >& PrototypeGeomMeshPrim = ChildGeomMeshPrims[Index];
+					Unwound[Index] = UsdGeomPointInstancerTranslatorImpl::UnwindToNonCollapsedPrim( FUsdSchemaTranslator::ECollapsingType::Assets, PrototypeGeomMeshPrim, Context );
+				}
+			);
+			PrototypesPathsSlowTask.EnterProgressFrame();
+			for (int32 Index = 0; Index < ChildGeomMeshPrims.Num(); Index++)
+			{
+				TUsdStore< pxr::SdfPath > PrototypeTargetPrimPath = Unwound[Index];
 
 				const FString UEPrototypeTargetPrimPath = UsdToUnreal::ConvertPath( PrototypeTargetPrimPath.Get() );
 
@@ -261,7 +294,18 @@ void FUsdGeomPointInstancerTranslator::UpdateComponents( USceneComponent* PointI
 
 							UUsdAssetCache& AssetCache = *Context->AssetCache.Get();
 							UStaticMesh* StaticMesh = UsdGeomPointInstancerTranslatorImpl::SetStaticMesh( PrototypeGeomMesh, *HismComponent, AssetCache );
-							UsdGeomPointInstancerTranslatorImpl::ConvertGeomPointInstancer( Prim.GetStage(), PointInstancer, PrototypeIndex, *HismComponent, pxr::UsdTimeCode( Context->Time ) );
+
+							// Evaluating point instancer can take a long time and is thread-safe. Move to async task while we work on something else.
+							Tasks.Emplace(
+								Async(EAsyncExecution::ThreadPool,
+									[this, Stage = Prim.GetStage(), PointInstancer, PrototypeIndex, HismComponent]()
+									{
+										TArray<FTransform> InstanceTransforms;
+										UsdGeomPointInstancerTranslatorImpl::GetPointInstancerTransforms( Stage, PointInstancer, PrototypeIndex, pxr::UsdTimeCode(Context->Time), InstanceTransforms );
+
+										return MakeTuple( HismComponent, MoveTemp(InstanceTransforms) );
+									})
+							);
 
 							// Handle material overrides
 							if ( StaticMesh )
@@ -297,10 +341,19 @@ void FUsdGeomPointInstancerTranslator::UpdateComponents( USceneComponent* PointI
 				}
 			}
 		}
+
+		// Wait on and assign results of the point instancer.
+		for ( auto& Future : Tasks )
+		{
+			TTuple<UHierarchicalInstancedStaticMeshComponent*, TArray<FTransform>> Result { Future.Get() };
+			UsdGeomPointInstancerTranslatorImpl::ApplyPointInstanceTransforms( Result.Key, Result.Value );
+		}
 	}
 
 	Super::UpdateComponents( PointInstancerRootComponent );
 }
 
+#undef LOCTEXT_NAMESPACE
+
 #endif // #if USE_USD_SDK
 
diff --git a/Engine/Plugins/Importers/USDImporter/Source/USDSchemas/Private/USDGeomXformableTranslator.cpp b/Engine/Plugins/Importers/USDImporter/Source/USDSchemas/Private/USDGeomXformableTranslator.cpp
--- a/Engine/Plugins/Importers/USDImporter/Source/USDSchemas/Private/USDGeomXformableTranslator.cpp
+++ b/Engine/Plugins/Importers/USDImporter/Source/USDSchemas/Private/USDGeomXformableTranslator.cpp
@@ -360,7 +360,7 @@ void FUsdGeomXformableTranslator::UpdateComponents( USceneComponent* SceneCompon
 		// UsdToUnreal::ConvertXformable will set a new transform, which will emit warnings during PIE/Runtime if the component
 		// has Static mobility, so here we unregister, set the new transform value, and reregister below
 		const bool bStaticMobility = SceneComponent->Mobility == EComponentMobility::Static;
-		if ( bStaticMobility )
+		if ( bStaticMobility && SceneComponent->IsRegistered() )
 		{
 			SceneComponent->UnregisterComponent();
 		}
@@ -405,6 +405,8 @@ void FUsdGeomXformableTranslator::UpdateComponents( USceneComponent* SceneCompon
 
 bool FUsdGeomXformableTranslator::CollapsesChildren( ECollapsingType CollapsingType ) const
 {
+	TRACE_CPUPROFILER_EVENT_SCOPE( FUsdGeomXformableTranslator::CollapsesChildren );
+
 	bool bCollapsesChildren = false;
 
 	FScopedUsdAllocs UsdAllocs;
diff --git a/Engine/Plugins/Importers/USDImporter/Source/USDSchemas/Private/USDSchemaTranslator.cpp b/Engine/Plugins/Importers/USDImporter/Source/USDSchemas/Private/USDSchemaTranslator.cpp
--- a/Engine/Plugins/Importers/USDImporter/Source/USDSchemas/Private/USDSchemaTranslator.cpp
+++ b/Engine/Plugins/Importers/USDImporter/Source/USDSchemas/Private/USDSchemaTranslator.cpp
@@ -306,48 +306,54 @@ FUsdSchemaTranslatorTaskChain& FUsdSchemaTranslatorTaskChain::Then( ESchemaTrans
 	return *this;
 }
 
-ESchemaTranslationStatus FUsdSchemaTranslatorTaskChain::Execute(bool bExclusiveSyncTasks)
+namespace UsdSchemaTranslatorTaskChainImpl
 {
-	if ( !CurrentTask )
+	FORCEINLINE bool CanStart( FSchemaTranslatorTask* Task, bool bExclusiveSyncTasks )
 	{
-		return ESchemaTranslationStatus::Done;
+		return ( Task->LaunchPolicy == ESchemaTranslationLaunchPolicy::ExclusiveSync ) == bExclusiveSyncTasks;
 	}
+}
 
-	FSchemaTranslatorTask& TranslatorTask = *CurrentTask;
+ESchemaTranslationStatus FUsdSchemaTranslatorTaskChain::Execute(bool bExclusiveSyncTasks)
+{
+	FSchemaTranslatorTask* TranslatorTask = CurrentTask.Get();
 
-	bool bCanStart =
-		(  bExclusiveSyncTasks && CurrentTask->LaunchPolicy == ESchemaTranslationLaunchPolicy::ExclusiveSync ) ||
-		( !bExclusiveSyncTasks && CurrentTask->LaunchPolicy != ESchemaTranslationLaunchPolicy::ExclusiveSync );
+	if ( TranslatorTask == nullptr )
+	{
+		return ESchemaTranslationStatus::Done;
+	}
 
-	if ( !TranslatorTask.IsDone() )
+	if ( !TranslatorTask->IsDone() )
 	{
-		if ( !TranslatorTask.IsStarted() )
+		if ( !TranslatorTask->IsStarted() )
 		{
-			if ( bCanStart )
+			if ( UsdSchemaTranslatorTaskChainImpl::CanStart( TranslatorTask, bExclusiveSyncTasks ) )
 			{
-				TranslatorTask.Start();
+				TranslatorTask->Start();
 			}
 			else
 			{
 				return ESchemaTranslationStatus::Pending;
 			}
 		}
+
+		return ESchemaTranslationStatus::InProgress;
 	}
 	else
 	{
 		CurrentTask = CurrentTask->Continuation;
 
-		if ( CurrentTask )
+		if (( TranslatorTask = CurrentTask.Get()) != nullptr )
 		{
-			if ( bCanStart )
+			if ( UsdSchemaTranslatorTaskChainImpl::CanStart( TranslatorTask, bExclusiveSyncTasks ) )
 			{
 				if ( IsInGameThread() )
 				{
-					CurrentTask->StartIfAsync(); // Queue the next task asap if async
+					TranslatorTask->StartIfAsync(); // Queue the next task asap if async
 				}
 				else
 				{
-					CurrentTask->Start();
+					TranslatorTask->Start();
 				}
 			}
 			else
@@ -357,7 +363,7 @@ ESchemaTranslationStatus FUsdSchemaTranslatorTaskChain::Execute(bool bExclusiveS
 		}
 	}
 
-	return CurrentTask ? ESchemaTranslationStatus::InProgress : ESchemaTranslationStatus::Done;
+	return TranslatorTask ? ESchemaTranslationStatus::InProgress : ESchemaTranslationStatus::Done;
 }
 
 #undef LOCTEXT_NAMESPACE
diff --git a/Engine/Plugins/Importers/USDImporter/Source/USDStage/Private/USDPrimTwin.cpp b/Engine/Plugins/Importers/USDImporter/Source/USDStage/Private/USDPrimTwin.cpp
--- a/Engine/Plugins/Importers/USDImporter/Source/USDStage/Private/USDPrimTwin.cpp
+++ b/Engine/Plugins/Importers/USDImporter/Source/USDStage/Private/USDPrimTwin.cpp
@@ -47,6 +47,8 @@ void UUsdPrimTwin::RemoveChild( const TCHAR* InPrimPath )
 
 void UUsdPrimTwin::Clear()
 {
+	TRACE_CPUPROFILER_EVENT_SCOPE( UUsdPrimTwin::Clear );
+
 	FScopedUnrealAllocs UnrealAllocs;
 
 	Modify();
diff --git a/Engine/Plugins/Importers/USDImporter/Source/USDStageImporter/Private/USDStageImporter.cpp b/Engine/Plugins/Importers/USDImporter/Source/USDStageImporter/Private/USDStageImporter.cpp
--- a/Engine/Plugins/Importers/USDImporter/Source/USDStageImporter/Private/USDStageImporter.cpp
+++ b/Engine/Plugins/Importers/USDImporter/Source/USDStageImporter/Private/USDStageImporter.cpp
@@ -194,16 +194,22 @@ namespace UsdStageImporterImpl
 #endif // #if USE_USD_SDK
 	}
 
-	void CreateAssetsForPrims(const TArray<UE::FUsdPrim>& Prims, FUsdSchemaTranslationContext& TranslationContext)
+	void CreateAssetsForPrims(const TArray<UE::FUsdPrim>& Prims, FUsdSchemaTranslationContext& TranslationContext, const FText& ProgressMessage )
 	{
+		TRACE_CPUPROFILER_EVENT_SCOPE( CreateAssetsForPrims );
+
 		IUsdSchemasModule& UsdSchemasModule = FModuleManager::Get().LoadModuleChecked<IUsdSchemasModule>(TEXT("USDSchemas"));
 
+		FScopedSlowTask SlowTask( Prims.Num(), ProgressMessage );
+
 		for (const UE::FUsdPrim& Prim : Prims)
 		{
 			if (TSharedPtr<FUsdSchemaTranslator> SchemaTranslator = UsdSchemasModule.GetTranslatorRegistry().CreateTranslatorForSchema(TranslationContext.AsShared(), UE::FUsdTyped(Prim)))
 			{
 				SchemaTranslator->CreateAssets();
 			}
+
+			SlowTask.EnterProgressFrame();
 		}
 
 		TranslationContext.CompleteTasks();
@@ -211,6 +217,8 @@ namespace UsdStageImporterImpl
 
 	void ImportMaterials(FUsdStageImportContext& ImportContext, FUsdSchemaTranslationContext& TranslationContext)
 	{
+		TRACE_CPUPROFILER_EVENT_SCOPE( ImportMaterials );
+
 		if (!ImportContext.ImportOptions->bImportMaterials)
 		{
 			return;
@@ -218,11 +226,13 @@ namespace UsdStageImporterImpl
 
 		TArray< UE::FUsdPrim > MaterialPrims = UsdUtils::GetAllPrimsOfType( ImportContext.Stage.GetPseudoRoot(), TEXT("UsdShadeMaterial") );
 
-		CreateAssetsForPrims(MaterialPrims, TranslationContext);
+		CreateAssetsForPrims(MaterialPrims, TranslationContext, LOCTEXT("CreateMaterials", "Creating materials") );
 	}
 
 	void ImportMeshes(FUsdStageImportContext& ImportContext, FUsdSchemaTranslationContext& TranslationContext)
 	{
+		TRACE_CPUPROFILER_EVENT_SCOPE( ImportMeshes );
+
 #if USE_USD_SDK
 		if (!ImportContext.ImportOptions->bImportGeometry)
 		{
@@ -242,12 +252,14 @@ namespace UsdStageImporterImpl
 		};
 
 		TArray< UE::FUsdPrim > MeshPrims = UsdUtils::GetAllPrimsOfType( ImportContext.Stage.GetPseudoRoot(), TEXT("UsdGeomXformable"), PruneCollapsedMeshes );
-		CreateAssetsForPrims(MeshPrims, TranslationContext);
+		CreateAssetsForPrims(MeshPrims, TranslationContext, LOCTEXT( "CreateMeshes", "Creating meshes" ) );
 #endif // #if USE_USD_SDK
 	}
 
 	void ImportAnimation(FUsdStageImportContext& ImportContext, UE::FUsdPrim& Prim, USceneComponent* SceneComponent)
 	{
+		TRACE_CPUPROFILER_EVENT_SCOPE( ImportAnimation );
+
 		if ( !ImportContext.ImportOptions->bImportLevelSequences )
 		{
 			return;
@@ -262,7 +274,12 @@ namespace UsdStageImporterImpl
 
 	void ImportActor(FUsdStageImportContext& ImportContext, UE::FUsdPrim& Prim, FUsdSchemaTranslationContext& TranslationContext)
 	{
+		TRACE_CPUPROFILER_EVENT_SCOPE( ImportActor );
+
 		IUsdSchemasModule& UsdSchemasModule = FModuleManager::Get().LoadModuleChecked< IUsdSchemasModule >(TEXT("USDSchemas"));
+
+		FScopedSlowTask SlowTask( 3, LOCTEXT("ImportActor", "Importing Actor") );
+
 		bool bExpandChilren = true;
 		USceneComponent* Component = nullptr;
 
@@ -274,6 +291,7 @@ namespace UsdStageImporterImpl
 			bExpandChilren = !SchemaTranslator->CollapsesChildren(FUsdSchemaTranslator::ECollapsingType::Components);
 		}
 
+		SlowTask.EnterProgressFrame();
 		// Recurse to children
 		if (bExpandChilren)
 		{
@@ -286,7 +304,7 @@ namespace UsdStageImporterImpl
 				ImportActor(ImportContext, ChildStore, TranslationContext);
 			}
 		}
-
+		SlowTask.EnterProgressFrame();
 		if ( Component )
 		{
 			// LightComponents specifically need this to setup static lighting
@@ -304,6 +322,7 @@ namespace UsdStageImporterImpl
 			}
 #endif // USE_USD_SDK
 		}
+		SlowTask.EnterProgressFrame();
 	}
 
 	void ImportActors(FUsdStageImportContext& ImportContext, FUsdSchemaTranslationContext& TranslationContext)
@@ -1542,7 +1561,7 @@ bool UUsdStageImporter::ReimportSingleAsset(FUsdStageImportContext& ImportContex
 		UE::FUsdPrim TargetPrim = ImportContext.Stage.GetPrimAtPath( UE::FSdfPath( *OriginalImportData->PrimPath ) );
 		if ( TargetPrim )
 		{
-			UsdStageImporterImpl::CreateAssetsForPrims({TargetPrim}, TranslationContext.Get());
+			UsdStageImporterImpl::CreateAssetsForPrims({TargetPrim}, TranslationContext.Get(), LOCTEXT("CreateAssets", "Creating assets"));
 		}
 	}
 	TranslationContext->CompleteTasks();
diff --git a/Engine/Plugins/Importers/USDImporter/Source/USDUtilities/Private/USDGeomMeshConversion.cpp b/Engine/Plugins/Importers/USDImporter/Source/USDUtilities/Private/USDGeomMeshConversion.cpp
--- a/Engine/Plugins/Importers/USDImporter/Source/USDUtilities/Private/USDGeomMeshConversion.cpp
+++ b/Engine/Plugins/Importers/USDImporter/Source/USDUtilities/Private/USDGeomMeshConversion.cpp
@@ -1050,6 +1050,8 @@ bool UsdToUnreal::ConvertGeomMeshHierarchy(
 	bool bCombineIdenticalMaterialSlots
 )
 {
+	TRACE_CPUPROFILER_EVENT_SCOPE( UsdToUnreal::ConvertGeomMeshHierarchy );
+
 	FStaticMeshAttributes StaticMeshAttributes( OutMeshDescription );
 	StaticMeshAttributes.Register();
 
diff --git a/Engine/Plugins/Importers/USDImporter/Source/UnrealUSDWrapper/Private/USDMemory.cpp b/Engine/Plugins/Importers/USDImporter/Source/UnrealUSDWrapper/Private/USDMemory.cpp
--- a/Engine/Plugins/Importers/USDImporter/Source/UnrealUSDWrapper/Private/USDMemory.cpp
+++ b/Engine/Plugins/Importers/USDImporter/Source/UnrealUSDWrapper/Private/USDMemory.cpp
@@ -162,14 +162,16 @@ void FUsdMemoryManager::Free( void* Original )
 #if USD_USES_SYSTEM_MALLOC
 	// Because USD is multi-threaded, it might call us back to free an object after we've exited our allocator scope.
 	// This can happen for inlined USD functions that call delete.
-	if ( FUsdMemoryManager::IsUsingSystemMalloc() || SystemAllocedPtrs.Contains( Original ) )
+	int32 Removed = 0;
 	{
-		{
-			FScopeLock Lock( &CriticalSection );
-			SystemAllocedPtrs.Remove( Original );
-		}
+		FScopeLock Lock( &CriticalSection );
+		Removed = SystemAllocedPtrs.Remove( Original );
+	}
 
-		FMemory::SystemFree( Original );
+	if ( FUsdMemoryManager::IsUsingSystemMalloc() || Removed > 0 )
+	{
+		// System allocations are so slow that we send them to another thread for freeing
+		Async( EAsyncExecution::ThreadPool, [Original](){ FMemory::SystemFree(Original); } );
 	}
 	else
 #endif // #if USD_USES_SYSTEM_MALLOC
