USD: Add cvar to allow collapsing prims that don't have a specified Kind.

This is copied over from swarm 19085926, intended originally for 5.0.1.

diff --git a/Engine/Plugins/Importers/USDImporter/Source/USDSchemas/Private/USDGeomXformableTranslator.cpp b/Engine/Plugins/Importers/USDImporter/Source/USDSchemas/Private/USDGeomXformableTranslator.cpp
--- a/Engine/Plugins/Importers/USDImporter/Source/USDSchemas/Private/USDGeomXformableTranslator.cpp
+++ b/Engine/Plugins/Importers/USDImporter/Source/USDSchemas/Private/USDGeomXformableTranslator.cpp
@@ -40,6 +40,12 @@
 	#include "pxr/usd/usdShade/materialBindingAPI.h"
 #include "USDIncludesEnd.h"
 
+static bool GCollapsePrimsWithoutKind = true;
+static FAutoConsoleVariableRef CVarCollapsePrimsWithoutKind(
+	TEXT( "USD.CollapsePrimsWithoutKind" ),
+	GCollapsePrimsWithoutKind,
+	TEXT( "Allow collapsing prims that have no authored 'Kind' value" ) );
+
 static int32 GMaxNumVerticesCollapsedMesh = 5000000;
 static FAutoConsoleVariableRef CVarMaxNumVerticesCollapsedMesh(
 	TEXT( "USD.MaxNumVerticesCollapsedMesh" ),
@@ -416,10 +422,8 @@ bool FUsdGeomXformableTranslator::CollapsesChildren( ECollapsingType CollapsingT
 
 	if ( Model )
 	{
-		// We need KindValidationNone here or else we get inconsistent results when a prim references another prim that is a component.
-		// For example, when referencing a component prim in another file, this returns 'true' if the referencer is a root prim,
-		// but false if the referencer is within another Xform prim, for whatever reason.
-		bCollapsesChildren = EnumHasAnyFlags( Context->KindsToCollapse, UsdUtils::GetDefaultKind( Prim ) );
+		EUsdDefaultKind PrimKind = UsdUtils::GetDefaultKind( Prim );
+		bCollapsesChildren = Context->KindsToCollapse != EUsdDefaultKind::None && ( EnumHasAnyFlags( Context->KindsToCollapse, PrimKind ) || ( PrimKind == EUsdDefaultKind::None && GCollapsePrimsWithoutKind ) );
 
 		if ( !bCollapsesChildren )
 		{
diff --git a/Engine/Plugins/Importers/USDImporter/Source/USDUtilities/Private/USDConversionUtils.cpp b/Engine/Plugins/Importers/USDImporter/Source/USDUtilities/Private/USDConversionUtils.cpp
--- a/Engine/Plugins/Importers/USDImporter/Source/USDUtilities/Private/USDConversionUtils.cpp
+++ b/Engine/Plugins/Importers/USDImporter/Source/USDUtilities/Private/USDConversionUtils.cpp
@@ -773,6 +773,9 @@ EUsdDefaultKind UsdUtils::GetDefaultKind( const pxr::UsdPrim& Prim )
 		return Result;
 	}
 
+	// We need KindValidationNone here or else we get inconsistent results when a prim references another prim that is a component.
+	// For example, when referencing a component prim in another file, this returns 'true' if the referencer is a root prim,
+	// but false if the referencer is within another Xform prim, for whatever reason.
 	if ( Model.IsKind( pxr::KindTokens->model, pxr::UsdModelAPI::KindValidationNone ) )
 	{
 		Result |= EUsdDefaultKind::Model;
