USDImporter: use native USD file formats when opening stages and saving/exporting layers

In general, non-native file formats like .abc and .mtlx will typically be
authored elsewhere and then composed into native USD layers. Making
edits to a layer in a non-native format layer and then saving that stage
would be the typical way of saving those edits back to the non-native
format layer they came from.

diff --git a/Engine/Plugins/Importers/USDImporter/Source/USDExporter/Private/AnimSequenceExporterUSD.cpp b/Engine/Plugins/Importers/USDImporter/Source/USDExporter/Private/AnimSequenceExporterUSD.cpp
--- a/Engine/Plugins/Importers/USDImporter/Source/USDExporter/Private/AnimSequenceExporterUSD.cpp
+++ b/Engine/Plugins/Importers/USDImporter/Source/USDExporter/Private/AnimSequenceExporterUSD.cpp
@@ -63,7 +63,7 @@ namespace UE
 UAnimSequenceExporterUSD::UAnimSequenceExporterUSD()
 {
 #if USE_USD_SDK
-	for ( const FString& Extension : UnrealUSDWrapper::GetAllSupportedFileFormats() )
+	for ( const FString& Extension : UnrealUSDWrapper::GetNativeFileFormats() )
 	{
 		// USDZ is not supported for writing for now
 		if ( Extension.Equals( TEXT( "usdz" ) ) )
diff --git a/Engine/Plugins/Importers/USDImporter/Source/USDExporter/Private/LevelExporterUSD.cpp b/Engine/Plugins/Importers/USDImporter/Source/USDExporter/Private/LevelExporterUSD.cpp
--- a/Engine/Plugins/Importers/USDImporter/Source/USDExporter/Private/LevelExporterUSD.cpp
+++ b/Engine/Plugins/Importers/USDImporter/Source/USDExporter/Private/LevelExporterUSD.cpp
@@ -23,7 +23,7 @@
 ULevelExporterUSD::ULevelExporterUSD()
 {
 #if USE_USD_SDK
-	for ( const FString& Extension : UnrealUSDWrapper::GetAllSupportedFileFormats() )
+	for ( const FString& Extension : UnrealUSDWrapper::GetNativeFileFormats() )
 	{
 		// USDZ is not supported for writing for now
 		if ( Extension.Equals( TEXT( "usdz" ) ) )
diff --git a/Engine/Plugins/Importers/USDImporter/Source/USDExporter/Private/LevelSequenceExporterUSD.cpp b/Engine/Plugins/Importers/USDImporter/Source/USDExporter/Private/LevelSequenceExporterUSD.cpp
--- a/Engine/Plugins/Importers/USDImporter/Source/USDExporter/Private/LevelSequenceExporterUSD.cpp
+++ b/Engine/Plugins/Importers/USDImporter/Source/USDExporter/Private/LevelSequenceExporterUSD.cpp
@@ -959,7 +959,7 @@ namespace UE
 ULevelSequenceExporterUsd::ULevelSequenceExporterUsd()
 {
 #if USE_USD_SDK
-	for ( const FString& Extension : UnrealUSDWrapper::GetAllSupportedFileFormats() )
+	for ( const FString& Extension : UnrealUSDWrapper::GetNativeFileFormats() )
 	{
 		// USDZ is not supported for writing for now
 		if ( Extension.Equals( TEXT( "usdz" ) ) )
diff --git a/Engine/Plugins/Importers/USDImporter/Source/USDExporter/Private/MaterialExporterUSD.cpp b/Engine/Plugins/Importers/USDImporter/Source/USDExporter/Private/MaterialExporterUSD.cpp
--- a/Engine/Plugins/Importers/USDImporter/Source/USDExporter/Private/MaterialExporterUSD.cpp
+++ b/Engine/Plugins/Importers/USDImporter/Source/USDExporter/Private/MaterialExporterUSD.cpp
@@ -27,7 +27,7 @@
 UMaterialExporterUsd::UMaterialExporterUsd()
 {
 #if USE_USD_SDK
-	for ( const FString& Extension : UnrealUSDWrapper::GetAllSupportedFileFormats() )
+	for ( const FString& Extension : UnrealUSDWrapper::GetNativeFileFormats() )
 	{
 		// USDZ is not supported for writing for now
 		if ( Extension.Equals( TEXT( "usdz" ) ) )
diff --git a/Engine/Plugins/Importers/USDImporter/Source/USDExporter/Private/SkeletalMeshExporterUSD.cpp b/Engine/Plugins/Importers/USDImporter/Source/USDExporter/Private/SkeletalMeshExporterUSD.cpp
--- a/Engine/Plugins/Importers/USDImporter/Source/USDExporter/Private/SkeletalMeshExporterUSD.cpp
+++ b/Engine/Plugins/Importers/USDImporter/Source/USDExporter/Private/SkeletalMeshExporterUSD.cpp
@@ -59,7 +59,7 @@ namespace UE
 USkeletalMeshExporterUsd::USkeletalMeshExporterUsd()
 {
 #if USE_USD_SDK
-	for ( const FString& Extension : UnrealUSDWrapper::GetAllSupportedFileFormats() )
+	for ( const FString& Extension : UnrealUSDWrapper::GetNativeFileFormats() )
 	{
 		// USDZ is not supported for writing for now
 		if ( Extension.Equals( TEXT( "usdz" ) ) )
diff --git a/Engine/Plugins/Importers/USDImporter/Source/USDExporter/Private/StaticMeshExporterUSD.cpp b/Engine/Plugins/Importers/USDImporter/Source/USDExporter/Private/StaticMeshExporterUSD.cpp
--- a/Engine/Plugins/Importers/USDImporter/Source/USDExporter/Private/StaticMeshExporterUSD.cpp
+++ b/Engine/Plugins/Importers/USDImporter/Source/USDExporter/Private/StaticMeshExporterUSD.cpp
@@ -69,7 +69,7 @@ bool UStaticMeshExporterUsd::IsUsdAvailable()
 UStaticMeshExporterUsd::UStaticMeshExporterUsd()
 {
 #if USE_USD_SDK
-	for ( const FString& Extension : UnrealUSDWrapper::GetAllSupportedFileFormats() )
+	for ( const FString& Extension : UnrealUSDWrapper::GetNativeFileFormats() )
 	{
 		// USDZ is not supported for writing for now
 		if ( Extension.Equals( TEXT( "usdz" ) ) )
diff --git a/Engine/Plugins/Importers/USDImporter/Source/USDExporter/Public/LevelExporterUSDOptions.h b/Engine/Plugins/Importers/USDImporter/Source/USDExporter/Public/LevelExporterUSDOptions.h
--- a/Engine/Plugins/Importers/USDImporter/Source/USDExporter/Public/LevelExporterUSDOptions.h
+++ b/Engine/Plugins/Importers/USDImporter/Source/USDExporter/Public/LevelExporterUSDOptions.h
@@ -106,7 +106,7 @@ private:
 	UFUNCTION()
 	static TArray<FString> GetUsdExtensions()
 	{
-		TArray<FString> Extensions = UnrealUSDWrapper::GetAllSupportedFileFormats();
+		TArray<FString> Extensions = UnrealUSDWrapper::GetNativeFileFormats();
 		Extensions.Remove( TEXT( "usdz" ) );
 		return Extensions;
 	}
diff --git a/Engine/Plugins/Importers/USDImporter/Source/USDStageEditor/Private/SUSDLayersTreeView.cpp b/Engine/Plugins/Importers/USDImporter/Source/USDStageEditor/Private/SUSDLayersTreeView.cpp
--- a/Engine/Plugins/Importers/USDImporter/Source/USDStageEditor/Private/SUSDLayersTreeView.cpp
+++ b/Engine/Plugins/Importers/USDImporter/Source/USDStageEditor/Private/SUSDLayersTreeView.cpp
@@ -506,7 +506,7 @@ bool SUsdLayersTreeView::CanAddSubLayer() const
 
 void SUsdLayersTreeView::OnAddSubLayer()
 {
-	TOptional< FString > SubLayerFile = UsdUtils::BrowseUsdFile( UsdUtils::EBrowseFileMode::Open, AsShared() );
+	TOptional< FString > SubLayerFile = UsdUtils::BrowseUsdFile( UsdUtils::EBrowseFileMode::Composition, AsShared() );
 
 	if ( !SubLayerFile )
 	{
diff --git a/Engine/Plugins/Importers/USDImporter/Source/USDStageEditor/Private/SUSDStageTreeView.cpp b/Engine/Plugins/Importers/USDImporter/Source/USDStageEditor/Private/SUSDStageTreeView.cpp
--- a/Engine/Plugins/Importers/USDImporter/Source/USDStageEditor/Private/SUSDStageTreeView.cpp
+++ b/Engine/Plugins/Importers/USDImporter/Source/USDStageEditor/Private/SUSDStageTreeView.cpp
@@ -749,7 +749,7 @@ void SUsdStageTreeView::OnAddReference()
 		return;
 	}
 
-	TOptional< FString > PickedFile = UsdUtils::BrowseUsdFile( UsdUtils::EBrowseFileMode::Open, AsShared() );
+	TOptional< FString > PickedFile = UsdUtils::BrowseUsdFile( UsdUtils::EBrowseFileMode::Composition, AsShared() );
 
 	if ( !PickedFile )
 	{
diff --git a/Engine/Plugins/Importers/USDImporter/Source/USDUtilities/Private/USDLayerUtils.cpp b/Engine/Plugins/Importers/USDImporter/Source/USDUtilities/Private/USDLayerUtils.cpp
--- a/Engine/Plugins/Importers/USDImporter/Source/USDUtilities/Private/USDLayerUtils.cpp
+++ b/Engine/Plugins/Importers/USDImporter/Source/USDUtilities/Private/USDLayerUtils.cpp
@@ -134,7 +134,14 @@ TOptional< FString > UsdUtils::BrowseUsdFile( EBrowseFileMode Mode, TSharedRef<
 
 	TArray< FString > OutFiles;
 
-	TArray< FString > SupportedExtensions = UnrealUSDWrapper::GetAllSupportedFileFormats();
+	// When browsing files for the purposes of opening a stage or saving layers,
+	// we offer the native USD file formats as options. Browsing files in order
+	// to use them as the targets of composition arcs (e.g. sublayers,
+	// references, payloads, etc.) also offers any plugin file formats that are
+	// registered.
+	TArray< FString > SupportedExtensions = ( Mode == EBrowseFileMode::Composition ) ?
+		UnrealUSDWrapper::GetAllSupportedFileFormats() :
+		UnrealUSDWrapper::GetNativeFileFormats();
 	if ( SupportedExtensions.Num() == 0 )
 	{
 		UE_LOG(LogUsd, Error, TEXT("No file extensions supported by the USD SDK!"));
@@ -159,6 +166,7 @@ TOptional< FString > UsdUtils::BrowseUsdFile( EBrowseFileMode Mode, TSharedRef<
 	switch ( Mode )
 	{
 		case EBrowseFileMode::Open :
+		case EBrowseFileMode::Composition :
 		{
 			if ( !DesktopPlatform->OpenFileDialog( ParentWindowHandle, LOCTEXT( "ChooseFile", "Choose file").ToString(), TEXT(""), TEXT(""), *FileTypes, EFileDialogFlags::None, OutFiles ) )
 			{
diff --git a/Engine/Plugins/Importers/USDImporter/Source/USDUtilities/Public/USDLayerUtils.h b/Engine/Plugins/Importers/USDImporter/Source/USDUtilities/Public/USDLayerUtils.h
--- a/Engine/Plugins/Importers/USDImporter/Source/USDUtilities/Public/USDLayerUtils.h
+++ b/Engine/Plugins/Importers/USDImporter/Source/USDUtilities/Public/USDLayerUtils.h
@@ -36,7 +36,13 @@ namespace UsdUtils
 {
 	enum class EBrowseFileMode
 	{
+		/** Browse for a file for the purpose of opening a stage with that file as the stage's root layer. */
 		Open,
+
+		/** Browse for a file for the purpose of using that file as the target of a composition arc (e.g. a sublayer, reference, payload, etc.). */
+		Composition,
+
+		/** Browse for a file for the purpose of saving or exporting a layer to that file. */
 		Save
 	};
 
