USD: Fix missing include on cl 19224061.

diff --git a/Engine/Plugins/Importers/USDImporter/Source/UnrealUSDWrapper/Private/USDMemory.cpp b/Engine/Plugins/Importers/USDImporter/Source/UnrealUSDWrapper/Private/USDMemory.cpp
--- a/Engine/Plugins/Importers/USDImporter/Source/UnrealUSDWrapper/Private/USDMemory.cpp
+++ b/Engine/Plugins/Importers/USDImporter/Source/UnrealUSDWrapper/Private/USDMemory.cpp
@@ -6,6 +6,7 @@
 #include "HAL/PlatformTLS.h"
 #include "HAL/TlsAutoCleanup.h"
 #include "Misc/ScopeLock.h"
+#include "Async/Async.h"
 
 class UNREALUSDWRAPPER_API FTlsSlot final
 {
@@ -113,7 +114,7 @@ void FUsdMemoryManager::ActivateAllocator( EUsdActiveAllocator Allocator )
 bool FUsdMemoryManager::DeactivateAllocator( EUsdActiveAllocator Allocator )
 {
 	FActiveAllocatorsStack& ActiveAllocatorStack = GetActiveAllocatorsStackForThread();
-		
+
 	// Remove the topmost instance of Allocator on the stack
 	if ( ActiveAllocatorStack->Num() > 0 )
 	{
@@ -189,7 +190,7 @@ bool FUsdMemoryManager::IsUsingSystemMalloc()
 	}
 
 	FActiveAllocatorsStack* ActiveAllocatorStack = reinterpret_cast< FActiveAllocatorsStack* >( ActiveAllocatorsStackTLS->GetTlsValue() );
-	
+
 	if ( ActiveAllocatorStack && (*ActiveAllocatorStack)->Num() > 0 )
 	{
 		return (*ActiveAllocatorStack)->Top() == EUsdActiveAllocator::System;
