USD: Cache whether a prim can be collapsed or not.

diff --git a/Engine/Plugins/Importers/USDImporter/Source/USDSchemas/Private/USDCollapsingCache.cpp b/Engine/Plugins/Importers/USDImporter/Source/USDSchemas/Private/USDCollapsingCache.cpp
new file mode 100644
--- /dev/null
+++ b/Engine/Plugins/Importers/USDImporter/Source/USDSchemas/Private/USDCollapsingCache.cpp
@@ -0,0 +1,189 @@
+// Copyright Epic Games, Inc. All Rights Reserved.
+
+#include "USDCollapsingCache.h"
+
+#include "USDSchemasModule.h"
+#include "USDSchemaTranslator.h"
+#include "USDTypesConversion.h"
+
+#include "UsdWrappers/SdfPath.h"
+#include "UsdWrappers/UsdPrim.h"
+
+#include "Modules/ModuleManager.h"
+
+#if USE_USD_SDK
+#include "USDIncludesStart.h"
+	#include "pxr/usd/sdf/path.h"
+	#include "pxr/usd/usd/prim.h"
+#include "USDIncludesEnd.h"
+#endif // USE_USD_SDK
+
+bool FUsdCollapsingCache::Serialize( FArchive& Ar )
+{
+	Ar << AssetPathsToCollapsedRoot;
+	Ar << ComponentPathsToCollapsedRoot;
+
+	return true;
+}
+
+bool FUsdCollapsingCache::IsPathCollapsed( const UE::FSdfPath& Path, ECollapsingType CollapsingType ) const
+{
+	const TMap< UE::FSdfPath, UE::FSdfPath >* MapToUse =
+		CollapsingType == ECollapsingType::Assets
+			? &AssetPathsToCollapsedRoot
+			: &ComponentPathsToCollapsedRoot;
+
+	if ( const UE::FSdfPath* FoundResult = MapToUse->Find( Path ) )
+	{
+		// A non-empty path to another prim means this prim is collapsed into that one
+		return !FoundResult->IsEmpty() && ( *FoundResult ) != Path;
+	}
+
+	// This should never happen: We should have cached the entire tree
+	ensureMsgf( false, TEXT( "Prim path '%s' has not been cached!" ), *Path.GetString() );
+	return false;
+}
+
+bool FUsdCollapsingCache::DoesPathCollapseChildren( const UE::FSdfPath& Path, ECollapsingType CollapsingType ) const
+{
+	const TMap< UE::FSdfPath, UE::FSdfPath >* MapToUse =
+		CollapsingType == ECollapsingType::Assets
+			? &AssetPathsToCollapsedRoot
+			: &ComponentPathsToCollapsedRoot;
+
+	if ( const UE::FSdfPath* FoundResult = MapToUse->Find( Path ) )
+	{
+		// We store our own Path in there when we collapse children. Otherwise we hold the path of our collapse root, or empty (in case nothing is collapsed up to here)
+		return (*FoundResult) == Path;
+	}
+
+	// This should never happen: We should have cached the entire tree
+	ensureMsgf( false, TEXT( "Prim path '%s' has not been cached!" ), *Path.GetString() );
+	return false;
+}
+
+UE::FSdfPath FUsdCollapsingCache::UnwindToNonCollapsedPath( const UE::FSdfPath& Path, ECollapsingType CollapsingType ) const
+{
+	const TMap< UE::FSdfPath, UE::FSdfPath >* MapToUse =
+		CollapsingType == ECollapsingType::Assets
+			? &AssetPathsToCollapsedRoot
+			: &ComponentPathsToCollapsedRoot;
+
+	if ( const UE::FSdfPath* FoundResult = MapToUse->Find( Path ) )
+	{
+		// An empty path here means that we are not collapsed at all
+		if ( FoundResult->IsEmpty() )
+		{
+			return Path;
+		}
+		// Otherwise we have our own path in there (in case we collapse children) or the path to the prim that collapsed us
+		else
+		{
+			return *FoundResult;
+		}
+	}
+
+	// This should never happen: We should have cached the entire tree
+	ensureMsgf( false, TEXT( "Prim path '%s' has not been cached!" ), *Path.GetString() );
+	return Path;
+}
+
+namespace UE::USDCollapsingCacheImpl::Private
+{
+#if USE_USD_SDK
+	void RecursiveRebuildCache(
+		const pxr::UsdPrim& UsdPrim,
+		FUsdSchemaTranslationContext& Context,
+		FUsdSchemaTranslatorRegistry& Registry,
+		TMap< UE::FSdfPath, UE::FSdfPath >& AssetPathsToCollapsedRoot,
+		TMap< UE::FSdfPath, UE::FSdfPath >& ComponentPathsToCollapsedRoot,
+		pxr::SdfPath AssetCollapsedRoot = pxr::SdfPath::EmptyPath(),
+		pxr::SdfPath ComponentCollapsedRoot = pxr::SdfPath::EmptyPath()
+	)
+	{
+		FScopedUsdAllocs Allocs;
+
+		pxr::SdfPath UsdPrimPath = UsdPrim.GetPrimPath();
+
+		bool bIsAssetCollapsed = !AssetCollapsedRoot.IsEmpty();
+		bool bIsComponentCollapsed = !ComponentCollapsedRoot.IsEmpty();
+
+		if ( !bIsAssetCollapsed || !bIsComponentCollapsed )
+		{
+			if ( TSharedPtr< FUsdSchemaTranslator > SchemaTranslator = Registry.CreateTranslatorForSchema( Context.AsShared(), UE::FUsdTyped( UsdPrim ) ) )
+			{
+				if ( !bIsAssetCollapsed )
+				{
+					if ( SchemaTranslator->CollapsesChildren( ECollapsingType::Assets ) )
+					{
+						AssetCollapsedRoot = UsdPrimPath;
+					}
+				}
+
+				if ( !bIsComponentCollapsed )
+				{
+					if ( SchemaTranslator->CollapsesChildren( ECollapsingType::Components ) )
+					{
+						ComponentCollapsedRoot = UsdPrimPath;
+					}
+				}
+			}
+		}
+
+		UE::FSdfPath PrimPath{ UsdPrimPath };
+		UE::FSdfPath UEAssetCollapsedRoot{ AssetCollapsedRoot };
+		UE::FSdfPath UEComponentCollapsedRoot{ ComponentCollapsedRoot };
+
+		// These paths will be still empty in case nothing has collapsed yet, hold UsdPrimPath in case UsdPrim collapses that type, or hold the path to the
+		// collapsed root passed in via our caller, in case we're collapsed
+		AssetPathsToCollapsedRoot.Add( PrimPath, UEAssetCollapsedRoot );
+		ComponentPathsToCollapsedRoot.Add( PrimPath, UEComponentCollapsedRoot );
+
+		pxr::UsdPrimSiblingRange PrimChildren = UsdPrim.GetFilteredChildren( pxr::UsdTraverseInstanceProxies( pxr::UsdPrimAllPrimsPredicate ) );
+		for ( pxr::UsdPrim Child : PrimChildren )
+		{
+			RecursiveRebuildCache( Child, Context, Registry, AssetPathsToCollapsedRoot, ComponentPathsToCollapsedRoot, AssetCollapsedRoot, ComponentCollapsedRoot );
+		}
+	}
+#endif // USE_USD_SDK
+}
+
+void FUsdCollapsingCache::RebuildCacheForSubtree( const UE::FUsdPrim& Prim, FUsdSchemaTranslationContext& Context )
+{
+	TRACE_CPUPROFILER_EVENT_SCOPE( FUsdCollapsingCache::RebuildCacheForSubtree );
+
+#if USE_USD_SDK
+	// We can't deallocate our collapsing cache pointer with the Usd allocator
+	FScopedUnrealAllocs UEAllocs;
+
+	// We don't want the translation context to try using its collapsing cache during the rebuild process, as that's the entire point
+	TSharedPtr<FUsdCollapsingCache> Pin = Context.CollapsingCache;
+	Context.CollapsingCache = nullptr;
+	{
+		FScopedUsdAllocs Allocs;
+
+		pxr::UsdPrim UsdPrim{Prim};
+		if ( !UsdPrim )
+		{
+			return;
+		}
+
+		IUsdSchemasModule& UsdSchemasModule = FModuleManager::Get().LoadModuleChecked< IUsdSchemasModule >( TEXT( "USDSchemas" ) );
+		FUsdSchemaTranslatorRegistry& Registry = UsdSchemasModule.GetTranslatorRegistry();
+
+		UE::USDCollapsingCacheImpl::Private::RecursiveRebuildCache( UsdPrim, Context, Registry, AssetPathsToCollapsedRoot, ComponentPathsToCollapsedRoot );
+	}
+	Context.CollapsingCache = Pin;
+#endif // USE_USD_SDK
+}
+
+void FUsdCollapsingCache::Clear()
+{
+	AssetPathsToCollapsedRoot.Empty();
+	ComponentPathsToCollapsedRoot.Empty();
+}
+
+bool FUsdCollapsingCache::IsEmpty()
+{
+	return AssetPathsToCollapsedRoot.IsEmpty() && ComponentPathsToCollapsedRoot.IsEmpty();
+}
\ No newline at end of file
diff --git a/Engine/Plugins/Importers/USDImporter/Source/USDSchemas/Private/USDGeomMeshTranslator.cpp b/Engine/Plugins/Importers/USDImporter/Source/USDSchemas/Private/USDGeomMeshTranslator.cpp
--- a/Engine/Plugins/Importers/USDImporter/Source/USDSchemas/Private/USDGeomMeshTranslator.cpp
+++ b/Engine/Plugins/Importers/USDImporter/Source/USDSchemas/Private/USDGeomMeshTranslator.cpp
@@ -1462,7 +1462,7 @@ bool FUsdGeomMeshTranslator::CanBeCollapsed( ECollapsingType CollapsingType ) co
 
 	// Don't collapse if our final UStaticMesh would have multiple LODs
 	if ( Context->bAllowInterpretingLODs &&
-		 CollapsingType == FUsdSchemaTranslator::ECollapsingType::Assets &&
+		 CollapsingType == ECollapsingType::Assets &&
 		 UsdUtils::IsGeomMeshALOD( Prim ) )
 	{
 		return false;
diff --git a/Engine/Plugins/Importers/USDImporter/Source/USDSchemas/Private/USDGeomPointInstancerTranslator.cpp b/Engine/Plugins/Importers/USDImporter/Source/USDSchemas/Private/USDGeomPointInstancerTranslator.cpp
--- a/Engine/Plugins/Importers/USDImporter/Source/USDSchemas/Private/USDGeomPointInstancerTranslator.cpp
+++ b/Engine/Plugins/Importers/USDImporter/Source/USDSchemas/Private/USDGeomPointInstancerTranslator.cpp
@@ -109,31 +109,6 @@ namespace UsdGeomPointInstancerTranslatorImpl
 
 		return StaticMesh;
 	}
-
-	TUsdStore< pxr::SdfPath > UnwindToNonCollapsedPrim( FUsdSchemaTranslator::ECollapsingType CollapsingType, TUsdStore< pxr::UsdPrim > UsdPrim, const TSharedRef< FUsdSchemaTranslationContext >& TranslationContext )
-	{
-		IUsdSchemasModule& UsdSchemasModule = FModuleManager::Get().LoadModuleChecked< IUsdSchemasModule >( TEXT("USDSchemas") );
-
-		TUsdStore< pxr::SdfPath > UsdPrimPath = UsdPrim.Get().GetPrimPath();
-
-		if ( TSharedPtr< FUsdSchemaTranslator > SchemaTranslator = UsdSchemasModule.GetTranslatorRegistry().CreateTranslatorForSchema( TranslationContext, UE::FUsdTyped( UsdPrim.Get() ) ) )
-		{
-			while ( SchemaTranslator->IsCollapsed( CollapsingType ) )
-			{
-				UsdPrimPath = UsdPrimPath.Get().GetParentPath();
-				UsdPrim = UsdPrim.Get().GetStage()->GetPrimAtPath( UsdPrimPath.Get() );
-
-				SchemaTranslator = UsdSchemasModule.GetTranslatorRegistry().CreateTranslatorForSchema( TranslationContext, UE::FUsdTyped( UsdPrim.Get() ) );
-
-				if ( !SchemaTranslator.IsValid() )
-				{
-					break;
-				}
-			}
-		}
-
-		return UsdPrimPath;
-	};
 }
 
 void FUsdGeomPointInstancerTranslator::UpdateComponents( USceneComponent* PointInstancerRootComponent )
@@ -253,30 +228,23 @@ void FUsdGeomPointInstancerTranslator::UpdateComponents( USceneComponent* PointI
 			}
 
 			TArray< TUsdStore< pxr::UsdPrim > > ChildGeomMeshPrims = UsdUtils::GetAllPrimsOfType( PrototypeUsdPrim, pxr::TfType::Find< pxr::UsdGeomMesh >() );
-			TArray< TUsdStore< pxr::SdfPath > > Unwound;
-			Unwound.AddZeroed( ChildGeomMeshPrims.Num() );
-			PrototypesPathsSlowTask.EnterProgressFrame();
-			ParallelFor(
-				TEXT("Unwinding"),
-				ChildGeomMeshPrims.Num(), 1,
-				[&](int32 Index)
+
+			for ( const TUsdStore< pxr::UsdPrim >& PrototypeGeomMeshPrim : ChildGeomMeshPrims )
+			{
+				pxr::SdfPath PrototypeTargetPrimPath = PrototypeGeomMeshPrim.Get().GetPrimPath();
+
+				if ( Context->CollapsingCache.IsValid() )
 				{
-					const TUsdStore< pxr::UsdPrim >& PrototypeGeomMeshPrim = ChildGeomMeshPrims[Index];
-					Unwound[Index] = UsdGeomPointInstancerTranslatorImpl::UnwindToNonCollapsedPrim( FUsdSchemaTranslator::ECollapsingType::Assets, PrototypeGeomMeshPrim, Context );
+					PrototypeTargetPrimPath = Context->CollapsingCache->UnwindToNonCollapsedPath( UE::FSdfPath{ PrototypeTargetPrimPath }, ECollapsingType::Assets );
 				}
-			);
-			PrototypesPathsSlowTask.EnterProgressFrame();
-			for (int32 Index = 0; Index < ChildGeomMeshPrims.Num(); Index++)
-			{
-				TUsdStore< pxr::SdfPath > PrototypeTargetPrimPath = Unwound[Index];
 
-				const FString UEPrototypeTargetPrimPath = UsdToUnreal::ConvertPath( PrototypeTargetPrimPath.Get() );
+				const FString UEPrototypeTargetPrimPath = UsdToUnreal::ConvertPath( PrototypeTargetPrimPath );
 
 				if ( !ProcessedPrims.Contains( UEPrototypeTargetPrimPath ) )
 				{
 					ProcessedPrims.Add( UEPrototypeTargetPrimPath );
 
-					if ( pxr::UsdPrim PrototypeTargetUsdPrim = Prim.GetStage()->GetPrimAtPath( PrototypeTargetPrimPath.Get() ) )
+					if ( pxr::UsdPrim PrototypeTargetUsdPrim = Prim.GetStage()->GetPrimAtPath( PrototypeTargetPrimPath ) )
 					{
 						pxr::UsdGeomMesh PrototypeGeomMesh( PrototypeTargetUsdPrim );
 						FUsdGeomXformableTranslator PrototypeGeomXformableTranslator( UHierarchicalInstancedStaticMeshComponent::StaticClass(), Context, UE::FUsdTyped( PrototypeGeomMesh ) );
diff --git a/Engine/Plugins/Importers/USDImporter/Source/USDSchemas/Private/USDGeomXformableTranslator.cpp b/Engine/Plugins/Importers/USDImporter/Source/USDSchemas/Private/USDGeomXformableTranslator.cpp
--- a/Engine/Plugins/Importers/USDImporter/Source/USDSchemas/Private/USDGeomXformableTranslator.cpp
+++ b/Engine/Plugins/Importers/USDImporter/Source/USDSchemas/Private/USDGeomXformableTranslator.cpp
@@ -413,6 +413,11 @@ bool FUsdGeomXformableTranslator::CollapsesChildren( ECollapsingType CollapsingT
 {
 	TRACE_CPUPROFILER_EVENT_SCOPE( FUsdGeomXformableTranslator::CollapsesChildren );
 
+	if ( Context->CollapsingCache.IsValid() )
+	{
+		return Context->CollapsingCache->DoesPathCollapseChildren( PrimPath, CollapsingType );
+	}
+
 	bool bCollapsesChildren = false;
 
 	FScopedUsdAllocs UsdAllocs;
@@ -435,8 +440,12 @@ bool FUsdGeomXformableTranslator::CollapsesChildren( ECollapsingType CollapsingT
 		{
 			IUsdSchemasModule& UsdSchemasModule = FModuleManager::Get().LoadModuleChecked< IUsdSchemasModule >( TEXT("USDSchemas") );
 
+			// TODO: This can be optimized in order to make FUsdCollapsingCache::RebuildCacheForSubtree faster: If we have a child prim that we know doesn't collapse,
+			// any of our parents should be able to know they can't collapse *us* either.
+			// This is somewhat niche though: Realistically to waste time here a prim and its children need to have a kind that allows collapsing, and also not be able to collapse.
+			// Also, if any of these prims *does* manage to collapse, FUsdCollapsingCache will already not actually query the subtree children if they can collapse or not anymore,
+			// and just consider them collapsed by the parent
 			TArray< TUsdStore< pxr::UsdPrim > > ChildXformPrims = UsdUtils::GetAllPrimsOfType( Prim, pxr::TfType::Find< pxr::UsdGeomXformable >() );
-
 			for ( const TUsdStore< pxr::UsdPrim >& ChildXformPrim : ChildXformPrims )
 			{
 				if ( TSharedPtr< FUsdSchemaTranslator > SchemaTranslator = UsdSchemasModule.GetTranslatorRegistry().CreateTranslatorForSchema( Context, UE::FUsdTyped( ChildXformPrim.Get() ) ) )
diff --git a/Engine/Plugins/Importers/USDImporter/Source/USDSchemas/Private/USDSchemaTranslator.cpp b/Engine/Plugins/Importers/USDImporter/Source/USDSchemas/Private/USDSchemaTranslator.cpp
--- a/Engine/Plugins/Importers/USDImporter/Source/USDSchemas/Private/USDSchemaTranslator.cpp
+++ b/Engine/Plugins/Importers/USDImporter/Source/USDSchemas/Private/USDSchemaTranslator.cpp
@@ -2,6 +2,7 @@
 
 #include "USDSchemaTranslator.h"
 
+#include "USDCollapsingCache.h"
 #include "USDErrorUtils.h"
 #include "USDSchemasModule.h"
 #include "USDTypesConversion.h"
@@ -203,29 +204,13 @@ bool FUsdSchemaTranslator::IsCollapsed( ECollapsingType CollapsingType ) const
 #if USE_USD_SDK
 	TRACE_CPUPROFILER_EVENT_SCOPE( FUsdSchemaTranslator::IsCollapsed );
 
-	if ( !CanBeCollapsed( CollapsingType ) )
+	if ( Context->CollapsingCache.IsValid() )
 	{
-		return false;
+		return Context->CollapsingCache->IsPathCollapsed( PrimPath, CollapsingType );
 	}
 
-	TUsdStore< pxr::UsdPrim > Prim = pxr::UsdPrim( GetPrim() );
-	TUsdStore< pxr::UsdPrim > ParentPrim = Prim.Get().GetParent();
-
-	IUsdSchemasModule& UsdSchemasModule = FModuleManager::Get().LoadModuleChecked< IUsdSchemasModule >( TEXT("USDSchemas") );
-
-	while ( ParentPrim.Get() )
-	{
-		TSharedPtr< FUsdSchemaTranslator > ParentSchemaTranslator = UsdSchemasModule.GetTranslatorRegistry().CreateTranslatorForSchema( Context, UE::FUsdTyped( ParentPrim.Get() ) );
-
-		if ( ParentSchemaTranslator && ParentSchemaTranslator->CollapsesChildren( CollapsingType ) )
-		{
-			return true;
-		}
-		else
-		{
-			ParentPrim = ParentPrim.Get().GetParent();
-		}
-	}
+	// This is merely a fallback, and we should never need this
+	return CanBeCollapsed( CollapsingType );
 #endif // #if USE_USD_SDK
 
 	return false;
diff --git a/Engine/Plugins/Importers/USDImporter/Source/USDSchemas/Private/USDSkelRootTranslator.cpp b/Engine/Plugins/Importers/USDImporter/Source/USDSchemas/Private/USDSkelRootTranslator.cpp
--- a/Engine/Plugins/Importers/USDImporter/Source/USDSchemas/Private/USDSkelRootTranslator.cpp
+++ b/Engine/Plugins/Importers/USDImporter/Source/USDSchemas/Private/USDSkelRootTranslator.cpp
@@ -840,6 +840,7 @@ namespace UsdSkelRootTranslatorImpl
 
 				if ( pxr::UsdSkelRoot SkeletonRoot{ GetPrim() } )
 				{
+					FString SkelRootPath = PrimPath.GetString();
 					std::vector< pxr::UsdSkelBinding > SkeletonBindings;
 					SkeletonCache.Get().Populate( SkeletonRoot, pxr::UsdTraverseInstanceProxies() );
 					SkeletonCache.Get().ComputeSkelBindings( SkeletonRoot, &SkeletonBindings, pxr::UsdTraverseInstanceProxies() );
@@ -859,7 +860,6 @@ namespace UsdSkelRootTranslatorImpl
 						{
 							continue;
 						}
-						FString SkelAnimationPrimPath = UsdToUnreal::ConvertPath( SkelAnimationPrim.GetPath() );
 
 						if ( !AnimQuery.JointTransformsMightBeTimeVarying() &&
 							( NewBlendShapes.Num() == 0 || !AnimQuery.BlendShapeWeightsMightBeTimeVarying() ) )
@@ -894,7 +894,7 @@ namespace UsdSkelRootTranslatorImpl
 								UUsdAnimSequenceAssetImportData* ImportData = NewObject< UUsdAnimSequenceAssetImportData >( AnimSequence, TEXT( "USDAssetImportData" ) );
 								AnimSequence->AssetImportData = ImportData;
 
-								ImportData->PrimPath = SkelAnimationPrimPath;
+								ImportData->PrimPath = SkelRootPath;
 								ImportData->LayerStartOffsetSeconds = LayerStartOffsetSeconds;
 
 								Context->AssetCache->CacheAsset( HashString, AnimSequence );
@@ -907,7 +907,7 @@ namespace UsdSkelRootTranslatorImpl
 
 						if ( AnimSequence )
 						{
-							Context->AssetCache->LinkAssetToPrim( SkelAnimationPrimPath, AnimSequence );
+							Context->AssetCache->LinkAssetToPrim( SkelRootPath, AnimSequence );
 						}
 					}
 				}
diff --git a/Engine/Plugins/Importers/USDImporter/Source/USDSchemas/Public/USDCollapsingCache.h b/Engine/Plugins/Importers/USDImporter/Source/USDSchemas/Public/USDCollapsingCache.h
new file mode 100644
--- /dev/null
+++ b/Engine/Plugins/Importers/USDImporter/Source/USDSchemas/Public/USDCollapsingCache.h
@@ -0,0 +1,45 @@
+// Copyright Epic Games, Inc. All Rights Reserved.
+
+#pragma once
+
+#include "Containers/Set.h"
+#include "CoreMinimal.h"
+
+#include "UsdWrappers/SdfPath.h"
+
+struct FUsdSchemaTranslationContext;
+namespace UE
+{
+	class FSdfPath;
+	class FUsdPrim;
+}
+
+enum class ECollapsingType
+{
+	Assets,
+	Components
+};
+
+/**
+ * Caches whether given prim paths represent prims that can have their assets or components collapsed when importing/opening
+ * a USD Stage.
+ */
+struct USDSCHEMAS_API FUsdCollapsingCache
+{
+	bool Serialize( FArchive& Ar );
+
+	bool IsPathCollapsed( const UE::FSdfPath& Path, ECollapsingType CollapsingType ) const;
+	bool DoesPathCollapseChildren( const UE::FSdfPath& Path, ECollapsingType CollapsingType ) const;
+
+	// Returns Path in case it represents an uncollapsed prim, or returns the path to the prim that collapsed it
+	UE::FSdfPath UnwindToNonCollapsedPath( const UE::FSdfPath& Path, ECollapsingType CollapsingType ) const;
+
+	void RebuildCacheForSubtree( const UE::FUsdPrim& Prim, FUsdSchemaTranslationContext& Context );
+
+	void Clear();
+	bool IsEmpty();
+
+private:
+	TMap< UE::FSdfPath, UE::FSdfPath > AssetPathsToCollapsedRoot;
+	TMap< UE::FSdfPath, UE::FSdfPath > ComponentPathsToCollapsedRoot;
+};
\ No newline at end of file
diff --git a/Engine/Plugins/Importers/USDImporter/Source/USDSchemas/Public/USDGeomPointInstancerTranslator.h b/Engine/Plugins/Importers/USDImporter/Source/USDSchemas/Public/USDGeomPointInstancerTranslator.h
--- a/Engine/Plugins/Importers/USDImporter/Source/USDSchemas/Public/USDGeomPointInstancerTranslator.h
+++ b/Engine/Plugins/Importers/USDImporter/Source/USDSchemas/Public/USDGeomPointInstancerTranslator.h
@@ -20,7 +20,7 @@ public:
 	using FUsdGeomXformableTranslator::FUsdGeomXformableTranslator;
 
 	virtual void UpdateComponents( USceneComponent* SceneComponent ) override;
-	virtual bool CollapsesChildren( ECollapsingType CollapsingType ) const override { return CollapsingType == FUsdSchemaTranslator::ECollapsingType::Components; }
+	virtual bool CollapsesChildren( ECollapsingType CollapsingType ) const override { return CollapsingType == ECollapsingType::Components; }
 	virtual bool CanBeCollapsed( ECollapsingType CollapsingType ) const override { return false; }
 };
 
diff --git a/Engine/Plugins/Importers/USDImporter/Source/USDSchemas/Public/USDSchemaTranslator.h b/Engine/Plugins/Importers/USDImporter/Source/USDSchemas/Public/USDSchemaTranslator.h
--- a/Engine/Plugins/Importers/USDImporter/Source/USDSchemas/Public/USDSchemaTranslator.h
+++ b/Engine/Plugins/Importers/USDImporter/Source/USDSchemas/Public/USDSchemaTranslator.h
@@ -6,6 +6,7 @@
 
 #include "UnrealUSDWrapper.h"
 #include "USDAssetCache.h"
+#include "USDCollapsingCache.h"
 #include "USDMemory.h"
 #include "USDSkeletalDataConversion.h"
 
@@ -153,6 +154,8 @@ struct USDSCHEMAS_API FUsdSchemaTranslationContext : public TSharedFromThis< FUs
 	/** Where the translated assets will be stored */
 	TStrongObjectPtr< UUsdAssetCache > AssetCache;
 
+	/** Caches the collapsing state of prim paths */
+	TSharedPtr<FUsdCollapsingCache> CollapsingCache;
 
 	/** Where we place imported blend shapes, if available */
 	UsdUtils::FBlendShapeMap* BlendShapesByPath = nullptr;
@@ -232,12 +235,6 @@ public:
 	virtual USceneComponent* CreateComponents() { return nullptr; }
 	virtual void UpdateComponents( USceneComponent* SceneComponent ) {}
 
-	enum class ECollapsingType
-	{
-		Assets,
-		Components
-	};
-
 	virtual bool CollapsesChildren( ECollapsingType CollapsingType ) const { return false; }
 
 	bool IsCollapsed( ECollapsingType CollapsingType ) const;
diff --git a/Engine/Plugins/Importers/USDImporter/Source/USDStage/Private/USDStageActor.cpp b/Engine/Plugins/Importers/USDImporter/Source/USDStage/Private/USDStageActor.cpp
--- a/Engine/Plugins/Importers/USDImporter/Source/USDStage/Private/USDStageActor.cpp
+++ b/Engine/Plugins/Importers/USDImporter/Source/USDStage/Private/USDStageActor.cpp
@@ -6,6 +6,7 @@
 #include "USDAssetCache.h"
 #include "USDAssetImportData.h"
 #include "USDClassesModule.h"
+#include "USDCollapsingCache.h"
 #include "USDConversionUtils.h"
 #include "USDErrorUtils.h"
 #include "USDGeomMeshConversion.h"
@@ -95,6 +96,7 @@ struct FUsdStageActorImpl
 		TranslationContext->RenderContext = StageActor->RenderContext;
 		TranslationContext->MaterialToPrimvarToUVIndex = &StageActor->MaterialToPrimvarToUVIndex;
 		TranslationContext->BlendShapesByPath = &StageActor->BlendShapesByPath;
+		TranslationContext->CollapsingCache = StageActor->CollapsingCache;
 
 		// Its more convenient to toggle between variants using the USDStage window, as opposed to parsing LODs
 		TranslationContext->bAllowInterpretingLODs = false;
@@ -253,49 +255,6 @@ struct FUsdStageActorImpl
 		UnrealUSDWrapper::EraseStageFromCache( Stage );
 	}
 
-	static UE::FSdfPath UnwindToNonCollapsedPrim( AUsdStageActor* StageActor, const FString& InPrimPath, FUsdSchemaTranslator::ECollapsingType CollapsingType )
-	{
-		IUsdSchemasModule& UsdSchemasModule = FModuleManager::Get().LoadModuleChecked< IUsdSchemasModule >( TEXT( "USDSchemas" ) );
-
-		TSharedRef< FUsdSchemaTranslationContext > TranslationContext = FUsdStageActorImpl::CreateUsdSchemaTranslationContext( StageActor, InPrimPath );
-
-		UE::FUsdStage UsdStage = StageActor->GetOrLoadUsdStage();
-		UE::FSdfPath UsdPrimPath( *InPrimPath );
-		UE::FUsdPrim UsdPrim = UsdStage.GetPrimAtPath( UsdPrimPath );
-
-		if ( TSharedPtr< FUsdSchemaTranslator > SchemaTranslator = UsdSchemasModule.GetTranslatorRegistry().CreateTranslatorForSchema( TranslationContext, UE::FUsdTyped( UsdPrim ) ) )
-		{
-			while ( SchemaTranslator->IsCollapsed( CollapsingType ) )
-			{
-				UE::FSdfPath ParentUsdPrimPath = UsdPrimPath.GetParentPath();
-				UE::FUsdPrim ParentUsdPrim = UsdStage.GetPrimAtPath( ParentUsdPrimPath );
-				if ( ParentUsdPrim.IsPseudoRoot() )
-				{
-					// It doesn't matter if we're collapsed when our parent is the root: We'll be a separate component/asset anyway.
-					// At that point we don't want to return "/" from this function though, so break here
-					break;
-				}
-
-				UsdPrimPath = ParentUsdPrimPath;
-				UsdPrim = ParentUsdPrim;
-
-				TranslationContext = FUsdStageActorImpl::CreateUsdSchemaTranslationContext( StageActor, UsdPrimPath.GetString() );
-				TSharedPtr< FUsdSchemaTranslator > ParentSchemaTranslator = UsdSchemasModule.GetTranslatorRegistry().CreateTranslatorForSchema( TranslationContext, UE::FUsdTyped( UsdPrim ) );
-
-				// Note how we continue looping with the child translator when the parent has an invalid translator.
-				// This is intentional: If our parent here has no valid translator, it will always be collapsed, so we need to check
-				// whether we have a valid schema translator for our grandparent. If that one is valid and collapses, both us and our parent will
-				// be collapsed into the grandparent, and that's the path we need to return
-				if ( ParentSchemaTranslator.IsValid() )
-				{
-					SchemaTranslator = ParentSchemaTranslator;
-				}
-			}
-		}
-
-		return UsdPrimPath;
-	};
-
 	static bool ObjectNeedsMultiUserTag( UObject* Object, AUsdStageActor* StageActor )
 	{
 		// Don't need to tag non-transient stuff
@@ -828,6 +787,8 @@ void AUsdStageActor::OnUsdObjectsChanged( const UsdUtils::FObjectChangesByPath&
 	// back to USD when these objects change though.
 	FScopedBlockMonitoringChangesForTransaction BlockMonitoring{ LevelSequenceHelper };
 
+	bool bHasResync = ResyncChanges.Num() > 0;
+
 	// The most important thing here is to iterate in parent to child order, so build SortedPrimsChangedList
 	TMap< FString, bool > SortedPrimsChangedList;
 	for ( const TPair<FString, TArray<UsdUtils::FObjectChangeNotice>>& InfoChange : InfoChanges )
@@ -846,6 +807,7 @@ void AUsdStageActor::OnUsdObjectsChanged( const UsdUtils::FObjectChangesByPath&
 					if ( ResyncProperties.Contains( AttributeChange.PropertyName ) )
 					{
 						bIsResync = true;
+						bHasResync = true;
 					}
 				}
 			}
@@ -881,18 +843,29 @@ void AUsdStageActor::OnUsdObjectsChanged( const UsdUtils::FObjectChangesByPath&
 
 	FScopedUsdMessageLog ScopedMessageLog;
 
-	TSet< FString > UpdatedAssets;
-	TSet< FString > ResyncedAssets;
-	TSet< FString > UpdatedComponents;
-	TSet< FString > ResyncedComponents;
+	TSet< UE::FSdfPath > UpdatedAssets;
+	TSet< UE::FSdfPath > ResyncedAssets;
+	TSet< UE::FSdfPath > UpdatedComponents;
+	TSet< UE::FSdfPath > ResyncedComponents;
 
 	bool bDeselected = false;
 
+	if ( bHasResync && CollapsingCache.IsValid() )
+	{
+		// TODO: Selective rebuild of only the required parts of the cache.
+		// If a prim changes from CanBeCollapsed to not (or vice-versa), that means its parent may change, and its grandparent, etc. so we'd need to check a large
+		// part of the tree (here we're just rebuilding the whole thing for now).
+		// However, if we know that only Prim resynced, we can traverse the tree root down and if we reach Prim and its collapsing state
+		// hasn't updated from before, we don't have to update its subtree at all, or sibling subtrees.
+		TSharedRef< FUsdSchemaTranslationContext > TranslationContext = FUsdStageActorImpl::CreateUsdSchemaTranslationContext( this, TEXT( "/" ) );
+		CollapsingCache->RebuildCacheForSubtree( UsdStage.GetPseudoRoot(), TranslationContext.Get() );
+	}
+
 	for ( const TPair< FString, bool >& PrimChangedInfo : SortedPrimsChangedList )
 	{
 		RefreshStageTask.EnterProgressFrame();
 
-		const FString& PrimPath = PrimChangedInfo.Key;
+		const UE::FSdfPath PrimPath = UE::FSdfPath{ *PrimChangedInfo.Key };
 		const bool bIsResync = PrimChangedInfo.Value;
 
 		if ( bIsResync && !bDeselected )
@@ -902,21 +875,21 @@ void AUsdStageActor::OnUsdObjectsChanged( const UsdUtils::FObjectChangesByPath&
 		}
 
 		// Return if the path or any of its higher level paths are already processed
-		auto IsPathAlreadyProcessed = []( TSet< FString >& PathsProcessed, FString PathToProcess ) -> bool
+		auto IsPathAlreadyProcessed = []( TSet< UE::FSdfPath >& PathsProcessed, UE::FSdfPath PathToProcess ) -> bool
 		{
 			FString SubPath;
 			FString ParentPath;
 
-			if ( PathsProcessed.Contains( TEXT("/") ) )
+			if ( PathsProcessed.Contains( UE::FSdfPath::AbsoluteRootPath() ) )
 			{
 				return true;
 			}
 
 			while ( !PathToProcess.IsEmpty() && !PathsProcessed.Contains( PathToProcess ) )
 			{
-				if ( PathToProcess.Split( TEXT("/"), &ParentPath, &SubPath, ESearchCase::IgnoreCase, ESearchDir::FromEnd ) )
+				if ( !PathToProcess.IsAbsoluteRootPath() )
 				{
-					PathToProcess = ParentPath;
+					PathToProcess = PathToProcess.GetParentPath();
 				}
 				else
 				{
@@ -927,35 +900,39 @@ void AUsdStageActor::OnUsdObjectsChanged( const UsdUtils::FObjectChangesByPath&
 			return !PathToProcess.IsEmpty() && PathsProcessed.Contains( PathToProcess );
 		};
 
-		auto UpdateComponents = [&]( const FString& InPrimPath, const bool bInResync )
+		auto UpdateComponents = [&]( const UE::FSdfPath& InPrimPath, const bool bInResync )
 		{
-			UE::FSdfPath ComponentsPrimPath = FUsdStageActorImpl::UnwindToNonCollapsedPrim( this, InPrimPath, FUsdSchemaTranslator::ECollapsingType::Components );
+			UE::FSdfPath ComponentsPrimPath = CollapsingCache.IsValid()
+				? CollapsingCache->UnwindToNonCollapsedPath( InPrimPath, ECollapsingType::Components )
+				: InPrimPath;
 
-			TSet< FString >& RefreshedComponents = bInResync ? ResyncedComponents : UpdatedComponents;
+			TSet< UE::FSdfPath >& RefreshedComponents = bInResync ? ResyncedComponents : UpdatedComponents;
 
-			if ( !IsPathAlreadyProcessed( RefreshedComponents, ComponentsPrimPath.GetString() ) )
+			if ( !IsPathAlreadyProcessed( RefreshedComponents, ComponentsPrimPath ) )
 			{
 				TSharedRef< FUsdSchemaTranslationContext > TranslationContext = FUsdStageActorImpl::CreateUsdSchemaTranslationContext( this, ComponentsPrimPath.GetString() );
 				UpdatePrim( ComponentsPrimPath, bInResync, *TranslationContext );
 				TranslationContext->CompleteTasks();
 
-				RefreshedComponents.Add( ComponentsPrimPath.GetString() );
+				RefreshedComponents.Add( ComponentsPrimPath );
 
 				if ( bInResync )
 				{
 					// Consider that the path has been updated in the case of a resync
-					UpdatedComponents.Add( ComponentsPrimPath.GetString() );
+					UpdatedComponents.Add( ComponentsPrimPath );
 				}
 			}
 		};
 
-		auto ReloadAssets = [&]( const FString& InPrimPath, const bool bInResync )
+		auto ReloadAssets = [&]( const UE::FSdfPath& InPrimPath, const bool bInResync )
 		{
-			UE::FSdfPath AssetsPrimPath = FUsdStageActorImpl::UnwindToNonCollapsedPrim( this, InPrimPath, FUsdSchemaTranslator::ECollapsingType::Assets );
+			UE::FSdfPath AssetsPrimPath = CollapsingCache.IsValid()
+				? CollapsingCache->UnwindToNonCollapsedPath( InPrimPath, ECollapsingType::Assets )
+				: InPrimPath;
 
-			TSet< FString >& RefreshedAssets = bInResync ? ResyncedAssets : UpdatedAssets;
+			TSet< UE::FSdfPath >& RefreshedAssets = bInResync ? ResyncedAssets : UpdatedAssets;
 
-			if ( !IsPathAlreadyProcessed( RefreshedAssets, AssetsPrimPath.GetString() ) )
+			if ( !IsPathAlreadyProcessed( RefreshedAssets, AssetsPrimPath ) )
 			{
 				TSharedRef< FUsdSchemaTranslationContext > TranslationContext = FUsdStageActorImpl::CreateUsdSchemaTranslationContext( this, AssetsPrimPath.GetString() );
 
@@ -977,20 +954,20 @@ void AUsdStageActor::OnUsdObjectsChanged( const UsdUtils::FObjectChangesByPath&
 						for ( const FString& MaterialUserPrim : UsdUtils::GetMaterialUsers( PrimToResync ) )
 						{
 							const bool bResyncComponent = true; // We need to force resync to reassign materials
-							UpdateComponents( MaterialUserPrim, bResyncComponent );
+							UpdateComponents( UE::FSdfPath{ *MaterialUserPrim }, bResyncComponent );
 						}
 					}
 #endif // USE_USD_SDK
 
 					// Resyncing also includes "updating" the prim
-					UpdatedAssets.Add( AssetsPrimPath.GetString() );
+					UpdatedAssets.Add( AssetsPrimPath );
 				}
 				else
 				{
 					LoadAsset( *TranslationContext, GetOrLoadUsdStage().GetPrimAtPath( AssetsPrimPath ) );
 				}
 
-				RefreshedAssets.Add( AssetsPrimPath.GetString() );
+				RefreshedAssets.Add( AssetsPrimPath );
 			}
 		};
 
@@ -1109,7 +1086,7 @@ UUsdPrimTwin* AUsdStageActor::ExpandPrim( const UE::FUsdPrim& Prim, FUsdSchemaTr
 			SchemaTranslator->UpdateComponents( UsdPrimTwin->SceneComponent.Get() );
 		}
 
-		bExpandChildren = !SchemaTranslator->CollapsesChildren( FUsdSchemaTranslator::ECollapsingType::Components );
+		bExpandChildren = !SchemaTranslator->CollapsesChildren( ECollapsingType::Components );
 	}
 
 	if ( bExpandChildren )
@@ -1338,7 +1315,23 @@ void AUsdStageActor::SetTime(float InTime)
 
 USceneComponent* AUsdStageActor::GetGeneratedComponent( const FString& PrimPath )
 {
-	FString UncollapsedPath = FUsdStageActorImpl::UnwindToNonCollapsedPrim( this, PrimPath, FUsdSchemaTranslator::ECollapsingType::Components ).GetString();
+	if ( !UsdStage )
+	{
+		return nullptr;
+	}
+
+	// We can't query our CollapsingCache with invalid paths, as we're using ensures to track when we miss the cache (which shouldn't ever happen)
+	UE::FSdfPath UsdPath{ *PrimPath };
+	if ( !UsdStage.GetPrimAtPath( UsdPath ) )
+	{
+		return nullptr;
+	}
+
+	FString UncollapsedPath = PrimPath;
+	if ( CollapsingCache.IsValid() )
+	{
+		UncollapsedPath = CollapsingCache->UnwindToNonCollapsedPath( UsdPath, ECollapsingType::Components ).GetString();
+	}
 
 	if ( UUsdPrimTwin* UsdPrimTwin = GetRootPrimTwin()->Find( UncollapsedPath ) )
 	{
@@ -1350,12 +1343,28 @@ USceneComponent* AUsdStageActor::GetGeneratedComponent( const FString& PrimPath
 
 TArray<UObject*> AUsdStageActor::GetGeneratedAssets( const FString& PrimPath )
 {
+	if ( !UsdStage )
+	{
+		return {};
+	}
+
+	// We can't query our CollapsingCache with invalid paths, as we're using ensures to track when we miss the cache (which shouldn't ever happen)
+	UE::FSdfPath UsdPath{ *PrimPath };
+	if ( !UsdStage.GetPrimAtPath( UsdPath ) )
+	{
+		return {};
+	}
+
 	if ( !AssetCache )
 	{
 		return {};
 	}
 
-	FString UncollapsedPath = FUsdStageActorImpl::UnwindToNonCollapsedPrim( this, PrimPath, FUsdSchemaTranslator::ECollapsingType::Assets ).GetString();
+	FString UncollapsedPath = PrimPath;
+	if ( CollapsingCache.IsValid() )
+	{
+		UncollapsedPath = CollapsingCache->UnwindToNonCollapsedPath( UsdPath, ECollapsingType::Assets ).GetString();
+	}
 
 	TSet<UObject*> Result;
 
@@ -1576,6 +1585,9 @@ void AUsdStageActor::OnObjectsReplaced( const TMap<UObject*, UObject*>& ObjectRe
 			AssetCache->Rename( nullptr, NewActor );
 			NewActor->AssetCache = AssetCache;
 
+			NewActor->CollapsingCache = CollapsingCache;
+			CollapsingCache = nullptr;
+
 			// It could be that we're automatically recompiling when going into PIE because our blueprint was dirty.
 			// In that case we also need bIsTransitioningIntoPIE to be true to prevent us from calling LoadUsdStage from PostRegisterAllComponents
 			NewActor->bIsTransitioningIntoPIE = bIsTransitioningIntoPIE;
@@ -1657,7 +1669,14 @@ void AUsdStageActor::LoadUsdStage()
 
 	TSharedRef< FUsdSchemaTranslationContext > TranslationContext = FUsdStageActorImpl::CreateUsdSchemaTranslationContext( this, RootTwin->PrimPath );
 
-	SlowTask.EnterProgressFrame( 0.8f );
+	SlowTask.EnterProgressFrame( 0.1f );
+	if ( !CollapsingCache.IsValid() )
+	{
+		CollapsingCache = MakeShared<FUsdCollapsingCache>();
+	}
+	CollapsingCache->RebuildCacheForSubtree( UsdStage.GetPseudoRoot(), TranslationContext.Get() );
+
+	SlowTask.EnterProgressFrame( 0.7f );
 	LoadAssets( *TranslationContext, UsdStage.GetPseudoRoot() );
 
 	SlowTask.EnterProgressFrame( 0.2f );
@@ -1720,6 +1739,11 @@ void AUsdStageActor::UnloadUsdStage()
 		AssetCache->Reset();
 	}
 
+	if ( CollapsingCache )
+	{
+		CollapsingCache->Clear();
+	}
+
 	ObjectsToWatch.Reset();
 	BlendShapesByPath.Reset();
 	MaterialToPrimvarToUVIndex.Reset();
@@ -2034,6 +2058,16 @@ void AUsdStageActor::Serialize(FArchive& Ar)
 		Ar << MaterialToPrimvarToUVIndex;
 		Ar << bIsTransitioningIntoPIE;
 	}
+
+	if ( ( Ar.GetPortFlags() & PPF_DuplicateForPIE ) || Ar.IsTransacting() )
+	{
+		if ( !CollapsingCache.IsValid() )
+		{
+			CollapsingCache = MakeShared<FUsdCollapsingCache>();
+		}
+
+		CollapsingCache->Serialize( Ar );
+	}
 }
 
 void AUsdStageActor::Destroyed()
@@ -2570,7 +2604,7 @@ void AUsdStageActor::LoadAssets( FUsdSchemaTranslationContext& TranslationContex
 	{
 		if ( TSharedPtr< FUsdSchemaTranslator > SchemaTranslator = UsdSchemasModule.GetTranslatorRegistry().CreateTranslatorForSchema( TranslationContext.AsShared(), UE::FUsdTyped( UsdPrim ) ) )
 		{
-			return SchemaTranslator->CollapsesChildren( FUsdSchemaTranslator::ECollapsingType::Assets );
+			return SchemaTranslator->CollapsesChildren( ECollapsingType::Assets );
 		}
 
 		return false;
diff --git a/Engine/Plugins/Importers/USDImporter/Source/USDStage/Public/USDStageActor.h b/Engine/Plugins/Importers/USDImporter/Source/USDStage/Public/USDStageActor.h
--- a/Engine/Plugins/Importers/USDImporter/Source/USDStage/Public/USDStageActor.h
+++ b/Engine/Plugins/Importers/USDImporter/Source/USDStage/Public/USDStageActor.h
@@ -6,6 +6,7 @@
 #include "Misc/SecureHash.h"
 #include "Templates/Function.h"
 
+#include "USDCollapsingCache.h"
 #include "USDLevelSequenceHelper.h"
 #include "USDListener.h"
 #include "USDPrimTwin.h"
@@ -174,6 +175,7 @@ public:
 	void ReloadAnimations();
 	float GetTime() { return Time; }
 	UUsdAssetCache* GetAssetCache() { return AssetCache; }
+	TSharedPtr<FUsdCollapsingCache> GetCollapsingCache() { return CollapsingCache; }
 	TMap< FString, TMap< FString, int32 > > GetMaterialToPrimvarToUVIndex() { return MaterialToPrimvarToUVIndex; }
 
 public:
@@ -236,20 +238,9 @@ private:
 	UPROPERTY( VisibleAnywhere, Category = "USD", AdvancedDisplay )
 	UUsdAssetCache* AssetCache;
 
-private:
 	UPROPERTY( Transient )
 	UUsdTransactor* Transactor;
 
-	/** Keep track of blend shapes so that we can map 'inbetween shapes' to their separate morph targets when animating */
-	UsdUtils::FBlendShapeMap BlendShapesByPath;
-
-	/**
-	 * When parsing materials, we keep track of which primvar we mapped to which UV channel.
-	 * When parsing meshes later, we use this data to place the correct primvar values in each UV channel.
-	 * We keep this here as these are generated when the materials stored in AssetsCache are parsed, so it should accompany them
-	 */
-	TMap< FString, TMap< FString, int32 > > MaterialToPrimvarToUVIndex;
-
 public:
 	UE_DEPRECATED( 4.27, "Use the const version if you don't wish to load the stage on-demand, or use GetOrLoadUsdStage if you do" )
 	USDSTAGE_API UE::FUsdStage& GetUsdStage();
@@ -284,6 +275,18 @@ protected:
 	void AnimatePrims();
 
 private:
+	TSharedPtr<FUsdCollapsingCache> CollapsingCache;
+
+	/** Keep track of blend shapes so that we can map 'inbetween shapes' to their separate morph targets when animating */
+	UsdUtils::FBlendShapeMap BlendShapesByPath;
+
+	/**
+	 * When parsing materials, we keep track of which primvar we mapped to which UV channel.
+	 * When parsing meshes later, we use this data to place the correct primvar values in each UV channel.
+	 * We keep this here as these are generated when the materials stored in AssetsCache are parsed, so it should accompany them
+	 */
+	TMap< FString, TMap< FString, int32 > > MaterialToPrimvarToUVIndex;
+
 	UE::FUsdStage UsdStage;
 
 	/**
diff --git a/Engine/Plugins/Importers/USDImporter/Source/USDStage/USDStage.Build.cs b/Engine/Plugins/Importers/USDImporter/Source/USDStage/USDStage.Build.cs
--- a/Engine/Plugins/Importers/USDImporter/Source/USDStage/USDStage.Build.cs
+++ b/Engine/Plugins/Importers/USDImporter/Source/USDStage/USDStage.Build.cs
@@ -9,6 +9,13 @@ namespace UnrealBuildTool.Rules
 	{
 		public USDStage(ReadOnlyTargetRules Target) : base(Target)
 		{
+			PublicDependencyModuleNames.AddRange(
+				new string[]
+				{
+					"USDSchemas" // Has to be a public dependency because the stage actor has a collapsing cache member
+				}
+			);
+
 			PrivateDependencyModuleNames.AddRange(
 				new string[]
 				{
@@ -26,7 +33,6 @@ namespace UnrealBuildTool.Rules
 					"StaticMeshDescription",
 					"UnrealUSDWrapper",
 					"USDClasses",
-					"USDSchemas",
 					"USDUtilities",
 				}
 			);
diff --git a/Engine/Plugins/Importers/USDImporter/Source/USDStageEditorViewModels/Private/USDStageViewModel.cpp b/Engine/Plugins/Importers/USDImporter/Source/USDStageEditorViewModels/Private/USDStageViewModel.cpp
--- a/Engine/Plugins/Importers/USDImporter/Source/USDStageEditorViewModels/Private/USDStageViewModel.cpp
+++ b/Engine/Plugins/Importers/USDImporter/Source/USDStageEditorViewModels/Private/USDStageViewModel.cpp
@@ -307,6 +307,7 @@ void FUsdStageViewModel::ImportStage()
 			// Let the importer reuse our assets, but force it to spawn new actors and components always
 			// This allows a different setting for asset/component collapsing, and doesn't require modifying the PrimTwins
 			ImportContext.AssetCache = StageActor->GetAssetCache();
+			ImportContext.CollapsingCache = StageActor->GetCollapsingCache();
 			ImportContext.LevelSequenceHelper.SetAssetCache( StageActor->GetAssetCache() );
 			ImportContext.MaterialToPrimvarToUVIndex = StageActor->GetMaterialToPrimvarToUVIndex();
 
diff --git a/Engine/Plugins/Importers/USDImporter/Source/USDStageImporter/Private/USDStageAssetImportFactory.cpp b/Engine/Plugins/Importers/USDImporter/Source/USDStageImporter/Private/USDStageAssetImportFactory.cpp
--- a/Engine/Plugins/Importers/USDImporter/Source/USDStageImporter/Private/USDStageAssetImportFactory.cpp
+++ b/Engine/Plugins/Importers/USDImporter/Source/USDStageImporter/Private/USDStageAssetImportFactory.cpp
@@ -102,7 +102,7 @@ bool UUsdStageAssetImportFactory::FactoryCanImport(const FString& Filename)
 
 void UUsdStageAssetImportFactory::CleanUp()
 {
-	ImportContext = FUsdStageImportContext();
+	ImportContext.Reset();
 	Super::CleanUp();
 }
 
diff --git a/Engine/Plugins/Importers/USDImporter/Source/USDStageImporter/Private/USDStageImportContext.cpp b/Engine/Plugins/Importers/USDImporter/Source/USDStageImporter/Private/USDStageImportContext.cpp
--- a/Engine/Plugins/Importers/USDImporter/Source/USDStageImporter/Private/USDStageImportContext.cpp
+++ b/Engine/Plugins/Importers/USDImporter/Source/USDStageImporter/Private/USDStageImportContext.cpp
@@ -17,24 +17,7 @@
 
 FUsdStageImportContext::FUsdStageImportContext()
 {
-	World = nullptr;
-	ImportOptions = NewObject< UUsdStageImportOptions >();
-
-	// Always reading from the stage cache is a good default because while we can have multiple instances of the
-	// same stage open, USD will only open a particular layer once. If we try importing without using the stage
-	// cache and the stage we want to import uses an existing open layer, we will forcibly reload
-	// that layer (check UnrealUSDWrapper::OpenStage), which would erase our previous changes to it and lead to modifications
-	// on existing open stages (e.g. we have cube.usda open with a stage actor and with some changes and we click
-	// file -> import into level and import cube.usda again)
-	bReadFromStageCache = true;
-
-	bStageWasOriginallyOpenInCache = false;
-	SceneActor = nullptr;
-	ImportedAsset = nullptr;
-	AssetCache = nullptr;
-	OriginalMetersPerUnit = 0.01f;
-	TargetSceneActorTargetTransform = FTransform::Identity;
-	TargetSceneActorAttachParent = nullptr;
+	Reset();
 }
 
 bool FUsdStageImportContext::Init(const FString& InName, const FString& InFilePath, const FString& InInitialPackagePath, EObjectFlags InFlags, bool bInIsAutomated, bool bIsReimport, bool bAllowActorImport)
@@ -85,3 +68,35 @@ bool FUsdStageImportContext::Init(const FString& InName, const FString& InFilePa
 
 	return true;
 }
+
+void FUsdStageImportContext::Reset()
+{
+	World = nullptr;
+	TargetSceneActorTargetTransform = FTransform::Identity;
+	TargetSceneActorAttachParent = nullptr;
+	SceneActor = nullptr;
+	ObjectName = FString{};
+	PackagePath = FString{};
+	FilePath = FString{};
+	ImportOptions = NewObject< UUsdStageImportOptions >();
+	ImportedAsset = nullptr;
+	LevelSequenceHelper.Clear();
+	AssetCache = nullptr;
+	CollapsingCache = nullptr;
+	MaterialToPrimvarToUVIndex.Empty();
+	Stage = UE::FUsdStage{};
+	ImportObjectFlags = EObjectFlags::RF_NoFlags;
+	bIsAutomated = false;
+
+	// Always reading from the stage cache is a good default because while we can have multiple instances of the
+	// same stage open, USD will only open a particular layer once. If we try importing without using the stage
+	// cache and the stage we want to import uses an existing open layer, we will forcibly reload
+	// that layer (check UnrealUSDWrapper::OpenStage), which would erase our previous changes to it and lead to modifications
+	// on existing open stages (e.g. we have cube.usda open with a stage actor and with some changes and we click
+	// file -> import into level and import cube.usda again)
+	bReadFromStageCache = true;
+
+	bStageWasOriginallyOpenInCache = false;
+	OriginalMetersPerUnit = 0.01f;
+	OriginalUpAxis = EUsdUpAxis::ZAxis;
+}
\ No newline at end of file
diff --git a/Engine/Plugins/Importers/USDImporter/Source/USDStageImporter/Private/USDStageImportFactory.cpp b/Engine/Plugins/Importers/USDImporter/Source/USDStageImporter/Private/USDStageImportFactory.cpp
--- a/Engine/Plugins/Importers/USDImporter/Source/USDStageImporter/Private/USDStageImportFactory.cpp
+++ b/Engine/Plugins/Importers/USDImporter/Source/USDStageImporter/Private/USDStageImportFactory.cpp
@@ -93,7 +93,7 @@ bool UUsdStageImportFactory::FactoryCanImport(const FString& Filename)
 
 void UUsdStageImportFactory::CleanUp()
 {
-	ImportContext = FUsdStageImportContext();
+	ImportContext.Reset();
 	Super::CleanUp();
 }
 
diff --git a/Engine/Plugins/Importers/USDImporter/Source/USDStageImporter/Private/USDStageImporter.cpp b/Engine/Plugins/Importers/USDImporter/Source/USDStageImporter/Private/USDStageImporter.cpp
--- a/Engine/Plugins/Importers/USDImporter/Source/USDStageImporter/Private/USDStageImporter.cpp
+++ b/Engine/Plugins/Importers/USDImporter/Source/USDStageImporter/Private/USDStageImporter.cpp
@@ -215,6 +215,19 @@ namespace UsdStageImporterImpl
 		TranslationContext.CompleteTasks();
 	}
 
+	void CacheCollapsingState( FUsdSchemaTranslationContext& TranslationContext )
+	{
+		if ( !TranslationContext.CollapsingCache.IsValid() )
+		{
+			TranslationContext.CollapsingCache = MakeShared<FUsdCollapsingCache>();
+		}
+
+		if ( TranslationContext.CollapsingCache->IsEmpty() )
+		{
+			TranslationContext.CollapsingCache->RebuildCacheForSubtree( TranslationContext.Stage.GetPseudoRoot(), TranslationContext );
+		}
+	}
+
 	void ImportMaterials(FUsdStageImportContext& ImportContext, FUsdSchemaTranslationContext& TranslationContext)
 	{
 		TRACE_CPUPROFILER_EVENT_SCOPE( ImportMaterials );
@@ -245,7 +258,7 @@ namespace UsdStageImporterImpl
 		{
 			if (TSharedPtr< FUsdSchemaTranslator > SchemaTranslator = UsdSchemasModule.GetTranslatorRegistry().CreateTranslatorForSchema(TranslationContext.AsShared(), UE::FUsdTyped(UsdPrim)))
 			{
-				return SchemaTranslator->CollapsesChildren(FUsdSchemaTranslator::ECollapsingType::Assets);
+				return SchemaTranslator->CollapsesChildren(ECollapsingType::Assets);
 			}
 
 			return false;
@@ -288,7 +301,7 @@ namespace UsdStageImporterImpl
 		{
 			Component = SchemaTranslator->CreateComponents();
 
-			bExpandChilren = !SchemaTranslator->CollapsesChildren(FUsdSchemaTranslator::ECollapsingType::Components);
+			bExpandChilren = !SchemaTranslator->CollapsesChildren(ECollapsingType::Components);
 		}
 
 		SlowTask.EnterProgressFrame();
@@ -1035,6 +1048,7 @@ namespace UsdStageImporterImpl
 		// Always discard the context's reference to the stage because it may be a persistent import context (like
 		// the non-static data member of UUsdStageImportFactory
 		ImportContext.Stage = UE::FUsdStage();
+		ImportContext.LevelSequenceHelper.Init( UE::FUsdStage() );
 #endif // #if USE_USD_SDK
 	}
 
@@ -1429,19 +1443,19 @@ void UUsdStageImporter::ImportFromFile(FUsdStageImportContext& ImportContext)
 		return;
 	}
 
-	UsdStageImporterImpl::SetupSceneActor(ImportContext);
-	if (!ImportContext.SceneActor && ImportContext.ImportOptions->bImportActors)
+	UsdStageImporterImpl::SetupSceneActor( ImportContext );
+	if ( !ImportContext.SceneActor && ImportContext.ImportOptions->bImportActors )
 	{
 		return;
 	}
 
 	FUsdDelegates::OnPreUsdImport.Broadcast( ImportContext.FilePath );
 
-	AActor* ExistingSceneActor = UsdStageImporterImpl::GetExistingSceneActor(ImportContext);
+	AActor* ExistingSceneActor = UsdStageImporterImpl::GetExistingSceneActor( ImportContext );
 
-	UsdStageImporterImpl::SetupStageForImport(ImportContext);
+	UsdStageImporterImpl::SetupStageForImport( ImportContext );
 
-	ImportContext.LevelSequenceHelper.Init(ImportContext.Stage);
+	ImportContext.LevelSequenceHelper.Init( ImportContext.Stage );
 
 	TMap<FSoftObjectPath, FSoftObjectPath> SoftObjectsToRemap;
 	TMap<UObject*, UObject*> ObjectsToRemap;
@@ -1472,8 +1486,10 @@ void UUsdStageImporter::ImportFromFile(FUsdStageImportContext& ImportContext)
 	TranslationContext->bAllowInterpretingLODs = ImportContext.ImportOptions->bInterpretLODs;
 	TranslationContext->bAllowParsingSkeletalAnimations = ImportContext.ImportOptions->bImportGeometry && ImportContext.ImportOptions->bImportSkeletalAnimations;
 	TranslationContext->MaterialToPrimvarToUVIndex = &ImportContext.MaterialToPrimvarToUVIndex;
+	TranslationContext->CollapsingCache = ImportContext.CollapsingCache;
 	TranslationContext->BlendShapesByPath = &BlendShapesByPath;
 	{
+		UsdStageImporterImpl::CacheCollapsingState( TranslationContext.Get() );
 		UsdStageImporterImpl::ImportMaterials( ImportContext, TranslationContext.Get() );
 		UsdStageImporterImpl::ImportMeshes( ImportContext, TranslationContext.Get() );
 		UsdStageImporterImpl::ImportActors( ImportContext, TranslationContext.Get() );
@@ -1556,8 +1572,11 @@ bool UUsdStageImporter::ReimportSingleAsset(FUsdStageImportContext& ImportContex
 	TranslationContext->bAllowInterpretingLODs = ImportContext.ImportOptions->bInterpretLODs;
 	TranslationContext->bAllowParsingSkeletalAnimations = ImportContext.ImportOptions->bImportGeometry && ImportContext.ImportOptions->bImportSkeletalAnimations;
 	TranslationContext->MaterialToPrimvarToUVIndex = &ImportContext.MaterialToPrimvarToUVIndex;
+	TranslationContext->CollapsingCache = ImportContext.CollapsingCache;
 	TranslationContext->BlendShapesByPath = &BlendShapesByPath;
 	{
+		UsdStageImporterImpl::CacheCollapsingState( TranslationContext.Get() );
+
 		UE::FUsdPrim TargetPrim = ImportContext.Stage.GetPrimAtPath( UE::FSdfPath( *OriginalImportData->PrimPath ) );
 		if ( TargetPrim )
 		{
diff --git a/Engine/Plugins/Importers/USDImporter/Source/USDStageImporter/Public/USDStageImportContext.h b/Engine/Plugins/Importers/USDImporter/Source/USDStageImporter/Public/USDStageImportContext.h
--- a/Engine/Plugins/Importers/USDImporter/Source/USDStageImporter/Public/USDStageImportContext.h
+++ b/Engine/Plugins/Importers/USDImporter/Source/USDStageImporter/Public/USDStageImportContext.h
@@ -3,6 +3,7 @@
 #pragma once
 
 #include "UnrealUSDWrapper.h"
+#include "USDCollapsingCache.h"
 #include "USDLevelSequenceHelper.h"
 #include "USDMemory.h"
 #include "USDStageOptions.h"
@@ -59,6 +60,8 @@ struct USDSTAGEIMPORTER_API FUsdStageImportContext
 	UPROPERTY()
 	UUsdAssetCache* AssetCache;
 
+	TSharedPtr<FUsdCollapsingCache> CollapsingCache;
+
 	/**
 	 * When parsing materials, we keep track of which primvar we mapped to which UV channel.
 	 * When parsing meshes later, we use this data to place the correct primvar values in each UV channel.
@@ -91,6 +94,7 @@ public:
 	FUsdStageImportContext();
 
 	bool Init(const FString& InName, const FString& InFilePath, const FString& InInitialPackagePath, EObjectFlags InFlags, bool bInIsAutomated, bool bIsReimport = false, bool bAllowActorImport = true);
+	void Reset();
 
 private:
 	/** Error messages **/
diff --git a/Engine/Plugins/Importers/USDImporter/Source/UnrealUSDWrapper/Private/UsdWrappers/SdfPath.cpp b/Engine/Plugins/Importers/USDImporter/Source/UnrealUSDWrapper/Private/UsdWrappers/SdfPath.cpp
--- a/Engine/Plugins/Importers/USDImporter/Source/UnrealUSDWrapper/Private/UsdWrappers/SdfPath.cpp
+++ b/Engine/Plugins/Importers/USDImporter/Source/UnrealUSDWrapper/Private/UsdWrappers/SdfPath.cpp
@@ -102,6 +102,34 @@ namespace UE
 		return !( *this == Other );
 	}
 
+	uint32 GetTypeHash( const UE::FSdfPath& Path )
+	{
+		uint32 Result = 0;
+#if USE_USD_SDK
+		FScopedUnrealAllocs UnrealAllocs;
+		Result = static_cast< uint32 >( Path.Impl->PxrSdfPath.Get().GetHash() );
+#endif // #if USE_USD_SDK
+		return Result;
+	}
+
+	FArchive& operator<<( FArchive& Ar, UE::FSdfPath& Path )
+	{
+		if ( Ar.IsSaving() )
+		{
+			FString String = Path.GetString();
+			Ar << String;
+		}
+		else if ( Ar.IsLoading() )
+		{
+			FString String;
+			Ar << String;
+
+			Path = UE::FSdfPath{ *String };
+		}
+
+		return Ar;
+	}
+
 #if USE_USD_SDK
 	FSdfPath::FSdfPath( const pxr::SdfPath& InSdfPath )
 		: Impl( MakeUnique< Internal::FSdfPathImpl >( InSdfPath ) )
@@ -156,6 +184,15 @@ namespace UE
 #endif // #if USE_USD_SDK
 	}
 
+	bool FSdfPath::IsAbsoluteRootPath() const
+	{
+#if USE_USD_SDK
+		return Impl->PxrSdfPath.Get().IsAbsoluteRootPath();
+#else
+		return false;
+#endif // #if USE_USD_SDK
+	}
+
 	bool FSdfPath::IsAbsoluteRootOrPrimPath() const
 	{
 #if USE_USD_SDK
@@ -236,4 +273,21 @@ namespace UE
 		return FString();
 #endif // #if USE_USD_SDK
 	}
+
+	TArray<FSdfPath> FSdfPath::GetPrefixes() const
+	{
+		TArray<FSdfPath> Result;
+#if USE_USD_SDK
+		FScopedUsdAllocs Allocs;
+
+		std::vector<pxr::SdfPath> UsdPrefixes = Impl->PxrSdfPath.Get().GetPrefixes();
+		Result.Reserve( UsdPrefixes.size() );
+
+		for ( const pxr::SdfPath& UsdPrefix : UsdPrefixes )
+		{
+			Result.Add( UE::FSdfPath{ ANSI_TO_TCHAR( UsdPrefix.GetString().c_str() ) } );
+		}
+#endif // #if USE_USD_SDK
+		return Result;
+	}
 }
diff --git a/Engine/Plugins/Importers/USDImporter/Source/UnrealUSDWrapper/Public/UsdWrappers/SdfPath.h b/Engine/Plugins/Importers/USDImporter/Source/UnrealUSDWrapper/Public/UsdWrappers/SdfPath.h
--- a/Engine/Plugins/Importers/USDImporter/Source/UnrealUSDWrapper/Public/UsdWrappers/SdfPath.h
+++ b/Engine/Plugins/Importers/USDImporter/Source/UnrealUSDWrapper/Public/UsdWrappers/SdfPath.h
@@ -43,6 +43,9 @@ namespace UE
 		bool operator==( const FSdfPath& Other ) const;
 		bool operator!=( const FSdfPath& Other ) const;
 
+		friend UNREALUSDWRAPPER_API uint32 GetTypeHash( const UE::FSdfPath& Path );
+		friend UNREALUSDWRAPPER_API FArchive& operator<<( FArchive& Ar, FSdfPath& Path );
+
 	// Auto conversion from/to pxr::SdfPath
 	public:
 #if USE_USD_SDK
@@ -61,6 +64,7 @@ namespace UE
 
 		bool IsEmpty() const noexcept;
 
+		bool IsAbsoluteRootPath() const;
 		bool IsAbsoluteRootOrPrimPath() const;
 		FString GetName() const;
 		FString GetElementString() const;
@@ -75,6 +79,8 @@ namespace UE
 
 		FString GetString() const;
 
+		TArray<FSdfPath> GetPrefixes() const;
+
 	private:
 		TUniquePtr< Internal::FSdfPathImpl > Impl;
 	};
