Fix high contention on USD memory manager by using lock striping
Remove async task for USD memory free because the task itself will allocate memory and cause even more free
Parallellize USD collapse cache build
  - 45s to 2.5s for building the collapse cache for Moana scene

diff --git a/Engine/Plugins/Importers/USDImporter/Source/USDSchemas/Private/USDCollapsingCache.cpp b/Engine/Plugins/Importers/USDImporter/Source/USDSchemas/Private/USDCollapsingCache.cpp
--- a/Engine/Plugins/Importers/USDImporter/Source/USDSchemas/Private/USDCollapsingCache.cpp
+++ b/Engine/Plugins/Importers/USDImporter/Source/USDSchemas/Private/USDCollapsingCache.cpp
@@ -9,6 +9,7 @@
 #include "UsdWrappers/SdfPath.h"
 #include "UsdWrappers/UsdPrim.h"
 
+#include "Async/ParallelFor.h"
 #include "Modules/ModuleManager.h"
 
 #if USE_USD_SDK
@@ -20,6 +21,7 @@
 
 bool FUsdCollapsingCache::Serialize( FArchive& Ar )
 {
+	FWriteScopeLock ScopeLock( Lock );
 	Ar << AssetPathsToCollapsedRoot;
 	Ar << ComponentPathsToCollapsedRoot;
 
@@ -28,6 +30,7 @@ bool FUsdCollapsingCache::Serialize( FArchive& Ar )
 
 bool FUsdCollapsingCache::IsPathCollapsed( const UE::FSdfPath& Path, ECollapsingType CollapsingType ) const
 {
+	FReadScopeLock ScopeLock( Lock );
 	const TMap< UE::FSdfPath, UE::FSdfPath >* MapToUse =
 		CollapsingType == ECollapsingType::Assets
 			? &AssetPathsToCollapsedRoot
@@ -46,6 +49,7 @@ bool FUsdCollapsingCache::IsPathCollapsed( const UE::FSdfPath& Path, ECollapsing
 
 bool FUsdCollapsingCache::DoesPathCollapseChildren( const UE::FSdfPath& Path, ECollapsingType CollapsingType ) const
 {
+	FReadScopeLock ScopeLock( Lock );
 	const TMap< UE::FSdfPath, UE::FSdfPath >* MapToUse =
 		CollapsingType == ECollapsingType::Assets
 			? &AssetPathsToCollapsedRoot
@@ -64,6 +68,7 @@ bool FUsdCollapsingCache::DoesPathCollapseChildren( const UE::FSdfPath& Path, EC
 
 UE::FSdfPath FUsdCollapsingCache::UnwindToNonCollapsedPath( const UE::FSdfPath& Path, ECollapsingType CollapsingType ) const
 {
+	FReadScopeLock ScopeLock( Lock );
 	const TMap< UE::FSdfPath, UE::FSdfPath >* MapToUse =
 		CollapsingType == ECollapsingType::Assets
 			? &AssetPathsToCollapsedRoot
@@ -97,14 +102,20 @@ namespace UE::USDCollapsingCacheImpl::Private
 		FUsdSchemaTranslatorRegistry& Registry,
 		TMap< UE::FSdfPath, UE::FSdfPath >& AssetPathsToCollapsedRoot,
 		TMap< UE::FSdfPath, UE::FSdfPath >& ComponentPathsToCollapsedRoot,
-		pxr::SdfPath AssetCollapsedRoot = pxr::SdfPath::EmptyPath(),
-		pxr::SdfPath ComponentCollapsedRoot = pxr::SdfPath::EmptyPath()
+		FRWLock& Lock,
+		const pxr::SdfPath& AssetCollapsedRoot = pxr::SdfPath::EmptyPath(),
+		const pxr::SdfPath& ComponentCollapsedRoot = pxr::SdfPath::EmptyPath()
 	)
 	{
+		TRACE_CPUPROFILER_EVENT_SCOPE( RecursiveRebuildCache );
 		FScopedUsdAllocs Allocs;
 
 		pxr::SdfPath UsdPrimPath = UsdPrim.GetPrimPath();
 
+		// Prevents allocation by referencing instead of copying
+		const pxr::SdfPath* AssetCollapsedRootOverride = &AssetCollapsedRoot;
+		const pxr::SdfPath* ComponentCollapsedRootOverride = &ComponentCollapsedRoot;
+
 		bool bIsAssetCollapsed = !AssetCollapsedRoot.IsEmpty();
 		bool bIsComponentCollapsed = !ComponentCollapsedRoot.IsEmpty();
 
@@ -116,7 +127,7 @@ namespace UE::USDCollapsingCacheImpl::Private
 				{
 					if ( SchemaTranslator->CollapsesChildren( ECollapsingType::Assets ) )
 					{
-						AssetCollapsedRoot = UsdPrimPath;
+						AssetCollapsedRootOverride = &UsdPrimPath;
 					}
 				}
 
@@ -124,35 +135,45 @@ namespace UE::USDCollapsingCacheImpl::Private
 				{
 					if ( SchemaTranslator->CollapsesChildren( ECollapsingType::Components ) )
 					{
-						ComponentCollapsedRoot = UsdPrimPath;
+						ComponentCollapsedRootOverride = &UsdPrimPath;
 					}
 				}
 			}
 		}
 
-		UE::FSdfPath PrimPath{ UsdPrimPath };
-		UE::FSdfPath UEAssetCollapsedRoot{ AssetCollapsedRoot };
-		UE::FSdfPath UEComponentCollapsedRoot{ ComponentCollapsedRoot };
-
 		// These paths will be still empty in case nothing has collapsed yet, hold UsdPrimPath in case UsdPrim collapses that type, or hold the path to the
 		// collapsed root passed in via our caller, in case we're collapsed
-		AssetPathsToCollapsedRoot.Add( PrimPath, UEAssetCollapsedRoot );
-		ComponentPathsToCollapsedRoot.Add( PrimPath, UEComponentCollapsedRoot );
+		{
+			FWriteScopeLock ScopeLock(Lock);
+			AssetPathsToCollapsedRoot.Emplace( UsdPrimPath, *AssetCollapsedRootOverride );
+			ComponentPathsToCollapsedRoot.Emplace( UsdPrimPath, *ComponentCollapsedRootOverride );
+		}
 
 		pxr::UsdPrimSiblingRange PrimChildren = UsdPrim.GetFilteredChildren( pxr::UsdTraverseInstanceProxies( pxr::UsdPrimAllPrimsPredicate ) );
+		
+		TArray<pxr::UsdPrim> Prims;
 		for ( pxr::UsdPrim Child : PrimChildren )
 		{
-			RecursiveRebuildCache( Child, Context, Registry, AssetPathsToCollapsedRoot, ComponentPathsToCollapsedRoot, AssetCollapsedRoot, ComponentCollapsedRoot );
+			Prims.Emplace( Child );
 		}
+
+		ParallelFor(
+			TEXT("RecursiveRebuildCache"),
+			Prims.Num(), 1 /* MinBatchSize */,
+			[&Prims, &Context, &Registry, &AssetPathsToCollapsedRoot, &ComponentPathsToCollapsedRoot, &Lock, &AssetCollapsedRoot, &ComponentCollapsedRoot](int32 Index)
+			{
+				RecursiveRebuildCache( Prims[Index], Context, Registry, AssetPathsToCollapsedRoot, ComponentPathsToCollapsedRoot, Lock, AssetCollapsedRoot, ComponentCollapsedRoot );
+			}
+		);
 	}
 #endif // USE_USD_SDK
 }
 
 void FUsdCollapsingCache::RebuildCacheForSubtree( const UE::FUsdPrim& Prim, FUsdSchemaTranslationContext& Context )
 {
+#if USE_USD_SDK
 	TRACE_CPUPROFILER_EVENT_SCOPE( FUsdCollapsingCache::RebuildCacheForSubtree );
 
-#if USE_USD_SDK
 	// We can't deallocate our collapsing cache pointer with the Usd allocator
 	FScopedUnrealAllocs UEAllocs;
 
@@ -171,7 +192,7 @@ void FUsdCollapsingCache::RebuildCacheForSubtree( const UE::FUsdPrim& Prim, FUsd
 		IUsdSchemasModule& UsdSchemasModule = FModuleManager::Get().LoadModuleChecked< IUsdSchemasModule >( TEXT( "USDSchemas" ) );
 		FUsdSchemaTranslatorRegistry& Registry = UsdSchemasModule.GetTranslatorRegistry();
 
-		UE::USDCollapsingCacheImpl::Private::RecursiveRebuildCache( UsdPrim, Context, Registry, AssetPathsToCollapsedRoot, ComponentPathsToCollapsedRoot );
+		UE::USDCollapsingCacheImpl::Private::RecursiveRebuildCache( UsdPrim, Context, Registry, AssetPathsToCollapsedRoot, ComponentPathsToCollapsedRoot, Lock );
 	}
 	Context.CollapsingCache = Pin;
 #endif // USE_USD_SDK
@@ -179,11 +200,13 @@ void FUsdCollapsingCache::RebuildCacheForSubtree( const UE::FUsdPrim& Prim, FUsd
 
 void FUsdCollapsingCache::Clear()
 {
+	FWriteScopeLock ScopeLock( Lock );
 	AssetPathsToCollapsedRoot.Empty();
 	ComponentPathsToCollapsedRoot.Empty();
 }
 
 bool FUsdCollapsingCache::IsEmpty()
 {
+	FReadScopeLock ScopeLock( Lock );
 	return AssetPathsToCollapsedRoot.IsEmpty() && ComponentPathsToCollapsedRoot.IsEmpty();
 }
\ No newline at end of file
diff --git a/Engine/Plugins/Importers/USDImporter/Source/USDSchemas/Public/USDCollapsingCache.h b/Engine/Plugins/Importers/USDImporter/Source/USDSchemas/Public/USDCollapsingCache.h
--- a/Engine/Plugins/Importers/USDImporter/Source/USDSchemas/Public/USDCollapsingCache.h
+++ b/Engine/Plugins/Importers/USDImporter/Source/USDSchemas/Public/USDCollapsingCache.h
@@ -40,6 +40,7 @@ struct USDSCHEMAS_API FUsdCollapsingCache
 	bool IsEmpty();
 
 private:
+	mutable FRWLock Lock;
 	TMap< UE::FSdfPath, UE::FSdfPath > AssetPathsToCollapsedRoot;
 	TMap< UE::FSdfPath, UE::FSdfPath > ComponentPathsToCollapsedRoot;
 };
\ No newline at end of file
diff --git a/Engine/Plugins/Importers/USDImporter/Source/UnrealUSDWrapper/Private/USDMemory.cpp b/Engine/Plugins/Importers/USDImporter/Source/UnrealUSDWrapper/Private/USDMemory.cpp
--- a/Engine/Plugins/Importers/USDImporter/Source/UnrealUSDWrapper/Private/USDMemory.cpp
+++ b/Engine/Plugins/Importers/USDImporter/Source/UnrealUSDWrapper/Private/USDMemory.cpp
@@ -92,7 +92,7 @@ private:
 };
 
 TOptional< FTlsSlot > FUsdMemoryManager::ActiveAllocatorsStackTLS {};
-TSet< void* > FUsdMemoryManager::SystemAllocedPtrs {};
+FUsdMemoryManager::FThreadSafeSet FUsdMemoryManager::SystemAllocedPtrs {};
 FCriticalSection FUsdMemoryManager::CriticalSection{};
 
 void FUsdMemoryManager::Initialize()
@@ -144,8 +144,8 @@ void* FUsdMemoryManager::Malloc( SIZE_T Count )
 	{
 		Result = FMemory::SystemMalloc( Count );
 
+		if ( Result )
 		{
-			FScopeLock Lock( &CriticalSection );
 			SystemAllocedPtrs.Add( Result );
 		}
 	}
@@ -163,16 +163,11 @@ void FUsdMemoryManager::Free( void* Original )
 #if USD_USES_SYSTEM_MALLOC
 	// Because USD is multi-threaded, it might call us back to free an object after we've exited our allocator scope.
 	// This can happen for inlined USD functions that call delete.
-	int32 Removed = 0;
+	// Skip nullptr to avoid collision between threads for something that was not inserted in the first place.
+	const bool bRemoved = Original != nullptr && SystemAllocedPtrs.Remove( Original );
+	if ( FUsdMemoryManager::IsUsingSystemMalloc() || bRemoved )
 	{
-		FScopeLock Lock( &CriticalSection );
-		Removed = SystemAllocedPtrs.Remove( Original );
-	}
-
-	if ( FUsdMemoryManager::IsUsingSystemMalloc() || Removed > 0 )
-	{
-		// System allocations are so slow that we send them to another thread for freeing
-		Async( EAsyncExecution::TaskGraph, [Original](){ FMemory::SystemFree(Original); } );
+		FMemory::SystemFree( Original );
 	}
 	else
 #endif // #if USD_USES_SYSTEM_MALLOC
diff --git a/Engine/Plugins/Importers/USDImporter/Source/UnrealUSDWrapper/Public/USDMemory.h b/Engine/Plugins/Importers/USDImporter/Source/UnrealUSDWrapper/Public/USDMemory.h
--- a/Engine/Plugins/Importers/USDImporter/Source/UnrealUSDWrapper/Public/USDMemory.h
+++ b/Engine/Plugins/Importers/USDImporter/Source/UnrealUSDWrapper/Public/USDMemory.h
@@ -6,6 +6,7 @@
 #include "Misc/Optional.h"
 #include "Modules/Boilerplate/ModuleBoilerplate.h"
 #include "Templates/SharedPointer.h"
+#include "Misc/ScopeRWLock.h"
 #include "Templates/UnrealTemplate.h"
 
 /**
@@ -76,7 +77,34 @@ private:
 
 	static TOptional< FTlsSlot > ActiveAllocatorsStackTLS;
 
-	static TSet< void* > SystemAllocedPtrs;
+	struct FThreadSafeSet
+	{
+	public:
+		void Add( void* Ptr )
+		{
+			FSubSet& SubSet = SubSets[GetTypeHash( Ptr ) % BucketCount];
+			FWriteScopeLock ScopeLock( SubSet.Lock );
+			SubSet.Set.Add( Ptr );
+		}
+
+		bool Remove( void* Ptr )
+		{
+			FSubSet& SubSet = SubSets[GetTypeHash( Ptr ) % BucketCount];
+			FWriteScopeLock ScopeLock( SubSet.Lock );
+			return SubSet.Set.Remove( Ptr ) > 0;
+		}
+
+	private:
+		struct FSubSet
+		{
+			FRWLock Lock;
+			TSet< void* > Set;
+		};
+		static constexpr int32 BucketCount = 61; /* Prime number for better modulo distribution */
+		FSubSet SubSets[BucketCount];
+	};
+
+	static FThreadSafeSet SystemAllocedPtrs;
 
 	static FCriticalSection CriticalSection;
 };
