USD: Fix refactoring oversight preventing prims from ever being acknowledged as collapsed.

diff --git a/Engine/Plugins/Importers/USDImporter/Source/USDSchemas/Private/USDCollapsingCache.cpp b/Engine/Plugins/Importers/USDImporter/Source/USDSchemas/Private/USDCollapsingCache.cpp
--- a/Engine/Plugins/Importers/USDImporter/Source/USDSchemas/Private/USDCollapsingCache.cpp
+++ b/Engine/Plugins/Importers/USDImporter/Source/USDSchemas/Private/USDCollapsingCache.cpp
@@ -160,9 +160,9 @@ namespace UE::USDCollapsingCacheImpl::Private
 		ParallelFor(
 			TEXT("RecursiveRebuildCache"),
 			Prims.Num(), 1 /* MinBatchSize */,
-			[&Prims, &Context, &Registry, &AssetPathsToCollapsedRoot, &ComponentPathsToCollapsedRoot, &Lock, &AssetCollapsedRoot, &ComponentCollapsedRoot](int32 Index)
+			[&Prims, &Context, &Registry, &AssetPathsToCollapsedRoot, &ComponentPathsToCollapsedRoot, &Lock, AssetCollapsedRootOverride, ComponentCollapsedRootOverride](int32 Index)
 			{
-				RecursiveRebuildCache( Prims[Index], Context, Registry, AssetPathsToCollapsedRoot, ComponentPathsToCollapsedRoot, Lock, AssetCollapsedRoot, ComponentCollapsedRoot );
+				RecursiveRebuildCache( Prims[Index], Context, Registry, AssetPathsToCollapsedRoot, ComponentPathsToCollapsedRoot, Lock, *AssetCollapsedRootOverride, *ComponentCollapsedRootOverride );
 			}
 		);
 	}
