USD: Fix UAnimSequence assets and USkeletalMesh assets overwriting each other on the USD asset cache.

This is also a revert of the USDSkelRootTranslator.cpp changes from cl 19303961

diff --git a/Engine/Plugins/Importers/USDImporter/Source/USDSchemas/Private/USDSkelRootTranslator.cpp b/Engine/Plugins/Importers/USDImporter/Source/USDSchemas/Private/USDSkelRootTranslator.cpp
--- a/Engine/Plugins/Importers/USDImporter/Source/USDSchemas/Private/USDSkelRootTranslator.cpp
+++ b/Engine/Plugins/Importers/USDImporter/Source/USDSchemas/Private/USDSkelRootTranslator.cpp
@@ -808,6 +808,23 @@ namespace UsdSkelRootTranslatorImpl
 
 					Context->AssetCache->LinkAssetToPrim( SkelRootPath, SkeletalMesh );
 
+					// Track our Skeleton by the source skeleton prim path
+					if ( USkeleton* Skeleton = SkeletalMesh->GetSkeleton() )
+					{
+						FScopedUsdAllocs Allocs;
+
+						// Our SkeletonCache should already have been populated by our call to UsdSkelRootTranslatorImpl::LoadAllSkeletalData, so we should
+						// be able to quickly query our skeleton again to find its path
+						std::vector< pxr::UsdSkelBinding > SkeletonBindings;
+						SkeletonCache.Get().ComputeSkelBindings( pxr::UsdSkelRoot( GetPrim() ), &SkeletonBindings, pxr::UsdTraverseInstanceProxies() );
+						if ( SkeletonBindings.size() > 0 )
+						{
+							pxr::UsdSkelBinding& SkeletonBinding = SkeletonBindings[ 0 ];
+							const pxr::UsdSkelSkeleton& UsdSkeleton = SkeletonBinding.GetSkeleton();
+							Context->AssetCache->LinkAssetToPrim( UsdToUnreal::ConvertPath( UsdSkeleton.GetPrim().GetPrimPath() ), Skeleton );
+						}
+					}
+
 					// We may be reusing a skeletal mesh we got in the cache, but we always need the BlendShapesByPath stored on the
 					// actor to be up-to-date with the Skeletal Mesh that is actually being displayed
 					if ( Context->BlendShapesByPath )
@@ -840,7 +857,6 @@ namespace UsdSkelRootTranslatorImpl
 
 				if ( pxr::UsdSkelRoot SkeletonRoot{ GetPrim() } )
 				{
-					FString SkelRootPath = PrimPath.GetString();
 					std::vector< pxr::UsdSkelBinding > SkeletonBindings;
 					SkeletonCache.Get().Populate( SkeletonRoot, pxr::UsdTraverseInstanceProxies() );
 					SkeletonCache.Get().ComputeSkelBindings( SkeletonRoot, &SkeletonBindings, pxr::UsdTraverseInstanceProxies() );
@@ -860,6 +876,7 @@ namespace UsdSkelRootTranslatorImpl
 						{
 							continue;
 						}
+						FString SkelAnimationPrimPath = UsdToUnreal::ConvertPath( SkelAnimationPrim.GetPath() );
 
 						if ( !AnimQuery.JointTransformsMightBeTimeVarying() &&
 							( NewBlendShapes.Num() == 0 || !AnimQuery.BlendShapeWeightsMightBeTimeVarying() ) )
@@ -894,7 +911,7 @@ namespace UsdSkelRootTranslatorImpl
 								UUsdAnimSequenceAssetImportData* ImportData = NewObject< UUsdAnimSequenceAssetImportData >( AnimSequence, TEXT( "USDAssetImportData" ) );
 								AnimSequence->AssetImportData = ImportData;
 
-								ImportData->PrimPath = SkelRootPath;
+								ImportData->PrimPath = SkelAnimationPrimPath;
 								ImportData->LayerStartOffsetSeconds = LayerStartOffsetSeconds;
 
 								Context->AssetCache->CacheAsset( HashString, AnimSequence );
@@ -907,7 +924,7 @@ namespace UsdSkelRootTranslatorImpl
 
 						if ( AnimSequence )
 						{
-							Context->AssetCache->LinkAssetToPrim( SkelRootPath, AnimSequence );
+							Context->AssetCache->LinkAssetToPrim( SkelAnimationPrimPath, AnimSequence );
 						}
 					}
 				}
diff --git a/Engine/Plugins/Importers/USDImporter/Source/USDStage/Private/USDStageActor.cpp b/Engine/Plugins/Importers/USDImporter/Source/USDStage/Private/USDStageActor.cpp
--- a/Engine/Plugins/Importers/USDImporter/Source/USDStage/Private/USDStageActor.cpp
+++ b/Engine/Plugins/Importers/USDImporter/Source/USDStage/Private/USDStageActor.cpp
@@ -1360,24 +1360,32 @@ TArray<UObject*> AUsdStageActor::GetGeneratedAssets( const FString& PrimPath )
 		return {};
 	}
 
-	FString UncollapsedPath = PrimPath;
-	if ( CollapsingCache.IsValid() )
-	{
-		UncollapsedPath = CollapsingCache->UnwindToNonCollapsedPath( UsdPath, ECollapsingType::Assets ).GetString();
-	}
-
 	TSet<UObject*> Result;
+	FString PathToUse = PrimPath;
 
-	if ( UObject* FoundAsset = AssetCache->GetAssetForPrim( UncollapsedPath ) )
+	// If we found an asset generated specifically for a prim path, use that. If this prim has been collapsed,
+	// the asset generated for it will be assigned to its collapsing root, so we won't find anything here
+	if ( UObject* FoundAsset = AssetCache->GetAssetForPrim( PrimPath ) )
 	{
 		Result.Add( FoundAsset );
 	}
+	// If we haven't try seeing if we collapsed into another asset
+	else if ( CollapsingCache.IsValid() )
+	{
+		PathToUse = CollapsingCache->UnwindToNonCollapsedPath( UsdPath, ECollapsingType::Assets ).GetString();
+		if ( UObject* FoundCollapsedAsset = AssetCache->GetAssetForPrim( PathToUse ) )
+		{
+			Result.Add( FoundCollapsedAsset );
+		}
+	}
 
+	// Collect any other asset that claims they came from the same prim (e.g. also return the skeleton if we query
+	// the SkelRoot, or return textures if we query a material prim that used them, etc.)
 	for ( const TPair<FString, UObject*>& HashToAsset : AssetCache->GetCachedAssets() )
 	{
 		if ( UUsdAssetImportData* ImportData = UsdUtils::GetAssetImportData( HashToAsset.Value ) )
 		{
-			if ( ImportData->PrimPath == UncollapsedPath )
+			if ( ImportData->PrimPath == PathToUse )
 			{
 				Result.Add( HashToAsset.Value );
 			}
