USD: Fix skeletal animations not exporting to USD on LevelSequences where we used "Bake to Control Rig".

diff --git a/Engine/Plugins/Importers/USDImporter/Source/USDExporter/Private/LevelSequenceExporterUSD.cpp b/Engine/Plugins/Importers/USDImporter/Source/USDExporter/Private/LevelSequenceExporterUSD.cpp
--- a/Engine/Plugins/Importers/USDImporter/Source/USDExporter/Private/LevelSequenceExporterUSD.cpp
+++ b/Engine/Plugins/Importers/USDImporter/Source/USDExporter/Private/LevelSequenceExporterUSD.cpp
@@ -19,8 +19,8 @@
 #include "AssetExportTask.h"
 #include "Compilation/MovieSceneCompiledDataManager.h"
 #include "Components/SkeletalMeshComponent.h"
-#include "Evaluation/MovieSceneSequenceHierarchy.h"
 #include "Editor.h"
+#include "Evaluation/MovieSceneSequenceHierarchy.h"
 #include "ISequencer.h"
 #include "ISequencerModule.h"
 #include "LevelEditorSequencerIntegration.h"
@@ -32,6 +32,7 @@
 #include "MovieSceneTimeHelpers.h"
 #include "MovieSceneTrack.h"
 #include "Sections/MovieSceneSubSection.h"
+#include "Sequencer/MovieSceneControlRigParameterTrack.h"
 #include "Subsystems/AssetEditorSubsystem.h"
 #include "Tracks/MovieScenePropertyTrack.h"
 #include "Tracks/MovieSceneSkeletalAnimationTrack.h"
@@ -667,7 +668,7 @@ namespace UE
 								const FString PropertyPath = PropertyTrack->GetPropertyPath().ToString();
 								UnrealToUsd::CreateComponentPropertyBaker( Prim, *BoundComponent, PropertyPath, Baker );
 							}
-							else if ( const UMovieSceneSpawnTrack* SpawnTrack = Cast<UMovieSceneSpawnTrack>( Track ) )
+							else if ( Track->IsA< UMovieSceneSpawnTrack >() )
 							{
 								// Just handle spawnable tracks as if they're visibility tracks, and hide the prim when not "spawned"
 								// Remember that our spawn register just hides the spawnables when they're not spawned anyway, so this
@@ -675,7 +676,9 @@ namespace UE
 								const FString PropertyPath = TEXT( "bHidden" );
 								UnrealToUsd::CreateComponentPropertyBaker( Prim, *BoundComponent, PropertyPath, Baker);
 							}
-							else if ( const UMovieSceneSkeletalAnimationTrack* SkeletalTrack = Cast<UMovieSceneSkeletalAnimationTrack>( Track ) )
+							// Check for the control rig tracks too, because if the user did "Bake to Control Rig" the controlrig code will silently set the
+							// original skeletal animation track sections as disabled, so they'd fail the "IsTrackAnimated" check above
+							else if ( Track->IsA< UMovieSceneSkeletalAnimationTrack >() || Track->IsA< UMovieSceneControlRigParameterTrack >() )
 							{
 								if ( USkeletalMeshComponent* SkeletalBoundComponent = Cast<USkeletalMeshComponent>( BoundComponent ) )
 								{
diff --git a/Engine/Plugins/Importers/USDImporter/Source/USDExporter/USDExporter.Build.cs b/Engine/Plugins/Importers/USDImporter/Source/USDExporter/USDExporter.Build.cs
--- a/Engine/Plugins/Importers/USDImporter/Source/USDExporter/USDExporter.Build.cs
+++ b/Engine/Plugins/Importers/USDImporter/Source/USDExporter/USDExporter.Build.cs
@@ -26,6 +26,7 @@ namespace UnrealBuildTool.Rules
 					"Analytics",
 					"AnalyticsBlueprintLibrary",
 					"CinematicCamera",
+					"ControlRig", // For detecting rigged skeletal animations when exporting level sequences to USD
 					"EditorStyle",
 					"GeometryCache",
 					"InputCore",
diff --git a/Engine/Plugins/Importers/USDImporter/USDImporter.uplugin b/Engine/Plugins/Importers/USDImporter/USDImporter.uplugin
--- a/Engine/Plugins/Importers/USDImporter/USDImporter.uplugin
+++ b/Engine/Plugins/Importers/USDImporter/USDImporter.uplugin
@@ -76,13 +76,21 @@
 	],
 	"Plugins": [
 		{
-			"Name": "GeometryCache",
+			"Name": "AnalyticsBlueprintLibrary",
+			"Enabled": true
+		},
+		{
+			"Name": "ControlRig",
 			"Enabled": true
 		},
 		{
 			"Name": "EditorScriptingUtilities",
 			"Enabled": true
 		},
+		{
+			"Name": "GeometryCache",
+			"Enabled": true
+		},
 		{
 			"Name": "LevelSequenceEditor",
 			"Enabled": true
@@ -98,10 +106,6 @@
 		{
 			"Name": "SequencerScripting",
 			"Enabled": true
-		},
-		{
-			"Name": "AnalyticsBlueprintLibrary",
-			"Enabled": true
 		}
 	]
 }
