USD: Add cvar to (attempt to) disable physics mesh generation.

diff --git a/Engine/Plugins/Importers/USDImporter/Source/USDSchemas/Private/USDGeomMeshTranslator.cpp b/Engine/Plugins/Importers/USDImporter/Source/USDSchemas/Private/USDGeomMeshTranslator.cpp
--- a/Engine/Plugins/Importers/USDImporter/Source/USDSchemas/Private/USDGeomMeshTranslator.cpp
+++ b/Engine/Plugins/Importers/USDImporter/Source/USDSchemas/Private/USDGeomMeshTranslator.cpp
@@ -680,9 +680,21 @@ namespace UsdGeomMeshTranslatorImpl
 		}
 
 #endif // WITH_EDITOR
+
+		static IConsoleVariable* CVar = IConsoleManager::Get().FindConsoleVariable( TEXT( "USD.EnableCollision" ) );
+		bool bEnableCollision = CVar && CVar->GetBool();
+
 		if ( StaticMesh.GetBodySetup())
 		{
-			StaticMesh.GetBodySetup()->CreatePhysicsMeshes();
+			if ( bEnableCollision )
+			{
+				StaticMesh.GetBodySetup()->CreatePhysicsMeshes();
+			}
+			else
+			{
+				StaticMesh.GetBodySetup()->DefaultInstance.SetCollisionEnabled( ECollisionEnabled::NoCollision );
+				StaticMesh.GetBodySetup()->DefaultInstance.SetCollisionProfileName( UCollisionProfile::NoCollision_ProfileName );
+			}
 		}
 		return true;
 	}
diff --git a/Engine/Plugins/Importers/USDImporter/Source/USDSchemas/Private/USDGeomXformableTranslator.cpp b/Engine/Plugins/Importers/USDImporter/Source/USDSchemas/Private/USDGeomXformableTranslator.cpp
--- a/Engine/Plugins/Importers/USDImporter/Source/USDSchemas/Private/USDGeomXformableTranslator.cpp
+++ b/Engine/Plugins/Importers/USDImporter/Source/USDSchemas/Private/USDGeomXformableTranslator.cpp
@@ -52,6 +52,12 @@ static FAutoConsoleVariableRef CVarMaxNumVerticesCollapsedMesh(
 	GMaxNumVerticesCollapsedMesh,
 	TEXT( "Maximum number of vertices that a combined Mesh can have for us to collapse it into a single StaticMesh" ) );
 
+static bool GEnableCollision = true;
+static FAutoConsoleVariableRef CVarEnableCollision(
+	TEXT( "USD.EnableCollision" ),
+	GEnableCollision,
+	TEXT( "Whether to have collision enabled for spawned components and generated meshes" ) );
+
 class FUsdGeomXformableCreateAssetsTaskChain : public FBuildStaticMeshTaskChain
 {
 public:
@@ -320,6 +326,18 @@ USceneComponent* FUsdGeomXformableTranslator::CreateComponentsEx( TOptional< TSu
 
 	if ( SceneComponent )
 	{
+		if ( !GEnableCollision )
+		{
+			// In most cases this will have no benefit memory-wise, as regular UStaticMeshComponents build their physics meshes anyway
+			// when registering, regardless of these. HISM components will *not* build them though, so disabling the cvar may lead
+			// to some memory savings
+			if ( UPrimitiveComponent* PrimComp = Cast<UPrimitiveComponent>( SceneComponent ) )
+			{
+				PrimComp->SetCollisionEnabled( ECollisionEnabled::NoCollision );
+				PrimComp->SetCollisionProfileName( UCollisionProfile::NoCollision_ProfileName );
+			}
+		}
+
 		if ( !SceneComponent->GetOwner()->GetRootComponent() )
 		{
 			SceneComponent->GetOwner()->SetRootComponent( SceneComponent );
