USD: Don't unload and reload the stage when reregistering components;

Fix not being able to animate the preview actor on the blueprint editor (of blueprints that derive stage actors) after an instance of the blueprint has been placed.

diff --git a/Engine/Plugins/Importers/USDImporter/Source/USDStage/Private/USDStageActor.cpp b/Engine/Plugins/Importers/USDImporter/Source/USDStage/Private/USDStageActor.cpp
--- a/Engine/Plugins/Importers/USDImporter/Source/USDStage/Private/USDStageActor.cpp
+++ b/Engine/Plugins/Importers/USDImporter/Source/USDStage/Private/USDStageActor.cpp
@@ -1697,16 +1697,22 @@ void AUsdStageActor::LoadUsdStage()
 	{
 		SetTime( UsdStage.GetRootLayer().GetStartTimeCode() );
 
-		// Our CDO will never load the stage, so it will remain with some other Time value. If we don't update it, it will desync with the
-		// the Time value of the instance on the preview editor (because the instance will load the stage and update its Time), and so our
-		// manipulation of the CDO's Time value on the blueprint editor won't be propagated to the instance.
-		// This means that we wouldn't be able to animate the preview actor at all. Here we fix that by resyncing our Time with the CDO
+		// If we're an instance of a blueprint that derives the stage actor and we're in the editor preview world, it means we're the
+		// blueprint preview actor. We (the instance) will load the stage and update our Time to StartTimeCode, but our CDO will not.
+		// The blueprint editor shows the property values of the CDO however, so our Time may desync with the CDO's. If that happens, setting the Time
+		// value in the blueprint editor won't be propagated to the instance, so we wouldn't be able to animate the preview actor at all.
+		// Here we fix that by updating our CDO to our new Time value. Note how we just do this if we're the preview instance though, we don't
+		// want other instances driving the CDO like this
 		if ( UBlueprintGeneratedClass* BPClass = Cast<UBlueprintGeneratedClass>( GetClass() ) )
 		{
-			// Note: CDO is an instance of a BlueprintGeneratedClass here and this is just a base class pointer. We're not changing the actual AUsdStageActor's CDO
-			if ( AUsdStageActor* CDO = Cast<AUsdStageActor>( GetClass()->GetDefaultObject() ) )
+			UWorld* World = GetWorld();
+			if ( World && World->WorldType == EWorldType::EditorPreview )
 			{
-				CDO->SetTime( GetTime() );
+				// Note: CDO is an instance of a BlueprintGeneratedClass here and this is just a base class pointer. We're not changing the actual AUsdStageActor's CDO
+				if ( AUsdStageActor* CDO = Cast<AUsdStageActor>( GetClass()->GetDefaultObject() ) )
+				{
+					CDO->SetTime( GetTime() );
+				}
 			}
 		}
 	}
@@ -2150,6 +2156,27 @@ void AUsdStageActor::PostRegisterAllComponents()
 {
 	Super::PostRegisterAllComponents();
 
+	// This may fail if our stage happened to not spawn any components, actors or assets, but by that
+	// point "being loaded" doesn't really mean anything anyway
+	const bool bStageIsLoaded = UsdStage && ( ( RootUsdTwin && RootUsdTwin->GetSceneComponent() != nullptr ) || ( AssetCache && AssetCache->GetNumAssets() > 0 ) );
+
+	// Blocks loading stage when going into PIE, if we already have something loaded (we'll want to duplicate stuff instead).
+	// We need to allow loading when going into PIE when we have nothing loaded yet because the MovieRenderQueue (or other callers)
+	// may directly trigger PIE sessions providing an override world. Without this exception a map saved with a loaded stage
+	// wouldn't load it at all when opening the level in that way
+	UWorld* World = GetWorld();
+	if ( bIsTransitioningIntoPIE && bStageIsLoaded && ( !World || World->WorldType == EWorldType::PIE ) )
+	{
+		return;
+	}
+
+	// We get an inactive world when dragging a ULevel asset
+	// This is just hiding though, so we shouldn't actively load/unload anything
+	if ( !World || World->WorldType == EWorldType::Inactive )
+	{
+		return;
+	}
+
 #if WITH_EDITOR
 	// Prevent loading on bHiddenEdLevel because PostRegisterAllComponents gets called in the process of hiding our level, if we're in the persistent level.
 	if ( bIsEditorPreviewActor || bHiddenEdLevel )
@@ -2157,15 +2184,22 @@ void AUsdStageActor::PostRegisterAllComponents()
 		return;
 	}
 
-	// We can't load stage when recompiling our blueprint because blueprint recompilation is not a transaction. We're forced
-	// to reuse the existing spawned components, actors and prim twins instead ( which we move over on OnObjectsReplaced ), or
-	// we'd get tons of undo/redo bugs.
 	if ( UBlueprintGeneratedClass* BPClass = Cast<UBlueprintGeneratedClass>( GetClass() ) )
 	{
+		// We can't load stage when recompiling our blueprint because blueprint recompilation is not a transaction. We're forced
+		// to reuse the existing spawned components, actors and prim twins instead ( which we move over on OnObjectsReplaced ), or
+		// we'd get tons of undo/redo bugs.
 		if ( FRecompilationTracker::IsBeingCompiled( Cast<UBlueprint>( BPClass->ClassGeneratedBy ) ) )
 		{
 			return;
 		}
+
+		// For blueprints that derive from the stage actor, any property change on the blueprint preview window will trigger a full
+		// PostRegisterAllComponents. We don't want to reload the stage when e.g. changing the Time property, so we have to return here
+		if ( World && World->WorldType == EWorldType::EditorPreview && bStageIsLoaded )
+		{
+			return;
+		}
 	}
 #endif // WITH_EDITOR
 
@@ -2174,24 +2208,12 @@ void AUsdStageActor::PostRegisterAllComponents()
 	// as our spawned actors/components should be invisible too
 	ULevel* Level = GetLevel();
 	const bool bIsLevelHidden = !Level || ( !Level->bIsVisible && !Level->bIsAssociatingLevel );
-
-	// This may say fail if our stage happened to not spawn any components, actors or assets, but by that
-	// point "being loaded" doesn't really mean anything anyway
-	const bool bStageIsLoaded = UsdStage && ( ( RootUsdTwin && RootUsdTwin->GetSceneComponent() != nullptr ) || ( AssetCache && AssetCache->GetNumAssets() > 0 ) );
-
-	// Blocks loading stage when going into PIE, if we already have something loaded (we'll want to duplicate stuff instead).
-	// We need to allow loading when going into PIE when we have nothing loaded yet because the MovieRenderQueue (or other callers)
-	// may directly trigger PIE sessions providing an override world. Without this exception a map saved with a loaded stage
-	// wouldn't load it at all when opening the level in that way
-	UWorld* World = GetWorld();
-	if ( bIsTransitioningIntoPIE && bStageIsLoaded && ( !World || World->WorldType == EWorldType::PIE ) )
+	if ( bIsLevelHidden )
 	{
 		return;
 	}
 
-	// We get an inactive world when dragging a ULevel asset
-	// This is just hiding though, so we shouldn't actively load/unload anything
-	if ( IsTemplate() || !World || World->WorldType == EWorldType::Inactive || bIsLevelHidden || bIsModifyingAProperty || bIsUndoRedoing )
+	if ( IsTemplate() || bIsModifyingAProperty || bIsUndoRedoing )
 	{
 		return;
 	}
@@ -2202,9 +2224,14 @@ void AUsdStageActor::PostRegisterAllComponents()
 	LoadUsdStage();
 }
 
-void AUsdStageActor::PostUnregisterAllComponents()
+void AUsdStageActor::UnregisterAllComponents( bool bForReregister )
 {
-	Super::PostUnregisterAllComponents();
+	Super::UnregisterAllComponents( bForReregister );
+
+	if ( bForReregister || bIsModifyingAProperty || bIsUndoRedoing )
+	{
+		return;
+	}
 
 #if WITH_EDITOR
 	if ( bIsEditorPreviewActor )
@@ -2234,7 +2261,12 @@ void AUsdStageActor::PostUnregisterAllComponents()
 	// We get an inactive world when dragging a ULevel asset
 	// Unlike on PostRegister, we still want to unload our stage if our world is nullptr, as that likely means we were in
 	// a sublevel that got unloaded
-	if ( IsTemplate() || IsEngineExitRequested() || ( World && World->WorldType == EWorldType::Inactive ) || bIsModifyingAProperty || bIsUndoRedoing )
+	if ( World && World->WorldType == EWorldType::Inactive )
+	{
+		return;
+	}
+
+	if ( IsTemplate() || IsEngineExitRequested() )
 	{
 		return;
 	}
@@ -2242,6 +2274,11 @@ void AUsdStageActor::PostUnregisterAllComponents()
 	UnloadUsdStage();
 }
 
+void AUsdStageActor::PostUnregisterAllComponents()
+{
+	Super::PostUnregisterAllComponents();
+}
+
 void AUsdStageActor::OnPreUsdImport( FString FilePath )
 {
 	if ( !UsdStage || !HasAuthorityOverStage() )
diff --git a/Engine/Plugins/Importers/USDImporter/Source/USDStage/Public/USDStageActor.h b/Engine/Plugins/Importers/USDImporter/Source/USDStage/Public/USDStageActor.h
--- a/Engine/Plugins/Importers/USDImporter/Source/USDStage/Public/USDStageActor.h
+++ b/Engine/Plugins/Importers/USDImporter/Source/USDStage/Public/USDStageActor.h
@@ -194,6 +194,7 @@ public:
 	virtual void BeginDestroy() override;
 
 	virtual void PostRegisterAllComponents() override;
+	virtual void UnregisterAllComponents( bool bForReregister = false ) override;
 	virtual void PostUnregisterAllComponents() override;
 
 	void OnPreUsdImport( FString FilePath );
