USD: Refactor point instancer prototype handling to always attempt to collapse one mesh for each prototype;

Add common automated tests for point instancers.

diff --git a/Engine/Plugins/Importers/USDImporter/Source/USDSchemas/Private/USDGeomPointInstancerTranslator.cpp b/Engine/Plugins/Importers/USDImporter/Source/USDSchemas/Private/USDGeomPointInstancerTranslator.cpp
--- a/Engine/Plugins/Importers/USDImporter/Source/USDSchemas/Private/USDGeomPointInstancerTranslator.cpp
+++ b/Engine/Plugins/Importers/USDImporter/Source/USDSchemas/Private/USDGeomPointInstancerTranslator.cpp
@@ -6,6 +6,7 @@
 #include "USDAssetCache.h"
 #include "USDAssetImportData.h"
 #include "USDConversionUtils.h"
+#include "USDGeomMeshConversion.h"
 #include "USDLog.h"
 #include "USDSchemasModule.h"
 #include "USDTypesConversion.h"
@@ -47,10 +48,9 @@ namespace UsdGeomPointInstancerTranslatorImpl
 
 		pxr::VtMatrix4dArray UsdInstanceTransforms;
 
-		if ( !PointInstancer.ComputeInstanceTransformsAtTime(
-				&UsdInstanceTransforms,
-				EvalTime,
-				EvalTime ) )
+		// We don't want the prototype root prim's transforms to be included in these, as they'll already be baked into the meshes
+		// themselves
+		if ( !PointInstancer.ComputeInstanceTransformsAtTime( &UsdInstanceTransforms, EvalTime, EvalTime, pxr::UsdGeomPointInstancer::ExcludeProtoXform ) )
 		{
 			return false;
 		}
@@ -86,43 +86,86 @@ namespace UsdGeomPointInstancerTranslatorImpl
 		}
 	}
 
-	UStaticMesh* SetStaticMesh( const pxr::UsdGeomMesh& UsdMesh, UStaticMeshComponent& MeshComponent, const UUsdAssetCache& AssetCache )
+	void SetStaticMesh( UStaticMesh* StaticMesh, UStaticMeshComponent& MeshComponent )
 	{
-		FScopedUnrealAllocs UnrealAllocs;
-
-		FString MeshPrimPath = UsdToUnreal::ConvertPath( UsdMesh.GetPrim().GetPrimPath() );
-
-		UStaticMesh* StaticMesh = Cast< UStaticMesh >( AssetCache.GetAssetForPrim( MeshPrimPath ) );
+		if ( StaticMesh == MeshComponent.GetStaticMesh() )
+		{
+			return;
+		}
 
-		if ( StaticMesh != MeshComponent.GetStaticMesh() )
+		if ( MeshComponent.IsRegistered() )
 		{
-			if ( MeshComponent.IsRegistered() )
-			{
-				MeshComponent.UnregisterComponent();
-			}
+			MeshComponent.UnregisterComponent();
+		}
 
+		if ( StaticMesh )
+		{
 			StaticMesh->CreateBodySetup(); // BodySetup is required for HISM component
-			MeshComponent.SetStaticMesh( StaticMesh );
-
-			MeshComponent.RegisterComponent();
 		}
 
-		return StaticMesh;
+		MeshComponent.SetStaticMesh( StaticMesh );
+
+		MeshComponent.RegisterComponent();
 	}
 }
 
-void FUsdGeomPointInstancerTranslator::UpdateComponents( USceneComponent* PointInstancerRootComponent )
+FUsdGeomPointInstancerCreateAssetsTaskChain::FUsdGeomPointInstancerCreateAssetsTaskChain( const TSharedRef< FUsdSchemaTranslationContext >& InContext, const UE::FSdfPath& InPrimPath )
+	: FBuildStaticMeshTaskChain( InContext, InPrimPath )
 {
-	if ( !PointInstancerRootComponent )
-	{
-		return;
-	}
+	SetupTasks();
+}
 
-	TRACE_CPUPROFILER_EVENT_SCOPE( FUsdGeomPointInstancerTranslator::UpdateComponents );
+void FUsdGeomPointInstancerCreateAssetsTaskChain::SetupTasks()
+{
+	FScopedUnrealAllocs UnrealAllocs;
 
-	PointInstancerRootComponent->Modify();
+	// Create mesh description (Async)
+	Do( ESchemaTranslationLaunchPolicy::Async,
+		[this]() -> bool
+		{
+			// TODO: Restore support for LOD prototypes
+			LODIndexToMeshDescription.Reset( 1 );
+			LODIndexToMaterialInfo.Reset( 1 );
 
-	FScopedUsdAllocs UsdAllocs;
+			FMeshDescription& AddedMeshDescription = LODIndexToMeshDescription.Emplace_GetRef();
+			UsdUtils::FUsdPrimMaterialAssignmentInfo& AssignmentInfo = LODIndexToMaterialInfo.Emplace_GetRef();
+
+			TMap< FString, TMap< FString, int32 > > Unused;
+			TMap< FString, TMap< FString, int32 > >* MaterialToPrimvarToUVIndex = Context->MaterialToPrimvarToUVIndex ? Context->MaterialToPrimvarToUVIndex : &Unused;
+
+			pxr::TfToken RenderContextToken = pxr::UsdShadeTokens->universalRenderContext;
+			if ( !Context->RenderContext.IsNone() )
+			{
+				RenderContextToken = UnrealToUsd::ConvertToken( *Context->RenderContext.ToString() ).Get();
+			}
+
+			// We want to bake even the top level prim's transform and visibility into the combined mesh, as we'll only apply the instancing
+			// transform later
+			const bool bSkipRootPrimTransformAndVis = false;
+
+			UsdToUnreal::ConvertGeomMeshHierarchy(
+				GetPrim(),
+				pxr::UsdTimeCode( Context->Time ),
+				Context->PurposesToLoad,
+				RenderContextToken,
+				*MaterialToPrimvarToUVIndex,
+				AddedMeshDescription,
+				AssignmentInfo,
+				bSkipRootPrimTransformAndVis,
+				Context->bMergeIdenticalMaterialSlots
+			);
+
+			return !AddedMeshDescription.IsEmpty();
+		});
+
+	FBuildStaticMeshTaskChain::SetupTasks();
+}
+
+void FUsdGeomPointInstancerTranslator::CreateAssets()
+{
+	TRACE_CPUPROFILER_EVENT_SCOPE( FUsdGeomPointInstancerTranslator::CreateAssets );
+
+	FScopedUsdAllocs Allocs;
 
 	pxr::UsdPrim Prim = GetPrim();
 	pxr::UsdGeomPointInstancer PointInstancer( Prim );
@@ -132,192 +175,151 @@ void FUsdGeomPointInstancerTranslator::UpdateComponents( USceneComponent* PointI
 		return;
 	}
 
-	// The components that we will spawn in here will not have prim twins, and so would otherwise never be destroyed.
-	// Here we make sure we get rid of the old ones because we will create new ones right afterwards
-	PointInstancerRootComponent->Modify();
-	TArray<USceneComponent*> Children;
-	const bool bIncludeAllDescendants = true;
-	PointInstancerRootComponent->GetChildrenComponents( bIncludeAllDescendants, Children );
-	for ( USceneComponent* Child : Children )
+	const pxr::UsdRelationship& Prototypes = PointInstancer.GetPrototypesRel();
+
+	pxr::SdfPathVector PrototypePaths;
+	if ( !Prototypes.GetTargets( &PrototypePaths ) )
+	{
+		return;
+	}
+
+	for ( const pxr::SdfPath& PrototypePath : PrototypePaths )
 	{
-		// Attempt at making sure it's one of our components just in case the user happened to put something in there
-		if ( !Child->ComponentTags.Contains( UnrealIdentifiers::Invisible) &&
-			 !Child->ComponentTags.Contains( UnrealIdentifiers::Inherited) )
+		pxr::UsdPrim PrototypeUsdPrim = Prim.GetStage()->GetPrimAtPath( PrototypePath );
+		if ( !PrototypeUsdPrim )
 		{
+			UE_LOG( LogUsd, Warning, TEXT( "Failed to find prototype '%s' for PointInstancer '%s' when collapsing assets" ), *UsdToUnreal::ConvertPath( PrototypePath ), *PrimPath.GetString() );
 			continue;
 		}
 
-		Child->Modify();
+		Context->TranslatorTasks.Add( MakeShared< FUsdGeomPointInstancerCreateAssetsTaskChain >( Context, UE::FSdfPath{ *UsdToUnreal::ConvertPath( PrototypePath ) } ) );
+	}
+}
 
-		const bool bPromoteChildren = true;
-		Child->DestroyComponent( bPromoteChildren );
+USceneComponent* FUsdGeomPointInstancerTranslator::CreateComponents()
+{
+	TRACE_CPUPROFILER_EVENT_SCOPE( FUsdGeomPointInstancerTranslator::CreateComponents );
+
+	// The plan here is to create an USceneComponent that corresponds to the PointInstancer prim itself, and then spawn a child
+	// HISM component for each prototype.
+	// We always request a scene component here explicitly or else we'll be upgraded to a static mesh component by the mechanism that
+	// handles collapsed Meshes/static mesh components for the GeomXFormable translator.
+	USceneComponent* MainSceneComponent = CreateComponentsEx( { USceneComponent::StaticClass() }, {} );
+	if ( !MainSceneComponent )
+	{
+		return nullptr;
 	}
 
-#if WITH_EDITOR
-	// If we have this actor selected we need to refresh the details panel manually on the next frame because
-	// deleting those components will leave it in some type of invalid state
-	if ( PointInstancerRootComponent->GetOwner()->IsSelected() )
+	FScopedUsdAllocs UsdAllocs;
+
+	pxr::UsdPrim Prim = GetPrim();
+	pxr::UsdGeomPointInstancer PointInstancer( Prim );
+	if ( !PointInstancer )
 	{
-		GEditor->GetTimerManager()->SetTimerForNextTick( []()
-		{
-			GEditor->NoteSelectionChange();
-		} );
+		return MainSceneComponent;
 	}
-#endif // WITH_EDITOR
 
-	// We have to manage the visibilities of the components that we spawn ourselves, because they will not have prim twins
-	// and so will not be handled by the stage actor. Plus there are some weird rules involved, like how we don't care about the
-	// visibilities of *parents* of prototype prims, only theirs and down from there
-	const bool bPointInstancerVisible = UsdUtils::IsVisible( UE::FUsdPrim{ Prim }, Context->Time );
+	pxr::SdfPathVector PrototypePaths;
+	if ( !PointInstancer.GetPrototypesRel().GetTargets( &PrototypePaths ) )
+	{
+		return MainSceneComponent;
+	}
 
-	// For each prototype
-	const pxr::UsdRelationship& Prototypes = PointInstancer.GetPrototypesRel();
-	pxr::SdfPathVector PrototypesPaths;
+	UUsdAssetCache& AssetCache = *Context->AssetCache.Get();
 
-	TSet< FString > ProcessedPrims;
+	// Lets pretend ParentComponent is pointing to the parent USceneComponent while we create the child HISMs, so they get
+	// automatically attached to it as children
+	TGuardValue< USceneComponent* > ParentComponentGuard{ Context->ParentComponent, MainSceneComponent };
 
-	if ( Prototypes.GetTargets( &PrototypesPaths ) )
+	TArray<TFuture<TTuple<UHierarchicalInstancedStaticMeshComponent*, TArray<FTransform>>>> Tasks;
+	FScopedSlowTask PrototypePathsSlowTask( ( float ) PrototypePaths.size(), LOCTEXT( "GeomPointCreateComponents", "Creating HierarchicalInstancedStaticMeshComponents for point instancers" ) );
+	for ( int32 PrototypeIndex = 0; PrototypeIndex < PrototypePaths.size(); ++PrototypeIndex )
 	{
-		TGuardValue< USceneComponent* > ParentComponentGuard( Context->ParentComponent, PointInstancerRootComponent );
+		PrototypePathsSlowTask.EnterProgressFrame();
 
-		TArray<TFuture<TTuple<UHierarchicalInstancedStaticMeshComponent*, TArray<FTransform>>>> Tasks;
+		const pxr::SdfPath& PrototypePath = PrototypePaths[ PrototypeIndex ];
+		const FString PrototypePathStr = UsdToUnreal::ConvertPath( PrototypePath );
 
-		FScopedSlowTask PrototypesPathsSlowTask( (float)PrototypesPaths.size() * 3, LOCTEXT("GeomPointUpdateComponents", "Updating Point Components") );
-		for ( int32 PrototypeIndex = 0; PrototypeIndex < PrototypesPaths.size(); ++PrototypeIndex)
+		pxr::UsdPrim PrototypeUsdPrim = Prim.GetStage()->GetPrimAtPath( PrototypePath );
+		if ( !PrototypeUsdPrim )
 		{
-			PrototypesPathsSlowTask.EnterProgressFrame();
+			UE_LOG( LogUsd, Warning, TEXT( "Failed to find prototype '%s' for PointInstancer '%s' when creating components" ), *PrototypePathStr, *PrimPath.GetString() );
+			continue;
+		}
 
-			pxr::UsdPrim PrototypeUsdPrim = Prim.GetStage()->GetPrimAtPath( PrototypesPaths[PrototypeIndex] );
-			if ( !PrototypeUsdPrim )
-			{
-				UE_LOG( LogUsd, Warning, TEXT( "Failed to find prototype '%s' for PointInstancer '%s'" ), *UsdToUnreal::ConvertPath( PrototypesPaths[ PrototypeIndex ] ), *PrimPath.GetString() );
-				continue;
-			}
-			UE::FUsdPrim PrototypePrim{ PrototypeUsdPrim };
+		using UHISMComponent = UHierarchicalInstancedStaticMeshComponent;
 
-			USceneComponent* PrototypeParentComponent = Context->ParentComponent;
+		FUsdGeomXformableTranslator XformableTranslator{ UHISMComponent::StaticClass(), Context, UE::FUsdTyped( PrototypeUsdPrim ) };
 
-			bool bPrototypeRootVisible = true;
+		const bool bNeedsActor = false;
+		UHISMComponent* HISMComponent = Cast<UHISMComponent>( XformableTranslator.CreateComponentsEx( { UHISMComponent::StaticClass() }, bNeedsActor ) );
+		if ( !HISMComponent )
+		{
+			UE_LOG( LogUsd, Warning, TEXT( "Failed to generate HISM component for prototype '%s' for PointInstancer '%s'" ), *PrototypePathStr, *PrimPath.GetString() );
+			continue;
+		}
 
-			// Temp fix to prevent us from creating yet another UStaticMeshComponent as a parent prim to the HISM component
-			const bool bNeedsActor = false;
-			TOptional<TGuardValue< USceneComponent* >> ParentComponentGuard2;
-			if ( !PrototypeUsdPrim.IsA<pxr::UsdGeomMesh>() )
-			{
-				FUsdGeomXformableTranslator PrototypeXformTranslator( Context, UE::FUsdTyped( PrototypeUsdPrim ) );
+		UStaticMesh* StaticMesh = Cast< UStaticMesh >( AssetCache.GetAssetForPrim( PrototypePathStr ) );
+		UsdGeomPointInstancerTranslatorImpl::SetStaticMesh( StaticMesh, *HISMComponent );
 
-				if ( USceneComponent* PrototypeXformComponent = PrototypeXformTranslator.CreateComponentsEx( {}, bNeedsActor ) )
+		// Evaluating point instancer can take a long time and is thread-safe. Move to async task while we work on something else.
+		pxr::UsdTimeCode TimeCode{ Context->Time };
+		Tasks.Emplace(
+			Async(
+				EAsyncExecution::ThreadPool,
+				[ TimeCode, Stage = Prim.GetStage(), PointInstancer, PrototypeIndex, HISMComponent ]()
 				{
-					// We have to be careful with prototype visibilities and can't defer to just using ComputeVisibility() because sometimes source apps
-					// can mark the prototype parent prim as invisible, as an attempt to prevent them from being parsed directly as scene meshes.
-					// We will handle the visibility of all the components we spawn manually
-					bPrototypeRootVisible = UsdUtils::HasInheritedVisibility( PrototypePrim, Context->Time );
-					PrototypeXformComponent->SetHiddenInGame( !bPointInstancerVisible || !bPrototypeRootVisible );
+					TArray<FTransform> InstanceTransforms;
+					UsdGeomPointInstancerTranslatorImpl::GetPointInstancerTransforms( Stage, PointInstancer, PrototypeIndex, TimeCode, InstanceTransforms );
 
-					PrototypeParentComponent = PrototypeXformComponent;
+					return MakeTuple( HISMComponent, MoveTemp( InstanceTransforms ) );
 				}
+			)
+		);
 
-				ParentComponentGuard2.Emplace(Context->ParentComponent, PrototypeParentComponent );
-			}
-
-			// #ueent_todo: Fix support nested point instancers. For now, just skip them.
-			TArray< TUsdStore< pxr::UsdPrim > > ChildPointInstancerPrims = UsdUtils::GetAllPrimsOfType(PrototypeUsdPrim, pxr::TfType::Find< pxr::UsdGeomPointInstancer >());
-			if ( ChildPointInstancerPrims.Num() > 0 )
-			{
-				continue;
-			}
-
-			TArray< TUsdStore< pxr::UsdPrim > > ChildGeomMeshPrims = UsdUtils::GetAllPrimsOfType( PrototypeUsdPrim, pxr::TfType::Find< pxr::UsdGeomMesh >() );
-
-			for ( const TUsdStore< pxr::UsdPrim >& PrototypeGeomMeshPrim : ChildGeomMeshPrims )
+		// Handle material overrides
+		if ( StaticMesh )
+		{
+#if WITH_EDITOR
+			// If the prim paths match, it means that it was this prim that created (and so "owns") the static mesh,
+			// so its material assignments will already be directly on the mesh. If they differ, we're using some other prim's mesh,
+			// so we may need material overrides on our component
+			UUsdAssetImportData* UsdImportData = Cast<UUsdAssetImportData>( StaticMesh->AssetImportData );
+			if ( UsdImportData && UsdImportData->PrimPath != PrototypePathStr )
+#endif // WITH_EDITOR
 			{
-				pxr::SdfPath PrototypeTargetPrimPath = PrototypeGeomMeshPrim.Get().GetPrimPath();
-
-				if ( Context->CollapsingCache.IsValid() )
+				TArray<UMaterialInterface*> ExistingAssignments;
+				for ( FStaticMaterial& StaticMaterial : StaticMesh->GetStaticMaterials() )
 				{
-					PrototypeTargetPrimPath = Context->CollapsingCache->UnwindToNonCollapsedPath( UE::FSdfPath{ PrototypeTargetPrimPath }, ECollapsingType::Assets );
+					ExistingAssignments.Add( StaticMaterial.MaterialInterface );
 				}
 
-				const FString UEPrototypeTargetPrimPath = UsdToUnreal::ConvertPath( PrototypeTargetPrimPath );
-
-				if ( !ProcessedPrims.Contains( UEPrototypeTargetPrimPath ) )
-				{
-					ProcessedPrims.Add( UEPrototypeTargetPrimPath );
-
-					if ( pxr::UsdPrim PrototypeTargetUsdPrim = Prim.GetStage()->GetPrimAtPath( PrototypeTargetPrimPath ) )
-					{
-						pxr::UsdGeomMesh PrototypeGeomMesh( PrototypeTargetUsdPrim );
-						FUsdGeomXformableTranslator PrototypeGeomXformableTranslator( UHierarchicalInstancedStaticMeshComponent::StaticClass(), Context, UE::FUsdTyped( PrototypeGeomMesh ) );
-
-						USceneComponent* UsdGeomPrimComponent = PrototypeGeomXformableTranslator.CreateComponentsEx( {}, bNeedsActor );
-
-						if ( UHierarchicalInstancedStaticMeshComponent* HismComponent = Cast< UHierarchicalInstancedStaticMeshComponent >( UsdGeomPrimComponent ) )
-						{
-							UE::FUsdPrim PrototypeTargetPrim{ PrototypeTargetUsdPrim };
-
-							// Traverse up every time because we have no idea how far deep each UsdGeomMesh prim is
-							const bool bHasInvisibleParent = UsdUtils::HasInvisibleParent( PrototypeTargetPrim, PrototypePrim, Context->Time );
-							const bool bPrototypeMeshVisible = UsdUtils::HasInheritedVisibility( PrototypeTargetPrim, Context->Time );
-							HismComponent->SetHiddenInGame( !bPointInstancerVisible || !bPrototypeRootVisible || !bPrototypeMeshVisible || bHasInvisibleParent );
-
-							UUsdAssetCache& AssetCache = *Context->AssetCache.Get();
-							UStaticMesh* StaticMesh = UsdGeomPointInstancerTranslatorImpl::SetStaticMesh( PrototypeGeomMesh, *HismComponent, AssetCache );
-
-							// Evaluating point instancer can take a long time and is thread-safe. Move to async task while we work on something else.
-							Tasks.Emplace(
-								Async(EAsyncExecution::ThreadPool,
-									[this, Stage = Prim.GetStage(), PointInstancer, PrototypeIndex, HismComponent]()
-									{
-										TArray<FTransform> InstanceTransforms;
-										UsdGeomPointInstancerTranslatorImpl::GetPointInstancerTransforms( Stage, PointInstancer, PrototypeIndex, pxr::UsdTimeCode(Context->Time), InstanceTransforms );
-
-										return MakeTuple( HismComponent, MoveTemp(InstanceTransforms) );
-									})
-							);
-
-							// Handle material overrides
-							if ( StaticMesh )
-							{
-#if WITH_EDITOR
-								// If the prim paths match, it means that it was this prim that created (and so "owns") the static mesh,
-								// so its material assignments will already be directly on the mesh. If they differ, we're using some other prim's mesh,
-								// so we may need material overrides on our component
-								UUsdAssetImportData* UsdImportData = Cast<UUsdAssetImportData>( StaticMesh->AssetImportData );
-								if ( UsdImportData && UsdImportData->PrimPath != PrimPath.GetString() )
-#endif // WITH_EDITOR
-								{
-									TArray<UMaterialInterface*> ExistingAssignments;
-									for ( FStaticMaterial& StaticMaterial : StaticMesh->GetStaticMaterials() )
-									{
-										ExistingAssignments.Add( StaticMaterial.MaterialInterface );
-									}
-
-									MeshTranslationImpl::SetMaterialOverrides(
-										PrototypeTargetUsdPrim,
-										ExistingAssignments,
-										*HismComponent,
-										AssetCache,
-										Context->Time,
-										Context->ObjectFlags,
-										Context->bAllowInterpretingLODs,
-										Context->RenderContext
-									);
-								}
-							}
-						}
-					}
-				}
+				MeshTranslationImpl::SetMaterialOverrides(
+					PrototypeUsdPrim,
+					ExistingAssignments,
+					*HISMComponent,
+					AssetCache,
+					Context->Time,
+					Context->ObjectFlags,
+					Context->bAllowInterpretingLODs,
+					Context->RenderContext
+				);
 			}
 		}
+	}
 
-		// Wait on and assign results of the point instancer.
-		for ( auto& Future : Tasks )
-		{
-			TTuple<UHierarchicalInstancedStaticMeshComponent*, TArray<FTransform>> Result { Future.Get() };
-			UsdGeomPointInstancerTranslatorImpl::ApplyPointInstanceTransforms( Result.Key, Result.Value );
-		}
+	// Wait on and assign results of the point instancer.
+	for ( auto& Future : Tasks )
+	{
+		TTuple<UHierarchicalInstancedStaticMeshComponent*, TArray<FTransform>> Result{ Future.Get() };
+		UsdGeomPointInstancerTranslatorImpl::ApplyPointInstanceTransforms( Result.Key, Result.Value );
 	}
 
+	return MainSceneComponent;
+}
+
+void FUsdGeomPointInstancerTranslator::UpdateComponents( USceneComponent* PointInstancerRootComponent )
+{
 	Super::UpdateComponents( PointInstancerRootComponent );
 }
 
diff --git a/Engine/Plugins/Importers/USDImporter/Source/USDSchemas/Private/USDGeomXformableTranslator.cpp b/Engine/Plugins/Importers/USDImporter/Source/USDSchemas/Private/USDGeomXformableTranslator.cpp
--- a/Engine/Plugins/Importers/USDImporter/Source/USDSchemas/Private/USDGeomXformableTranslator.cpp
+++ b/Engine/Plugins/Importers/USDImporter/Source/USDSchemas/Private/USDGeomXformableTranslator.cpp
@@ -14,6 +14,7 @@
 #include "USDSchemasModule.h"
 #include "USDTypesConversion.h"
 
+#include "Components/HierarchicalInstancedStaticMeshComponent.h"
 #include "Components/SceneComponent.h"
 #include "Components/StaticMeshComponent.h"
 #include "Engine/StaticMesh.h"
@@ -95,6 +96,9 @@ void FUsdGeomXformableCreateAssetsTaskChain::SetupTasks()
 				RenderContextToken = UnrealToUsd::ConvertToken( *Context->RenderContext.ToString() ).Get();
 			}
 
+			// We're going to put Prim's transform and visibility on the component, so we don't need to bake it into the combined mesh
+			const bool bSkipRootPrimTransformAndVis = true;
+
 			UsdToUnreal::ConvertGeomMeshHierarchy(
 				GetPrim(),
 				pxr::UsdTimeCode( Context->Time ),
@@ -103,6 +107,7 @@ void FUsdGeomXformableCreateAssetsTaskChain::SetupTasks()
 				*MaterialToPrimvarToUVIndex,
 				AddedMeshDescription,
 				AssignmentInfo,
+				bSkipRootPrimTransformAndVis,
 				Context->bMergeIdenticalMaterialSlots
 			);
 
@@ -389,7 +394,12 @@ void FUsdGeomXformableTranslator::UpdateComponents( USceneComponent* SceneCompon
 			SceneComponent->UnregisterComponent();
 		}
 
-		UsdToUnreal::ConvertXformable( Context->Stage, pxr::UsdGeomXformable( GetPrim() ), *SceneComponent, Context->Time );
+		// Don't convert transform/visibility for HISM components as we only create these for point instancer prototype prims,
+		// and we don't want the actual prototype prim transform/visibility on the component itself: It will always be baked into the mesh
+		if ( !SceneComponent->IsA< UHierarchicalInstancedStaticMeshComponent >() )
+		{
+			UsdToUnreal::ConvertXformable( Context->Stage, pxr::UsdGeomXformable( GetPrim() ), *SceneComponent, Context->Time );
+		}
 
 		// If the user modified a mesh parameter (e.g. vertex color), the hash will be different and it will become a separate asset
 		// so we must check for this and assign the new StaticMesh
diff --git a/Engine/Plugins/Importers/USDImporter/Source/USDSchemas/Public/USDGeomPointInstancerTranslator.h b/Engine/Plugins/Importers/USDImporter/Source/USDSchemas/Public/USDGeomPointInstancerTranslator.h
--- a/Engine/Plugins/Importers/USDImporter/Source/USDSchemas/Public/USDGeomPointInstancerTranslator.h
+++ b/Engine/Plugins/Importers/USDImporter/Source/USDSchemas/Public/USDGeomPointInstancerTranslator.h
@@ -2,25 +2,33 @@
 
 #pragma once
 
+#include "USDGeomMeshTranslator.h"
 #include "USDGeomXformableTranslator.h"
 
 #if USE_USD_SDK
 
-#include "USDIncludesStart.h"
-	#include "pxr/pxr.h"
-	#include "pxr/usd/usdGeom/pointInstancer.h"
-#include "USDIncludesEnd.h"
-
 struct FUsdSchemaTranslationContext;
 
+class FUsdGeomPointInstancerCreateAssetsTaskChain : public FBuildStaticMeshTaskChain
+{
+public:
+	explicit FUsdGeomPointInstancerCreateAssetsTaskChain( const TSharedRef< FUsdSchemaTranslationContext >& InContext, const UE::FSdfPath& InPrimPath );
+
+protected:
+	virtual void SetupTasks() override;
+};
+
 class USDSCHEMAS_API FUsdGeomPointInstancerTranslator : public FUsdGeomXformableTranslator
 {
 public:
 	using Super = FUsdGeomXformableTranslator;
 	using FUsdGeomXformableTranslator::FUsdGeomXformableTranslator;
 
+	virtual void CreateAssets() override;
+	virtual USceneComponent* CreateComponents() override;
 	virtual void UpdateComponents( USceneComponent* SceneComponent ) override;
-	virtual bool CollapsesChildren( ECollapsingType CollapsingType ) const override { return CollapsingType == ECollapsingType::Components; }
+
+	virtual bool CollapsesChildren( ECollapsingType CollapsingType ) const override { return true; }
 	virtual bool CanBeCollapsed( ECollapsingType CollapsingType ) const override { return false; }
 };
 
diff --git a/Engine/Plugins/Importers/USDImporter/Source/USDUtilities/Private/USDGeomMeshConversion.cpp b/Engine/Plugins/Importers/USDImporter/Source/USDUtilities/Private/USDGeomMeshConversion.cpp
--- a/Engine/Plugins/Importers/USDImporter/Source/USDUtilities/Private/USDGeomMeshConversion.cpp
+++ b/Engine/Plugins/Importers/USDImporter/Source/USDUtilities/Private/USDGeomMeshConversion.cpp
@@ -1047,6 +1047,7 @@ bool UsdToUnreal::ConvertGeomMeshHierarchy(
 	const TMap< FString, TMap<FString, int32> >& MaterialToPrimvarToUVIndex,
 	FMeshDescription& OutMeshDescription,
 	UsdUtils::FUsdPrimMaterialAssignmentInfo& OutMaterialAssignments,
+	bool bSkipRootPrimTransformAndVisibility,
 	bool bCombineIdenticalMaterialSlots
 )
 {
@@ -1055,10 +1056,6 @@ bool UsdToUnreal::ConvertGeomMeshHierarchy(
 	FStaticMeshAttributes StaticMeshAttributes( OutMeshDescription );
 	StaticMeshAttributes.Register();
 
-	// Skip the first transform, because if we're parsing Prim's mesh we don't want to bake its own transform into the vertices,
-	// as that transform will be used on components instead
-	const bool bIsFirstPrim = true;
-
 	return UE::UsdGeomMeshConversion::Private::RecursivelyCollapseChildMeshes(
 		Prim,
 		FTransform::Identity,
@@ -1068,7 +1065,7 @@ bool UsdToUnreal::ConvertGeomMeshHierarchy(
 		MaterialToPrimvarToUVIndex,
 		OutMeshDescription,
 		OutMaterialAssignments,
-		bIsFirstPrim,
+		bSkipRootPrimTransformAndVisibility,
 		bCombineIdenticalMaterialSlots
 	);
 }
diff --git a/Engine/Plugins/Importers/USDImporter/Source/USDUtilities/Public/USDGeomMeshConversion.h b/Engine/Plugins/Importers/USDImporter/Source/USDUtilities/Public/USDGeomMeshConversion.h
--- a/Engine/Plugins/Importers/USDImporter/Source/USDUtilities/Public/USDGeomMeshConversion.h
+++ b/Engine/Plugins/Importers/USDImporter/Source/USDUtilities/Public/USDGeomMeshConversion.h
@@ -73,9 +73,10 @@ namespace UsdToUnreal
 	 * @param OutMeshDescription - Output parameter that will be filled with the converted mesh data
 	 * @param OutMaterialAssignments - Output parameter that will be filled with the material data extracted from UsdSchema
 	 * @param bCombineIdenticalMaterialSlots - Whether to try reusing material slots (both local and the ones already in MaterialAssignments) when converting the material assignments from the mesh
+	 * @param bSkipRootPrimTransformAndVisibility - Whether to ignore Prim's transform and visibility when generating the combined mesh. This is useful if you plan on displaying that data on a component with the resulting mesh instead.
 	 * @return Whether the conversion was successful or not.
 	 */
-	USDUTILITIES_API bool ConvertGeomMeshHierarchy( const pxr::UsdPrim& Prim, const pxr::UsdTimeCode& TimeCode, const EUsdPurpose PurposesToLoad, const pxr::TfToken& RenderContext, const TMap< FString, TMap<FString, int32> >& MaterialToPrimvarToUVIndex, FMeshDescription& OutMeshDescription, UsdUtils::FUsdPrimMaterialAssignmentInfo& OutMaterialAssignments, bool bCombineIdenticalMaterialSlots = true );
+	USDUTILITIES_API bool ConvertGeomMeshHierarchy( const pxr::UsdPrim& Prim, const pxr::UsdTimeCode& TimeCode, const EUsdPurpose PurposesToLoad, const pxr::TfToken& RenderContext, const TMap< FString, TMap<FString, int32> >& MaterialToPrimvarToUVIndex, FMeshDescription& OutMeshDescription, UsdUtils::FUsdPrimMaterialAssignmentInfo& OutMaterialAssignments, bool bSkipRootPrimTransformAndVisibility = false, bool bCombineIdenticalMaterialSlots = true );
 
 	/** Configure Material to become a vertex color/displayColor material, according to the given description */
 	USDUTILITIES_API bool ConvertDisplayColor( const UsdUtils::FDisplayColorMaterial& DisplayColorDescription, UMaterialInstanceConstant& Material );
