USD: Also bind actors whenever we bind one of its child components to the LevelSequence.

diff --git a/Engine/Plugins/Importers/USDImporter/Source/USDStage/Private/USDLevelSequenceHelper.cpp b/Engine/Plugins/Importers/USDImporter/Source/USDStage/Private/USDLevelSequenceHelper.cpp
--- a/Engine/Plugins/Importers/USDImporter/Source/USDStage/Private/USDLevelSequenceHelper.cpp
+++ b/Engine/Plugins/Importers/USDImporter/Source/USDStage/Private/USDLevelSequenceHelper.cpp
@@ -1552,12 +1552,44 @@ TrackType* FUsdLevelSequenceHelperImpl::AddTrack( const FName& TrackName, const
 		}
 		else
 		{
+			FGuid ComponentBinding;
+			FGuid ActorBinding;
+			UObject* ComponentContext = ComponentToBind.GetWorld();
+
+			FString PrimName = FPaths::GetBaseFilename( PrimTwin.PrimPath );
+
+			// Make sure we always bind the parent actor too
+			if ( AActor* Actor = ComponentToBind.GetOwner() )
+			{
+				ActorBinding = Sequence.FindBindingFromObject( Actor, Actor->GetWorld() );
+				if ( !ActorBinding.IsValid() )
+				{
+					// We use the label here because that will always be named after the prim that caused the actor
+					// to be generated. If we just used our own PrimName in here we may run into situations where a child Camera prim
+					// of a decomposed camera ends up naming the actor binding after itself, even though the parent Xform prim, and the
+					// actor on the level, maybe named something else
+					ActorBinding = MovieScene->AddPossessable( Actor->GetActorLabel(), Actor->GetClass() );
+					Sequence.BindPossessableObject( ActorBinding, *Actor, Actor->GetWorld() );
+				}
+
+				ComponentContext = Actor;
+			}
+
+			ComponentBinding = MovieScene->AddPossessable( PrimName, ComponentToBind.GetClass() );
+
+			if ( ActorBinding.IsValid() && ComponentBinding.IsValid() )
+			{
+				if ( FMovieScenePossessable* ComponentPossessable = MovieScene->FindPossessable( ComponentBinding ) )
+				{
+					ComponentPossessable->SetParent( ActorBinding );
+				}
+			}
+
 			// Bind component
-			FGuid Binding = MovieScene->AddPossessable( FPaths::GetBaseFilename( PrimTwin.PrimPath ), ComponentToBind.GetClass() );
-			Sequence.BindPossessableObject( Binding, ComponentToBind, ComponentToBind.GetWorld() );
+			Sequence.BindPossessableObject( ComponentBinding, ComponentToBind, ComponentContext );
 
-			SceneComponentsBindings.Emplace( &PrimTwin ) = TPair< ULevelSequence*, FGuid >( &Sequence, Binding );
-			return Binding;
+			SceneComponentsBindings.Emplace( &PrimTwin ) = TPair< ULevelSequence*, FGuid >( &Sequence, ComponentBinding );
+			return ComponentBinding;
 		}
 	}( );
 
@@ -1694,12 +1726,44 @@ void FUsdLevelSequenceHelperImpl::RemovePossessable( ULevelSequence& Sequence, c
 		return;
 	}
 
+	// The RemovePossessable calls Modify the MovieScene already, but the UnbindPossessableObject
+	// ones don't modify the Sequence and change properties, so we must modify them here
+	Sequence.Modify();
+
 	if ( const TPair< ULevelSequence*, FGuid >* SceneComponentBinding = SceneComponentsBindings.Find( &PrimTwin ) )
 	{
+		const FGuid& ComponentPossessableGuid = SceneComponentBinding->Value;
+
+		FGuid ActorPossessableGuid;
+		if ( FMovieScenePossessable* ComponentPossessable = MovieScene->FindPossessable( ComponentPossessableGuid ) )
+		{
+			ActorPossessableGuid = ComponentPossessable->GetParent();
+		}
+
 		// This will also remove all tracks bound to this guid
-		if ( MovieScene->RemovePossessable( SceneComponentBinding->Value ) )
+		if ( MovieScene->RemovePossessable( ComponentPossessableGuid ) )
 		{
-			Sequence.UnbindPossessableObjects( SceneComponentBinding->Value );
+			Sequence.UnbindPossessableObjects( ComponentPossessableGuid );
+		}
+
+		// If our parent binding has nothing else in it, we should remove it too
+		bool bRemoveActorBinding = true;
+		if ( ActorPossessableGuid.IsValid() )
+		{
+			for ( int32 PossessableIndex = 0; PossessableIndex < MovieScene->GetPossessableCount(); ++PossessableIndex )
+			{
+				const FMovieScenePossessable& SomePossessable = MovieScene->GetPossessable( PossessableIndex );
+				if ( SomePossessable.GetParent() == ActorPossessableGuid )
+				{
+					bRemoveActorBinding = false;
+					break;
+				}
+			}
+		}
+		if ( bRemoveActorBinding )
+		{
+			MovieScene->RemovePossessable( ActorPossessableGuid );
+			Sequence.UnbindPossessableObjects( ActorPossessableGuid );
 		}
 
 		SceneComponentsBindings.Remove( &PrimTwin );
@@ -2260,16 +2324,8 @@ void FUsdLevelSequenceHelperImpl::HandleTrackChange( const UMovieSceneTrack& Tra
 		return;
 	}
 
-	USceneComponent* BoundSceneComponent = nullptr;
-	if ( AActor* BoundActor = Cast< AActor >( BoundObject ) )
-	{
-		BoundSceneComponent = BoundActor->GetRootComponent();
-	}
-	else
-	{
-		BoundSceneComponent = Cast< USceneComponent >( BoundObject );
-	}
-
+	// Our tracked bindings are always directly to components. If we don't have one here just abort
+	USceneComponent* BoundSceneComponent = Cast< USceneComponent >( BoundObject );
 	if ( !BoundSceneComponent )
 	{
 		return;
