USD: Manually group notices emitted when toggling payloads for multiple prims.

diff --git a/Engine/Plugins/Importers/USDImporter/Source/USDStageEditor/Private/SUSDStageTreeView.cpp b/Engine/Plugins/Importers/USDImporter/Source/USDStageEditor/Private/SUSDStageTreeView.cpp
--- a/Engine/Plugins/Importers/USDImporter/Source/USDStageEditor/Private/SUSDStageTreeView.cpp
+++ b/Engine/Plugins/Importers/USDImporter/Source/USDStageEditor/Private/SUSDStageTreeView.cpp
@@ -12,6 +12,7 @@
 
 #include "UsdWrappers/SdfChangeBlock.h"
 #include "UsdWrappers/SdfLayer.h"
+#include "UsdWrappers/SdfPath.h"
 #include "UsdWrappers/UsdPrim.h"
 #include "UsdWrappers/UsdStage.h"
 
@@ -872,30 +873,50 @@ void SUsdStageTreeView::RestoreExpansionStates()
 
 void SUsdStageTreeView::OnToggleAllPayloads( EPayloadsTrigger PayloadsTrigger )
 {
+	if ( !UsdStageActor.IsValid() )
+	{
+		return;
+	}
+
 	TArray< FUsdPrimViewModelRef > MySelectedItems = GetSelectedItems();
 
+	// Ideally we'd just use a UE::FSdfChangeBlock here, but for whatever reason this doesn't seem to affect the
+	// notices USD emits when loading/unloading prim payloads, so we must do this via the UsdStage directly
+
+	TSet<UE::FSdfPath> PrimsToLoad;
+	TSet<UE::FSdfPath> PrimsToUnload;
+
 	for ( FUsdPrimViewModelRef SelectedItem : MySelectedItems )
 	{
 		if ( SelectedItem->UsdPrim )
 		{
 			TFunction< void( FUsdPrimViewModelRef ) > RecursiveTogglePayloads;
-			RecursiveTogglePayloads = [ &RecursiveTogglePayloads, PayloadsTrigger ]( FUsdPrimViewModelRef InSelectedItem ) -> void
+			RecursiveTogglePayloads = [ &RecursiveTogglePayloads, PayloadsTrigger, &PrimsToLoad, &PrimsToUnload ]( FUsdPrimViewModelRef InSelectedItem ) -> void
 			{
 				UE::FUsdPrim& UsdPrim = InSelectedItem->UsdPrim;
 
 				if ( UsdPrim.HasPayload() )
 				{
+					bool bPrimIsLoaded = UsdPrim.IsLoaded();
+
 					if ( PayloadsTrigger == EPayloadsTrigger::Toggle )
 					{
-						InSelectedItem->TogglePayload();
+						if ( bPrimIsLoaded )
+						{
+							PrimsToUnload.Add( UsdPrim.GetPrimPath() );
+						}
+						else
+						{
+							PrimsToLoad.Add( UsdPrim.GetPrimPath() );
+						}
 					}
-					else if ( PayloadsTrigger == EPayloadsTrigger::Load && !UsdPrim.IsLoaded() )
+					else if ( PayloadsTrigger == EPayloadsTrigger::Load && !bPrimIsLoaded )
 					{
-						InSelectedItem->TogglePayload();
+						PrimsToLoad.Add( UsdPrim.GetPrimPath() );
 					}
-					else if ( PayloadsTrigger == EPayloadsTrigger::Unload && UsdPrim.IsLoaded() )
+					else if ( PayloadsTrigger == EPayloadsTrigger::Unload && bPrimIsLoaded )
 					{
-						InSelectedItem->TogglePayload();
+						PrimsToUnload.Add( UsdPrim.GetPrimPath() );
 					}
 				}
 				else
@@ -910,6 +931,14 @@ void SUsdStageTreeView::OnToggleAllPayloads( EPayloadsTrigger PayloadsTrigger )
 			RecursiveTogglePayloads( SelectedItem );
 		}
 	}
+
+	if ( PrimsToLoad.Num() + PrimsToUnload.Num() > 0 )
+	{
+		UE::FSdfChangeBlock GroupNotices;
+
+		UE::FUsdStage UsdStage = UsdStageActor->GetOrLoadUsdStage();
+		UsdStage.LoadAndUnload( PrimsToLoad, PrimsToUnload );
+	}
 }
 
 void SUsdStageTreeView::ScrollItemIntoView( FUsdPrimViewModelRef TreeItem )
diff --git a/Engine/Plugins/Importers/USDImporter/Source/UnrealUSDWrapper/Private/UsdWrappers/UsdStage.cpp b/Engine/Plugins/Importers/USDImporter/Source/UnrealUSDWrapper/Private/UsdWrappers/UsdStage.cpp
--- a/Engine/Plugins/Importers/USDImporter/Source/UnrealUSDWrapper/Private/UsdWrappers/UsdStage.cpp
+++ b/Engine/Plugins/Importers/USDImporter/Source/UnrealUSDWrapper/Private/UsdWrappers/UsdStage.cpp
@@ -200,6 +200,37 @@ namespace UE
 	}
 #endif // USE_USD_SDK
 
+	template<typename PtrType>
+	void FUsdStageBase<PtrType>::LoadAndUnload( const TSet<FSdfPath>& LoadSet, const TSet<FSdfPath>& UnloadSet, EUsdLoadPolicy Policy )
+	{
+#if USE_USD_SDK
+		if ( PtrType& Ptr = Impl->GetInner() )
+		{
+			FScopedUsdAllocs Allocs;
+
+			pxr::SdfPathSet UsdLoadSet;
+			for ( const FSdfPath& Path : LoadSet )
+			{
+				UsdLoadSet.insert( static_cast< pxr::SdfPath >( Path ) );
+			}
+
+			pxr::SdfPathSet UsdUnloadSet;
+			for ( const FSdfPath& Path : UnloadSet )
+			{
+				UsdUnloadSet.insert( static_cast< pxr::SdfPath >( Path ) );
+			}
+
+			Ptr->LoadAndUnload(
+				UsdLoadSet,
+				UsdUnloadSet,
+				Policy == EUsdLoadPolicy::UsdLoadWithDescendants
+					? pxr::UsdLoadPolicy::UsdLoadWithDescendants
+					: pxr::UsdLoadPolicy::UsdLoadWithoutDescendants
+			);
+		}
+#endif
+	}
+
 	template<typename PtrType>
 	FSdfLayer FUsdStageBase<PtrType>::GetRootLayer() const
 	{
diff --git a/Engine/Plugins/Importers/USDImporter/Source/UnrealUSDWrapper/Public/UsdWrappers/UsdStage.h b/Engine/Plugins/Importers/USDImporter/Source/UnrealUSDWrapper/Public/UsdWrappers/UsdStage.h
--- a/Engine/Plugins/Importers/USDImporter/Source/UnrealUSDWrapper/Public/UsdWrappers/UsdStage.h
+++ b/Engine/Plugins/Importers/USDImporter/Source/UnrealUSDWrapper/Public/UsdWrappers/UsdStage.h
@@ -18,6 +18,13 @@ namespace UE
 		template< typename PtrType > class FUsdStageImpl;
 	}
 
+	/** Corresponds to pxr::UsdLoadPolicy, refer to the USD SDK documentation */
+	enum class EUsdLoadPolicy
+	{
+		UsdLoadWithDescendants,    // Load a prim plus all its descendants.
+		UsdLoadWithoutDescendants  // Load a prim by itself with no descendants.
+	};
+
 	/**
 	 * Minimal pxr::UsdStage pointer wrapper for Unreal that can be used from no-rtti modules.
 	 * Use the aliases FUsdStage and FUsdStageWeak instead (defined on ForwardDeclarations.h)
@@ -63,6 +70,8 @@ namespace UE
 
 	// Wrapped pxr::UsdStage functions, refer to the USD SDK documentation
 	public:
+		void LoadAndUnload( const TSet<UE::FSdfPath>& LoadSet, const TSet<UE::FSdfPath>& UnloadSet, EUsdLoadPolicy Policy = EUsdLoadPolicy::UsdLoadWithDescendants );
+
 		/**
 		 * Saves a flattened copy of the stage to the given path (e.g. "C:/Folder/FlattenedStage.usda"). Will use the corresponding file writer depending on FilePath extension.
 		 * Will not alter the current stage.
