USD: Fix subtree not updating visibility when a prim with animated visibilities is hidden/shown.

diff --git a/Engine/Plugins/Importers/USDImporter/Source/USDStage/Private/USDLevelSequenceHelper.cpp b/Engine/Plugins/Importers/USDImporter/Source/USDStage/Private/USDLevelSequenceHelper.cpp
--- a/Engine/Plugins/Importers/USDImporter/Source/USDStage/Private/USDLevelSequenceHelper.cpp
+++ b/Engine/Plugins/Importers/USDImporter/Source/USDStage/Private/USDLevelSequenceHelper.cpp
@@ -279,7 +279,7 @@ private:
 	void CreateTimeTrack(const FUsdLevelSequenceHelperImpl::FLayerTimeInfo& Info);
 	void RemoveTimeTrack(const FUsdLevelSequenceHelperImpl::FLayerTimeInfo* Info);
 
-	void AddCommonTracks( const UUsdPrimTwin& PrimTwin, const UE::FUsdPrim& Prim );
+	void AddCommonTracks( const UUsdPrimTwin& PrimTwin, const UE::FUsdPrim& Prim, bool bForceVisibilityTracks = false );
 	void AddCameraTracks( const UUsdPrimTwin& PrimTwin, const UE::FUsdPrim& Prim );
 	void AddLightTracks( const UUsdPrimTwin& PrimTwin, const UE::FUsdPrim& Prim, const TSet<FName>& PropertyPathsToRead = {} );
 	void AddSkeletalTracks( const UUsdPrimTwin& PrimTwin, const UE::FUsdPrim& Prim );
@@ -291,7 +291,10 @@ private:
 
 // Prims handling
 public:
-	void AddPrim( UUsdPrimTwin& PrimTwin );
+	// If bForceVisibilityTracks is true, we'll add and bake the visibility tracks for this prim even if the
+	// prim itself doesn't have animated visibility (so that we can handle its visibility in case one of its
+	// parents does have visibility animations)
+	void AddPrim( UUsdPrimTwin& PrimTwin, bool bForceVisibilityTracks = false );
 	void RemovePrim(const UUsdPrimTwin& PrimTwin);
 
 private:
@@ -1065,7 +1068,7 @@ void FUsdLevelSequenceHelperImpl::RemoveTimeTrack(const FLayerTimeInfo* LayerTim
 	}
 }
 
-void FUsdLevelSequenceHelperImpl::AddCommonTracks( const UUsdPrimTwin& PrimTwin, const UE::FUsdPrim& Prim )
+void FUsdLevelSequenceHelperImpl::AddCommonTracks( const UUsdPrimTwin& PrimTwin, const UE::FUsdPrim& Prim, bool bForceVisibilityTracks )
 {
 	USceneComponent* ComponentToBind = PrimTwin.GetSceneComponent();
 	if ( !ComponentToBind )
@@ -1121,31 +1124,92 @@ void FUsdLevelSequenceHelperImpl::AddCommonTracks( const UUsdPrimTwin& PrimTwin,
 	if ( Attrs.Num() > 0 )
 	{
 		UE::FUsdAttribute VisibilityAttribute = Attrs[ 0 ];
-		if ( VisibilityAttribute && VisibilityAttribute.ValueMightBeTimeVarying() )
+		if ( VisibilityAttribute )
 		{
-			if ( ULevelSequence* AttributeSequence = FindOrAddSequenceForAttribute( VisibilityAttribute ) )
+			const bool bAttrMightBeTimeVarying = VisibilityAttribute.ValueMightBeTimeVarying();
+
+			// Collect all the time samples we'll need to sample our visibility at (USD has inherited visibilities, so every time
+			// a parent has a key, we need to recompute the child visibility at that moment too)
+			TArray<double> TotalVisibilityTimeSamples;
+
+			// If we're adding a visibility track because a parent has visibility animations, we want to write our baked visibility
+			// tracks on the same layer as the first one of our parents that actually has animated visibility.
+			// There's no ideal place for this because it's essentially a fake track that we're creating, and we may have arbitrarily
+			// many parents and specs on multiple layers, but this is hopefully at least *a* reasonable answer.
+			UE::FUsdAttribute FirstAnimatedVisibilityParentAttr;
+
+			if ( bAttrMightBeTimeVarying || bForceVisibilityTracks )
 			{
-				const bool bIsMuted = UsdUtils::IsAttributeMuted( VisibilityAttribute, UsdStage );
+				VisibilityAttribute.GetTimeSamples( TotalVisibilityTimeSamples );
 
-				FMovieSceneSequenceTransform SequenceTransform;
-				FMovieSceneSequenceID SequenceID = SequencesID.FindRef( AttributeSequence );
-				if ( FMovieSceneSubSequenceData* SubSequenceData = SequenceHierarchyCache.FindSubData( SequenceID ) )
+				// TODO: Improve this, as this is extremely inefficient since we'll be parsing this tree for the root down and repeatedly
+				// redoing this one child at a time...
+				UE::FUsdPrim ParentPrim = Prim.GetParent();
+				while ( ParentPrim && !ParentPrim.IsPseudoRoot() )
 				{
-					SequenceTransform = SubSequenceData->RootToSequenceTransform;
+					if ( UsdUtils::HasAnimatedVisibility( ParentPrim ) )
+					{
+						TArray<UE::FUsdAttribute> ParentAttrs = UnrealToUsd::GetAttributesForProperty( ParentPrim, UnrealIdentifiers::HiddenInGamePropertyName );
+						if ( ParentAttrs.Num() > 0 )
+						{
+							UE::FUsdAttribute ParentVisAttr = ParentAttrs[ 0 ];
+							if ( ParentVisAttr && ParentVisAttr.ValueMightBeTimeVarying() )
+							{
+								if ( !FirstAnimatedVisibilityParentAttr )
+								{
+									FirstAnimatedVisibilityParentAttr = ParentVisAttr;
+								}
+
+								TArray<double> TimeSamples;
+								if ( ParentVisAttr.GetTimeSamples( TimeSamples ) )
+								{
+									TotalVisibilityTimeSamples.Append( TimeSamples );
+								}
+							}
+						}
+					}
+
+					ParentPrim = ParentPrim.GetParent();
 				}
 
-				if ( UMovieScene* MovieScene = AttributeSequence->GetMovieScene() )
+				// Put these in order for the sampling below, but don't worry about duplicates: The baking process already skips
+				// consecutive duplicates anyway
+				TotalVisibilityTimeSamples.Sort();
+			}
+
+			// Pick which attribute we will use to fetch the target LevelSequence to put our baked tracks
+			UE::FUsdAttribute AttributeForSequence;
+			if ( bAttrMightBeTimeVarying )
+			{
+				AttributeForSequence = VisibilityAttribute;
+			}
+			if ( !AttributeForSequence && bForceVisibilityTracks && FirstAnimatedVisibilityParentAttr )
+			{
+				AttributeForSequence = FirstAnimatedVisibilityParentAttr;
+			}
+
+			if ( AttributeForSequence && TotalVisibilityTimeSamples.Num() > 0 )
+			{
+				if ( ULevelSequence* AttributeSequence = FindOrAddSequenceForAttribute( AttributeForSequence ) )
 				{
-					const bool bReadonly = false;
-					UsdLevelSequenceHelperImpl::FMovieSceneReadonlyGuard MovieSceneReadonlyGuard{ *MovieScene, bReadonly };
+					const bool bIsMuted = UsdUtils::IsAttributeMuted( AttributeForSequence, UsdStage );
 
-					TArray<double> TimeSamples;
-					if ( VisibilityAttribute.GetTimeSamples( TimeSamples ) )
+					FMovieSceneSequenceTransform SequenceTransform;
+					FMovieSceneSequenceID SequenceID = SequencesID.FindRef( AttributeSequence );
+					if ( FMovieSceneSubSequenceData* SubSequenceData = SequenceHierarchyCache.FindSubData( SequenceID ) )
 					{
+						SequenceTransform = SubSequenceData->RootToSequenceTransform;
+					}
+
+					if ( UMovieScene* MovieScene = AttributeSequence->GetMovieScene() )
+					{
+						const bool bReadonly = false;
+						UsdLevelSequenceHelperImpl::FMovieSceneReadonlyGuard MovieSceneReadonlyGuard{ *MovieScene, bReadonly };
+
 						if ( UMovieSceneVisibilityTrack* VisibilityTrack = AddTrack<UMovieSceneVisibilityTrack>( UnrealIdentifiers::HiddenInGamePropertyName, PrimTwin, *ComponentToBind, *AttributeSequence, bIsMuted ) )
 						{
 							UsdToUnreal::FPropertyTrackReader Reader = UsdToUnreal::CreatePropertyTrackReader( Prim, UnrealIdentifiers::HiddenInGamePropertyName );
-							UsdToUnreal::ConvertBoolTimeSamples( UsdStage, TimeSamples, Reader.BoolReader, *VisibilityTrack, SequenceTransform );
+							UsdToUnreal::ConvertBoolTimeSamples( UsdStage, TotalVisibilityTimeSamples, Reader.BoolReader, *VisibilityTrack, SequenceTransform );
 						}
 
 						PrimPathByLevelSequenceName.AddUnique( AttributeSequence->GetFName(), Prim.GetPrimPath().GetString() );
@@ -1477,7 +1541,7 @@ void FUsdLevelSequenceHelperImpl::AddSkeletalTracks( const UUsdPrimTwin& PrimTwi
 	PrimPathByLevelSequenceName.AddUnique( SkelAnimationSequence->GetFName(), PrimPath );
 }
 
-void FUsdLevelSequenceHelperImpl::AddPrim( UUsdPrimTwin& PrimTwin )
+void FUsdLevelSequenceHelperImpl::AddPrim( UUsdPrimTwin& PrimTwin, bool bForceVisibilityTracks )
 {
 	if ( !UsdStage )
 	{
@@ -1525,7 +1589,7 @@ void FUsdLevelSequenceHelperImpl::AddPrim( UUsdPrimTwin& PrimTwin )
 		AddSkeletalTracks( PrimTwin, UsdPrim );
 	}
 
-	AddCommonTracks( PrimTwin, UsdPrim );
+	AddCommonTracks( PrimTwin, UsdPrim, bForceVisibilityTracks );
 
 	RefreshSequencer();
 }
@@ -2584,7 +2648,7 @@ public:
 	void UnbindFromUsdStageActor() {}
 	void OnStageActorRenamed() {};
 
-	void AddPrim( UUsdPrimTwin& PrimTwin ) {}
+	void AddPrim( UUsdPrimTwin& PrimTwin, bool bForceVisibilityTracks ) {}
 	void RemovePrim(const UUsdPrimTwin& PrimTwin) {}
 
 	void StartMonitoringChanges() {}
@@ -2701,11 +2765,11 @@ void FUsdLevelSequenceHelper::UnbindFromUsdStageActor()
 	}
 }
 
-void FUsdLevelSequenceHelper::AddPrim(UUsdPrimTwin& PrimTwin)
+void FUsdLevelSequenceHelper::AddPrim(UUsdPrimTwin& PrimTwin, bool bForceVisibilityTracks)
 {
 	if (UsdSequencerImpl.IsValid())
 	{
-		UsdSequencerImpl->AddPrim(PrimTwin);
+		UsdSequencerImpl->AddPrim(PrimTwin, bForceVisibilityTracks);
 	}
 }
 
diff --git a/Engine/Plugins/Importers/USDImporter/Source/USDStage/Private/USDStageActor.cpp b/Engine/Plugins/Importers/USDImporter/Source/USDStage/Private/USDStageActor.cpp
--- a/Engine/Plugins/Importers/USDImporter/Source/USDStage/Private/USDStageActor.cpp
+++ b/Engine/Plugins/Importers/USDImporter/Source/USDStage/Private/USDStageActor.cpp
@@ -1129,6 +1129,32 @@ UUsdPrimTwin* AUsdStageActor::ExpandPrim( const UE::FUsdPrim& Prim, FUsdSchemaTr
 	{
 		if ( !PrimsToAnimate.Contains( UsdPrimTwin->PrimPath ) )
 		{
+			// Unfortunately if we have animated visibility we need to be ready to update the visibility
+			// of all components that we spawned for child prims whenever this prim's visibility updates.
+			// We can't just have this prim's FUsdGeomXformableTranslator::UpdateComponents ->
+			// -> UsdToUnreal::ConvertXformable call use SetHiddenInGame recursively, because we may have
+			// child prims that are themselves also invisible, and so their own subtrees should be invisible
+			// even if this prim goes visible. Also keep in mind that technically we'll always update each
+			// prim in the order that they are within PrimsToAnimate, but that order is not strictly enforced
+			// to be e.g. a breadth first traversal on the prim tree or anything like this, so these updates
+			// need to be order-independent, which means we really should add the entire subtree to the list
+			// and have UpdateComponents called on all components.
+			if ( UsdUtils::HasAnimatedVisibility( Prim ) )
+			{
+				const bool bRecursive = true;
+				UsdPrimTwin->Iterate(
+					[this]( UUsdPrimTwin& Twin )
+					{
+						if ( !PrimsToAnimate.Contains( Twin.PrimPath ) )
+						{
+							PrimsToAnimate.Add( Twin.PrimPath );
+							LevelSequenceHelper.AddPrim( Twin );
+						}
+					},
+					bRecursive
+				);
+			}
+
 			PrimsToAnimate.Add( UsdPrimTwin->PrimPath );
 			LevelSequenceHelper.AddPrim( *UsdPrimTwin );
 		}
diff --git a/Engine/Plugins/Importers/USDImporter/Source/USDStage/Public/USDLevelSequenceHelper.h b/Engine/Plugins/Importers/USDImporter/Source/USDStage/Public/USDLevelSequenceHelper.h
--- a/Engine/Plugins/Importers/USDImporter/Source/USDStage/Public/USDLevelSequenceHelper.h
+++ b/Engine/Plugins/Importers/USDImporter/Source/USDStage/Public/USDLevelSequenceHelper.h
@@ -54,8 +54,14 @@ public:
 	void BindToUsdStageActor(AUsdStageActor* StageActor);
 	void UnbindFromUsdStageActor();
 
-	/** Adds the necessary tracks for a given prim to the level sequence */
-	void AddPrim(UUsdPrimTwin& PrimTwin);
+	/**
+	 * Adds the necessary tracks for a given prim to the level sequence.
+	 * If bForceVisibilityTracks is true, will add visibility tracks even if this prim
+	 * doesn't actually have timeSamples on its visibility attribute (use this when
+	 * a parent does have animated visibility, and we need to "bake" that out to a dedicated
+	 * visibility track so that the standalone LevelSequence asset behaves as expected)
+	 */
+	void AddPrim(UUsdPrimTwin& PrimTwin, bool bForceVisibilityTracks = false);
 
 	/** Removes any track associated with this prim */
 	void RemovePrim(const UUsdPrimTwin& PrimTwin);
diff --git a/Engine/Plugins/Importers/USDImporter/Source/USDStageImporter/Private/USDStageImporter.cpp b/Engine/Plugins/Importers/USDImporter/Source/USDStageImporter/Private/USDStageImporter.cpp
--- a/Engine/Plugins/Importers/USDImporter/Source/USDStageImporter/Private/USDStageImporter.cpp
+++ b/Engine/Plugins/Importers/USDImporter/Source/USDStageImporter/Private/USDStageImporter.cpp
@@ -269,7 +269,7 @@ namespace UsdStageImporterImpl
 #endif // #if USE_USD_SDK
 	}
 
-	void ImportAnimation(FUsdStageImportContext& ImportContext, UE::FUsdPrim& Prim, USceneComponent* SceneComponent)
+	void ImportAnimation(FUsdStageImportContext& ImportContext, UE::FUsdPrim& Prim, bool bAnimatedVisibility, USceneComponent* SceneComponent)
 	{
 		TRACE_CPUPROFILER_EVENT_SCOPE( ImportAnimation );
 
@@ -282,10 +282,10 @@ namespace UsdStageImporterImpl
 		UsdPrimTwin->PrimPath = Prim.GetPrimPath().GetString();
 		UsdPrimTwin->SceneComponent = SceneComponent;
 
-		ImportContext.LevelSequenceHelper.AddPrim( *UsdPrimTwin );
+		ImportContext.LevelSequenceHelper.AddPrim( *UsdPrimTwin, bAnimatedVisibility );
 	}
 
-	void ImportActor(FUsdStageImportContext& ImportContext, UE::FUsdPrim& Prim, FUsdSchemaTranslationContext& TranslationContext)
+	void ImportActor(FUsdStageImportContext& ImportContext, UE::FUsdPrim& Prim, bool bForceVisibilityAnimationTracks, FUsdSchemaTranslationContext& TranslationContext)
 	{
 		TRACE_CPUPROFILER_EVENT_SCOPE( ImportActor );
 
@@ -304,6 +304,12 @@ namespace UsdStageImporterImpl
 			bExpandChilren = !SchemaTranslator->CollapsesChildren(ECollapsingType::Components);
 		}
 
+		// In USD if a parent prim has animated visibility, it will affect the entire subtree. In UE this doesn't
+		// happen by default, so if our visibility is animated (or if we're forced to animate visibility from a parent prim),
+		// then we must ensure that we generate visibility tracks for our child prims that don't collapse either, so that
+		// they get the chance to bake their "computed visibilities" and e.g. be hidden whenever a parent prim goes invisible
+		const bool bAnimatedVisibility = bForceVisibilityAnimationTracks || UsdUtils::HasAnimatedVisibility( Prim );
+
 		SlowTask.EnterProgressFrame();
 		// Recurse to children
 		if (bExpandChilren)
@@ -312,9 +318,9 @@ namespace UsdStageImporterImpl
 			TGuardValue<USceneComponent*> ParentComponentGuard(TranslationContext.ParentComponent, ContextParentComponent);
 
 			const bool bTraverseInstanceProxies = true;
-			for (UE::FUsdPrim ChildStore : Prim.GetFilteredChildren(bTraverseInstanceProxies))
+			for (UE::FUsdPrim Child : Prim.GetFilteredChildren(bTraverseInstanceProxies))
 			{
-				ImportActor(ImportContext, ChildStore, TranslationContext);
+				ImportActor(ImportContext, Child, bAnimatedVisibility, TranslationContext);
 			}
 		}
 		SlowTask.EnterProgressFrame();
@@ -329,9 +335,9 @@ namespace UsdStageImporterImpl
 			}
 
 #if USE_USD_SDK
-			if (UsdUtils::IsAnimated(Prim))
+			if ( bAnimatedVisibility || UsdUtils::IsAnimated( Prim ) )
 			{
-				ImportAnimation(ImportContext, Prim, Component);
+				ImportAnimation(ImportContext, Prim, bAnimatedVisibility, Component);
 			}
 #endif // USE_USD_SDK
 		}
@@ -345,8 +351,9 @@ namespace UsdStageImporterImpl
 			return;
 		}
 
+		const bool bForceVisibilityAnimationTracks = false;
 		UE::FUsdPrim RootPrim = ImportContext.Stage.GetPseudoRoot();
-		ImportActor(ImportContext, RootPrim, TranslationContext);
+		ImportActor(ImportContext, RootPrim, bForceVisibilityAnimationTracks, TranslationContext);
 	}
 
 	// Assets coming out of USDSchemas module have default names, so here we do our best to provide them with
diff --git a/Engine/Plugins/Importers/USDImporter/Source/USDUtilities/Private/USDConversionUtils.cpp b/Engine/Plugins/Importers/USDImporter/Source/USDUtilities/Private/USDConversionUtils.cpp
--- a/Engine/Plugins/Importers/USDImporter/Source/USDUtilities/Private/USDConversionUtils.cpp
+++ b/Engine/Plugins/Importers/USDImporter/Source/USDUtilities/Private/USDConversionUtils.cpp
@@ -55,6 +55,7 @@
 	#include "pxr/usd/usd/stage.h"
 	#include "pxr/usd/usd/variantSets.h"
 	#include "pxr/usd/usdGeom/camera.h"
+	#include "pxr/usd/usdGeom/imageable.h"
 	#include "pxr/usd/usdGeom/mesh.h"
 	#include "pxr/usd/usdGeom/metrics.h"
 	#include "pxr/usd/usdGeom/primvar.h"
@@ -760,6 +761,30 @@ bool UsdUtils::IsAnimated( const pxr::UsdPrim& Prim )
 	return false;
 }
 
+bool UsdUtils::HasAnimatedVisibility( const pxr::UsdPrim& Prim )
+{
+	if ( !Prim || !Prim.IsActive() )
+	{
+		return false;
+	}
+
+	FScopedUsdAllocs UsdAllocs;
+
+	pxr::UsdGeomImageable Imageable( Prim );
+	if ( Imageable )
+	{
+		if ( pxr::UsdAttribute Attr = Imageable.GetVisibilityAttr() )
+		{
+			if ( Attr.ValueMightBeTimeVarying() )
+			{
+				return true;
+			}
+		}
+	}
+
+	return false;
+}
+
 EUsdDefaultKind UsdUtils::GetDefaultKind( const pxr::UsdPrim& Prim )
 {
 	FScopedUsdAllocs Allocs;
diff --git a/Engine/Plugins/Importers/USDImporter/Source/USDUtilities/Private/USDPrimConversion.cpp b/Engine/Plugins/Importers/USDImporter/Source/USDUtilities/Private/USDPrimConversion.cpp
--- a/Engine/Plugins/Importers/USDImporter/Source/USDUtilities/Private/USDPrimConversion.cpp
+++ b/Engine/Plugins/Importers/USDImporter/Source/USDUtilities/Private/USDPrimConversion.cpp
@@ -536,8 +536,16 @@ bool UsdToUnreal::ConvertBoolTimeSamples( const UE::FUsdStage& Stage, const TArr
 	const double StageTimeCodesPerSecond = UsdStage->GetTimeCodesPerSecond();
 	const FFrameRate StageFrameRate( StageTimeCodesPerSecond, 1 );
 
+	double LastTimeSample = TNumericLimits<double>::Lowest();
 	for ( const double UsdTimeSample : UsdTimeSamples )
 	{
+		// We never want to evaluate the same time twice
+		if ( FMath::IsNearlyEqual( UsdTimeSample, LastTimeSample ) )
+		{
+			continue;
+		}
+		LastTimeSample = UsdTimeSample;
+
 		int32 FrameNumber = FMath::FloorToInt( UsdTimeSample );
 		float SubFrameNumber = UsdTimeSample - FrameNumber;
 
@@ -599,8 +607,16 @@ bool UsdToUnreal::ConvertFloatTimeSamples( const UE::FUsdStage& Stage, const TAr
 
 	const ERichCurveInterpMode InterpMode = ( UsdStage->GetInterpolationType() == pxr::UsdInterpolationTypeLinear ) ? ERichCurveInterpMode::RCIM_Linear : ERichCurveInterpMode::RCIM_Constant;
 
+	double LastTimeSample = TNumericLimits<double>::Lowest();
 	for ( const double UsdTimeSample : UsdTimeSamples )
 	{
+		// We never want to evaluate the same time twice
+		if ( FMath::IsNearlyEqual( UsdTimeSample, LastTimeSample ) )
+		{
+			continue;
+		}
+		LastTimeSample = UsdTimeSample;
+
 		int32 FrameNumber = FMath::FloorToInt( UsdTimeSample );
 		float SubFrameNumber = UsdTimeSample - FrameNumber;
 
@@ -667,8 +683,16 @@ bool UsdToUnreal::ConvertColorTimeSamples( const UE::FUsdStage& Stage, const TAr
 
 	const ERichCurveInterpMode InterpMode = ( UsdStage->GetInterpolationType() == pxr::UsdInterpolationTypeLinear ) ? ERichCurveInterpMode::RCIM_Linear : ERichCurveInterpMode::RCIM_Constant;
 
+	double LastTimeSample = TNumericLimits<double>::Lowest();
 	for ( const double UsdTimeSample : UsdTimeSamples )
 	{
+		// We never want to evaluate the same time twice
+		if ( FMath::IsNearlyEqual( UsdTimeSample, LastTimeSample ) )
+		{
+			continue;
+		}
+		LastTimeSample = UsdTimeSample;
+
 		int32 FrameNumber = FMath::FloorToInt( UsdTimeSample );
 		float SubFrameNumber = UsdTimeSample - FrameNumber;
 
@@ -758,8 +782,16 @@ bool UsdToUnreal::ConvertTransformTimeSamples( const UE::FUsdStage& Stage, const
 
 	const ERichCurveInterpMode InterpMode = ( UsdStage->GetInterpolationType() == pxr::UsdInterpolationTypeLinear ) ? ERichCurveInterpMode::RCIM_Linear : ERichCurveInterpMode::RCIM_Constant;
 
+	double LastTimeSample = TNumericLimits<double>::Lowest();
 	for ( const double UsdTimeSample : UsdTimeSamples )
 	{
+		// We never want to evaluate the same time twice
+		if ( FMath::IsNearlyEqual( UsdTimeSample, LastTimeSample ) )
+		{
+			continue;
+		}
+		LastTimeSample = UsdTimeSample;
+
 		int32 FrameNumber = FMath::FloorToInt( UsdTimeSample );
 		float SubFrameNumber = UsdTimeSample - FrameNumber;
 
diff --git a/Engine/Plugins/Importers/USDImporter/Source/USDUtilities/Public/USDConversionUtils.h b/Engine/Plugins/Importers/USDImporter/Source/USDUtilities/Public/USDConversionUtils.h
--- a/Engine/Plugins/Importers/USDImporter/Source/USDUtilities/Public/USDConversionUtils.h
+++ b/Engine/Plugins/Importers/USDImporter/Source/USDUtilities/Public/USDConversionUtils.h
@@ -152,6 +152,7 @@ namespace UsdUtils
 	USDUTILITIES_API TArray< TUsdStore< pxr::UsdGeomPrimvar > > GetUVSetPrimvars( const pxr::UsdGeomMesh& UsdMesh, const TMap< FString, TMap< FString, int32 > >& MaterialToPrimvarsUVSetNames, const UsdUtils::FUsdPrimMaterialAssignmentInfo& UsdMeshMaterialAssignmentInfo );
 
 	USDUTILITIES_API bool IsAnimated( const pxr::UsdPrim& Prim );
+	USDUTILITIES_API bool HasAnimatedVisibility( const pxr::UsdPrim& Prim );
 
 	/** Returns whether Prim belongs to any of the default kinds, or a kind derived from them. The result can be a union of different kinds. */
 	USDUTILITIES_API EUsdDefaultKind GetDefaultKind( const pxr::UsdPrim& Prim );
