USD: Fix not being able to duplicate a stage actor with a loaded stage.

diff --git a/Engine/Plugins/Importers/USDImporter/Source/USDClasses/Private/USDAssetCache.cpp b/Engine/Plugins/Importers/USDImporter/Source/USDClasses/Private/USDAssetCache.cpp
--- a/Engine/Plugins/Importers/USDImporter/Source/USDClasses/Private/USDAssetCache.cpp
+++ b/Engine/Plugins/Importers/USDImporter/Source/USDClasses/Private/USDAssetCache.cpp
@@ -91,9 +91,9 @@ void UUsdAssetCache::DiscardAsset( const FString& Hash )
 
 	if ( FoundObject )
 	{
-		for ( TMap< FString, UObject* >::TIterator PrimPathToAssetIt = PrimPathToAssets.CreateIterator(); PrimPathToAssetIt; ++PrimPathToAssetIt )
+		for ( TMap< FString, TWeakObjectPtr<UObject> >::TIterator PrimPathToAssetIt = PrimPathToAssets.CreateIterator(); PrimPathToAssetIt; ++PrimPathToAssetIt )
 		{
-			if ( *FoundObject == PrimPathToAssetIt.Value() )
+			if ( *FoundObject == PrimPathToAssetIt.Value().Get() )
 			{
 				PrimPathToAssetIt.RemoveCurrent();
 			}
@@ -163,10 +163,13 @@ UObject* UUsdAssetCache::GetAssetForPrim( const FString& PrimPath ) const
 {
 	FScopeLock Lock( &CriticalSection );
 
-	if ( UObject* const* FoundObject = PrimPathToAssets.Find( PrimPath ) )
+	if ( TWeakObjectPtr<UObject> const* FoundObjectPtr = PrimPathToAssets.Find( PrimPath ) )
 	{
-		ActiveAssets.Add( *FoundObject );
-		return *FoundObject;
+		if ( UObject* FoundObject = FoundObjectPtr->Get() )
+		{
+			ActiveAssets.Add( FoundObject );
+			return FoundObject;
+		}
 	}
 
 	return nullptr;
diff --git a/Engine/Plugins/Importers/USDImporter/Source/USDClasses/Public/USDAssetCache.h b/Engine/Plugins/Importers/USDImporter/Source/USDClasses/Public/USDAssetCache.h
--- a/Engine/Plugins/Importers/USDImporter/Source/USDClasses/Public/USDAssetCache.h
+++ b/Engine/Plugins/Importers/USDImporter/Source/USDClasses/Public/USDAssetCache.h
@@ -32,7 +32,7 @@ public:
 	/** Returns the first prim path associated with an asset (see TMap::FindKey) */
 	FString GetPrimForAsset( UObject* Asset ) const;
 
-	TMap< FString, UObject* > GetAssetPrimLinks() const { return PrimPathToAssets; }; // Can't return a reference as it wouldn't be thread-safe
+	TMap< FString, TWeakObjectPtr<UObject> > GetAssetPrimLinks() const { return PrimPathToAssets; }; // Can't return a reference as it wouldn't be thread-safe
 
 	bool IsAssetOwnedByCache( UObject* Asset ) const { return OwnedAssets.Contains( Asset ); }
 
@@ -56,22 +56,25 @@ public:
 	virtual void Serialize( FArchive& Ar ) override;
 
 private:
-	UPROPERTY( Transient, VisibleAnywhere, Category = "Assets" )
+	// For now everything is NonPIEDuplicateTransient as this is mostly a subobject of AUsdStageActor. When the actor is duplicated
+	// it will need to reload the stage anyway to rebuild its prim links to components and assets, so there is no point in duplicating
+	// the properties here just yet. Obviously we want to duplicate these properties
+	UPROPERTY( NonPIEDuplicateTransient, Transient, VisibleAnywhere, Category = "Assets" )
 	TMap< FString, UObject* > TransientStorage;
 
-	UPROPERTY( VisibleAnywhere, Category = "Assets" )
+	UPROPERTY( NonPIEDuplicateTransient, VisibleAnywhere, Category = "Assets" )
 	TMap< FString, UObject* > PersistentStorage;
 
 	UPROPERTY( EditAnywhere, Category = "Assets", AdvancedDisplay )
 	bool bAllowPersistentStorage;
 
 	// Points to the assets in primary storage, used to quickly check if we own an asset
-	UPROPERTY()
-	TSet< UObject* > OwnedAssets;
+	UPROPERTY( NonPIEDuplicateTransient )
+	TSet< TWeakObjectPtr<UObject> > OwnedAssets;
 
 	// Keeps associations from prim paths to assets that we own in primary storage
-	UPROPERTY()
-    TMap< FString, UObject* > PrimPathToAssets;
+	UPROPERTY( NonPIEDuplicateTransient )
+    TMap< FString, TWeakObjectPtr<UObject> > PrimPathToAssets;
 
 	// Assets that were added/retrieved since the last call to MarkAssetsAsSlate();
 	mutable TSet<UObject*> ActiveAssets;
diff --git a/Engine/Plugins/Importers/USDImporter/Source/USDStage/Private/USDStageActor.cpp b/Engine/Plugins/Importers/USDImporter/Source/USDStage/Private/USDStageActor.cpp
--- a/Engine/Plugins/Importers/USDImporter/Source/USDStage/Private/USDStageActor.cpp
+++ b/Engine/Plugins/Importers/USDImporter/Source/USDStage/Private/USDStageActor.cpp
@@ -1433,9 +1433,9 @@ FString AUsdStageActor::GetSourcePrimPath( UObject* Object )
 
 	if ( AssetCache )
 	{
-		for ( const TPair<FString, UObject*>& PrimPathToAsset : AssetCache->GetAssetPrimLinks() )
+		for ( const TPair<FString, TWeakObjectPtr<UObject>>& PrimPathToAsset : AssetCache->GetAssetPrimLinks() )
 		{
-			if ( PrimPathToAsset.Value == Object )
+			if ( PrimPathToAsset.Value.Get() == Object )
 			{
 				return PrimPathToAsset.Key;
 			}
@@ -2081,6 +2081,10 @@ void AUsdStageActor::PostDuplicate( bool bDuplicateForPIE )
 	{
 		OpenUsdStage();
 	}
+	else
+	{
+		LoadUsdStage();
+	}
 }
 
 void AUsdStageActor::Serialize(FArchive& Ar)
@@ -2637,7 +2641,7 @@ void AUsdStageActor::LoadAssets( FUsdSchemaTranslationContext& TranslationContex
 	{
 		FString StartPrimPath = StartPrim.GetPrimPath().GetString();
 		TSet<FString> PrimPathsToRemove;
-		for ( const TPair< FString, UObject* >& PrimPathToAssetIt : AssetCache->GetAssetPrimLinks() )
+		for ( const TPair< FString, TWeakObjectPtr<UObject> >& PrimPathToAssetIt : AssetCache->GetAssetPrimLinks() )
 		{
 			const FString& PrimPath = PrimPathToAssetIt.Key;
 			if ( PrimPath.StartsWith( StartPrimPath ) || PrimPath == StartPrimPath )
