USD: Merge of 19383969 to safeguard against crashes when calling BroadcastLevelActorListChanged() during garbage collection.

diff --git a/Engine/Plugins/Importers/USDImporter/Source/USDStage/Private/USDStageActor.cpp b/Engine/Plugins/Importers/USDImporter/Source/USDStage/Private/USDStageActor.cpp
--- a/Engine/Plugins/Importers/USDImporter/Source/USDStage/Private/USDStageActor.cpp
+++ b/Engine/Plugins/Importers/USDImporter/Source/USDStage/Private/USDStageActor.cpp
@@ -1202,7 +1202,7 @@ void AUsdStageActor::UpdatePrim( const UE::FSdfPath& InUsdPrimPath, bool bResync
 		UUsdPrimTwin* UsdPrimTwin = ExpandPrim( PrimToExpand, TranslationContext );
 
 #if WITH_EDITOR
-		if ( GIsEditor && GEditor ) // Make sure we're not in standalone either
+		if ( GIsEditor && GEditor && !IsGarbageCollecting() ) // Make sure we're not in standalone either
 		{
 			GEditor->BroadcastLevelActorListChanged();
 			GEditor->RedrawLevelEditingViewports();
@@ -1744,7 +1744,7 @@ void AUsdStageActor::LoadUsdStage()
 	}
 
 #if WITH_EDITOR
-	if ( GIsEditor && GEditor )
+	if ( GIsEditor && GEditor && !IsGarbageCollecting() )
 	{
 		GEditor->BroadcastLevelActorListChanged();
 	}
@@ -2727,7 +2727,7 @@ void AUsdStageActor::AnimatePrims()
 	TranslationContext->CompleteTasks();
 
 #if WITH_EDITOR
-	if ( GIsEditor && GEditor )
+	if ( GIsEditor && GEditor && !IsGarbageCollecting() )
 	{
 		GEditor->BroadcastLevelActorListChanged();
 		GEditor->RedrawLevelEditingViewports();
