Fix for Clang compiler error

diff --git a/Engine/Plugins/Importers/USDImporter/Source/USDStageImporter/Private/USDStageImporter.cpp b/Engine/Plugins/Importers/USDImporter/Source/USDStageImporter/Private/USDStageImporter.cpp
--- a/Engine/Plugins/Importers/USDImporter/Source/USDStageImporter/Private/USDStageImporter.cpp
+++ b/Engine/Plugins/Importers/USDImporter/Source/USDStageImporter/Private/USDStageImporter.cpp
@@ -308,7 +308,11 @@ namespace UsdStageImporterImpl
 		// happen by default, so if our visibility is animated (or if we're forced to animate visibility from a parent prim),
 		// then we must ensure that we generate visibility tracks for our child prims that don't collapse either, so that
 		// they get the chance to bake their "computed visibilities" and e.g. be hidden whenever a parent prim goes invisible
+#if USE_USD_SDK
 		const bool bAnimatedVisibility = bForceVisibilityAnimationTracks || UsdUtils::HasAnimatedVisibility( Prim );
+#else
+		const bool bAnimatedVisibility = bForceVisibilityAnimationTracks;
+#endif
 
 		SlowTask.EnterProgressFrame();
 		// Recurse to children
