USD: Fix crash when using Action->Import and having the stage actor's LevelSequence open in the Sequencer.

diff --git a/Engine/Plugins/Importers/USDImporter/Source/USDStage/Private/USDLevelSequenceHelper.cpp b/Engine/Plugins/Importers/USDImporter/Source/USDStage/Private/USDLevelSequenceHelper.cpp
--- a/Engine/Plugins/Importers/USDImporter/Source/USDStage/Private/USDLevelSequenceHelper.cpp
+++ b/Engine/Plugins/Importers/USDImporter/Source/USDStage/Private/USDLevelSequenceHelper.cpp
@@ -1849,8 +1849,13 @@ void FUsdLevelSequenceHelperImpl::RefreshSequencer()
 
 	if ( TSharedPtr< ISequencer > Sequencer = WeakSequencer.Pin() )
 	{
-		Sequencer->RefreshTree();
-		Sequencer->NotifyMovieSceneDataChanged( EMovieSceneDataChangeType::TrackValueChanged );
+		// Don't try refreshing the sequencer if its displaying a stale sequence (e.g. during busy transitions like import) as it
+		// can crash
+		if ( UMovieSceneSequence* FocusedSequence = Sequencer->GetFocusedMovieSceneSequence() )
+		{
+			Sequencer->RefreshTree();
+			Sequencer->NotifyMovieSceneDataChanged( EMovieSceneDataChangeType::TrackValueChanged );
+		}
 	}
 #endif // WITH_EDITOR
 }
