USD: Add support for reading resetXformStack xformOp.

Prevent sampling the USD stage more than once at the same time code when parsing animations.

diff --git a/Engine/Plugins/Importers/USDImporter/Source/USDSchemas/Private/USDGeomXformableTranslator.cpp b/Engine/Plugins/Importers/USDImporter/Source/USDSchemas/Private/USDGeomXformableTranslator.cpp
--- a/Engine/Plugins/Importers/USDImporter/Source/USDSchemas/Private/USDGeomXformableTranslator.cpp
+++ b/Engine/Plugins/Importers/USDImporter/Source/USDSchemas/Private/USDGeomXformableTranslator.cpp
@@ -364,11 +364,14 @@ USceneComponent* FUsdGeomXformableTranslator::CreateComponentsEx( TOptional< TSu
 			SceneComponent->Mobility = UsdUtils::IsAnimated( Prim ) ? EComponentMobility::Movable : EComponentMobility::Static;
 		}
 
-		UpdateComponents( SceneComponent );
-
 		// Attach to parent
+		// Do this before UpdatingComponents as we may need to use the parent transform to set a world transform directly
+		// (in case of resetXformStack). Besides, this is more consistent anyway as during stage updates we'll call
+		// UpdateComponents with all the components already attached
 		SceneComponent->AttachToComponent( Context->ParentComponent, FAttachmentTransformRules::KeepRelativeTransform );
 
+		UpdateComponents( SceneComponent );
+
 		if ( !SceneComponent->IsRegistered() )
 		{
 			SceneComponent->RegisterComponent();
diff --git a/Engine/Plugins/Importers/USDImporter/Source/USDStage/Private/USDLevelSequenceHelper.cpp b/Engine/Plugins/Importers/USDImporter/Source/USDStage/Private/USDLevelSequenceHelper.cpp
--- a/Engine/Plugins/Importers/USDImporter/Source/USDStage/Private/USDLevelSequenceHelper.cpp
+++ b/Engine/Plugins/Importers/USDImporter/Source/USDStage/Private/USDLevelSequenceHelper.cpp
@@ -1080,7 +1080,47 @@ void FUsdLevelSequenceHelperImpl::AddCommonTracks( const UUsdPrimTwin& PrimTwin,
 	ULevelSequence* PrimSequence = FindSequenceForIdentifier( PrimLayer.GetIdentifier() );
 
 	UE::FUsdGeomXformable Xformable( Prim );
-	if ( Xformable.TransformMightBeTimeVarying() ) // Test that transform might be time varying and not TransformAttribute since it will check each xform ops
+
+	// If this xformable has an op to reset the xform stack and one of its ancestors is animated, then we need to make
+	// a transform track for it even if its transform is not animated by itself. This because that op effectively means
+	// "discard the parent transform and treat this as a direct world transform", but when reading we'll manually recompute
+	// the relative transform to its parent anyway (for simplicity's sake). If that parent (or any of its ancestors) is
+	// being animated, we'll need to recompute this for every animation keyframe
+	TArray<double> AncestorTimeSamples;
+	bool bNeedTrackToCompensateResetXformOp = false;
+	if ( Xformable.GetResetXformStack() )
+	{
+		UE::FUsdPrim AncestorPrim = Prim.GetParent();
+		while ( AncestorPrim && !AncestorPrim.IsPseudoRoot() )
+		{
+			if ( UE::FUsdGeomXformable AncestorXformable{ AncestorPrim } )
+			{
+				if ( AncestorXformable.TransformMightBeTimeVarying() )
+				{
+					bNeedTrackToCompensateResetXformOp = true;
+
+					TArray<double> TimeSamples;
+					if ( AncestorXformable.GetTimeSamples( &TimeSamples ) )
+					{
+						AncestorTimeSamples.Append( TimeSamples );
+					}
+				}
+
+				// The exception is if our ancestor also wants to reset its xform stack (i.e. its transform is meant to be
+				// used as the world transform). In this case we don't need to care about higher up ancestors anymore, as
+				// their transforms wouldn't affect below this prim anyway
+				if ( AncestorXformable.GetResetXformStack() )
+				{
+					break;
+				}
+			}
+
+			AncestorPrim = AncestorPrim.GetParent();
+		}
+	}
+
+	// Test that transform might be time varying and not TransformAttribute since it will check each xform ops
+	if ( Xformable.TransformMightBeTimeVarying() || bNeedTrackToCompensateResetXformOp )
 	{
 		TArray<UE::FUsdAttribute> Attrs = UnrealToUsd::GetAttributesForProperty( Prim, UnrealIdentifiers::TransformPropertyName );
 		if ( Attrs.Num() > 0 )
@@ -1106,6 +1146,12 @@ void FUsdLevelSequenceHelperImpl::AddCommonTracks( const UUsdPrimTwin& PrimTwin,
 						TArray<double> TimeSamples;
 						if ( Xformable.GetTimeSamples( &TimeSamples ) )
 						{
+							if ( bNeedTrackToCompensateResetXformOp )
+							{
+								TimeSamples.Append( AncestorTimeSamples );
+								TimeSamples.Sort();
+							}
+
 							if ( UMovieScene3DTransformTrack* TransformTrack = AddTrack<UMovieScene3DTransformTrack>( UnrealIdentifiers::TransformPropertyName, PrimTwin, *ComponentToBind, *AttributeSequence, bIsMuted ) )
 							{
 								UsdToUnreal::FPropertyTrackReader Reader = UsdToUnreal::CreatePropertyTrackReader( Prim, UnrealIdentifiers::TransformPropertyName );
diff --git a/Engine/Plugins/Importers/USDImporter/Source/USDUtilities/Private/USDConversionUtils.cpp b/Engine/Plugins/Importers/USDImporter/Source/USDUtilities/Private/USDConversionUtils.cpp
--- a/Engine/Plugins/Importers/USDImporter/Source/USDUtilities/Private/USDConversionUtils.cpp
+++ b/Engine/Plugins/Importers/USDImporter/Source/USDUtilities/Private/USDConversionUtils.cpp
@@ -722,6 +722,36 @@ bool UsdUtils::IsAnimated( const pxr::UsdPrim& Prim )
 		{
 			return true;
 		}
+
+		// If this xformable has an op to reset the xform stack and one of its ancestors is animated, then we need to pretend
+		// its transform is also animated. This because that op effectively means "discard the parent transform and treat this
+		// as a direct world transform", but when reading we'll manually recompute the relative transform to its parent anyway
+		// (for simplicity's sake). If that parent (or any of its ancestors) is being animated, we'll need to recompute this
+		// for every animation keyframe, which basically means we're animated too
+		if ( Xformable.GetResetXformStack() )
+		{
+			pxr::UsdPrim AncestorPrim = Prim.GetParent();
+			while ( AncestorPrim && !AncestorPrim.IsPseudoRoot() )
+			{
+				if ( pxr::UsdGeomXformable AncestorXformable{ AncestorPrim } )
+				{
+					if ( AncestorXformable.TransformMightBeTimeVarying() )
+					{
+						return true;
+					}
+
+					// The exception is if our ancestor also wants to reset its xform stack (i.e. its transform is meant to be
+					// used as the world transform). In this case we don't need to care about higher up ancestors anymore, as
+					// their transforms wouldn't affect below this prim anyway
+					if ( AncestorXformable.GetResetXformStack() )
+					{
+						break;
+					}
+				}
+
+				AncestorPrim = AncestorPrim.GetParent();
+			}
+		}
 	}
 
 	const std::vector< pxr::UsdAttribute >& Attributes = Prim.GetAttributes();
diff --git a/Engine/Plugins/Importers/USDImporter/Source/USDUtilities/Private/USDPrimConversion.cpp b/Engine/Plugins/Importers/USDImporter/Source/USDUtilities/Private/USDPrimConversion.cpp
--- a/Engine/Plugins/Importers/USDImporter/Source/USDUtilities/Private/USDPrimConversion.cpp
+++ b/Engine/Plugins/Importers/USDImporter/Source/USDUtilities/Private/USDPrimConversion.cpp
@@ -232,11 +232,91 @@ namespace UE
 
 				return NewOp;
 			}
+
+			// Turns OutTransform into the UE-space relative (local to parent) transform for Xformable, paying attention to if it
+			// or any of its ancestors has the '!resetXformStack!' xformOp.
+			void GetPrimConvertedRelativeTransform( pxr::UsdGeomXformable Xformable, double UsdTimeCode, FTransform& OutTransform )
+			{
+				if ( !Xformable )
+				{
+					return;
+				}
+
+				FScopedUsdAllocs Allocs;
+
+				pxr::UsdPrim UsdPrim = Xformable.GetPrim();
+				pxr::UsdStageRefPtr UsdStage = UsdPrim.GetStage();
+
+				bool bResetTransformStack = false;
+				UsdToUnreal::ConvertXformable( UsdStage, Xformable, OutTransform, UsdTimeCode, &bResetTransformStack );
+
+				// If we have the resetXformStack op on this prim's xformOpOrder we have to essentially use its transform
+				// as the world transform (i.e. we have to discard the parent transforms). We won't do this here, and will instead
+				// keep relative transforms everywhere for consistency, which means we must manually invert the ParentToWorld transform
+				// and compute our relative transform ourselves.
+				//
+				// Ideally we'd query the components for this for performance reasons, but not only we don't have access to them here,
+				// but neither the stage actor's PrimsToAnimate nor the sequencer guarantee a particular evaluation order anyway,
+				// which means that if our parent is also animated, we could end up computing our relative transforms using the outdated
+				// parent's transform instead. This means we must compute our relative transform using the actual prim hierarchy.
+				//
+				// Additionally, our parent prims may be animated, so we must query all of our ancestors for a new world matrix every frame.
+				//
+				// We could use UsdGeomXformCache for this, but given that we won't actually cache anything (since we'll have to resample
+				// all ancestors every frame anyway) and that we would have to manually handle the camera/light compensation at least for
+				// our immediate parent, it's simpler to just recursively call our own UsdToUnreal::ConvertXformable and concatenate the
+				// results. Its not as fast, but we'll only do this on the initial read for prims with `resetXformStack`, so it should
+				// be very rare. We don't ever write out the resetXformStack either, so after that initial read this op should just disappear.
+				//
+				// Note that, alternatively, we could also handle this whole situation by having the scene components specify their transforms
+				// as absolute, and the Sequencer would work with that as well. However that would spread out the handling of
+				// resetXformStack through all USD workflows, and mean we'd have to *write out* resetXformStack when writing/exporting
+				// absolute transform components, and also convert between them when the user toggles between relative/absolute manually,
+				// which is probably worse than just baking it as relative transforms on first read and forgetting about it.
+				if ( bResetTransformStack )
+				{
+					FTransform ParentToWorld = FTransform::Identity;
+
+					pxr::UsdPrim AncestorPrim = UsdPrim.GetParent();
+					while ( AncestorPrim && !AncestorPrim.IsPseudoRoot() )
+					{
+						FTransform AncestorTransform = FTransform::Identity;
+						bool bAncestorResetTransformStack = false;
+						UsdToUnreal::ConvertXformable( UsdStage, pxr::UsdGeomXformable{ AncestorPrim }, AncestorTransform, UsdTimeCode, &bAncestorResetTransformStack );
+
+						ParentToWorld = ParentToWorld * AncestorTransform;
+
+						// If we find a parent that also has the resetXformStack, then we're in luck: That transform value will be its world
+						// transform already, so we can stop concatenating stuff. Yes, on the component-side of things we'd have done the same
+						// thing of making a fake relative transform for it, but the end result would have been the same final world transform
+						if ( bAncestorResetTransformStack )
+						{
+							break;
+						}
+
+						AncestorPrim = AncestorPrim.GetParent();
+					}
+
+					const FVector& Scale = ParentToWorld.GetScale3D();
+					if ( !FMath::IsNearlyEqual( Scale.X, Scale.Y ) || !FMath::IsNearlyEqual( Scale.X, Scale.Z ) )
+					{
+						UE_LOG( LogUsd, Warning, TEXT( "Inverting transform with non-uniform scaling '%s' when computing relative transform for prim '%s'! Result will likely be incorrect, since FTransforms can't invert non-uniform scalings. You can work around this by baking your non-uniform scaling transform into the vertices, or by not using the !resetXformStack! Xform op." ),
+							*Scale.ToString(),
+							*UsdToUnreal::ConvertPath( UsdPrim.GetPrimPath() )
+						);
+					}
+
+					// Multiplying with matrices here helps mitigate the issues encountered with non-uniform scaling, however it will stil
+					// never be perfect, as it is not possible to generate an FTransform that can properly invert a complex transform with non-uniform
+					// scaling when just multiplying them (which is what downstream code within USceneComponent will do).
+					OutTransform = FTransform{ OutTransform.ToMatrixWithScale() * ParentToWorld.ToInverseMatrixWithScale() };
+				}
+			}
 		}
 	}
 }
 
-bool UsdToUnreal::ConvertXformable( const pxr::UsdStageRefPtr& Stage, const pxr::UsdTyped& Schema, FTransform& OutTransform, double EvalTime )
+bool UsdToUnreal::ConvertXformable( const pxr::UsdStageRefPtr& Stage, const pxr::UsdTyped& Schema, FTransform& OutTransform, double EvalTime, bool* bOutResetTransformStack )
 {
 	pxr::UsdGeomXformable Xformable( Schema );
 	if ( !Xformable )
@@ -248,21 +328,24 @@ bool UsdToUnreal::ConvertXformable( const pxr::UsdStageRefPtr& Stage, const pxr:
 
 	// Transform
 	pxr::GfMatrix4d UsdMatrix;
-	bool bResetXFormStack = false;
-	Xformable.GetLocalTransformation( &UsdMatrix, &bResetXFormStack, EvalTime );
+	bool bResetXformStack = false;
+	bool* bResetXformStackPtr = bOutResetTransformStack ? bOutResetTransformStack : &bResetXformStack;
+	Xformable.GetLocalTransformation( &UsdMatrix, bResetXformStackPtr, EvalTime );
 
 	FUsdStageInfo StageInfo( Stage );
 	OutTransform = UsdToUnreal::ConvertMatrix( StageInfo, UsdMatrix );
 
-	// Extra rotation to match different camera facing direction convention
-	// Note: The camera space is always Y-up, yes, but this is not what this is: This is the camera's transform wrt the stage,
-	// which follows the stage up axis
-	if ( Xformable.GetPrim().IsA< pxr::UsdGeomCamera >() ||
+	const bool bPrimIsLight = 
 #if defined(HAS_USDLUX_LIGHTAPI)
-			Xformable.GetPrim().HasAPI< pxr::UsdLuxLightAPI >() )
+			Xformable.GetPrim().HasAPI< pxr::UsdLuxLightAPI >();
 #else
-			Xformable.GetPrim().IsA< pxr::UsdLuxLight >() )
+			Xformable.GetPrim().IsA< pxr::UsdLuxLight >();
 #endif // #if defined(HAS_USDLUX_LIGHTAPI)
+
+	// Extra rotation to match different camera facing direction convention
+	// Note: The camera space is always Y-up, yes, but this is not what this is: This is the camera's transform wrt the stage,
+	// which follows the stage up axis
+	if ( Xformable.GetPrim().IsA< pxr::UsdGeomCamera >() || bPrimIsLight )
 	{
 		if ( StageInfo.UpAxis == EUsdUpAxis::YAxis )
 		{
@@ -276,12 +359,17 @@ bool UsdToUnreal::ConvertXformable( const pxr::UsdStageRefPtr& Stage, const pxr:
 	// Invert the compensation applied to our parents, in case they're a camera or a light
 	if ( pxr::UsdPrim Parent = Xformable.GetPrim().GetParent() )
 	{
-		if ( Parent.IsA< pxr::UsdGeomCamera >() ||
+		const bool bParentIsLight = 
 #if defined(HAS_USDLUX_LIGHTAPI)
-				Parent.HasAPI< pxr::UsdLuxLightAPI >() )
+				Parent.HasAPI< pxr::UsdLuxLightAPI >();
 #else
-				Parent.IsA< pxr::UsdLuxLight >() )
+				Parent.IsA< pxr::UsdLuxLight >();
 #endif // #if defined(HAS_USDLUX_LIGHTAPI)
+
+		// If bResetXFormStack is true, then the prim's local transform will be used directly as the world transform, and we will
+		// already invert the parent transform fully, regardless of what it is. This means it doesn't really matter if our parent
+		// has a camera/light compensation or not, and so we don't have to have the explicit inverse compensation here anyway!
+		if ( !(*bResetXformStackPtr) && ( Parent.IsA< pxr::UsdGeomCamera >() || bParentIsLight ) )
 		{
 			if ( StageInfo.UpAxis == EUsdUpAxis::YAxis )
 			{
@@ -311,8 +399,7 @@ bool UsdToUnreal::ConvertXformable( const pxr::UsdStageRefPtr& Stage, const pxr:
 
 	// Transform
 	FTransform Transform;
-	UsdToUnreal::ConvertXformable( Stage, Xformable, Transform, EvalTime );
-
+	UE::USDPrimConversionImpl::Private::GetPrimConvertedRelativeTransform( Xformable, EvalTime, Transform );
 	SceneComponent.SetRelativeTransform( Transform );
 
 	SceneComponent.Modify();
@@ -860,12 +947,12 @@ UsdToUnreal::FPropertyTrackReader UsdToUnreal::CreatePropertyTrackReader( const
 		if ( PropertyPath == UnrealIdentifiers::TransformPropertyName )
 		{
 			FTransform Default = FTransform::Identity;
-			UsdToUnreal::ConvertXformable( UsdStage, Xformable, Default, UsdUtils::GetDefaultTimeCode() );
+			UE::USDPrimConversionImpl::Private::GetPrimConvertedRelativeTransform( Xformable, UsdUtils::GetDefaultTimeCode(), Default );
 
 			Reader.TransformReader = [UsdStage, Xformable, Default]( double UsdTimeCode )
 			{
 				FTransform Result = Default;
-				UsdToUnreal::ConvertXformable( UsdStage, Xformable, Result, UsdTimeCode );
+				UE::USDPrimConversionImpl::Private::GetPrimConvertedRelativeTransform( Xformable, UsdTimeCode, Result );
 				return Result;
 			};
 			return Reader;
@@ -2352,10 +2439,6 @@ bool UnrealToUsd::ConvertXformable( const FTransform& RelativeTransform, pxr::Us
 
 	const pxr::UsdTimeCode UsdTimeCode( TimeCode );
 
-	pxr::GfMatrix4d UsdMatrix;
-	bool bResetXFormStack = false;
-	XForm.GetLocalTransformation( &UsdMatrix, &bResetXFormStack, UsdTimeCode );
-
 	if ( pxr::UsdGeomXformOp MatrixXform = UE::USDPrimConversionImpl::Private::ForceMatrixXform( XForm ) )
 	{
 		MatrixXform.Set( UsdTransform, UsdTimeCode );
diff --git a/Engine/Plugins/Importers/USDImporter/Source/USDUtilities/Public/USDPrimConversion.h b/Engine/Plugins/Importers/USDImporter/Source/USDUtilities/Public/USDPrimConversion.h
--- a/Engine/Plugins/Importers/USDImporter/Source/USDUtilities/Public/USDPrimConversion.h
+++ b/Engine/Plugins/Importers/USDImporter/Source/USDUtilities/Public/USDPrimConversion.h
@@ -61,7 +61,7 @@ namespace UsdToUnreal
 	USDUTILITIES_API bool ConvertXformable( const pxr::UsdTyped& Schema, UMovieScene3DTransformTrack& MovieSceneTrack, const FMovieSceneSequenceTransform& SequenceTransform );
 
 	USDUTILITIES_API bool ConvertXformable( const pxr::UsdStageRefPtr& Stage, const pxr::UsdTyped& Schema, USceneComponent& SceneComponent, double EvalTime );
-	USDUTILITIES_API bool ConvertXformable( const pxr::UsdStageRefPtr& Stage, const pxr::UsdTyped& Schema, FTransform& OutTransform, double EvalTime );
+	USDUTILITIES_API bool ConvertXformable( const pxr::UsdStageRefPtr& Stage, const pxr::UsdTyped& Schema, FTransform& OutTransform, double EvalTime, bool* bOutResetTransformStack = nullptr );
 
 	UE_DEPRECATED( 5.0, "Prefer the overload that receives an UE::FUsdPrim" )
 	USDUTILITIES_API bool ConvertGeomCamera( const pxr::UsdStageRefPtr& Stage, const pxr::UsdGeomCamera& GeomCamera, UCineCameraComponent& CameraComponent, double EvalTime );
diff --git a/Engine/Plugins/Importers/USDImporter/Source/UnrealUSDWrapper/Private/UsdWrappers/UsdGeomXformable.cpp b/Engine/Plugins/Importers/USDImporter/Source/UnrealUSDWrapper/Private/UsdWrappers/UsdGeomXformable.cpp
--- a/Engine/Plugins/Importers/USDImporter/Source/UnrealUSDWrapper/Private/UsdWrappers/UsdGeomXformable.cpp
+++ b/Engine/Plugins/Importers/USDImporter/Source/UnrealUSDWrapper/Private/UsdWrappers/UsdGeomXformable.cpp
@@ -113,7 +113,7 @@ namespace UE
 	}
 
 #if USE_USD_SDK
-	FUsdGeomXformable::FUsdGeomXformable( const pxr::UsdGeomXformable& InUsdGeomXformable )		
+	FUsdGeomXformable::FUsdGeomXformable( const pxr::UsdGeomXformable& InUsdGeomXformable )
 	{
 		FScopedUnrealAllocs UnrealAllocs;
 		Impl = MakeUnique< Internal::FUsdGeomXformableImpl >( InUsdGeomXformable );
@@ -142,6 +142,15 @@ namespace UE
 	}
 #endif // #if USE_USD_SDK
 
+	bool FUsdGeomXformable::GetResetXformStack() const
+	{
+#if USE_USD_SDK
+		return Impl->PxrUsdGeomXformable.Get().GetResetXformStack();
+#else
+		return false;
+#endif // #if USE_USD_SDK
+	}
+
 	bool FUsdGeomXformable::TransformMightBeTimeVarying() const
 	{
 #if USE_USD_SDK
diff --git a/Engine/Plugins/Importers/USDImporter/Source/UnrealUSDWrapper/Public/UsdWrappers/UsdGeomXformable.h b/Engine/Plugins/Importers/USDImporter/Source/UnrealUSDWrapper/Public/UsdWrappers/UsdGeomXformable.h
--- a/Engine/Plugins/Importers/USDImporter/Source/UnrealUSDWrapper/Public/UsdWrappers/UsdGeomXformable.h
+++ b/Engine/Plugins/Importers/USDImporter/Source/UnrealUSDWrapper/Public/UsdWrappers/UsdGeomXformable.h
@@ -60,6 +60,8 @@ namespace UE
 
 		// Wrapped pxr::UsdGeomXformable functions, refer to the USD SDK documentation
 	public:
+		bool GetResetXformStack() const;
+
 		bool TransformMightBeTimeVarying() const;
 		bool GetTimeSamples( TArray< double >* Times) const;
 
