USD: Fix packaging failure due to trying to use GetActorLabel() outside WITH_EDITOR block.

diff --git a/Engine/Plugins/Importers/USDImporter/Source/USDStage/Private/USDLevelSequenceHelper.cpp b/Engine/Plugins/Importers/USDImporter/Source/USDStage/Private/USDLevelSequenceHelper.cpp
--- a/Engine/Plugins/Importers/USDImporter/Source/USDStage/Private/USDLevelSequenceHelper.cpp
+++ b/Engine/Plugins/Importers/USDImporter/Source/USDStage/Private/USDLevelSequenceHelper.cpp
@@ -1678,7 +1678,14 @@ TrackType* FUsdLevelSequenceHelperImpl::AddTrack( const FName& TrackName, const
 					// to be generated. If we just used our own PrimName in here we may run into situations where a child Camera prim
 					// of a decomposed camera ends up naming the actor binding after itself, even though the parent Xform prim, and the
 					// actor on the level, maybe named something else
-					ActorBinding = MovieScene->AddPossessable( Actor->GetActorLabel(), Actor->GetClass() );
+					ActorBinding = MovieScene->AddPossessable(
+#if WITH_EDITOR
+						Actor->GetActorLabel(),
+#else
+						Actor->GetName(),
+#endif // WITH_EDITOR
+						Actor->GetClass()
+					);
 					Sequence.BindPossessableObject( ActorBinding, *Actor, Actor->GetWorld() );
 				}
 
