USD: Consistently retrieve SkelAnimations only from the first skeletal binding of a SkelRoot.

diff --git a/Engine/Plugins/Importers/USDImporter/Source/USDSchemas/Private/USDSkelRootTranslator.cpp b/Engine/Plugins/Importers/USDImporter/Source/USDSchemas/Private/USDSkelRootTranslator.cpp
--- a/Engine/Plugins/Importers/USDImporter/Source/USDSchemas/Private/USDSkelRootTranslator.cpp
+++ b/Engine/Plugins/Importers/USDImporter/Source/USDSchemas/Private/USDSkelRootTranslator.cpp
@@ -926,6 +926,12 @@ namespace UsdSkelRootTranslatorImpl
 						{
 							Context->AssetCache->LinkAssetToPrim( SkelAnimationPrimPath, AnimSequence );
 						}
+
+						// For now we shouldn't try to parse the SkelAnimations from skeletal bindings other than the first one as we only
+						// ever actually parse the first skeleton from the SkelRoot (check LoadAllSkeletalData). This is not only a matter of
+						// not wasting time, because as we would try reusing the first generated USkeleton for all other UAnimSequences,
+						// they'd likely end up not being usable or crash due to unexpected number of joints/blend shapes/etc.
+						break;
 					}
 				}
 
@@ -967,7 +973,7 @@ void FUsdSkelRootTranslator::UpdateComponents( USceneComponent* SceneComponent )
 	UE::FUsdPrim SkelAnimPrim;
 	if ( SkeletalMeshComponent->AnimationData.AnimToPlay == nullptr )
 	{
-		SkelAnimPrim = UsdUtils::FindAnimationSource( GetPrim() );
+		SkelAnimPrim = UsdUtils::FindFirstAnimationSource( GetPrim() );
 		if ( SkelAnimPrim )
 		{
 			if ( UAnimSequence* TargetAnimSequence = Cast< UAnimSequence >( Context->AssetCache->GetAssetForPrim( SkelAnimPrim.GetPrimPath().GetString() ) ) )
@@ -1021,7 +1027,7 @@ void FUsdSkelRootTranslator::UpdateComponents( USceneComponent* SceneComponent )
 			UE::FSdfLayerOffset CombinedOffset;
 			if ( !SkelAnimPrim )
 			{
-				SkelAnimPrim = UsdUtils::FindAnimationSource( GetPrim() );
+				SkelAnimPrim = UsdUtils::FindFirstAnimationSource( GetPrim() );
 			}
 			if ( SkelAnimPrim )
 			{
diff --git a/Engine/Plugins/Importers/USDImporter/Source/USDStage/Private/USDLevelSequenceHelper.cpp b/Engine/Plugins/Importers/USDImporter/Source/USDStage/Private/USDLevelSequenceHelper.cpp
--- a/Engine/Plugins/Importers/USDImporter/Source/USDStage/Private/USDLevelSequenceHelper.cpp
+++ b/Engine/Plugins/Importers/USDImporter/Source/USDStage/Private/USDLevelSequenceHelper.cpp
@@ -1515,7 +1515,7 @@ void FUsdLevelSequenceHelperImpl::AddSkeletalTracks( const UUsdPrimTwin& PrimTwi
 
 	// We'll place the skeletal animation track wherever the SkelAnimation prim is defined (not necessarily the
 	// same layer as the skel root)
-	UE::FUsdPrim SkelAnimationPrim = UsdUtils::FindAnimationSource( Prim );
+	UE::FUsdPrim SkelAnimationPrim = UsdUtils::FindFirstAnimationSource( Prim );
 	if ( !SkelAnimationPrim )
 	{
 		return;
@@ -2380,7 +2380,7 @@ void FUsdLevelSequenceHelperImpl::HandleMovieSceneChange( UMovieScene& MovieScen
 				{
 					if ( !MovieScene.FindTrack( UMovieSceneSkeletalAnimationTrack::StaticClass(), Guid ) )
 					{
-						if ( UE::FUsdPrim SkelAnimationPrim = UsdUtils::FindAnimationSource( UsdPrim ) )
+						if ( UE::FUsdPrim SkelAnimationPrim = UsdUtils::FindFirstAnimationSource( UsdPrim ) )
 						{
 							if ( UE::FSdfLayer SkelAnimationLayer = UsdUtils::FindLayerForPrim( SkelAnimationPrim ) )
 							{
@@ -2529,7 +2529,7 @@ void FUsdLevelSequenceHelperImpl::HandleTrackChange( const UMovieSceneTrack& Tra
 					bAllSectionsMuted &= !Section->IsActive();
 				}
 
-				if ( UE::FUsdPrim SkelAnimationPrim = UsdUtils::FindAnimationSource( UsdPrim ) )
+				if ( UE::FUsdPrim SkelAnimationPrim = UsdUtils::FindFirstAnimationSource( UsdPrim ) )
 				{
 					UE::FUsdAttribute TranslationsAttr = SkelAnimationPrim.GetAttribute( TEXT( "translations" ) );
 					UE::FUsdAttribute RotationsAttr = SkelAnimationPrim.GetAttribute( TEXT( "rotations" ) );
diff --git a/Engine/Plugins/Importers/USDImporter/Source/USDUtilities/Private/USDConversionUtils.cpp b/Engine/Plugins/Importers/USDImporter/Source/USDUtilities/Private/USDConversionUtils.cpp
--- a/Engine/Plugins/Importers/USDImporter/Source/USDUtilities/Private/USDConversionUtils.cpp
+++ b/Engine/Plugins/Importers/USDImporter/Source/USDUtilities/Private/USDConversionUtils.cpp
@@ -785,6 +785,10 @@ bool UsdUtils::IsAnimated( const pxr::UsdPrim& Prim )
 			{
 				return true;
 			}
+
+			// We only parse Skeletons and SkelAnimations from the first skeletal binding of a SkelRoot, so
+			// if that one is not animated then this entire SkelRoot is not animated to us either (for now)
+			break;
 		}
 	}
 
diff --git a/Engine/Plugins/Importers/USDImporter/Source/USDUtilities/Private/USDSkeletalDataConversion.cpp b/Engine/Plugins/Importers/USDImporter/Source/USDUtilities/Private/USDSkeletalDataConversion.cpp
--- a/Engine/Plugins/Importers/USDImporter/Source/USDUtilities/Private/USDSkeletalDataConversion.cpp
+++ b/Engine/Plugins/Importers/USDImporter/Source/USDUtilities/Private/USDSkeletalDataConversion.cpp
@@ -2387,28 +2387,44 @@ void UsdUtils::BindAnimationSource( pxr::UsdPrim& Prim, const pxr::UsdPrim& Anim
 	SkelBindingAPI.CreateAnimationSourceRel().SetTargets( pxr::SdfPathVector( { AnimationSource.GetPath() } ) );
 }
 
-UE::FUsdPrim UsdUtils::FindAnimationSource( const UE::FUsdPrim& SkelRootPrim )
+UE::FUsdPrim UsdUtils::FindAnimationSource( const UE::FUsdPrim& Prim )
 {
-	if ( !SkelRootPrim )
+	if ( !Prim )
 	{
 		return {};
 	}
 
 	FScopedUsdAllocs UsdAllocs;
 
-	const pxr::UsdPrim UsdSkelRootPrim{ SkelRootPrim };
-	if ( !UsdSkelRootPrim.IsA<pxr::UsdSkelRoot>() )
+	pxr::UsdPrim UsdPrim{ Prim };
+	if ( !UsdPrim.HasAPI<pxr::UsdSkelBindingAPI>() )
 	{
+		UE_LOG( LogUsd, Warning, TEXT( "Failed to find animation source for prim '%s' because it doesn't have a UsdSkelBindingAPI!" ), *Prim.GetPrimPath().GetString() );
 		return {};
 	}
 
-	if ( !UsdSkelRootPrim.HasAPI<pxr::UsdSkelBindingAPI>() )
+	// According to https://graphics.pixar.com/usd/release/api/_usd_skel__schemas.html#UsdSkel_BindingAPI_Skeletons,
+	// "Since a UsdSkelRoot primitive provides encapsulation of skeletal data, bindings defined on any primitives
+	// that do not have a UsdSkelRoot primitive as one of their ancestors have no meaning, and should be ignored."
+	// So here we make sure that Prim has at least one SkelRoot ancestor
+	bool bFoundSkelRoot = false;
+	pxr::UsdPrim Parent = UsdPrim.GetParent();
+	while ( Parent && !Parent.IsPseudoRoot() )
+	{
+		if ( Parent.IsA<pxr::UsdSkelRoot>() )
+		{
+			bFoundSkelRoot = true;
+			break;
+		}
+
+		Parent = Parent.GetParent();
+	}
+	if ( !bFoundSkelRoot )
 	{
-		UE_LOG( LogUsd, Warning, TEXT( "Failed to find animation source for prim '%s' because it is a SkelRoot without a UsdSkelBindingAPI!" ), *SkelRootPrim.GetPrimPath().GetString() );
 		return {};
 	}
 
-	pxr::UsdSkelBindingAPI SkelBindingAPI{ UsdSkelRootPrim };
+	pxr::UsdSkelBindingAPI SkelBindingAPI{ Prim };
 
 	pxr::SdfPathVector AnimationSources;
 	const bool bSuccess = SkelBindingAPI.GetAnimationSourceRel().GetTargets( &AnimationSources );
@@ -2418,11 +2434,50 @@ UE::FUsdPrim UsdUtils::FindAnimationSource( const UE::FUsdPrim& SkelRootPrim )
 	}
 
 	const pxr::SdfPath& FirstSource = AnimationSources[ 0 ];
-	pxr::UsdStageRefPtr UsdStage = UsdSkelRootPrim.GetStage();
+	pxr::UsdStageRefPtr UsdStage = Prim.GetStage();
 
 	return UE::FUsdPrim{ UsdStage->GetPrimAtPath( FirstSource ) };
 }
 
+UE::FUsdPrim UsdUtils::FindFirstAnimationSource( const UE::FUsdPrim& SkelRootPrim )
+{
+	if ( !SkelRootPrim )
+	{
+		return {};
+	}
+
+	FScopedUsdAllocs UsdAllocs;
+
+	// For now we really only parse the first skeletal binding of a SkelRoot (check USDSkelRootTranslator.cpp,
+	// LoadAllSkeletalData) and its SkelAnimation, if any.
+	// Note that we don't check the SkelRoot prim directly for the SkelAnimation binding: If it has a valid one
+	// it will propagate down to child namespaces and affect our first skeletal binding anyway
+
+	if ( pxr::UsdSkelRoot SkeletonRoot{ pxr::UsdPrim{ SkelRootPrim } } )
+	{
+		std::vector< pxr::UsdSkelBinding > SkeletonBindings;
+
+		pxr::UsdSkelCache SkeletonCache;
+		SkeletonCache.Populate( SkeletonRoot, pxr::UsdTraverseInstanceProxies() );
+		SkeletonCache.ComputeSkelBindings( SkeletonRoot, &SkeletonBindings, pxr::UsdTraverseInstanceProxies() );
+
+		for ( const pxr::UsdSkelBinding& Binding : SkeletonBindings )
+		{
+			const pxr::UsdSkelSkeleton& Skeleton = Binding.GetSkeleton();
+			pxr::UsdSkelSkeletonQuery SkelQuery = SkeletonCache.GetSkelQuery( Skeleton );
+			pxr::UsdSkelAnimQuery AnimQuery = SkelQuery.GetAnimQuery();
+			if ( !AnimQuery )
+			{
+				continue;
+			}
+
+			return UE::FUsdPrim{ AnimQuery.GetPrim() };
+		}
+	}
+
+	return {};
+}
+
 #endif // USE_USD_SDK
 
 #if USE_USD_SDK && WITH_EDITOR
diff --git a/Engine/Plugins/Importers/USDImporter/Source/USDUtilities/Public/USDSkeletalDataConversion.h b/Engine/Plugins/Importers/USDImporter/Source/USDUtilities/Public/USDSkeletalDataConversion.h
--- a/Engine/Plugins/Importers/USDImporter/Source/USDUtilities/Public/USDSkeletalDataConversion.h
+++ b/Engine/Plugins/Importers/USDImporter/Source/USDUtilities/Public/USDSkeletalDataConversion.h
@@ -110,8 +110,16 @@ namespace UsdUtils
 	 */
 	USDUTILITIES_API void BindAnimationSource( pxr::UsdPrim& Prim, const pxr::UsdPrim& AnimationSource );
 
-	// Finds the strongest SkelAnimation prim that is bound to SkelRootPrim as its animation source
-	USDUTILITIES_API UE::FUsdPrim FindAnimationSource( const UE::FUsdPrim& SkelRootPrim );
+	// Finds the SkelAnimation prim that is bound to Prim as its animationSource, if Prim has the UsdSkelBindingAPI. Returns an invalid prim otherwise.
+	UE_DEPRECATED( 5.1, "Prefer UsdUtils::FindFirstAnimationSource instead" )
+	USDUTILITIES_API UE::FUsdPrim FindAnimationSource( const UE::FUsdPrim& Prim );
+
+	/**
+	 * Returns the SkelAnimation prim that is resolved for the first skeletal binding of SkelRootPrim, if it is a SkelRoot.
+	 * We use this as we currently only parse a single Skeleton per SkelRoot (and so a single SkelAnimation), but in the future we may
+	 * decide to do something else.
+	 */
+	USDUTILITIES_API UE::FUsdPrim FindFirstAnimationSource( const UE::FUsdPrim& SkelRootPrim );
 #endif // USE_USD_SDK
 }
 
