USDImporter: register translators for boundable and nonboundable lights in USD 21.11 and later

This is another bit of prep work in advance of upgrading the engine's version of USD
from 21.08 to 22.03.

In USD 21.11, UsdLuxLight was removed and replaced with UsdLuxLightAPI,
which is an API schema instead of a typed one. The concrete base schemas for
the built-in light types are now UsdLuxBoundableLightBase and
UsdLuxNonboundableLightBase. When building against a version of USD that
uses the UsdLuxLightAPI and these new(-ish) base schemas, we register
translators for them instead of the now defunct UsdLuxLight schema.

diff --git a/Engine/Plugins/Importers/USDImporter/Source/USDSchemas/Private/USDSchemasModule.cpp b/Engine/Plugins/Importers/USDImporter/Source/USDSchemas/Private/USDSchemasModule.cpp
--- a/Engine/Plugins/Importers/USDImporter/Source/USDSchemas/Private/USDSchemasModule.cpp
+++ b/Engine/Plugins/Importers/USDImporter/Source/USDSchemas/Private/USDSchemasModule.cpp
@@ -34,7 +34,13 @@ public:
 		UsdGeomPointInstancerTranslatorHandle = GetTranslatorRegistry().Register< FUsdGeomPointInstancerTranslator >( TEXT("UsdGeomPointInstancer") );
 		UsdGeomXformableTranslatorHandle = GetTranslatorRegistry().Register< FUsdGeomXformableTranslator >( TEXT("UsdGeomXformable") );
 		UsdShadeMaterialTranslatorHandle = GetTranslatorRegistry().Register< FUsdShadeMaterialTranslator >( TEXT("UsdShadeMaterial") );
+
+#if defined(HAS_USDLUX_LIGHTAPI)
+		UsdLuxBoundableLightBaseTranslatorHandle = GetTranslatorRegistry().Register< FUsdLuxLightTranslator >( TEXT("UsdLuxBoundableLightBase") );
+		UsdLuxNonboundableLightBaseTranslatorHandle = GetTranslatorRegistry().Register< FUsdLuxLightTranslator >( TEXT("UsdLuxNonboundableLightBase") );
+#else
 		UsdLuxLightTranslatorHandle = GetTranslatorRegistry().Register< FUsdLuxLightTranslator >( TEXT("UsdLuxLight") );
+#endif
 
 #if WITH_EDITOR
 		// Creating skeletal meshes technically works in Standalone mode, but by checking for this we artificially block it
@@ -71,7 +77,14 @@ public:
 #endif // WITH_EDITOR
 		GetTranslatorRegistry().Unregister( UsdGeomXformableTranslatorHandle );
 		GetTranslatorRegistry().Unregister( UsdShadeMaterialTranslatorHandle );
+
+#if defined(HAS_USDLUX_LIGHTAPI)
+		GetTranslatorRegistry().Unregister( UsdLuxBoundableLightBaseTranslatorHandle );
+		GetTranslatorRegistry().Unregister( UsdLuxNonboundableLightBaseTranslatorHandle );
+#else
 		GetTranslatorRegistry().Unregister( UsdLuxLightTranslatorHandle );
+#endif
+
 #endif // #if USE_USD_SDK
 	}
 
@@ -95,7 +108,13 @@ protected:
 	FRegisteredSchemaTranslatorHandle UsdSkelRootTranslatorHandle;
 	FRegisteredSchemaTranslatorHandle UsdGeomXformableTranslatorHandle;
 	FRegisteredSchemaTranslatorHandle UsdShadeMaterialTranslatorHandle;
+
+#if defined(HAS_USDLUX_LIGHTAPI)
+	FRegisteredSchemaTranslatorHandle UsdLuxBoundableLightBaseTranslatorHandle;
+	FRegisteredSchemaTranslatorHandle UsdLuxNonboundableLightBaseTranslatorHandle;
+#else
 	FRegisteredSchemaTranslatorHandle UsdLuxLightTranslatorHandle;
+#endif
 
 	// Custom schemas
 	FRegisteredSchemaTranslatorHandle MdlUsdShadeMaterialTranslatorHandle;
