USD: Allow collapsing point instancer prims into static meshes;

Always collapse nested point instancers;
Fix transforms/visibility on Mesh prims being ignored when handling LODs;
Add more automated tests for point instancers.

diff --git a/Engine/Plugins/Importers/USDImporter/Source/USDSchemas/Private/USDGeomMeshTranslator.cpp b/Engine/Plugins/Importers/USDImporter/Source/USDSchemas/Private/USDGeomMeshTranslator.cpp
--- a/Engine/Plugins/Importers/USDImporter/Source/USDSchemas/Private/USDGeomMeshTranslator.cpp
+++ b/Engine/Plugins/Importers/USDImporter/Source/USDSchemas/Private/USDGeomMeshTranslator.cpp
@@ -11,6 +11,7 @@
 #include "USDConversionUtils.h"
 #include "USDGeomMeshConversion.h"
 #include "USDLog.h"
+#include "USDPrimConversion.h"
 #include "USDTypesConversion.h"
 
 #include "UsdWrappers/SdfPath.h"
@@ -48,6 +49,7 @@
 	#include "pxr/usd/usd/editContext.h"
 	#include "pxr/usd/usd/prim.h"
 	#include "pxr/usd/usd/stage.h"
+	#include "pxr/usd/usd/typed.h"
 	#include "pxr/usd/usdGeom/mesh.h"
 	#include "pxr/usd/usdGeom/xformable.h"
 	#include "pxr/usd/usdShade/material.h"
@@ -451,7 +453,7 @@ namespace UsdGeomMeshTranslatorImpl
 	// Note that these other LODs will be hidden in other variants, and won't show up on traversal unless we actively switch the variants (which we do here).
 	// We use a separate function for this because there is a very specific set of conditions where we successfully can do this, and we
 	// want to fall back to just parsing UsdMesh as a simple single-LOD mesh if we fail.
-	bool TryLoadingMultipleLODs( const pxr::UsdTyped& UsdMesh, TArray<FMeshDescription>& OutLODIndexToMeshDescription, TArray<UsdUtils::FUsdPrimMaterialAssignmentInfo>& OutLODIndexToMaterialInfo, const TMap< FString, TMap< FString, int32 > >& InMaterialToPrimvarToUVIndex, const pxr::UsdTimeCode InTimeCode, const pxr::TfToken& RenderContext, bool bCombineIdenticalMaterialSlots )
+	bool TryLoadingMultipleLODs( const pxr::UsdTyped& UsdMesh, TArray<FMeshDescription>& OutLODIndexToMeshDescription, TArray<UsdUtils::FUsdPrimMaterialAssignmentInfo>& OutLODIndexToMaterialInfo, const TMap< FString, TMap< FString, int32 > >& InMaterialToPrimvarToUVIndex, const FTransform& AdditionalTransform, const pxr::UsdTimeCode InTimeCode, const pxr::TfToken& RenderContext, bool bCombineIdenticalMaterialSlots )
 	{
 		FScopedUsdAllocs Allocs;
 
@@ -478,7 +480,38 @@ namespace UsdGeomMeshTranslatorImpl
 			FStaticMeshAttributes StaticMeshAttributes( TempMeshDescription );
 			StaticMeshAttributes.Register();
 
-			bool bSuccess = UsdToUnreal::ConvertGeomMesh( LODMesh, TempMeshDescription, TempMaterialInfo, FTransform::Identity, InMaterialToPrimvarToUVIndex, InTimeCode, RenderContext, bCombineIdenticalMaterialSlots );
+			bool bSuccess = true;
+
+			// The user can't manually hide/unhide a particular LOD in the engine after it is imported, so we should probably bake
+			// the particular LOD visibility into the combined mesh. Note how we don't use computed visibility here, as we only really
+			// care if this mesh in particular has been marked as invisible
+			pxr::TfToken Visibility;
+			pxr::UsdAttribute VisibilityAttr = LODMesh.GetVisibilityAttr();
+			if ( VisibilityAttr && VisibilityAttr.Get( &Visibility, InTimeCode ) && Visibility == pxr::UsdGeomTokens->inherited )
+			{
+				// If we're interpreting LODs we must bake the transform from each LOD Mesh into the vertices, because there's no guarantee
+				// all LODs have the same transform, so we can't just put the transforms directly on the component. If we are not interpreting
+				// LODs we can do that though
+				// TODO: Handle resetXformOp here
+				bool bResetXformStack = false;
+				FTransform MeshTransform = FTransform::Identity;
+				bSuccess &= UsdToUnreal::ConvertXformable( LODMesh.GetPrim().GetStage(), LODMesh, MeshTransform, InTimeCode.GetValue(), &bResetXformStack );
+
+				if ( bSuccess )
+				{
+					bSuccess &= UsdToUnreal::ConvertGeomMesh(
+						LODMesh,
+						TempMeshDescription,
+						TempMaterialInfo,
+						MeshTransform * AdditionalTransform,
+						InMaterialToPrimvarToUVIndex,
+						InTimeCode,
+						RenderContext,
+						bCombineIdenticalMaterialSlots
+					);
+				}
+			}
+
 			if ( bSuccess )
 			{
 				LODIndexToMeshDescriptionMap.Add( LODIndex, MoveTemp( TempMeshDescription ) );
@@ -504,7 +537,7 @@ namespace UsdGeomMeshTranslatorImpl
 	}
 
 	void LoadMeshDescriptions( const pxr::UsdTyped& UsdMesh, TArray<FMeshDescription>& OutLODIndexToMeshDescription, TArray<UsdUtils::FUsdPrimMaterialAssignmentInfo>& OutLODIndexToMaterialInfo,
-		const TMap< FString, TMap< FString, int32 > >& MaterialToPrimvarToUVIndex, const pxr::UsdTimeCode TimeCode, bool bInterpretLODs, const FName& RenderContext, bool bCombineIdenticalMaterialSlots = true )
+		const TMap< FString, TMap< FString, int32 > >& MaterialToPrimvarToUVIndex, const FTransform& AdditionalTransform, const pxr::UsdTimeCode TimeCode, bool bInterpretLODs, const FName& RenderContext, bool bCombineIdenticalMaterialSlots = true )
 	{
 		if ( !UsdMesh )
 		{
@@ -520,7 +553,7 @@ namespace UsdGeomMeshTranslatorImpl
 		bool bInterpretedLODs = false;
 		if ( bInterpretLODs )
 		{
-			bInterpretedLODs = TryLoadingMultipleLODs( UsdMesh, OutLODIndexToMeshDescription, OutLODIndexToMaterialInfo, MaterialToPrimvarToUVIndex, TimeCode, RenderContextToken, bCombineIdenticalMaterialSlots );
+			bInterpretedLODs = TryLoadingMultipleLODs( UsdMesh, OutLODIndexToMeshDescription, OutLODIndexToMaterialInfo, MaterialToPrimvarToUVIndex, AdditionalTransform, TimeCode, RenderContextToken, bCombineIdenticalMaterialSlots );
 		}
 
 		if ( !bInterpretedLODs )
@@ -531,7 +564,30 @@ namespace UsdGeomMeshTranslatorImpl
 			FStaticMeshAttributes StaticMeshAttributes( TempMeshDescription );
 			StaticMeshAttributes.Register();
 
-			bool bSuccess = UsdToUnreal::ConvertGeomMesh( UsdMesh, TempMeshDescription, TempMaterialInfo, FTransform::Identity, MaterialToPrimvarToUVIndex, TimeCode, RenderContextToken, bCombineIdenticalMaterialSlots );
+			// If we're being *asked* to intrepret LODs (regardless of whether that succeeded or not), we won't place our mesh transform on
+			// the static mesh component itself, and will instead expect it to be baked into the mesh. So here we do that
+			bool bSuccess = true;
+			FTransform MeshTransform = FTransform::Identity;
+			if ( bInterpretLODs )
+			{
+				// TODO: Handle resetXformOp here
+				bool bResetXformStack = false;
+				bSuccess &= UsdToUnreal::ConvertXformable( UsdMesh.GetPrim().GetStage(), UsdMesh, MeshTransform, TimeCode.GetValue(), &bResetXformStack );
+			}
+
+			if ( bSuccess )
+			{
+				bSuccess &= UsdToUnreal::ConvertGeomMesh(
+					UsdMesh,
+					TempMeshDescription,
+					TempMaterialInfo,
+					MeshTransform * AdditionalTransform,
+					MaterialToPrimvarToUVIndex,
+					TimeCode,
+					RenderContextToken,
+					bCombineIdenticalMaterialSlots
+				);
+			}
 
 			if ( bSuccess )
 			{
@@ -851,6 +907,7 @@ namespace UsdGeomMeshTranslatorImpl
 						LODIndexToMeshDescription,
 						LODIndexToMaterialInfo,
 						Track->MaterialToPrimvarToUVIndex,
+						FTransform::Identity,
 						pxr::UsdTimeCode( Time ),
 						bAllowInterpretingLODs,
 						Track->RenderContext
@@ -1166,8 +1223,9 @@ void FBuildStaticMeshTaskChain::SetupTasks()
 		} );
 }
 
-FGeomMeshCreateAssetsTaskChain::FGeomMeshCreateAssetsTaskChain( const TSharedRef< FUsdSchemaTranslationContext >& InContext, const UE::FSdfPath& InPrimPath )
+FGeomMeshCreateAssetsTaskChain::FGeomMeshCreateAssetsTaskChain( const TSharedRef< FUsdSchemaTranslationContext >& InContext, const UE::FSdfPath& InPrimPath, const FTransform& InAdditionalTransform )
 	: FBuildStaticMeshTaskChain( InContext, InPrimPath )
+	, AdditionalTransform( InAdditionalTransform )
 {
 	SetupTasks();
 }
@@ -1196,6 +1254,7 @@ void FGeomMeshCreateAssetsTaskChain::SetupTasks()
 				LODIndexToMeshDescription,
 				LODIndexToMaterialInfo,
 				*MaterialToPrimvarToUVIndex,
+				AdditionalTransform,
 				pxr::UsdTimeCode( Context->Time ),
 				Context->bAllowInterpretingLODs,
 				Context->RenderContext,
@@ -1251,6 +1310,7 @@ void FGeometryCacheCreateAssetsTaskChain::SetupTasks()
 				LODIndexToMeshDescription,
 				LODIndexToMaterialInfo,
 				*MaterialToPrimvarToUVIndex,
+				FTransform::Identity,
 				pxr::UsdTimeCode( TimeCode ),
 				bAllowInterpretingLODs,
 				Context->RenderContext
@@ -1335,6 +1395,7 @@ USceneComponent* FUsdGeomMeshTranslator::CreateComponents()
 #endif // WITH_EDITOR
 
 	USceneComponent* SceneComponent = CreateComponentsEx( ComponentType, {} );
+	UpdateComponents( SceneComponent );
 
 	const FString PrimPathString = PrimPath.GetString();
 
diff --git a/Engine/Plugins/Importers/USDImporter/Source/USDSchemas/Private/USDGeomPointInstancerTranslator.cpp b/Engine/Plugins/Importers/USDImporter/Source/USDSchemas/Private/USDGeomPointInstancerTranslator.cpp
--- a/Engine/Plugins/Importers/USDImporter/Source/USDSchemas/Private/USDGeomPointInstancerTranslator.cpp
+++ b/Engine/Plugins/Importers/USDImporter/Source/USDSchemas/Private/USDGeomPointInstancerTranslator.cpp
@@ -7,7 +7,9 @@
 #include "USDAssetImportData.h"
 #include "USDConversionUtils.h"
 #include "USDGeomMeshConversion.h"
+#include "USDGeomMeshTranslator.h"
 #include "USDLog.h"
+#include "USDPrimConversion.h"
 #include "USDSchemasModule.h"
 #include "USDTypesConversion.h"
 
@@ -31,6 +33,7 @@
 	#include "pxr/usd/usdGeom/camera.h"
 	#include "pxr/usd/usdGeom/mesh.h"
 	#include "pxr/usd/usdGeom/pointInstancer.h"
+	#include "pxr/usd/usdGeom/xform.h"
 	#include "pxr/usd/usdGeom/xformable.h"
 #include "USDIncludesEnd.h"
 
@@ -38,44 +41,6 @@
 
 namespace UsdGeomPointInstancerTranslatorImpl
 {
-	bool GetPointInstancerTransforms( const pxr::UsdStageRefPtr& Stage, const pxr::UsdGeomPointInstancer& PointInstancer, const int32 ProtoIndex, pxr::UsdTimeCode EvalTime, TArray<FTransform>& OutInstanceTransforms )
-	{
-		TRACE_CPUPROFILER_EVENT_SCOPE( GetPointInstancerTransforms );
-
-		FScopedUsdAllocs UsdAllocs;
-
-		pxr::VtArray< int > ProtoIndices = UsdUtils::GetUsdValue< pxr::VtArray< int > >( PointInstancer.GetProtoIndicesAttr(), EvalTime );
-
-		pxr::VtMatrix4dArray UsdInstanceTransforms;
-
-		// We don't want the prototype root prim's transforms to be included in these, as they'll already be baked into the meshes
-		// themselves
-		if ( !PointInstancer.ComputeInstanceTransformsAtTime( &UsdInstanceTransforms, EvalTime, EvalTime, pxr::UsdGeomPointInstancer::ExcludeProtoXform ) )
-		{
-			return false;
-		}
-
-		FUsdStageInfo StageInfo( Stage );
-
-		int32 Index = 0;
-
-		FScopedUnrealAllocs UnrealAllocs;
-
-		OutInstanceTransforms.Reset( UsdInstanceTransforms.size() );
-
-		for ( pxr::GfMatrix4d& UsdMatrix : UsdInstanceTransforms )
-		{
-			if ( ProtoIndices[ Index ] == ProtoIndex )
-			{
-				OutInstanceTransforms.Add( UsdToUnreal::ConvertMatrix( StageInfo, UsdMatrix ) );
-			}
-
-			++Index;
-		}
-
-		return true;
-	}
-
 	void ApplyPointInstanceTransforms( UInstancedStaticMeshComponent* Component, TArray<FTransform>& InstanceTransforms )
 	{
 		if ( Component )
@@ -109,8 +74,13 @@ namespace UsdGeomPointInstancerTranslatorImpl
 	}
 }
 
-FUsdGeomPointInstancerCreateAssetsTaskChain::FUsdGeomPointInstancerCreateAssetsTaskChain( const TSharedRef< FUsdSchemaTranslationContext >& InContext, const UE::FSdfPath& InPrimPath )
+FUsdGeomPointInstancerCreateAssetsTaskChain::FUsdGeomPointInstancerCreateAssetsTaskChain(
+	const TSharedRef< FUsdSchemaTranslationContext >& InContext,
+	const UE::FSdfPath& InPrimPath,
+	bool bInIgnoreTopLevelTransformAndVisibility
+)
 	: FBuildStaticMeshTaskChain( InContext, InPrimPath )
+	, bIgnoreTopLevelTransformAndVisibility( bInIgnoreTopLevelTransformAndVisibility )
 {
 	SetupTasks();
 }
@@ -139,10 +109,6 @@ void FUsdGeomPointInstancerCreateAssetsTaskChain::SetupTasks()
 				RenderContextToken = UnrealToUsd::ConvertToken( *Context->RenderContext.ToString() ).Get();
 			}
 
-			// We want to bake even the top level prim's transform and visibility into the combined mesh, as we'll only apply the instancing
-			// transform later
-			const bool bSkipRootPrimTransformAndVis = false;
-
 			UsdToUnreal::ConvertGeomMeshHierarchy(
 				GetPrim(),
 				pxr::UsdTimeCode( Context->Time ),
@@ -151,7 +117,7 @@ void FUsdGeomPointInstancerCreateAssetsTaskChain::SetupTasks()
 				*MaterialToPrimvarToUVIndex,
 				AddedMeshDescription,
 				AssignmentInfo,
-				bSkipRootPrimTransformAndVis,
+				bIgnoreTopLevelTransformAndVisibility,
 				Context->bMergeIdenticalMaterialSlots
 			);
 
@@ -169,30 +135,85 @@ void FUsdGeomPointInstancerTranslator::CreateAssets()
 
 	pxr::UsdPrim Prim = GetPrim();
 	pxr::UsdGeomPointInstancer PointInstancer( Prim );
-
 	if ( !PointInstancer )
 	{
 		return;
 	}
 
-	const pxr::UsdRelationship& Prototypes = PointInstancer.GetPrototypesRel();
-
-	pxr::SdfPathVector PrototypePaths;
-	if ( !Prototypes.GetTargets( &PrototypePaths ) )
+	// If another FUsdGeomXformableTranslator is collapsing the point instancer prim, it will do so by calling
+	// UsdToUnreal::ConvertGeomMeshHierarchy which will consume the prim directly.
+	// This case right here is if we're collapsing *ourselves*, where we'll essentially pretend we're a single static mesh.
+	if ( Context->bCollapseTopLevelPointInstancers )
 	{
-		return;
+		// Don't bake our actual point instancer's transform or visibility into the mesh as its nice to have these on the static mesh component instead
+		const bool bIgnoreTopLevelTransformAndVisibility = true;
+		Context->TranslatorTasks.Add( MakeShared< FUsdGeomPointInstancerCreateAssetsTaskChain >( Context, PrimPath, bIgnoreTopLevelTransformAndVisibility ) );
 	}
-
-	for ( const pxr::SdfPath& PrototypePath : PrototypePaths )
+	// Otherwise we're just going to spawn HISM components instead
+	else
 	{
-		pxr::UsdPrim PrototypeUsdPrim = Prim.GetStage()->GetPrimAtPath( PrototypePath );
-		if ( !PrototypeUsdPrim )
+		const pxr::UsdRelationship& Prototypes = PointInstancer.GetPrototypesRel();
+
+		pxr::SdfPathVector PrototypePaths;
+		if ( !Prototypes.GetTargets( &PrototypePaths ) )
 		{
-			UE_LOG( LogUsd, Warning, TEXT( "Failed to find prototype '%s' for PointInstancer '%s' when collapsing assets" ), *UsdToUnreal::ConvertPath( PrototypePath ), *PrimPath.GetString() );
-			continue;
+			return;
 		}
 
-		Context->TranslatorTasks.Add( MakeShared< FUsdGeomPointInstancerCreateAssetsTaskChain >( Context, UE::FSdfPath{ *UsdToUnreal::ConvertPath( PrototypePath ) } ) );
+		for ( const pxr::SdfPath& PrototypePath : PrototypePaths )
+		{
+			pxr::UsdPrim PrototypeUsdPrim = Prim.GetStage()->GetPrimAtPath( PrototypePath );
+			UE::FSdfPath UEPrototypePath{ PrototypePath };
+
+			if ( !PrototypeUsdPrim )
+			{
+				UE_LOG( LogUsd, Warning, TEXT( "Failed to find prototype '%s' for PointInstancer '%s' when collapsing assets" ), *UEPrototypePath.GetString(), *PrimPath.GetString() );
+				continue;
+			}
+
+			if ( Context->bAllowInterpretingLODs && UsdUtils::DoesPrimContainMeshLODs( PrototypeUsdPrim ) )
+			{
+				// We have to provide one of the LOD meshes to the task chain, so find the path to one
+				UE::FSdfPath ChildMeshPath;
+				pxr::UsdPrimSiblingRange PrimRange = PrototypeUsdPrim.GetChildren();
+				for ( pxr::UsdPrimSiblingRange::iterator PrimRangeIt = PrimRange.begin(); PrimRangeIt != PrimRange.end(); ++PrimRangeIt )
+				{
+					const pxr::UsdPrim& Child = *PrimRangeIt;
+					if ( pxr::UsdGeomMesh ChildMesh{ Child } )
+					{
+						ChildMeshPath = UE::FSdfPath{ Child.GetPrimPath() };
+						break;
+					}
+				}
+
+				// This is in charge of baking in 'Prim's transform into the generated static mesh for the prototype, as it
+				// otherwise wouldn't end up anywhere else. Note that in the default/export case 'Prim' (the prim that actually contains the LOD
+				// variant set) is schema-less (and so not an Xform), but this is just in case the user manually made it an Xform instead
+				FTransform AdditionalUESpaceTransform = FTransform::Identity;
+				if ( pxr::UsdGeomXform ParentXform{ PrototypeUsdPrim } )
+				{
+					// Skip this LOD mesh if its invisible
+					pxr::TfToken Visibility;
+					pxr::UsdAttribute VisibilityAttr = ParentXform.GetVisibilityAttr();
+					if ( VisibilityAttr && VisibilityAttr.Get( &Visibility, Context->Time ) && Visibility == pxr::UsdGeomTokens->invisible )
+					{
+						continue;
+					}
+
+					// TODO: Handle the resetXformStack op for LOD parents
+					bool bOutResetTransformStack = false;
+					UsdToUnreal::ConvertXformable( Prim.GetStage(), ParentXform, AdditionalUESpaceTransform, Context->Time, &bOutResetTransformStack );
+				}
+
+				Context->TranslatorTasks.Add( MakeShared< FGeomMeshCreateAssetsTaskChain >( Context, ChildMeshPath, AdditionalUESpaceTransform ) );
+			}
+			else
+			{
+				// Fully bake the prototype top level transform and visibility into its own static mesh
+				const bool bIgnoreTopLevelTransformAndVisibility = false;
+				Context->TranslatorTasks.Add( MakeShared< FUsdGeomPointInstancerCreateAssetsTaskChain >( Context, UEPrototypePath, bIgnoreTopLevelTransformAndVisibility ) );
+			}
+		}
 	}
 }
 
@@ -200,15 +221,24 @@ USceneComponent* FUsdGeomPointInstancerTranslator::CreateComponents()
 {
 	TRACE_CPUPROFILER_EVENT_SCOPE( FUsdGeomPointInstancerTranslator::CreateComponents );
 
-	// The plan here is to create an USceneComponent that corresponds to the PointInstancer prim itself, and then spawn a child
+	// If we're going to be collapsing ourselves, just make a static mesh component that can receive the collapsed mesh
+	if ( Context->bCollapseTopLevelPointInstancers )
+	{
+		USceneComponent* MeshComponent = CreateComponentsEx( { UStaticMeshComponent::StaticClass() }, {} );
+		UpdateComponents( MeshComponent );
+		return MeshComponent;
+	}
+
+	// Otherwise, the plan here is to create an USceneComponent that corresponds to the PointInstancer prim itself, and then spawn a child
 	// HISM component for each prototype.
 	// We always request a scene component here explicitly or else we'll be upgraded to a static mesh component by the mechanism that
-	// handles collapsed Meshes/static mesh components for the GeomXFormable translator.
+	// handles collapsed meshes/static mesh components for the GeomXFormable translator.
 	USceneComponent* MainSceneComponent = CreateComponentsEx( { USceneComponent::StaticClass() }, {} );
 	if ( !MainSceneComponent )
 	{
-		return nullptr;
+		return MainSceneComponent;
 	}
+	UpdateComponents( MainSceneComponent );
 
 	FScopedUsdAllocs UsdAllocs;
 
@@ -238,7 +268,7 @@ USceneComponent* FUsdGeomPointInstancerTranslator::CreateComponents()
 		PrototypePathsSlowTask.EnterProgressFrame();
 
 		const pxr::SdfPath& PrototypePath = PrototypePaths[ PrototypeIndex ];
-		const FString PrototypePathStr = UsdToUnreal::ConvertPath( PrototypePath );
+		FString PrototypePathStr = UsdToUnreal::ConvertPath( PrototypePath );
 
 		pxr::UsdPrim PrototypeUsdPrim = Prim.GetStage()->GetPrimAtPath( PrototypePath );
 		if ( !PrototypeUsdPrim )
@@ -247,6 +277,23 @@ USceneComponent* FUsdGeomPointInstancerTranslator::CreateComponents()
 			continue;
 		}
 
+		// If our prototype was a LOD mesh we will have used the path of one of the actual LOD meshes to start the FGeomMeshCreateAssetsTaskChain,
+		// so we have to look for our resulting mesh with the same path
+		if ( Context->bAllowInterpretingLODs && UsdUtils::DoesPrimContainMeshLODs( PrototypeUsdPrim ) )
+		{
+			pxr::UsdPrimSiblingRange PrimRange = PrototypeUsdPrim.GetChildren();
+			for ( pxr::UsdPrimSiblingRange::iterator PrimRangeIt = PrimRange.begin(); PrimRangeIt != PrimRange.end(); ++PrimRangeIt )
+			{
+				const pxr::UsdPrim& Child = *PrimRangeIt;
+				if ( pxr::UsdGeomMesh ChildMesh{ Child } )
+				{
+					PrototypeUsdPrim = Child;
+					PrototypePathStr = UsdToUnreal::ConvertPath( Child.GetPrimPath() );
+					break;
+				}
+			}
+		}
+
 		using UHISMComponent = UHierarchicalInstancedStaticMeshComponent;
 
 		FUsdGeomXformableTranslator XformableTranslator{ UHISMComponent::StaticClass(), Context, UE::FUsdTyped( PrototypeUsdPrim ) };
@@ -264,13 +311,14 @@ USceneComponent* FUsdGeomPointInstancerTranslator::CreateComponents()
 
 		// Evaluating point instancer can take a long time and is thread-safe. Move to async task while we work on something else.
 		pxr::UsdTimeCode TimeCode{ Context->Time };
+		FUsdStageInfo StageInfo{ Prim.GetStage() };
 		Tasks.Emplace(
 			Async(
 				EAsyncExecution::ThreadPool,
-				[ TimeCode, Stage = Prim.GetStage(), PointInstancer, PrototypeIndex, HISMComponent ]()
+				[ TimeCode, StageInfo, PointInstancer, PrototypeIndex, HISMComponent ]()
 				{
 					TArray<FTransform> InstanceTransforms;
-					UsdGeomPointInstancerTranslatorImpl::GetPointInstancerTransforms( Stage, PointInstancer, PrototypeIndex, TimeCode, InstanceTransforms );
+					UsdUtils::GetPointInstancerTransforms( StageInfo, PointInstancer, PrototypeIndex, TimeCode, InstanceTransforms );
 
 					return MakeTuple( HISMComponent, MoveTemp( InstanceTransforms ) );
 				}
diff --git a/Engine/Plugins/Importers/USDImporter/Source/USDSchemas/Private/USDGeomXformableTranslator.cpp b/Engine/Plugins/Importers/USDImporter/Source/USDSchemas/Private/USDGeomXformableTranslator.cpp
--- a/Engine/Plugins/Importers/USDImporter/Source/USDSchemas/Private/USDGeomXformableTranslator.cpp
+++ b/Engine/Plugins/Importers/USDImporter/Source/USDSchemas/Private/USDGeomXformableTranslator.cpp
@@ -138,7 +138,12 @@ FUsdGeomXformableTranslator::FUsdGeomXformableTranslator( TSubclassOf< USceneCom
 
 USceneComponent* FUsdGeomXformableTranslator::CreateComponents()
 {
-	return CreateComponentsEx( {}, {} );
+	USceneComponent* Component = CreateComponentsEx( {}, {} );
+
+	// We pulled UpdateComponents outside CreateComponentsEx as in some cases we don't want to do it
+	// right away (like on FUsdGeomPointInstancerTranslator::CreateComponents)
+	UpdateComponents( Component );
+	return Component;
 }
 
 USceneComponent* FUsdGeomXformableTranslator::CreateComponentsEx( TOptional< TSubclassOf< USceneComponent > > ComponentType, TOptional< bool > bNeedsActor )
@@ -370,8 +375,6 @@ USceneComponent* FUsdGeomXformableTranslator::CreateComponentsEx( TOptional< TSu
 		// UpdateComponents with all the components already attached
 		SceneComponent->AttachToComponent( Context->ParentComponent, FAttachmentTransformRules::KeepRelativeTransform );
 
-		UpdateComponents( SceneComponent );
-
 		if ( !SceneComponent->IsRegistered() )
 		{
 			SceneComponent->RegisterComponent();
@@ -397,18 +400,16 @@ void FUsdGeomXformableTranslator::UpdateComponents( USceneComponent* SceneCompon
 			SceneComponent->UnregisterComponent();
 		}
 
-		// Don't convert transform/visibility for HISM components as we only create these for point instancer prototype prims,
-		// and we don't want the actual prototype prim transform/visibility on the component itself: It will always be baked into the mesh
-		if ( !SceneComponent->IsA< UHierarchicalInstancedStaticMeshComponent >() )
-		{
-			UsdToUnreal::ConvertXformable( Context->Stage, pxr::UsdGeomXformable( GetPrim() ), *SceneComponent, Context->Time );
-		}
-
 		// If the user modified a mesh parameter (e.g. vertex color), the hash will be different and it will become a separate asset
 		// so we must check for this and assign the new StaticMesh
+		bool bHasMultipleLODs = false;
 		if ( UStaticMeshComponent* StaticMeshComponent = Cast< UStaticMeshComponent >( SceneComponent ) )
 		{
 			UStaticMesh* PrimStaticMesh = Cast< UStaticMesh >( Context->AssetCache->GetAssetForPrim( PrimPath.GetString() ) );
+			if ( PrimStaticMesh )
+			{
+				bHasMultipleLODs = PrimStaticMesh->GetNumLODs() > 1;
+			}
 
 			if ( PrimStaticMesh != StaticMeshComponent->GetStaticMesh() )
 			{
@@ -430,6 +431,13 @@ void FUsdGeomXformableTranslator::UpdateComponents( USceneComponent* SceneCompon
 			}
 		}
 
+		// Only put the transform into the component if we haven't parsed LODs for our static mesh: The Mesh transforms will already be baked
+		// into the mesh at that case, as each LOD could technically have a separate transform
+		if ( !Context->bAllowInterpretingLODs || !bHasMultipleLODs )
+		{
+			UsdToUnreal::ConvertXformable( Context->Stage, pxr::UsdGeomXformable( GetPrim() ), *SceneComponent, Context->Time );
+		}
+
 		// Note how we should only register if we unregistered ourselves: If we did this every time we would
 		// register too early during the process of duplicating into PIE, and that would prevent a future RegisterComponent
 		// call from naturally creating the required render state
diff --git a/Engine/Plugins/Importers/USDImporter/Source/USDSchemas/Private/USDLuxLightTranslator.cpp b/Engine/Plugins/Importers/USDImporter/Source/USDSchemas/Private/USDLuxLightTranslator.cpp
--- a/Engine/Plugins/Importers/USDImporter/Source/USDSchemas/Private/USDLuxLightTranslator.cpp
+++ b/Engine/Plugins/Importers/USDImporter/Source/USDSchemas/Private/USDLuxLightTranslator.cpp
@@ -35,7 +35,9 @@
 USceneComponent* FUsdLuxLightTranslator::CreateComponents()
 {
 	const bool bNeedsActor = true;
-	return CreateComponentsEx( {}, bNeedsActor );
+	USceneComponent* LightComponent = CreateComponentsEx( {}, bNeedsActor );
+	UpdateComponents( LightComponent );
+	return LightComponent;
 }
 
 void FUsdLuxLightTranslator::UpdateComponents( USceneComponent* SceneComponent )
diff --git a/Engine/Plugins/Importers/USDImporter/Source/USDSchemas/Public/USDGeomMeshTranslator.h b/Engine/Plugins/Importers/USDImporter/Source/USDSchemas/Public/USDGeomMeshTranslator.h
--- a/Engine/Plugins/Importers/USDImporter/Source/USDSchemas/Public/USDGeomMeshTranslator.h
+++ b/Engine/Plugins/Importers/USDImporter/Source/USDSchemas/Public/USDGeomMeshTranslator.h
@@ -55,7 +55,11 @@ protected:
 class FGeomMeshCreateAssetsTaskChain : public FBuildStaticMeshTaskChain
 {
 public:
-	explicit FGeomMeshCreateAssetsTaskChain( const TSharedRef< FUsdSchemaTranslationContext >& InContext, const UE::FSdfPath& PrimPath );
+	explicit FGeomMeshCreateAssetsTaskChain( const TSharedRef< FUsdSchemaTranslationContext >& InContext, const UE::FSdfPath& PrimPath, const FTransform& AdditionalTransform = FTransform::Identity );
+
+protected:
+	// Inputs
+	FTransform AdditionalTransform;
 
 protected:
 	virtual void SetupTasks() override;
diff --git a/Engine/Plugins/Importers/USDImporter/Source/USDSchemas/Public/USDGeomPointInstancerTranslator.h b/Engine/Plugins/Importers/USDImporter/Source/USDSchemas/Public/USDGeomPointInstancerTranslator.h
--- a/Engine/Plugins/Importers/USDImporter/Source/USDSchemas/Public/USDGeomPointInstancerTranslator.h
+++ b/Engine/Plugins/Importers/USDImporter/Source/USDSchemas/Public/USDGeomPointInstancerTranslator.h
@@ -12,10 +12,13 @@ struct FUsdSchemaTranslationContext;
 class FUsdGeomPointInstancerCreateAssetsTaskChain : public FBuildStaticMeshTaskChain
 {
 public:
-	explicit FUsdGeomPointInstancerCreateAssetsTaskChain( const TSharedRef< FUsdSchemaTranslationContext >& InContext, const UE::FSdfPath& InPrimPath );
+	explicit FUsdGeomPointInstancerCreateAssetsTaskChain( const TSharedRef< FUsdSchemaTranslationContext >& InContext, const UE::FSdfPath& InPrimPath, bool bIgnoreTopLevelTransformAndVisibility );
 
 protected:
 	virtual void SetupTasks() override;
+
+private:
+	bool bIgnoreTopLevelTransformAndVisibility = false;
 };
 
 class USDSCHEMAS_API FUsdGeomPointInstancerTranslator : public FUsdGeomXformableTranslator
@@ -28,8 +31,13 @@ public:
 	virtual USceneComponent* CreateComponents() override;
 	virtual void UpdateComponents( USceneComponent* SceneComponent ) override;
 
+	// Point instancers will always collapse (i.e. "be in charge of") their entire prim subtree.
+	// If we have a nested point instancer situation, the parentmost point instancer will be in charge of collapsing everything
+	// below it, so returning true to "CanBeCollapsed" in that case won't really help anything.
+	// We only return true from CanBeCollapsed if e.g. a generic Xform prim should be able to collapse us too (like when we're
+	// collapsing even simple, not-nested point instancers).
 	virtual bool CollapsesChildren( ECollapsingType CollapsingType ) const override { return true; }
-	virtual bool CanBeCollapsed( ECollapsingType CollapsingType ) const override { return false; }
+	virtual bool CanBeCollapsed( ECollapsingType CollapsingType ) const override { return Context->bCollapseTopLevelPointInstancers; }
 };
 
 #endif // #if USE_USD_SDK
diff --git a/Engine/Plugins/Importers/USDImporter/Source/USDSchemas/Public/USDSchemaTranslator.h b/Engine/Plugins/Importers/USDImporter/Source/USDSchemas/Public/USDSchemaTranslator.h
--- a/Engine/Plugins/Importers/USDImporter/Source/USDSchemas/Public/USDSchemaTranslator.h
+++ b/Engine/Plugins/Importers/USDImporter/Source/USDSchemas/Public/USDSchemaTranslator.h
@@ -178,6 +178,13 @@ struct USDSCHEMAS_API FUsdSchemaTranslationContext : public TSharedFromThis< FUs
 	 */
 	bool bMergeIdenticalMaterialSlots = true;
 
+	/**
+	 * If true, will cause us to collapse any point instancer prim into a single static mesh and static mesh component.
+	 * If false, will cause us to use HierarchicalInstancedStaticMeshComponents to replicate the instancing behavior.
+	 * Point instancers inside other point instancer prototypes are *always* collapsed into the prototype's static mesh.
+	 */
+	bool bCollapseTopLevelPointInstancers = false;
+
 	/**
 	 * If true, prims with a "LOD" variant set, and "LOD0", "LOD1", etc. variants containing each
 	 * a prim can be parsed into a single UStaticMesh asset with multiple LODs
diff --git a/Engine/Plugins/Importers/USDImporter/Source/USDStage/Private/USDStageActor.cpp b/Engine/Plugins/Importers/USDImporter/Source/USDStage/Private/USDStageActor.cpp
--- a/Engine/Plugins/Importers/USDImporter/Source/USDStage/Private/USDStageActor.cpp
+++ b/Engine/Plugins/Importers/USDImporter/Source/USDStage/Private/USDStageActor.cpp
@@ -106,6 +106,7 @@ struct FUsdStageActorImpl
 
 		TranslationContext->KindsToCollapse = (EUsdDefaultKind) StageActor->KindsToCollapse;
 		TranslationContext->bMergeIdenticalMaterialSlots = StageActor->bMergeIdenticalMaterialSlots;
+		TranslationContext->bCollapseTopLevelPointInstancers = StageActor->bCollapseTopLevelPointInstancers;
 
 		UE::FSdfPath UsdPrimPath( *PrimPath );
 		UUsdPrimTwin* ParentUsdPrimTwin = StageActor->GetRootPrimTwin()->Find( UsdPrimPath.GetParentPath().GetString() );
@@ -521,6 +522,7 @@ struct FUsdStageActorImpl
 			EventAttributes.Emplace( TEXT( "InterpolationType" ), LexToString( (uint8)StageActor->InterpolationType) );
 			EventAttributes.Emplace( TEXT( "KindsToCollapse" ), LexToString( StageActor->KindsToCollapse ) );
 			EventAttributes.Emplace( TEXT( "MergeIdenticalMaterialSlots" ), LexToString( StageActor->bMergeIdenticalMaterialSlots ) );
+			EventAttributes.Emplace( TEXT( "CollapseTopLevelPointInstancers" ), LexToString( StageActor->bCollapseTopLevelPointInstancers ) );
 			EventAttributes.Emplace( TEXT( "PurposesToLoad" ), LexToString( StageActor->PurposesToLoad ) );
 			EventAttributes.Emplace( TEXT( "NaniteTriangleThreshold" ), LexToString( StageActor->NaniteTriangleThreshold ) );
 			EventAttributes.Emplace( TEXT( "RenderContext" ), StageActor->RenderContext.ToString() );
@@ -605,6 +607,7 @@ AUsdStageActor::AUsdStageActor()
 	, InterpolationType( EUsdInterpolationType::Linear )
 	, KindsToCollapse( ( int32 ) ( EUsdDefaultKind::Component | EUsdDefaultKind::Subcomponent ) )
 	, bMergeIdenticalMaterialSlots( true )
+	, bCollapseTopLevelPointInstancers( false )
 	, PurposesToLoad( (int32) EUsdPurpose::Proxy )
 	, NaniteTriangleThreshold( (uint64) 1000000 )
 	, Time( 0.0f )
@@ -1307,6 +1310,14 @@ void AUsdStageActor::SetMergeIdenticalMaterialSlots( bool bMerge )
 	LoadUsdStage();
 }
 
+void AUsdStageActor::SetCollapseTopLevelPointInstancers( bool bCollapse )
+{
+	Modify();
+
+	bCollapseTopLevelPointInstancers = bCollapse;
+	LoadUsdStage();
+}
+
 void AUsdStageActor::SetPurposesToLoad( int32 NewPurposesToLoad )
 {
 	Modify();
@@ -2573,6 +2584,10 @@ void AUsdStageActor::HandlePropertyChangedEvent( FPropertyChangedEvent& Property
 	{
 		SetMergeIdenticalMaterialSlots( bMergeIdenticalMaterialSlots );
 	}
+	else if ( PropertyName == GET_MEMBER_NAME_CHECKED( AUsdStageActor, bCollapseTopLevelPointInstancers ) )
+	{
+		SetCollapseTopLevelPointInstancers( bCollapseTopLevelPointInstancers );
+	}
 	else if ( PropertyName == GET_MEMBER_NAME_CHECKED( AUsdStageActor, PurposesToLoad ) )
 	{
 		SetPurposesToLoad( PurposesToLoad );
diff --git a/Engine/Plugins/Importers/USDImporter/Source/USDStage/Public/USDStageActor.h b/Engine/Plugins/Importers/USDImporter/Source/USDStage/Public/USDStageActor.h
--- a/Engine/Plugins/Importers/USDImporter/Source/USDStage/Public/USDStageActor.h
+++ b/Engine/Plugins/Importers/USDImporter/Source/USDStage/Public/USDStageActor.h
@@ -62,6 +62,14 @@ public:
 	UPROPERTY(EditAnywhere, BlueprintReadOnly, Category = "USD")
 	bool bMergeIdenticalMaterialSlots;
 
+	/**
+	 * If true, will cause us to collapse any point instancer prim into a single static mesh and static mesh component.
+	 * If false, will cause us to use HierarchicalInstancedStaticMeshComponents to replicate the instancing behavior.
+	 * Point instancers inside other point instancer prototypes are *always* collapsed into the prototype's static mesh.
+	 */
+	UPROPERTY( EditAnywhere, BlueprintReadOnly, Category = "USD" )
+	bool bCollapseTopLevelPointInstancers;
+
 	/* Only load prims with these specific purposes from the USD file */
 	UPROPERTY(EditAnywhere, BlueprintReadOnly, Category = "USD", meta = (Bitmask, BitmaskEnum=EUsdPurpose))
 	int32 PurposesToLoad;
@@ -90,6 +98,9 @@ public:
 	UFUNCTION( BlueprintCallable, Category = "USD", meta = ( CallInEditor = "true" ) )
 	USDSTAGE_API void SetMergeIdenticalMaterialSlots( bool bMerge );
 
+	UFUNCTION( BlueprintCallable, Category = "USD", meta = ( CallInEditor = "true" ) )
+	USDSTAGE_API void SetCollapseTopLevelPointInstancers( bool bCollapse );
+
 	UFUNCTION(BlueprintCallable, Category = "USD", meta = (CallInEditor = "true"))
 	USDSTAGE_API void SetPurposesToLoad( int32 NewPurposesToLoad );
 
diff --git a/Engine/Plugins/Importers/USDImporter/Source/USDStageEditor/Private/SUSDStage.cpp b/Engine/Plugins/Importers/USDImporter/Source/USDStageEditor/Private/SUSDStage.cpp
--- a/Engine/Plugins/Importers/USDImporter/Source/USDStageEditor/Private/SUSDStage.cpp
+++ b/Engine/Plugins/Importers/USDImporter/Source/USDStageEditor/Private/SUSDStage.cpp
@@ -989,7 +989,7 @@ void SUsdStage::FillCollapsingSubMenu( FMenuBuilder& MenuBuilder )
 				if ( AUsdStageActor* StageActor = ViewModel.UsdStageActor.Get() )
 				{
 					FScopedTransaction Transaction( FText::Format(
-						LOCTEXT( "ToggleCollapsingTransaction", "Toggle bMergeIdenticalMaterialSlots on USD stage actor '{1}'" ),
+						LOCTEXT( "MergeIdenticalMaterialSlotsTransaction", "Toggle bMergeIdenticalMaterialSlots on USD stage actor '{1}'" ),
 						FText::FromString( StageActor->GetActorLabel() )
 					) );
 
@@ -1014,6 +1014,41 @@ void SUsdStage::FillCollapsingSubMenu( FMenuBuilder& MenuBuilder )
 		EUserInterfaceActionType::ToggleButton
 	);
 
+	MenuBuilder.AddMenuEntry(
+		LOCTEXT( "CollapseTopLevelPointInstancers", "Collapse simple point instancers" ),
+		LOCTEXT( "CollapseTopLevelPointInstancersToolTip", "If true, will cause us to collapse any point instancer prim into a single static mesh and static mesh component.\nIf false, will cause us to use HierarchicalInstancedStaticMeshComponents to replicate the instancing behavior.\nPoint instancers inside other point instancer prototypes are * always * collapsed into the prototype's static mesh." ),
+		FSlateIcon(),
+		FUIAction(
+			FExecuteAction::CreateLambda( [this]()
+			{
+				if ( AUsdStageActor* StageActor = ViewModel.UsdStageActor.Get() )
+				{
+					FScopedTransaction Transaction( FText::Format(
+						LOCTEXT( "CollapseTopLevelPointInstancersTransaction", "Toggle bCollapseTopLevelPointInstancers on USD stage actor '{1}'" ),
+						FText::FromString( StageActor->GetActorLabel() )
+					) );
+
+					TGuardValue<bool> MaintainSelectionGuard( bUpdatingViewportSelection, true );
+					StageActor->SetCollapseTopLevelPointInstancers( !StageActor->bCollapseTopLevelPointInstancers );
+				}
+			}),
+			FCanExecuteAction::CreateLambda( [this]()
+			{
+				return ViewModel.UsdStageActor.Get() != nullptr;
+			}),
+			FIsActionChecked::CreateLambda( [this]()
+			{
+				if ( AUsdStageActor* StageActor = ViewModel.UsdStageActor.Get() )
+				{
+					return StageActor->bCollapseTopLevelPointInstancers;
+				}
+				return false;
+			})
+		),
+		NAME_None,
+		EUserInterfaceActionType::ToggleButton
+	);
+
 	auto AddKindToCollapseEntry = [&]( const EUsdDefaultKind Kind, const FText& Text, FCanExecuteAction CanExecuteAction )
 	{
 		MenuBuilder.AddMenuEntry(
diff --git a/Engine/Plugins/Importers/USDImporter/Source/USDStageImporter/Private/USDStageImportOptions.cpp b/Engine/Plugins/Importers/USDImporter/Source/USDStageImporter/Private/USDStageImportOptions.cpp
--- a/Engine/Plugins/Importers/USDImporter/Source/USDStageImporter/Private/USDStageImportOptions.cpp
+++ b/Engine/Plugins/Importers/USDImporter/Source/USDStageImporter/Private/USDStageImportOptions.cpp
@@ -32,6 +32,7 @@ UUsdStageImportOptions::UUsdStageImportOptions(const FObjectInitializer& ObjectI
 	bPrimPathFolderStructure = false;
 	KindsToCollapse = ( int32 ) ( EUsdDefaultKind::Component | EUsdDefaultKind::Subcomponent );
 	bMergeIdenticalMaterialSlots = true;
+	bCollapseTopLevelPointInstancers = false;
 	bInterpretLODs = true;
 }
 
diff --git a/Engine/Plugins/Importers/USDImporter/Source/USDStageImporter/Private/USDStageImporter.cpp b/Engine/Plugins/Importers/USDImporter/Source/USDStageImporter/Private/USDStageImporter.cpp
--- a/Engine/Plugins/Importers/USDImporter/Source/USDStageImporter/Private/USDStageImporter.cpp
+++ b/Engine/Plugins/Importers/USDImporter/Source/USDStageImporter/Private/USDStageImporter.cpp
@@ -1283,6 +1283,7 @@ namespace UsdStageImporterImpl
 				EventAttributes.Emplace( TEXT( "PrimPathFolderStructure" ), LexToString( ImportContext.ImportOptions->bPrimPathFolderStructure ) );
 				EventAttributes.Emplace( TEXT( "KindsToCollapse" ), LexToString( ImportContext.ImportOptions->KindsToCollapse ) );
 				EventAttributes.Emplace( TEXT( "MergeIdenticalMaterialSlots" ), LexToString( ImportContext.ImportOptions->bMergeIdenticalMaterialSlots ) );
+				EventAttributes.Emplace( TEXT( "CollapseTopLevelPointInstancers" ), LexToString( ImportContext.ImportOptions->bCollapseTopLevelPointInstancers ) );
 				EventAttributes.Emplace( TEXT( "InterpretLODs" ), LexToString( ImportContext.ImportOptions->bInterpretLODs ) );
 			}
 
@@ -1494,6 +1495,7 @@ void UUsdStageImporter::ImportFromFile(FUsdStageImportContext& ImportContext)
 	TranslationContext->ParentComponent = ImportContext.SceneActor ? ImportContext.SceneActor->GetRootComponent() : nullptr;
 	TranslationContext->KindsToCollapse = ( EUsdDefaultKind ) ImportContext.ImportOptions->KindsToCollapse;
 	TranslationContext->bMergeIdenticalMaterialSlots = ImportContext.ImportOptions->bMergeIdenticalMaterialSlots;
+	TranslationContext->bCollapseTopLevelPointInstancers = ImportContext.ImportOptions->bCollapseTopLevelPointInstancers;
 	TranslationContext->bAllowInterpretingLODs = ImportContext.ImportOptions->bInterpretLODs;
 	TranslationContext->bAllowParsingSkeletalAnimations = ImportContext.ImportOptions->bImportGeometry && ImportContext.ImportOptions->bImportSkeletalAnimations;
 	TranslationContext->MaterialToPrimvarToUVIndex = &ImportContext.MaterialToPrimvarToUVIndex;
@@ -1580,6 +1582,7 @@ bool UUsdStageImporter::ReimportSingleAsset(FUsdStageImportContext& ImportContex
 	TranslationContext->NaniteTriangleThreshold = ImportContext.ImportOptions->NaniteTriangleThreshold;
 	TranslationContext->KindsToCollapse = ( EUsdDefaultKind ) ImportContext.ImportOptions->KindsToCollapse;
 	TranslationContext->bMergeIdenticalMaterialSlots = ImportContext.ImportOptions->bMergeIdenticalMaterialSlots;
+	TranslationContext->bCollapseTopLevelPointInstancers = ImportContext.ImportOptions->bCollapseTopLevelPointInstancers;
 	TranslationContext->bAllowInterpretingLODs = ImportContext.ImportOptions->bInterpretLODs;
 	TranslationContext->bAllowParsingSkeletalAnimations = ImportContext.ImportOptions->bImportGeometry && ImportContext.ImportOptions->bImportSkeletalAnimations;
 	TranslationContext->MaterialToPrimvarToUVIndex = &ImportContext.MaterialToPrimvarToUVIndex;
diff --git a/Engine/Plugins/Importers/USDImporter/Source/USDStageImporter/Public/USDStageImportOptions.h b/Engine/Plugins/Importers/USDImporter/Source/USDStageImporter/Public/USDStageImportOptions.h
--- a/Engine/Plugins/Importers/USDImporter/Source/USDStageImporter/Public/USDStageImportOptions.h
+++ b/Engine/Plugins/Importers/USDImporter/Source/USDStageImporter/Public/USDStageImportOptions.h
@@ -124,6 +124,14 @@ public:
 	UPROPERTY(BlueprintReadWrite, config, EditAnywhere, Category = "Processing")
 	bool bMergeIdenticalMaterialSlots;
 
+	/**
+	 * If true, will cause us to collapse any point instancer prim into a single static mesh and static mesh component.
+	 * If false, will cause us to use HierarchicalInstancedStaticMeshComponents to replicate the instancing behavior.
+	 * Point instancers inside other point instancer prototypes are *always* collapsed into the prototype's static mesh.
+	 */
+	UPROPERTY( BlueprintReadWrite, config, EditAnywhere, Category = "Processing" )
+	bool bCollapseTopLevelPointInstancers;
+
 	/** When true, if a prim has a "LOD" variant set with variants named "LOD0", "LOD1", etc. where each contains a UsdGeomMesh, the importer will attempt to parse the meshes as separate LODs of a single UStaticMesh. When false, only the selected variant will be parsed as LOD0 of the UStaticMesh.  */
 	UPROPERTY(BlueprintReadWrite, config, EditAnywhere, Category="Processing", meta=(DisplayName="Interpret LOD variant sets", EditCondition=bImportGeometry) )
 	bool bInterpretLODs;
diff --git a/Engine/Plugins/Importers/USDImporter/Source/USDUtilities/Private/USDGeomMeshConversion.cpp b/Engine/Plugins/Importers/USDImporter/Source/USDUtilities/Private/USDGeomMeshConversion.cpp
--- a/Engine/Plugins/Importers/USDImporter/Source/USDUtilities/Private/USDGeomMeshConversion.cpp
+++ b/Engine/Plugins/Importers/USDImporter/Source/USDUtilities/Private/USDGeomMeshConversion.cpp
@@ -27,6 +27,7 @@
 #include "Misc/Paths.h"
 #include "Modules/ModuleManager.h"
 #include "StaticMeshAttributes.h"
+#include "StaticMeshOperations.h"
 #include "StaticMeshResources.h"
 
 #if WITH_EDITOR
@@ -43,6 +44,7 @@
 	#include "pxr/usd/usd/prim.h"
 	#include "pxr/usd/usd/primRange.h"
 	#include "pxr/usd/usdGeom/mesh.h"
+	#include "pxr/usd/usdGeom/pointInstancer.h"
 	#include "pxr/usd/usdGeom/primvarsAPI.h"
 	#include "pxr/usd/usdGeom/subset.h"
 	#include "pxr/usd/usdGeom/tokens.h"
@@ -565,32 +567,54 @@ namespace UE
 				}
 
 				bool bSuccess = true;
+				bool bTraverseChildren = true;
+
 				if ( pxr::UsdGeomMesh Mesh = pxr::UsdGeomMesh( Prim ) )
 				{
 					bSuccess = UsdToUnreal::ConvertGeomMesh( Mesh, OutMeshDescription, OutMaterialAssignments, ChildTransform, MaterialToPrimvarToUVIndex, TimeCode, RenderContext, bCombineIdenticalMaterialSlots );
 				}
-
-				for ( const pxr::UsdPrim& ChildPrim : Prim.GetFilteredChildren( pxr::UsdTraverseInstanceProxies() ) )
+				else if ( pxr::UsdGeomPointInstancer PointInstancer = pxr::UsdGeomPointInstancer{ Prim } )
 				{
-					if ( !bSuccess )
-					{
-						break;
-					}
-
-					const bool bChildIsFirstPrim = false;
-
-					bSuccess &= RecursivelyCollapseChildMeshes(
-						ChildPrim,
+					bSuccess = UsdToUnreal::ConvertPointInstancerToMesh(
+						PointInstancer,
 						ChildTransform,
-						TimeCode,
-						PurposesToLoad,
-						RenderContext,
 						MaterialToPrimvarToUVIndex,
+						PurposesToLoad,
 						OutMeshDescription,
 						OutMaterialAssignments,
-						bChildIsFirstPrim,
+						TimeCode,
+						RenderContext,
 						bCombineIdenticalMaterialSlots
 					);
+
+					// We never want to step into point instancers when fetching prims for drawing
+					bTraverseChildren = false;
+				}
+
+				if ( bTraverseChildren )
+				{
+					for ( const pxr::UsdPrim& ChildPrim : Prim.GetFilteredChildren( pxr::UsdTraverseInstanceProxies() ) )
+					{
+						if ( !bSuccess )
+						{
+							break;
+						}
+
+						const bool bChildIsFirstPrim = false;
+
+						bSuccess &= RecursivelyCollapseChildMeshes(
+							ChildPrim,
+							ChildTransform,
+							TimeCode,
+							PurposesToLoad,
+							RenderContext,
+							MaterialToPrimvarToUVIndex,
+							OutMeshDescription,
+							OutMaterialAssignments,
+							bChildIsFirstPrim,
+							bCombineIdenticalMaterialSlots
+						);
+					}
 				}
 
 				return bSuccess;
@@ -1039,6 +1063,195 @@ bool UsdToUnreal::ConvertGeomMesh( const pxr::UsdTyped& UsdSchema, FMeshDescript
 	return true;
 }
 
+bool UsdToUnreal::ConvertPointInstancerToMesh(
+	const pxr::UsdGeomPointInstancer& PointInstancer,
+	const FTransform& AdditionalTransform,
+	const TMap< FString, TMap< FString, int32 > >& MaterialToPrimvarsUVSetNames,
+	const EUsdPurpose PurposesToLoad,
+	FMeshDescription& MeshDescription,
+	UsdUtils::FUsdPrimMaterialAssignmentInfo& MaterialAssignments,
+	const pxr::UsdTimeCode TimeCode,
+	const pxr::TfToken& RenderContext,
+	bool bCombineIdenticalMaterialSlots
+)
+{
+	TRACE_CPUPROFILER_EVENT_SCOPE( UsdToUnreal::ConvertPointInstancerToMesh );
+
+	FScopedUsdAllocs Allocs;
+
+	if ( !PointInstancer )
+	{
+		return false;
+	}
+
+	pxr::UsdStageRefPtr Stage = PointInstancer.GetPrim().GetStage();
+	if ( !Stage )
+	{
+		return false;
+	}
+
+	// Bake each prototype to a single mesh description and material assignment struct
+	TArray<FMeshDescription> PrototypeMeshDescriptions;
+	TArray<UsdUtils::FUsdPrimMaterialAssignmentInfo> PrototypeMaterialAssignments;
+	TArray<TMap<FPolygonGroupID, FPolygonGroupID>> PrototypePolygonGroupRemapping;
+	uint32 NumPrototypes = 0;
+	{
+		const pxr::UsdRelationship& Prototypes = PointInstancer.GetPrototypesRel();
+
+		pxr::SdfPathVector PrototypePaths;
+		if ( !Prototypes.GetTargets( &PrototypePaths ) )
+		{
+			return false;
+		}
+
+		NumPrototypes = PrototypePaths.size();
+		if ( NumPrototypes == 0 )
+		{
+			return true;
+		}
+
+		PrototypeMeshDescriptions.SetNum( NumPrototypes );
+		PrototypeMaterialAssignments.SetNum( NumPrototypes );
+		PrototypePolygonGroupRemapping.SetNum( NumPrototypes );
+
+		for ( uint32 PrototypeIndex = 0; PrototypeIndex < NumPrototypes; ++PrototypeIndex )
+		{
+			const pxr::SdfPath& PrototypePath = PrototypePaths[ PrototypeIndex ];
+
+			pxr::UsdPrim PrototypeUsdPrim = Stage->GetPrimAtPath( PrototypePath );
+			if ( !PrototypeUsdPrim )
+			{
+				UE_LOG( LogUsd, Warning, TEXT( "Failed to find prototype '%s' for PointInstancer '%s' within ConvertPointInstancerToMesh" ),
+					*UsdToUnreal::ConvertPath( PrototypePath ),
+					*UsdToUnreal::ConvertPath( PointInstancer.GetPrim().GetPrimPath() )
+				);
+				continue;
+			}
+
+			const bool bSkipRootPrimTransformAndVisibility = false;
+			ConvertGeomMeshHierarchy(
+				PrototypeUsdPrim,
+				TimeCode,
+				PurposesToLoad,
+				RenderContext,
+				MaterialToPrimvarsUVSetNames,
+				PrototypeMeshDescriptions[ PrototypeIndex ],
+				PrototypeMaterialAssignments[ PrototypeIndex ],
+				bSkipRootPrimTransformAndVisibility,
+				bCombineIdenticalMaterialSlots
+			);
+		}
+	}
+
+	// Handle combined prototype material slots.
+	// Sets up PrototypePolygonGroupRemapping so that our new faces are remapped from the prototype's mesh description polygon groups to the
+	// combined mesh description's polygon groups when AppendMeshDescription is called.
+	// Note: We always setup our mesh description polygon groups in the same order as the material assignment slots, so this is not so complicated
+	for ( uint32 PrototypeIndex = 0; PrototypeIndex < NumPrototypes; ++PrototypeIndex )
+	{
+		UsdUtils::FUsdPrimMaterialAssignmentInfo& PrototypeMaterialAssignment = PrototypeMaterialAssignments[ PrototypeIndex ];
+		TMap<FPolygonGroupID, FPolygonGroupID>& PrototypeToCombinedMeshPolygonGroupMap = PrototypePolygonGroupRemapping[ PrototypeIndex ];
+
+		if ( bCombineIdenticalMaterialSlots )
+		{
+			// Build a map of our existing slots since we can hash the entire slot, and our incoming mesh may have an arbitrary number of new slots
+			TMap< UsdUtils::FUsdPrimMaterialSlot, int32 > CombinedMaterialSlotsToIndex;
+			for ( int32 Index = 0; Index < MaterialAssignments.Slots.Num(); ++Index )
+			{
+				const UsdUtils::FUsdPrimMaterialSlot& Slot = MaterialAssignments.Slots[ Index ];
+				CombinedMaterialSlotsToIndex.Add( Slot, Index );
+			}
+
+			for ( int32 PrototypeMaterialSlotIndex = 0; PrototypeMaterialSlotIndex < PrototypeMaterialAssignment.Slots.Num(); ++PrototypeMaterialSlotIndex )
+			{
+				const UsdUtils::FUsdPrimMaterialSlot& LocalSlot = PrototypeMaterialAssignment.Slots[ PrototypeMaterialSlotIndex ];
+				if ( int32* ExistingCombinedIndex = CombinedMaterialSlotsToIndex.Find( LocalSlot ) )
+				{
+					PrototypeToCombinedMeshPolygonGroupMap.Add( PrototypeMaterialSlotIndex, *ExistingCombinedIndex );
+				}
+				else
+				{
+					MaterialAssignments.Slots.Add( LocalSlot );
+					PrototypeToCombinedMeshPolygonGroupMap.Add( PrototypeMaterialSlotIndex, MaterialAssignments.Slots.Num() - 1 );
+				}
+			}
+		}
+		else
+		{
+			const int32 NumExistingMaterialSlots = MaterialAssignments.Slots.Num();
+			MaterialAssignments.Slots.Append( PrototypeMaterialAssignment.Slots );
+
+			for ( int32 PrototypeMaterialSlotIndex = 0; PrototypeMaterialSlotIndex < PrototypeMaterialAssignment.Slots.Num(); ++PrototypeMaterialSlotIndex )
+			{
+				PrototypeToCombinedMeshPolygonGroupMap.Add( PrototypeMaterialSlotIndex, NumExistingMaterialSlots + PrototypeMaterialSlotIndex );
+			}
+		}
+	}
+
+	// Make sure we have the polygon groups we expect. Appending the mesh descriptions will not create new polygon groups if we're using a
+	// PolygonGroupsDelegate, which we will
+	const int32 NumExistingPolygonGroups = MeshDescription.PolygonGroups().Num();
+	for ( int32 NumMissingPolygonGroups = MaterialAssignments.Slots.Num() - NumExistingPolygonGroups; NumMissingPolygonGroups > 0; --NumMissingPolygonGroups )
+	{
+		MeshDescription.CreatePolygonGroup();
+	}
+
+	// Double-check our target mesh description has the attributes we need
+	FStaticMeshAttributes StaticMeshAttributes( MeshDescription );
+	StaticMeshAttributes.Register();
+
+	// Append mesh descriptions
+	FUsdStageInfo StageInfo{ PointInstancer.GetPrim().GetStage() };
+	for ( uint32 PrototypeIndex = 0; PrototypeIndex < NumPrototypes; ++PrototypeIndex )
+	{
+		const FMeshDescription& PrototypeMeshDescription = PrototypeMeshDescriptions[ PrototypeIndex ];
+		UsdUtils::FUsdPrimMaterialAssignmentInfo& PrototypeMaterialAssignment = PrototypeMaterialAssignments[ PrototypeIndex ];
+
+		// We may generate some empty meshes in case a prototype is invisible, for example
+		if ( PrototypeMeshDescription.IsEmpty() )
+		{
+			continue;
+		}
+
+		TArray<FTransform> InstanceTransforms;
+		bool bSuccess = UsdUtils::GetPointInstancerTransforms( StageInfo, PointInstancer, PrototypeIndex, TimeCode, InstanceTransforms );
+		if ( !bSuccess )
+		{
+			UE_LOG( LogUsd, Error, TEXT( "Failed to retrieve point instancer transforms for prototype index '%u' of point instancer '%s'" ),
+				PrototypeIndex,
+				*UsdToUnreal::ConvertPath( PointInstancer.GetPrim().GetPrimPath() )
+			);
+
+			continue;
+		}
+
+		const int32 NumInstances = InstanceTransforms.Num();
+
+		MeshDescription.ReserveNewVertices( PrototypeMeshDescription.Vertices().Num() * NumInstances );
+		MeshDescription.ReserveNewVertexInstances( PrototypeMeshDescription.VertexInstances().Num() * NumInstances );
+		MeshDescription.ReserveNewEdges( PrototypeMeshDescription.Edges().Num() * NumInstances );
+		MeshDescription.ReserveNewTriangles( PrototypeMeshDescription.Triangles().Num() * NumInstances );
+
+		FStaticMeshOperations::FAppendSettings Settings;
+		Settings.PolygonGroupsDelegate = FAppendPolygonGroupsDelegate::CreateLambda(
+			[&PrototypePolygonGroupRemapping, PrototypeIndex]( const FMeshDescription& SourceMesh, FMeshDescription& TargetMesh, PolygonGroupMap& RemapPolygonGroups )
+			{
+				RemapPolygonGroups = PrototypePolygonGroupRemapping[ PrototypeIndex ];
+			}
+		);
+
+		// TODO: Maybe we should make a new overload of AppendMeshDescriptions that can do this more efficiently, since all we need is to change the
+		// transform repeatedly?
+		for ( const FTransform& Transform : InstanceTransforms )
+		{
+			Settings.MeshTransform = Transform * AdditionalTransform;
+			FStaticMeshOperations::AppendMeshDescription( PrototypeMeshDescription, MeshDescription, Settings );
+		}
+	}
+
+	return true;
+}
+
 bool UsdToUnreal::ConvertGeomMeshHierarchy(
 	const pxr::UsdPrim& Prim,
 	const pxr::UsdTimeCode& TimeCode,
@@ -1913,38 +2126,70 @@ TOptional<UsdUtils::FDisplayColorMaterial> UsdUtils::ExtractDisplayColorMaterial
 	return Desc;
 }
 
-bool UsdUtils::IsGeomMeshALOD( const pxr::UsdPrim& UsdMeshPrim )
+namespace UE::UsdGeomMeshConversion::Private
 {
-	pxr::UsdGeomMesh UsdMesh{ UsdMeshPrim };
-	if ( !UsdMesh )
+	bool DoesPrimContainMeshLODsInternal( const pxr::UsdPrim& Prim )
 	{
-		return false;
-	}
+		FScopedUsdAllocs Allocs;
 
-	FScopedUsdAllocs Allocs;
+		if ( !Prim )
+		{
+			return false;
+		}
 
-	pxr::UsdPrim ParentPrim = UsdMeshPrim.GetParent();
-	if ( !ParentPrim )
-	{
-		return false;
-	}
+		const std::string LODString = UnrealIdentifiers::LOD.GetString();
 
-	const std::string LODString = UnrealIdentifiers::LOD.GetString();
+		pxr::UsdVariantSets VariantSets = Prim.GetVariantSets();
+		if ( !VariantSets.HasVariantSet( LODString ) )
+		{
+			return false;
+		}
 
-	pxr::UsdVariantSets VariantSets = ParentPrim.GetVariantSets();
-	if ( !VariantSets.HasVariantSet( LODString ) )
+		std::string Selection = VariantSets.GetVariantSet( LODString ).GetVariantSelection();
+		int32 LODIndex = UsdGeomMeshImpl::GetLODIndexFromName( Selection );
+		if ( LODIndex == INDEX_NONE )
+		{
+			return false;
+		}
+
+		return true;
+	}
+}
+
+bool UsdUtils::DoesPrimContainMeshLODs( const pxr::UsdPrim& Prim )
+{
+	const bool bHasValidVariantSetup = UE::UsdGeomMeshConversion::Private::DoesPrimContainMeshLODsInternal( Prim );
+	if ( bHasValidVariantSetup )
 	{
-		return false;
+		FScopedUsdAllocs Allocs;
+
+		// Check if it has at least one mesh too
+		pxr::UsdPrimSiblingRange PrimRange = Prim.GetChildren();
+		for ( pxr::UsdPrimSiblingRange::iterator PrimRangeIt = PrimRange.begin(); PrimRangeIt != PrimRange.end(); ++PrimRangeIt )
+		{
+			const pxr::UsdPrim& Child = *PrimRangeIt;
+			if ( pxr::UsdGeomMesh ChildMesh{ Child } )
+			{
+				return true;
+				break;
+			}
+		}
 	}
 
-	std::string Selection = VariantSets.GetVariantSet( LODString ).GetVariantSelection();
-	int32 LODIndex = UsdGeomMeshImpl::GetLODIndexFromName( Selection );
-	if ( LODIndex == INDEX_NONE )
+	return false;
+}
+
+bool UsdUtils::IsGeomMeshALOD( const pxr::UsdPrim& UsdMeshPrim )
+{
+	FScopedUsdAllocs Allocs;
+
+	pxr::UsdGeomMesh UsdMesh{ UsdMeshPrim };
+	if ( !UsdMesh )
 	{
 		return false;
 	}
 
-	return true;
+	return UE::UsdGeomMeshConversion::Private::DoesPrimContainMeshLODsInternal( UsdMeshPrim.GetParent() );
 }
 
 int32 UsdUtils::GetNumberOfLODVariants( const pxr::UsdPrim& Prim )
@@ -2473,6 +2718,46 @@ FString UsdUtils::HashGeomMeshPrim( const UE::FUsdStage& Stage, const FString& P
 	return Hash;
 }
 
+bool UsdUtils::GetPointInstancerTransforms( const FUsdStageInfo& StageInfo, const pxr::UsdGeomPointInstancer& PointInstancer, const int32 ProtoIndex, pxr::UsdTimeCode EvalTime, TArray<FTransform>& OutInstanceTransforms )
+{
+	TRACE_CPUPROFILER_EVENT_SCOPE( GetPointInstancerTransforms );
+
+	if ( !PointInstancer )
+	{
+		return false;
+	}
+
+	FScopedUsdAllocs UsdAllocs;
+
+	pxr::VtArray< int > ProtoIndices = UsdUtils::GetUsdValue< pxr::VtArray< int > >( PointInstancer.GetProtoIndicesAttr(), EvalTime );
+
+	pxr::VtMatrix4dArray UsdInstanceTransforms;
+
+	// We don't want the prototype root prim's transforms to be included in these, as they'll already be baked into the meshes themselves
+	if ( !PointInstancer.ComputeInstanceTransformsAtTime( &UsdInstanceTransforms, EvalTime, EvalTime, pxr::UsdGeomPointInstancer::ExcludeProtoXform ) )
+	{
+		return false;
+	}
+
+	int32 Index = 0;
+
+	FScopedUnrealAllocs UnrealAllocs;
+
+	OutInstanceTransforms.Reset( UsdInstanceTransforms.size() );
+
+	for ( pxr::GfMatrix4d& UsdMatrix : UsdInstanceTransforms )
+	{
+		if ( ProtoIndices[ Index ] == ProtoIndex )
+		{
+			OutInstanceTransforms.Add( UsdToUnreal::ConvertMatrix( StageInfo, UsdMatrix ) );
+		}
+
+		++Index;
+	}
+
+	return true;
+}
+
 #undef LOCTEXT_NAMESPACE
 
 #endif // #if USE_USD_SDK
diff --git a/Engine/Plugins/Importers/USDImporter/Source/USDUtilities/Public/USDGeomMeshConversion.h b/Engine/Plugins/Importers/USDImporter/Source/USDUtilities/Public/USDGeomMeshConversion.h
--- a/Engine/Plugins/Importers/USDImporter/Source/USDUtilities/Public/USDGeomMeshConversion.h
+++ b/Engine/Plugins/Importers/USDImporter/Source/USDUtilities/Public/USDGeomMeshConversion.h
@@ -17,6 +17,7 @@
 PXR_NAMESPACE_OPEN_SCOPE
 	class TfToken;
 	class UsdGeomMesh;
+	class UsdGeomPointInstancer;
 	class UsdPrim;
 	class UsdShadeMaterial;
 	class UsdTyped;
@@ -27,6 +28,7 @@ class UMaterialInstanceDynamic;
 class UStaticMesh;
 struct FMeshDescription;
 struct FStaticMeshLODResources;
+struct FUsdStageInfo;
 namespace UsdUtils
 {
 	struct FDisplayColorMaterial;
@@ -61,6 +63,31 @@ namespace UsdToUnreal
 	USDUTILITIES_API bool ConvertGeomMesh( const pxr::UsdTyped& UsdSchema, FMeshDescription& MeshDescription, const FTransform& AdditionalTransform, const pxr::UsdTimeCode TimeCode = pxr::UsdTimeCode::EarliestTime() );
 	USDUTILITIES_API bool ConvertGeomMesh( const pxr::UsdTyped& UsdSchema, FMeshDescription& MeshDescription, const FTransform& AdditionalTransform, const TMap< FString, TMap< FString, int32 > >& MaterialToPrimvarsUVSetNames, const pxr::UsdTimeCode TimeCode = pxr::UsdTimeCode::EarliestTime() );
 
+	/** Bake a point instancer prim into a single FMeshDescription and MaterialAssignments, manually "instancing" all the prototype instances into a single combined mesh.
+	 *
+	 * @param PointInstancer - UsdGeomPointInstancer to bake
+	 * @param AdditionalTransform - Transform (in UnrealEditor space) to bake into the vertex positions and normals of the converted mesh
+	 * @param MaterialToPrimvarsUVSetNames - Maps from a material prim path, to pairs indicating which primvar names are used as 'st' coordinates for this mesh, and which UVIndex materials will sample from (e.g. ["st0", 0], ["myUvSet2", 2], etc). This is used to pick which primvars will become UV sets.
+	 * @param PurposesToLoad - Display purposes to consider when baking the prototypes into meshes
+	 * @param OutMeshDescription - Output parameter that will be filled with the converted data
+	 * @param OutMaterialAssignments - Output parameter that will be filled with the material data extracted from UsdSchema
+	 * @param TimeCode - USD timecode when the UsdGeomPointInstancer should be sampled for mesh data and converted
+	 * @param RenderContext - The render context to use when parsing material assignments from the prototype meshes
+	 * @param bCombineIdenticalMaterialSlots - Whether to try reusing material slots (both local and the ones already in MaterialAssignments) when converting the material assignments from the prototype meshes
+	 * @return Whether the conversion was successful or not.
+	 */
+	USDUTILITIES_API bool ConvertPointInstancerToMesh(
+		const pxr::UsdGeomPointInstancer& PointInstancer,
+		const FTransform& AdditionalTransform,
+		const TMap< FString, TMap< FString, int32 > >& MaterialToPrimvarsUVSetNames,
+		const EUsdPurpose PurposesToLoad,
+		FMeshDescription& OutMeshDescription,
+		UsdUtils::FUsdPrimMaterialAssignmentInfo& OutMaterialAssignments,
+		const pxr::UsdTimeCode TimeCode = pxr::UsdTimeCode::EarliestTime(),
+		const pxr::TfToken& RenderContext = pxr::UsdShadeTokens->universalRenderContext,
+		bool bCombineIdenticalMaterialSlots = true
+	);
+
 	/**
 	 * Recursively traverses down Prim's subtree parsing all Mesh data (including Prim's), baking all transforms and putting the resulting
 	 * mesh data within OutMeshDescription and OutMaterialAssignments. It will not bake Prim's own transform into the Mesh data.
@@ -191,7 +218,10 @@ namespace UsdUtils
 	 */
 	USDUTILITIES_API TArray<FString> GetMaterialUsers( const UE::FUsdPrim& MaterialPrim );
 
-	/** Returns whether this UsdMesh can be interpreted as a LOD of a mesh with multiple LODs */
+	/** Returns whether this prim can be interpreted as describing a static mesh with multiple LODs (i.e. if the prim holds the LOD variant set itself) */
+	USDUTILITIES_API bool DoesPrimContainMeshLODs( const pxr::UsdPrim& Prim );
+
+	/** Returns whether this UsdMesh can be interpreted as a LOD of a mesh with multiple LODs (i.e. if the Mesh prim is inside a LOD variant) */
 	USDUTILITIES_API bool IsGeomMeshALOD( const pxr::UsdPrim& UsdMeshPrim );
 
 	/** Returns how many LOD variants the Prim has. Note that this will return 0 if called on one of the LOD meshes themselves, it's meant to be called on its parent */
@@ -226,6 +256,12 @@ namespace UsdUtils
 	 * If the prim path is invalid or not a GeomMesh, it will return an empty hash
 	 */
 	USDUTILITIES_API FString HashGeomMeshPrim( const UE::FUsdStage& Stage, const FString& PrimPath, double TimeCode );
+
+	/**
+	 * Places in OutInstanceTransforms the UE-space instance transforms for a given point instancer prototype index.
+	 * Returns whether the transforms were successfully retrieved or not.
+	 */
+	USDUTILITIES_API bool GetPointInstancerTransforms( const FUsdStageInfo& StageInfo, const pxr::UsdGeomPointInstancer& PointInstancer, const int32 ProtoIndex, pxr::UsdTimeCode EvalTime, TArray<FTransform>& OutInstanceTransforms );
 }
 
 #endif // #if USE_USD_SDK
