USD: Fix incorrect UI string that escaped batch-rename refactor.

diff --git a/Engine/Plugins/Importers/USDImporter/Source/USDStageEditor/Private/SUSDStage.cpp b/Engine/Plugins/Importers/USDImporter/Source/USDStageEditor/Private/SUSDStage.cpp
--- a/Engine/Plugins/Importers/USDImporter/Source/USDStageEditor/Private/SUSDStage.cpp
+++ b/Engine/Plugins/Importers/USDImporter/Source/USDStageEditor/Private/SUSDStage.cpp
@@ -1015,7 +1015,7 @@ void SUsdStage::FillCollapsingSubMenu( FMenuBuilder& MenuBuilder )
 	);
 
 	MenuBuilder.AddMenuEntry(
-		LOCTEXT( "CollapseTopLevelPointInstancers", "Collapse simple point instancers" ),
+		LOCTEXT( "CollapseTopLevelPointInstancers", "Collapse top-level point instancers" ),
 		LOCTEXT( "CollapseTopLevelPointInstancersToolTip", "If true, will cause us to collapse any point instancer prim into a single static mesh and static mesh component.\nIf false, will cause us to use HierarchicalInstancedStaticMeshComponents to replicate the instancing behavior.\nPoint instancers inside other point instancer prototypes are * always * collapsed into the prototype's static mesh." ),
 		FSlateIcon(),
 		FUIAction(
