USD: Use restTransforms for reference skeleton pose, but account for bindTransforms on skeletal mesh.

diff --git a/Engine/Plugins/Importers/USDImporter/Source/USDSchemas/Private/USDSkelRootTranslator.cpp b/Engine/Plugins/Importers/USDImporter/Source/USDSchemas/Private/USDSkelRootTranslator.cpp
--- a/Engine/Plugins/Importers/USDImporter/Source/USDSchemas/Private/USDSkelRootTranslator.cpp
+++ b/Engine/Plugins/Importers/USDImporter/Source/USDSchemas/Private/USDSkelRootTranslator.cpp
@@ -235,6 +235,16 @@ namespace UsdSkelRootTranslatorImpl
 			HashState.Update( ( uint8* ) JointTransforms.data(), JointTransforms.size() * sizeof( pxr::GfMatrix4d ) );
 		}
 
+		// restTransforms
+		pxr::VtArray<pxr::GfMatrix4d> Transforms;
+		const bool bAtRest = true;
+		InUsdSkeletonQuery.ComputeJointLocalTransforms( &Transforms, pxr::UsdTimeCode::EarliestTime(), bAtRest );
+		HashState.Update( ( uint8* ) Transforms.data(), Transforms.size() * sizeof( pxr::GfMatrix4d ) );
+
+		// bindTransforms
+		InUsdSkeletonQuery.GetJointWorldBindTransforms( &Transforms );
+		HashState.Update( ( uint8* ) Transforms.data(), Transforms.size() * sizeof( pxr::GfMatrix4d ) );
+
 		// Time samples for blend shape curves
 		AnimQuery.GetBlendShapeWeightTimeSamples( &TimeData );
 		HashState.Update( ( uint8* ) TimeData.data(), TimeData.size() * sizeof( double ) );
@@ -395,9 +405,6 @@ namespace UsdSkelRootTranslatorImpl
 				return true;
 			}
 
-			pxr::GfMatrix4d GeomBindTransformUSD = SkinningQuery.GetGeomBindTransform( pxr::UsdTimeCode( InTime ) );
-			FTransform AdditionalTransform = FTransform( UsdToUnreal::ConvertMatrix( StageInfo, GeomBindTransformUSD ) );
-
 			FSkeletalMeshImportData& LODImportData = LODIndexToSkeletalMeshImportDataMap.FindOrAdd( LODIndex );
 			TArray<UsdUtils::FUsdPrimMaterialSlot>& LODSlots = LODIndexToMaterialInfoMap.FindOrAdd( LODIndex ).Slots;
 
@@ -405,7 +412,7 @@ namespace UsdSkelRootTranslatorImpl
 			// into one FSkeletalMeshImportData per LOD, so we need to remap the indices using this
 			uint32 NumPointsBeforeThisMesh = static_cast< uint32 >( LODImportData.Points.Num() );
 
-			bool bSuccess = UsdToUnreal::ConvertSkinnedMesh( SkinningQuery, AdditionalTransform, LODImportData, LODSlots, InMaterialToPrimvarsUVSetNames, RenderContext );
+			bool bSuccess = UsdToUnreal::ConvertSkinnedMesh( SkinningQuery, SkelQuery, LODImportData, LODSlots, InMaterialToPrimvarsUVSetNames, RenderContext );
 			if ( !bSuccess )
 			{
 				return true;
@@ -423,7 +430,6 @@ namespace UsdSkelRootTranslatorImpl
 							BlendShapeQuery.GetBlendShape( BlendShapeIndex ),
 							StageInfo,
 							LODIndex,
-							AdditionalTransform,
 							NumPointsBeforeThisMesh,
 							InOutUsedMorphTargetNames,
 							*OutBlendShapes
diff --git a/Engine/Plugins/Importers/USDImporter/Source/USDUtilities/Private/USDSkeletalDataConversion.cpp b/Engine/Plugins/Importers/USDImporter/Source/USDUtilities/Private/USDSkeletalDataConversion.cpp
--- a/Engine/Plugins/Importers/USDImporter/Source/USDUtilities/Private/USDSkeletalDataConversion.cpp
+++ b/Engine/Plugins/Importers/USDImporter/Source/USDUtilities/Private/USDSkeletalDataConversion.cpp
@@ -47,6 +47,7 @@
 	#include "pxr/usd/usdGeom/primvarsAPI.h"
 	#include "pxr/usd/usdGeom/subset.h"
 	#include "pxr/usd/usdGeom/tokens.h"
+	#include "pxr/usd/usdGeom/xformCache.h"
 	#include "pxr/usd/usdShade/tokens.h"
 	#include "pxr/usd/usdSkel/animation.h"
 	#include "pxr/usd/usdSkel/animMapper.h"
@@ -56,6 +57,7 @@
 	#include "pxr/usd/usdSkel/cache.h"
 	#include "pxr/usd/usdSkel/root.h"
 	#include "pxr/usd/usdSkel/skeletonQuery.h"
+	#include "pxr/usd/usdSkel/skinningQuery.h"
 	#include "pxr/usd/usdSkel/topology.h"
 	#include "pxr/usd/usdSkel/utils.h"
 #include "USDIncludesEnd.h"
@@ -255,7 +257,7 @@ namespace SkelDataConversionImpl
 	}
 
 	/** Converts the given offsets into UnrealEditor space and fills in an FUsdBlendShape object with all the data that will become a morph target */
-	bool CreateUsdBlendShape( const FString& Name, const pxr::VtArray< pxr::GfVec3f >& PointOffsets, const pxr::VtArray< pxr::GfVec3f >& NormalOffsets, const pxr::VtArray< int >& PointIndices, const FUsdStageInfo& StageInfo, const FTransform& AdditionalTransform, uint32 PointIndexOffset, int32 LODIndex, UsdUtils::FUsdBlendShape& OutBlendShape )
+	bool CreateUsdBlendShape( const FString& Name, const pxr::VtArray< pxr::GfVec3f >& PointOffsets, const pxr::VtArray< pxr::GfVec3f >& NormalOffsets, const pxr::VtArray< int >& PointIndices, const FUsdStageInfo& StageInfo, uint32 PointIndexOffset, int32 LODIndex, UsdUtils::FUsdBlendShape& OutBlendShape )
 	{
 		uint32 NumOffsets = PointOffsets.size();
 		uint32 NumIndices = PointIndices.size();
@@ -309,9 +311,6 @@ namespace SkelDataConversionImpl
 			}
 		}
 
-		// This comes from geomBindTransform, which is a manually-input transform, and so can have non-uniform scales, shears, etc
-		FTransform NormalTransform = FTransform(AdditionalTransform.ToInverseMatrixWithScale().GetTransposed());
-
 		OutBlendShape.Vertices.SetNumUninitialized( NumOffsets );
 		for ( uint32 OffsetIndex = 0; OffsetIndex < NumOffsets; ++OffsetIndex )
 		{
@@ -321,11 +320,8 @@ namespace SkelDataConversionImpl
 				: FVector( 0, 0, 0 );
 
 			FMorphTargetDelta& ModifiedVertex = OutBlendShape.Vertices[ OffsetIndex ];
-
-			// Intentionally ignore translation on PositionDelta as this is really a direction vector,
-			// and geomBindTransform's translation is already applied to the mesh vertices
-			ModifiedVertex.PositionDelta = (FVector3f)AdditionalTransform.TransformVector(UEOffset);
-			ModifiedVertex.TangentZDelta = (FVector3f)NormalTransform.TransformVector(UENormal);
+			ModifiedVertex.PositionDelta = ( FVector3f ) UEOffset;
+			ModifiedVertex.TangentZDelta = ( FVector3f ) UENormal;
 			ModifiedVertex.SourceIdx = BaseIndices[ OffsetIndex ];
 		}
 
@@ -1075,18 +1071,21 @@ bool UsdToUnreal::ConvertSkeleton(const pxr::UsdSkelSkeletonQuery& SkeletonQuery
 	}
 
 	// Retrieve the bone transforms to be used as the reference pose
-	VtArray<GfMatrix4d> UsdBoneTransforms;
 	TArray<FTransform> BoneTransforms;
 
-	bool bJointTransformsComputed = SkeletonQuery.ComputeJointLocalTransforms(&UsdBoneTransforms, UsdTimeCode::Default());
-	if (bJointTransformsComputed)
+	// Note that for the FReferenceSkeleton we'll want to use the data in restTransforms, so that the skeletal mesh
+	// assumes the pose described by them when non-animated. See the gigantic comment inside ConvertSkinnedMesh.
+	VtArray<GfMatrix4d> JointLocalRestTransforms;
+	const bool bAtRest = true;
+	bool bJointTransformsComputed = SkeletonQuery.ComputeJointLocalTransforms( &JointLocalRestTransforms, UsdTimeCode::EarliestTime(), bAtRest );
+	if ( bJointTransformsComputed )
 	{
 		UsdStageWeakPtr Stage = SkeletonQuery.GetSkeleton().GetPrim().GetStage();
 		const FUsdStageInfo StageInfo(Stage);
 
-		for (uint32 Index = 0; Index < UsdBoneTransforms.size(); ++Index)
+		for (uint32 Index = 0; Index < JointLocalRestTransforms.size(); ++Index)
 		{
-			const GfMatrix4d& UsdMatrix = UsdBoneTransforms[Index];
+			const GfMatrix4d& UsdMatrix = JointLocalRestTransforms[Index];
 			FTransform BoneTransform = UsdToUnreal::ConvertMatrix( StageInfo, UsdMatrix );
 			BoneTransforms.Add(BoneTransform);
 		}
@@ -1126,7 +1125,7 @@ bool UsdToUnreal::ConvertSkeleton(const pxr::UsdSkelSkeletonQuery& SkeletonQuery
 	return true;
 }
 
-bool UsdToUnreal::ConvertSkinnedMesh(const pxr::UsdSkelSkinningQuery& SkinningQuery, const FTransform& AdditionalTransform, FSkeletalMeshImportData& SkelMeshImportData, TArray< UsdUtils::FUsdPrimMaterialSlot >& MaterialAssignments, const TMap< FString, TMap< FString, int32 > >& MaterialToPrimvarsUVSetNames, const pxr::TfToken& RenderContext )
+bool UsdToUnreal::ConvertSkinnedMesh(const pxr::UsdSkelSkinningQuery& SkinningQuery, const pxr::UsdSkelSkeletonQuery& SkeletonQuery, FSkeletalMeshImportData& SkelMeshImportData, TArray< UsdUtils::FUsdPrimMaterialSlot >& MaterialAssignments, const TMap< FString, TMap< FString, int32 > >& MaterialToPrimvarsUVSetNames, const pxr::TfToken& RenderContext )
 {
 	TRACE_CPUPROFILER_EVENT_SCOPE( UsdToUnreal::ConvertSkinnedMesh );
 
@@ -1148,24 +1147,86 @@ bool UsdToUnreal::ConvertSkinnedMesh(const pxr::UsdSkelSkinningQuery& SkinningQu
 	uint32 NumPoints = 0;
 	uint32 NumExistingPoints = SkelMeshImportData.Points.Num();
 
+	// In USD, bindTransforms is used to place the skeleton in world space, and geomBindTransforms is used to
+	// place the mesh in world space. This can be understood as "bind time", and the relationship between the joints
+	// and triangles there will then be maintained throughout the skinning/animation process. In UE the analogue of
+	// this "bind time" happens when we just fill in FReferenceSkeleton and USkeletalMesh with data, and the
+	// relationship between triangles and joints there is maintained throughout animation.
+	//
+	// When we convert our skinned mesh and fill in our USkeletalMesh, then, we should have two goals in mind:
+	// - Our USkeletalMesh's final triangle positions must match the final FReferenceSkeleton pose in the same way
+	//   that the USD skeleton matches the Mesh prim at USD "bind time" (so we must use bindTransforms and
+	//   geomBindTransform in some way);
+	// - We want to be able to use the restTransforms pose on our FReferenceSkeleton, so that the rest pose in UE is
+	//   the same rest pose used in USD (and not just the bind pose instead).
+	//
+	// To accomplish this, we'll have our FReferenceSkeleton assume the restTransforms of the USD skeleton,
+	// and then skin the mesh data to deform it according to restTransforms too, so that they match.
+	//
+	// This means that if we're placing restTransforms on our FReferenceSkeleton, we must put our mesh data skinned
+	// by restTransforms on our USkeletalMesh. The problem is that the USD mesh raw points is specified in the "bind
+	// time" pose already (except for an additional, optional transform called geomBindTransform that we can apply).
+	// In order to move the mesh data from its source space to restTransforms pose, we must then:
+	// - Apply geomBindTransform: It is now at "bind time", directly in world space (note: don't confuse this world
+	//   space with the final concatenated prim transform after we account for parent Xforms and so on. We're only
+	//   concerned with "bind time" here, and a different Xform/skinning process will get our final mesh at world
+	//   space. i.e. Here we don't care if our parent SkelRoot has e.g. a +30 translation on the X);
+	// - Invert the skinning done by bindTransforms: This reverses the effect of the bindTransforms pose, but note
+	//   that it most importantly *deforms the mesh according to the skeleton's bindTransforms*. This is essential,
+	//   and this step is what will ensure that the triangles-joints relationship specified by bindTransforms +
+	//   geomBindTransforms is maintained. Practically, this means that the mesh will act as if all joints are on
+	//   top of each other at the origin, and the mesh will become a packed mess of triangles;
+	// - Skin the mesh by restTransforms: This puts the mesh nicely in the final rest pose we want.
+	//
+	// We'll do this by creating a set of skinning transforms (MeshToSkeletonRestPose) that concatenates the two latter
+	// steps, and by calling ComputeSkinnedPoints, which applies geomBindTransform first and then skins the mesh with
+	// those transforms.
+	//
+	// Reference: https://graphics.pixar.com/usd/release/api/_usd_skel__schema_overview.html
+	VtArray<GfMatrix4d> MeshToSkeletonRestPose;
+	{
+		bool bSuccess = true;
+
+		// Get world-space restTransforms
+		VtArray<GfMatrix4d> WorldSpaceRestTransforms;
+		{
+			VtArray<GfMatrix4d> JointLocalRestTransforms;
+			const bool bAtRest = true;
+			bSuccess &= SkeletonQuery.ComputeJointLocalTransforms( &JointLocalRestTransforms, UsdTimeCode::EarliestTime(), bAtRest );
+			bSuccess &= UsdSkelConcatJointTransforms( SkeletonQuery.GetTopology(), JointLocalRestTransforms, &WorldSpaceRestTransforms );
+		}
+
+		// Get world-space bindTransforms
+		VtArray<GfMatrix4d> WorldSpaceBindTransforms;
+		bSuccess &= SkeletonQuery.GetJointWorldBindTransforms( &WorldSpaceBindTransforms );
+
+		if ( !bSuccess )
+		{
+			return false;
+		}
+
+		MeshToSkeletonRestPose.resize( WorldSpaceRestTransforms.size() );
+		for ( uint32 Index = 0; Index < WorldSpaceRestTransforms.size(); ++Index )
+		{
+			MeshToSkeletonRestPose[ Index ] = WorldSpaceBindTransforms[ Index ].GetInverse() * WorldSpaceRestTransforms[ Index ];
+		}
+	}
+
 	UsdAttribute PointsAttr = UsdMesh.GetPointsAttr();
 	if (PointsAttr)
 	{
 		VtArray<GfVec3f> UsdPoints;
 		PointsAttr.Get(&UsdPoints, UsdTimeCode::Default());
 
+		SkinningQuery.ComputeSkinnedPoints( MeshToSkeletonRestPose, &UsdPoints, UsdTimeCode::EarliestTime() );
+
 		NumPoints = UsdPoints.size();
 		SkelMeshImportData.Points.AddUninitialized(NumPoints);
 
 		for (uint32 PointIndex = 0; PointIndex < NumPoints; ++PointIndex)
 		{
 			const GfVec3f& Point = UsdPoints[PointIndex];
-
-			// Convert the USD vertex to Unreal and apply the GeomBindTransform to it
-			FVector Pos = UsdToUnreal::ConvertVector(StageInfo, Point);
-			Pos = AdditionalTransform.TransformPosition(Pos);
-
-			SkelMeshImportData.Points[PointIndex + NumExistingPoints] = (FVector3f)Pos;
+			SkelMeshImportData.Points[ PointIndex + NumExistingPoints ] = ( FVector3f ) UsdToUnreal::ConvertVector( StageInfo, Point );
 		}
 	}
 
@@ -1199,7 +1260,11 @@ bool UsdToUnreal::ConvertSkinnedMesh(const pxr::UsdSkelSkinningQuery& SkinningQu
 	UsdAttribute NormalsAttribute = UsdMesh.GetNormalsAttr();
 	if (NormalsAttribute)
 	{
-		NormalsAttribute.Get(&Normals, UsdTimeCode::Default());
+		if ( NormalsAttribute.Get( &Normals, UsdTimeCode::Default() ) && Normals.size() > 0 )
+		{
+			// Like for points, we need to ensure these normals are in the same coordinate space of the skeleton
+			SkinningQuery.ComputeSkinnedNormals( MeshToSkeletonRestPose, &Normals, UsdTimeCode::EarliestTime() );
+		}
 	}
 
 	uint32 NumExistingFaces = SkelMeshImportData.Faces.Num();
@@ -1653,6 +1718,85 @@ bool UsdToUnreal::ConvertSkinnedMesh(const pxr::UsdSkelSkinningQuery& SkinningQu
 	return true;
 }
 
+bool UsdToUnreal::ConvertSkinnedMesh( const pxr::UsdSkelSkinningQuery& UsdSkinningQuery, const FTransform& AdditionalTransform, FSkeletalMeshImportData& SkelMeshImportData, TArray< UsdUtils::FUsdPrimMaterialSlot >& MaterialAssignments, const TMap< FString, TMap< FString, int32 > >& MaterialToPrimvarsUVSetNames, const pxr::TfToken& RenderContext )
+{
+	FScopedUsdAllocs UEAllocs;
+
+	// We must use a SkeletonQuery to properly convert the skinned mesh to the target pose. If we weren't provided one,
+	// we must fetch it manually
+
+	pxr::UsdPrim SkinnedPrim = UsdSkinningQuery.GetPrim();
+	if ( !SkinnedPrim )
+	{
+		return false;
+	}
+
+	bool bFoundSkelRoot = false;
+	pxr::UsdPrim ParentPrim = SkinnedPrim.GetParent();
+	while ( ParentPrim && !ParentPrim.IsPseudoRoot() )
+	{
+		if ( ParentPrim.IsA<pxr::UsdSkelRoot>() )
+		{
+			bFoundSkelRoot = true;
+			break;
+		}
+
+		ParentPrim = ParentPrim.GetParent();
+	}
+	if ( !bFoundSkelRoot )
+	{
+		return false;
+	}
+
+	pxr::UsdSkelRoot SkelRoot{ ParentPrim };
+
+	pxr::UsdSkelCache SkelCache;
+	if ( !SkelCache.Populate( SkelRoot, pxr::UsdTraverseInstanceProxies() ) )
+	{
+		return false;
+	}
+
+	std::vector< pxr::UsdSkelBinding > SkeletonBindings;
+	if ( !SkelCache.ComputeSkelBindings( SkelRoot, &SkeletonBindings, pxr::UsdTraverseInstanceProxies() ) )
+	{
+		return false;
+	}
+
+	for ( const pxr::UsdSkelBinding& Binding : SkeletonBindings )
+	{
+		const pxr::UsdSkelSkeleton& Skeleton = Binding.GetSkeleton();
+		pxr::UsdSkelSkeletonQuery SkelQuery = SkelCache.GetSkelQuery( Skeleton );
+
+		// Double-check that we're talking about the same mesh, at least. Each mesh can only be bound to one
+		// skeleton so this is probably guaranteed to be the same SkelBinding
+		bool bFoundSkinningQuery = false;
+		for ( const pxr::UsdSkelSkinningQuery& SomeSkinningQuery : Binding.GetSkinningTargets() )
+		{
+			if ( SomeSkinningQuery.GetPrim() == SkinnedPrim )
+			{
+				bFoundSkinningQuery = true;
+				break;
+			}
+		}
+		if ( !bFoundSkinningQuery )
+		{
+			break;
+		}
+
+		// We really only ever deal with the first binding, so we can return here
+		return ConvertSkinnedMesh(
+			UsdSkinningQuery,
+			SkelQuery,
+			SkelMeshImportData,
+			MaterialAssignments,
+			MaterialToPrimvarsUVSetNames,
+			RenderContext
+		);
+	}
+
+	return false;
+}
+
 // Using UsdSkelSkeletonQuery instead of UsdSkelAnimQuery as it automatically does the joint remapping when we ask it to compute joint transforms.
 // It also initializes the joint transforms with the rest pose, if available, in case the animation doesn't provide data for all joints.
 bool UsdToUnreal::ConvertSkelAnim( const pxr::UsdSkelSkeletonQuery& InUsdSkeletonQuery, const pxr::VtArray<pxr::UsdSkelSkinningQuery>* InSkinningTargets, const UsdUtils::FBlendShapeMap* InBlendShapes, bool bInInterpretLODs, UAnimSequence* OutSkeletalAnimationAsset, float* OutStartOffsetSeconds )
@@ -2021,12 +2165,12 @@ bool UsdToUnreal::ConvertSkelAnim( const pxr::UsdSkelSkeletonQuery& InUsdSkeleto
 	return true;
 }
 
-bool UsdToUnreal::ConvertBlendShape( const pxr::UsdSkelBlendShape& UsdBlendShape, const FUsdStageInfo& StageInfo, const FTransform& AdditionalTransform, uint32 PointIndexOffset, TSet<FString>& UsedMorphTargetNames, UsdUtils::FBlendShapeMap& OutBlendShapes )
+bool UsdToUnreal::ConvertBlendShape( const pxr::UsdSkelBlendShape& UsdBlendShape, const FUsdStageInfo& StageInfo, uint32 PointIndexOffset, TSet<FString>& UsedMorphTargetNames, UsdUtils::FBlendShapeMap& OutBlendShapes )
 {
-	return ConvertBlendShape(UsdBlendShape, StageInfo, 0, AdditionalTransform, PointIndexOffset, UsedMorphTargetNames, OutBlendShapes);
+	return ConvertBlendShape(UsdBlendShape, StageInfo, 0, PointIndexOffset, UsedMorphTargetNames, OutBlendShapes);
 }
 
-bool UsdToUnreal::ConvertBlendShape( const pxr::UsdSkelBlendShape& UsdBlendShape, const FUsdStageInfo& StageInfo, int32 LODIndex, const FTransform& AdditionalTransform, uint32 PointIndexOffset, TSet<FString>& UsedMorphTargetNames, UsdUtils::FBlendShapeMap& OutBlendShapes )
+bool UsdToUnreal::ConvertBlendShape( const pxr::UsdSkelBlendShape& UsdBlendShape, const FUsdStageInfo& StageInfo, int32 LODIndex, uint32 PointIndexOffset, TSet<FString>& UsedMorphTargetNames, UsdUtils::FBlendShapeMap& OutBlendShapes )
 {
 	FScopedUsdAllocs Allocs;
 
@@ -2057,7 +2201,7 @@ bool UsdToUnreal::ConvertBlendShape( const pxr::UsdSkelBlendShape& UsdBlendShape
 	}
 
 	UsdUtils::FUsdBlendShape PrimaryBlendShape;
-	if ( !SkelDataConversionImpl::CreateUsdBlendShape( PrimaryName, Offsets, Normals, PointIndices, StageInfo, AdditionalTransform, PointIndexOffset, LODIndex, PrimaryBlendShape ) )
+	if ( !SkelDataConversionImpl::CreateUsdBlendShape( PrimaryName, Offsets, Normals, PointIndices, StageInfo, PointIndexOffset, LODIndex, PrimaryBlendShape ) )
 	{
 		return false;
 	}
@@ -2101,7 +2245,7 @@ bool UsdToUnreal::ConvertBlendShape( const pxr::UsdSkelBlendShape& UsdBlendShape
 		// Create separate blend shape for the inbetween
 		// Now how the inbetween always shares the same point indices as the parent
 		UsdUtils::FUsdBlendShape InbetweenShape;
-		if ( !SkelDataConversionImpl::CreateUsdBlendShape( InbetweenName, InbetweenPointsOffsets, InbetweenNormalOffsets, PointIndices, StageInfo, AdditionalTransform, PointIndexOffset, LODIndex, InbetweenShape ) )
+		if ( !SkelDataConversionImpl::CreateUsdBlendShape( InbetweenName, InbetweenPointsOffsets, InbetweenNormalOffsets, PointIndices, StageInfo, PointIndexOffset, LODIndex, InbetweenShape ) )
 		{
 			continue;
 		}
@@ -2126,6 +2270,16 @@ bool UsdToUnreal::ConvertBlendShape( const pxr::UsdSkelBlendShape& UsdBlendShape
 	return true;
 }
 
+bool UsdToUnreal::ConvertBlendShape( const pxr::UsdSkelBlendShape& UsdBlendShape, const FUsdStageInfo& StageInfo, const FTransform& AdditionalTransform, uint32 PointIndexOffset, TSet<FString>& UsedMorphTargetNames, UsdUtils::FBlendShapeMap& OutBlendShapes )
+{
+	return UsdToUnreal::ConvertBlendShape( UsdBlendShape, StageInfo, 0, PointIndexOffset, UsedMorphTargetNames, OutBlendShapes );
+}
+
+bool UsdToUnreal::ConvertBlendShape( const pxr::UsdSkelBlendShape& UsdBlendShape, const FUsdStageInfo& StageInfo, int32 LODIndex, const FTransform& AdditionalTransform, uint32 PointIndexOffset, TSet<FString>& UsedMorphTargetNames, UsdUtils::FBlendShapeMap& OutBlendShapes )
+{
+	return UsdToUnreal::ConvertBlendShape( UsdBlendShape, StageInfo, LODIndex, PointIndexOffset, UsedMorphTargetNames, OutBlendShapes );
+}
+
 USkeletalMesh* UsdToUnreal::GetSkeletalMeshFromImportData(
 	TArray<FSkeletalMeshImportData>& LODIndexToSkeletalMeshImportData,
 	const TArray<SkeletalMeshImportData::FBone>& InSkeletonBones,
diff --git a/Engine/Plugins/Importers/USDImporter/Source/USDUtilities/Public/USDSkeletalDataConversion.h b/Engine/Plugins/Importers/USDImporter/Source/USDUtilities/Public/USDSkeletalDataConversion.h
--- a/Engine/Plugins/Importers/USDImporter/Source/USDUtilities/Public/USDSkeletalDataConversion.h
+++ b/Engine/Plugins/Importers/USDImporter/Source/USDUtilities/Public/USDSkeletalDataConversion.h
@@ -139,25 +139,32 @@ namespace UsdToUnreal
 	 * @param UsdBlendShape - Source USD object with the blend shape data
 	 * @param StageInfo - Details about the stage, required to do the coordinate conversion from the USD blend shape to morph target data
 	 * @param LODIndex - LODIndex of the SkeletalMesh that will use the imported FUsdBlendShape
-	 * @param AdditionalTransform - Additional affine transform to apply to the blend shape deltas and tangents
+	 * @param AdditionalTransform - Additional affine transform to apply to the blend shape deltas and tangents (Note: No longer used, even in the deprecated overload of this function)
 	 * @param PointIndexOffset - Index into the corresponding SkelMeshImportData.Points where the points corresponding to the UsdBlendShape's prim start
 	 * @param UsedMorphTargetNames - Names that the newly created FUsdBlendShapes cannot have (this function will also add the names of the newly created FUsdBlendShapes to it)
 	 * @param OutBlendShapes - Where the newly created blend shapes will be added to
 	 * @return Whether the conversion was successful or not.
 	 */
+	USDUTILITIES_API bool ConvertBlendShape( const pxr::UsdSkelBlendShape& UsdBlendShape, const FUsdStageInfo& StageInfo, uint32 PointIndexOffset, TSet<FString>& UsedMorphTargetNames, UsdUtils::FBlendShapeMap& OutBlendShapes );
+	USDUTILITIES_API bool ConvertBlendShape( const pxr::UsdSkelBlendShape& UsdBlendShape, const FUsdStageInfo& StageInfo, int32 LODIndex, uint32 PointIndexOffset, TSet<FString>& UsedMorphTargetNames, UsdUtils::FBlendShapeMap& OutBlendShapes );
+	UE_DEPRECATED( 5.1, "AdditionalTransform is no longer used, please use the other overloads that don't receive it as a parameter." )
 	USDUTILITIES_API bool ConvertBlendShape( const pxr::UsdSkelBlendShape& UsdBlendShape, const FUsdStageInfo& StageInfo, const FTransform& AdditionalTransform, uint32 PointIndexOffset, TSet<FString>& UsedMorphTargetNames, UsdUtils::FBlendShapeMap& OutBlendShapes );
+	UE_DEPRECATED( 5.1, "AdditionalTransform is no longer used, please use the other overloads that don't receive it as a parameter." )
 	USDUTILITIES_API bool ConvertBlendShape( const pxr::UsdSkelBlendShape& UsdBlendShape, const FUsdStageInfo& StageInfo, int32 LODIndex, const FTransform& AdditionalTransform, uint32 PointIndexOffset, TSet<FString>& UsedMorphTargetNames, UsdUtils::FBlendShapeMap& OutBlendShapes );
 
 	/**
 	 * Extracts skeletal mesh data fro UsdSkinningQuery, and places the results in SkelMeshImportData.
 	 * @param UsdSkinningQuery - SkinningQuery with the data to convert
-	 * @param AdditionalTransform - Additional transform to apply to the vertices of the mesh
+	 * @param SkeletonQuery - Query object for the skeleton this mesh is bound to
+	 * @param AdditionalTransform - Additional transform to apply to the vertices of the mesh (Note: No longer used, even in the deprecated overload of this function)
 	 * @param SkelMeshImportData - Output parameter that will be filled with the converted data
 	 * @param MaterialAssignments - Output parameter that will be filled with the material assignment data extracted from UsdSkinningQuery
 	 * @param MaterialToPrimvarsUVSetNames - Maps from a material prim path, to pairs indicating which primvar names are used as 'st' coordinates for this mesh, and which UVIndex materials will sample from (e.g. ["st0", 0], ["myUvSet2", 2], etc). This is used to pick which primvars will become UV sets.
 	 * @param RenderContext - Render context to use when parsing the skinned mesh's materials (e.g. '' for universal, or 'mdl', or 'unreal', etc.)
 	 * @return Whether the conversion was successful or not.
 	 */
+	USDUTILITIES_API bool ConvertSkinnedMesh( const pxr::UsdSkelSkinningQuery& UsdSkinningQuery, const pxr::UsdSkelSkeletonQuery& SkeletonQuery, FSkeletalMeshImportData& SkelMeshImportData, TArray< UsdUtils::FUsdPrimMaterialSlot >& MaterialAssignments, const TMap< FString, TMap< FString, int32 > >& MaterialToPrimvarsUVSetNames, const pxr::TfToken& RenderContext = pxr::UsdShadeTokens->universalRenderContext );
+	UE_DEPRECATED( 5.1, "Please use the other overload that also receives the relevant UsdSkelSkeletonQuery object." )
 	USDUTILITIES_API bool ConvertSkinnedMesh( const pxr::UsdSkelSkinningQuery& UsdSkinningQuery, const FTransform& AdditionalTransform, FSkeletalMeshImportData& SkelMeshImportData, TArray< UsdUtils::FUsdPrimMaterialSlot >& MaterialAssignments, const TMap< FString, TMap< FString, int32 > >& MaterialToPrimvarsUVSetNames, const pxr::TfToken& RenderContext = pxr::UsdShadeTokens->universalRenderContext );
 
 	/**
