USD: Fix prim rename changes leading to ensures when trying to access the collapsing cache.

diff --git a/Engine/Plugins/Importers/USDImporter/Source/USDStage/Private/USDStageActor.cpp b/Engine/Plugins/Importers/USDImporter/Source/USDStage/Private/USDStageActor.cpp
--- a/Engine/Plugins/Importers/USDImporter/Source/USDStage/Private/USDStageActor.cpp
+++ b/Engine/Plugins/Importers/USDImporter/Source/USDStage/Private/USDStageActor.cpp
@@ -764,16 +764,16 @@ void AUsdStageActor::OnUsdObjectsChanged( const UsdUtils::FObjectChangesByPath&
 		Transactor->Update( InfoChanges, ResyncChanges );
 	}
 
+	const UE::FUsdStage& Stage = GetOrLoadUsdStage();
+	if ( !Stage )
+	{
+		return;
+	}
+
 	// If the stage was closed in a big transaction (e.g. undo open) a random UObject may be transacting before us and triggering USD changes,
 	// and the UE::FUsdStage will still be opened and valid (even though we intend on closing/changing it when we transact). It could be problematic/wasteful if we
 	// responded to those notices, so just early out here. We can do this check because our RootLayer property will already have the new value
 	{
-		const UE::FUsdStage& Stage = GetOrLoadUsdStage();
-		if ( !Stage )
-		{
-			return;
-		}
-
 		const UE::FSdfLayer& StageRoot = Stage.GetRootLayer();
 		if ( !StageRoot )
 		{
@@ -796,11 +796,17 @@ void AUsdStageActor::OnUsdObjectsChanged( const UsdUtils::FObjectChangesByPath&
 	TMap< FString, bool > SortedPrimsChangedList;
 	for ( const TPair<FString, TArray<UsdUtils::FObjectChangeNotice>>& InfoChange : InfoChanges )
 	{
-		const FString& PrimPath = InfoChange.Key;
+		UE::FSdfPath PrimPath = UE::FSdfPath( *InfoChange.Key ).StripAllVariantSelections();
+
+		const bool bPrimExists = static_cast< bool >( Stage.GetPrimAtPath( PrimPath ) );
+		if ( !bPrimExists )
+		{
+			continue;
+		}
 
 		// Some stage info should trigger some resyncs (even though technically info changes) because they should trigger reparsing of geometry
 		bool bIsResync = false;
-		if ( PrimPath == TEXT( "/" ) )
+		if ( PrimPath.IsAbsoluteRootPath() )
 		{
 			for ( const UsdUtils::FObjectChangeNotice& ObjectChange : InfoChange.Value )
 			{
@@ -824,8 +830,18 @@ void AUsdStageActor::OnUsdObjectsChanged( const UsdUtils::FObjectChangesByPath&
 	// Do Resyncs after so that they overwrite pure info changes if we have any
 	for ( const TPair<FString, TArray<UsdUtils::FObjectChangeNotice>>& ResyncChange : ResyncChanges )
 	{
+		UE::FSdfPath PrimPath = UE::FSdfPath( *ResyncChange.Key ).StripAllVariantSelections();
+
+		// In some scenarios (e.g. prim renaming) we'll get resync notices on prims that don't exist anymore. There's no point
+		// doing anything with these, so just ignore them
+		const bool bPrimExists = static_cast< bool >( Stage.GetPrimAtPath( PrimPath ) );
+		if ( !bPrimExists )
+		{
+			continue;
+		}
+
 		const bool bIsResync = true;
-		SortedPrimsChangedList.Add( UE::FSdfPath( *ResyncChange.Key ).StripAllVariantSelections().GetString(), bIsResync );
+		SortedPrimsChangedList.Add( PrimPath.GetString(), bIsResync );
 	}
 
 	SortedPrimsChangedList.KeySort([]( const FString& A, const FString& B ) -> bool { return A.Len() < B.Len(); } );
