USD: Refactor GeomMeshConversion interface interface.

diff --git a/Engine/Plugins/Importers/USDImporter/Source/USDSchemas/Private/USDGeomMeshTranslator.cpp b/Engine/Plugins/Importers/USDImporter/Source/USDSchemas/Private/USDGeomMeshTranslator.cpp
--- a/Engine/Plugins/Importers/USDImporter/Source/USDSchemas/Private/USDGeomMeshTranslator.cpp
+++ b/Engine/Plugins/Importers/USDImporter/Source/USDSchemas/Private/USDGeomMeshTranslator.cpp
@@ -453,7 +453,7 @@ namespace UsdGeomMeshTranslatorImpl
 	// Note that these other LODs will be hidden in other variants, and won't show up on traversal unless we actively switch the variants (which we do here).
 	// We use a separate function for this because there is a very specific set of conditions where we successfully can do this, and we
 	// want to fall back to just parsing UsdMesh as a simple single-LOD mesh if we fail.
-	bool TryLoadingMultipleLODs( const pxr::UsdTyped& UsdMesh, TArray<FMeshDescription>& OutLODIndexToMeshDescription, TArray<UsdUtils::FUsdPrimMaterialAssignmentInfo>& OutLODIndexToMaterialInfo, const TMap< FString, TMap< FString, int32 > >& InMaterialToPrimvarToUVIndex, const FTransform& AdditionalTransform, const pxr::UsdTimeCode InTimeCode, const pxr::TfToken& RenderContext, bool bCombineIdenticalMaterialSlots )
+	bool TryLoadingMultipleLODs( const pxr::UsdTyped& UsdMesh, TArray<FMeshDescription>& OutLODIndexToMeshDescription, TArray<UsdUtils::FUsdPrimMaterialAssignmentInfo>& OutLODIndexToMaterialInfo, const TMap< FString, TMap< FString, int32 > >& InMaterialToPrimvarToUVIndex, const FTransform& AdditionalTransform, const pxr::UsdTimeCode InTimeCode, const pxr::TfToken& RenderContext, bool bMergeIdenticalMaterialSlots )
 	{
 		FScopedUsdAllocs Allocs;
 
@@ -499,15 +499,18 @@ namespace UsdGeomMeshTranslatorImpl
 
 				if ( bSuccess )
 				{
+					UsdToUnreal::FUsdMeshConversionOptions Options;
+					Options.AdditionalTransform = MeshTransform * AdditionalTransform;
+					Options.TimeCode = InTimeCode;
+					Options.RenderContext = RenderContext;
+					Options.bMergeIdenticalMaterialSlots = bMergeIdenticalMaterialSlots;
+					Options.MaterialToPrimvarToUVIndex = &InMaterialToPrimvarToUVIndex;
+
 					bSuccess &= UsdToUnreal::ConvertGeomMesh(
 						LODMesh,
 						TempMeshDescription,
 						TempMaterialInfo,
-						MeshTransform * AdditionalTransform,
-						InMaterialToPrimvarToUVIndex,
-						InTimeCode,
-						RenderContext,
-						bCombineIdenticalMaterialSlots
+						Options
 					);
 				}
 			}
@@ -537,7 +540,7 @@ namespace UsdGeomMeshTranslatorImpl
 	}
 
 	void LoadMeshDescriptions( const pxr::UsdTyped& UsdMesh, TArray<FMeshDescription>& OutLODIndexToMeshDescription, TArray<UsdUtils::FUsdPrimMaterialAssignmentInfo>& OutLODIndexToMaterialInfo,
-		const TMap< FString, TMap< FString, int32 > >& MaterialToPrimvarToUVIndex, const FTransform& AdditionalTransform, const pxr::UsdTimeCode TimeCode, bool bInterpretLODs, const FName& RenderContext, bool bCombineIdenticalMaterialSlots = true )
+		const TMap< FString, TMap< FString, int32 > >& MaterialToPrimvarToUVIndex, const FTransform& AdditionalTransform, const pxr::UsdTimeCode TimeCode, bool bInterpretLODs, const FName& RenderContext, bool bMergeIdenticalMaterialSlots = true )
 	{
 		if ( !UsdMesh )
 		{
@@ -553,7 +556,7 @@ namespace UsdGeomMeshTranslatorImpl
 		bool bInterpretedLODs = false;
 		if ( bInterpretLODs )
 		{
-			bInterpretedLODs = TryLoadingMultipleLODs( UsdMesh, OutLODIndexToMeshDescription, OutLODIndexToMaterialInfo, MaterialToPrimvarToUVIndex, AdditionalTransform, TimeCode, RenderContextToken, bCombineIdenticalMaterialSlots );
+			bInterpretedLODs = TryLoadingMultipleLODs( UsdMesh, OutLODIndexToMeshDescription, OutLODIndexToMaterialInfo, MaterialToPrimvarToUVIndex, AdditionalTransform, TimeCode, RenderContextToken, bMergeIdenticalMaterialSlots );
 		}
 
 		if ( !bInterpretedLODs )
@@ -577,15 +580,18 @@ namespace UsdGeomMeshTranslatorImpl
 
 			if ( bSuccess )
 			{
+				UsdToUnreal::FUsdMeshConversionOptions Options;
+				Options.AdditionalTransform = MeshTransform * AdditionalTransform;
+				Options.TimeCode = TimeCode;
+				Options.RenderContext = RenderContextToken;
+				Options.bMergeIdenticalMaterialSlots = bMergeIdenticalMaterialSlots;
+				Options.MaterialToPrimvarToUVIndex = &MaterialToPrimvarToUVIndex;
+
 				bSuccess &= UsdToUnreal::ConvertGeomMesh(
-					UsdMesh,
+					pxr::UsdGeomMesh{ UsdMesh },
 					TempMeshDescription,
 					TempMaterialInfo,
-					MeshTransform * AdditionalTransform,
-					MaterialToPrimvarToUVIndex,
-					TimeCode,
-					RenderContextToken,
-					bCombineIdenticalMaterialSlots
+					Options
 				);
 			}
 
diff --git a/Engine/Plugins/Importers/USDImporter/Source/USDSchemas/Private/USDGeomPointInstancerTranslator.cpp b/Engine/Plugins/Importers/USDImporter/Source/USDSchemas/Private/USDGeomPointInstancerTranslator.cpp
--- a/Engine/Plugins/Importers/USDImporter/Source/USDSchemas/Private/USDGeomPointInstancerTranslator.cpp
+++ b/Engine/Plugins/Importers/USDImporter/Source/USDSchemas/Private/USDGeomPointInstancerTranslator.cpp
@@ -109,16 +109,19 @@ void FUsdGeomPointInstancerCreateAssetsTaskChain::SetupTasks()
 				RenderContextToken = UnrealToUsd::ConvertToken( *Context->RenderContext.ToString() ).Get();
 			}
 
+			UsdToUnreal::FUsdMeshConversionOptions Options;
+			Options.TimeCode = Context->Time;
+			Options.PurposesToLoad = Context->PurposesToLoad;
+			Options.RenderContext = RenderContextToken;
+			Options.MaterialToPrimvarToUVIndex = MaterialToPrimvarToUVIndex;
+			Options.bMergeIdenticalMaterialSlots = Context->bMergeIdenticalMaterialSlots;
+
 			UsdToUnreal::ConvertGeomMeshHierarchy(
 				GetPrim(),
-				pxr::UsdTimeCode( Context->Time ),
-				Context->PurposesToLoad,
-				RenderContextToken,
-				*MaterialToPrimvarToUVIndex,
 				AddedMeshDescription,
 				AssignmentInfo,
-				bIgnoreTopLevelTransformAndVisibility,
-				Context->bMergeIdenticalMaterialSlots
+				Options,
+				bIgnoreTopLevelTransformAndVisibility
 			);
 
 			return !AddedMeshDescription.IsEmpty();
diff --git a/Engine/Plugins/Importers/USDImporter/Source/USDSchemas/Private/USDGeomXformableTranslator.cpp b/Engine/Plugins/Importers/USDImporter/Source/USDSchemas/Private/USDGeomXformableTranslator.cpp
--- a/Engine/Plugins/Importers/USDImporter/Source/USDSchemas/Private/USDGeomXformableTranslator.cpp
+++ b/Engine/Plugins/Importers/USDImporter/Source/USDSchemas/Private/USDGeomXformableTranslator.cpp
@@ -99,16 +99,19 @@ void FUsdGeomXformableCreateAssetsTaskChain::SetupTasks()
 			// We're going to put Prim's transform and visibility on the component, so we don't need to bake it into the combined mesh
 			const bool bSkipRootPrimTransformAndVis = true;
 
+			UsdToUnreal::FUsdMeshConversionOptions Options;
+			Options.TimeCode = Context->Time;
+			Options.PurposesToLoad = Context->PurposesToLoad;
+			Options.RenderContext = RenderContextToken;
+			Options.MaterialToPrimvarToUVIndex = MaterialToPrimvarToUVIndex;
+			Options.bMergeIdenticalMaterialSlots = Context->bMergeIdenticalMaterialSlots;
+
 			UsdToUnreal::ConvertGeomMeshHierarchy(
 				GetPrim(),
-				pxr::UsdTimeCode( Context->Time ),
-				Context->PurposesToLoad,
-				RenderContextToken,
-				*MaterialToPrimvarToUVIndex,
 				AddedMeshDescription,
 				AssignmentInfo,
-				bSkipRootPrimTransformAndVis,
-				Context->bMergeIdenticalMaterialSlots
+				Options,
+				bSkipRootPrimTransformAndVis
 			);
 
 			return !AddedMeshDescription.IsEmpty();
diff --git a/Engine/Plugins/Importers/USDImporter/Source/USDUtilities/Private/USDGeomMeshConversion.cpp b/Engine/Plugins/Importers/USDImporter/Source/USDUtilities/Private/USDGeomMeshConversion.cpp
--- a/Engine/Plugins/Importers/USDImporter/Source/USDUtilities/Private/USDGeomMeshConversion.cpp
+++ b/Engine/Plugins/Importers/USDImporter/Source/USDUtilities/Private/USDGeomMeshConversion.cpp
@@ -55,633 +55,607 @@
 
 #define LOCTEXT_NAMESPACE "USDGeomMeshConversion"
 
-namespace UE
+namespace UE::UsdGeomMeshConversion::Private
 {
-	namespace UsdGeomMeshConversion
+	static const FString DisplayColorID = TEXT( "!DisplayColor" );
+
+	int32 GetPrimValueIndex( const pxr::TfToken& InterpType, const int32 VertexIndex, const int32 VertexInstanceIndex, const int32 PolygonIndex )
 	{
-		namespace Private
+		if ( InterpType == pxr::UsdGeomTokens->vertex )
+		{
+			return VertexIndex;
+		}
+		else if ( InterpType == pxr::UsdGeomTokens->varying )
+		{
+			return VertexIndex;
+		}
+		else if ( InterpType == pxr::UsdGeomTokens->faceVarying )
 		{
-			static const FString DisplayColorID = TEXT( "!DisplayColor" );
+			return VertexInstanceIndex;
+		}
+		else if ( InterpType == pxr::UsdGeomTokens->uniform )
+		{
+			return PolygonIndex;
+		}
+		else /* if ( InterpType == pxr::UsdGeomTokens->constant ) */
+		{
+			return 0; // return index 0 for constant or any other unsupported cases
+		}
+	}
 
-			int32 GetPrimValueIndex( const pxr::TfToken& InterpType, const int32 VertexIndex, const int32 VertexInstanceIndex, const int32 PolygonIndex )
-			{
-				if ( InterpType == pxr::UsdGeomTokens->vertex )
-				{
-					return VertexIndex;
-				}
-				else if ( InterpType == pxr::UsdGeomTokens->varying )
-				{
-					return VertexIndex;
-				}
-				else if ( InterpType == pxr::UsdGeomTokens->faceVarying )
-				{
-					return VertexInstanceIndex;
-				}
-				else if ( InterpType == pxr::UsdGeomTokens->uniform )
-				{
-					return PolygonIndex;
-				}
-				else /* if ( InterpType == pxr::UsdGeomTokens->constant ) */
-				{
-					return 0; // return index 0 for constant or any other unsupported cases
-				}
-			}
+	int32 GetLODIndexFromName( const std::string& Name )
+	{
+		const std::string LODString = UnrealIdentifiers::LOD.GetString();
 
-			int32 GetLODIndexFromName( const std::string& Name )
-			{
-				const std::string LODString = UnrealIdentifiers::LOD.GetString();
+		// True if Name does not start with "LOD"
+		if ( Name.rfind( LODString, 0 ) != 0 )
+		{
+			return INDEX_NONE;
+		}
 
-				// True if Name does not start with "LOD"
-				if ( Name.rfind( LODString, 0 ) != 0 )
-				{
-					return INDEX_NONE;
-				}
+		// After LODString there should be only numbers
+		if ( Name.find_first_not_of( "0123456789", LODString.size() ) != std::string::npos )
+		{
+			return INDEX_NONE;
+		}
 
-				// After LODString there should be only numbers
-				if ( Name.find_first_not_of( "0123456789", LODString.size() ) != std::string::npos )
-				{
-					return INDEX_NONE;
-				}
+		const int Base = 10;
+		char** EndPtr = nullptr;
+		return std::strtol( Name.c_str() + LODString.size(), EndPtr, Base );
+	}
 
-				const int Base = 10;
-				char** EndPtr = nullptr;
-				return std::strtol( Name.c_str() + LODString.size(), EndPtr, Base );
-			}
+	void ConvertStaticMeshLOD(
+		int32 LODIndex,
+		const FStaticMeshLODResources& LODRenderMesh,
+		pxr::UsdGeomMesh& UsdMesh,
+		const pxr::VtArray< std::string >& MaterialAssignments,
+		const pxr::UsdTimeCode TimeCode,
+		pxr::UsdPrim MaterialPrim
+	)
+	{
+		pxr::UsdPrim MeshPrim = UsdMesh.GetPrim();
+		pxr::UsdStageRefPtr Stage = MeshPrim.GetStage();
+		if ( !Stage )
+		{
+			return;
+		}
+		const FUsdStageInfo StageInfo{ Stage };
 
-			void ConvertStaticMeshLOD(
-				int32 LODIndex,
-				const FStaticMeshLODResources& LODRenderMesh,
-				pxr::UsdGeomMesh& UsdMesh,
-				const pxr::VtArray< std::string >& MaterialAssignments,
-				const pxr::UsdTimeCode TimeCode,
-				pxr::UsdPrim MaterialPrim
-			)
+		// Vertices
+		{
+			const int32 VertexCount = LODRenderMesh.VertexBuffers.StaticMeshVertexBuffer.GetNumVertices();
+
+			// Points
 			{
-				pxr::UsdPrim MeshPrim = UsdMesh.GetPrim();
-				pxr::UsdStageRefPtr Stage = MeshPrim.GetStage();
-				if ( !Stage )
+				pxr::UsdAttribute Points = UsdMesh.CreatePointsAttr();
+				if ( Points )
 				{
-					return;
+					pxr::VtArray< pxr::GfVec3f > PointsArray;
+					PointsArray.reserve( VertexCount );
+
+					for ( int32 VertexIndex = 0; VertexIndex < VertexCount; ++VertexIndex )
+					{
+						FVector VertexPosition = (FVector)LODRenderMesh.VertexBuffers.PositionVertexBuffer.VertexPosition( VertexIndex );
+						PointsArray.push_back( UnrealToUsd::ConvertVector( StageInfo, VertexPosition ) );
+					}
+
+					Points.Set( PointsArray, TimeCode );
 				}
-				const FUsdStageInfo StageInfo{ Stage };
+			}
 
-				// Vertices
+			// Normals
+			{
+				pxr::UsdAttribute NormalsAttribute = UsdMesh.CreateNormalsAttr();
+				if ( NormalsAttribute )
 				{
-					const int32 VertexCount = LODRenderMesh.VertexBuffers.StaticMeshVertexBuffer.GetNumVertices();
+					pxr::VtArray< pxr::GfVec3f > Normals;
+					Normals.reserve( VertexCount );
 
-					// Points
+					for ( int32 VertexIndex = 0; VertexIndex < VertexCount; ++VertexIndex )
 					{
-						pxr::UsdAttribute Points = UsdMesh.CreatePointsAttr();
-						if ( Points )
-						{
-							pxr::VtArray< pxr::GfVec3f > PointsArray;
-							PointsArray.reserve( VertexCount );
-
-							for ( int32 VertexIndex = 0; VertexIndex < VertexCount; ++VertexIndex )
-							{
-								FVector VertexPosition = (FVector)LODRenderMesh.VertexBuffers.PositionVertexBuffer.VertexPosition( VertexIndex );
-								PointsArray.push_back( UnrealToUsd::ConvertVector( StageInfo, VertexPosition ) );
-							}
-
-							Points.Set( PointsArray, TimeCode );
-						}
+						FVector VertexNormal = (FVector4)LODRenderMesh.VertexBuffers.StaticMeshVertexBuffer.VertexTangentZ( VertexIndex );
+						Normals.push_back( UnrealToUsd::ConvertVector( StageInfo, VertexNormal ) );
 					}
 
-					// Normals
-					{
-						pxr::UsdAttribute NormalsAttribute = UsdMesh.CreateNormalsAttr();
-						if ( NormalsAttribute )
-						{
-							pxr::VtArray< pxr::GfVec3f > Normals;
-							Normals.reserve( VertexCount );
+					NormalsAttribute.Set( Normals, TimeCode );
+				}
+			}
 
-							for ( int32 VertexIndex = 0; VertexIndex < VertexCount; ++VertexIndex )
-							{
-								FVector VertexNormal = (FVector4)LODRenderMesh.VertexBuffers.StaticMeshVertexBuffer.VertexTangentZ( VertexIndex );
-								Normals.push_back( UnrealToUsd::ConvertVector( StageInfo, VertexNormal ) );
-							}
+			// UVs
+			{
+				const int32 TexCoordSourceCount = LODRenderMesh.VertexBuffers.StaticMeshVertexBuffer.GetNumTexCoords();
 
-							NormalsAttribute.Set( Normals, TimeCode );
-						}
-					}
+				for ( int32 TexCoordSourceIndex = 0; TexCoordSourceIndex < TexCoordSourceCount; ++TexCoordSourceIndex )
+				{
+					pxr::TfToken UsdUVSetName = UsdUtils::GetUVSetName( TexCoordSourceIndex ).Get();
 
-					// UVs
+					pxr::UsdGeomPrimvar PrimvarST = UsdMesh.CreatePrimvar( UsdUVSetName, pxr::SdfValueTypeNames->TexCoord2fArray, pxr::UsdGeomTokens->vertex );
+
+					if ( PrimvarST )
 					{
-						const int32 TexCoordSourceCount = LODRenderMesh.VertexBuffers.StaticMeshVertexBuffer.GetNumTexCoords();
+						pxr::VtVec2fArray UVs;
 
-						for ( int32 TexCoordSourceIndex = 0; TexCoordSourceIndex < TexCoordSourceCount; ++TexCoordSourceIndex )
+						for ( int32 VertexIndex = 0; VertexIndex < VertexCount; ++VertexIndex )
 						{
-							pxr::TfToken UsdUVSetName = UsdUtils::GetUVSetName( TexCoordSourceIndex ).Get();
+							FVector2D TexCoord = FVector2D(LODRenderMesh.VertexBuffers.StaticMeshVertexBuffer.GetVertexUV( VertexIndex, TexCoordSourceIndex ));
+							TexCoord[ 1 ] = 1.f - TexCoord[ 1 ];
 
-							pxr::UsdGeomPrimvar PrimvarST = UsdMesh.CreatePrimvar( UsdUVSetName, pxr::SdfValueTypeNames->TexCoord2fArray, pxr::UsdGeomTokens->vertex );
+							UVs.push_back( UnrealToUsd::ConvertVector( TexCoord ) );
+						}
 
-							if ( PrimvarST )
-							{
-								pxr::VtVec2fArray UVs;
+						PrimvarST.Set( UVs, TimeCode );
+					}
+				}
+			}
 
-								for ( int32 VertexIndex = 0; VertexIndex < VertexCount; ++VertexIndex )
-								{
-									FVector2D TexCoord = FVector2D(LODRenderMesh.VertexBuffers.StaticMeshVertexBuffer.GetVertexUV( VertexIndex, TexCoordSourceIndex ));
-									TexCoord[ 1 ] = 1.f - TexCoord[ 1 ];
+			// Vertex colors
+			if ( LODRenderMesh.bHasColorVertexData )
+			{
+				pxr::UsdGeomPrimvar DisplayColorPrimvar = UsdMesh.CreateDisplayColorPrimvar( pxr::UsdGeomTokens->vertex );
+				pxr::UsdGeomPrimvar DisplayOpacityPrimvar = UsdMesh.CreateDisplayOpacityPrimvar( pxr::UsdGeomTokens->vertex );
 
-									UVs.push_back( UnrealToUsd::ConvertVector( TexCoord ) );
-								}
+				if ( DisplayColorPrimvar )
+				{
+					pxr::VtArray< pxr::GfVec3f > DisplayColors;
+					DisplayColors.reserve( VertexCount );
 
-								PrimvarST.Set( UVs, TimeCode );
-							}
-						}
-					}
+					pxr::VtArray< float > DisplayOpacities;
+					DisplayOpacities.reserve( VertexCount );
 
-					// Vertex colors
-					if ( LODRenderMesh.bHasColorVertexData )
+					for ( int32 VertexIndex = 0; VertexIndex < VertexCount; ++VertexIndex )
 					{
-						pxr::UsdGeomPrimvar DisplayColorPrimvar = UsdMesh.CreateDisplayColorPrimvar( pxr::UsdGeomTokens->vertex );
-						pxr::UsdGeomPrimvar DisplayOpacityPrimvar = UsdMesh.CreateDisplayOpacityPrimvar( pxr::UsdGeomTokens->vertex );
+						const FColor& VertexColor = LODRenderMesh.VertexBuffers.ColorVertexBuffer.VertexColor( VertexIndex );
 
-						if ( DisplayColorPrimvar )
-						{
-							pxr::VtArray< pxr::GfVec3f > DisplayColors;
-							DisplayColors.reserve( VertexCount );
-
-							pxr::VtArray< float > DisplayOpacities;
-							DisplayOpacities.reserve( VertexCount );
+						pxr::GfVec4f Color = UnrealToUsd::ConvertColor( VertexColor );
+						DisplayColors.push_back( pxr::GfVec3f( Color[ 0 ], Color[ 1 ], Color[ 2 ] ) );
+						DisplayOpacities.push_back( Color[ 3 ] );
+					}
 
-							for ( int32 VertexIndex = 0; VertexIndex < VertexCount; ++VertexIndex )
-							{
-								const FColor& VertexColor = LODRenderMesh.VertexBuffers.ColorVertexBuffer.VertexColor( VertexIndex );
+					DisplayColorPrimvar.Set( DisplayColors, TimeCode );
+					DisplayOpacityPrimvar.Set( DisplayOpacities, TimeCode );
+				}
+			}
+		}
 
-								pxr::GfVec4f Color = UnrealToUsd::ConvertColor( VertexColor );
-								DisplayColors.push_back( pxr::GfVec3f( Color[ 0 ], Color[ 1 ], Color[ 2 ] ) );
-								DisplayOpacities.push_back( Color[ 3 ] );
-							}
+		// Faces
+		{
+			const int32 FaceCount = LODRenderMesh.GetNumTriangles();
 
-							DisplayColorPrimvar.Set( DisplayColors, TimeCode );
-							DisplayOpacityPrimvar.Set( DisplayOpacities, TimeCode );
-						}
-					}
-				}
+			// Face Vertex Counts
+			{
+				pxr::UsdAttribute FaceCountsAttribute = UsdMesh.CreateFaceVertexCountsAttr();
 
-				// Faces
+				if ( FaceCountsAttribute )
 				{
-					const int32 FaceCount = LODRenderMesh.GetNumTriangles();
+					pxr::VtArray< int > FaceVertexCounts;
+					FaceVertexCounts.reserve( FaceCount );
 
-					// Face Vertex Counts
+					for ( int32 FaceIndex = 0; FaceIndex < FaceCount; ++FaceIndex )
 					{
-						pxr::UsdAttribute FaceCountsAttribute = UsdMesh.CreateFaceVertexCountsAttr();
-
-						if ( FaceCountsAttribute )
-						{
-							pxr::VtArray< int > FaceVertexCounts;
-							FaceVertexCounts.reserve( FaceCount );
-
-							for ( int32 FaceIndex = 0; FaceIndex < FaceCount; ++FaceIndex )
-							{
-								FaceVertexCounts.push_back( 3 );
-							}
-
-							FaceCountsAttribute.Set( FaceVertexCounts, TimeCode );
-						}
+						FaceVertexCounts.push_back( 3 );
 					}
 
-					// Face Vertex Indices
-					{
-						pxr::UsdAttribute FaceVertexIndicesAttribute = UsdMesh.GetFaceVertexIndicesAttr();
+					FaceCountsAttribute.Set( FaceVertexCounts, TimeCode );
+				}
+			}
 
-						if ( FaceVertexIndicesAttribute )
-						{
-							FIndexArrayView Indices = LODRenderMesh.IndexBuffer.GetArrayView();
-							ensure( Indices.Num() == FaceCount * 3 );
+			// Face Vertex Indices
+			{
+				pxr::UsdAttribute FaceVertexIndicesAttribute = UsdMesh.GetFaceVertexIndicesAttr();
 
-							pxr::VtArray< int > FaceVertexIndices;
-							FaceVertexIndices.reserve( FaceCount * 3 );
+				if ( FaceVertexIndicesAttribute )
+				{
+					FIndexArrayView Indices = LODRenderMesh.IndexBuffer.GetArrayView();
+					ensure( Indices.Num() == FaceCount * 3 );
 
-							for ( int32 Index = 0; Index < FaceCount * 3; ++Index )
-							{
-								FaceVertexIndices.push_back( Indices[ Index ] );
-							}
+					pxr::VtArray< int > FaceVertexIndices;
+					FaceVertexIndices.reserve( FaceCount * 3 );
 
-							FaceVertexIndicesAttribute.Set( FaceVertexIndices, TimeCode );
-						}
+					for ( int32 Index = 0; Index < FaceCount * 3; ++Index )
+					{
+						FaceVertexIndices.push_back( Indices[ Index ] );
 					}
+
+					FaceVertexIndicesAttribute.Set( FaceVertexIndices, TimeCode );
 				}
+			}
+		}
+
+		// Material assignments
+		{
+			bool bHasUEMaterialAssignements = false;
 
-				// Material assignments
+			pxr::VtArray< std::string > UnrealMaterialsForLOD;
+			for ( const FStaticMeshSection& Section : LODRenderMesh.Sections )
+			{
+				if ( Section.MaterialIndex >= 0 && Section.MaterialIndex < MaterialAssignments.size() )
+				{
+					UnrealMaterialsForLOD.push_back( MaterialAssignments[ Section.MaterialIndex ] );
+					bHasUEMaterialAssignements = true;
+				}
+				else
 				{
-					bool bHasUEMaterialAssignements = false;
+					// Keep unrealMaterials with the same number of elements as our MaterialIndices expect
+					UnrealMaterialsForLOD.push_back( "" );
+				}
+			}
 
-					pxr::VtArray< std::string > UnrealMaterialsForLOD;
-					for ( const FStaticMeshSection& Section : LODRenderMesh.Sections )
+			// This LOD has a single material assignment, just add an unrealMaterials attribute to the mesh prim
+			if ( bHasUEMaterialAssignements && UnrealMaterialsForLOD.size() == 1 )
+			{
+				const bool bHasMaterialAttribute = MaterialPrim.HasAttribute( UnrealIdentifiers::MaterialAssignments );
+				if ( bHasUEMaterialAssignements )
+				{
+					if ( pxr::UsdAttribute UEMaterialsAttribute = MaterialPrim.CreateAttribute( UnrealIdentifiers::MaterialAssignment, pxr::SdfValueTypeNames->String ) )
 					{
-						if ( Section.MaterialIndex >= 0 && Section.MaterialIndex < MaterialAssignments.size() )
-						{
-							UnrealMaterialsForLOD.push_back( MaterialAssignments[ Section.MaterialIndex ] );
-							bHasUEMaterialAssignements = true;
-						}
-						else
-						{
-							// Keep unrealMaterials with the same number of elements as our MaterialIndices expect
-							UnrealMaterialsForLOD.push_back( "" );
-						}
+						UEMaterialsAttribute.Set( UnrealMaterialsForLOD[ 0 ] );
 					}
+				}
+				else if ( bHasMaterialAttribute )
+				{
+					MaterialPrim.GetAttribute( UnrealIdentifiers::MaterialAssignments ).Clear();
+				}
+			}
+			// Multiple material assignments to the same LOD (and so the same mesh prim). Need to create a GeomSubset for each UE mesh section
+			else if ( UnrealMaterialsForLOD.size() > 1 )
+			{
+				// Need to fetch all triangles of a section, and add their indices
+				for ( int32 SectionIndex = 0; SectionIndex < LODRenderMesh.Sections.Num(); ++SectionIndex )
+				{
+					const FStaticMeshSection& Section = LODRenderMesh.Sections[ SectionIndex ];
 
-					// This LOD has a single material assignment, just add an unrealMaterials attribute to the mesh prim
-					if ( bHasUEMaterialAssignements && UnrealMaterialsForLOD.size() == 1 )
+					// Note that we will continue on even if we have no material assignment, so as to satisfy the "partition" family condition (below)
+					std::string SectionMaterial;
+					if ( Section.MaterialIndex >= 0 && Section.MaterialIndex < MaterialAssignments.size() )
 					{
-						const bool bHasMaterialAttribute = MaterialPrim.HasAttribute( UnrealIdentifiers::MaterialAssignments );
-						if ( bHasUEMaterialAssignements )
-						{
-							if ( pxr::UsdAttribute UEMaterialsAttribute = MaterialPrim.CreateAttribute( UnrealIdentifiers::MaterialAssignment, pxr::SdfValueTypeNames->String ) )
-							{
-								UEMaterialsAttribute.Set( UnrealMaterialsForLOD[ 0 ] );
-							}
-						}
-						else if ( bHasMaterialAttribute )
-						{
-							MaterialPrim.GetAttribute( UnrealIdentifiers::MaterialAssignments ).Clear();
-						}
+						SectionMaterial = MaterialAssignments[ Section.MaterialIndex ];
 					}
-					// Multiple material assignments to the same LOD (and so the same mesh prim). Need to create a GeomSubset for each UE mesh section
-					else if ( UnrealMaterialsForLOD.size() > 1 )
-					{
-						// Need to fetch all triangles of a section, and add their indices
-						for ( int32 SectionIndex = 0; SectionIndex < LODRenderMesh.Sections.Num(); ++SectionIndex )
-						{
-							const FStaticMeshSection& Section = LODRenderMesh.Sections[ SectionIndex ];
-
-							// Note that we will continue on even if we have no material assignment, so as to satisfy the "partition" family condition (below)
-							std::string SectionMaterial;
-							if ( Section.MaterialIndex >= 0 && Section.MaterialIndex < MaterialAssignments.size() )
-							{
-								SectionMaterial = MaterialAssignments[ Section.MaterialIndex ];
-							}
 
-							pxr::UsdPrim GeomSubsetPrim = Stage->DefinePrim(
-								MeshPrim.GetPath().AppendPath( pxr::SdfPath( "Section" + std::to_string( SectionIndex ) ) ),
-								UnrealToUsd::ConvertToken( TEXT( "GeomSubset" ) ).Get()
-							);
+					pxr::UsdPrim GeomSubsetPrim = Stage->DefinePrim(
+						MeshPrim.GetPath().AppendPath( pxr::SdfPath( "Section" + std::to_string( SectionIndex ) ) ),
+						UnrealToUsd::ConvertToken( TEXT( "GeomSubset" ) ).Get()
+					);
 
-							// MaterialPrim may be in another stage, so we may need another GeomSubset there
-							pxr::UsdPrim MaterialGeomSubsetPrim = GeomSubsetPrim;
-							if ( MaterialPrim.GetStage() != MeshPrim.GetStage() )
-							{
-								MaterialGeomSubsetPrim = MaterialPrim.GetStage()->OverridePrim(
-									MaterialPrim.GetPath().AppendPath( pxr::SdfPath( "Section" + std::to_string( SectionIndex ) ) )
-								);
-							}
+					// MaterialPrim may be in another stage, so we may need another GeomSubset there
+					pxr::UsdPrim MaterialGeomSubsetPrim = GeomSubsetPrim;
+					if ( MaterialPrim.GetStage() != MeshPrim.GetStage() )
+					{
+						MaterialGeomSubsetPrim = MaterialPrim.GetStage()->OverridePrim(
+							MaterialPrim.GetPath().AppendPath( pxr::SdfPath( "Section" + std::to_string( SectionIndex ) ) )
+						);
+					}
 
-							pxr::UsdGeomSubset GeomSubsetSchema{ GeomSubsetPrim };
+					pxr::UsdGeomSubset GeomSubsetSchema{ GeomSubsetPrim };
 
-							// Element type attribute
-							pxr::UsdAttribute ElementTypeAttr = GeomSubsetSchema.CreateElementTypeAttr();
-							ElementTypeAttr.Set( pxr::UsdGeomTokens->face, TimeCode );
+					// Element type attribute
+					pxr::UsdAttribute ElementTypeAttr = GeomSubsetSchema.CreateElementTypeAttr();
+					ElementTypeAttr.Set( pxr::UsdGeomTokens->face, TimeCode );
 
-							// Indices attribute
-							const uint32 TriangleCount = Section.NumTriangles;
-							const uint32 FirstTriangleIndex = Section.FirstIndex / 3; // FirstIndex is the first *vertex* instance index
-							FIndexArrayView VertexInstances = LODRenderMesh.IndexBuffer.GetArrayView();
-							pxr::VtArray<int> IndicesAttrValue;
-							for ( uint32 TriangleIndex = FirstTriangleIndex; TriangleIndex - FirstTriangleIndex < TriangleCount; ++TriangleIndex )
-							{
-								// Note that we add VertexInstances in sequence to the usda file for the faceVertexInstances attribute, which
-								// also constitutes our triangle order
-								IndicesAttrValue.push_back( static_cast< int >( TriangleIndex ) );
-							}
+					// Indices attribute
+					const uint32 TriangleCount = Section.NumTriangles;
+					const uint32 FirstTriangleIndex = Section.FirstIndex / 3; // FirstIndex is the first *vertex* instance index
+					FIndexArrayView VertexInstances = LODRenderMesh.IndexBuffer.GetArrayView();
+					pxr::VtArray<int> IndicesAttrValue;
+					for ( uint32 TriangleIndex = FirstTriangleIndex; TriangleIndex - FirstTriangleIndex < TriangleCount; ++TriangleIndex )
+					{
+						// Note that we add VertexInstances in sequence to the usda file for the faceVertexInstances attribute, which
+						// also constitutes our triangle order
+						IndicesAttrValue.push_back( static_cast< int >( TriangleIndex ) );
+					}
 
-							pxr::UsdAttribute IndicesAttr = GeomSubsetSchema.CreateIndicesAttr();
-							IndicesAttr.Set( IndicesAttrValue, TimeCode );
+					pxr::UsdAttribute IndicesAttr = GeomSubsetSchema.CreateIndicesAttr();
+					IndicesAttr.Set( IndicesAttrValue, TimeCode );
 
-							// Family name attribute
-							pxr::UsdAttribute FamilyNameAttr = GeomSubsetSchema.CreateFamilyNameAttr();
-							FamilyNameAttr.Set( pxr::UsdShadeTokens->materialBind, TimeCode );
+					// Family name attribute
+					pxr::UsdAttribute FamilyNameAttr = GeomSubsetSchema.CreateFamilyNameAttr();
+					FamilyNameAttr.Set( pxr::UsdShadeTokens->materialBind, TimeCode );
 
-							// Family type
-							pxr::UsdGeomSubset::SetFamilyType( UsdMesh, pxr::UsdShadeTokens->materialBind, pxr::UsdGeomTokens->partition );
+					// Family type
+					pxr::UsdGeomSubset::SetFamilyType( UsdMesh, pxr::UsdShadeTokens->materialBind, pxr::UsdGeomTokens->partition );
 
-							// unrealMaterial attribute
-							if ( pxr::UsdAttribute UEMaterialsAttribute = MaterialGeomSubsetPrim.CreateAttribute( UnrealIdentifiers::MaterialAssignment, pxr::SdfValueTypeNames->String ) )
-							{
-								UEMaterialsAttribute.Set( UnrealMaterialsForLOD[ SectionIndex ] );
-							}
-						}
+					// unrealMaterial attribute
+					if ( pxr::UsdAttribute UEMaterialsAttribute = MaterialGeomSubsetPrim.CreateAttribute( UnrealIdentifiers::MaterialAssignment, pxr::SdfValueTypeNames->String ) )
+					{
+						UEMaterialsAttribute.Set( UnrealMaterialsForLOD[ SectionIndex ] );
 					}
 				}
 			}
+		}
+	}
+
+	bool ConvertMeshDescription( const FMeshDescription& MeshDescription, pxr::UsdGeomMesh& UsdMesh, const FMatrix& AdditionalTransform, const pxr::UsdTimeCode TimeCode )
+	{
+		pxr::UsdPrim MeshPrim = UsdMesh.GetPrim();
+		pxr::UsdStageRefPtr Stage = MeshPrim.GetStage();
+		if ( !Stage )
+		{
+			return false;
+		}
+		const FUsdStageInfo StageInfo{ Stage };
 
-			bool ConvertMeshDescription( const FMeshDescription& MeshDescription, pxr::UsdGeomMesh& UsdMesh, const FMatrix& AdditionalTransform, const pxr::UsdTimeCode TimeCode )
+		FStaticMeshConstAttributes Attributes(MeshDescription);
+		TVertexAttributesConstRef<FVector3f> VertexPositions = Attributes.GetVertexPositions();
+		TPolygonGroupAttributesConstRef<FName> PolygonGroupImportedMaterialSlotNames = Attributes.GetPolygonGroupMaterialSlotNames();
+		TVertexInstanceAttributesConstRef<FVector3f> VertexInstanceNormals = Attributes.GetVertexInstanceNormals();
+		TVertexInstanceAttributesConstRef<FVector4f> VertexInstanceColors = Attributes.GetVertexInstanceColors();
+		TVertexInstanceAttributesConstRef<FVector2f> VertexInstanceUVs = Attributes.GetVertexInstanceUVs();
+
+		const int32 VertexCount = VertexPositions.GetNumElements();
+		const int32 VertexInstanceCount = VertexInstanceNormals.GetNumElements();
+		const int32 FaceCount = MeshDescription.Polygons().Num();
+
+		// Points
+		{
+			if ( pxr::UsdAttribute Points = UsdMesh.CreatePointsAttr() )
 			{
-				pxr::UsdPrim MeshPrim = UsdMesh.GetPrim();
-				pxr::UsdStageRefPtr Stage = MeshPrim.GetStage();
-				if ( !Stage )
+				pxr::VtArray< pxr::GfVec3f > PointsArray;
+				PointsArray.reserve( VertexCount );
+
+				for ( const FVertexID VertexID : MeshDescription.Vertices().GetElementIDs() )
 				{
-					return false;
+					FVector UEPosition = AdditionalTransform.TransformPosition( (FVector)VertexPositions[ VertexID ] );
+					PointsArray.push_back( UnrealToUsd::ConvertVector( StageInfo, UEPosition ) );
 				}
-				const FUsdStageInfo StageInfo{ Stage };
 
-				FStaticMeshConstAttributes Attributes(MeshDescription);
-				TVertexAttributesConstRef<FVector3f> VertexPositions = Attributes.GetVertexPositions();
-				TPolygonGroupAttributesConstRef<FName> PolygonGroupImportedMaterialSlotNames = Attributes.GetPolygonGroupMaterialSlotNames();
-				TVertexInstanceAttributesConstRef<FVector3f> VertexInstanceNormals = Attributes.GetVertexInstanceNormals();
-				TVertexInstanceAttributesConstRef<FVector4f> VertexInstanceColors = Attributes.GetVertexInstanceColors();
-				TVertexInstanceAttributesConstRef<FVector2f> VertexInstanceUVs = Attributes.GetVertexInstanceUVs();
+				Points.Set( PointsArray, TimeCode );
+			}
+		}
 
-				const int32 VertexCount = VertexPositions.GetNumElements();
-				const int32 VertexInstanceCount = VertexInstanceNormals.GetNumElements();
-				const int32 FaceCount = MeshDescription.Polygons().Num();
+		// Normals
+		{
+			if ( pxr::UsdAttribute NormalsAttribute = UsdMesh.CreateNormalsAttr() )
+			{
+				pxr::VtArray< pxr::GfVec3f > Normals;
+				Normals.reserve( VertexInstanceCount );
 
-				// Points
+				for ( const FVertexInstanceID InstanceID : MeshDescription.VertexInstances().GetElementIDs() )
 				{
-					if ( pxr::UsdAttribute Points = UsdMesh.CreatePointsAttr() )
-					{
-						pxr::VtArray< pxr::GfVec3f > PointsArray;
-						PointsArray.reserve( VertexCount );
-
-						for ( const FVertexID VertexID : MeshDescription.Vertices().GetElementIDs() )
-						{
-							FVector UEPosition = AdditionalTransform.TransformPosition( (FVector)VertexPositions[ VertexID ] );
-							PointsArray.push_back( UnrealToUsd::ConvertVector( StageInfo, UEPosition ) );
-						}
-
-						Points.Set( PointsArray, TimeCode );
-					}
+					FVector UENormal = (FVector)VertexInstanceNormals[ InstanceID ].GetSafeNormal();
+					Normals.push_back( UnrealToUsd::ConvertVector( StageInfo, UENormal ) );
 				}
 
-				// Normals
-				{
-					if ( pxr::UsdAttribute NormalsAttribute = UsdMesh.CreateNormalsAttr() )
-					{
-						pxr::VtArray< pxr::GfVec3f > Normals;
-						Normals.reserve( VertexInstanceCount );
+				NormalsAttribute.Set( Normals, TimeCode );
+				UsdMesh.SetNormalsInterpolation( pxr::UsdGeomTokens->faceVarying );
+			}
+		}
 
-						for ( const FVertexInstanceID InstanceID : MeshDescription.VertexInstances().GetElementIDs() )
-						{
-							FVector UENormal = (FVector)VertexInstanceNormals[ InstanceID ].GetSafeNormal();
-							Normals.push_back( UnrealToUsd::ConvertVector( StageInfo, UENormal ) );
-						}
+		// UVs
+		{
+			int32 NumUVs = VertexInstanceUVs.GetNumChannels();
 
-						NormalsAttribute.Set( Normals, TimeCode );
-						UsdMesh.SetNormalsInterpolation( pxr::UsdGeomTokens->faceVarying );
-					}
-				}
+			for ( int32 UVIndex = 0; UVIndex < NumUVs; ++UVIndex )
+			{
+				pxr::TfToken UsdUVSetName = UsdUtils::GetUVSetName( UVIndex ).Get();
 
-				// UVs
+				pxr::UsdGeomPrimvar PrimvarST = UsdMesh.CreatePrimvar( UsdUVSetName, pxr::SdfValueTypeNames->TexCoord2fArray, pxr::UsdGeomTokens->vertex );
+				if ( PrimvarST )
 				{
-					int32 NumUVs = VertexInstanceUVs.GetNumChannels();
+					pxr::VtVec2fArray UVs;
 
-					for ( int32 UVIndex = 0; UVIndex < NumUVs; ++UVIndex )
+					for ( const FVertexInstanceID InstanceID : MeshDescription.VertexInstances().GetElementIDs() )
 					{
-						pxr::TfToken UsdUVSetName = UsdUtils::GetUVSetName( UVIndex ).Get();
-
-						pxr::UsdGeomPrimvar PrimvarST = UsdMesh.CreatePrimvar( UsdUVSetName, pxr::SdfValueTypeNames->TexCoord2fArray, pxr::UsdGeomTokens->vertex );
-						if ( PrimvarST )
-						{
-							pxr::VtVec2fArray UVs;
-
-							for ( const FVertexInstanceID InstanceID : MeshDescription.VertexInstances().GetElementIDs() )
-							{
-								FVector2D UV = FVector2D(VertexInstanceUVs.Get( InstanceID, UVIndex ));
-								UV[ 1 ] = 1.f - UV[ 1 ];
-								UVs.push_back( UnrealToUsd::ConvertVector( UV ) );
-							}
-
-							PrimvarST.Set( UVs, TimeCode );
-							PrimvarST.SetInterpolation( pxr::UsdGeomTokens->faceVarying );
-						}
+						FVector2D UV = FVector2D(VertexInstanceUVs.Get( InstanceID, UVIndex ));
+						UV[ 1 ] = 1.f - UV[ 1 ];
+						UVs.push_back( UnrealToUsd::ConvertVector( UV ) );
 					}
-				}
 
-				// Vertex colors
-				if ( VertexInstanceColors.GetNumElements() > 0 )
-				{
-					pxr::UsdGeomPrimvar DisplayColorPrimvar = UsdMesh.CreateDisplayColorPrimvar( pxr::UsdGeomTokens->faceVarying );
-					pxr::UsdGeomPrimvar DisplayOpacityPrimvar = UsdMesh.CreateDisplayOpacityPrimvar( pxr::UsdGeomTokens->faceVarying );
-					if ( DisplayColorPrimvar && DisplayOpacityPrimvar )
-					{
-						pxr::VtArray< pxr::GfVec3f > DisplayColors;
-						DisplayColors.reserve( VertexInstanceCount );
+					PrimvarST.Set( UVs, TimeCode );
+					PrimvarST.SetInterpolation( pxr::UsdGeomTokens->faceVarying );
+				}
+			}
+		}
 
-						pxr::VtArray< float > DisplayOpacities;
-						DisplayOpacities.reserve( VertexInstanceCount );
+		// Vertex colors
+		if ( VertexInstanceColors.GetNumElements() > 0 )
+		{
+			pxr::UsdGeomPrimvar DisplayColorPrimvar = UsdMesh.CreateDisplayColorPrimvar( pxr::UsdGeomTokens->faceVarying );
+			pxr::UsdGeomPrimvar DisplayOpacityPrimvar = UsdMesh.CreateDisplayOpacityPrimvar( pxr::UsdGeomTokens->faceVarying );
+			if ( DisplayColorPrimvar && DisplayOpacityPrimvar )
+			{
+				pxr::VtArray< pxr::GfVec3f > DisplayColors;
+				DisplayColors.reserve( VertexInstanceCount );
 
-						for ( const FVertexInstanceID InstanceID : MeshDescription.VertexInstances().GetElementIDs() )
-						{
-							pxr::GfVec4f Color = UnrealToUsd::ConvertColor( FLinearColor( VertexInstanceColors[ InstanceID ] ) );
-							DisplayColors.push_back( pxr::GfVec3f( Color[ 0 ], Color[ 1 ], Color[ 2 ] ) );
-							DisplayOpacities.push_back( Color[ 3 ] );
-						}
+				pxr::VtArray< float > DisplayOpacities;
+				DisplayOpacities.reserve( VertexInstanceCount );
 
-						DisplayColorPrimvar.Set( DisplayColors, TimeCode );
-						DisplayOpacityPrimvar.Set( DisplayOpacities, TimeCode );
-					}
+				for ( const FVertexInstanceID InstanceID : MeshDescription.VertexInstances().GetElementIDs() )
+				{
+					pxr::GfVec4f Color = UnrealToUsd::ConvertColor( FLinearColor( VertexInstanceColors[ InstanceID ] ) );
+					DisplayColors.push_back( pxr::GfVec3f( Color[ 0 ], Color[ 1 ], Color[ 2 ] ) );
+					DisplayOpacities.push_back( Color[ 3 ] );
 				}
 
-				// Faces
-				{
-					pxr::UsdAttribute FaceCountsAttribute = UsdMesh.CreateFaceVertexCountsAttr();
-					pxr::UsdAttribute FaceVertexIndicesAttribute = UsdMesh.GetFaceVertexIndicesAttr();
+				DisplayColorPrimvar.Set( DisplayColors, TimeCode );
+				DisplayOpacityPrimvar.Set( DisplayOpacities, TimeCode );
+			}
+		}
 
-					pxr::VtArray< int > FaceVertexCounts;
-					FaceVertexCounts.reserve( FaceCount );
+		// Faces
+		{
+			pxr::UsdAttribute FaceCountsAttribute = UsdMesh.CreateFaceVertexCountsAttr();
+			pxr::UsdAttribute FaceVertexIndicesAttribute = UsdMesh.GetFaceVertexIndicesAttr();
 
-					pxr::VtArray< int > FaceVertexIndices;
+			pxr::VtArray< int > FaceVertexCounts;
+			FaceVertexCounts.reserve( FaceCount );
 
-					for ( FPolygonID PolygonID : MeshDescription.Polygons().GetElementIDs() )
-					{
-						const TArray<FVertexInstanceID>& PolygonVertexInstances = MeshDescription.GetPolygonVertexInstances( PolygonID );
-						FaceVertexCounts.push_back( static_cast< int >( PolygonVertexInstances.Num() ) );
+			pxr::VtArray< int > FaceVertexIndices;
 
-						for ( FVertexInstanceID VertexInstanceID : PolygonVertexInstances )
-						{
-							int32 VertexIndex = MeshDescription.GetVertexInstanceVertex( VertexInstanceID ).GetValue();
-							FaceVertexIndices.push_back( static_cast< int >( VertexIndex ) );
-						}
-					}
+			for ( FPolygonID PolygonID : MeshDescription.Polygons().GetElementIDs() )
+			{
+				const TArray<FVertexInstanceID>& PolygonVertexInstances = MeshDescription.GetPolygonVertexInstances( PolygonID );
+				FaceVertexCounts.push_back( static_cast< int >( PolygonVertexInstances.Num() ) );
 
-					FaceCountsAttribute.Set( FaceVertexCounts, TimeCode );
-					FaceVertexIndicesAttribute.Set( FaceVertexIndices, TimeCode );
+				for ( FVertexInstanceID VertexInstanceID : PolygonVertexInstances )
+				{
+					int32 VertexIndex = MeshDescription.GetVertexInstanceVertex( VertexInstanceID ).GetValue();
+					FaceVertexIndices.push_back( static_cast< int >( VertexIndex ) );
 				}
-
-				return true;
 			}
 
-			bool RecursivelyCollapseChildMeshes(
-				const pxr::UsdPrim& Prim,
-				const FTransform& CurrentTransform,
-				const pxr::UsdTimeCode& TimeCode,
-				const EUsdPurpose PurposesToLoad,
-				const pxr::TfToken& RenderContext,
-				const TMap< FString, TMap<FString, int32> >& MaterialToPrimvarToUVIndex,
-				FMeshDescription& OutMeshDescription,
-				UsdUtils::FUsdPrimMaterialAssignmentInfo& OutMaterialAssignments,
-				bool bIsFirstPrim,
-				const bool bCombineIdenticalMaterialSlots
-			)
-			{
-				// Ignore meshes from disabled purposes
-				if ( !EnumHasAllFlags( PurposesToLoad, IUsdPrim::GetPurpose( Prim ) ) )
-				{
-					return true;
-				}
+			FaceCountsAttribute.Set( FaceVertexCounts, TimeCode );
+			FaceVertexIndicesAttribute.Set( FaceVertexIndices, TimeCode );
+		}
+
+		return true;
+	}
 
-				FTransform ChildTransform = CurrentTransform;
+	bool RecursivelyCollapseChildMeshes(
+		const pxr::UsdPrim& Prim,
+		FMeshDescription& OutMeshDescription,
+		UsdUtils::FUsdPrimMaterialAssignmentInfo& OutMaterialAssignments,
+		UsdToUnreal::FUsdMeshConversionOptions& Options,
+		bool bIsFirstPrim
+	)
+	{
+		// Ignore meshes from disabled purposes
+		if ( !EnumHasAllFlags( Options.PurposesToLoad, IUsdPrim::GetPurpose( Prim ) ) )
+		{
+			return true;
+		}
 
-				if ( !bIsFirstPrim )
+		FTransform ChildTransform = Options.AdditionalTransform;
+
+		if ( !bIsFirstPrim )
+		{
+			// Ignore invisible child meshes.
+			//
+			// We used to compute visibility here and flat out ignore any invisible meshes. However, it could be that this mesh is invisible
+			// due to the first prim (the parentmost prim of the recursive calls) being invisible. If the first is invisible but animated
+			// then its possible it will become visible later, so if the child meshes are all invisible due to that fact alone then we should still
+			// consider them. If the first is invisible but *not* animated then we should still consider it in the same way, because that's sort
+			// of what you'd expect by calling ConvertGeomMeshHierarchy: We shouldn't just return nothing if the prim happens to be invisible.
+			// Besides, it could be that first is invisible due to itself having a parent that is invisible but has visibility animations:
+			// In that case we'd also want to generate meshes even if first is effectively invisible, since those parents can become visible
+			// later as well. The only case left is if first is invisible due having parents that are invisible and not animated: Checking for
+			// this would involve checking visibility and animations of all of its parents though, which is probably a bit too much, and like
+			// in the case where first itself is invisible and not animated, the caller may still expect to receive a valid mesh even if the
+			// prim's parents are invisible.
+			//
+			// The only case in which we'll truly discard invisible submeshes now is if they're invisible *by themselves*. If we're collapsing
+			// them then we know they're not animated either, so they will basically never be visible at all, at any time code.
+			//
+			// Note that if we were to ever manually set any of these back to visible again via the editor, the visibility changes are
+			// now resyncs and we'll reparse this entire asset, which will give us the chance to add them back to the collapsed mesh.
+			if ( pxr::UsdGeomImageable UsdGeomImageable = pxr::UsdGeomImageable( Prim ) )
+			{
+				if ( pxr::UsdAttribute VisibilityAttr = UsdGeomImageable.GetVisibilityAttr() )
 				{
-					// Ignore invisible child meshes.
-					//
-					// We used to compute visibility here and flat out ignore any invisible meshes. However, it could be that this mesh is invisible
-					// due to the first prim (the parentmost prim of the recursive calls) being invisible. If the first is invisible but animated
-					// then its possible it will become visible later, so if the child meshes are all invisible due to that fact alone then we should still
-					// consider them. If the first is invisible but *not* animated then we should still consider it in the same way, because that's sort
-					// of what you'd expect by calling ConvertGeomMeshHierarchy: We shouldn't just return nothing if the prim happens to be invisible.
-					// Besides, it could be that first is invisible due to itself having a parent that is invisible but has visibility animations:
-					// In that case we'd also want to generate meshes even if first is effectively invisible, since those parents can become visible
-					// later as well. The only case left is if first is invisible due having parents that are invisible and not animated: Checking for
-					// this would involve checking visibility and animations of all of its parents though, which is probably a bit too much, and like
-					// in the case where first itself is invisible and not animated, the caller may still expect to receive a valid mesh even if the
-					// prim's parents are invisible.
-					//
-					// The only case in which we'll truly discard invisible submeshes now is if they're invisible *by themselves*. If we're collapsing
-					// them then we know they're not animated either, so they will basically never be visible at all, at any time code.
-					//
-					// Note that if we were to ever manually set any of these back to visible again via the editor, the visibility changes are
-					// now resyncs and we'll reparse this entire asset, which will give us the chance to add them back to the collapsed mesh.
-					if ( pxr::UsdGeomImageable UsdGeomImageable = pxr::UsdGeomImageable( Prim ) )
+					pxr::TfToken VisibilityToken;
+					if ( VisibilityAttr.Get( &VisibilityToken ) && VisibilityToken == pxr::UsdGeomTokens->invisible )
 					{
-						if ( pxr::UsdAttribute VisibilityAttr = UsdGeomImageable.GetVisibilityAttr() )
-						{
-							pxr::TfToken VisibilityToken;
-							if ( VisibilityAttr.Get( &VisibilityToken ) && VisibilityToken == pxr::UsdGeomTokens->invisible )
-							{
-								return true;
-							}
-						}
+						return true;
 					}
+				}
+			}
 
-					if ( pxr::UsdGeomXformable Xformable = pxr::UsdGeomXformable( Prim ) )
-					{
-						FTransform LocalChildTransform;
-						UsdToUnreal::ConvertXformable( Prim.GetStage(), Xformable, LocalChildTransform, TimeCode.GetValue() );
+			if ( pxr::UsdGeomXformable Xformable = pxr::UsdGeomXformable( Prim ) )
+			{
+				FTransform LocalChildTransform;
+				UsdToUnreal::ConvertXformable( Prim.GetStage(), Xformable, LocalChildTransform, Options.TimeCode.GetValue() );
 
-						ChildTransform = LocalChildTransform * CurrentTransform;
-					}
-				}
+				ChildTransform = LocalChildTransform * Options.AdditionalTransform;
+			}
+		}
 
-				bool bSuccess = true;
-				bool bTraverseChildren = true;
+		bool bSuccess = true;
+		bool bTraverseChildren = true;
 
-				if ( pxr::UsdGeomMesh Mesh = pxr::UsdGeomMesh( Prim ) )
-				{
-					bSuccess = UsdToUnreal::ConvertGeomMesh( Mesh, OutMeshDescription, OutMaterialAssignments, ChildTransform, MaterialToPrimvarToUVIndex, TimeCode, RenderContext, bCombineIdenticalMaterialSlots );
-				}
-				else if ( pxr::UsdGeomPointInstancer PointInstancer = pxr::UsdGeomPointInstancer{ Prim } )
-				{
-					bSuccess = UsdToUnreal::ConvertPointInstancerToMesh(
-						PointInstancer,
-						ChildTransform,
-						MaterialToPrimvarToUVIndex,
-						PurposesToLoad,
-						OutMeshDescription,
-						OutMaterialAssignments,
-						TimeCode,
-						RenderContext,
-						bCombineIdenticalMaterialSlots
-					);
+		// Since ConvertGeomMesh and ConvertPointInstancerToMesh take the Options object by const ref and we traverse
+		// children afterwards, its fine to overwrite Options.AdditionalTransform. We do have to put it back to our
+		// original value after we're done though, as calls to sibling prims that would run after this call would need
+		// the original AdditionalTransform in place. The alternative is to copy the entire options object...
+		TGuardValue<FTransform> Guard{ Options.AdditionalTransform, ChildTransform };
 
-					// We never want to step into point instancers when fetching prims for drawing
-					bTraverseChildren = false;
-				}
+		if ( pxr::UsdGeomMesh Mesh = pxr::UsdGeomMesh( Prim ) )
+		{
+			bSuccess = UsdToUnreal::ConvertGeomMesh( Mesh, OutMeshDescription, OutMaterialAssignments, Options );
+		}
+		else if ( pxr::UsdGeomPointInstancer PointInstancer = pxr::UsdGeomPointInstancer{ Prim } )
+		{
+			bSuccess = UsdToUnreal::ConvertPointInstancerToMesh(
+				PointInstancer,
+				OutMeshDescription,
+				OutMaterialAssignments,
+				Options
+			);
 
-				if ( bTraverseChildren )
-				{
-					for ( const pxr::UsdPrim& ChildPrim : Prim.GetFilteredChildren( pxr::UsdTraverseInstanceProxies() ) )
-					{
-						if ( !bSuccess )
-						{
-							break;
-						}
+			// We never want to step into point instancers when fetching prims for drawing
+			bTraverseChildren = false;
+		}
 
-						const bool bChildIsFirstPrim = false;
-
-						bSuccess &= RecursivelyCollapseChildMeshes(
-							ChildPrim,
-							ChildTransform,
-							TimeCode,
-							PurposesToLoad,
-							RenderContext,
-							MaterialToPrimvarToUVIndex,
-							OutMeshDescription,
-							OutMaterialAssignments,
-							bChildIsFirstPrim,
-							bCombineIdenticalMaterialSlots
-						);
-					}
+		if ( bTraverseChildren )
+		{
+			for ( const pxr::UsdPrim& ChildPrim : Prim.GetFilteredChildren( pxr::UsdTraverseInstanceProxies() ) )
+			{
+				if ( !bSuccess )
+				{
+					break;
 				}
 
-				return bSuccess;
+				const bool bChildIsFirstPrim = false;
+
+				bSuccess &= RecursivelyCollapseChildMeshes(
+					ChildPrim,
+					OutMeshDescription,
+					OutMaterialAssignments,
+					Options,
+					bChildIsFirstPrim
+				);
 			}
 		}
+
+		return bSuccess;
 	}
 }
 namespace UsdGeomMeshImpl = UE::UsdGeomMeshConversion::Private;
 
-bool UsdToUnreal::ConvertGeomMesh( const pxr::UsdTyped& UsdSchema, FMeshDescription& MeshDescription, const pxr::UsdTimeCode TimeCode )
+namespace UsdToUnreal
 {
-	UsdUtils::FUsdPrimMaterialAssignmentInfo MaterialAssignments;
-	return ConvertGeomMesh( UsdSchema, MeshDescription, MaterialAssignments, TimeCode );
-}
-
-bool UsdToUnreal::ConvertGeomMesh( const pxr::UsdTyped& UsdSchema, FMeshDescription& MeshDescription, const FTransform& AdditionalTransform, const pxr::UsdTimeCode TimeCode )
-{
-	UsdUtils::FUsdPrimMaterialAssignmentInfo MaterialAssignments;
-	return ConvertGeomMesh( UsdSchema, MeshDescription, MaterialAssignments, AdditionalTransform, TimeCode );
-}
-
-bool UsdToUnreal::ConvertGeomMesh( const pxr::UsdTyped& UsdSchema, FMeshDescription& MeshDescription, const FTransform& AdditionalTransform, const TMap< FString, TMap< FString, int32 > >& MaterialToPrimvarsUVSetNames, const pxr::UsdTimeCode TimeCode )
-{
-	UsdUtils::FUsdPrimMaterialAssignmentInfo MaterialAssignments;
-	return ConvertGeomMesh( UsdSchema, MeshDescription, MaterialAssignments, AdditionalTransform, MaterialToPrimvarsUVSetNames, TimeCode );
-}
-
-bool UsdToUnreal::ConvertGeomMesh( const pxr::UsdTyped& UsdSchema, FMeshDescription& MeshDescription, UsdUtils::FUsdPrimMaterialAssignmentInfo& MaterialAssignments,
-	const pxr::UsdTimeCode TimeCode, const pxr::TfToken& RenderContext )
-{
-	TMap< FString, TMap< FString, int32 > > MaterialToPrimvarsUVSetNames;
-	return ConvertGeomMesh( UsdSchema, MeshDescription, MaterialAssignments, FTransform::Identity, MaterialToPrimvarsUVSetNames, TimeCode, RenderContext );
-}
+	const FUsdMeshConversionOptions FUsdMeshConversionOptions::DefaultOptions;
 
-bool UsdToUnreal::ConvertGeomMesh( const pxr::UsdTyped& UsdSchema, FMeshDescription& MeshDescription, UsdUtils::FUsdPrimMaterialAssignmentInfo& MaterialAssignments, const FTransform& AdditionalTransform,
-	const pxr::UsdTimeCode TimeCode, const pxr::TfToken& RenderContext )
-{
-	TMap< FString, TMap< FString, int32 > > MaterialToPrimvarsUVSetNames;
-	return ConvertGeomMesh( UsdSchema, MeshDescription, MaterialAssignments, AdditionalTransform, MaterialToPrimvarsUVSetNames, TimeCode, RenderContext );
+	FUsdMeshConversionOptions::FUsdMeshConversionOptions()
+		: AdditionalTransform( FTransform::Identity )
+		, PurposesToLoad( EUsdPurpose::Render )
+		, RenderContext( pxr::UsdShadeTokens->universalRenderContext )
+		, TimeCode( pxr::UsdTimeCode::EarliestTime() )
+		, MaterialToPrimvarToUVIndex( nullptr )
+		, bMergeIdenticalMaterialSlots( true )
+	{
+	}
 }
 
-bool UsdToUnreal::ConvertGeomMesh( const pxr::UsdTyped& UsdSchema, FMeshDescription& MeshDescription, UsdUtils::FUsdPrimMaterialAssignmentInfo& MaterialAssignments, const FTransform& AdditionalTransform,
-	const TMap< FString, TMap< FString, int32 > >& MaterialToPrimvarsUVSetNames, const pxr::UsdTimeCode TimeCode, const pxr::TfToken& RenderContext, bool bCombineIdenticalMaterialSlots )
+bool UsdToUnreal::ConvertGeomMesh(
+	const pxr::UsdGeomMesh& UsdMesh,
+	FMeshDescription& OutMeshDescription,
+	UsdUtils::FUsdPrimMaterialAssignmentInfo& OutMaterialAssignments,
+	const FUsdMeshConversionOptions& Options
+)
 {
 	TRACE_CPUPROFILER_EVENT_SCOPE( UsdToUnreal::ConvertGeomMesh );
 
-	using namespace pxr;
-
-	FScopedUsdAllocs UsdAllocs;
-
-	pxr::UsdGeomMesh UsdMesh( UsdSchema );
-	pxr::UsdPrim UsdPrim = UsdMesh.GetPrim();
-	if ( !UsdMesh || !UsdPrim )
+	if ( !UsdMesh )
 	{
 		return false;
 	}
 
-	pxr::UsdStageRefPtr Stage = UsdMesh.GetPrim().GetStage();
+	FScopedUsdAllocs Allocs;
+
+	pxr::UsdPrim UsdPrim = UsdMesh.GetPrim();
+	pxr::UsdStageRefPtr Stage = UsdPrim.GetStage();
 	const FUsdStageInfo StageInfo( Stage );
 
-	const double TimeCodeValue = TimeCode.GetValue();
+	const double TimeCodeValue = Options.TimeCode.GetValue();
 
-	const int32 MaterialIndexOffset = MaterialAssignments.Slots.Num();
+	const int32 MaterialIndexOffset = OutMaterialAssignments.Slots.Num();
 
 	// Material assignments
 	const bool bProvideMaterialIndices = true;
-	UsdUtils::FUsdPrimMaterialAssignmentInfo LocalInfo = UsdUtils::GetPrimMaterialAssignments( UsdPrim, TimeCode, bProvideMaterialIndices, RenderContext );
+	UsdUtils::FUsdPrimMaterialAssignmentInfo LocalInfo = UsdUtils::GetPrimMaterialAssignments(
+		UsdPrim,
+		Options.TimeCode,
+		bProvideMaterialIndices,
+		Options.RenderContext
+	);
 	TArray< UsdUtils::FUsdPrimMaterialSlot >& LocalMaterialSlots = LocalInfo.Slots;
 	TArray< int32 >& FaceMaterialIndices = LocalInfo.MaterialIndices;
 
@@ -689,13 +663,13 @@ bool UsdToUnreal::ConvertGeomMesh( const pxr::UsdTyped& UsdSchema, FMeshDescript
 	TArray<int32> LocalToCombinedMaterialSlotIndices;
 	LocalToCombinedMaterialSlotIndices.SetNumZeroed( LocalInfo.Slots.Num() );
 
-	if ( bCombineIdenticalMaterialSlots )
+	if ( Options.bMergeIdenticalMaterialSlots )
 	{
 		// Build a map of our existing slots since we can hash the entire slot, and our incoming mesh may have an arbitrary number of new slots
 		TMap< UsdUtils::FUsdPrimMaterialSlot, int32 > CombinedMaterialSlotsToIndex;
-		for ( int32 Index = 0; Index < MaterialAssignments.Slots.Num(); ++Index )
+		for ( int32 Index = 0; Index < OutMaterialAssignments.Slots.Num(); ++Index )
 		{
-			const UsdUtils::FUsdPrimMaterialSlot& Slot = MaterialAssignments.Slots[ Index ];
+			const UsdUtils::FUsdPrimMaterialSlot& Slot = OutMaterialAssignments.Slots[ Index ];
 			CombinedMaterialSlotsToIndex.Add( Slot, Index );
 		}
 
@@ -708,26 +682,26 @@ bool UsdToUnreal::ConvertGeomMesh( const pxr::UsdTyped& UsdSchema, FMeshDescript
 			}
 			else
 			{
-				MaterialAssignments.Slots.Add( LocalSlot );
-				LocalToCombinedMaterialSlotIndices[ LocalIndex ] = MaterialAssignments.Slots.Num() - 1;
+				OutMaterialAssignments.Slots.Add( LocalSlot );
+				LocalToCombinedMaterialSlotIndices[ LocalIndex ] = OutMaterialAssignments.Slots.Num() - 1;
 			}
 		}
 	}
 	else
 	{
 		// Just append our new local material slots at the end of MaterialAssignments
-		MaterialAssignments.Slots.Append( LocalInfo.Slots );
+		OutMaterialAssignments.Slots.Append( LocalInfo.Slots );
 		for ( int32 LocalIndex = 0; LocalIndex < LocalInfo.Slots.Num(); ++LocalIndex )
 		{
 			LocalToCombinedMaterialSlotIndices[ LocalIndex ] = LocalIndex + MaterialIndexOffset;
 		}
 	}
 
-	const int32 VertexOffset = MeshDescription.Vertices().Num();
-	const int32 VertexInstanceOffset = MeshDescription.VertexInstances().Num();
-	const int32 PolygonOffset = MeshDescription.Polygons().Num();
+	const int32 VertexOffset = OutMeshDescription.Vertices().Num();
+	const int32 VertexInstanceOffset = OutMeshDescription.VertexInstances().Num();
+	const int32 PolygonOffset = OutMeshDescription.Polygons().Num();
 
-	FStaticMeshAttributes StaticMeshAttributes( MeshDescription );
+	FStaticMeshAttributes StaticMeshAttributes( OutMeshDescription );
 
 	// Vertex positions
 	TVertexAttributesRef< FVector3f > MeshDescriptionVertexPositions = StaticMeshAttributes.GetVertexPositions();
@@ -738,16 +712,16 @@ bool UsdToUnreal::ConvertGeomMesh( const pxr::UsdTyped& UsdSchema, FMeshDescript
 			pxr::VtArray< pxr::GfVec3f > PointsArray;
 			Points.Get( &PointsArray, TimeCodeValue );
 
-			MeshDescription.ReserveNewVertices( PointsArray.size() );
+			OutMeshDescription.ReserveNewVertices( PointsArray.size() );
 
 			for ( int32 LocalPointIndex = 0; LocalPointIndex < PointsArray.size(); ++LocalPointIndex )
 			{
-				const GfVec3f& Point = PointsArray[ LocalPointIndex ];
+				const pxr::GfVec3f& Point = PointsArray[ LocalPointIndex ];
 
-				FVector Position = AdditionalTransform.TransformPosition( UsdToUnreal::ConvertVector( StageInfo, Point ) );
+				FVector Position = Options.AdditionalTransform.TransformPosition( UsdToUnreal::ConvertVector( StageInfo, Point ) );
 
-				FVertexID AddedVertexId = MeshDescription.CreateVertex();
-				MeshDescriptionVertexPositions[ AddedVertexId ] = (FVector3f)Position;
+				FVertexID AddedVertexId = OutMeshDescription.CreateVertex();
+				MeshDescriptionVertexPositions[ AddedVertexId ] = ( FVector3f ) Position;
 			}
 		}
 	}
@@ -765,9 +739,9 @@ bool UsdToUnreal::ConvertGeomMesh( const pxr::UsdTyped& UsdSchema, FMeshDescript
 
 		// If we're going to share existing material slots, we'll need to share existing PolygonGroups in the mesh description,
 		// so we need to traverse it and prefill PolygonGroupMapping
-		if( bCombineIdenticalMaterialSlots )
+		if ( Options.bMergeIdenticalMaterialSlots )
 		{
-			for ( FPolygonGroupID PolygonGroupID : MeshDescription.PolygonGroups().GetElementIDs() )
+			for ( FPolygonGroupID PolygonGroupID : OutMeshDescription.PolygonGroups().GetElementIDs() )
 			{
 				// We always create our polygon groups in order with our combined material slots, so its easy to reconstruct this mapping here
 				PolygonGroupMapping.Add( PolygonGroupID.GetValue(), PolygonGroupID );
@@ -783,7 +757,7 @@ bool UsdToUnreal::ConvertGeomMesh( const pxr::UsdTyped& UsdSchema, FMeshDescript
 			const int32 CombinedMaterialIndex = LocalToCombinedMaterialSlotIndices[ LocalMaterialIndex ];
 			if ( !PolygonGroupMapping.Contains( CombinedMaterialIndex ) )
 			{
-				FPolygonGroupID NewPolygonGroup = MeshDescription.CreatePolygonGroup();
+				FPolygonGroupID NewPolygonGroup = OutMeshDescription.CreatePolygonGroup();
 				PolygonGroupMapping.Add( CombinedMaterialIndex, NewPolygonGroup );
 
 				// This is important for runtime, where the material slots are matched to LOD sections based on their material slot name
@@ -799,8 +773,8 @@ bool UsdToUnreal::ConvertGeomMesh( const pxr::UsdTyped& UsdSchema, FMeshDescript
 		}
 
 		// Face counts
-		UsdAttribute FaceCountsAttribute = UsdMesh.GetFaceVertexCountsAttr();
-		VtArray< int > FaceCounts;
+		pxr::UsdAttribute FaceCountsAttribute = UsdMesh.GetFaceVertexCountsAttr();
+		pxr::VtArray< int > FaceCounts;
 
 		if ( FaceCountsAttribute )
 		{
@@ -809,8 +783,8 @@ bool UsdToUnreal::ConvertGeomMesh( const pxr::UsdTyped& UsdSchema, FMeshDescript
 		}
 
 		// Face indices
-		UsdAttribute FaceIndicesAttribute = UsdMesh.GetFaceVertexIndicesAttr();
-		VtArray< int > FaceIndices;
+		pxr::UsdAttribute FaceIndicesAttribute = UsdMesh.GetFaceVertexIndicesAttr();
+		pxr::VtArray< int > FaceIndices;
 
 		if ( FaceIndicesAttribute )
 		{
@@ -818,8 +792,8 @@ bool UsdToUnreal::ConvertGeomMesh( const pxr::UsdTyped& UsdSchema, FMeshDescript
 		}
 
 		// Normals
-		UsdAttribute NormalsAttribute = UsdMesh.GetNormalsAttr();
-		VtArray< GfVec3f > Normals;
+		pxr::UsdAttribute NormalsAttribute = UsdMesh.GetNormalsAttr();
+		pxr::VtArray< pxr::GfVec3f > Normals;
 
 		if ( NormalsAttribute )
 		{
@@ -834,15 +808,15 @@ bool UsdToUnreal::ConvertGeomMesh( const pxr::UsdTyped& UsdSchema, FMeshDescript
 		struct FUVSet
 		{
 			int32 UVSetIndexUE; // The user may only have 'uv4' and 'uv5', so we can't just use array indices to find the target UV channel
-			TOptional< VtIntArray > UVIndices; // UVs might be indexed or they might be flat (one per vertex)
-			VtVec2fArray UVs;
+			TOptional< pxr::VtIntArray > UVIndices; // UVs might be indexed or they might be flat (one per vertex)
+			pxr::VtVec2fArray UVs;
 
-			pxr::TfToken InterpType = UsdGeomTokens->faceVarying;
+			pxr::TfToken InterpType = pxr::UsdGeomTokens->faceVarying;
 		};
 
 		TArray< FUVSet > UVSets;
 
-		TArray< TUsdStore< UsdGeomPrimvar > > PrimvarsByUVIndex = UsdUtils::GetUVSetPrimvars( UsdMesh, MaterialToPrimvarsUVSetNames, LocalInfo );
+		TArray< TUsdStore< pxr::UsdGeomPrimvar > > PrimvarsByUVIndex = UsdUtils::GetUVSetPrimvars( UsdMesh, *Options.MaterialToPrimvarToUVIndex, LocalInfo );
 
 		int32 HighestAddedUVChannel = 0;
 		for ( int32 UVChannelIndex = 0; UVChannelIndex < PrimvarsByUVIndex.Num(); ++UVChannelIndex )
@@ -852,7 +826,7 @@ bool UsdToUnreal::ConvertGeomMesh( const pxr::UsdTyped& UsdSchema, FMeshDescript
 				break;
 			}
 
-			UsdGeomPrimvar& Primvar = PrimvarsByUVIndex[UVChannelIndex].Get();
+			pxr::UsdGeomPrimvar& Primvar = PrimvarsByUVIndex[ UVChannelIndex ].Get();
 			if ( !Primvar )
 			{
 				// The user may have name their UV sets 'uv4' and 'uv5', in which case we have no UV sets below 4, so just skip them
@@ -867,7 +841,7 @@ bool UsdToUnreal::ConvertGeomMesh( const pxr::UsdTyped& UsdSchema, FMeshDescript
 			{
 				UVSet.UVIndices.Emplace();
 
-				if ( Primvar.GetIndices( &UVSet.UVIndices.GetValue(), TimeCode ) && Primvar.Get( &UVSet.UVs, TimeCode ) )
+				if ( Primvar.GetIndices( &UVSet.UVIndices.GetValue(), Options.TimeCode ) && Primvar.Get( &UVSet.UVs, Options.TimeCode ) )
 				{
 					if ( UVSet.UVs.size() > 0 )
 					{
@@ -892,37 +866,37 @@ bool UsdToUnreal::ConvertGeomMesh( const pxr::UsdTyped& UsdSchema, FMeshDescript
 		// When importing multiple mesh pieces to the same static mesh.  Ensure each mesh piece has the same number of Uv's
 		{
 			int32 ExistingUVCount = MeshDescriptionUVs.GetNumChannels();
-			int32 NumUVs = FMath::Max( HighestAddedUVChannel + 1, ExistingUVCount);
-			NumUVs = FMath::Min<int32>(MAX_MESH_TEXTURE_COORDS_MD, NumUVs);
+			int32 NumUVs = FMath::Max( HighestAddedUVChannel + 1, ExistingUVCount );
+			NumUVs = FMath::Min<int32>( MAX_MESH_TEXTURE_COORDS_MD, NumUVs );
 			// At least one UV set must exist.
-			NumUVs = FMath::Max<int32>(1, NumUVs);
+			NumUVs = FMath::Max<int32>( 1, NumUVs );
 
 			//Make sure all Vertex instance have the correct number of UVs
-			MeshDescriptionUVs.SetNumChannels(NumUVs);
+			MeshDescriptionUVs.SetNumChannels( NumUVs );
 		}
 
 		TVertexInstanceAttributesRef< FVector3f > MeshDescriptionNormals = StaticMeshAttributes.GetVertexInstanceNormals();
 
-		MeshDescription.ReserveNewVertexInstances( FaceCounts.size() * 3 );
-		MeshDescription.ReserveNewPolygons( FaceCounts.size() );
-		MeshDescription.ReserveNewEdges( FaceCounts.size() * 2 );
+		OutMeshDescription.ReserveNewVertexInstances( FaceCounts.size() * 3 );
+		OutMeshDescription.ReserveNewPolygons( FaceCounts.size() );
+		OutMeshDescription.ReserveNewEdges( FaceCounts.size() * 2 );
 
 		// Vertex color
 		TVertexInstanceAttributesRef< FVector4f > MeshDescriptionColors = StaticMeshAttributes.GetVertexInstanceColors();
 
-		UsdGeomPrimvar ColorPrimvar = UsdMesh.GetDisplayColorPrimvar();
-		pxr::TfToken ColorInterpolation = UsdGeomTokens->constant;
+		pxr::UsdGeomPrimvar ColorPrimvar = UsdMesh.GetDisplayColorPrimvar();
+		pxr::TfToken ColorInterpolation = pxr::UsdGeomTokens->constant;
 		pxr::VtArray< pxr::GfVec3f > UsdColors;
 
 		if ( ColorPrimvar )
 		{
-			ColorPrimvar.ComputeFlattened( &UsdColors, TimeCode );
+			ColorPrimvar.ComputeFlattened( &UsdColors, Options.TimeCode );
 			ColorInterpolation = ColorPrimvar.GetInterpolation();
 		}
 
 		// Vertex opacity
-		UsdGeomPrimvar OpacityPrimvar = UsdMesh.GetDisplayOpacityPrimvar();
-		pxr::TfToken OpacityInterpolation = UsdGeomTokens->constant;
+		pxr::UsdGeomPrimvar OpacityPrimvar = UsdMesh.GetDisplayOpacityPrimvar();
+		pxr::TfToken OpacityInterpolation = pxr::UsdGeomTokens->constant;
 		pxr::VtArray< float > UsdOpacities;
 
 		if ( OpacityPrimvar )
@@ -933,17 +907,17 @@ bool UsdToUnreal::ConvertGeomMesh( const pxr::UsdTyped& UsdSchema, FMeshDescript
 
 		for ( int32 PolygonIndex = 0; PolygonIndex < FaceCounts.size(); ++PolygonIndex )
 		{
-			int32 PolygonVertexCount = FaceCounts[PolygonIndex];
+			int32 PolygonVertexCount = FaceCounts[ PolygonIndex ];
 			CornerInstanceIDs.Reset( PolygonVertexCount );
 			CornerVerticesIDs.Reset( PolygonVertexCount );
 
-			for (int32 CornerIndex = 0; CornerIndex < PolygonVertexCount; ++CornerIndex, ++CurrentVertexInstanceIndex)
+			for ( int32 CornerIndex = 0; CornerIndex < PolygonVertexCount; ++CornerIndex, ++CurrentVertexInstanceIndex )
 			{
 				int32 VertexInstanceIndex = VertexInstanceOffset + CurrentVertexInstanceIndex;
-				const FVertexInstanceID VertexInstanceID(VertexInstanceIndex);
-				const int32 ControlPointIndex = FaceIndices[CurrentVertexInstanceIndex];
-				const FVertexID VertexID(VertexOffset + ControlPointIndex);
-				const FVector VertexPosition = (FVector)MeshDescriptionVertexPositions[VertexID];
+				const FVertexInstanceID VertexInstanceID( VertexInstanceIndex );
+				const int32 ControlPointIndex = FaceIndices[ CurrentVertexInstanceIndex ];
+				const FVertexID VertexID( VertexOffset + ControlPointIndex );
+				const FVector VertexPosition = ( FVector ) MeshDescriptionVertexPositions[ VertexID ];
 
 				// Make sure a face doesn't use the same vertex twice as MeshDescription doesn't like that
 				if ( CornerVerticesIDs.Contains( VertexID ) )
@@ -953,7 +927,7 @@ bool UsdToUnreal::ConvertGeomMesh( const pxr::UsdTyped& UsdSchema, FMeshDescript
 
 				CornerVerticesIDs.Add( VertexID );
 
-				FVertexInstanceID AddedVertexInstanceId = MeshDescription.CreateVertexInstance(VertexID);
+				FVertexInstanceID AddedVertexInstanceId = OutMeshDescription.CreateVertexInstance( VertexID );
 				CornerInstanceIDs.Add( AddedVertexInstanceId );
 
 				if ( Normals.size() > 0 )
@@ -962,10 +936,10 @@ bool UsdToUnreal::ConvertGeomMesh( const pxr::UsdTyped& UsdSchema, FMeshDescript
 
 					if ( NormalIndex < Normals.size() )
 					{
-						const GfVec3f& Normal = Normals[NormalIndex];
-						FVector TransformedNormal = AdditionalTransform.TransformVector( UsdToUnreal::ConvertVector( StageInfo, Normal ) ).GetSafeNormal();
+						const pxr::GfVec3f& Normal = Normals[ NormalIndex ];
+						FVector TransformedNormal = Options.AdditionalTransform.TransformVector( UsdToUnreal::ConvertVector( StageInfo, Normal ) ).GetSafeNormal();
 
-						MeshDescriptionNormals[AddedVertexInstanceId] = (FVector3f)TransformedNormal.GetSafeNormal();
+						MeshDescriptionNormals[ AddedVertexInstanceId ] = ( FVector3f ) TransformedNormal.GetSafeNormal();
 					}
 				}
 
@@ -973,7 +947,7 @@ bool UsdToUnreal::ConvertGeomMesh( const pxr::UsdTyped& UsdSchema, FMeshDescript
 				{
 					const int32 ValueIndex = UsdGeomMeshImpl::GetPrimValueIndex( UVSet.InterpType, ControlPointIndex, CurrentVertexInstanceIndex, PolygonIndex );
 
-					GfVec2f UV( 0.f, 0.f );
+					pxr::GfVec2f UV( 0.f, 0.f );
 
 					if ( UVSet.UVIndices.IsSet() )
 					{
@@ -988,7 +962,7 @@ bool UsdToUnreal::ConvertGeomMesh( const pxr::UsdTyped& UsdSchema, FMeshDescript
 					}
 
 					// Flip V for Unreal uv's which match directx
-					FVector2f FinalUVVector( UV[0], 1.f - UV[1] );
+					FVector2f FinalUVVector( UV[ 0 ], 1.f - UV[ 1 ] );
 					MeshDescriptionUVs.Set( AddedVertexInstanceId, UVSet.UVSetIndexUE, FinalUVVector );
 				}
 
@@ -996,7 +970,7 @@ bool UsdToUnreal::ConvertGeomMesh( const pxr::UsdTyped& UsdSchema, FMeshDescript
 				{
 					const int32 ValueIndex = UsdGeomMeshImpl::GetPrimValueIndex( ColorInterpolation, ControlPointIndex, CurrentVertexInstanceIndex, PolygonIndex );
 
-					GfVec3f UsdColor( 1.f, 1.f, 1.f );
+					pxr::GfVec3f UsdColor( 1.f, 1.f, 1.f );
 
 					if ( !UsdColors.empty() && ensure( UsdColors.size() > ValueIndex ) )
 					{
@@ -1012,7 +986,7 @@ bool UsdToUnreal::ConvertGeomMesh( const pxr::UsdTyped& UsdSchema, FMeshDescript
 
 					if ( !UsdOpacities.empty() && ensure( UsdOpacities.size() > ValueIndex ) )
 					{
-						MeshDescriptionColors[ AddedVertexInstanceId ][3] = UsdOpacities[ ValueIndex ];
+						MeshDescriptionColors[ AddedVertexInstanceId ][ 3 ] = UsdOpacities[ ValueIndex ];
 					}
 				}
 			}
@@ -1040,14 +1014,14 @@ bool UsdToUnreal::ConvertGeomMesh( const pxr::UsdTyped& UsdSchema, FMeshDescript
 
 			if ( bFlipThisGeometry )
 			{
-				for (int32 i = 0; i < CornerInstanceIDs.Num() / 2; ++i)
+				for ( int32 i = 0; i < CornerInstanceIDs.Num() / 2; ++i )
 				{
-					Swap(CornerInstanceIDs[i], CornerInstanceIDs[CornerInstanceIDs.Num() - i - 1]);
+					Swap( CornerInstanceIDs[ i ], CornerInstanceIDs[ CornerInstanceIDs.Num() - i - 1 ] );
 				}
 			}
 			// Insert a polygon into the mesh
 			FPolygonGroupID PolygonGroupID = PolygonGroupMapping[ CombinedMaterialIndex ];
-			const FPolygonID NewPolygonID = MeshDescription.CreatePolygon( PolygonGroupID, CornerInstanceIDs );
+			const FPolygonID NewPolygonID = OutMeshDescription.CreatePolygon( PolygonGroupID, CornerInstanceIDs );
 		}
 	}
 
@@ -1063,27 +1037,78 @@ bool UsdToUnreal::ConvertGeomMesh( const pxr::UsdTyped& UsdSchema, FMeshDescript
 	return true;
 }
 
-bool UsdToUnreal::ConvertPointInstancerToMesh(
-	const pxr::UsdGeomPointInstancer& PointInstancer,
+// Deprecated
+bool UsdToUnreal::ConvertGeomMesh(
+	const pxr::UsdTyped& UsdSchema,
+	FMeshDescription& MeshDescription,
+	UsdUtils::FUsdPrimMaterialAssignmentInfo& MaterialAssignments,
+	const pxr::UsdTimeCode TimeCode,
+	const pxr::TfToken& RenderContext
+)
+{
+	FUsdMeshConversionOptions Options;
+	Options.TimeCode = TimeCode;
+	Options.RenderContext = RenderContext;
+
+	return ConvertGeomMesh( pxr::UsdGeomMesh{ UsdSchema }, MeshDescription, MaterialAssignments, Options );
+}
+
+// Deprecated
+bool UsdToUnreal::ConvertGeomMesh(
+	const pxr::UsdTyped& UsdSchema,
+	FMeshDescription& MeshDescription,
+	UsdUtils::FUsdPrimMaterialAssignmentInfo& MaterialAssignments,
 	const FTransform& AdditionalTransform,
-	const TMap< FString, TMap< FString, int32 > >& MaterialToPrimvarsUVSetNames,
-	const EUsdPurpose PurposesToLoad,
+	const pxr::UsdTimeCode TimeCode,
+	const pxr::TfToken& RenderContext
+)
+{
+	FUsdMeshConversionOptions Options;
+	Options.AdditionalTransform = AdditionalTransform;
+	Options.TimeCode = TimeCode;
+	Options.RenderContext = RenderContext;
+
+	return ConvertGeomMesh( pxr::UsdGeomMesh{ UsdSchema }, MeshDescription, MaterialAssignments, Options );
+}
+
+// Deprecated
+bool UsdToUnreal::ConvertGeomMesh(
+	const pxr::UsdTyped& UsdSchema,
 	FMeshDescription& MeshDescription,
 	UsdUtils::FUsdPrimMaterialAssignmentInfo& MaterialAssignments,
+	const FTransform& AdditionalTransform,
+	const TMap< FString, TMap< FString, int32 > >& MaterialToPrimvarsUVSetNames,
 	const pxr::UsdTimeCode TimeCode,
 	const pxr::TfToken& RenderContext,
-	bool bCombineIdenticalMaterialSlots
+	bool bMergeIdenticalMaterialSlots
 )
 {
-	TRACE_CPUPROFILER_EVENT_SCOPE( UsdToUnreal::ConvertPointInstancerToMesh );
+	FUsdMeshConversionOptions Options;
+	Options.AdditionalTransform = AdditionalTransform;
+	Options.MaterialToPrimvarToUVIndex = &MaterialToPrimvarsUVSetNames;
+	Options.TimeCode = TimeCode;
+	Options.RenderContext = RenderContext;
+	Options.bMergeIdenticalMaterialSlots = bMergeIdenticalMaterialSlots;
+
+	return ConvertGeomMesh( pxr::UsdGeomMesh{ UsdSchema }, MeshDescription, MaterialAssignments, Options );
+}
 
-	FScopedUsdAllocs Allocs;
+bool UsdToUnreal::ConvertPointInstancerToMesh(
+	const pxr::UsdGeomPointInstancer& PointInstancer,
+	FMeshDescription& OutMeshDescription,
+	UsdUtils::FUsdPrimMaterialAssignmentInfo& OutMaterialAssignments,
+	const FUsdMeshConversionOptions& Options
+)
+{
+	TRACE_CPUPROFILER_EVENT_SCOPE( UsdToUnreal::ConvertPointInstancerToMesh );
 
 	if ( !PointInstancer )
 	{
 		return false;
 	}
 
+	FScopedUsdAllocs Allocs;
+
 	pxr::UsdStageRefPtr Stage = PointInstancer.GetPrim().GetStage();
 	if ( !Stage )
 	{
@@ -1114,6 +1139,11 @@ bool UsdToUnreal::ConvertPointInstancerToMesh(
 		PrototypeMaterialAssignments.SetNum( NumPrototypes );
 		PrototypePolygonGroupRemapping.SetNum( NumPrototypes );
 
+		// Our AdditionalTransform should be applied after even the instance transforms, we don't want to apply it
+		// directly to our prototypes
+		FUsdMeshConversionOptions OptionsCopy = Options;
+		OptionsCopy.AdditionalTransform = FTransform::Identity;
+
 		for ( uint32 PrototypeIndex = 0; PrototypeIndex < NumPrototypes; ++PrototypeIndex )
 		{
 			const pxr::SdfPath& PrototypePath = PrototypePaths[ PrototypeIndex ];
@@ -1131,14 +1161,10 @@ bool UsdToUnreal::ConvertPointInstancerToMesh(
 			const bool bSkipRootPrimTransformAndVisibility = false;
 			ConvertGeomMeshHierarchy(
 				PrototypeUsdPrim,
-				TimeCode,
-				PurposesToLoad,
-				RenderContext,
-				MaterialToPrimvarsUVSetNames,
 				PrototypeMeshDescriptions[ PrototypeIndex ],
 				PrototypeMaterialAssignments[ PrototypeIndex ],
-				bSkipRootPrimTransformAndVisibility,
-				bCombineIdenticalMaterialSlots
+				OptionsCopy,
+				bSkipRootPrimTransformAndVisibility
 			);
 		}
 	}
@@ -1152,13 +1178,13 @@ bool UsdToUnreal::ConvertPointInstancerToMesh(
 		UsdUtils::FUsdPrimMaterialAssignmentInfo& PrototypeMaterialAssignment = PrototypeMaterialAssignments[ PrototypeIndex ];
 		TMap<FPolygonGroupID, FPolygonGroupID>& PrototypeToCombinedMeshPolygonGroupMap = PrototypePolygonGroupRemapping[ PrototypeIndex ];
 
-		if ( bCombineIdenticalMaterialSlots )
+		if ( Options.bMergeIdenticalMaterialSlots )
 		{
 			// Build a map of our existing slots since we can hash the entire slot, and our incoming mesh may have an arbitrary number of new slots
 			TMap< UsdUtils::FUsdPrimMaterialSlot, int32 > CombinedMaterialSlotsToIndex;
-			for ( int32 Index = 0; Index < MaterialAssignments.Slots.Num(); ++Index )
+			for ( int32 Index = 0; Index < OutMaterialAssignments.Slots.Num(); ++Index )
 			{
-				const UsdUtils::FUsdPrimMaterialSlot& Slot = MaterialAssignments.Slots[ Index ];
+				const UsdUtils::FUsdPrimMaterialSlot& Slot = OutMaterialAssignments.Slots[ Index ];
 				CombinedMaterialSlotsToIndex.Add( Slot, Index );
 			}
 
@@ -1171,15 +1197,15 @@ bool UsdToUnreal::ConvertPointInstancerToMesh(
 				}
 				else
 				{
-					MaterialAssignments.Slots.Add( LocalSlot );
-					PrototypeToCombinedMeshPolygonGroupMap.Add( PrototypeMaterialSlotIndex, MaterialAssignments.Slots.Num() - 1 );
+					OutMaterialAssignments.Slots.Add( LocalSlot );
+					PrototypeToCombinedMeshPolygonGroupMap.Add( PrototypeMaterialSlotIndex, OutMaterialAssignments.Slots.Num() - 1 );
 				}
 			}
 		}
 		else
 		{
-			const int32 NumExistingMaterialSlots = MaterialAssignments.Slots.Num();
-			MaterialAssignments.Slots.Append( PrototypeMaterialAssignment.Slots );
+			const int32 NumExistingMaterialSlots = OutMaterialAssignments.Slots.Num();
+			OutMaterialAssignments.Slots.Append( PrototypeMaterialAssignment.Slots );
 
 			for ( int32 PrototypeMaterialSlotIndex = 0; PrototypeMaterialSlotIndex < PrototypeMaterialAssignment.Slots.Num(); ++PrototypeMaterialSlotIndex )
 			{
@@ -1190,14 +1216,14 @@ bool UsdToUnreal::ConvertPointInstancerToMesh(
 
 	// Make sure we have the polygon groups we expect. Appending the mesh descriptions will not create new polygon groups if we're using a
 	// PolygonGroupsDelegate, which we will
-	const int32 NumExistingPolygonGroups = MeshDescription.PolygonGroups().Num();
-	for ( int32 NumMissingPolygonGroups = MaterialAssignments.Slots.Num() - NumExistingPolygonGroups; NumMissingPolygonGroups > 0; --NumMissingPolygonGroups )
+	const int32 NumExistingPolygonGroups = OutMeshDescription.PolygonGroups().Num();
+	for ( int32 NumMissingPolygonGroups = OutMaterialAssignments.Slots.Num() - NumExistingPolygonGroups; NumMissingPolygonGroups > 0; --NumMissingPolygonGroups )
 	{
-		MeshDescription.CreatePolygonGroup();
+		OutMeshDescription.CreatePolygonGroup();
 	}
 
 	// Double-check our target mesh description has the attributes we need
-	FStaticMeshAttributes StaticMeshAttributes( MeshDescription );
+	FStaticMeshAttributes StaticMeshAttributes( OutMeshDescription );
 	StaticMeshAttributes.Register();
 
 	// Append mesh descriptions
@@ -1214,7 +1240,7 @@ bool UsdToUnreal::ConvertPointInstancerToMesh(
 		}
 
 		TArray<FTransform> InstanceTransforms;
-		bool bSuccess = UsdUtils::GetPointInstancerTransforms( StageInfo, PointInstancer, PrototypeIndex, TimeCode, InstanceTransforms );
+		bool bSuccess = UsdUtils::GetPointInstancerTransforms( StageInfo, PointInstancer, PrototypeIndex, Options.TimeCode, InstanceTransforms );
 		if ( !bSuccess )
 		{
 			UE_LOG( LogUsd, Error, TEXT( "Failed to retrieve point instancer transforms for prototype index '%u' of point instancer '%s'" ),
@@ -1227,10 +1253,10 @@ bool UsdToUnreal::ConvertPointInstancerToMesh(
 
 		const int32 NumInstances = InstanceTransforms.Num();
 
-		MeshDescription.ReserveNewVertices( PrototypeMeshDescription.Vertices().Num() * NumInstances );
-		MeshDescription.ReserveNewVertexInstances( PrototypeMeshDescription.VertexInstances().Num() * NumInstances );
-		MeshDescription.ReserveNewEdges( PrototypeMeshDescription.Edges().Num() * NumInstances );
-		MeshDescription.ReserveNewTriangles( PrototypeMeshDescription.Triangles().Num() * NumInstances );
+		OutMeshDescription.ReserveNewVertices( PrototypeMeshDescription.Vertices().Num() * NumInstances );
+		OutMeshDescription.ReserveNewVertexInstances( PrototypeMeshDescription.VertexInstances().Num() * NumInstances );
+		OutMeshDescription.ReserveNewEdges( PrototypeMeshDescription.Edges().Num() * NumInstances );
+		OutMeshDescription.ReserveNewTriangles( PrototypeMeshDescription.Triangles().Num() * NumInstances );
 
 		FStaticMeshOperations::FAppendSettings Settings;
 		Settings.PolygonGroupsDelegate = FAppendPolygonGroupsDelegate::CreateLambda(
@@ -1244,8 +1270,8 @@ bool UsdToUnreal::ConvertPointInstancerToMesh(
 		// transform repeatedly?
 		for ( const FTransform& Transform : InstanceTransforms )
 		{
-			Settings.MeshTransform = Transform * AdditionalTransform;
-			FStaticMeshOperations::AppendMeshDescription( PrototypeMeshDescription, MeshDescription, Settings );
+			Settings.MeshTransform = Transform * Options.AdditionalTransform;
+			FStaticMeshOperations::AppendMeshDescription( PrototypeMeshDescription, OutMeshDescription, Settings );
 		}
 	}
 
@@ -1254,32 +1280,61 @@ bool UsdToUnreal::ConvertPointInstancerToMesh(
 
 bool UsdToUnreal::ConvertGeomMeshHierarchy(
 	const pxr::UsdPrim& Prim,
-	const pxr::UsdTimeCode& TimeCode,
-	const EUsdPurpose PurposesToLoad,
-	const pxr::TfToken& RenderContext,
-	const TMap< FString, TMap<FString, int32> >& MaterialToPrimvarToUVIndex,
 	FMeshDescription& OutMeshDescription,
 	UsdUtils::FUsdPrimMaterialAssignmentInfo& OutMaterialAssignments,
-	bool bSkipRootPrimTransformAndVisibility,
-	bool bCombineIdenticalMaterialSlots
+	const FUsdMeshConversionOptions& Options,
+	bool bSkipRootPrimTransformAndVisibility
 )
 {
 	TRACE_CPUPROFILER_EVENT_SCOPE( UsdToUnreal::ConvertGeomMeshHierarchy );
 
+	if ( !Prim )
+	{
+		return false;
+	}
+
 	FStaticMeshAttributes StaticMeshAttributes( OutMeshDescription );
 	StaticMeshAttributes.Register();
 
-	return UE::UsdGeomMeshConversion::Private::RecursivelyCollapseChildMeshes(
+	// Pass a copy down so that we can repeatedly overwrite the AdditionalTransform and still
+	// provide the options object to ConvertGeomMesh and ConvertPointInstancerToMesh
+	FUsdMeshConversionOptions OptionsCopy = Options;
+
+	return UsdGeomMeshImpl::RecursivelyCollapseChildMeshes(
+		Prim,
+		OutMeshDescription,
+		OutMaterialAssignments,
+		OptionsCopy,
+		bSkipRootPrimTransformAndVisibility
+	);
+}
+
+// Deprecated
+bool UsdToUnreal::ConvertGeomMeshHierarchy(
+	const pxr::UsdPrim& Prim,
+	const pxr::UsdTimeCode& TimeCode,
+	const EUsdPurpose PurposesToLoad,
+	const pxr::TfToken& RenderContext,
+	const TMap< FString, TMap<FString, int32> >& MaterialToPrimvarToUVIndex,
+	FMeshDescription& OutMeshDescription,
+	UsdUtils::FUsdPrimMaterialAssignmentInfo& OutMaterialAssignments,
+	bool bSkipRootPrimTransformAndVisibility,
+	bool bMergeIdenticalMaterialSlots
+)
+{
+	FUsdMeshConversionOptions Options;
+	Options.TimeCode = TimeCode;
+	Options.PurposesToLoad = PurposesToLoad;
+	Options.RenderContext = RenderContext;
+	Options.MaterialToPrimvarToUVIndex = &MaterialToPrimvarToUVIndex;
+	Options.bMergeIdenticalMaterialSlots = bMergeIdenticalMaterialSlots;
+
+	return ConvertGeomMeshHierarchy(
 		Prim,
-		FTransform::Identity,
-		TimeCode,
-		PurposesToLoad,
-		RenderContext,
-		MaterialToPrimvarToUVIndex,
 		OutMeshDescription,
 		OutMaterialAssignments,
-		bSkipRootPrimTransformAndVisibility,
-		bCombineIdenticalMaterialSlots
+		Options,
+		bSkipRootPrimTransformAndVisibility
 	);
 }
 
diff --git a/Engine/Plugins/Importers/USDImporter/Source/USDUtilities/Public/USDGeomMeshConversion.h b/Engine/Plugins/Importers/USDImporter/Source/USDUtilities/Public/USDGeomMeshConversion.h
--- a/Engine/Plugins/Importers/USDImporter/Source/USDUtilities/Public/USDGeomMeshConversion.h
+++ b/Engine/Plugins/Importers/USDImporter/Source/USDUtilities/Public/USDGeomMeshConversion.h
@@ -28,6 +28,7 @@ class UMaterialInstanceDynamic;
 class UStaticMesh;
 struct FMeshDescription;
 struct FStaticMeshLODResources;
+struct FUsdCollapsingCache;
 struct FUsdStageInfo;
 namespace UsdUtils
 {
@@ -41,72 +42,133 @@ namespace UE
 
 namespace UsdToUnreal
 {
+	/** Common options to mesh conversion functions */
+	struct USDUTILITIES_API FUsdMeshConversionOptions
+	{
+		const static FUsdMeshConversionOptions DefaultOptions;
+		FUsdMeshConversionOptions();
+
+		// Additional UE-space transform to bake into the mesh's vertices
+		FTransform AdditionalTransform;
+
+		// Only prims with these purposes will be traversed
+		EUsdPurpose PurposesToLoad;
+
+		// Render context to use when parsing Mesh materials
+		pxr::TfToken RenderContext;
+
+		// USD timecode when the prims should be sampled for mesh data and converted
+		pxr::UsdTimeCode TimeCode;
+
+		// Maps from a material prim path, to pairs indicating which primvar names are used as 'st' coordinates for
+		// this mesh, and which UVIndex materials will sample from( e.g.[ "st0", 0 ], [ "myUvSet2", 2 ], etc ).
+		// This is used to pick which primvars will become UV sets.
+		const TMap< FString, TMap<FString, int32> >* MaterialToPrimvarToUVIndex;
+
+		// Whether to try reusing material slots (both local and the ones already in MaterialAssignments) when
+		// converting the material assignments from the mesh
+		bool bMergeIdenticalMaterialSlots;
+	};
+
 	/**
-	 * Extracts mesh data from UsdSchema at TimeCode and places the results in MeshDescription.
+	 * Converts an UsdGeomMesh and appends the converted data to OutMeshDescription and OutMaterialAssignments.
 	 *
-	 * @param UsdSchema - UsdGeomMesh that contains the data to convert
-	 * @param MeshDescription - Output parameter that will be filled with the converted data
-	 * @param MaterialAssignments - Output parameter that will be filled with the material data extracted from UsdSchema
-	 * @param AdditionalTransform - Transform (in UnrealEditor space) to bake into the vertex positions and normals of the converted mesh
-	 * @param MaterialToPrimvarsUVSetNames - Maps from a material prim path, to pairs indicating which primvar names are used as 'st' coordinates for this mesh, and which UVIndex materials will sample from (e.g. ["st0", 0], ["myUvSet2", 2], etc). This is used to pick which primvars will become UV sets.
-	 * @param TimeCode - USD timecode when the UsdGeomMesh should be sampled for mesh data and converted
-	 * @param RenderContext - The render context to use when parsing material assignments from the Mesh prim
-	 * @param bCombineIdenticalMaterialSlots - Whether to try reusing material slots (both local and the ones already in MaterialAssignments) when converting the material assignments from the mesh
+	 * @param Mesh - Mesh prim to convert
+	 * @param OutMeshDescription - Output parameter that will be filled with the converted data
+	 * @param OutMaterialAssignments - Output parameter that will be filled with the material data extracted from Mesh
+	 * @param CommonOptions - Shared options used for the conversion
 	 * @return Whether the conversion was successful or not.
 	 */
-	USDUTILITIES_API bool ConvertGeomMesh( const pxr::UsdTyped& UsdSchema, FMeshDescription& MeshDescription, UsdUtils::FUsdPrimMaterialAssignmentInfo& MaterialAssignments, const pxr::UsdTimeCode TimeCode = pxr::UsdTimeCode::EarliestTime(), const pxr::TfToken& RenderContext = pxr::UsdShadeTokens->universalRenderContext );
-	USDUTILITIES_API bool ConvertGeomMesh( const pxr::UsdTyped& UsdSchema, FMeshDescription& MeshDescription, UsdUtils::FUsdPrimMaterialAssignmentInfo& MaterialAssignments, const FTransform& AdditionalTransform, const pxr::UsdTimeCode TimeCode = pxr::UsdTimeCode::EarliestTime(), const pxr::TfToken& RenderContext = pxr::UsdShadeTokens->universalRenderContext );
-	USDUTILITIES_API bool ConvertGeomMesh( const pxr::UsdTyped& UsdSchema, FMeshDescription& MeshDescription, UsdUtils::FUsdPrimMaterialAssignmentInfo& MaterialAssignments, const FTransform& AdditionalTransform, const TMap< FString, TMap< FString, int32 > >& MaterialToPrimvarsUVSetNames, const pxr::UsdTimeCode TimeCode = pxr::UsdTimeCode::EarliestTime(), const pxr::TfToken& RenderContext = pxr::UsdShadeTokens->universalRenderContext, bool bCombineIdenticalMaterialSlots = true );
-
-	/** DEPRECATED and will not convert material information. Use the three first signatures above this one */
-	USDUTILITIES_API bool ConvertGeomMesh( const pxr::UsdTyped& UsdSchema, FMeshDescription& MeshDescription, const pxr::UsdTimeCode TimeCode = pxr::UsdTimeCode::EarliestTime() );
-	USDUTILITIES_API bool ConvertGeomMesh( const pxr::UsdTyped& UsdSchema, FMeshDescription& MeshDescription, const FTransform& AdditionalTransform, const pxr::UsdTimeCode TimeCode = pxr::UsdTimeCode::EarliestTime() );
-	USDUTILITIES_API bool ConvertGeomMesh( const pxr::UsdTyped& UsdSchema, FMeshDescription& MeshDescription, const FTransform& AdditionalTransform, const TMap< FString, TMap< FString, int32 > >& MaterialToPrimvarsUVSetNames, const pxr::UsdTimeCode TimeCode = pxr::UsdTimeCode::EarliestTime() );
+	USDUTILITIES_API bool ConvertGeomMesh(
+		const pxr::UsdGeomMesh& Mesh,
+		FMeshDescription& OutMeshDescription,
+		UsdUtils::FUsdPrimMaterialAssignmentInfo& OutMaterialAssignments,
+		const FUsdMeshConversionOptions& CommonOptions = FUsdMeshConversionOptions::DefaultOptions
+	);
+	UE_DEPRECATED( 5.1, "Please use the overload that receives FUsdMeshConversionOptions." )
+	USDUTILITIES_API bool ConvertGeomMesh(
+		const pxr::UsdTyped& UsdSchema,
+		FMeshDescription& MeshDescription,
+		UsdUtils::FUsdPrimMaterialAssignmentInfo& MaterialAssignments,
+		const pxr::UsdTimeCode TimeCode = pxr::UsdTimeCode::EarliestTime(),
+		const pxr::TfToken& RenderContext = pxr::UsdShadeTokens->universalRenderContext
+	);
+	UE_DEPRECATED( 5.1, "Please use the overload that receives FUsdMeshConversionOptions." )
+	USDUTILITIES_API bool ConvertGeomMesh(
+		const pxr::UsdTyped& UsdSchema,
+		FMeshDescription& MeshDescription,
+		UsdUtils::FUsdPrimMaterialAssignmentInfo& MaterialAssignments,
+		const FTransform& AdditionalTransform,
+		const pxr::UsdTimeCode TimeCode = pxr::UsdTimeCode::EarliestTime(),
+		const pxr::TfToken& RenderContext = pxr::UsdShadeTokens->universalRenderContext
+	);
+	UE_DEPRECATED( 5.1, "Please use the overload that receives FUsdMeshConversionOptions." )
+	USDUTILITIES_API bool ConvertGeomMesh(
+		const pxr::UsdTyped& UsdSchema,
+		FMeshDescription& MeshDescription,
+		UsdUtils::FUsdPrimMaterialAssignmentInfo& MaterialAssignments,
+		const FTransform& AdditionalTransform,
+		const TMap< FString, TMap< FString, int32 > >& MaterialToPrimvarsUVSetNames,
+		const pxr::UsdTimeCode TimeCode = pxr::UsdTimeCode::EarliestTime(),
+		const pxr::TfToken& RenderContext = pxr::UsdShadeTokens->universalRenderContext,
+		bool bMergeIdenticalMaterialSlots = true
+	);
 
-	/** Bake a point instancer prim into a single FMeshDescription and MaterialAssignments, manually "instancing" all the prototype instances into a single combined mesh.
+	/** Bake a point instancer prim into a single FMeshDescription and MaterialAssignments, manually "instancing" all
+	 *  the prototype instances into a single combined mesh.
 	 *
 	 * @param PointInstancer - UsdGeomPointInstancer to bake
-	 * @param AdditionalTransform - Transform (in UnrealEditor space) to bake into the vertex positions and normals of the converted mesh
-	 * @param MaterialToPrimvarsUVSetNames - Maps from a material prim path, to pairs indicating which primvar names are used as 'st' coordinates for this mesh, and which UVIndex materials will sample from (e.g. ["st0", 0], ["myUvSet2", 2], etc). This is used to pick which primvars will become UV sets.
-	 * @param PurposesToLoad - Display purposes to consider when baking the prototypes into meshes
 	 * @param OutMeshDescription - Output parameter that will be filled with the converted data
-	 * @param OutMaterialAssignments - Output parameter that will be filled with the material data extracted from UsdSchema
-	 * @param TimeCode - USD timecode when the UsdGeomPointInstancer should be sampled for mesh data and converted
-	 * @param RenderContext - The render context to use when parsing material assignments from the prototype meshes
-	 * @param bCombineIdenticalMaterialSlots - Whether to try reusing material slots (both local and the ones already in MaterialAssignments) when converting the material assignments from the prototype meshes
+	 * @param OutMaterialAssignments - Output parameter that will be filled with the material data extracted from
+	 *								   PointInstancer
+	 * @param CommonOptions - Shared options used for the conversion
 	 * @return Whether the conversion was successful or not.
 	 */
 	USDUTILITIES_API bool ConvertPointInstancerToMesh(
 		const pxr::UsdGeomPointInstancer& PointInstancer,
-		const FTransform& AdditionalTransform,
-		const TMap< FString, TMap< FString, int32 > >& MaterialToPrimvarsUVSetNames,
-		const EUsdPurpose PurposesToLoad,
 		FMeshDescription& OutMeshDescription,
 		UsdUtils::FUsdPrimMaterialAssignmentInfo& OutMaterialAssignments,
-		const pxr::UsdTimeCode TimeCode = pxr::UsdTimeCode::EarliestTime(),
-		const pxr::TfToken& RenderContext = pxr::UsdShadeTokens->universalRenderContext,
-		bool bCombineIdenticalMaterialSlots = true
+		const FUsdMeshConversionOptions& CommonOptions = FUsdMeshConversionOptions::DefaultOptions
 	);
 
 	/**
-	 * Recursively traverses down Prim's subtree parsing all Mesh data (including Prim's), baking all transforms and putting the resulting
-	 * mesh data within OutMeshDescription and OutMaterialAssignments. It will not bake Prim's own transform into the Mesh data.
+	 * Recursively traverses down Prim's subtree parsing all Mesh data (including Prim's), baking all transforms and
+	 * putting the resulting mesh data within OutMeshDescription and OutMaterialAssignments. It will not bake Prim's
+	 * own transform into the Mesh data.
 	 *
 	 * @param Prim - Prim that will be converted (result will include its mesh and its children's)
-	 * @param TimeCode - USD timecode when the UsdGeomMeshes should be sampled for mesh data and converted
-	 * @param PurposesToLoad - Only prims with these purposes will be traversed
-	 * @param RenderContext - Render context to use when parsing Mesh materials
-	 * @param MaterialToPrimvarToUVIndex - Maps from a material prim path, to pairs indicating which primvar names are used as 'st' coordinates for this mesh, and which UVIndex materials will sample from (e.g. ["st0", 0], ["myUvSet2", 2], etc). This is used to pick which primvars will become UV sets.
-	 * @param OutMeshDescription - Output parameter that will be filled with the converted mesh data
-	 * @param OutMaterialAssignments - Output parameter that will be filled with the material data extracted from UsdSchema
-	 * @param bCombineIdenticalMaterialSlots - Whether to try reusing material slots (both local and the ones already in MaterialAssignments) when converting the material assignments from the mesh
-	 * @param bSkipRootPrimTransformAndVisibility - Whether to ignore Prim's transform and visibility when generating the combined mesh. This is useful if you plan on displaying that data on a component with the resulting mesh instead.
+	 * @param OutMeshDescription - Output parameter that will be filled with the converted mesh descriptions
+	 * @param OutMaterialAssignments - Output parameter that will be filled with the material data extracted from
+									   UsdSchema. Each item of the array matches the mesh description of same index
+	 * @param CommonOptions - Shared options used for the conversion
 	 * @return Whether the conversion was successful or not.
 	 */
-	USDUTILITIES_API bool ConvertGeomMeshHierarchy( const pxr::UsdPrim& Prim, const pxr::UsdTimeCode& TimeCode, const EUsdPurpose PurposesToLoad, const pxr::TfToken& RenderContext, const TMap< FString, TMap<FString, int32> >& MaterialToPrimvarToUVIndex, FMeshDescription& OutMeshDescription, UsdUtils::FUsdPrimMaterialAssignmentInfo& OutMaterialAssignments, bool bSkipRootPrimTransformAndVisibility = false, bool bCombineIdenticalMaterialSlots = true );
+	USDUTILITIES_API bool ConvertGeomMeshHierarchy(
+		const pxr::UsdPrim& Prim,
+		FMeshDescription& OutMeshDescription,
+		UsdUtils::FUsdPrimMaterialAssignmentInfo& OutMaterialAssignments,
+		const FUsdMeshConversionOptions& CommonOptions = FUsdMeshConversionOptions::DefaultOptions,
+		bool bSkipRootPrimTransformAndVisibility = false
+	);
+
+	UE_DEPRECATED( 5.1, "Please use the overload that receives FUsdMeshConversionOptions." )
+	USDUTILITIES_API bool ConvertGeomMeshHierarchy(
+		const pxr::UsdPrim& Prim,
+		const pxr::UsdTimeCode& TimeCode,
+		const EUsdPurpose PurposesToLoad,
+		const pxr::TfToken& RenderContext,
+		const TMap< FString, TMap<FString, int32> >& MaterialToPrimvarToUVIndex,
+		FMeshDescription& OutMeshDescription,
+		UsdUtils::FUsdPrimMaterialAssignmentInfo& OutMaterialAssignments,
+		bool bSkipRootPrimTransformAndVisibility = false,
+		bool bMergeIdenticalMaterialSlots = true
+	);
 
 	/** Configure Material to become a vertex color/displayColor material, according to the given description */
-	USDUTILITIES_API bool ConvertDisplayColor( const UsdUtils::FDisplayColorMaterial& DisplayColorDescription, UMaterialInstanceConstant& Material );
+	USDUTILITIES_API bool ConvertDisplayColor(
+		const UsdUtils::FDisplayColorMaterial& DisplayColorDescription,
+		UMaterialInstanceConstant& Material
+	);
 }
 
 namespace UnrealToUsd
