USD: Add cvar to limit number of point instancer instances that are parsed.

diff --git a/Engine/Plugins/Importers/USDImporter/Source/USDUtilities/Private/USDGeomMeshConversion.cpp b/Engine/Plugins/Importers/USDImporter/Source/USDUtilities/Private/USDGeomMeshConversion.cpp
--- a/Engine/Plugins/Importers/USDImporter/Source/USDUtilities/Private/USDGeomMeshConversion.cpp
+++ b/Engine/Plugins/Importers/USDImporter/Source/USDUtilities/Private/USDGeomMeshConversion.cpp
@@ -55,6 +55,12 @@
 
 #define LOCTEXT_NAMESPACE "USDGeomMeshConversion"
 
+static int32 GMaxInstancesPerPointInstancer = -1;
+static FAutoConsoleVariableRef CVarMaxInstancesPerPointInstancer(
+	TEXT( "USD.MaxInstancesPerPointInstancer" ),
+	GMaxInstancesPerPointInstancer,
+	TEXT( "We will only parse up to this many instances from any point instancer when reading from USD to UE. Set this to -1 to disable this limit." ) );
+
 namespace UE::UsdGeomMeshConversion::Private
 {
 	static const FString DisplayColorID = TEXT( "!DisplayColor" );
@@ -2798,10 +2804,19 @@ bool UsdUtils::GetPointInstancerTransforms( const FUsdStageInfo& StageInfo, cons
 
 	FScopedUnrealAllocs UnrealAllocs;
 
-	OutInstanceTransforms.Reset( UsdInstanceTransforms.size() );
+	const int32 NumInstances = GMaxInstancesPerPointInstancer >= 0
+		? FMath::Min( static_cast< int32 >( UsdInstanceTransforms.size() ), GMaxInstancesPerPointInstancer )
+		: static_cast< int32 >( UsdInstanceTransforms.size() );
+
+	OutInstanceTransforms.Reset( NumInstances );
 
 	for ( pxr::GfMatrix4d& UsdMatrix : UsdInstanceTransforms )
 	{
+		if ( Index == NumInstances )
+		{
+			break;
+		}
+
 		if ( ProtoIndices[ Index ] == ProtoIndex )
 		{
 			OutInstanceTransforms.Add( UsdToUnreal::ConvertMatrix( StageInfo, UsdMatrix ) );
