USD: Refactor collapsing cache to also cache subtree vertex counts and material slots for other features.

Modify a couple of the point instancer test scenes to add material slots that don't merge together for testing purposes;
Update test_point_instancer_scenes screenshots to account for the modified material slots of the previous item;
Add a couple of testing library functions so we can query the collapsing cache from Python during tests.

diff --git a/Engine/Plugins/Importers/USDImporter/Source/USDSchemas/Private/USDCollapsingCache.cpp b/Engine/Plugins/Importers/USDImporter/Source/USDSchemas/Private/USDCollapsingCache.cpp
deleted file mode 100644
diff --git a/Engine/Plugins/Importers/USDImporter/Source/USDSchemas/Private/USDGeomXformableTranslator.cpp b/Engine/Plugins/Importers/USDImporter/Source/USDSchemas/Private/USDGeomXformableTranslator.cpp
--- a/Engine/Plugins/Importers/USDImporter/Source/USDSchemas/Private/USDGeomXformableTranslator.cpp
+++ b/Engine/Plugins/Importers/USDImporter/Source/USDSchemas/Private/USDGeomXformableTranslator.cpp
@@ -455,9 +455,9 @@ bool FUsdGeomXformableTranslator::CollapsesChildren( ECollapsingType CollapsingT
 {
 	TRACE_CPUPROFILER_EVENT_SCOPE( FUsdGeomXformableTranslator::CollapsesChildren );
 
-	if ( Context->CollapsingCache.IsValid() )
+	if ( Context->InfoCache.IsValid() )
 	{
-		return Context->CollapsingCache->DoesPathCollapseChildren( PrimPath, CollapsingType );
+		return Context->InfoCache->DoesPathCollapseChildren( PrimPath, CollapsingType );
 	}
 
 	bool bCollapsesChildren = false;
@@ -482,10 +482,10 @@ bool FUsdGeomXformableTranslator::CollapsesChildren( ECollapsingType CollapsingT
 		{
 			IUsdSchemasModule& UsdSchemasModule = FModuleManager::Get().LoadModuleChecked< IUsdSchemasModule >( TEXT("USDSchemas") );
 
-			// TODO: This can be optimized in order to make FUsdCollapsingCache::RebuildCacheForSubtree faster: If we have a child prim that we know doesn't collapse,
+			// TODO: This can be optimized in order to make FUsdInfoCache::RebuildCacheForSubtree faster: If we have a child prim that we know doesn't collapse,
 			// any of our parents should be able to know they can't collapse *us* either.
 			// This is somewhat niche though: Realistically to waste time here a prim and its children need to have a kind that allows collapsing, and also not be able to collapse.
-			// Also, if any of these prims *does* manage to collapse, FUsdCollapsingCache will already not actually query the subtree children if they can collapse or not anymore,
+			// Also, if any of these prims *does* manage to collapse, FUsdInfoCache will already not actually query the subtree children if they can collapse or not anymore,
 			// and just consider them collapsed by the parent
 			TArray< TUsdStore< pxr::UsdPrim > > ChildXformPrims = UsdUtils::GetAllPrimsOfType( Prim, pxr::TfType::Find< pxr::UsdGeomXformable >() );
 			for ( const TUsdStore< pxr::UsdPrim >& ChildXformPrim : ChildXformPrims )
@@ -546,36 +546,17 @@ bool FUsdGeomXformableTranslator::CollapsesChildren( ECollapsingType CollapsingT
 		{
 			bCollapsesChildren = false;
 		}
-		else
+		else if ( Context->InfoCache.IsValid() )
 		{
-			pxr::TfToken RenderContextToken = pxr::UsdShadeTokens->universalRenderContext;
-			if ( !Context->RenderContext.IsNone() )
+			TOptional<uint64> NumExpectedVertices = Context->InfoCache->GetSubtreeVertexCount( PrimPath );
+			if ( !NumExpectedVertices.IsSet() || NumExpectedVertices.GetValue() > GMaxNumVerticesCollapsedMesh )
 			{
-				RenderContextToken = UnrealToUsd::ConvertToken( *Context->RenderContext.ToString() ).Get();
+				bCollapsesChildren = false;
 			}
 
-			int32 NumVertices = 0;
-
-			TSet<UsdUtils::FUsdPrimMaterialSlot> CombinedSlots;
-			int32 NumMaxExpectedMaterialSlots = 0;
-
-			for ( const TUsdStore< pxr::UsdPrim >& ChildPrim : ChildGeomMeshes )
+			if ( bChildrenWantNanite )
 			{
-				pxr::UsdGeomMesh ChildGeomMesh( ChildPrim.Get() );
-
-				if ( pxr::UsdAttribute Points = ChildGeomMesh.GetPointsAttr() )
-				{
-					pxr::VtArray< pxr::GfVec3f > PointsArray;
-					Points.Get( &PointsArray, pxr::UsdTimeCode( Context->Time ) );
-
-					NumVertices += PointsArray.size();
-
-					if ( NumVertices > GMaxNumVerticesCollapsedMesh )
-					{
-						bCollapsesChildren = false;
-						break;
-					}
-				}
+				TOptional<uint64> NumExpectedMaterialSlots = Context->InfoCache->GetSubtreeMaterialSlotCount( PrimPath ).Get( 0 );
 
 				// Note that we wont try to prevent collapsing in general if the combined mesh would have a triangle count above the threshold but too many material slots:
 				// We'll just disable Nanite with a message on the log instead.
@@ -587,33 +568,14 @@ bool FUsdGeomXformableTranslator::CollapsesChildren( ECollapsingType CollapsingT
 				//  - If we don't enable Nanite for them, then what is the benefit? Now the mesh hasn't collapsed but we don't have Nanite anywhere anyway...
 				// At least for the case below (with the explicit overrides) we would end up with some meshes having Nanite, according to how the user set them.
 				// In the future we could expose the collapsing controls on the stage actor to let the user control this a bit better
-
-				// Don't collapse children if the child meshes have Nanite override opinions but the combined mesh would lead to over 64 material slots.
-				if ( bChildrenWantNanite )
+				const int32 MaxNumSections = 64; // There is no define for this, but it's checked for on NaniteBuilder.cpp, FBuilderModule::Build
+				if ( !NumExpectedMaterialSlots.IsSet() || NumExpectedMaterialSlots.GetValue() > MaxNumSections )
 				{
-					if ( Context->bMergeIdenticalMaterialSlots )
-					{
-						const bool bProvideMaterialIndices = false;
-						UsdUtils::FUsdPrimMaterialAssignmentInfo LocalInfo = UsdUtils::GetPrimMaterialAssignments( ChildPrim.Get(), Context->Time, bProvideMaterialIndices, RenderContextToken );
-						CombinedSlots.Append( LocalInfo.Slots );
-						NumMaxExpectedMaterialSlots = CombinedSlots.Num();
-					}
-					else
-					{
-						std::vector<pxr::UsdGeomSubset> GeomSubsets = pxr::UsdShadeMaterialBindingAPI( ChildPrim.Get() ).GetMaterialBindSubsets();
-						NumMaxExpectedMaterialSlots += FMath::Max<int32>( 1, GeomSubsets.size() + 1 ); // +1 because we may create an additional slot if it's not properly partitioned
-					}
-
-					const int32 MaxNumSections = 64; // There is no define for this, but it's checked for on NaniteBuilder.cpp, FBuilderModule::Build
-					if ( NumMaxExpectedMaterialSlots > MaxNumSections )
-					{
-						UE_LOG( LogUsd, Log, TEXT( "Not collapsing down from prim '%s' as child meshes want Nanite to be abled but the generated static mesh would have more than '%d' material slots" ),
-							*PrimPath.GetString(),
-							MaxNumSections
-						);
-						bCollapsesChildren = false;
-						break;
-					}
+					UE_LOG( LogUsd, Log, TEXT( "Not collapsing down from prim '%s' as child meshes want Nanite to be abled but the generated static mesh would have more than '%d' material slots" ),
+						*PrimPath.GetString(),
+						MaxNumSections
+					);
+					bCollapsesChildren = false;
 				}
 			}
 		}
diff --git a/Engine/Plugins/Importers/USDImporter/Source/USDSchemas/Private/USDInfoCache.cpp b/Engine/Plugins/Importers/USDImporter/Source/USDSchemas/Private/USDInfoCache.cpp
new file mode 100644
--- /dev/null
+++ b/Engine/Plugins/Importers/USDImporter/Source/USDSchemas/Private/USDInfoCache.cpp
@@ -0,0 +1,578 @@
+// Copyright Epic Games, Inc. All Rights Reserved.
+
+#include "USDInfoCache.h"
+
+#include "USDGeomMeshConversion.h"
+#include "USDSchemasModule.h"
+#include "USDSchemaTranslator.h"
+#include "USDTypesConversion.h"
+
+#include "UsdWrappers/SdfPath.h"
+#include "UsdWrappers/UsdPrim.h"
+
+#include "Async/ParallelFor.h"
+#include "Modules/ModuleManager.h"
+
+#if USE_USD_SDK
+#include "USDIncludesStart.h"
+	#include "pxr/usd/sdf/path.h"
+	#include "pxr/usd/usd/prim.h"
+	#include "pxr/usd/usdGeom/mesh.h"
+	#include "pxr/usd/usdGeom/pointInstancer.h"
+	#include "pxr/usd/usdGeom/subset.h"
+#include "USDIncludesEnd.h"
+#endif // USE_USD_SDK
+
+bool FUsdInfoCache::Serialize( FArchive& Ar )
+{
+	FWriteScopeLock ScopeLock( Lock );
+	Ar << AssetPathsToCollapsedRoot;
+	Ar << ComponentPathsToCollapsedRoot;
+	Ar << ExpectedVertexCountPerSubtree;
+	Ar << ExpectedMaterialSlotCountPerSubtree;
+
+	return true;
+}
+
+bool FUsdInfoCache::IsPathCollapsed( const UE::FSdfPath& Path, ECollapsingType CollapsingType ) const
+{
+	FReadScopeLock ScopeLock( Lock );
+	const TMap< UE::FSdfPath, UE::FSdfPath >* MapToUse =
+		CollapsingType == ECollapsingType::Assets
+			? &AssetPathsToCollapsedRoot
+			: &ComponentPathsToCollapsedRoot;
+
+	if ( const UE::FSdfPath* FoundResult = MapToUse->Find( Path ) )
+	{
+		// A non-empty path to another prim means this prim is collapsed into that one
+		return !FoundResult->IsEmpty() && ( *FoundResult ) != Path;
+	}
+
+	// This should never happen: We should have cached the entire tree
+	ensureMsgf( false, TEXT( "Prim path '%s' has not been cached!" ), *Path.GetString() );
+	return false;
+}
+
+bool FUsdInfoCache::DoesPathCollapseChildren( const UE::FSdfPath& Path, ECollapsingType CollapsingType ) const
+{
+	FReadScopeLock ScopeLock( Lock );
+	const TMap< UE::FSdfPath, UE::FSdfPath >* MapToUse =
+		CollapsingType == ECollapsingType::Assets
+			? &AssetPathsToCollapsedRoot
+			: &ComponentPathsToCollapsedRoot;
+
+	if ( const UE::FSdfPath* FoundResult = MapToUse->Find( Path ) )
+	{
+		// We store our own Path in there when we collapse children. Otherwise we hold the path of our collapse root, or empty (in case nothing is collapsed up to here)
+		return (*FoundResult) == Path;
+	}
+
+	// This should never happen: We should have cached the entire tree
+	ensureMsgf( false, TEXT( "Prim path '%s' has not been cached!" ), *Path.GetString() );
+	return false;
+}
+
+UE::FSdfPath FUsdInfoCache::UnwindToNonCollapsedPath( const UE::FSdfPath& Path, ECollapsingType CollapsingType ) const
+{
+	FReadScopeLock ScopeLock( Lock );
+	const TMap< UE::FSdfPath, UE::FSdfPath >* MapToUse =
+		CollapsingType == ECollapsingType::Assets
+			? &AssetPathsToCollapsedRoot
+			: &ComponentPathsToCollapsedRoot;
+
+	if ( const UE::FSdfPath* FoundResult = MapToUse->Find( Path ) )
+	{
+		// An empty path here means that we are not collapsed at all
+		if ( FoundResult->IsEmpty() )
+		{
+			return Path;
+		}
+		// Otherwise we have our own path in there (in case we collapse children) or the path to the prim that collapsed us
+		else
+		{
+			return *FoundResult;
+		}
+	}
+
+	// This should never happen: We should have cached the entire tree
+	ensureMsgf( false, TEXT( "Prim path '%s' has not been cached!" ), *Path.GetString() );
+	return Path;
+}
+
+namespace UE::USDInfoCacheImpl::Private
+{
+#if USE_USD_SDK
+	void GetPrimVertexCountAndSlots(
+		const pxr::UsdPrim& UsdPrim,
+		const FUsdSchemaTranslationContext& Context,
+		const TMap< UE::FSdfPath, uint64 >& InSubtreeToVertexCounts,
+		const TMap< UE::FSdfPath, TArray<UsdUtils::FUsdPrimMaterialSlot> >& InSubtreeToMaterialSlots,
+		uint64& OutVertexCount,
+		TArray<UsdUtils::FUsdPrimMaterialSlot>& OutMaterialSlots,
+		FRWLock& Lock
+	)
+	{
+		TRACE_CPUPROFILER_EVENT_SCOPE( UE::USDInfoCacheImpl::Private::GetPrimVertexCountAndSlots );
+
+		if ( UsdPrim.IsA<pxr::UsdGeomMesh>() || UsdPrim.IsA<pxr::UsdGeomSubset>() )
+		{
+			if ( pxr::UsdGeomMesh Mesh{ UsdPrim } )
+			{
+				if ( pxr::UsdAttribute Points = Mesh.GetPointsAttr() )
+				{
+					pxr::VtArray< pxr::GfVec3f > PointsArray;
+					Points.Get( &PointsArray, pxr::UsdTimeCode( Context.Time ) );
+					OutVertexCount = PointsArray.size();
+				}
+			}
+
+			pxr::TfToken RenderContextToken = pxr::UsdShadeTokens->universalRenderContext;
+			if ( !Context.RenderContext.IsNone() )
+			{
+				RenderContextToken = UnrealToUsd::ConvertToken( *Context.RenderContext.ToString() ).Get();
+			}
+
+			const bool bProvideMaterialIndices = false;
+			UsdUtils::FUsdPrimMaterialAssignmentInfo LocalInfo = UsdUtils::GetPrimMaterialAssignments(
+				UsdPrim,
+				Context.Time,
+				bProvideMaterialIndices,
+				RenderContextToken
+			);
+
+			OutMaterialSlots.Append( MoveTemp( LocalInfo.Slots ) );
+		}
+		else if ( pxr::UsdGeomPointInstancer PointInstancer{ UsdPrim } )
+		{
+			const pxr::UsdRelationship& Prototypes = PointInstancer.GetPrototypesRel();
+
+			pxr::SdfPathVector PrototypePaths;
+			if ( Prototypes.GetTargets( &PrototypePaths ) )
+			{
+				TArray<uint64> PrototypeVertexCounts;
+				PrototypeVertexCounts.SetNumZeroed( PrototypePaths.size() );
+
+				{
+					FReadScopeLock ScopeLock( Lock );
+					for ( uint32 PrototypeIndex = 0; PrototypeIndex < PrototypePaths.size(); ++PrototypeIndex )
+					{
+						const pxr::SdfPath& PrototypePath = PrototypePaths[ PrototypeIndex ];
+
+						// If we're calling this for a point instancer we should have parsed the results for our
+						// prototype subtrees already
+						if ( const uint64* FoundPrototypeVertexCount = InSubtreeToVertexCounts.Find( UE::FSdfPath{ PrototypePath } ) )
+						{
+							PrototypeVertexCounts[ PrototypeIndex ] = *FoundPrototypeVertexCount;
+						}
+
+						if ( const TArray<UsdUtils::FUsdPrimMaterialSlot>* FoundPrototypeSlots = InSubtreeToMaterialSlots.Find( UE::FSdfPath{ PrototypePath } ) )
+						{
+							OutMaterialSlots.Append( *FoundPrototypeSlots );
+						}
+					}
+				}
+
+				if ( pxr::UsdAttribute ProtoIndicesAttr = PointInstancer.GetProtoIndicesAttr() )
+				{
+					pxr::VtArray<int> ProtoIndicesArr;
+					if ( ProtoIndicesAttr.Get( &ProtoIndicesArr, pxr::UsdTimeCode::EarliestTime() ) )
+					{
+						for ( int ProtoIndex : ProtoIndicesArr )
+						{
+							OutVertexCount += PrototypeVertexCounts[ ProtoIndex ];
+						}
+					}
+				}
+			}
+		}
+	}
+
+	void RecursiveRebuildCache(
+		const pxr::UsdPrim& UsdPrim,
+		FUsdSchemaTranslationContext& Context,
+		FUsdSchemaTranslatorRegistry& Registry,
+		TMap< UE::FSdfPath, UE::FSdfPath >& AssetPathsToCollapsedRoot,
+		TMap< UE::FSdfPath, UE::FSdfPath >& ComponentPathsToCollapsedRoot,
+		FRWLock& Lock,
+		TMap< UE::FSdfPath, uint64 >& InOutSubtreeToVertexCounts,
+		TMap< UE::FSdfPath, TArray<UsdUtils::FUsdPrimMaterialSlot>>& InOutSubtreeToMaterialSlots,
+		TArray< FString >& InOutPointInstancerPaths,
+		uint64& OutSubtreeVertexCount,
+		TArray<UsdUtils::FUsdPrimMaterialSlot>& OutSubtreeSlots,
+		const pxr::SdfPath& AssetCollapsedRoot = pxr::SdfPath::EmptyPath(),
+		const pxr::SdfPath& ComponentCollapsedRoot = pxr::SdfPath::EmptyPath()
+	)
+	{
+		TRACE_CPUPROFILER_EVENT_SCOPE( UE::USDInfoCacheImpl::Private::RecursiveRebuildCache );
+		FScopedUsdAllocs Allocs;
+
+		pxr::SdfPath UsdPrimPath = UsdPrim.GetPrimPath();
+
+		// Prevents allocation by referencing instead of copying
+		const pxr::SdfPath* AssetCollapsedRootOverride = &AssetCollapsedRoot;
+		const pxr::SdfPath* ComponentCollapsedRootOverride = &ComponentCollapsedRoot;
+
+		bool bIsAssetCollapsed = !AssetCollapsedRoot.IsEmpty();
+		bool bIsComponentCollapsed = !ComponentCollapsedRoot.IsEmpty();
+
+		if ( !bIsAssetCollapsed || !bIsComponentCollapsed )
+		{
+			if ( TSharedPtr< FUsdSchemaTranslator > SchemaTranslator = Registry.CreateTranslatorForSchema( Context.AsShared(), UE::FUsdTyped( UsdPrim ) ) )
+			{
+				if ( !bIsAssetCollapsed )
+				{
+					if ( SchemaTranslator->CollapsesChildren( ECollapsingType::Assets ) )
+					{
+						AssetCollapsedRootOverride = &UsdPrimPath;
+					}
+				}
+
+				if ( !bIsComponentCollapsed )
+				{
+					if ( SchemaTranslator->CollapsesChildren( ECollapsingType::Components ) )
+					{
+						ComponentCollapsedRootOverride = &UsdPrimPath;
+					}
+				}
+			}
+		}
+
+		pxr::UsdPrimSiblingRange PrimChildren = UsdPrim.GetFilteredChildren( pxr::UsdTraverseInstanceProxies( pxr::UsdPrimAllPrimsPredicate ) );
+
+		TArray<pxr::UsdPrim> Prims;
+		for ( pxr::UsdPrim Child : PrimChildren )
+		{
+			Prims.Emplace( Child );
+		}
+
+		const uint32 NumChildren = Prims.Num();
+
+		TArray<uint64> ChildSubtreeVertexCounts;
+		ChildSubtreeVertexCounts.SetNumUninitialized( NumChildren );
+
+		TArray<TArray<UsdUtils::FUsdPrimMaterialSlot>> ChildSubtreeMaterialSlots;
+		ChildSubtreeMaterialSlots.SetNum( NumChildren );
+
+		const int32 MinBatchSize = 1;
+
+		ParallelFor( TEXT( "RecursiveRebuildCache" ), Prims.Num(), MinBatchSize,
+			[&]( int32 Index )
+			{
+				RecursiveRebuildCache(
+					Prims[ Index ],
+					Context,
+					Registry,
+					AssetPathsToCollapsedRoot,
+					ComponentPathsToCollapsedRoot,
+					Lock,
+					InOutSubtreeToVertexCounts,
+					InOutSubtreeToMaterialSlots,
+					InOutPointInstancerPaths,
+					ChildSubtreeVertexCounts[ Index ],
+					ChildSubtreeMaterialSlots[ Index ],
+					*AssetCollapsedRootOverride,
+					*ComponentCollapsedRootOverride
+				);
+			}
+		);
+
+		OutSubtreeVertexCount = 0;
+		OutSubtreeSlots.Empty();
+
+		bool bIsPointInstancer = false;
+		if ( pxr::UsdGeomPointInstancer PointInstancer{ UsdPrim } )
+		{
+			bIsPointInstancer = true;
+		}
+		else
+		{
+			GetPrimVertexCountAndSlots(
+				UsdPrim,
+				Context,
+				InOutSubtreeToVertexCounts,
+				InOutSubtreeToMaterialSlots,
+				OutSubtreeVertexCount,
+				OutSubtreeSlots,
+				Lock
+			);
+
+			for ( uint32 ChildIndex = 0; ChildIndex < NumChildren; ++ChildIndex )
+			{
+				OutSubtreeVertexCount += ChildSubtreeVertexCounts[ ChildIndex ];
+				OutSubtreeSlots.Append( ChildSubtreeMaterialSlots[ ChildIndex ] );
+			}
+		}
+
+		{
+			FWriteScopeLock ScopeLock( Lock );
+
+			// For point instancers we can't guarantee we parsed the prototypes yet because they
+			// could technically be anywhere, so store them here for a later pass
+			if ( bIsPointInstancer )
+			{
+				InOutPointInstancerPaths.Emplace( UE::FSdfPath{ UsdPrimPath }.GetString() );
+			}
+			// While we will compute the totals for any and all children normally, don't just append the regular
+			// traversal vertex count to the point instancer prim itself just yet, as that doesn't really represent
+			// what will happen. We'll later do another pass to handle point instancers where we'll properly instance
+			// stuff, and then we'll updadte all ancestors
+			else
+			{
+				InOutSubtreeToVertexCounts.Emplace( UsdPrimPath, OutSubtreeVertexCount );
+				InOutSubtreeToMaterialSlots.Emplace( UsdPrimPath, OutSubtreeSlots );
+			}
+
+			// These paths will be still empty in case nothing has collapsed yet, hold UsdPrimPath in case UsdPrim
+			// collapses that type, or hold the path to the collapsed root passed in via our caller, in case we're
+			// collapsed
+			AssetPathsToCollapsedRoot.Emplace( UsdPrimPath, *AssetCollapsedRootOverride );
+			ComponentPathsToCollapsedRoot.Emplace( UsdPrimPath, *ComponentCollapsedRootOverride );
+
+		}
+	}
+
+	/**
+	 * Updates the subtree counts with point instancer instancing info.
+	 *
+	 * This has to be done outside of the main recursion because point instancers may reference any prim in the
+	 * stage to be their prototypes (including other point instancers), so we must first parse the entire
+	 * stage (forcing point instancer vertex/material slot counts to zero), and only then use the parsed counts
+	 * of prim subtrees all over to build the final counts of point instancers that use them as prototypes, and
+	 * then update their parents.
+	 */
+	void UpdateInfoForPointInstancers(
+		const pxr::UsdStageRefPtr& Stage,
+		TArray< FString >& PointInstancerPaths,
+		const FUsdSchemaTranslationContext& Context,
+		TMap< UE::FSdfPath, uint64 >& InOutExpectedVertexCountPerSubtree,
+		TMap< UE::FSdfPath, TArray<UsdUtils::FUsdPrimMaterialSlot>>& InOutSubtreeMaterialSlots,
+		FRWLock& Lock
+	)
+	{
+		// We must sort point instancers in a particular order in case they depend on each other.
+		// At least we know that an ordering like this should be possible, because A with B as a prototype and B with A
+		// as a prototype leads to an invalid USD stage.
+		TFunction<bool(const FString&, const FString&)> SortFunction = [Stage]( const FString& LHS, const FString& RHS )
+		{
+			FScopedUsdAllocs Allocs;
+
+			pxr::SdfPath LPath = UnrealToUsd::ConvertPath( *LHS ).Get();
+			pxr::SdfPath RPath = UnrealToUsd::ConvertPath( *RHS ).Get();
+
+			pxr::UsdGeomPointInstancer LPointInstancer = pxr::UsdGeomPointInstancer{ Stage->GetPrimAtPath( LPath ) };
+			pxr::UsdGeomPointInstancer RPointInstancer = pxr::UsdGeomPointInstancer{ Stage->GetPrimAtPath( RPath ) };
+			if ( LPointInstancer && RPointInstancer )
+			{
+				const pxr::UsdRelationship& LPrototypes = LPointInstancer.GetPrototypesRel();
+				pxr::SdfPathVector LPrototypePaths;
+				if ( LPrototypes.GetTargets( &LPrototypePaths ) )
+				{
+					for ( const pxr::SdfPath& LPrototypePath : LPrototypePaths )
+					{
+						// Consider RPointInstancer at RPath "/LPointInstancer/Prototypes/Nest/RPointInstancer", and
+						// LPointInstancer has prototype "/LPointInstancer/Prototypes/Nest". If RPath has the LPrototypePath as prefix,
+						// we should have R come before L in the sort order.
+						// Of course, in this scenario we could get away with just sorting by length, but that wouldn't help if the
+						// point instancers were not inside each other (e.g. siblings).
+						if ( RPath.HasPrefix( LPrototypePath ) )
+						{
+							return false;
+						}
+					}
+
+					// Give it the benefit of the doubt here and say that if R doesn't *need* to come before L, let's ensure L
+					// goes before R just in case
+					return true;
+				}
+			}
+
+			return LHS < RHS;
+		};
+		PointInstancerPaths.Sort(SortFunction);
+
+		for ( const FString& PointInstancerPath : PointInstancerPaths )
+		{
+			UE::FSdfPath UsdPointInstancerPath{ *PointInstancerPath };
+
+			if ( pxr::UsdPrim PointInstancer = Stage->GetPrimAtPath( UnrealToUsd::ConvertPath( *PointInstancerPath ).Get() ) )
+			{
+				uint64 PointInstancerVertexCount = 0;
+				TArray<UsdUtils::FUsdPrimMaterialSlot> PointInstancerMaterialSlots;
+
+				GetPrimVertexCountAndSlots(
+					PointInstancer,
+					Context,
+					InOutExpectedVertexCountPerSubtree,
+					InOutSubtreeMaterialSlots,
+					PointInstancerVertexCount,
+					PointInstancerMaterialSlots,
+					Lock
+				);
+
+				InOutExpectedVertexCountPerSubtree.Emplace( UsdPointInstancerPath, PointInstancerVertexCount );
+				InOutSubtreeMaterialSlots.Emplace( UsdPointInstancerPath, PointInstancerMaterialSlots );
+
+				// Now that we have info on the point instancer itself, update the counts of all ancestors.
+				// Note: The vertex/material slot count for the entire point instancer subtree are just the counts
+				// for the point instancer itself, as we stop regular traversal when we hit them
+				UE::FSdfPath ParentPath = UsdPointInstancerPath.GetParentPath();
+				pxr::UsdPrim Prim = Stage->GetPrimAtPath( ParentPath );
+				while ( Prim )
+				{
+					// If our ancestor is a point instancer itself, just abort as we'll only get the actual counts
+					// when we handle that ancestor directly. We don't want to update the ancestor point instancer's
+					// ancestors with incorrect values
+					if ( Prim.IsA<pxr::UsdGeomPointInstancer>() )
+					{
+						break;
+					}
+
+					InOutExpectedVertexCountPerSubtree[ ParentPath ] += PointInstancerVertexCount;
+					InOutSubtreeMaterialSlots[ ParentPath ].Append( PointInstancerMaterialSlots );
+
+					// Break only here so we update the pseudoroot too
+					if ( Prim.IsPseudoRoot() )
+					{
+						break;
+					}
+
+					ParentPath = ParentPath.GetParentPath();
+					Prim = Stage->GetPrimAtPath( ParentPath );
+				}
+			}
+		}
+	}
+
+	/**
+	 * Condenses our collected material slots for all subtress (SubtreeMaterialSlots) into just material slot counts,
+	 * (OutMaterialSlotCounts), according to bMergeIdenticalSlots.
+	 *
+	 * We do this after the main pass because then the main material slot collecting code on
+	 * the main recursive pass just adds them to arrays, and we're allowed to handle bMergeIdenticalSlots
+	 * only here.
+	 */
+	void CollectMaterialSlotCounts(
+		const TMap< UE::FSdfPath, TArray<UsdUtils::FUsdPrimMaterialSlot>>& SubtreeMaterialSlots,
+		TMap< UE::FSdfPath, uint64 >& OutMaterialSlotCounts,
+		bool bMergeIdenticalSlots
+	)
+	{
+		if ( bMergeIdenticalSlots )
+		{
+			for ( const TPair< UE::FSdfPath, TArray<UsdUtils::FUsdPrimMaterialSlot>>& Pair : SubtreeMaterialSlots )
+			{
+				TSet<UsdUtils::FUsdPrimMaterialSlot> SlotsSet{ Pair.Value };
+				OutMaterialSlotCounts.FindOrAdd( Pair.Key ) = SlotsSet.Num();
+			}
+		}
+		else
+		{
+			for ( const TPair< UE::FSdfPath, TArray<UsdUtils::FUsdPrimMaterialSlot>>& Pair : SubtreeMaterialSlots )
+			{
+				OutMaterialSlotCounts.FindOrAdd( Pair.Key ) = Pair.Value.Num();
+			}
+		}
+	}
+#endif // USE_USD_SDK
+}
+
+TOptional<uint64> FUsdInfoCache::GetSubtreeVertexCount( const UE::FSdfPath& Path )
+{
+	if ( uint64* FoundCount = ExpectedVertexCountPerSubtree.Find( Path ) )
+	{
+		return *FoundCount;
+	}
+
+	return {};
+}
+
+TOptional<uint64> FUsdInfoCache::GetSubtreeMaterialSlotCount( const UE::FSdfPath& Path )
+{
+	if ( uint64* FoundCount = ExpectedMaterialSlotCountPerSubtree.Find( Path ) )
+	{
+		return *FoundCount;
+	}
+
+	return {};
+}
+
+void FUsdInfoCache::RebuildCacheForSubtree( const UE::FUsdPrim& Prim, FUsdSchemaTranslationContext& Context )
+{
+#if USE_USD_SDK
+	TRACE_CPUPROFILER_EVENT_SCOPE( FUsdInfoCache::RebuildCacheForSubtree );
+
+	// We can't deallocate our info cache pointer with the Usd allocator
+	FScopedUnrealAllocs UEAllocs;
+
+	// We don't want the translation context to try using its info cache during the rebuild process, as that's the entire point
+	TGuardValue< TSharedPtr<FUsdInfoCache> > Guard{ Context.InfoCache, nullptr };
+	{
+		FScopedUsdAllocs Allocs;
+
+		pxr::UsdPrim UsdPrim{Prim};
+		if ( !UsdPrim )
+		{
+			return;
+		}
+
+		Clear();
+
+		IUsdSchemasModule& UsdSchemasModule = FModuleManager::Get().LoadModuleChecked< IUsdSchemasModule >( TEXT( "USDSchemas" ) );
+		FUsdSchemaTranslatorRegistry& Registry = UsdSchemasModule.GetTranslatorRegistry();
+
+		TMap< UE::FSdfPath, TArray<UsdUtils::FUsdPrimMaterialSlot>> TempSubtreeSlots;
+		TArray< FString > PointInstancerPaths;
+
+		uint64 SubtreeVertexCount = 0;
+		TArray<UsdUtils::FUsdPrimMaterialSlot> SubtreeSlots;
+		UE::USDInfoCacheImpl::Private::RecursiveRebuildCache(
+			UsdPrim,
+			Context,
+			Registry,
+			AssetPathsToCollapsedRoot,
+			ComponentPathsToCollapsedRoot,
+			Lock,
+			ExpectedVertexCountPerSubtree,
+			TempSubtreeSlots,
+			PointInstancerPaths,
+			SubtreeVertexCount,
+			SubtreeSlots
+		);
+
+		UE::USDInfoCacheImpl::Private::UpdateInfoForPointInstancers(
+			UsdPrim.GetStage(),
+			PointInstancerPaths,
+			Context,
+			ExpectedVertexCountPerSubtree,
+			TempSubtreeSlots,
+			Lock
+		);
+
+		UE::USDInfoCacheImpl::Private::CollectMaterialSlotCounts(
+			TempSubtreeSlots,
+			ExpectedMaterialSlotCountPerSubtree,
+			Context.bMergeIdenticalMaterialSlots
+		);
+	}
+#endif // USE_USD_SDK
+}
+
+void FUsdInfoCache::Clear()
+{
+	FWriteScopeLock ScopeLock( Lock );
+
+	AssetPathsToCollapsedRoot.Empty();
+	ComponentPathsToCollapsedRoot.Empty();
+	ExpectedVertexCountPerSubtree.Empty();
+	ExpectedMaterialSlotCountPerSubtree.Empty();
+}
+
+bool FUsdInfoCache::IsEmpty()
+{
+	FReadScopeLock ScopeLock( Lock );
+
+	return AssetPathsToCollapsedRoot.IsEmpty()
+		&& ComponentPathsToCollapsedRoot.IsEmpty()
+		&& ExpectedVertexCountPerSubtree.IsEmpty()
+		&& ExpectedMaterialSlotCountPerSubtree.IsEmpty();
+}
\ No newline at end of file
diff --git a/Engine/Plugins/Importers/USDImporter/Source/USDSchemas/Private/USDSchemaTranslator.cpp b/Engine/Plugins/Importers/USDImporter/Source/USDSchemas/Private/USDSchemaTranslator.cpp
--- a/Engine/Plugins/Importers/USDImporter/Source/USDSchemas/Private/USDSchemaTranslator.cpp
+++ b/Engine/Plugins/Importers/USDImporter/Source/USDSchemas/Private/USDSchemaTranslator.cpp
@@ -2,7 +2,7 @@
 
 #include "USDSchemaTranslator.h"
 
-#include "USDCollapsingCache.h"
+#include "USDInfoCache.h"
 #include "USDErrorUtils.h"
 #include "USDSchemasModule.h"
 #include "USDTypesConversion.h"
@@ -204,9 +204,9 @@ bool FUsdSchemaTranslator::IsCollapsed( ECollapsingType CollapsingType ) const
 #if USE_USD_SDK
 	TRACE_CPUPROFILER_EVENT_SCOPE( FUsdSchemaTranslator::IsCollapsed );
 
-	if ( Context->CollapsingCache.IsValid() )
+	if ( Context->InfoCache.IsValid() )
 	{
-		return Context->CollapsingCache->IsPathCollapsed( PrimPath, CollapsingType );
+		return Context->InfoCache->IsPathCollapsed( PrimPath, CollapsingType );
 	}
 
 	// This is merely a fallback, and we should never need this
diff --git a/Engine/Plugins/Importers/USDImporter/Source/USDSchemas/Public/USDCollapsingCache.h b/Engine/Plugins/Importers/USDImporter/Source/USDSchemas/Public/USDCollapsingCache.h
deleted file mode 100644
diff --git a/Engine/Plugins/Importers/USDImporter/Source/USDSchemas/Public/USDInfoCache.h b/Engine/Plugins/Importers/USDImporter/Source/USDSchemas/Public/USDInfoCache.h
new file mode 100644
--- /dev/null
+++ b/Engine/Plugins/Importers/USDImporter/Source/USDSchemas/Public/USDInfoCache.h
@@ -0,0 +1,58 @@
+// Copyright Epic Games, Inc. All Rights Reserved.
+
+#pragma once
+
+#include "Containers/Set.h"
+#include "CoreMinimal.h"
+
+#include "UsdWrappers/SdfPath.h"
+
+struct FUsdSchemaTranslationContext;
+namespace UE
+{
+	class FSdfPath;
+	class FUsdPrim;
+}
+
+enum class ECollapsingType
+{
+	Assets,
+	Components
+};
+
+/**
+ * Caches whether given prim paths represent prims that can have their assets or components collapsed when importing/opening
+ * a USD Stage.
+ */
+struct USDSCHEMAS_API FUsdInfoCache
+{
+	bool Serialize( FArchive& Ar );
+
+	bool IsPathCollapsed( const UE::FSdfPath& Path, ECollapsingType CollapsingType ) const;
+	bool DoesPathCollapseChildren( const UE::FSdfPath& Path, ECollapsingType CollapsingType ) const;
+
+	// Returns Path in case it represents an uncollapsed prim, or returns the path to the prim that collapsed it
+	UE::FSdfPath UnwindToNonCollapsedPath( const UE::FSdfPath& Path, ECollapsingType CollapsingType ) const;
+
+	// Provides the total vertex or material slots counts for each prim *and* its subtree.
+	// This is built inside RebuildCacheForSubtree, so it will factor in the used Context's bMergeIdenticalMaterialSlots.
+	// Note that these aren't affected by actual collapsing: A prim that doesn't collapse its children will still
+	// provide the total sum of vertex counts of its entire subtree when queried
+	TOptional<uint64> GetSubtreeVertexCount( const UE::FSdfPath& Path );
+	TOptional<uint64> GetSubtreeMaterialSlotCount( const UE::FSdfPath& Path );
+
+	void RebuildCacheForSubtree( const UE::FUsdPrim& Prim, FUsdSchemaTranslationContext& Context );
+
+	void Clear();
+	bool IsEmpty();
+
+private:
+	mutable FRWLock Lock;
+
+	// TODO: Maybe use an impl class instead so that we can also clean up some of the parameters of the implementations
+	// TODO: Combine these maps into a single map to a struct with all this info
+	TMap< UE::FSdfPath, UE::FSdfPath > AssetPathsToCollapsedRoot;
+	TMap< UE::FSdfPath, UE::FSdfPath > ComponentPathsToCollapsedRoot;
+	TMap< UE::FSdfPath, uint64 > ExpectedVertexCountPerSubtree;
+	TMap< UE::FSdfPath, uint64 > ExpectedMaterialSlotCountPerSubtree;
+};
\ No newline at end of file
diff --git a/Engine/Plugins/Importers/USDImporter/Source/USDSchemas/Public/USDSchemaTranslator.h b/Engine/Plugins/Importers/USDImporter/Source/USDSchemas/Public/USDSchemaTranslator.h
--- a/Engine/Plugins/Importers/USDImporter/Source/USDSchemas/Public/USDSchemaTranslator.h
+++ b/Engine/Plugins/Importers/USDImporter/Source/USDSchemas/Public/USDSchemaTranslator.h
@@ -6,7 +6,7 @@
 
 #include "UnrealUSDWrapper.h"
 #include "USDAssetCache.h"
-#include "USDCollapsingCache.h"
+#include "USDInfoCache.h"
 #include "USDMemory.h"
 #include "USDSkeletalDataConversion.h"
 
@@ -154,8 +154,8 @@ struct USDSCHEMAS_API FUsdSchemaTranslationContext : public TSharedFromThis< FUs
 	/** Where the translated assets will be stored */
 	TStrongObjectPtr< UUsdAssetCache > AssetCache;
 
-	/** Caches the collapsing state of prim paths */
-	TSharedPtr<FUsdCollapsingCache> CollapsingCache;
+	/** Caches various information about prims that are expensive to query */
+	TSharedPtr<FUsdInfoCache> InfoCache;
 
 	/** Where we place imported blend shapes, if available */
 	UsdUtils::FBlendShapeMap* BlendShapesByPath = nullptr;
diff --git a/Engine/Plugins/Importers/USDImporter/Source/USDStage/Private/USDStageActor.cpp b/Engine/Plugins/Importers/USDImporter/Source/USDStage/Private/USDStageActor.cpp
--- a/Engine/Plugins/Importers/USDImporter/Source/USDStage/Private/USDStageActor.cpp
+++ b/Engine/Plugins/Importers/USDImporter/Source/USDStage/Private/USDStageActor.cpp
@@ -6,7 +6,7 @@
 #include "USDAssetCache.h"
 #include "USDAssetImportData.h"
 #include "USDClassesModule.h"
-#include "USDCollapsingCache.h"
+#include "USDInfoCache.h"
 #include "USDConversionUtils.h"
 #include "USDErrorUtils.h"
 #include "USDGeomMeshConversion.h"
@@ -96,7 +96,7 @@ struct FUsdStageActorImpl
 		TranslationContext->RenderContext = StageActor->RenderContext;
 		TranslationContext->MaterialToPrimvarToUVIndex = &StageActor->MaterialToPrimvarToUVIndex;
 		TranslationContext->BlendShapesByPath = &StageActor->BlendShapesByPath;
-		TranslationContext->CollapsingCache = StageActor->CollapsingCache;
+		TranslationContext->InfoCache = StageActor->InfoCache;
 
 		// Its more convenient to toggle between variants using the USDStage window, as opposed to parsing LODs
 		TranslationContext->bAllowInterpretingLODs = false;
@@ -869,7 +869,7 @@ void AUsdStageActor::OnUsdObjectsChanged( const UsdUtils::FObjectChangesByPath&
 
 	bool bDeselected = false;
 
-	if ( bHasResync && CollapsingCache.IsValid() )
+	if ( bHasResync && InfoCache.IsValid() )
 	{
 		// TODO: Selective rebuild of only the required parts of the cache.
 		// If a prim changes from CanBeCollapsed to not (or vice-versa), that means its parent may change, and its grandparent, etc. so we'd need to check a large
@@ -877,7 +877,7 @@ void AUsdStageActor::OnUsdObjectsChanged( const UsdUtils::FObjectChangesByPath&
 		// However, if we know that only Prim resynced, we can traverse the tree root down and if we reach Prim and its collapsing state
 		// hasn't updated from before, we don't have to update its subtree at all, or sibling subtrees.
 		TSharedRef< FUsdSchemaTranslationContext > TranslationContext = FUsdStageActorImpl::CreateUsdSchemaTranslationContext( this, TEXT( "/" ) );
-		CollapsingCache->RebuildCacheForSubtree( UsdStage.GetPseudoRoot(), TranslationContext.Get() );
+		InfoCache->RebuildCacheForSubtree( UsdStage.GetPseudoRoot(), TranslationContext.Get() );
 	}
 
 	for ( const TPair< FString, bool >& PrimChangedInfo : SortedPrimsChangedList )
@@ -921,8 +921,8 @@ void AUsdStageActor::OnUsdObjectsChanged( const UsdUtils::FObjectChangesByPath&
 
 		auto UpdateComponents = [&]( const UE::FSdfPath& InPrimPath, const bool bInResync )
 		{
-			UE::FSdfPath ComponentsPrimPath = CollapsingCache.IsValid()
-				? CollapsingCache->UnwindToNonCollapsedPath( InPrimPath, ECollapsingType::Components )
+			UE::FSdfPath ComponentsPrimPath = InfoCache.IsValid()
+				? InfoCache->UnwindToNonCollapsedPath( InPrimPath, ECollapsingType::Components )
 				: InPrimPath;
 
 			TSet< UE::FSdfPath >& RefreshedComponents = bInResync ? ResyncedComponents : UpdatedComponents;
@@ -945,8 +945,8 @@ void AUsdStageActor::OnUsdObjectsChanged( const UsdUtils::FObjectChangesByPath&
 
 		auto ReloadAssets = [&]( const UE::FSdfPath& InPrimPath, const bool bInResync )
 		{
-			UE::FSdfPath AssetsPrimPath = CollapsingCache.IsValid()
-				? CollapsingCache->UnwindToNonCollapsedPath( InPrimPath, ECollapsingType::Assets )
+			UE::FSdfPath AssetsPrimPath = InfoCache.IsValid()
+				? InfoCache->UnwindToNonCollapsedPath( InPrimPath, ECollapsingType::Assets )
 				: InPrimPath;
 
 			TSet< UE::FSdfPath >& RefreshedAssets = bInResync ? ResyncedAssets : UpdatedAssets;
@@ -1373,7 +1373,7 @@ USceneComponent* AUsdStageActor::GetGeneratedComponent( const FString& PrimPath
 		return nullptr;
 	}
 
-	// We can't query our CollapsingCache with invalid paths, as we're using ensures to track when we miss the cache (which shouldn't ever happen)
+	// We can't query our InfoCache with invalid paths, as we're using ensures to track when we miss the cache (which shouldn't ever happen)
 	UE::FSdfPath UsdPath{ *PrimPath };
 	if ( !UsdStage.GetPrimAtPath( UsdPath ) )
 	{
@@ -1381,9 +1381,9 @@ USceneComponent* AUsdStageActor::GetGeneratedComponent( const FString& PrimPath
 	}
 
 	FString UncollapsedPath = PrimPath;
-	if ( CollapsingCache.IsValid() )
+	if ( InfoCache.IsValid() )
 	{
-		UncollapsedPath = CollapsingCache->UnwindToNonCollapsedPath( UsdPath, ECollapsingType::Components ).GetString();
+		UncollapsedPath = InfoCache->UnwindToNonCollapsedPath( UsdPath, ECollapsingType::Components ).GetString();
 	}
 
 	if ( UUsdPrimTwin* UsdPrimTwin = GetRootPrimTwin()->Find( UncollapsedPath ) )
@@ -1401,7 +1401,7 @@ TArray<UObject*> AUsdStageActor::GetGeneratedAssets( const FString& PrimPath )
 		return {};
 	}
 
-	// We can't query our CollapsingCache with invalid paths, as we're using ensures to track when we miss the cache (which shouldn't ever happen)
+	// We can't query our InfoCache with invalid paths, as we're using ensures to track when we miss the cache (which shouldn't ever happen)
 	UE::FSdfPath UsdPath{ *PrimPath };
 	if ( !UsdStage.GetPrimAtPath( UsdPath ) )
 	{
@@ -1423,9 +1423,9 @@ TArray<UObject*> AUsdStageActor::GetGeneratedAssets( const FString& PrimPath )
 		Result.Add( FoundAsset );
 	}
 	// If we haven't try seeing if we collapsed into another asset
-	else if ( CollapsingCache.IsValid() )
+	else if ( InfoCache.IsValid() )
 	{
-		PathToUse = CollapsingCache->UnwindToNonCollapsedPath( UsdPath, ECollapsingType::Assets ).GetString();
+		PathToUse = InfoCache->UnwindToNonCollapsedPath( UsdPath, ECollapsingType::Assets ).GetString();
 		if ( UObject* FoundCollapsedAsset = AssetCache->GetAssetForPrim( PathToUse ) )
 		{
 			Result.Add( FoundCollapsedAsset );
@@ -1646,8 +1646,8 @@ void AUsdStageActor::OnObjectsReplaced( const TMap<UObject*, UObject*>& ObjectRe
 			AssetCache->Rename( nullptr, NewActor );
 			NewActor->AssetCache = AssetCache;
 
-			NewActor->CollapsingCache = CollapsingCache;
-			CollapsingCache = nullptr;
+			NewActor->InfoCache = InfoCache;
+			InfoCache = nullptr;
 
 			// It could be that we're automatically recompiling when going into PIE because our blueprint was dirty.
 			// In that case we also need bIsTransitioningIntoPIE to be true to prevent us from calling LoadUsdStage from PostRegisterAllComponents
@@ -1731,11 +1731,11 @@ void AUsdStageActor::LoadUsdStage()
 	TSharedRef< FUsdSchemaTranslationContext > TranslationContext = FUsdStageActorImpl::CreateUsdSchemaTranslationContext( this, RootTwin->PrimPath );
 
 	SlowTask.EnterProgressFrame( 0.1f );
-	if ( !CollapsingCache.IsValid() )
+	if ( !InfoCache.IsValid() )
 	{
-		CollapsingCache = MakeShared<FUsdCollapsingCache>();
+		InfoCache = MakeShared<FUsdInfoCache>();
 	}
-	CollapsingCache->RebuildCacheForSubtree( UsdStage.GetPseudoRoot(), TranslationContext.Get() );
+	InfoCache->RebuildCacheForSubtree( UsdStage.GetPseudoRoot(), TranslationContext.Get() );
 
 	SlowTask.EnterProgressFrame( 0.7f );
 	LoadAssets( *TranslationContext, UsdStage.GetPseudoRoot() );
@@ -1806,9 +1806,9 @@ void AUsdStageActor::UnloadUsdStage()
 		AssetCache->Reset();
 	}
 
-	if ( CollapsingCache )
+	if ( InfoCache )
 	{
-		CollapsingCache->Clear();
+		InfoCache->Clear();
 	}
 
 	ObjectsToWatch.Reset();
@@ -2132,12 +2132,12 @@ void AUsdStageActor::Serialize(FArchive& Ar)
 
 	if ( ( Ar.GetPortFlags() & PPF_DuplicateForPIE ) || Ar.IsTransacting() )
 	{
-		if ( !CollapsingCache.IsValid() )
+		if ( !InfoCache.IsValid() )
 		{
-			CollapsingCache = MakeShared<FUsdCollapsingCache>();
+			InfoCache = MakeShared<FUsdInfoCache>();
 		}
 
-		CollapsingCache->Serialize( Ar );
+		InfoCache->Serialize( Ar );
 	}
 }
 
diff --git a/Engine/Plugins/Importers/USDImporter/Source/USDStage/Public/USDStageActor.h b/Engine/Plugins/Importers/USDImporter/Source/USDStage/Public/USDStageActor.h
--- a/Engine/Plugins/Importers/USDImporter/Source/USDStage/Public/USDStageActor.h
+++ b/Engine/Plugins/Importers/USDImporter/Source/USDStage/Public/USDStageActor.h
@@ -6,7 +6,7 @@
 #include "Misc/SecureHash.h"
 #include "Templates/Function.h"
 
-#include "USDCollapsingCache.h"
+#include "USDInfoCache.h"
 #include "USDLevelSequenceHelper.h"
 #include "USDListener.h"
 #include "USDPrimTwin.h"
@@ -186,7 +186,7 @@ public:
 	void ReloadAnimations();
 	float GetTime() { return Time; }
 	UUsdAssetCache* GetAssetCache() { return AssetCache; }
-	TSharedPtr<FUsdCollapsingCache> GetCollapsingCache() { return CollapsingCache; }
+	TSharedPtr<FUsdInfoCache> GetInfoCache() { return InfoCache; }
 	TMap< FString, TMap< FString, int32 > > GetMaterialToPrimvarToUVIndex() { return MaterialToPrimvarToUVIndex; }
 
 public:
@@ -287,7 +287,8 @@ protected:
 	void AnimatePrims();
 
 private:
-	TSharedPtr<FUsdCollapsingCache> CollapsingCache;
+	/** Caches various information about prims that are expensive to query */
+	TSharedPtr<FUsdInfoCache> InfoCache;
 
 	/** Keep track of blend shapes so that we can map 'inbetween shapes' to their separate morph targets when animating */
 	UsdUtils::FBlendShapeMap BlendShapesByPath;
diff --git a/Engine/Plugins/Importers/USDImporter/Source/USDStageEditorViewModels/Private/USDStageViewModel.cpp b/Engine/Plugins/Importers/USDImporter/Source/USDStageEditorViewModels/Private/USDStageViewModel.cpp
--- a/Engine/Plugins/Importers/USDImporter/Source/USDStageEditorViewModels/Private/USDStageViewModel.cpp
+++ b/Engine/Plugins/Importers/USDImporter/Source/USDStageEditorViewModels/Private/USDStageViewModel.cpp
@@ -307,7 +307,7 @@ void FUsdStageViewModel::ImportStage()
 			// Let the importer reuse our assets, but force it to spawn new actors and components always
 			// This allows a different setting for asset/component collapsing, and doesn't require modifying the PrimTwins
 			ImportContext.AssetCache = StageActor->GetAssetCache();
-			ImportContext.CollapsingCache = StageActor->GetCollapsingCache();
+			ImportContext.InfoCache = StageActor->GetInfoCache();
 			ImportContext.LevelSequenceHelper.SetAssetCache( StageActor->GetAssetCache() );
 			ImportContext.MaterialToPrimvarToUVIndex = StageActor->GetMaterialToPrimvarToUVIndex();
 
diff --git a/Engine/Plugins/Importers/USDImporter/Source/USDStageImporter/Private/USDStageImportContext.cpp b/Engine/Plugins/Importers/USDImporter/Source/USDStageImporter/Private/USDStageImportContext.cpp
--- a/Engine/Plugins/Importers/USDImporter/Source/USDStageImporter/Private/USDStageImportContext.cpp
+++ b/Engine/Plugins/Importers/USDImporter/Source/USDStageImporter/Private/USDStageImportContext.cpp
@@ -82,7 +82,7 @@ void FUsdStageImportContext::Reset()
 	ImportedAsset = nullptr;
 	LevelSequenceHelper.Clear();
 	AssetCache = nullptr;
-	CollapsingCache = nullptr;
+	InfoCache = nullptr;
 	MaterialToPrimvarToUVIndex.Empty();
 	Stage = UE::FUsdStage{};
 	ImportObjectFlags = EObjectFlags::RF_NoFlags;
diff --git a/Engine/Plugins/Importers/USDImporter/Source/USDStageImporter/Private/USDStageImporter.cpp b/Engine/Plugins/Importers/USDImporter/Source/USDStageImporter/Private/USDStageImporter.cpp
--- a/Engine/Plugins/Importers/USDImporter/Source/USDStageImporter/Private/USDStageImporter.cpp
+++ b/Engine/Plugins/Importers/USDImporter/Source/USDStageImporter/Private/USDStageImporter.cpp
@@ -217,14 +217,14 @@ namespace UsdStageImporterImpl
 
 	void CacheCollapsingState( FUsdSchemaTranslationContext& TranslationContext )
 	{
-		if ( !TranslationContext.CollapsingCache.IsValid() )
+		if ( !TranslationContext.InfoCache.IsValid() )
 		{
-			TranslationContext.CollapsingCache = MakeShared<FUsdCollapsingCache>();
+			TranslationContext.InfoCache = MakeShared<FUsdInfoCache>();
 		}
 
-		if ( TranslationContext.CollapsingCache->IsEmpty() )
+		if ( TranslationContext.InfoCache->IsEmpty() )
 		{
-			TranslationContext.CollapsingCache->RebuildCacheForSubtree( TranslationContext.Stage.GetPseudoRoot(), TranslationContext );
+			TranslationContext.InfoCache->RebuildCacheForSubtree( TranslationContext.Stage.GetPseudoRoot(), TranslationContext );
 		}
 	}
 
@@ -1499,7 +1499,7 @@ void UUsdStageImporter::ImportFromFile(FUsdStageImportContext& ImportContext)
 	TranslationContext->bAllowInterpretingLODs = ImportContext.ImportOptions->bInterpretLODs;
 	TranslationContext->bAllowParsingSkeletalAnimations = ImportContext.ImportOptions->bImportGeometry && ImportContext.ImportOptions->bImportSkeletalAnimations;
 	TranslationContext->MaterialToPrimvarToUVIndex = &ImportContext.MaterialToPrimvarToUVIndex;
-	TranslationContext->CollapsingCache = ImportContext.CollapsingCache;
+	TranslationContext->InfoCache = ImportContext.InfoCache;
 	TranslationContext->BlendShapesByPath = &BlendShapesByPath;
 	{
 		UsdStageImporterImpl::CacheCollapsingState( TranslationContext.Get() );
@@ -1586,7 +1586,7 @@ bool UUsdStageImporter::ReimportSingleAsset(FUsdStageImportContext& ImportContex
 	TranslationContext->bAllowInterpretingLODs = ImportContext.ImportOptions->bInterpretLODs;
 	TranslationContext->bAllowParsingSkeletalAnimations = ImportContext.ImportOptions->bImportGeometry && ImportContext.ImportOptions->bImportSkeletalAnimations;
 	TranslationContext->MaterialToPrimvarToUVIndex = &ImportContext.MaterialToPrimvarToUVIndex;
-	TranslationContext->CollapsingCache = ImportContext.CollapsingCache;
+	TranslationContext->InfoCache = ImportContext.InfoCache;
 	TranslationContext->BlendShapesByPath = &BlendShapesByPath;
 	{
 		UsdStageImporterImpl::CacheCollapsingState( TranslationContext.Get() );
diff --git a/Engine/Plugins/Importers/USDImporter/Source/USDStageImporter/Public/USDStageImportContext.h b/Engine/Plugins/Importers/USDImporter/Source/USDStageImporter/Public/USDStageImportContext.h
--- a/Engine/Plugins/Importers/USDImporter/Source/USDStageImporter/Public/USDStageImportContext.h
+++ b/Engine/Plugins/Importers/USDImporter/Source/USDStageImporter/Public/USDStageImportContext.h
@@ -3,7 +3,7 @@
 #pragma once
 
 #include "UnrealUSDWrapper.h"
-#include "USDCollapsingCache.h"
+#include "USDInfoCache.h"
 #include "USDLevelSequenceHelper.h"
 #include "USDMemory.h"
 #include "USDStageOptions.h"
@@ -60,7 +60,8 @@ struct USDSTAGEIMPORTER_API FUsdStageImportContext
 	UPROPERTY()
 	UUsdAssetCache* AssetCache;
 
-	TSharedPtr<FUsdCollapsingCache> CollapsingCache;
+	/** Caches various information about prims that are expensive to query */
+	TSharedPtr<FUsdInfoCache> InfoCache;
 
 	/**
 	 * When parsing materials, we keep track of which primvar we mapped to which UV channel.
diff --git a/Engine/Plugins/Importers/USDImporter/Source/USDTests/Private/USDTestsBlueprintLibrary.cpp b/Engine/Plugins/Importers/USDImporter/Source/USDTests/Private/USDTestsBlueprintLibrary.cpp
--- a/Engine/Plugins/Importers/USDImporter/Source/USDTests/Private/USDTestsBlueprintLibrary.cpp
+++ b/Engine/Plugins/Importers/USDImporter/Source/USDTests/Private/USDTestsBlueprintLibrary.cpp
@@ -90,3 +90,46 @@ void USDTestsBlueprintLibrary::DirtyStageActorBlueprint( AUsdStageActor* Bluepri
 	}
 #endif // WITH_EDITOR
 }
+
+int64 USDTestsBlueprintLibrary::GetSubtreeVertexCount( AUsdStageActor* StageActor, const FString& PrimPath )
+{
+	if ( StageActor )
+	{
+		if ( TSharedPtr<FUsdInfoCache> Cache = StageActor->GetInfoCache() )
+		{
+			TOptional<uint64> Result = Cache->GetSubtreeVertexCount( UE::FSdfPath{ *PrimPath } );
+			if ( Result.IsSet() )
+			{
+				// Narrowing conversion here, but we're only using this for our test scenes, which have at most a few thousand vertices
+				return static_cast< int64 >( Result.GetValue() );
+			}
+			else
+			{
+				return -1;
+			}
+		}
+	}
+
+	return -1;
+}
+
+int64 USDTestsBlueprintLibrary::GetSubtreeMaterialSlotCount( AUsdStageActor* StageActor, const FString& PrimPath )
+{
+	if ( StageActor )
+	{
+		if ( TSharedPtr<FUsdInfoCache> Cache = StageActor->GetInfoCache() )
+		{
+			TOptional<uint64> Result = Cache->GetSubtreeMaterialSlotCount( UE::FSdfPath{ *PrimPath } );
+			if ( Result.IsSet() )
+			{
+				return static_cast< int64 >( Result.GetValue() );
+			}
+			else
+			{
+				return -1;
+			}
+		}
+	}
+
+	return -1;
+}
diff --git a/Engine/Plugins/Importers/USDImporter/Source/USDTests/Public/USDTestsBlueprintLibrary.h b/Engine/Plugins/Importers/USDImporter/Source/USDTests/Public/USDTestsBlueprintLibrary.h
--- a/Engine/Plugins/Importers/USDImporter/Source/USDTests/Public/USDTestsBlueprintLibrary.h
+++ b/Engine/Plugins/Importers/USDImporter/Source/USDTests/Public/USDTestsBlueprintLibrary.h
@@ -32,4 +32,18 @@ public:
 	 */
 	UFUNCTION( BlueprintCallable, Category = "Blueprint" )
 	static void DirtyStageActorBlueprint( AUsdStageActor* BlueprintDerivedStageActor );
+
+	/**
+	 * Queries a subtree vertex count using the stage actor's info cache, which is not yet exposed to blueprint.
+	 * May return -1 in case of an error.
+	 */
+	UFUNCTION( BlueprintCallable, Category = "Counts" )
+	static int64 GetSubtreeVertexCount( AUsdStageActor* StageActor, const FString& PrimPath );
+
+	/**
+	 * Queries a subtree material slot count using the stage actor's info cache, which is not yet exposed to blueprint.
+	 * May return -1 in case of an error.
+	 */
+	UFUNCTION( BlueprintCallable, Category = "Counts" )
+	static int64 GetSubtreeMaterialSlotCount( AUsdStageActor* StageActor, const FString& PrimPath );
 };
diff --git a/Engine/Plugins/Importers/USDImporter/Source/USDTests/USDTests.Build.cs b/Engine/Plugins/Importers/USDImporter/Source/USDTests/USDTests.Build.cs
--- a/Engine/Plugins/Importers/USDImporter/Source/USDTests/USDTests.Build.cs
+++ b/Engine/Plugins/Importers/USDImporter/Source/USDTests/USDTests.Build.cs
@@ -17,6 +17,7 @@ namespace UnrealBuildTool.Rules
 					"Engine",
 					"UnrealUSDWrapper",
 					"USDClasses",
+					"USDSchemas",
 					"USDStage",
 					"USDUtilities",
 				}
