USD: Fix resync changes to collapsed prims not refreshing the collapsed root.

diff --git a/Engine/Plugins/Importers/USDImporter/Source/USDStage/Private/USDStageActor.cpp b/Engine/Plugins/Importers/USDImporter/Source/USDStage/Private/USDStageActor.cpp
--- a/Engine/Plugins/Importers/USDImporter/Source/USDStage/Private/USDStageActor.cpp
+++ b/Engine/Plugins/Importers/USDImporter/Source/USDStage/Private/USDStageActor.cpp
@@ -844,6 +844,32 @@ void AUsdStageActor::OnUsdObjectsChanged( const UsdUtils::FObjectChangesByPath&
 		SortedPrimsChangedList.Add( PrimPath.GetString(), bIsResync );
 	}
 
+	// If we have just resynced a child prim that was collapsed, we actually need to resync from its collapsed root
+	// downwards instead.
+	// UpdateComponents/ReloadAssets below will already do this of course, but the point of this is that we unwind
+	// the collapsed prim paths *before* we rebuild the cache: We're really interested in whether this prim was
+	// collapsed or not before this change, as the change itself may have made the prim uncollapsible now
+	if ( InfoCache.IsValid() )
+	{
+		TMap<FString, bool> NewEntries;
+		NewEntries.Reserve( SortedPrimsChangedList.Num() );
+
+		for ( TMap< FString, bool >::TIterator ChangeIt = SortedPrimsChangedList.CreateIterator(); ChangeIt; ++ChangeIt )
+		{
+			if ( ChangeIt.Value() )
+			{
+				UE::FSdfPath ChildPath{ *ChangeIt.Key() };
+
+				NewEntries.Add( InfoCache->UnwindToNonCollapsedPath( ChildPath, ECollapsingType::Assets ).GetString(), true );
+				NewEntries.Add( InfoCache->UnwindToNonCollapsedPath( ChildPath, ECollapsingType::Components ).GetString(), true );
+
+				ChangeIt.RemoveCurrent();
+			}
+		}
+
+		SortedPrimsChangedList.Append( NewEntries );
+	}
+
 	SortedPrimsChangedList.KeySort([]( const FString& A, const FString& B ) -> bool { return A.Len() < B.Len(); } );
 
 	// During PIE, the PIE and the editor world will respond to notices. We have to prevent any PIE
