USD: Fixed USD Stage Actor was unnecessarily dirtying itself when being played back in Level Sequence.

diff --git a/Engine/Plugins/Importers/USDImporter/Source/USDStage/Private/USDStageActor.cpp b/Engine/Plugins/Importers/USDImporter/Source/USDStage/Private/USDStageActor.cpp
--- a/Engine/Plugins/Importers/USDImporter/Source/USDStage/Private/USDStageActor.cpp
+++ b/Engine/Plugins/Importers/USDImporter/Source/USDStage/Private/USDStageActor.cpp
@@ -1277,7 +1277,8 @@ UE::FUsdStage& AUsdStageActor::GetOrLoadUsdStage()
 
 void AUsdStageActor::SetRootLayer( const FString& RootFilePath )
 {
-	Modify();
+	const bool bMarkDirty = false;
+	Modify(bMarkDirty);
 
 	FString RelativeFilePath = RootFilePath;
 #if USE_USD_SDK
@@ -1307,7 +1308,8 @@ void AUsdStageActor::SetRootLayer( const FString& RootFilePath )
 
 void AUsdStageActor::SetInitialLoadSet( EUsdInitialLoadSet NewLoadSet )
 {
-	Modify();
+	const bool bMarkDirty = false;
+	Modify(bMarkDirty);
 
 	InitialLoadSet = NewLoadSet;
 	LoadUsdStage();
@@ -1315,7 +1317,8 @@ void AUsdStageActor::SetInitialLoadSet( EUsdInitialLoadSet NewLoadSet )
 
 void AUsdStageActor::SetInterpolationType( EUsdInterpolationType NewType )
 {
-	Modify();
+	const bool bMarkDirty = false;
+	Modify(bMarkDirty);
 
 	InterpolationType = NewType;
 	LoadUsdStage();
@@ -1323,7 +1326,8 @@ void AUsdStageActor::SetInterpolationType( EUsdInterpolationType NewType )
 
 void AUsdStageActor::SetKindsToCollapse( int32 NewKindsToCollapse )
 {
-	Modify();
+	const bool bMarkDirty = false;
+	Modify(bMarkDirty);
 
 	const EUsdDefaultKind NewEnum = ( EUsdDefaultKind ) NewKindsToCollapse;
 	EUsdDefaultKind Result = NewEnum;
@@ -1362,7 +1366,8 @@ void AUsdStageActor::SetCollapseTopLevelPointInstancers( bool bCollapse )
 
 void AUsdStageActor::SetPurposesToLoad( int32 NewPurposesToLoad )
 {
-	Modify();
+	const bool bMarkDirty = false;
+	Modify(bMarkDirty);
 
 	PurposesToLoad = NewPurposesToLoad;
 	LoadUsdStage();
@@ -1370,7 +1375,8 @@ void AUsdStageActor::SetPurposesToLoad( int32 NewPurposesToLoad )
 
 void AUsdStageActor::SetNaniteTriangleThreshold( int32 NewNaniteTriangleThreshold )
 {
-	Modify();
+	const bool bMarkDirty = false;
+	Modify(bMarkDirty);
 
 	NaniteTriangleThreshold = NewNaniteTriangleThreshold;
 	LoadUsdStage();
@@ -1378,7 +1384,8 @@ void AUsdStageActor::SetNaniteTriangleThreshold( int32 NewNaniteTriangleThreshol
 
 void AUsdStageActor::SetRenderContext( const FName& NewRenderContext )
 {
-	Modify();
+	const bool bMarkDirty = false;
+	Modify(bMarkDirty);
 
 	RenderContext = NewRenderContext;
 	LoadUsdStage();
@@ -1386,7 +1393,8 @@ void AUsdStageActor::SetRenderContext( const FName& NewRenderContext )
 
 void AUsdStageActor::SetTime(float InTime)
 {
-	Modify();
+	const bool bMarkDirty = false;
+	Modify(bMarkDirty);
 
 	Time = InTime;
 	Refresh();
@@ -1818,7 +1826,8 @@ void AUsdStageActor::LoadUsdStage()
 
 void AUsdStageActor::UnloadUsdStage()
 {
-	Modify();
+	const bool bMarkDirty = false;
+	Modify(bMarkDirty);
 
 	FUsdStageActorImpl::DeselectActorsAndComponents( this );
 
@@ -1934,7 +1943,8 @@ void AUsdStageActor::ReloadAnimations()
 		}
 
 		// We need to guarantee we'll record our change of LevelSequence into the transaction, as Init() will create a new one
-		Modify();
+		const bool bMarkDirty = false;
+		Modify(bMarkDirty);
 
 		LevelSequence = LevelSequenceHelper.Init( UsdStage );
 		LevelSequenceHelper.BindToUsdStageActor( this );
