USD: Fix generated LevelSequence names not being made unique when they should be.

diff --git a/Engine/Plugins/Importers/USDImporter/Source/USDStage/Private/USDLevelSequenceHelper.cpp b/Engine/Plugins/Importers/USDImporter/Source/USDStage/Private/USDLevelSequenceHelper.cpp
--- a/Engine/Plugins/Importers/USDImporter/Source/USDStage/Private/USDLevelSequenceHelper.cpp
+++ b/Engine/Plugins/Importers/USDImporter/Source/USDStage/Private/USDLevelSequenceHelper.cpp
@@ -748,8 +748,10 @@ ULevelSequence* FUsdLevelSequenceHelperImpl::FindOrAddSequenceForLayer( const UE
 		// previous one. Also note that the previous level sequence, even though unreferenced by the stage actor, is likely still alive and valid due to references
 		// from the transaction buffer, so we would basically end up creating a identical new object on top of an existing one (the new object has the same address as the existing one).
 		// When importing we don't actually want to do this though, because we want these assets name to conflict so that we can publish/replace old assets if desired. The stage
-		// importer will make these names unique later if needed
-		const bool bIsImporting = StageActor.IsExplicitlyNull();
+		// importer will make these names unique later if needed.
+		// We only get an AssetCache when importing (from UUsdStageImporter::ImportFromFile) or when BindToUsdStageActor is called,
+		// which also gives us a stage actor. So if we don't have an actor but have a cache, we're importing
+		const bool bIsImporting = StageActor.IsExplicitlyNull() && AssetCache;
 		FName UniqueSequenceName = bIsImporting
 			? *UsdLevelSequenceHelperImpl::SanitizeObjectName( FPaths::GetBaseFilename( SequenceDisplayName ) )
 			: MakeUniqueObjectName( GetTransientPackage(), ULevelSequence::StaticClass(), *UsdLevelSequenceHelperImpl::SanitizeObjectName( FPaths::GetBaseFilename( SequenceDisplayName ) ) );
