USD: Fix ensure when responding to a notice about a prim being created.

diff --git a/Engine/Plugins/Importers/USDImporter/Source/USDSchemas/Private/USDInfoCache.cpp b/Engine/Plugins/Importers/USDImporter/Source/USDSchemas/Private/USDInfoCache.cpp
--- a/Engine/Plugins/Importers/USDImporter/Source/USDSchemas/Private/USDInfoCache.cpp
+++ b/Engine/Plugins/Importers/USDImporter/Source/USDSchemas/Private/USDInfoCache.cpp
@@ -34,6 +34,14 @@ bool FUsdInfoCache::Serialize( FArchive& Ar )
 	return true;
 }
 
+bool FUsdInfoCache::ContainsInfoAboutPrim( const UE::FSdfPath& Path ) const
+{
+	return AssetPathsToCollapsedRoot.Contains( Path )
+		|| ComponentPathsToCollapsedRoot.Contains( Path )
+		|| ExpectedVertexCountPerSubtree.Contains( Path )
+		|| ExpectedMaterialSlotCountPerSubtree.Contains( Path );
+}
+
 bool FUsdInfoCache::IsPathCollapsed( const UE::FSdfPath& Path, ECollapsingType CollapsingType ) const
 {
 	FReadScopeLock ScopeLock( Lock );
diff --git a/Engine/Plugins/Importers/USDImporter/Source/USDSchemas/Public/USDInfoCache.h b/Engine/Plugins/Importers/USDImporter/Source/USDSchemas/Public/USDInfoCache.h
--- a/Engine/Plugins/Importers/USDImporter/Source/USDSchemas/Public/USDInfoCache.h
+++ b/Engine/Plugins/Importers/USDImporter/Source/USDSchemas/Public/USDInfoCache.h
@@ -28,6 +28,9 @@ struct USDSCHEMAS_API FUsdInfoCache
 {
 	bool Serialize( FArchive& Ar );
 
+	// Returns whether we contain any info about prim at 'Path' at all
+	bool ContainsInfoAboutPrim( const UE::FSdfPath& Path ) const;
+
 	bool IsPathCollapsed( const UE::FSdfPath& Path, ECollapsingType CollapsingType ) const;
 	bool DoesPathCollapseChildren( const UE::FSdfPath& Path, ECollapsingType CollapsingType ) const;
 
diff --git a/Engine/Plugins/Importers/USDImporter/Source/USDStage/Private/USDStageActor.cpp b/Engine/Plugins/Importers/USDImporter/Source/USDStage/Private/USDStageActor.cpp
--- a/Engine/Plugins/Importers/USDImporter/Source/USDStage/Private/USDStageActor.cpp
+++ b/Engine/Plugins/Importers/USDImporter/Source/USDStage/Private/USDStageActor.cpp
@@ -860,10 +860,15 @@ void AUsdStageActor::OnUsdObjectsChanged( const UsdUtils::FObjectChangesByPath&
 			{
 				UE::FSdfPath ChildPath{ *ChangeIt.Key() };
 
-				NewEntries.Add( InfoCache->UnwindToNonCollapsedPath( ChildPath, ECollapsingType::Assets ).GetString(), true );
-				NewEntries.Add( InfoCache->UnwindToNonCollapsedPath( ChildPath, ECollapsingType::Components ).GetString(), true );
+				// Usually the cache should contain info about any prim on the stage, but if this is a
+				// notice about a prim being *created*, it may not have anything about this prim yet
+				if ( InfoCache->ContainsInfoAboutPrim( ChildPath ) )
+				{
+					NewEntries.Add( InfoCache->UnwindToNonCollapsedPath( ChildPath, ECollapsingType::Assets ).GetString(), true );
+					NewEntries.Add( InfoCache->UnwindToNonCollapsedPath( ChildPath, ECollapsingType::Components ).GetString(), true );
 
-				ChangeIt.RemoveCurrent();
+					ChangeIt.RemoveCurrent();
+				}
 			}
 		}
 
