USD: fix duplicated LOCTEXT id

diff --git a/Engine/Plugins/Importers/USDImporter/Source/USDStageEditor/Private/SUSDStage.cpp b/Engine/Plugins/Importers/USDImporter/Source/USDStageEditor/Private/SUSDStage.cpp
--- a/Engine/Plugins/Importers/USDImporter/Source/USDStageEditor/Private/SUSDStage.cpp
+++ b/Engine/Plugins/Importers/USDImporter/Source/USDStageEditor/Private/SUSDStage.cpp
@@ -981,7 +981,7 @@ void SUsdStage::FillCollapsingSubMenu( FMenuBuilder& MenuBuilder )
 {
 	MenuBuilder.AddMenuEntry(
 		LOCTEXT( "MergeIdenticalMaterialSlots", "Merge identical material slots" ),
-		LOCTEXT( "MergeIdenticalMaterialSlots", "If enabled, when multiple mesh prims are collapsed into a single static mesh, identical material slots are merged into one slot.\nOtherwise, material slots are simply appended to the list." ),
+		LOCTEXT( "MergeIdenticalMaterialSlots_ToolTip", "If enabled, when multiple mesh prims are collapsed into a single static mesh, identical material slots are merged into one slot.\nOtherwise, material slots are simply appended to the list." ),
 		FSlateIcon(),
 		FUIAction(
 			FExecuteAction::CreateLambda( [this]()
