USD: Fix test_inherited_visibilities_import occasionally failing.

diff --git a/Engine/Plugins/Importers/USDImporter/Source/USDStage/Private/USDLevelSequenceHelper.cpp b/Engine/Plugins/Importers/USDImporter/Source/USDStage/Private/USDLevelSequenceHelper.cpp
--- a/Engine/Plugins/Importers/USDImporter/Source/USDStage/Private/USDLevelSequenceHelper.cpp
+++ b/Engine/Plugins/Importers/USDImporter/Source/USDStage/Private/USDLevelSequenceHelper.cpp
@@ -364,8 +364,8 @@ private:
 	static const TCHAR* TimeTrackName;
 	static const double EmptySubSectionRange; // How many frames should an empty subsection cover, only needed so that the subsection is visible and the user can edit it
 
-	TWeakObjectPtr<AUsdStageActor> StageActor;
-	UUsdAssetCache* AssetCache;  // We keep a pointer to this directly because we may be called via the USDStageImporter directly, when we don't have an available actor
+	TWeakObjectPtr<AUsdStageActor> StageActor = nullptr;
+	UUsdAssetCache* AssetCache = nullptr;  // We keep a pointer to this directly because we may be called via the USDStageImporter directly, when we don't have an available actor
 	FGuid StageActorBinding;
 
 	// Only when this is zero we write LevelSequence object (tracks, moviescene, sections, etc.) transactions back to the USD stage
diff --git a/Engine/Plugins/Importers/USDImporter/Source/USDStageImporter/Private/USDStageImporter.cpp b/Engine/Plugins/Importers/USDImporter/Source/USDStageImporter/Private/USDStageImporter.cpp
--- a/Engine/Plugins/Importers/USDImporter/Source/USDStageImporter/Private/USDStageImporter.cpp
+++ b/Engine/Plugins/Importers/USDImporter/Source/USDStageImporter/Private/USDStageImporter.cpp
@@ -1467,8 +1467,6 @@ void UUsdStageImporter::ImportFromFile(FUsdStageImportContext& ImportContext)
 
 	UsdStageImporterImpl::SetupStageForImport( ImportContext );
 
-	ImportContext.LevelSequenceHelper.Init( ImportContext.Stage );
-
 	TMap<FSoftObjectPath, FSoftObjectPath> SoftObjectsToRemap;
 	TMap<UObject*, UObject*> ObjectsToRemap;
 	TSet<UObject*> UsedAssetsAndDependencies;
@@ -1481,6 +1479,7 @@ void UUsdStageImporter::ImportFromFile(FUsdStageImportContext& ImportContext)
 	}
 	ImportContext.AssetCache->MarkAssetsAsStale();
 	ImportContext.LevelSequenceHelper.SetAssetCache( ImportContext.AssetCache );
+	ImportContext.LevelSequenceHelper.Init( ImportContext.Stage );  // Must happen after the context gets an AssetCache!
 
 	// Shotgun approach to recreate all render states because we may want to reimport/delete/reassign a material/static/skeletalmesh while it is currently being drawn
 	FGlobalComponentRecreateRenderStateContext RecreateRenderStateContext;
