USD: Fix Mesh prim transforms being baked into the mesh during import.

diff --git a/Engine/Plugins/Importers/USDImporter/Source/USDSchemas/Private/USDGeomMeshTranslator.cpp b/Engine/Plugins/Importers/USDImporter/Source/USDSchemas/Private/USDGeomMeshTranslator.cpp
--- a/Engine/Plugins/Importers/USDImporter/Source/USDSchemas/Private/USDGeomMeshTranslator.cpp
+++ b/Engine/Plugins/Importers/USDImporter/Source/USDSchemas/Private/USDGeomMeshTranslator.cpp
@@ -539,7 +539,7 @@ namespace UsdGeomMeshTranslatorImpl
 		return bFoundLODs;
 	}
 
-	void LoadMeshDescriptions( const pxr::UsdTyped& UsdMesh, TArray<FMeshDescription>& OutLODIndexToMeshDescription, TArray<UsdUtils::FUsdPrimMaterialAssignmentInfo>& OutLODIndexToMaterialInfo,
+	void LoadMeshDescriptions( pxr::UsdTyped UsdMesh, TArray<FMeshDescription>& OutLODIndexToMeshDescription, TArray<UsdUtils::FUsdPrimMaterialAssignmentInfo>& OutLODIndexToMaterialInfo,
 		const TMap< FString, TMap< FString, int32 > >& MaterialToPrimvarToUVIndex, const FTransform& AdditionalTransform, const pxr::UsdTimeCode TimeCode, bool bInterpretLODs, const FName& RenderContext, bool bMergeIdenticalMaterialSlots = true )
 	{
 		if ( !UsdMesh )
@@ -547,16 +547,38 @@ namespace UsdGeomMeshTranslatorImpl
 			return;
 		}
 
+		FScopedUsdAllocs Allocs;
+
 		pxr::TfToken RenderContextToken = pxr::UsdShadeTokens->universalRenderContext;
 		if ( !RenderContext.IsNone() )
 		{
 			RenderContextToken = UnrealToUsd::ConvertToken( *RenderContext.ToString() ).Get();
 		}
 
+		pxr::UsdPrim Prim = UsdMesh.GetPrim();
+		pxr::UsdStageRefPtr Stage = Prim.GetStage();
+		pxr::SdfPath Path = Prim.GetPrimPath();
+
 		bool bInterpretedLODs = false;
 		if ( bInterpretLODs )
 		{
 			bInterpretedLODs = TryLoadingMultipleLODs( UsdMesh, OutLODIndexToMeshDescription, OutLODIndexToMaterialInfo, MaterialToPrimvarToUVIndex, AdditionalTransform, TimeCode, RenderContextToken, bMergeIdenticalMaterialSlots );
+
+			// Have to be very careful here as flipping through LODs invalidates prim references, so we need to
+			// re-acquire them
+			UsdMesh = pxr::UsdGeomMesh{ Stage->GetPrimAtPath( Path ) };
+		}
+
+		// If we've managed to interpret LODs, we won't place our mesh transform on the static mesh component itself
+		// (c.f. FUsdGeomXformableTranslator::UpdateComponents), and will instead expect it to be baked into the mesh.
+		// So here we do that
+		bool bSuccess = true;
+		FTransform MeshTransform = FTransform::Identity;
+		if ( bInterpretedLODs && OutLODIndexToMeshDescription.Num() > 1 )
+		{
+			// TODO: Handle resetXformOp here
+			bool bResetXformStack = false;
+			bSuccess &= UsdToUnreal::ConvertXformable( Stage, UsdMesh, MeshTransform, TimeCode.GetValue(), &bResetXformStack );
 		}
 
 		if ( !bInterpretedLODs )
@@ -567,17 +589,6 @@ namespace UsdGeomMeshTranslatorImpl
 			FStaticMeshAttributes StaticMeshAttributes( TempMeshDescription );
 			StaticMeshAttributes.Register();
 
-			// If we're being *asked* to intrepret LODs (regardless of whether that succeeded or not), we won't place our mesh transform on
-			// the static mesh component itself, and will instead expect it to be baked into the mesh. So here we do that
-			bool bSuccess = true;
-			FTransform MeshTransform = FTransform::Identity;
-			if ( bInterpretLODs )
-			{
-				// TODO: Handle resetXformOp here
-				bool bResetXformStack = false;
-				bSuccess &= UsdToUnreal::ConvertXformable( UsdMesh.GetPrim().GetStage(), UsdMesh, MeshTransform, TimeCode.GetValue(), &bResetXformStack );
-			}
-
 			if ( bSuccess )
 			{
 				UsdToUnreal::FUsdMeshConversionOptions Options;
