USD: Python sample script to trigger the LiveLink support code.

diff --git a/Engine/Plugins/Importers/USDImporter/Resources/PythonExamples/enable_live_link.py b/Engine/Plugins/Importers/USDImporter/Resources/PythonExamples/enable_live_link.py
new file mode 100644
--- /dev/null
+++ b/Engine/Plugins/Importers/USDImporter/Resources/PythonExamples/enable_live_link.py
@@ -0,0 +1,35 @@
+"""
+This script describes how to enable LiveLink for the components generated for a particular prim when
+imported in UE.
+
+After the script is executed, the stage can be saved and the connection details will persist on the USD
+files themselves. Opening the stage again will automatically re-establish the connection.
+
+To use this example, follow these steps:
+- Enable the "USD Importer" and "Live Link" plugins in UE;
+- Open the USD Stage editor window by going to Window -> Virtual Production -> USD Stage;
+- Open your desired USD Stage by going to File -> Open on the USD Stage editor and picking a USD file;
+- Edit the script below, replacing the prim path, AnimBlueprint path and live link subject names with
+  your desired values;
+- Run the script, either by copy-pasting it into the UE Python console, or by pasting the path to this
+  file onto the same console.
+"""
+
+from pxr import Usd, UsdUtils, Sdf
+
+stage = UsdUtils.StageCache().Get().GetAllStages()[0]
+prim = stage.GetPrimAtPath("/cube_default")
+
+schema = Usd.SchemaRegistry.GetTypeFromSchemaTypeName("LiveLinkAPI")
+prim.ApplyAPI(schema)
+
+with Sdf.ChangeBlock():
+    anim_bp_attr = prim.CreateAttribute("unrealLiveLink:animBlueprintPath", Sdf.ValueTypeNames.String)
+    anim_bp_attr.Set("/USDImporter/Blueprint/DefaultLiveLinkAnimBP.DefaultLiveLinkAnimBP")
+
+    subject_name_attr = prim.CreateAttribute("unrealLiveLink:subjectName", Sdf.ValueTypeNames.String)
+    subject_name_attr.Set("cube")
+
+    enabled_attr = prim.CreateAttribute("unrealLiveLink:enabled", Sdf.ValueTypeNames.Bool)
+    enabled_attr.Set(True)
+
