USD: Cleanup done as part of the incremental parsing cl (19700631).

diff --git a/Engine/Plugins/Importers/USDImporter/Source/USDSchemas/Private/USDGeomMeshTranslator.cpp b/Engine/Plugins/Importers/USDImporter/Source/USDSchemas/Private/USDGeomMeshTranslator.cpp
--- a/Engine/Plugins/Importers/USDImporter/Source/USDSchemas/Private/USDGeomMeshTranslator.cpp
+++ b/Engine/Plugins/Importers/USDImporter/Source/USDSchemas/Private/USDGeomMeshTranslator.cpp
@@ -453,7 +453,12 @@ namespace UsdGeomMeshTranslatorImpl
 	// Note that these other LODs will be hidden in other variants, and won't show up on traversal unless we actively switch the variants (which we do here).
 	// We use a separate function for this because there is a very specific set of conditions where we successfully can do this, and we
 	// want to fall back to just parsing UsdMesh as a simple single-LOD mesh if we fail.
-	bool TryLoadingMultipleLODs( const pxr::UsdTyped& UsdMesh, TArray<FMeshDescription>& OutLODIndexToMeshDescription, TArray<UsdUtils::FUsdPrimMaterialAssignmentInfo>& OutLODIndexToMaterialInfo, const TMap< FString, TMap< FString, int32 > >& InMaterialToPrimvarToUVIndex, const FTransform& AdditionalTransform, const pxr::UsdTimeCode InTimeCode, const pxr::TfToken& RenderContext, bool bMergeIdenticalMaterialSlots )
+	bool TryLoadingMultipleLODs(
+		const pxr::UsdTyped& UsdMesh,
+		TArray<FMeshDescription>& OutLODIndexToMeshDescription,
+		TArray<UsdUtils::FUsdPrimMaterialAssignmentInfo>& OutLODIndexToMaterialInfo,
+		const UsdToUnreal::FUsdMeshConversionOptions& Options
+	)
 	{
 		FScopedUsdAllocs Allocs;
 
@@ -472,6 +477,8 @@ namespace UsdGeomMeshTranslatorImpl
 		TMap<int32, FMeshDescription> LODIndexToMeshDescriptionMap;
 		TMap<int32, UsdUtils::FUsdPrimMaterialAssignmentInfo> LODIndexToMaterialInfoMap;
 
+		UsdToUnreal::FUsdMeshConversionOptions OptionsCopy = Options;
+
 		TFunction<bool( const pxr::UsdGeomMesh&, int32 )> ConvertLOD = [ & ]( const pxr::UsdGeomMesh& LODMesh, int32 LODIndex )
 		{
 			FMeshDescription TempMeshDescription;
@@ -487,7 +494,7 @@ namespace UsdGeomMeshTranslatorImpl
 			// care if this mesh in particular has been marked as invisible
 			pxr::TfToken Visibility;
 			pxr::UsdAttribute VisibilityAttr = LODMesh.GetVisibilityAttr();
-			if ( VisibilityAttr && VisibilityAttr.Get( &Visibility, InTimeCode ) && Visibility == pxr::UsdGeomTokens->inherited )
+			if ( VisibilityAttr && VisibilityAttr.Get( &Visibility, Options.TimeCode ) && Visibility == pxr::UsdGeomTokens->inherited )
 			{
 				// If we're interpreting LODs we must bake the transform from each LOD Mesh into the vertices, because there's no guarantee
 				// all LODs have the same transform, so we can't just put the transforms directly on the component. If we are not interpreting
@@ -495,22 +502,17 @@ namespace UsdGeomMeshTranslatorImpl
 				// TODO: Handle resetXformOp here
 				bool bResetXformStack = false;
 				FTransform MeshTransform = FTransform::Identity;
-				bSuccess &= UsdToUnreal::ConvertXformable( LODMesh.GetPrim().GetStage(), LODMesh, MeshTransform, InTimeCode.GetValue(), &bResetXformStack );
+				bSuccess &= UsdToUnreal::ConvertXformable( LODMesh.GetPrim().GetStage(), LODMesh, MeshTransform, Options.TimeCode.GetValue(), &bResetXformStack );
 
 				if ( bSuccess )
 				{
-					UsdToUnreal::FUsdMeshConversionOptions Options;
-					Options.AdditionalTransform = MeshTransform * AdditionalTransform;
-					Options.TimeCode = InTimeCode;
-					Options.RenderContext = RenderContext;
-					Options.bMergeIdenticalMaterialSlots = bMergeIdenticalMaterialSlots;
-					Options.MaterialToPrimvarToUVIndex = &InMaterialToPrimvarToUVIndex;
+					OptionsCopy.AdditionalTransform = MeshTransform * Options.AdditionalTransform;
 
 					bSuccess &= UsdToUnreal::ConvertGeomMesh(
 						LODMesh,
 						TempMeshDescription,
 						TempMaterialInfo,
-						Options
+						OptionsCopy
 					);
 				}
 			}
@@ -539,8 +541,13 @@ namespace UsdGeomMeshTranslatorImpl
 		return bFoundLODs;
 	}
 
-	void LoadMeshDescriptions( pxr::UsdTyped UsdMesh, TArray<FMeshDescription>& OutLODIndexToMeshDescription, TArray<UsdUtils::FUsdPrimMaterialAssignmentInfo>& OutLODIndexToMaterialInfo,
-		const TMap< FString, TMap< FString, int32 > >& MaterialToPrimvarToUVIndex, const FTransform& AdditionalTransform, const pxr::UsdTimeCode TimeCode, bool bInterpretLODs, const FName& RenderContext, bool bMergeIdenticalMaterialSlots = true )
+	void LoadMeshDescriptions(
+		pxr::UsdTyped UsdMesh,
+		TArray<FMeshDescription>& OutLODIndexToMeshDescription,
+		TArray<UsdUtils::FUsdPrimMaterialAssignmentInfo>& OutLODIndexToMaterialInfo,
+		const UsdToUnreal::FUsdMeshConversionOptions& Options,
+		bool bInterpretLODs = false
+	)
 	{
 		if ( !UsdMesh )
 		{
@@ -549,12 +556,6 @@ namespace UsdGeomMeshTranslatorImpl
 
 		FScopedUsdAllocs Allocs;
 
-		pxr::TfToken RenderContextToken = pxr::UsdShadeTokens->universalRenderContext;
-		if ( !RenderContext.IsNone() )
-		{
-			RenderContextToken = UnrealToUsd::ConvertToken( *RenderContext.ToString() ).Get();
-		}
-
 		pxr::UsdPrim Prim = UsdMesh.GetPrim();
 		pxr::UsdStageRefPtr Stage = Prim.GetStage();
 		pxr::SdfPath Path = Prim.GetPrimPath();
@@ -562,7 +563,7 @@ namespace UsdGeomMeshTranslatorImpl
 		bool bInterpretedLODs = false;
 		if ( bInterpretLODs )
 		{
-			bInterpretedLODs = TryLoadingMultipleLODs( UsdMesh, OutLODIndexToMeshDescription, OutLODIndexToMaterialInfo, MaterialToPrimvarToUVIndex, AdditionalTransform, TimeCode, RenderContextToken, bMergeIdenticalMaterialSlots );
+			bInterpretedLODs = TryLoadingMultipleLODs( UsdMesh, OutLODIndexToMeshDescription, OutLODIndexToMaterialInfo, Options );
 
 			// Have to be very careful here as flipping through LODs invalidates prim references, so we need to
 			// re-acquire them
@@ -578,7 +579,7 @@ namespace UsdGeomMeshTranslatorImpl
 		{
 			// TODO: Handle resetXformOp here
 			bool bResetXformStack = false;
-			bSuccess &= UsdToUnreal::ConvertXformable( Stage, UsdMesh, MeshTransform, TimeCode.GetValue(), &bResetXformStack );
+			bSuccess &= UsdToUnreal::ConvertXformable( Stage, UsdMesh, MeshTransform, Options.TimeCode.GetValue(), &bResetXformStack );
 		}
 
 		if ( !bInterpretedLODs )
@@ -591,18 +592,14 @@ namespace UsdGeomMeshTranslatorImpl
 
 			if ( bSuccess )
 			{
-				UsdToUnreal::FUsdMeshConversionOptions Options;
-				Options.AdditionalTransform = MeshTransform * AdditionalTransform;
-				Options.TimeCode = TimeCode;
-				Options.RenderContext = RenderContextToken;
-				Options.bMergeIdenticalMaterialSlots = bMergeIdenticalMaterialSlots;
-				Options.MaterialToPrimvarToUVIndex = &MaterialToPrimvarToUVIndex;
+				UsdToUnreal::FUsdMeshConversionOptions OptionsCopy = Options;
+				OptionsCopy.AdditionalTransform = MeshTransform * Options.AdditionalTransform;
 
 				bSuccess &= UsdToUnreal::ConvertGeomMesh(
 					pxr::UsdGeomMesh{ UsdMesh },
 					TempMeshDescription,
 					TempMaterialInfo,
-					Options
+					OptionsCopy
 				);
 			}
 
@@ -885,6 +882,17 @@ namespace UsdGeomMeshTranslatorImpl
 			// The GeometryCache module expects the end frame to be one past the last animation frame
 			EndFrame += 1;
 
+			pxr::TfToken RenderContextToken = pxr::UsdShadeTokens->universalRenderContext;
+			if ( !Context->RenderContext.IsNone() )
+			{
+				RenderContextToken = UnrealToUsd::ConvertToken( *Context->RenderContext.ToString() ).Get();
+			}
+
+			UsdToUnreal::FUsdMeshConversionOptions Options;
+			Options.PurposesToLoad = Context->PurposesToLoad;
+			Options.bMergeIdenticalMaterialSlots = Context->bMergeIdenticalMaterialSlots;
+			Options.RenderContext = RenderContextToken;
+
 			// Create and configure a new USDTrack to be added to the GeometryCache
 			UGeometryCacheTrackUsd* UsdTrack = NewObject< UGeometryCacheTrackUsd >( GeometryCache );
 			UsdTrack->Initialize(
@@ -894,7 +902,7 @@ namespace UsdGeomMeshTranslatorImpl
 				*MaterialToPrimvarToUVIndex,
 				StartFrame,
 				EndFrame,
-				[InPrimPath]( const TWeakObjectPtr<UGeometryCacheTrackUsd> TrackPtr, float Time, FGeometryCacheMeshData& OutMeshData )
+				[InPrimPath, Options]( const TWeakObjectPtr<UGeometryCacheTrackUsd> TrackPtr, float Time, FGeometryCacheMeshData& OutMeshData ) mutable
 				{
 					UGeometryCacheTrackUsd* Track = TrackPtr.Get();
 					if ( !Track )
@@ -919,15 +927,15 @@ namespace UsdGeomMeshTranslatorImpl
 					TArray< UsdUtils::FUsdPrimMaterialAssignmentInfo > LODIndexToMaterialInfo;
 					const bool bAllowInterpretingLODs = false;  // GeometryCaches don't have LODs, so we will never do this
 
+					Options.TimeCode = pxr::UsdTimeCode{ Time };
+					Options.MaterialToPrimvarToUVIndex = &Track->MaterialToPrimvarToUVIndex;
+
 					UsdGeomMeshTranslatorImpl::LoadMeshDescriptions(
 						pxr::UsdTyped( Prim ),
 						LODIndexToMeshDescription,
 						LODIndexToMaterialInfo,
-						Track->MaterialToPrimvarToUVIndex,
-						FTransform::Identity,
-						pxr::UsdTimeCode( Time ),
-						bAllowInterpretingLODs,
-						Track->RenderContext
+						Options,
+						bAllowInterpretingLODs
 					);
 
 					// Convert the MeshDescription to MeshData
@@ -1107,7 +1115,7 @@ void FBuildStaticMeshTaskChain::SetupTasks()
 			}
 			else
 			{
-				MeshName = PrimPath.GetString();
+				MeshName = PrimPathString;
 			}
 
 			bool bIsNew = true;
@@ -1266,16 +1274,26 @@ void FGeomMeshCreateAssetsTaskChain::SetupTasks()
 			TMap< FString, TMap< FString, int32 > > Unused;
 			TMap< FString, TMap< FString, int32 > >* MaterialToPrimvarToUVIndex = Context->MaterialToPrimvarToUVIndex ? Context->MaterialToPrimvarToUVIndex : &Unused;
 
+			pxr::TfToken RenderContextToken = pxr::UsdShadeTokens->universalRenderContext;
+			if ( !Context->RenderContext.IsNone() )
+			{
+				RenderContextToken = UnrealToUsd::ConvertToken( *Context->RenderContext.ToString() ).Get();
+			}
+
+			UsdToUnreal::FUsdMeshConversionOptions Options;
+			Options.TimeCode = Context->Time;
+			Options.PurposesToLoad = Context->PurposesToLoad;
+			Options.RenderContext = RenderContextToken;
+			Options.MaterialToPrimvarToUVIndex = MaterialToPrimvarToUVIndex;
+			Options.bMergeIdenticalMaterialSlots = Context->bMergeIdenticalMaterialSlots;
+			Options.AdditionalTransform = AdditionalTransform;
+
 			UsdGeomMeshTranslatorImpl::LoadMeshDescriptions(
 				pxr::UsdTyped( GetPrim() ),
 				LODIndexToMeshDescription,
 				LODIndexToMaterialInfo,
-				*MaterialToPrimvarToUVIndex,
-				AdditionalTransform,
-				pxr::UsdTimeCode( Context->Time ),
-				Context->bAllowInterpretingLODs,
-				Context->RenderContext,
-				Context->bMergeIdenticalMaterialSlots
+				Options,
+				Context->bAllowInterpretingLODs
 			);
 
 			// If we have at least one valid LOD, we should keep going
@@ -1317,20 +1335,29 @@ void FGeometryCacheCreateAssetsTaskChain::SetupTasks()
 			// Always hash the mesh at the same time because it may be animated, and
 			// otherwise we may think it's a new asset just because the context is at a different timecode (e.g. if we reload)
 			// TODO: Hash all timecodes, or else our mesh may change at t=5 and we never reload it because we only hash t=0
-			const double TimeCode = UsdUtils::GetEarliestTimeCode();
 			const bool bAllowInterpretingLODs = false;  // GeometryCaches don't have LODs
 			TMap< FString, TMap< FString, int32 > > Unused;
 			TMap< FString, TMap< FString, int32 > >* MaterialToPrimvarToUVIndex = Context->MaterialToPrimvarToUVIndex ? Context->MaterialToPrimvarToUVIndex : &Unused;
 
+			pxr::TfToken RenderContextToken = pxr::UsdShadeTokens->universalRenderContext;
+			if ( !Context->RenderContext.IsNone() )
+			{
+				RenderContextToken = UnrealToUsd::ConvertToken( *Context->RenderContext.ToString() ).Get();
+			}
+
+			UsdToUnreal::FUsdMeshConversionOptions Options;
+			Options.TimeCode = UsdUtils::GetEarliestTimeCode();
+			Options.PurposesToLoad = Context->PurposesToLoad;
+			Options.RenderContext = RenderContextToken;
+			Options.MaterialToPrimvarToUVIndex = MaterialToPrimvarToUVIndex;
+			Options.bMergeIdenticalMaterialSlots = Context->bMergeIdenticalMaterialSlots;
+
 			UsdGeomMeshTranslatorImpl::LoadMeshDescriptions(
 				pxr::UsdTyped( GetPrim() ),
 				LODIndexToMeshDescription,
 				LODIndexToMaterialInfo,
-				*MaterialToPrimvarToUVIndex,
-				FTransform::Identity,
-				pxr::UsdTimeCode( TimeCode ),
-				bAllowInterpretingLODs,
-				Context->RenderContext
+				Options,
+				bAllowInterpretingLODs
 			);
 
 			// If we have at least one valid LOD, we should keep going
diff --git a/Engine/Plugins/Importers/USDImporter/Source/USDSchemas/Private/USDGeomPointInstancerTranslator.cpp b/Engine/Plugins/Importers/USDImporter/Source/USDSchemas/Private/USDGeomPointInstancerTranslator.cpp
--- a/Engine/Plugins/Importers/USDImporter/Source/USDSchemas/Private/USDGeomPointInstancerTranslator.cpp
+++ b/Engine/Plugins/Importers/USDImporter/Source/USDSchemas/Private/USDGeomPointInstancerTranslator.cpp
@@ -224,12 +224,10 @@ USceneComponent* FUsdGeomPointInstancerTranslator::CreateComponents()
 {
 	TRACE_CPUPROFILER_EVENT_SCOPE( FUsdGeomPointInstancerTranslator::CreateComponents );
 
-	// If we're going to be collapsing ourselves, just make a static mesh component that can receive the collapsed mesh
+	// If we're collapsing ourselves, we're really just a collapsed Xform prim, so let that translator handle it
 	if ( Context->bCollapseTopLevelPointInstancers )
 	{
-		USceneComponent* MeshComponent = CreateComponentsEx( { UStaticMeshComponent::StaticClass() }, {} );
-		UpdateComponents( MeshComponent );
-		return MeshComponent;
+		return FUsdGeomXformableTranslator::CreateComponents();
 	}
 
 	// Otherwise, the plan here is to create an USceneComponent that corresponds to the PointInstancer prim itself, and then spawn a child
diff --git a/Engine/Plugins/Importers/USDImporter/Source/USDSchemas/Private/USDGeomXformableTranslator.cpp b/Engine/Plugins/Importers/USDImporter/Source/USDSchemas/Private/USDGeomXformableTranslator.cpp
--- a/Engine/Plugins/Importers/USDImporter/Source/USDSchemas/Private/USDGeomXformableTranslator.cpp
+++ b/Engine/Plugins/Importers/USDImporter/Source/USDSchemas/Private/USDGeomXformableTranslator.cpp
@@ -20,7 +20,6 @@
 #include "Engine/StaticMesh.h"
 #include "Engine/World.h"
 #include "Modules/ModuleManager.h"
-#include "StaticMeshAttributes.h"
 
 #include "UsdWrappers/SdfPath.h"
 #include "UsdWrappers/UsdGeomXformable.h"
@@ -36,6 +35,7 @@
 	#include "pxr/usd/usd/primRange.h"
 	#include "pxr/usd/usd/variantSets.h"
 	#include "pxr/usd/usdGeom/mesh.h"
+	#include "pxr/usd/usdGeom/pointInstancer.h"
 	#include "pxr/usd/usdGeom/subset.h"
 	#include "pxr/usd/usdGeom/xformable.h"
 	#include "pxr/usd/usdShade/materialBindingAPI.h"
@@ -322,6 +322,12 @@ USceneComponent* FUsdGeomXformableTranslator::CreateComponentsEx( TOptional< TSu
 						ComponentType = UStaticMeshComponent::StaticClass();
 					}
 				}
+				// If this is a component for a point instancer that just collapsed itself into a static mesh, just make
+				// a static mesh component that can receive it
+				else if ( Context->bCollapseTopLevelPointInstancers && pxr::UsdPrim{ Prim }.IsA<pxr::UsdGeomPointInstancer>() )
+				{
+					ComponentType = UStaticMeshComponent::StaticClass();
+				}
 			}
 		}
 
