USD: Add option to enable or disable LiveLink at the stage editor;

Cleanup and prevent repeated widget recreation inside SUSDPrimInfo every time a prim is selected;
Fix duplicate call to FUsdSkelRootTranslator::UpdateComponents when creating skeletal mesh components;

diff --git a/Engine/Plugins/Animation/LiveLink/Source/LiveLinkGraphNode/LiveLinkGraphNode.Build.cs b/Engine/Plugins/Animation/LiveLink/Source/LiveLinkGraphNode/LiveLinkGraphNode.Build.cs
--- a/Engine/Plugins/Animation/LiveLink/Source/LiveLinkGraphNode/LiveLinkGraphNode.Build.cs
+++ b/Engine/Plugins/Animation/LiveLink/Source/LiveLinkGraphNode/LiveLinkGraphNode.Build.cs
@@ -33,13 +33,19 @@ namespace UnrealBuildTool.Rules
 
 			if (Target.bBuildEditor == true)
 			{
+				PublicDependencyModuleNames.AddRange(
+					new string[]
+					{
+						"AnimGraph",
+					}
+				);
+
 				PrivateDependencyModuleNames.AddRange(
 					new string[]
 					{
 						"EditorFramework",
 						"UnrealEd",
 						"Kismet",
-						"AnimGraph",
 						"BlueprintGraph",
 					}
 				);
diff --git a/Engine/Plugins/Animation/LiveLink/Source/LiveLinkGraphNode/Private/AnimGraphNode_LiveLinkPose.h b/Engine/Plugins/Animation/LiveLink/Source/LiveLinkGraphNode/Private/AnimGraphNode_LiveLinkPose.h
deleted file mode 100644
diff --git a/Engine/Plugins/Animation/LiveLink/Source/LiveLinkGraphNode/Public/AnimGraphNode_LiveLinkPose.h b/Engine/Plugins/Animation/LiveLink/Source/LiveLinkGraphNode/Public/AnimGraphNode_LiveLinkPose.h
new file mode 100644
--- /dev/null
+++ b/Engine/Plugins/Animation/LiveLink/Source/LiveLinkGraphNode/Public/AnimGraphNode_LiveLinkPose.h
@@ -0,0 +1,35 @@
+// Copyright Epic Games, Inc. All Rights Reserved.
+
+#pragma once
+
+#include "AnimGraphNode_Base.h"
+#include "AnimNode_LiveLinkPose.h"
+
+#include "AnimGraphNode_LiveLinkPose.generated.h"
+
+
+
+UCLASS()
+class LIVELINKGRAPHNODE_API UAnimGraphNode_LiveLinkPose : public UAnimGraphNode_Base
+{
+	GENERATED_BODY()
+
+public:
+
+	//~ UEdGraphNode interface
+	virtual FText GetNodeTitle(ENodeTitleType::Type TitleType) const override;
+	virtual FText GetTooltipText() const override;
+	virtual FText GetMenuCategory() const;
+	//~ End of UEdGraphNode
+
+	//~ Begin UK2Node interface
+	virtual void ConvertDeprecatedNode(UEdGraph* Graph, bool bOnlySafeChanges) override;
+	virtual void GetOutputLinkAttributes(FNodeAttributeArray& OutAttributes) const override;
+	//~ End UK2Node interface
+
+public:
+
+	UPROPERTY(EditAnywhere, Category = Settings)
+	FAnimNode_LiveLinkPose Node;
+
+};
diff --git a/Engine/Plugins/Importers/USDImporter/Source/USDSchemas/Private/USDGeomMeshTranslator.cpp b/Engine/Plugins/Importers/USDImporter/Source/USDSchemas/Private/USDGeomMeshTranslator.cpp
--- a/Engine/Plugins/Importers/USDImporter/Source/USDSchemas/Private/USDGeomMeshTranslator.cpp
+++ b/Engine/Plugins/Importers/USDImporter/Source/USDSchemas/Private/USDGeomMeshTranslator.cpp
@@ -1585,7 +1585,7 @@ bool FUsdGeomMeshTranslator::CanBeCollapsed( ECollapsingType CollapsingType ) co
 		return false;
 	}
 
-	return !UsdGeomMeshTranslatorImpl::IsAnimated( Prim );
+	return Super::CanBeCollapsed( CollapsingType );
 }
 
 #endif // #if USE_USD_SDK
diff --git a/Engine/Plugins/Importers/USDImporter/Source/USDSchemas/Private/USDGeomXformableTranslator.cpp b/Engine/Plugins/Importers/USDImporter/Source/USDSchemas/Private/USDGeomXformableTranslator.cpp
--- a/Engine/Plugins/Importers/USDImporter/Source/USDSchemas/Private/USDGeomXformableTranslator.cpp
+++ b/Engine/Plugins/Importers/USDImporter/Source/USDSchemas/Private/USDGeomXformableTranslator.cpp
@@ -8,6 +8,7 @@
 #include "USDErrorUtils.h"
 #include "USDGeomMeshConversion.h"
 #include "USDGeomMeshTranslator.h"
+#include "USDIntegrationUtils.h"
 #include "USDLog.h"
 #include "USDMemory.h"
 #include "USDPrimConversion.h"
@@ -16,10 +17,15 @@
 
 #include "Components/HierarchicalInstancedStaticMeshComponent.h"
 #include "Components/SceneComponent.h"
+#include "Components/SkeletalMeshComponent.h"
 #include "Components/StaticMeshComponent.h"
 #include "Engine/StaticMesh.h"
 #include "Engine/World.h"
+#include "LiveLinkComponentController.h"
+#include "LiveLinkRole.h"
 #include "Modules/ModuleManager.h"
+#include "Roles/LiveLinkTransformRole.h"
+#include "StaticMeshAttributes.h"
 
 #include "UsdWrappers/SdfPath.h"
 #include "UsdWrappers/UsdGeomXformable.h"
@@ -59,6 +65,135 @@ static FAutoConsoleVariableRef CVarEnableCollision(
 	GEnableCollision,
 	TEXT( "Whether to have collision enabled for spawned components and generated meshes" ) );
 
+namespace UE::UsdXformableTranslatorImpl::Private
+{
+	void SetUpSceneComponentForLiveLink( const FUsdSchemaTranslationContext& Context, USceneComponent* Component, const pxr::UsdPrim& Prim )
+	{
+		if ( !Component || !Prim )
+		{
+			return;
+		}
+
+		AActor* Parent = Component->GetOwner();
+		if ( !Parent )
+		{
+			return;
+		}
+
+		USceneComponent* RootComponent = Parent->GetRootComponent();
+		if ( !RootComponent )
+		{
+			return;
+		}
+
+		ULiveLinkComponentController* Controller = nullptr;
+		{
+			// We would have to traverse all top-level actor components to know if our component is set up for live link already
+			// or not, so this just helps us make that a little bit faster. Its important because UpdateComponents (which calls us)
+			// is the main function that is called to animate components, so it can be spammed in case this prim has animations
+			static TMap<TWeakObjectPtr<USceneComponent>, TWeakObjectPtr<ULiveLinkComponentController>> LiveLinkEnabledComponents;
+			if ( ULiveLinkComponentController* ExistingController = LiveLinkEnabledComponents.FindRef( Component ).Get() )
+			{
+				// We found an existing controller we created to track this component, so use that
+				Controller = ExistingController;
+			}
+			// We don't know of any controllers handling this component yet, get a new one
+			else
+			{
+				TArray<ULiveLinkComponentController*> LiveLinkComponents;
+				Parent->GetComponents< ULiveLinkComponentController >( LiveLinkComponents );
+
+				for ( ULiveLinkComponentController* LiveLinkComponent : LiveLinkComponents )
+				{
+					if ( LiveLinkComponent->ComponentToControl.GetComponent( Parent ) == Component )
+					{
+						// We found some other controller handling this component somehow, use that
+						Controller = LiveLinkComponent;
+						break;
+					}
+				}
+
+				if ( !Controller )
+				{
+					// We'll get a warning from the live link controller component in case the component its controlling is not movable
+					Component->Mobility = EComponentMobility::Movable;
+
+					Controller = NewObject<ULiveLinkComponentController>( Parent, NAME_None, Context.ObjectFlags );
+					Controller->bUpdateInEditor = true;
+					Controller->ComponentToControl.PathToComponent = Component->GetPathName( Parent );
+
+					// Important because of how ULiveLinkComponentController::TickComponent also checks for the sequencer
+					// tag to try and guess if the controlled component is a spawnable
+					Controller->bDisableEvaluateLiveLinkWhenSpawnable = false;
+
+					Parent->AddInstanceComponent( Controller );
+					Controller->RegisterComponent();
+				}
+
+				if ( Controller )
+				{
+					LiveLinkEnabledComponents.Add( Component, Controller );
+				}
+			}
+		}
+
+		// Configure controller with our desired parameters
+		if ( Controller )
+		{
+			FScopedUsdAllocs Allocs;
+
+			FLiveLinkSubjectRepresentation SubjectRepresentation;
+			SubjectRepresentation.Role = ULiveLinkTransformRole::StaticClass();
+
+			if ( pxr::UsdAttribute Attr = Prim.GetAttribute( UnrealIdentifiers::UnrealLiveLinkSubjectName ) )
+			{
+				std::string SubjectName;
+				if ( Attr.Get( &SubjectName ) )
+				{
+					SubjectRepresentation.Subject = FName{ *UsdToUnreal::ConvertString( SubjectName ) };
+				}
+			}
+			Controller->SetSubjectRepresentation( SubjectRepresentation );
+
+			if ( pxr::UsdAttribute Attr = Prim.GetAttribute( UnrealIdentifiers::UnrealLiveLinkEnabled ) )
+			{
+				bool bEnabled = true;
+				if ( Attr.Get( &bEnabled ) )
+				{
+					Controller->bEvaluateLiveLink = bEnabled;
+				}
+			}
+		}
+	}
+
+	void RemoveLiveLinkFromComponent( USceneComponent* Component )
+	{
+		if ( !Component )
+		{
+			return;
+		}
+
+		AActor* Parent = Component->GetOwner();
+		if ( !Parent )
+		{
+			return;
+		}
+
+		TArray<ULiveLinkComponentController*> LiveLinkComponents;
+		Parent->GetComponents< ULiveLinkComponentController >( LiveLinkComponents );
+
+		for ( ULiveLinkComponentController* LiveLinkComponent : LiveLinkComponents )
+		{
+			if ( LiveLinkComponent->ComponentToControl.GetComponent( Parent ) == Component )
+			{
+				LiveLinkComponent->ComponentToControl.ComponentProperty = NAME_None;
+				Parent->RemoveInstanceComponent( LiveLinkComponent );
+				break;
+			}
+		}
+	}
+}
+
 class FUsdGeomXformableCreateAssetsTaskChain : public FBuildStaticMeshTaskChain
 {
 public:
@@ -440,17 +575,33 @@ void FUsdGeomXformableTranslator::UpdateComponents( USceneComponent* SceneCompon
 			}
 		}
 
+		UE::FUsdPrim Prim = GetPrim();
+
+		// Handle LiveLink, but only i we're not a skeletal mesh component: The SkelRootTranslator will deal with the
+		// skeletal version of the LiveLink configuration, we only handle setting up LiveLink for simple transforms
+		if ( !SceneComponent->IsA<USkeletalMeshComponent>() )
+		{
+			if ( UsdUtils::PrimHasLiveLinkSchema( Prim ) )
+			{
+				UE::UsdXformableTranslatorImpl::Private::SetUpSceneComponentForLiveLink( Context.Get(), SceneComponent, Prim );
+			}
+			else
+			{
+				UE::UsdXformableTranslatorImpl::Private::RemoveLiveLinkFromComponent( SceneComponent );
+			}
+		}
+
 		// Only put the transform into the component if we haven't parsed LODs for our static mesh: The Mesh transforms will already be baked
 		// into the mesh at that case, as each LOD could technically have a separate transform
 		if ( !Context->bAllowInterpretingLODs || !bHasMultipleLODs )
 		{
-			UsdToUnreal::ConvertXformable( Context->Stage, pxr::UsdGeomXformable( GetPrim() ), *SceneComponent, Context->Time );
+			UsdToUnreal::ConvertXformable( Context->Stage, pxr::UsdGeomXformable( Prim ), *SceneComponent, Context->Time );
 		}
 
 		// Note how we should only register if we unregistered ourselves: If we did this every time we would
 		// register too early during the process of duplicating into PIE, and that would prevent a future RegisterComponent
 		// call from naturally creating the required render state
-		if ( bStaticMobility && !SceneComponent->IsRegistered() )
+		if ( !SceneComponent->IsRegistered() )
 		{
 			SceneComponent->RegisterComponent();
 		}
@@ -592,7 +743,27 @@ bool FUsdGeomXformableTranslator::CollapsesChildren( ECollapsingType CollapsingT
 
 bool FUsdGeomXformableTranslator::CanBeCollapsed( ECollapsingType CollapsingType ) const
 {
-	return !UsdUtils::IsAnimated( GetPrim() );
+	FScopedUsdAllocs UsdAllocs;
+
+	FString PrimPathStr = PrimPath.GetString();
+
+	pxr::UsdPrim UsdPrim{ GetPrim() };
+	if ( !UsdPrim )
+	{
+		return false;
+	}
+
+	if ( UsdUtils::IsAnimated( UsdPrim ) )
+	{
+		return false;
+	}
+
+	if ( UsdUtils::PrimHasLiveLinkSchema( UsdPrim ) )
+	{
+		return false;
+	}
+
+	return true;
 }
 
 #endif // #if USE_USD_SDK
diff --git a/Engine/Plugins/Importers/USDImporter/Source/USDSchemas/Private/USDSkelRootTranslator.cpp b/Engine/Plugins/Importers/USDImporter/Source/USDSchemas/Private/USDSkelRootTranslator.cpp
--- a/Engine/Plugins/Importers/USDImporter/Source/USDSchemas/Private/USDSkelRootTranslator.cpp
+++ b/Engine/Plugins/Importers/USDImporter/Source/USDSchemas/Private/USDSkelRootTranslator.cpp
@@ -10,6 +10,7 @@
 #include "USDConversionUtils.h"
 #include "USDErrorUtils.h"
 #include "USDGeomMeshConversion.h"
+#include "USDIntegrationUtils.h"
 #include "USDLayerUtils.h"
 #include "USDLog.h"
 #include "USDMemory.h"
@@ -20,6 +21,7 @@
 #include "UsdWrappers/SdfPath.h"
 #include "UsdWrappers/UsdPrim.h"
 
+#include "Animation/AnimBlueprint.h"
 #include "Animation/AnimSequence.h"
 #include "Animation/Skeleton.h"
 #include "AnimationUtils.h"
@@ -32,6 +34,14 @@
 #include "Serialization/BufferArchive.h"
 
 #if WITH_EDITOR
+	#include "AnimGraphNode_LiveLinkPose.h"
+	#include "AnimNode_LiveLinkPose.h"
+	#include "BlueprintCompilationManager.h"
+	#include "EdGraph/EdGraph.h"
+	#include "EdGraph/EdGraphNode.h"
+	#include "K2Node_VariableGet.h"
+	#include "Kismet2/BlueprintEditorUtils.h"
+	#include "Kismet2/CompilerResultsLog.h"
 	#include "MaterialEditingLibrary.h"
 #endif // WITH_EDITOR
 
@@ -944,6 +954,247 @@ namespace UsdSkelRootTranslatorImpl
 				return true;
 			});
 	}
+
+	void UpdateLiveLinkProperties( const FUsdSchemaTranslationContext& Context, USkeletalMeshComponent* Component, const pxr::UsdPrim& Prim )
+	{
+		if ( !Component || !Component->SkeletalMesh || !Prim )
+		{
+			return;
+		}
+
+		FString PrimName = UsdToUnreal::ConvertString( Prim.GetName() );
+
+		USkeleton* Skeleton = Component->SkeletalMesh->GetSkeleton();
+		if ( !Skeleton )
+		{
+			return;
+		}
+
+		// Fetch relevant attributes from prim, since we know it has the schema
+		FString SubjectName;
+		FString AnimBPPath;
+		{
+			FScopedUsdAllocs Allocs;
+
+			if ( pxr::UsdAttribute Attr = Prim.GetAttribute( UnrealIdentifiers::UnrealLiveLinkSubjectName ) )
+			{
+				std::string SubjectNameString;
+				if ( Attr.Get( &SubjectNameString ) )
+				{
+					SubjectName = UsdToUnreal::ConvertString( SubjectNameString );
+				}
+			}
+
+			if ( pxr::UsdAttribute Attr = Prim.GetAttribute( UnrealIdentifiers::UnrealLiveLinkAnimBlueprintPath ) )
+			{
+				std::string PathString;
+				if ( Attr.Get( &PathString ) )
+				{
+					AnimBPPath = UsdToUnreal::ConvertString( PathString );
+				}
+			}
+		}
+
+		UAnimBlueprint* ExistingAnimBP = nullptr;
+		if ( Component->AnimClass )
+		{
+			ExistingAnimBP = Cast<UAnimBlueprint>( Component->AnimClass->ClassGeneratedBy );
+		}
+
+		bool bNeedRecompile = false;
+
+		// Fetch AnimBP to use
+		UAnimBlueprint* AnimBP = ExistingAnimBP;
+		if ( !AnimBPPath.IsEmpty() && ( !ExistingAnimBP || ExistingAnimBP->GetPathName() != AnimBPPath ) )
+		{
+			const static FString DefaultAnimBPPath = TEXT( "/USDImporter/Blueprint/DefaultLiveLinkAnimBP.DefaultLiveLinkAnimBP" );
+
+			AnimBP = Cast<UAnimBlueprint>( FSoftObjectPath( AnimBPPath ).TryLoad() );
+			if ( AnimBP )
+			{
+				// Create transient AnimBP based on our template, so that we can assign it a proper skeleton
+				if ( DefaultAnimBPPath == AnimBPPath )
+				{
+					// Check if we can use the AnimBP we currently have
+					bool bAlreadyHasTransientAnimBP = false;
+					if ( ExistingAnimBP )
+					{
+						if ( ExistingAnimBP->GetOutermost() == GetTransientPackage() && ExistingAnimBP->TargetSkeleton == Skeleton )
+						{
+							AnimBP = ExistingAnimBP;
+							bAlreadyHasTransientAnimBP = true;
+						}
+					}
+
+					FString CacheKey = UsdToUnreal::ConvertPath( Prim.GetPrimPath() ) + TEXT( "_DefaultAnimBlueprint" );
+
+					// Check if we can find an AnimBP for this prim in the asset cache (useful when doing Action->Import)
+					if ( !bAlreadyHasTransientAnimBP )
+					{
+						if ( UAnimBlueprint* CachedAnimBP = Cast< UAnimBlueprint >( Context.AssetCache->GetCachedAsset( CacheKey ) ) )
+						{
+							if ( CachedAnimBP->TargetSkeleton == Skeleton )
+							{
+								AnimBP = CachedAnimBP;
+								bAlreadyHasTransientAnimBP = true;
+							}
+						}
+					}
+
+					// We have to generate a new transient AnimBP
+					if ( !bAlreadyHasTransientAnimBP )
+					{
+						FName UniqueName = MakeUniqueObjectName(
+							GetTransientPackage(),
+							UAnimBlueprint::StaticClass(),
+							*( PrimName + TEXT( "_DefaultAnimBlueprint" ) )
+						);
+
+						// Duplicate and never reuse these so that they can be assigned independent subject names if desired.
+						// Its not as if scenes will have thousands of these anyway.
+						AnimBP = DuplicateObject( AnimBP, GetTransientPackage(), UniqueName );
+
+						// Patch up the flags here or else the rest of the engine code (and our plugin code) will get confused
+						// as to what a non-transient asset in the transient package even means. It can lead to some crashes
+						// when saving after import if we don't do this
+						AnimBP->ClearFlags( AnimBP->GetFlags() );
+						AnimBP->SetFlags( Skeleton->GetFlags() );
+
+						AnimBP->TargetSkeleton = Skeleton;
+						AnimBP->bIsTemplate = false;
+
+						bNeedRecompile = true;
+
+						Context.AssetCache->CacheAsset( CacheKey, AnimBP );
+					}
+				}
+				// Path is pointing to an existing, persistent AnimBP
+				else
+				{
+					// Force skeletons to be compatible
+					AnimBP->TargetSkeleton->AddCompatibleSkeleton( Skeleton );
+
+					// Check ExistingAnimBP so that we only emit this warning when we first set up the component (when we'll realistically not have an ExistingAnimBP yet),
+					// to try and prevent some warning spam
+					if ( AnimBP->TargetSkeleton->GetReferenceSkeleton().GetRefBoneInfo() != Skeleton->GetReferenceSkeleton().GetRefBoneInfo() && !ExistingAnimBP )
+					{
+						UE_LOG(LogUsd, Warning, TEXT("Forcing AnimBlueprint '%s's Skeleton '%s' to be compatible with the Skeleton generated for prim '%s', but they may be different!"),
+							*AnimBP->GetPathName(),
+							*AnimBP->TargetSkeleton->GetPathName(),
+							*PrimName
+						);
+					}
+				}
+			}
+		}
+
+		// Apply subject name to live link pose AnimBlueprint node
+		// Reference: UAnimationBlueprintLibrary::AddNodeAssetOverride
+		if ( AnimBP )
+		{
+			TArray<UBlueprint*> BlueprintHierarchy;
+			AnimBP->GetBlueprintHierarchyFromClass( AnimBP->GetAnimBlueprintGeneratedClass(), BlueprintHierarchy );
+
+			TArray<UAnimGraphNode_LiveLinkPose*> LiveLinkNodes;
+
+			for ( int32 BlueprintIndex = 0; BlueprintIndex < BlueprintHierarchy.Num(); ++BlueprintIndex )
+			{
+				UBlueprint* CurrentBlueprint = BlueprintHierarchy[ BlueprintIndex ];
+
+				TArray<UEdGraph*> Graphs;
+				CurrentBlueprint->GetAllGraphs( Graphs );
+
+				for ( UEdGraph* Graph : Graphs )
+				{
+					for ( UEdGraphNode* Node : Graph->Nodes )
+					{
+						if ( UAnimGraphNode_LiveLinkPose* AnimNode = Cast<UAnimGraphNode_LiveLinkPose>( Node ) )
+						{
+							LiveLinkNodes.Add( AnimNode );
+						}
+					}
+				}
+			}
+
+			if ( LiveLinkNodes.Num() > 1 && !ExistingAnimBP )
+			{
+				UE_LOG( LogUsd, Warning, TEXT( "Found more than one LiveLinkPose blueprint node on AnimBlueprint '%s's graphs."
+					"Note that all of those nodes will have their LiveLink SubjectName's updated to '%s', as described for prim '%s'!" ),
+					*AnimBP->GetPathName(),
+					*SubjectName,
+					*UsdToUnreal::ConvertPath( Prim.GetPrimPath() )
+				);
+			}
+
+			for ( UAnimGraphNode_LiveLinkPose* Node : LiveLinkNodes )
+			{
+				const UEdGraphSchema* Schema = Node->GetSchema();
+				if ( !Schema )
+				{
+					continue;
+				}
+
+				UEdGraphPin* SubjectNamePin = nullptr;
+				for ( UEdGraphPin* Pin : Node->Pins )
+				{
+					if ( Pin->GetName() == GET_MEMBER_NAME_STRING_CHECKED( FAnimNode_LiveLinkPose, LiveLinkSubjectName ) )
+					{
+						SubjectNamePin = Pin;
+						break;
+					}
+				}
+
+				// The subject name pin is already connected to something...
+				if ( SubjectNamePin->LinkedTo.Num() != 0 )
+				{
+					if ( !ExistingAnimBP )
+					{
+						UE_LOG( LogUsd, Warning, TEXT( "Failed to update a LiveLinkPose node's 'Subject Name' to '%s' on AnimBlueprint '%s', "
+							"because the pin is already connected to some other node. Disconnect it if you want it to be updated automatically." ),
+							*SubjectName,
+							*AnimBP->GetPathName()
+						);
+					}
+
+					continue;
+				}
+
+				if ( SubjectNamePin )
+				{
+					// The pin type is FLiveLinkSubjectName, so we must create an instance of it and serialize it using
+					// UScriptStruct::ExportText to generate a proper default value string
+					FLiveLinkSubjectName Dummy;
+					Dummy.Name = *SubjectName;
+
+					FString ValueString;
+					const void* Defaults = nullptr;
+					const UObject* OwnerObject = nullptr;
+					const int32 PortFlags = EPropertyPortFlags::PPF_None;
+					const UObject* ExportRootScope = nullptr;
+					FLiveLinkSubjectName::StaticStruct()->ExportText( ValueString, &Dummy, nullptr, nullptr, EPropertyPortFlags::PPF_None, nullptr );
+
+					if ( !Schema->DoesDefaultValueMatch( *SubjectNamePin, ValueString ) )
+					{
+						SubjectNamePin->Modify();
+						Schema->TrySetDefaultValue( *SubjectNamePin, ValueString );
+
+						FBlueprintEditorUtils::MarkBlueprintAsModified( AnimBP );
+						bNeedRecompile = true;
+					}
+				}
+			}
+
+			if ( bNeedRecompile )
+			{
+				FCompilerResultsLog Results;
+				FBPCompileRequest Request( AnimBP, EBlueprintCompileOptions::None, &Results );
+				FBlueprintCompilationManager::CompileSynchronously( Request );
+			}
+
+			Component->SetAnimInstanceClass( AnimBP->GeneratedClass );
+		}
+	}
+
 #endif // WITH_EDITOR
 }
 
@@ -961,11 +1212,8 @@ void FUsdSkelRootTranslator::CreateAssets()
 
 USceneComponent* FUsdSkelRootTranslator::CreateComponents()
 {
-	USceneComponent* RootComponent = FUsdGeomXformableTranslator::CreateComponents();
-
-	UpdateComponents( RootComponent );
-
-	return RootComponent;
+	// Don't need to call UpdateComponents as FUsdGeomXformableTranslator will already do that for us
+	return FUsdGeomXformableTranslator::CreateComponents();
 }
 
 void FUsdSkelRootTranslator::UpdateComponents( USceneComponent* SceneComponent )
@@ -976,10 +1224,28 @@ void FUsdSkelRootTranslator::UpdateComponents( USceneComponent* SceneComponent )
 		return;
 	}
 
+	UE::FUsdPrim Prim = GetPrim();
+
+	const bool bPrimHasLiveLinkSchema = UsdUtils::PrimHasLiveLinkSchema( Prim );
+	bool bPrimHasLiveLinkEnabled = bPrimHasLiveLinkSchema;
+	if ( bPrimHasLiveLinkSchema )
+	{
+		FScopedUsdAllocs Allocs;
+
+		if ( pxr::UsdAttribute Attr = pxr::UsdPrim{ Prim }.GetAttribute( UnrealIdentifiers::UnrealLiveLinkEnabled ) )
+		{
+			bool bEnabled = true;
+			if ( Attr.Get( &bEnabled ) )
+			{
+				bPrimHasLiveLinkEnabled = bEnabled;
+			}
+		}
+	}
+
 	UE::FUsdPrim SkelAnimPrim;
 	if ( SkeletalMeshComponent->AnimationData.AnimToPlay == nullptr )
 	{
-		SkelAnimPrim = UsdUtils::FindFirstAnimationSource( GetPrim() );
+		SkelAnimPrim = UsdUtils::FindFirstAnimationSource( Prim );
 		if ( SkelAnimPrim )
 		{
 			if ( UAnimSequence* TargetAnimSequence = Cast< UAnimSequence >( Context->AssetCache->GetAssetForPrim( SkelAnimPrim.GetPrimPath().GetString() ) ) )
@@ -988,7 +1254,6 @@ void FUsdSkelRootTranslator::UpdateComponents( USceneComponent* SceneComponent )
 				SkeletalMeshComponent->AnimationData.bSavedLooping = false;
 				SkeletalMeshComponent->AnimationData.bSavedPlaying = false;
 				SkeletalMeshComponent->SetUpdateAnimationInEditor( true );
-				SkeletalMeshComponent->SetAnimationMode( EAnimationMode::AnimationSingleNode );
 				SkeletalMeshComponent->SetAnimation( TargetAnimSequence );
 			}
 		}
@@ -996,6 +1261,10 @@ void FUsdSkelRootTranslator::UpdateComponents( USceneComponent* SceneComponent )
 
 	Super::UpdateComponents( SceneComponent );
 
+	// We always want this, but need to be registered for this to work (Super::UpdateComponents should register us)
+	const bool bNewUpdateState = true;
+	SkeletalMeshComponent->SetUpdateAnimationInEditor( bNewUpdateState );
+
 #if WITH_EDITOR
 	// Re-set the skeletal mesh if we created a new one (maybe the hash changed, a skinned UsdGeomMesh was hidden, etc.)
 	USkeletalMesh* TargetSkeletalMesh = Cast< USkeletalMesh >( Context->AssetCache->GetAssetForPrim( PrimPath.GetString() ) );
@@ -1013,7 +1282,7 @@ void FUsdSkelRootTranslator::UpdateComponents( USceneComponent* SceneComponent )
 			}
 
 			UsdSkelRootTranslatorImpl::SetMaterialOverrides(
-				GetPrim(),
+				Prim,
 				ExistingAssignments,
 				*SkeletalMeshComponent,
 				*Context->AssetCache.Get(),
@@ -1025,15 +1294,20 @@ void FUsdSkelRootTranslator::UpdateComponents( USceneComponent* SceneComponent )
 		}
 	}
 
+	if ( bPrimHasLiveLinkSchema )
+	{
+		UsdSkelRootTranslatorImpl::UpdateLiveLinkProperties( Context.Get(), SkeletalMeshComponent, Prim );
+	}
+
 	// Update the animation state
-	if( SkeletalMeshComponent->SkeletalMesh )
+	if ( SkeletalMeshComponent->SkeletalMesh && !bPrimHasLiveLinkEnabled )
 	{
 		if ( UAnimSequence* AnimSequence = Cast<UAnimSequence>( SkeletalMeshComponent->AnimationData.AnimToPlay.Get() ) )
 		{
 			UE::FSdfLayerOffset CombinedOffset;
 			if ( !SkelAnimPrim )
 			{
-				SkelAnimPrim = UsdUtils::FindFirstAnimationSource( GetPrim() );
+				SkelAnimPrim = UsdUtils::FindFirstAnimationSource( Prim );
 			}
 			if ( SkelAnimPrim )
 			{
diff --git a/Engine/Plugins/Importers/USDImporter/Source/USDSchemas/USDSchemas.Build.cs b/Engine/Plugins/Importers/USDImporter/Source/USDSchemas/USDSchemas.Build.cs
--- a/Engine/Plugins/Importers/USDImporter/Source/USDSchemas/USDSchemas.Build.cs
+++ b/Engine/Plugins/Importers/USDImporter/Source/USDSchemas/USDSchemas.Build.cs
@@ -20,6 +20,9 @@ namespace UnrealBuildTool.Rules
 					"CoreUObject",
 					"Engine",
 					"GeometryCache",
+					"LiveLinkAnimationCore",
+					"LiveLinkComponents",
+					"LiveLinkInterface",
 					"MeshDescription",
 					"RenderCore",
 					"RHI", // For FMaterialUpdateContext and the right way of updating material instance constants
@@ -37,7 +40,10 @@ namespace UnrealBuildTool.Rules
 				PrivateDependencyModuleNames.AddRange(
 					new string[]
 					{
+						"BlueprintGraph",
 						"GeometryCacheUSD",
+						"Kismet",
+						"LiveLinkGraphNode",
 						"MaterialEditor",
 						"MDLImporter",
 						"MeshUtilities",
diff --git a/Engine/Plugins/Importers/USDImporter/Source/USDStage/Private/USDStageActor.cpp b/Engine/Plugins/Importers/USDImporter/Source/USDStage/Private/USDStageActor.cpp
--- a/Engine/Plugins/Importers/USDImporter/Source/USDStage/Private/USDStageActor.cpp
+++ b/Engine/Plugins/Importers/USDImporter/Source/USDStage/Private/USDStageActor.cpp
@@ -6,12 +6,13 @@
 #include "USDAssetCache.h"
 #include "USDAssetImportData.h"
 #include "USDClassesModule.h"
-#include "USDInfoCache.h"
 #include "USDConversionUtils.h"
 #include "USDErrorUtils.h"
 #include "USDGeomMeshConversion.h"
 #include "USDGeomMeshTranslator.h"
 #include "USDGeomXformableTranslator.h"
+#include "USDInfoCache.h"
+#include "USDIntegrationUtils.h"
 #include "USDLayerUtils.h"
 #include "USDLightConversion.h"
 #include "USDListener.h"
@@ -49,6 +50,7 @@
 #include "EngineAnalytics.h"
 #include "HAL/FileManager.h"
 #include "LevelSequence.h"
+#include "LiveLinkComponentController.h"
 #include "Materials/Material.h"
 #include "Materials/MaterialInstanceConstant.h"
 #include "Materials/MaterialInterface.h"
@@ -2477,6 +2479,14 @@ void AUsdStageActor::OnObjectPropertyChanged( UObject* ObjectBeingModified, FPro
 		FUsdStageActorImpl::AllowListComponentHierarchy( GetRootComponent(), VisitedObjects );
 	}
 
+	// So that we can detect when the user enables/disables live link properties on a ULiveLinkComponentController that may
+	// be controlling a component that we *do* care about
+	ULiveLinkComponentController* Controller = Cast< ULiveLinkComponentController >( ObjectBeingModified );
+	if ( Controller )
+	{
+		ObjectBeingModified = Controller->ComponentToControl.GetComponent( Controller->GetOwner() );
+	}
+
 	UObject* PrimObject = ObjectBeingModified;
 
 	if ( !ObjectsToWatch.Contains( ObjectBeingModified ) )
@@ -2539,6 +2549,9 @@ void AUsdStageActor::OnObjectPropertyChanged( UObject* ObjectBeingModified, FPro
 				}
 
 #if USE_USD_SDK
+
+				UnrealToUsd::ConvertLiveLinkProperties( Controller ? Cast<UActorComponent>( Controller ) : PrimSceneComponent, UsdPrim );
+
 				UnrealToUsd::ConvertSceneComponent( UsdStage, PrimSceneComponent, UsdPrim );
 
 				if ( UMeshComponent* MeshComponent = Cast< UMeshComponent >( PrimSceneComponent ) )
diff --git a/Engine/Plugins/Importers/USDImporter/Source/USDStage/USDStage.Build.cs b/Engine/Plugins/Importers/USDImporter/Source/USDStage/USDStage.Build.cs
--- a/Engine/Plugins/Importers/USDImporter/Source/USDStage/USDStage.Build.cs
+++ b/Engine/Plugins/Importers/USDImporter/Source/USDStage/USDStage.Build.cs
@@ -12,7 +12,7 @@ namespace UnrealBuildTool.Rules
 			PublicDependencyModuleNames.AddRange(
 				new string[]
 				{
-					"USDSchemas" // Has to be a public dependency because the stage actor has a collapsing cache member
+					"USDSchemas" // Has to be a public dependency because the stage actor has an FUsdInfoCache member
 				}
 			);
 
@@ -24,6 +24,7 @@ namespace UnrealBuildTool.Rules
 					"CoreUObject",
 					"Engine",
 					"LevelSequence",
+					"LiveLinkComponents", // For tracking edits to LiveLinkComponentController properties and writing to USD
 					"MeshDescription",
 					"MovieScene",
 					"MovieSceneTracks",
diff --git a/Engine/Plugins/Importers/USDImporter/Source/USDStageEditor/Private/SUSDPrimInfo.cpp b/Engine/Plugins/Importers/USDImporter/Source/USDStageEditor/Private/SUSDPrimInfo.cpp
--- a/Engine/Plugins/Importers/USDImporter/Source/USDStageEditor/Private/SUSDPrimInfo.cpp
+++ b/Engine/Plugins/Importers/USDImporter/Source/USDStageEditor/Private/SUSDPrimInfo.cpp
@@ -2,43 +2,36 @@
 
 #include "SUSDPrimInfo.h"
 
+#include "USDIntegrationUtils.h"
 #include "USDStageActor.h"
 #include "USDStageModule.h"
 #include "USDTypesConversion.h"
+#include "Widgets/SUSDIntegrationsPanel.h"
+#include "Widgets/SUSDPrimPropertiesList.h"
+#include "Widgets/SUSDReferencesList.h"
+#include "Widgets/SUSDVariantSetsList.h"
 
-#include "UsdWrappers/SdfPath.h"
-#include "UsdWrappers/UsdPrim.h"
 #include "UsdWrappers/UsdStage.h"
+#include "UsdWrappers/UsdPrim.h"
+#include "UsdWrappers/SdfPath.h"
 
-#include "Algo/Find.h"
-#include "EditorStyleSet.h"
-#include "Engine/World.h"
-#include "Modules/ModuleManager.h"
-
+#include "Widgets/Views/SListView.h"
+#include "Widgets/Text/STextBlock.h"
 #include "Widgets/SBoxPanel.h"
-#include "Widgets/SUSDPrimPropertiesList.h"
-#include "Widgets/SUSDReferencesList.h"
-#include "Widgets/SUSDVariantSetsList.h"
-#include "Widgets/Layout/SSeparator.h"
-#include "Widgets/Layout/SSpacer.h"
 #include "Widgets/Layout/SSplitter.h"
-#include "Widgets/Text/STextBlock.h"
-#include "Widgets/Views/SListView.h"
+#include "Widgets/Layout/SSpacer.h"
+#include "Widgets/Layout/SSeparator.h"
+#include "Modules/ModuleManager.h"
+#include "Engine/World.h"
+#include "EditorStyleSet.h"
+#include "Algo/Find.h"
 
 #if USE_USD_SDK
 
 #define LOCTEXT_NAMESPACE "SUSDPrimInfo"
 
-namespace UsdPrimInfoWidgetConstants
-{
-	const FMargin CategoryHeaderPadding( 4.0f, 4.0f, 4.0f, 4.0f );
-}
-
 void SUsdPrimInfo::Construct( const FArguments& InArgs, const UE::FUsdStageWeak& UsdStage, const TCHAR* PrimPath )
 {
-	TSharedRef< SWidget > VariantSetsWidget = GenerateVariantSetsWidget( UsdStage, PrimPath );
-	TSharedRef< SWidget > ReferencesListWidget = GenerateReferencesListWidget( UsdStage, PrimPath );
-
 	ChildSlot
 	[
 		SNew( SVerticalBox )
@@ -56,25 +49,33 @@ void SUsdPrimInfo::Construct( const FArguments& InArgs, const UE::FUsdStageWeak&
 		+SVerticalBox::Slot()
 		.AutoHeight()
 		[
-			SAssignNew( VariantSetsBox, SBox )
+			SNew( SBox )
 			.Content()
 			[
-				VariantSetsWidget
+				SAssignNew( IntegrationsPanel, SUsdIntegrationsPanel, UsdStage, PrimPath )
 			]
 		]
 
 		+SVerticalBox::Slot()
 		.AutoHeight()
 		[
-			SAssignNew( ReferencesBox, SBox )
+			SNew( SBox )
 			.Content()
 			[
-				ReferencesListWidget
+				SAssignNew( VariantsList, SVariantsList, UsdStage, PrimPath )
 			]
 		]
-	];
 
-	SetPrimPath( UsdStage, PrimPath );
+		+SVerticalBox::Slot()
+		.AutoHeight()
+		[
+			SNew( SBox )
+			.Content()
+			[
+				SAssignNew( ReferencesList, SUsdReferencesList, UsdStage, PrimPath )
+			]
+		]
+	];
 }
 
 void SUsdPrimInfo::SetPrimPath( const UE::FUsdStageWeak& UsdStage, const TCHAR* PrimPath )
@@ -84,70 +85,20 @@ void SUsdPrimInfo::SetPrimPath( const UE::FUsdStageWeak& UsdStage, const TCHAR*
 		PropertiesList->SetPrimPath( UsdStage, PrimPath );
 	}
 
-	if ( VariantsList )
-	{
-		VariantsList->SetPrimPath( UsdStage, PrimPath );
-	}
-
-	if ( VariantSetsBox )
-	{
-		VariantSetsBox->SetContent( GenerateVariantSetsWidget( UsdStage, PrimPath ) );
-	}
-
-	if ( ReferencesBox )
-	{
-		ReferencesBox->SetContent( GenerateReferencesListWidget( UsdStage, PrimPath ) );
-	}
-}
-
-TSharedRef< SWidget > SUsdPrimInfo::GenerateVariantSetsWidget( const UE::FUsdStageWeak& UsdStage, const TCHAR* PrimPath )
-{
-	TSharedRef< SWidget > VariantSetsWidget = SNullWidget::NullWidget;
-
-	if ( !UsdStage )
-	{
-		return VariantSetsWidget;
-	}
-
-	UE::FUsdPrim UsdPrim = UsdStage.GetPrimAtPath( UE::FSdfPath( PrimPath ) );
-
-	if ( UsdPrim && UsdPrim.HasVariantSets() )
+	if ( IntegrationsPanel )
 	{
-		SAssignNew( VariantSetsWidget, SVerticalBox )
-
-		+SVerticalBox::Slot()
-		.AutoHeight()
-		[
-			SAssignNew( VariantsList, SVariantsList, UsdStage, PrimPath )
-		];
+		IntegrationsPanel->SetPrimPath( UsdStage, PrimPath );
 	}
 
-	return VariantSetsWidget;
-}
-
-TSharedRef< SWidget > SUsdPrimInfo::GenerateReferencesListWidget( const UE::FUsdStageWeak& UsdStage, const TCHAR* PrimPath )
-{
-	TSharedRef< SWidget > ReferencesListWidget = SNullWidget::NullWidget;
-
-	if ( !UsdStage )
+	if ( VariantsList )
 	{
-		return ReferencesListWidget;
+		VariantsList->SetPrimPath( UsdStage, PrimPath );
 	}
 
-	UE::FUsdPrim UsdPrim = UsdStage.GetPrimAtPath( UE::FSdfPath( PrimPath ) );
-
-	if ( UsdPrim && UsdPrim.HasAuthoredReferences() )
+	if ( ReferencesList )
 	{
-		SAssignNew( ReferencesListWidget, SVerticalBox )
-
-		+SVerticalBox::Slot()
-		.AutoHeight()
-		[
-			SNew( SUsdReferencesList, UsdStage, PrimPath )
-		];
+		ReferencesList->SetPrimPath( UsdStage, PrimPath );
 	}
-
-	return ReferencesListWidget;
 }
 
 #undef LOCTEXT_NAMESPACE
diff --git a/Engine/Plugins/Importers/USDImporter/Source/USDStageEditor/Private/SUSDPrimInfo.h b/Engine/Plugins/Importers/USDImporter/Source/USDStageEditor/Private/SUSDPrimInfo.h
--- a/Engine/Plugins/Importers/USDImporter/Source/USDStageEditor/Private/SUSDPrimInfo.h
+++ b/Engine/Plugins/Importers/USDImporter/Source/USDStageEditor/Private/SUSDPrimInfo.h
@@ -9,8 +9,6 @@
 
 #if USE_USD_SDK
 
-class AUsdStageActor;
-
 class SUsdPrimInfo : public SCompoundWidget
 {
 public:
@@ -18,18 +16,13 @@ public:
 	SLATE_END_ARGS()
 
 	void Construct( const FArguments& InArgs, const UE::FUsdStageWeak& UsdStage, const TCHAR* PrimPath );
-
 	void SetPrimPath( const UE::FUsdStageWeak& UsdStage, const TCHAR* PrimPath );
 
-protected:
-	TSharedRef< SWidget > GenerateVariantSetsWidget( const UE::FUsdStageWeak& UsdStage, const TCHAR* PrimPath );
-	TSharedRef< SWidget > GenerateReferencesListWidget( const UE::FUsdStageWeak& UsdStage, const TCHAR* PrimPath );
-
 private:
 	TSharedPtr< class SUsdPrimPropertiesList > PropertiesList;
+	TSharedPtr< class SUsdIntegrationsPanel > IntegrationsPanel;
 	TSharedPtr< class SVariantsList > VariantsList;
-	TSharedPtr< class SBox > VariantSetsBox;
-	TSharedPtr< class SBox > ReferencesBox;
+	TSharedPtr< class SUsdReferencesList > ReferencesList;
 };
 
 #endif // #if USE_USD_SDK
diff --git a/Engine/Plugins/Importers/USDImporter/Source/USDStageEditor/Private/SUSDStageTreeView.cpp b/Engine/Plugins/Importers/USDImporter/Source/USDStageEditor/Private/SUSDStageTreeView.cpp
--- a/Engine/Plugins/Importers/USDImporter/Source/USDStageEditor/Private/SUSDStageTreeView.cpp
+++ b/Engine/Plugins/Importers/USDImporter/Source/USDStageEditor/Private/SUSDStageTreeView.cpp
@@ -665,6 +665,34 @@ TSharedPtr< SWidget > SUsdStageTreeView::ConstructPrimContextMenu()
 	}
 	PrimOptions.EndSection();
 
+	PrimOptions.BeginSection( "Integrations", LOCTEXT( "Integrations", "Integrations" ) );
+	{
+		PrimOptions.AddMenuEntry(
+			LOCTEXT( "SetUpLiveLink", "Set up Live Link" ),
+			LOCTEXT( "SetUpLiveLink_ToolTip", "Sets up the generated component for Live Link and store the connection details to the USD Stage" ),
+			FSlateIcon(),
+			FUIAction(
+				FExecuteAction::CreateSP( this, &SUsdStageTreeView::OnSetUpLiveLink ),
+				FCanExecuteAction::CreateSP( this, &SUsdStageTreeView::CanSetUpLiveLink )
+			),
+			NAME_None,
+			EUserInterfaceActionType::Button
+		);
+
+		PrimOptions.AddMenuEntry(
+			LOCTEXT( "RemoveLiveLink", "Remove Live Link" ),
+			LOCTEXT( "RemoveLiveLink_ToolTip", "Reverses the Live Link configuration on the component and removes the connection details from the USD Stage" ),
+			FSlateIcon(),
+			FUIAction(
+				FExecuteAction::CreateSP( this, &SUsdStageTreeView::OnRemoveLiveLink ),
+				FCanExecuteAction::CreateSP( this, &SUsdStageTreeView::CanRemoveLiveLink )
+			),
+			NAME_None,
+			EUserInterfaceActionType::Button
+		);
+	}
+	PrimOptions.EndSection();
+
 	MenuWidget = PrimOptions.MakeWidget();
 
 	return MenuWidget;
@@ -789,6 +817,44 @@ void SUsdStageTreeView::OnClearReferences()
 	}
 }
 
+void SUsdStageTreeView::OnSetUpLiveLink()
+{
+	if ( !UsdStageActor.IsValid() )
+	{
+		return;
+	}
+
+	FScopedTransaction Transaction( LOCTEXT( "SetUpLiveLinkTransaction", "Set up Live Link for selected prims" ) );
+
+	TArray< FUsdPrimViewModelRef > MySelectedItems = GetSelectedItems();
+
+	UE::FSdfChangeBlock Block;
+
+	for ( FUsdPrimViewModelRef SelectedItem : MySelectedItems )
+	{
+		SelectedItem->SetUpLiveLink();
+	}
+}
+
+void SUsdStageTreeView::OnRemoveLiveLink()
+{
+	if ( !UsdStageActor.IsValid() )
+	{
+		return;
+	}
+
+	FScopedTransaction Transaction( LOCTEXT( "RemoveLiveLinkTransaction", "Reverses the Live Link configuration for selected prims" ) );
+
+	TArray< FUsdPrimViewModelRef > MySelectedItems = GetSelectedItems();
+
+	UE::FSdfChangeBlock Block;
+
+	for ( FUsdPrimViewModelRef SelectedItem : MySelectedItems )
+	{
+		SelectedItem->RemoveLiveLink();
+	}
+}
+
 bool SUsdStageTreeView::CanAddPrim() const
 {
 	if ( !UsdStageActor.IsValid() )
@@ -835,6 +901,58 @@ bool SUsdStageTreeView::CanExecutePrimAction() const
 	return bHasPrimSpec;
 }
 
+bool SUsdStageTreeView::CanSetUpLiveLink() const
+{
+	if ( !UsdStageActor.IsValid() )
+	{
+		return false;
+	}
+
+	UE::FUsdStage UsdStage = UsdStageActor->GetOrLoadUsdStage();
+	if ( !UsdStage || !UsdStage.IsEditTargetValid() )
+	{
+		return false;
+	}
+
+	TArray< FUsdPrimViewModelRef > MySelectedItems = GetSelectedItems();
+
+	for ( FUsdPrimViewModelRef SelectedItem : MySelectedItems )
+	{
+		if ( SelectedItem->CanSetUpLiveLink() )
+		{
+			return true;
+		}
+	}
+
+	return false;
+}
+
+bool SUsdStageTreeView::CanRemoveLiveLink() const
+{
+	if ( !UsdStageActor.IsValid() )
+	{
+		return false;
+	}
+
+	UE::FUsdStage UsdStage = UsdStageActor->GetOrLoadUsdStage();
+	if ( !UsdStage || !UsdStage.IsEditTargetValid() )
+	{
+		return false;
+	}
+
+	TArray< FUsdPrimViewModelRef > MySelectedItems = GetSelectedItems();
+
+	for ( FUsdPrimViewModelRef SelectedItem : MySelectedItems )
+	{
+		if ( SelectedItem->CanRemoveLiveLink() )
+		{
+			return true;
+		}
+	}
+
+	return false;
+}
+
 void SUsdStageTreeView::RequestListRefresh()
 {
 	SUsdTreeView< FUsdPrimViewModelRef >::RequestListRefresh();
diff --git a/Engine/Plugins/Importers/USDImporter/Source/USDStageEditor/Private/SUSDStageTreeView.h b/Engine/Plugins/Importers/USDImporter/Source/USDStageEditor/Private/SUSDStageTreeView.h
--- a/Engine/Plugins/Importers/USDImporter/Source/USDStageEditor/Private/SUSDStageTreeView.h
+++ b/Engine/Plugins/Importers/USDImporter/Source/USDStageEditor/Private/SUSDStageTreeView.h
@@ -54,9 +54,13 @@ private:
 	void OnRemovePrim();
 	void OnAddReference();
 	void OnClearReferences();
+	void OnSetUpLiveLink();
+	void OnRemoveLiveLink();
 
 	bool CanAddPrim() const;
 	bool CanExecutePrimAction() const;
+	bool CanSetUpLiveLink() const;
+	bool CanRemoveLiveLink() const;
 
 	/** Uses TreeItemExpansionStates to travel the tree and call SetItemExpansion */
 	void RestoreExpansionStates();
diff --git a/Engine/Plugins/Importers/USDImporter/Source/USDStageEditor/Private/Widgets/SUSDIntegrationsPanel.cpp b/Engine/Plugins/Importers/USDImporter/Source/USDStageEditor/Private/Widgets/SUSDIntegrationsPanel.cpp
new file mode 100644
--- /dev/null
+++ b/Engine/Plugins/Importers/USDImporter/Source/USDStageEditor/Private/Widgets/SUSDIntegrationsPanel.cpp
@@ -0,0 +1,319 @@
+// Copyright Epic Games, Inc. All Rights Reserved.
+
+#include "SUSDIntegrationsPanel.h"
+
+#include "SUSDStageEditorStyle.h"
+#include "USDIntegrationUtils.h"
+#include "USDTypesConversion.h"
+#include "USDValueConversion.h"
+
+#include "UsdWrappers/SdfPath.h"
+#include "UsdWrappers/UsdAttribute.h"
+#include "UsdWrappers/UsdPrim.h"
+#include "UsdWrappers/UsdStage.h"
+#include "UsdWrappers/VtValue.h"
+
+#include "EditorStyleSet.h"
+#include "ISinglePropertyView.h"
+#include "Modules/ModuleManager.h"
+#include "PropertyEditorModule.h"
+#include "Roles/LiveLinkAnimationRole.h"
+#include "Roles/LiveLinkTransformRole.h"
+#include "SLiveLinkSubjectRepresentationPicker.h"
+#include "Widgets/Input/SCheckBox.h"
+#include "Widgets/Input/SEditableTextBox.h"
+#include "Widgets/Views/SHeaderRow.h"
+
+#define LOCTEXT_NAMESPACE "SUSDIntegrationsPanel"
+
+namespace UE::SUsdIntergrationsPanel::Private
+{
+	const FMargin LeftRowPadding( 6.0f, 0.0f, 2.0f, 0.0f );
+	const FMargin RightRowPadding( 3.0f, 0.0f, 2.0f, 0.0f );
+	const float DesiredNumericEntryBoxWidth = 80.0f;
+
+	const TCHAR* NormalFont = TEXT( "PropertyWindow.NormalFont" );
+}
+
+void SUsdIntegrationsPanelRow::Construct( const FArguments& InArgs, TSharedPtr<UE::FUsdAttribute> InAttr, const TSharedRef< STableViewBase >& OwnerTable )
+{
+	Attribute = InAttr;
+
+	SMultiColumnTableRow< TSharedPtr<UE::FUsdAttribute> >::Construct( SMultiColumnTableRow< TSharedPtr<UE::FUsdAttribute> >::FArguments(), OwnerTable );
+}
+
+TSharedRef< SWidget > SUsdIntegrationsPanelRow::GenerateWidgetForColumn( const FName& ColumnName )
+{
+	TSharedRef< SWidget > ColumnWidget = SNullWidget::NullWidget;
+
+	bool bIsLeftColumn = true;
+	FOptionalSize RowHeight = FUsdStageEditorStyle::Get()->GetFloat( "UsdStageEditor.ListItemHeight" );
+
+	FScopedUsdAllocs Allocs;
+
+	FName AttributeName = Attribute->GetName();
+
+	if ( ColumnName == TEXT("PropertyName") )
+	{
+		using DisplayTextForPropertiesEntry = TPairInitializer<const FName&, const FText&>;
+		const static TMap<FName, FText> DisplayTextForProperties
+		({
+			DisplayTextForPropertiesEntry( *UsdToUnreal::ConvertToken( UnrealIdentifiers::UnrealLiveLinkAnimBlueprintPath ), LOCTEXT( "AnimBlueprintPathText", "AnimBlueprint asset" ) ),
+			DisplayTextForPropertiesEntry( *UsdToUnreal::ConvertToken( UnrealIdentifiers::UnrealLiveLinkSubjectName ), LOCTEXT( "SubjectNameText", "LiveLink subject name" ) ),
+			DisplayTextForPropertiesEntry( *UsdToUnreal::ConvertToken( UnrealIdentifiers::UnrealLiveLinkEnabled ), LOCTEXT( "LiveLinkEnabledText", "Enable LiveLink" ) ),
+		});
+
+		if ( const FText* TextToDisplay = DisplayTextForProperties.Find( AttributeName ) )
+		{
+			SAssignNew( ColumnWidget, STextBlock )
+			.Text( *TextToDisplay )
+			.Font( FEditorStyle::GetFontStyle( UE::SUsdIntergrationsPanel::Private::NormalFont ) );
+		}
+		else
+		{
+			ensure(false);
+			ColumnWidget = SNullWidget::NullWidget;
+		}
+	}
+	else
+	{
+		bIsLeftColumn = false;
+
+		FName TypeName = Attribute->GetTypeName();
+		TSharedPtr<UE::FUsdAttribute> AttributeCopy = Attribute;
+
+		if ( AttributeName == *UsdToUnreal::ConvertToken( UnrealIdentifiers::UnrealLiveLinkAnimBlueprintPath ) )
+		{
+			if ( UUsdIntegrationsPanelPropertyDummy* Dummy = GetMutableDefault<UUsdIntegrationsPanelPropertyDummy>() )
+			{
+				// Let the object picker row be as tall as it wants
+				RowHeight = FOptionalSize();
+
+				FSinglePropertyParams Params;
+				Params.Font = FEditorStyle::GetFontStyle( UE::SUsdIntergrationsPanel::Private::NormalFont );
+				Params.NamePlacement = EPropertyNamePlacement::Hidden;
+
+				UE::FVtValue Value;
+				if ( AttributeCopy->Get( Value ) && !Value.IsEmpty() )
+				{
+					Dummy->AnimBPProperty = FSoftObjectPath( UsdUtils::Stringify( Value ) ).TryLoad();
+				}
+				else
+				{
+					Dummy->AnimBPProperty = nullptr;
+				}
+
+				FPropertyEditorModule& PropertyEditor = FModuleManager::LoadModuleChecked<FPropertyEditorModule>( TEXT( "PropertyEditor" ) );
+				TSharedPtr<class ISinglePropertyView> PropertyView = PropertyEditor.CreateSingleProperty(
+					Dummy,
+					GET_MEMBER_NAME_CHECKED( UUsdIntegrationsPanelPropertyDummy, AnimBPProperty ),
+					Params
+				);
+
+				FSimpleDelegate PropertyChanged = FSimpleDelegate::CreateLambda( [AttributeCopy]()
+				{
+					UUsdIntegrationsPanelPropertyDummy* Dummy = GetMutableDefault<UUsdIntegrationsPanelPropertyDummy>();
+
+					FString Path = ( Dummy && Dummy->AnimBPProperty ) ? Dummy->AnimBPProperty->GetPathName() : FString{};
+
+					TArray< UsdUtils::FConvertedVtValueComponent > Components;
+					Components.Add( UsdUtils::FConvertedVtValueComponent( TInPlaceType<FString>(), Path ) );
+
+					UsdUtils::FConvertedVtValue Converted;
+					Converted.SourceType = UsdUtils::EUsdBasicDataTypes::String;
+					Converted.bIsArrayValued = false;
+					Converted.bIsEmpty = false;
+					Converted.Entries.Add( Components );
+
+					UE::FVtValue Value;
+					if ( UnrealToUsd::ConvertValue( Converted, Value ) && !Value.IsEmpty() )
+					{
+						AttributeCopy->Set( Value );
+					}
+				} );
+				PropertyView->SetOnPropertyValueChanged( PropertyChanged );
+
+				ColumnWidget = PropertyView.ToSharedRef();
+			}
+			else
+			{
+				ColumnWidget = SNullWidget::NullWidget;
+			}
+		}
+		else if ( AttributeName == *UsdToUnreal::ConvertToken( UnrealIdentifiers::UnrealLiveLinkSubjectName ) )
+		{
+			// SLiveLinkSubjectRepresentationPicker is also a little bit larger than one of our rows, so let it expand a bit
+			RowHeight = FOptionalSize();
+
+			SAssignNew( ColumnWidget, SLiveLinkSubjectRepresentationPicker )
+			.ShowRole( false )
+			.ShowSource( false )
+			.Font( FEditorStyle::GetFontStyle( UE::SUsdIntergrationsPanel::Private::NormalFont ) )
+			.Value_Lambda( [AttributeCopy]()
+			{
+				SLiveLinkSubjectRepresentationPicker::FLiveLinkSourceSubjectRole Result;
+
+				Result.Role = AttributeCopy->GetPrim().IsA(TEXT("SkelRoot"))
+					? ULiveLinkAnimationRole::StaticClass()
+					: ULiveLinkTransformRole::StaticClass();
+
+				UE::FVtValue Value;
+				AttributeCopy->Get( Value );
+				Result.Subject = FLiveLinkSubjectName{ FName{*UsdUtils::Stringify( Value )} };
+
+				return Result;
+			})
+			.OnValueChanged_Lambda( [AttributeCopy]( SLiveLinkSubjectRepresentationPicker::FLiveLinkSourceSubjectRole NewValue )
+			{
+				TArray< UsdUtils::FConvertedVtValueComponent > Components;
+				Components.Add( UsdUtils::FConvertedVtValueComponent( TInPlaceType<FString>(), NewValue.Subject.ToString() ) );
+
+				UsdUtils::FConvertedVtValue Converted;
+				Converted.SourceType = UsdUtils::EUsdBasicDataTypes::String;
+				Converted.bIsArrayValued = false;
+				Converted.bIsEmpty = false;
+				Converted.Entries.Add( Components );
+
+				UE::FVtValue Value;
+				if ( UnrealToUsd::ConvertValue( Converted, Value ) && !Value.IsEmpty() )
+				{
+					AttributeCopy->Set( Value );
+				}
+			});
+		}
+		else if ( AttributeName == *UsdToUnreal::ConvertToken( UnrealIdentifiers::UnrealLiveLinkEnabled ) )
+		{
+			SAssignNew( ColumnWidget, SBox )
+			.HeightOverride( FUsdStageEditorStyle::Get()->GetFloat( "UsdStageEditor.ListItemHeight" ) )
+			.VAlign( VAlign_Center )
+			[
+				SNew( SCheckBox )
+				.IsChecked_Lambda( [AttributeCopy]()
+				{
+					UE::FVtValue Value;
+					AttributeCopy->Get( Value );
+
+					UsdUtils::FConvertedVtValue Converted;
+
+					if ( UsdToUnreal::ConvertValue( Value, Converted ) && Converted.SourceType == UsdUtils::EUsdBasicDataTypes::Bool && !Converted.bIsEmpty )
+					{
+						if ( bool* bBoolValue = Converted.Entries[ 0 ][ 0 ].TryGet<bool>() )
+						{
+							return *bBoolValue ? ECheckBoxState::Checked : ECheckBoxState::Unchecked;
+						}
+					}
+
+					return ECheckBoxState::Undetermined;
+				})
+				.OnCheckStateChanged_Lambda( [AttributeCopy]( ECheckBoxState NewState )
+				{
+					TArray< UsdUtils::FConvertedVtValueComponent > Components;
+					Components.Add( UsdUtils::FConvertedVtValueComponent( TInPlaceType<bool>(), NewState == ECheckBoxState::Checked ) );
+
+					UsdUtils::FConvertedVtValue Converted;
+					Converted.SourceType = UsdUtils::EUsdBasicDataTypes::Bool;
+					Converted.bIsArrayValued = false;
+					Converted.bIsEmpty = false;
+					Converted.Entries.Add( Components );
+
+					UE::FVtValue Value;
+					if ( UnrealToUsd::ConvertValue( Converted, Value ) && !Value.IsEmpty() )
+					{
+						AttributeCopy->Set( Value );
+					}
+				})
+			];
+		}
+		else
+		{
+			ensure(false);
+		}
+	}
+
+	return SNew(SBox)
+		.HeightOverride( RowHeight )
+		[
+			SNew( SHorizontalBox )
+			+ SHorizontalBox::Slot()
+			.HAlign( HAlign_Left )
+			.VAlign( VAlign_Center )
+			.Padding( bIsLeftColumn ? UE::SUsdIntergrationsPanel::Private::LeftRowPadding : UE::SUsdIntergrationsPanel::Private::RightRowPadding )
+			.AutoWidth()
+			[
+				ColumnWidget
+			]
+		];
+}
+
+void SUsdIntegrationsPanel::Construct( const FArguments& InArgs, const UE::FUsdStageWeak& InUsdStage, const TCHAR* InPrimPath )
+{
+	SAssignNew( HeaderRowWidget, SHeaderRow )
+
+	+SHeaderRow::Column( FName( TEXT("PropertyName") ) )
+	.DefaultLabel( NSLOCTEXT( "SUsdIntegrationsPanel", "Integrations", "Integrations" ) )
+	.FillWidth( 30.0f )
+
+	+SHeaderRow::Column( FName( TEXT("Value") ) )
+	.DefaultLabel( FText::GetEmpty() )
+	.FillWidth( 70.0f );
+
+	SListView::Construct
+	(
+		SListView::FArguments()
+		.ListItemsSource( &Attributes )
+		.OnGenerateRow( this, &SUsdIntegrationsPanel::OnGenerateRow )
+		.HeaderRow( HeaderRowWidget )
+	);
+
+	SetPrimPath( InUsdStage, InPrimPath );
+}
+
+void SUsdIntegrationsPanel::SetPrimPath( const UE::FUsdStageWeak& InUsdStage, const TCHAR* InPrimPath )
+{
+	PrimPath = InPrimPath;
+	UsdStage = InUsdStage;
+
+	Attributes.Reset();
+
+	EVisibility NewVisibility = EVisibility::Collapsed;
+
+	if ( InUsdStage )
+	{
+		if ( UE::FUsdPrim Prim = InUsdStage.GetPrimAtPath( UE::FSdfPath{ InPrimPath } ) )
+		{
+			if ( UsdUtils::PrimHasLiveLinkSchema( Prim ) )
+			{
+				NewVisibility = EVisibility::Visible;
+
+				// Only show the dedicated AnimBlueprint picker for SkelRoots
+				if ( Prim.IsA( TEXT( "SkelRoot" ) ) )
+				{
+					if ( UE::FUsdAttribute Attr = Prim.GetAttribute( *UsdToUnreal::ConvertToken( UnrealIdentifiers::UnrealLiveLinkAnimBlueprintPath ) ) )
+					{
+						Attributes.Add( MakeShared<UE::FUsdAttribute>( Attr ) );
+					}
+				}
+
+				if ( UE::FUsdAttribute Attr = Prim.GetAttribute( *UsdToUnreal::ConvertToken( UnrealIdentifiers::UnrealLiveLinkSubjectName ) ) )
+				{
+					Attributes.Add( MakeShared<UE::FUsdAttribute>( Attr ) );
+				}
+
+				if ( UE::FUsdAttribute Attr = Prim.GetAttribute( *UsdToUnreal::ConvertToken( UnrealIdentifiers::UnrealLiveLinkEnabled ) ) )
+				{
+					Attributes.Add( MakeShared<UE::FUsdAttribute>( Attr ) );
+				}
+			}
+		}
+	}
+
+	SetVisibility( NewVisibility );
+	RequestListRefresh();
+}
+
+TSharedRef< ITableRow > SUsdIntegrationsPanel::OnGenerateRow( TSharedPtr<UE::FUsdAttribute> InAttr, const TSharedRef< STableViewBase >& OwnerTable )
+{
+	return SNew( SUsdIntegrationsPanelRow, InAttr, OwnerTable );
+}
+
+#undef LOCTEXT_NAMESPACE
\ No newline at end of file
diff --git a/Engine/Plugins/Importers/USDImporter/Source/USDStageEditor/Private/Widgets/SUSDIntegrationsPanel.h b/Engine/Plugins/Importers/USDImporter/Source/USDStageEditor/Private/Widgets/SUSDIntegrationsPanel.h
new file mode 100644
--- /dev/null
+++ b/Engine/Plugins/Importers/USDImporter/Source/USDStageEditor/Private/Widgets/SUSDIntegrationsPanel.h
@@ -0,0 +1,67 @@
+// Copyright Epic Games, Inc. All Rights Reserved.
+
+#pragma once
+
+#include "UsdWrappers/UsdStage.h"
+
+#include "Widgets/Views/SListView.h"
+
+#include "SUSDIntegrationsPanel.generated.h"
+
+class SHeaderRow;
+namespace UE
+{
+	class FUsdAttribute;
+}
+
+// We need an actual UObject and UPROPERTY to use the property editor module and generate one of the
+// standard object picker widgets, so we'll be using the CDO of this class to do that
+UCLASS(Abstract, Transient, MinimalAPI)
+class UUsdIntegrationsPanelPropertyDummy : public UObject
+{
+public:
+	GENERATED_BODY()
+
+	// Ideally this would be an FSoftObjectPath, but FPropertyEditorModule doesn't support generating
+	// widgets for those properties yet
+	UPROPERTY(EditAnywhere, Transient, Category=Dummy, meta=(AllowedClasses="AnimBlueprint"))
+	UObject* AnimBPProperty;
+};
+
+class SUsdIntegrationsPanelRow : public SMultiColumnTableRow< TSharedPtr<UE::FUsdAttribute> >
+{
+public:
+	SLATE_BEGIN_ARGS( SUsdIntegrationsPanelRow ) {}
+	SLATE_END_ARGS()
+
+public:
+	void Construct( const FArguments& InArgs, TSharedPtr<UE::FUsdAttribute> InAttr, const TSharedRef< STableViewBase >& OwnerTable );
+
+	virtual TSharedRef< SWidget > GenerateWidgetForColumn( const FName& ColumnName ) override;
+
+private:
+	TSharedPtr<UE::FUsdAttribute> Attribute;
+};
+
+// We don't really need a list view here since we'll mostly always know exactly what attributes are going to
+// be displayed here beforehand, but doing so is a simple way of ensuring a consistent look between this panel
+// and the variants/references panels, that *do* need to be lists
+class SUsdIntegrationsPanel : public SListView< TSharedPtr<UE::FUsdAttribute> >
+{
+	SLATE_BEGIN_ARGS( SUsdIntegrationsPanel ) {}
+	SLATE_END_ARGS()
+
+public:
+	void Construct( const FArguments& InArgs, const UE::FUsdStageWeak& InUsdStage, const TCHAR* InPrimPath );
+	void SetPrimPath( const UE::FUsdStageWeak& UsdStage, const TCHAR* InPrimPath );
+
+protected:
+	TSharedRef< ITableRow > OnGenerateRow( TSharedPtr<UE::FUsdAttribute> InAttr, const TSharedRef< STableViewBase >& OwnerTable );
+
+private:
+	TSharedPtr< SHeaderRow > HeaderRowWidget;
+
+	TArray<TSharedPtr<UE::FUsdAttribute>> Attributes;
+	UE::FUsdStageWeak UsdStage;
+	FString PrimPath;
+};
diff --git a/Engine/Plugins/Importers/USDImporter/Source/USDStageEditor/Private/Widgets/SUSDReferencesList.cpp b/Engine/Plugins/Importers/USDImporter/Source/USDStageEditor/Private/Widgets/SUSDReferencesList.cpp
--- a/Engine/Plugins/Importers/USDImporter/Source/USDStageEditor/Private/Widgets/SUSDReferencesList.cpp
+++ b/Engine/Plugins/Importers/USDImporter/Source/USDStageEditor/Private/Widgets/SUSDReferencesList.cpp
@@ -57,8 +57,6 @@ TSharedRef< SWidget > SUsdReferenceRow::GenerateWidgetForColumn( const FName& Co
 
 void SUsdReferencesList::Construct( const FArguments& InArgs, const UE::FUsdStageWeak& UsdStage, const TCHAR* PrimPath )
 {
-	ViewModel.UpdateReferences( UsdStage, PrimPath );
-
 	SAssignNew( HeaderRowWidget, SHeaderRow )
 
 	+SHeaderRow::Column( FName( TEXT("AssetPath") ) )
@@ -72,6 +70,8 @@ void SUsdReferencesList::Construct( const FArguments& InArgs, const UE::FUsdStag
 		.OnGenerateRow( this, &SUsdReferencesList::OnGenerateRow )
 		.HeaderRow( HeaderRowWidget )
 	);
+
+	SetPrimPath( UsdStage, PrimPath );
 }
 
 TSharedRef< ITableRow > SUsdReferencesList::OnGenerateRow( TSharedPtr< FUsdReference > InDisplayNode, const TSharedRef< STableViewBase >& OwnerTable )
@@ -82,6 +82,20 @@ TSharedRef< ITableRow > SUsdReferencesList::OnGenerateRow( TSharedPtr< FUsdRefer
 void SUsdReferencesList::SetPrimPath( const UE::FUsdStageWeak& UsdStage, const TCHAR* PrimPath )
 {
 	ViewModel.UpdateReferences( UsdStage, PrimPath );
+
+	EVisibility NewVisibility = EVisibility::Collapsed;
+	if ( UsdStage )
+	{
+		if ( UE::FUsdPrim UsdPrim = UsdStage.GetPrimAtPath( UE::FSdfPath{ PrimPath } ) )
+		{
+			if ( UsdPrim.HasAuthoredReferences() )
+			{
+				NewVisibility = EVisibility::Visible;
+			}
+		}
+	}
+	SetVisibility( NewVisibility );
+
 	RequestListRefresh();
 }
 
diff --git a/Engine/Plugins/Importers/USDImporter/Source/USDStageEditor/Private/Widgets/SUSDVariantSetsList.cpp b/Engine/Plugins/Importers/USDImporter/Source/USDStageEditor/Private/Widgets/SUSDVariantSetsList.cpp
--- a/Engine/Plugins/Importers/USDImporter/Source/USDStageEditor/Private/Widgets/SUSDVariantSetsList.cpp
+++ b/Engine/Plugins/Importers/USDImporter/Source/USDStageEditor/Private/Widgets/SUSDVariantSetsList.cpp
@@ -86,10 +86,6 @@ void SUsdVariantRow::OnSelectionChanged( TSharedPtr< FString > NewValue, ESelect
 
 void SVariantsList::Construct( const FArguments& InArgs, const UE::FUsdStageWeak& UsdStage, const TCHAR* InPrimPath )
 {
-	PrimPath = InPrimPath;
-
-	ViewModel.UpdateVariantSets( UsdStage, InPrimPath );
-
 	SAssignNew( HeaderRowWidget, SHeaderRow )
 
 	+SHeaderRow::Column( FName( TEXT("VariantSetName") ) )
@@ -107,6 +103,8 @@ void SVariantsList::Construct( const FArguments& InArgs, const UE::FUsdStageWeak
 		.OnGenerateRow( this, &SVariantsList::OnGenerateRow )
 		.HeaderRow( HeaderRowWidget )
 	);
+
+	SetPrimPath( UsdStage, InPrimPath );
 }
 
 TSharedRef< ITableRow > SVariantsList::OnGenerateRow( TSharedPtr< FUsdVariantSetViewModel > InDisplayNode, const TSharedRef< STableViewBase >& OwnerTable )
@@ -119,6 +117,20 @@ void SVariantsList::SetPrimPath( const UE::FUsdStageWeak& UsdStage, const TCHAR*
 {
 	PrimPath = InPrimPath;
 	ViewModel.UpdateVariantSets( UsdStage, *PrimPath );
+
+	EVisibility NewVisibility = EVisibility::Collapsed;
+	if ( UsdStage )
+	{
+		if ( UE::FUsdPrim UsdPrim = UsdStage.GetPrimAtPath( UE::FSdfPath{ InPrimPath } ) )
+		{
+			if ( UsdPrim.HasVariantSets() )
+			{
+				NewVisibility = EVisibility::Visible;
+			}
+		}
+	}
+	SetVisibility( NewVisibility );
+
 	RequestListRefresh();
 }
 
diff --git a/Engine/Plugins/Importers/USDImporter/Source/USDStageEditor/USDStageEditor.Build.cs b/Engine/Plugins/Importers/USDImporter/Source/USDStageEditor/USDStageEditor.Build.cs
--- a/Engine/Plugins/Importers/USDImporter/Source/USDStageEditor/USDStageEditor.Build.cs
+++ b/Engine/Plugins/Importers/USDImporter/Source/USDStageEditor/USDStageEditor.Build.cs
@@ -21,6 +21,8 @@ namespace UnrealBuildTool.Rules
 					"Engine",
 					"InputCore",
 					"LevelEditor",
+					"LiveLinkEditor", // For SLiveLinkSubjectRepresentationPicker
+					"LiveLinkInterface", // For the ULiveLinkAnimation/TransformRole classes
 					"Projects", // So that we can use the IPluginManager, required for our custom style
 					"SceneOutliner",
 					"Slate",
diff --git a/Engine/Plugins/Importers/USDImporter/Source/USDStageEditorViewModels/Private/USDPrimViewModel.cpp b/Engine/Plugins/Importers/USDImporter/Source/USDStageEditorViewModels/Private/USDPrimViewModel.cpp
--- a/Engine/Plugins/Importers/USDImporter/Source/USDStageEditorViewModels/Private/USDPrimViewModel.cpp
+++ b/Engine/Plugins/Importers/USDImporter/Source/USDStageEditorViewModels/Private/USDPrimViewModel.cpp
@@ -237,6 +237,113 @@ void FUsdPrimViewModel::TogglePayload()
 	}
 }
 
+void FUsdPrimViewModel::SetUpLiveLink()
+{
+#if USE_USD_SDK
+
+	FScopedUsdAllocs Allocs;
+
+	pxr::UsdPrim Prim{ UsdPrim };
+	if ( !Prim )
+	{
+		return;
+	}
+
+	pxr::TfType Schema = pxr::UsdSchemaRegistry::GetTypeFromSchemaTypeName( UnrealIdentifiers::LiveLinkAPI );
+	ensure( static_cast< bool >( Schema ) );
+
+	if ( Prim.HasAPI( Schema ) )
+	{
+		return;
+	}
+
+	ensure( Prim.ApplyAPI( Schema ) );
+
+	// It seems codeless schemas can't have fallback values, so just create the attributes here, and set some
+	// sensible defaults
+	const bool bCustom = false;
+
+	if ( pxr::UsdAttribute SubjectNameAttr = Prim.CreateAttribute( UnrealIdentifiers::UnrealLiveLinkSubjectName, pxr::SdfValueTypeNames->String, bCustom, pxr::SdfVariabilityUniform ) )
+	{
+		if ( !SubjectNameAttr.HasAuthoredValue() )
+		{
+			SubjectNameAttr.Set( UnrealToUsd::ConvertString( TEXT( "root" ) ).Get(), pxr::UsdTimeCode::Default() );
+		}
+	}
+
+	if ( pxr::UsdAttribute AnimBPPathAttr = Prim.CreateAttribute( UnrealIdentifiers::UnrealLiveLinkAnimBlueprintPath, pxr::SdfValueTypeNames->String, bCustom, pxr::SdfVariabilityUniform ) )
+	{
+		if ( !AnimBPPathAttr.HasAuthoredValue() )
+		{
+			AnimBPPathAttr.Set( UnrealToUsd::ConvertString( TEXT( "/USDImporter/Blueprint/DefaultLiveLinkAnimBP.DefaultLiveLinkAnimBP" ) ).Get(), pxr::UsdTimeCode::Default() );
+		}
+	}
+
+	if ( pxr::UsdAttribute EnabledAttr = Prim.CreateAttribute( UnrealIdentifiers::UnrealLiveLinkEnabled, pxr::SdfValueTypeNames->Bool, bCustom, pxr::SdfVariabilityUniform ) )
+	{
+		if ( !EnabledAttr.HasAuthoredValue() )
+		{
+			EnabledAttr.Set( true, pxr::UsdTimeCode::Default() );
+		}
+	}
+
+#endif // #if USE_USD_SDK
+}
+
+bool FUsdPrimViewModel::CanSetUpLiveLink() const
+{
+#if USE_USD_SDK
+	FScopedUsdAllocs UsdAllocs;
+
+	pxr::UsdPrim Prim{ UsdPrim };
+	if ( !Prim )
+	{
+		return false;
+	}
+
+	pxr::TfType Schema = pxr::UsdSchemaRegistry::GetTypeFromSchemaTypeName( UnrealIdentifiers::LiveLinkAPI );
+	if ( !ensure( static_cast< bool >( Schema ) ) )
+	{
+		return false;
+	}
+
+	return !Prim.HasAPI( Schema ) && Prim.CanApplyAPI( Schema );
+#endif // #if USE_USD_SDK
+}
+
+void FUsdPrimViewModel::RemoveLiveLink()
+{
+#if USE_USD_SDK
+	FScopedUsdAllocs UsdAllocs;
+
+	pxr::TfType Schema = pxr::UsdSchemaRegistry::GetTypeFromSchemaTypeName( UnrealIdentifiers::LiveLinkAPI );
+	ensure( static_cast< bool >( Schema ) );
+
+	ensure( pxr::UsdPrim( UsdPrim ).RemoveAPI( Schema ) );
+#endif // #if USE_USD_SDK
+}
+
+bool FUsdPrimViewModel::CanRemoveLiveLink() const
+{
+#if USE_USD_SDK
+	FScopedUsdAllocs UsdAllocs;
+
+	pxr::UsdPrim Prim{ UsdPrim };
+	if ( !Prim )
+	{
+		return false;
+	}
+
+	pxr::TfType Schema = pxr::UsdSchemaRegistry::GetTypeFromSchemaTypeName( UnrealIdentifiers::LiveLinkAPI );
+	if ( !ensure( static_cast< bool >( Schema ) ) )
+	{
+		return false;
+	}
+
+	return Prim.HasAPI( Schema );
+#endif // #if USE_USD_SDK
+}
+
 void FUsdPrimViewModel::DefinePrim( const TCHAR* PrimName )
 {
 #if USE_USD_SDK
diff --git a/Engine/Plugins/Importers/USDImporter/Source/USDStageEditorViewModels/Private/USDReferencesViewModel.cpp b/Engine/Plugins/Importers/USDImporter/Source/USDStageEditorViewModels/Private/USDReferencesViewModel.cpp
--- a/Engine/Plugins/Importers/USDImporter/Source/USDStageEditorViewModels/Private/USDReferencesViewModel.cpp
+++ b/Engine/Plugins/Importers/USDImporter/Source/USDStageEditorViewModels/Private/USDReferencesViewModel.cpp
@@ -31,23 +31,24 @@ void FUsdReferencesViewModel::UpdateReferences( const UE::FUsdStageWeak& UsdStag
 #if USE_USD_SDK
 	FScopedUsdAllocs UsdAllocs;
 
-	pxr::UsdPrim Prim( UsdStage.GetPrimAtPath( UE::FSdfPath( PrimPath ) ) );
-
-	pxr::UsdPrimCompositionQuery PrimCompositionQuery = pxr::UsdPrimCompositionQuery::GetDirectReferences( Prim );
-
-	for ( const pxr::UsdPrimCompositionQueryArc& CompositionArc : PrimCompositionQuery.GetCompositionArcs() )
+	if ( pxr::UsdPrim Prim{ UsdStage.GetPrimAtPath( UE::FSdfPath( PrimPath ) ) } )
 	{
-		if ( CompositionArc.GetArcType() == pxr::PcpArcTypeReference )
-		{
-			pxr::SdfReferenceEditorProxy ReferenceEditor;
-			pxr::SdfReference UsdReference;
+		pxr::UsdPrimCompositionQuery PrimCompositionQuery = pxr::UsdPrimCompositionQuery::GetDirectReferences( Prim );
 
-			if ( CompositionArc.GetIntroducingListEditor( &ReferenceEditor, &UsdReference ) )
+		for ( const pxr::UsdPrimCompositionQueryArc& CompositionArc : PrimCompositionQuery.GetCompositionArcs() )
+		{
+			if ( CompositionArc.GetArcType() == pxr::PcpArcTypeReference )
 			{
-				FUsdReference Reference;
-				Reference.AssetPath = UsdToUnreal::ConvertString( UsdReference.GetAssetPath() );
+				pxr::SdfReferenceEditorProxy ReferenceEditor;
+				pxr::SdfReference UsdReference;
+
+				if ( CompositionArc.GetIntroducingListEditor( &ReferenceEditor, &UsdReference ) )
+				{
+					FUsdReference Reference;
+					Reference.AssetPath = UsdToUnreal::ConvertString( UsdReference.GetAssetPath() );
 
-				References.Add( MakeSharedUnreal< FUsdReference >( MoveTemp( Reference ) ) );
+					References.Add( MakeSharedUnreal< FUsdReference >( MoveTemp( Reference ) ) );
+				}
 			}
 		}
 	}
diff --git a/Engine/Plugins/Importers/USDImporter/Source/USDStageEditorViewModels/Public/USDPrimViewModel.h b/Engine/Plugins/Importers/USDImporter/Source/USDStageEditorViewModels/Public/USDPrimViewModel.h
--- a/Engine/Plugins/Importers/USDImporter/Source/USDStageEditorViewModels/Public/USDPrimViewModel.h
+++ b/Engine/Plugins/Importers/USDImporter/Source/USDStageEditorViewModels/Public/USDPrimViewModel.h
@@ -48,6 +48,12 @@ public:
 	void ToggleVisibility();
 	void TogglePayload();
 
+	void SetUpLiveLink();
+	bool CanSetUpLiveLink() const;
+
+	void RemoveLiveLink();
+	bool CanRemoveLiveLink() const;
+
 	void DefinePrim( const TCHAR* PrimName );
 
 	/** Delegate for hooking up an inline editable text block to be notified that a rename is requested. */
diff --git a/Engine/Plugins/Importers/USDImporter/Source/USDStageImporter/Private/USDStageImporter.cpp b/Engine/Plugins/Importers/USDImporter/Source/USDStageImporter/Private/USDStageImporter.cpp
--- a/Engine/Plugins/Importers/USDImporter/Source/USDStageImporter/Private/USDStageImporter.cpp
+++ b/Engine/Plugins/Importers/USDImporter/Source/USDStageImporter/Private/USDStageImporter.cpp
@@ -19,6 +19,7 @@
 #include "UsdWrappers/UsdStage.h"
 #include "UsdWrappers/UsdTyped.h"
 
+#include "Animation/AnimBlueprint.h"
 #include "Animation/AnimSequence.h"
 #include "Animation/Skeleton.h"
 #include "AssetRegistry/AssetRegistryModule.h"
@@ -664,7 +665,7 @@ namespace UsdStageImporterImpl
 				{
 					// For skeletal stuff, the primpaths point to the SkelRoot, so it is useful to place the assets in there,
 					// as we'll always have at least the skeletal mesh and the skeleton
-					if ( Asset->IsA( USkeletalMesh::StaticClass() ) || Asset->IsA( USkeleton::StaticClass() ) || Asset->IsA( UAnimSequence::StaticClass() ) )
+					if ( Asset->IsA( USkeletalMesh::StaticClass() ) || Asset->IsA( USkeleton::StaticClass() ) || Asset->IsA( UAnimSequence::StaticClass() ) || Asset->IsA( UAnimBlueprint::StaticClass() ) )
 					{
 						AssetTypeFolder = ImportData->PrimPath;
 					}
@@ -692,7 +693,7 @@ namespace UsdStageImporterImpl
 				{
 					AssetTypeFolder = TEXT("Textures");
 				}
-				else if (Asset->IsA(USkeletalMesh::StaticClass()) || Asset->IsA(USkeleton::StaticClass()) || Asset->IsA(UAnimSequence::StaticClass()))
+				else if (Asset->IsA(USkeletalMesh::StaticClass()) || Asset->IsA(USkeleton::StaticClass()) || Asset->IsA(UAnimSequence::StaticClass()) || Asset->IsA(UAnimBlueprint::StaticClass()))
 				{
 					AssetTypeFolder = TEXT("SkeletalMeshes");
 				}
@@ -1201,6 +1202,10 @@ namespace UsdStageImporterImpl
 				{
 					// Do nothing. Skeletons have no additional dependencies
 				}
+				else if ( UAnimBlueprint* AnimBP = Cast<UAnimBlueprint>( Asset ) )
+				{
+					// Do nothing. AnimBlueprints have no additional dependencies
+				}
 				else
 				{
 					UE_LOG( LogUsd, Warning, TEXT( "Unknown asset '%s' encountered when collecting used assets before USD import." ), Asset ? *Asset->GetName() : TEXT( "nullptr" ) );
diff --git a/Engine/Plugins/Importers/USDImporter/Source/USDUtilities/Private/USDConversionUtils.cpp b/Engine/Plugins/Importers/USDImporter/Source/USDUtilities/Private/USDConversionUtils.cpp
--- a/Engine/Plugins/Importers/USDImporter/Source/USDUtilities/Private/USDConversionUtils.cpp
+++ b/Engine/Plugins/Importers/USDImporter/Source/USDUtilities/Private/USDConversionUtils.cpp
@@ -15,6 +15,7 @@
 #include "UsdWrappers/UsdStage.h"
 
 #include "Algo/Copy.h"
+#include "Animation/AnimBlueprint.h"
 #include "Animation/AnimSequence.h"
 #include "Animation/Skeleton.h"
 #include "CineCameraActor.h"
@@ -1047,6 +1048,18 @@ UUsdAssetImportData* UsdUtils::GetAssetImportData( UObject* Asset )
 			ImportData = Cast<UUsdAssetImportData>( SkMesh->GetAssetImportData() );
 		}
 	}
+	else if ( UAnimBlueprint* AnimBP = Cast<UAnimBlueprint>( Asset ) )
+	{
+		// We will always have a skeleton, but not necessarily we will have a preview mesh directly
+		// on the UAnimBlueprint
+		if ( USkeleton* AnimBPSkeleton = AnimBP->TargetSkeleton.Get() )
+		{
+			if ( USkeletalMesh* SkMesh = AnimBPSkeleton->GetPreviewMesh() )
+			{
+				ImportData = Cast<UUsdAssetImportData>( SkMesh->GetAssetImportData() );
+			}
+		}
+	}
 	else if ( USkeletalMesh* SkMesh = Cast<USkeletalMesh>( Asset ) )
 	{
 		ImportData = Cast<UUsdAssetImportData>(SkMesh->GetAssetImportData());
diff --git a/Engine/Plugins/Importers/USDImporter/Source/USDUtilities/Private/USDIntegrationUtils.cpp b/Engine/Plugins/Importers/USDImporter/Source/USDUtilities/Private/USDIntegrationUtils.cpp
new file mode 100644
--- /dev/null
+++ b/Engine/Plugins/Importers/USDImporter/Source/USDUtilities/Private/USDIntegrationUtils.cpp
@@ -0,0 +1,83 @@
+// Copyright Epic Games, Inc. All Rights Reserved.
+
+#include "USDIntegrationUtils.h"
+
+#include "UnrealUSDWrapper.h"
+#include "USDTypesConversion.h"
+#include "UsdWrappers/SdfChangeBlock.h"
+
+#include "Animation/AnimInstance.h"
+#include "Components/SkeletalMeshComponent.h"
+#include "LiveLinkComponentController.h"
+
+#if USE_USD_SDK
+#include "USDIncludesStart.h"
+	#include "pxr/usd/usd/attribute.h"
+	#include "pxr/usd/usd/prim.h"
+#include "USDIncludesEnd.h"
+
+#define LOCTEXT_NAMESPACE "USDIntegrationUtils"
+
+bool UsdUtils::PrimHasLiveLinkSchema( const pxr::UsdPrim& Prim )
+{
+	if ( !Prim )
+	{
+		return false;
+	}
+
+    FScopedUsdAllocs Allocs;
+
+    pxr::TfType Schema = pxr::UsdSchemaRegistry::GetTypeFromSchemaTypeName( UnrealIdentifiers::LiveLinkAPI );
+    return Prim.HasAPI( Schema );
+}
+
+void UnrealToUsd::ConvertLiveLinkProperties( const UActorComponent* InComponent, pxr::UsdPrim& InOutPrim )
+{
+#if WITH_EDITOR
+	if ( !InOutPrim || !UsdUtils::PrimHasLiveLinkSchema( InOutPrim ) )
+	{
+		return;
+	}
+
+	FScopedUsdAllocs Allocs;
+
+	UE::FSdfChangeBlock ChangeBlock;
+
+	// Skeletal LiveLink case
+	if ( const USkeletalMeshComponent* SkeletalComponent = Cast<USkeletalMeshComponent>( InComponent ) )
+	{
+		if ( pxr::UsdAttribute Attr = InOutPrim.CreateAttribute( UnrealIdentifiers::UnrealLiveLinkAnimBlueprintPath, pxr::SdfValueTypeNames->String ) )
+		{
+			FString AnimBPPath = SkeletalComponent->AnimClass && SkeletalComponent->AnimClass->ClassGeneratedBy
+				? SkeletalComponent->AnimClass->ClassGeneratedBy->GetPathName()
+				: FString();
+
+			Attr.Set( UnrealToUsd::ConvertString( *AnimBPPath ).Get(), pxr::UsdTimeCode::Default() );
+		}
+
+		if ( pxr::UsdAttribute Attr = InOutPrim.CreateAttribute( UnrealIdentifiers::UnrealLiveLinkEnabled, pxr::SdfValueTypeNames->Bool ) )
+		{
+			const bool bEnabled = SkeletalComponent->GetAnimationMode() == EAnimationMode::AnimationBlueprint;
+			Attr.Set( bEnabled, pxr::UsdTimeCode::Default() );
+		}
+	}
+	// Non-skeletal LiveLink case
+	else if ( const ULiveLinkComponentController* InController = Cast<ULiveLinkComponentController>( InComponent ) )
+	{
+		if ( pxr::UsdAttribute Attr = InOutPrim.CreateAttribute( UnrealIdentifiers::UnrealLiveLinkSubjectName, pxr::SdfValueTypeNames->String ) )
+		{
+			Attr.Set( UnrealToUsd::ConvertString( *InController->SubjectRepresentation.Subject.Name.ToString() ).Get(), pxr::UsdTimeCode::Default() );
+		}
+
+		if ( pxr::UsdAttribute Attr = InOutPrim.CreateAttribute( UnrealIdentifiers::UnrealLiveLinkEnabled, pxr::SdfValueTypeNames->Bool ) )
+		{
+			Attr.Set( InController->bEvaluateLiveLink, pxr::UsdTimeCode::Default() );
+		}
+	}
+#else
+	return false;
+#endif // WITH_EDITOR
+}
+
+#undef LOCTEXT_NAMESPACE
+#endif // USE_USD_SDK
\ No newline at end of file
diff --git a/Engine/Plugins/Importers/USDImporter/Source/USDUtilities/Public/USDIntegrationUtils.h b/Engine/Plugins/Importers/USDImporter/Source/USDUtilities/Public/USDIntegrationUtils.h
new file mode 100644
--- /dev/null
+++ b/Engine/Plugins/Importers/USDImporter/Source/USDUtilities/Public/USDIntegrationUtils.h
@@ -0,0 +1,44 @@
+// Copyright Epic Games, Inc. All Rights Reserved.
+
+#pragma once
+
+#include "CoreMinimal.h"
+
+#if USE_USD_SDK
+#include "USDIncludesStart.h"
+	#include "pxr/pxr.h"
+#include "USDIncludesEnd.h"
+
+PXR_NAMESPACE_OPEN_SCOPE
+	class UsdPrim;
+PXR_NAMESPACE_CLOSE_SCOPE
+
+#endif // #if USE_USD_SDK
+
+class UActorComponent;
+class ULiveLinkComponentController;
+
+namespace UsdUtils
+{
+#if USE_USD_SDK
+
+	// Returns true in case Prim has our custom "LiveLinkAPI" schema
+	USDUTILITIES_API bool PrimHasLiveLinkSchema( const pxr::UsdPrim& Prim );
+
+#endif // USE_USD_SDK
+}
+
+namespace UnrealToUsd
+{
+#if USE_USD_SDK
+
+	/**
+	 * Converts UE component properties related to LiveLink into values for the attributes of our custom LiveLinkAPI
+	 * schema.
+	 * @param InComponent		The main component with data to convert
+	 * @param InOutPrim			Prim with the LiveLinkAPI schema to receive the data
+	 */
+	USDUTILITIES_API void ConvertLiveLinkProperties( const UActorComponent* InComponent, pxr::UsdPrim& InOutPrim );
+
+#endif // USE_USD_SDK
+}
diff --git a/Engine/Plugins/Importers/USDImporter/Source/USDUtilities/USDUtilities.Build.cs b/Engine/Plugins/Importers/USDImporter/Source/USDUtilities/USDUtilities.Build.cs
--- a/Engine/Plugins/Importers/USDImporter/Source/USDUtilities/USDUtilities.Build.cs
+++ b/Engine/Plugins/Importers/USDImporter/Source/USDUtilities/USDUtilities.Build.cs
@@ -30,6 +30,7 @@ namespace UnrealBuildTool.Rules
 					"GeometryCache", // Just so that we can fetch its AssetImportData
 					"IntelTBB",
 					"Landscape", // So that GetSchemaNameForComponent knows what to do with landscape proxies
+					"LiveLinkComponents", // For converting LiveLinkComponentController properties to USD
 					"MeshDescription",
 					"MovieScene",
 					"MovieSceneTracks",
diff --git a/Engine/Plugins/Importers/USDImporter/Source/UnrealUSDWrapper/Private/UnrealUSDWrapper.cpp b/Engine/Plugins/Importers/USDImporter/Source/UnrealUSDWrapper/Private/UnrealUSDWrapper.cpp
--- a/Engine/Plugins/Importers/USDImporter/Source/UnrealUSDWrapper/Private/UnrealUSDWrapper.cpp
+++ b/Engine/Plugins/Importers/USDImporter/Source/UnrealUSDWrapper/Private/UnrealUSDWrapper.cpp
@@ -106,6 +106,11 @@ namespace UnrealIdentifiers
 	const TfToken UnrealNaniteOverrideEnable = TfToken("enable");
 	const TfToken UnrealNaniteOverrideDisable = TfToken("disable");
 
+	const TfToken LiveLinkAPI = TfToken("LiveLinkAPI");
+	const TfToken UnrealLiveLinkAnimBlueprintPath = TfToken("unrealLiveLink:animBlueprintPath");
+	const TfToken UnrealLiveLinkSubjectName = TfToken("unrealLiveLink:subjectName");
+	const TfToken UnrealLiveLinkEnabled = TfToken("unrealLiveLink:enabled");
+
 	const TfToken DiffuseColor = TfToken("diffuseColor");
 	const TfToken EmissiveColor = TfToken("emissiveColor");
 	const TfToken Metallic = TfToken("metallic");
diff --git a/Engine/Plugins/Importers/USDImporter/Source/UnrealUSDWrapper/Public/UnrealUSDWrapper.h b/Engine/Plugins/Importers/USDImporter/Source/UnrealUSDWrapper/Public/UnrealUSDWrapper.h
--- a/Engine/Plugins/Importers/USDImporter/Source/UnrealUSDWrapper/Public/UnrealUSDWrapper.h
+++ b/Engine/Plugins/Importers/USDImporter/Source/UnrealUSDWrapper/Public/UnrealUSDWrapper.h
@@ -193,7 +193,7 @@ public:
 
 	/**
 	 * Returns the file extensions of all file formats supported by USD.
-	 * 
+	 *
 	 * These include the extensions for the file formats built into USD as
 	 * well as those for any other file formats introduced through the USD
 	 * plugin system (e.g. "abc" for the Alembic file format of the usdAbc
@@ -203,7 +203,7 @@ public:
 
 	/**
 	 * Returns the file extensions that are native to USD.
-	 * 
+	 *
 	 * These are the extensions for the file formats built into USD (i.e. "usd", "usda", "usdc", and "usdz").
 	 */
 	UNREALUSDWRAPPER_API static TArray<FString> GetNativeFileFormats();
@@ -347,6 +347,11 @@ namespace UnrealIdentifiers
 	extern UNREALUSDWRAPPER_API const pxr::TfToken UnrealNaniteOverrideEnable;
 	extern UNREALUSDWRAPPER_API const pxr::TfToken UnrealNaniteOverrideDisable;
 
+	extern UNREALUSDWRAPPER_API const pxr::TfToken LiveLinkAPI;
+	extern UNREALUSDWRAPPER_API const pxr::TfToken UnrealLiveLinkAnimBlueprintPath;
+	extern UNREALUSDWRAPPER_API const pxr::TfToken UnrealLiveLinkSubjectName;
+	extern UNREALUSDWRAPPER_API const pxr::TfToken UnrealLiveLinkEnabled;
+
 	extern UNREALUSDWRAPPER_API const pxr::TfToken DiffuseColor;
 	extern UNREALUSDWRAPPER_API const pxr::TfToken EmissiveColor;
 	extern UNREALUSDWRAPPER_API const pxr::TfToken Metallic;
diff --git a/Engine/Plugins/Importers/USDImporter/USDImporter.uplugin b/Engine/Plugins/Importers/USDImporter/USDImporter.uplugin
--- a/Engine/Plugins/Importers/USDImporter/USDImporter.uplugin
+++ b/Engine/Plugins/Importers/USDImporter/USDImporter.uplugin
@@ -95,6 +95,10 @@
 			"Name": "LevelSequenceEditor",
 			"Enabled": true
 		},
+		{
+			"Name": "LiveLink",
+			"Enabled": true
+		},
 		{
 			"Name": "MDLImporter",
 			"Enabled": true
