USD: Add analytics to ExportAll, ExportFlattened, and ExportLayers workflows.

diff --git a/Engine/Plugins/Importers/USDImporter/Source/USDExporter/Private/AnimSequenceExporterUSD.cpp b/Engine/Plugins/Importers/USDImporter/Source/USDExporter/Private/AnimSequenceExporterUSD.cpp
--- a/Engine/Plugins/Importers/USDImporter/Source/USDExporter/Private/AnimSequenceExporterUSD.cpp
+++ b/Engine/Plugins/Importers/USDImporter/Source/USDExporter/Private/AnimSequenceExporterUSD.cpp
@@ -266,7 +266,14 @@ bool UAnimSequenceExporterUSD::ExportBinary( UObject* Object, const TCHAR* Type,
 		double ElapsedSeconds = FPlatformTime::ToSeconds64( FPlatformTime::Cycles64() - StartTime );
 		FString Extension = FPaths::GetExtension( UExporter::CurrentFilename );
 
-		UE::AnimSequenceExporterUSD::Private::SendAnalytics( Object, Options, bAutomated, ElapsedSeconds, EndTimeCode - StartTimeCode, Extension );
+		UE::AnimSequenceExporterUSD::Private::SendAnalytics(
+			Object,
+			Options,
+			bAutomated,
+			ElapsedSeconds,
+			UsdUtils::GetUsdStageNumFrames( AssetStage ),
+			Extension
+		);
 	}
 
 	return true;
diff --git a/Engine/Plugins/Importers/USDImporter/Source/USDExporter/Private/SkeletalMeshExporterUSD.cpp b/Engine/Plugins/Importers/USDImporter/Source/USDExporter/Private/SkeletalMeshExporterUSD.cpp
--- a/Engine/Plugins/Importers/USDImporter/Source/USDExporter/Private/SkeletalMeshExporterUSD.cpp
+++ b/Engine/Plugins/Importers/USDImporter/Source/USDExporter/Private/SkeletalMeshExporterUSD.cpp
@@ -210,9 +210,16 @@ bool USkeletalMeshExporterUsd::ExportBinary( UObject* Object, const TCHAR* Type,
 		bool bAutomated = ExportTask ? ExportTask->bAutomated : false;
 		double ElapsedSeconds = FPlatformTime::ToSeconds64( FPlatformTime::Cycles64() - StartTime );
 		FString Extension = FPaths::GetExtension( UExporter::CurrentFilename );
-		double NumberOfFrames = UsdStage.GetEndTimeCode() - UsdStage.GetStartTimeCode();
-
-		UE::SkeletalMeshExporterUSD::Private::SendAnalytics( Object, Options, bAutomated, ElapsedSeconds, NumberOfFrames, Extension );
+		double NumberOfFrames = 1 + UsdStage.GetEndTimeCode() - UsdStage.GetStartTimeCode();
+
+		UE::SkeletalMeshExporterUSD::Private::SendAnalytics(
+			Object,
+			Options,
+			bAutomated,
+			ElapsedSeconds,
+			UsdUtils::GetUsdStageNumFrames( AssetStage ),
+			Extension
+		);
 	}
 
 	return true;
diff --git a/Engine/Plugins/Importers/USDImporter/Source/USDExporter/Private/StaticMeshExporterUSD.cpp b/Engine/Plugins/Importers/USDImporter/Source/USDExporter/Private/StaticMeshExporterUSD.cpp
--- a/Engine/Plugins/Importers/USDImporter/Source/USDExporter/Private/StaticMeshExporterUSD.cpp
+++ b/Engine/Plugins/Importers/USDImporter/Source/USDExporter/Private/StaticMeshExporterUSD.cpp
@@ -219,9 +219,15 @@ bool UStaticMeshExporterUsd::ExportBinary( UObject* Object, const TCHAR* Type, F
 		bool bAutomated = ExportTask ? ExportTask->bAutomated : false;
 		double ElapsedSeconds = FPlatformTime::ToSeconds64( FPlatformTime::Cycles64() - StartTime );
 		FString Extension = FPaths::GetExtension( UExporter::CurrentFilename );
-		double NumberOfFrames = UsdStage.GetEndTimeCode() - UsdStage.GetStartTimeCode();
 
-		UE::StaticMeshExporterUSD::Private::SendAnalytics( Object, Options, bAutomated, ElapsedSeconds, NumberOfFrames, Extension );
+		UE::StaticMeshExporterUSD::Private::SendAnalytics(
+			Object,
+			Options,
+			bAutomated,
+			ElapsedSeconds,
+			UsdUtils::GetUsdStageNumFrames( AssetStage ),
+			Extension
+		);
 	}
 
 	return true;
diff --git a/Engine/Plugins/Importers/USDImporter/Source/USDExporter/Private/USDConversionBlueprintContext.cpp b/Engine/Plugins/Importers/USDImporter/Source/USDExporter/Private/USDConversionBlueprintContext.cpp
--- a/Engine/Plugins/Importers/USDImporter/Source/USDExporter/Private/USDConversionBlueprintContext.cpp
+++ b/Engine/Plugins/Importers/USDImporter/Source/USDExporter/Private/USDConversionBlueprintContext.cpp
@@ -465,3 +465,17 @@ bool UUsdConversionBlueprintContext::RemoveUnrealSurfaceOutput( const FString& P
 #endif // USE_USD_SDK
 }
 
+int32 UUsdConversionBlueprintContext::GetUsdStageNumFrames()
+{
+#if USE_USD_SDK
+	if ( !Stage )
+	{
+		return 0;
+	}
+
+	return UsdUtils::GetUsdStageNumFrames( Stage );
+#else
+	return 0;
+#endif // USE_USD_SDK
+}
+
diff --git a/Engine/Plugins/Importers/USDImporter/Source/USDExporter/Public/USDConversionBlueprintContext.h b/Engine/Plugins/Importers/USDImporter/Source/USDExporter/Public/USDConversionBlueprintContext.h
--- a/Engine/Plugins/Importers/USDImporter/Source/USDExporter/Public/USDConversionBlueprintContext.h
+++ b/Engine/Plugins/Importers/USDImporter/Source/USDExporter/Public/USDConversionBlueprintContext.h
@@ -163,4 +163,8 @@ public:
 	 */
 	UFUNCTION( BlueprintCallable, Category = "Conversion utils" )
 	bool RemoveUnrealSurfaceOutput( const FString& PrimPath, const FFilePath& LayerToAuthorIn );
+
+	/** Returns how many frames of animation the stage would generate if imported */
+	UFUNCTION( BlueprintCallable, Category = "Conversion utils" )
+	int32 GetUsdStageNumFrames();
 };
diff --git a/Engine/Plugins/Importers/USDImporter/Source/USDStage/Private/USDStageActor.cpp b/Engine/Plugins/Importers/USDImporter/Source/USDStage/Private/USDStageActor.cpp
--- a/Engine/Plugins/Importers/USDImporter/Source/USDStage/Private/USDStageActor.cpp
+++ b/Engine/Plugins/Importers/USDImporter/Source/USDStage/Private/USDStageActor.cpp
@@ -1,3 +1,4 @@
+
 // Copyright Epic Games, Inc. All Rights Reserved.
 
 #include "USDStageActor.h"
@@ -1826,9 +1827,14 @@ void AUsdStageActor::LoadUsdStage()
 
 	UE_LOG( LogUsd, Log, TEXT("%s %s in [%d min %.3f s]"), TEXT("Stage loaded"), *FPaths::GetBaseFilename( RootLayer.FilePath ), ElapsedMin, ElapsedSeconds );
 
-	double NumberOfFrames = UsdStage ? (UsdStage.GetEndTimeCode() - UsdStage.GetStartTimeCode()) : 0.0;
-
-	FUsdStageActorImpl::SendAnalytics( this, ElapsedSeconds, NumberOfFrames, FPaths::GetExtension( RootLayer.FilePath ) );
+#if USE_USD_SDK
+	FUsdStageActorImpl::SendAnalytics(
+		this,
+		ElapsedSeconds,
+		UsdUtils::GetUsdStageNumFrames( UsdStage ),
+		FPaths::GetExtension( RootLayer.FilePath )
+	);
+#endif // #if USE_USD_SDK
 }
 
 void AUsdStageActor::UnloadUsdStage()
diff --git a/Engine/Plugins/Importers/USDImporter/Source/USDStageEditor/Private/SUSDLayersTreeView.cpp b/Engine/Plugins/Importers/USDImporter/Source/USDStageEditor/Private/SUSDLayersTreeView.cpp
--- a/Engine/Plugins/Importers/USDImporter/Source/USDStageEditor/Private/SUSDLayersTreeView.cpp
+++ b/Engine/Plugins/Importers/USDImporter/Source/USDStageEditor/Private/SUSDLayersTreeView.cpp
@@ -3,6 +3,7 @@
 #include "SUSDLayersTreeView.h"
 
 #include "SUSDStageEditorStyle.h"
+#include "USDClassesModule.h"
 #include "USDLayersViewModel.h"
 #include "USDLayerUtils.h"
 #include "USDMemory.h"
@@ -16,6 +17,7 @@
 #include "DesktopPlatformModule.h"
 #include "EditorStyleSet.h"
 #include "Engine/World.h"
+#include "EngineAnalytics.h"
 #include "Framework/MultiBox/MultiBoxBuilder.h"
 #include "IDesktopPlatform.h"
 #include "Modules/ModuleManager.h"
@@ -447,6 +449,9 @@ void SUsdLayersTreeView::OnExportSelectedLayers() const
 		LayersToExport.Add( SelectedLayer );
 	}
 
+	double StartTime = FPlatformTime::Cycles64();
+	FString Extension;
+
 	// Single layer -> Allow picking the target layer filename
 	if ( LayersToExport.Num() == 1 )
 	{
@@ -456,10 +461,13 @@ void SUsdLayersTreeView::OnExportSelectedLayers() const
 			return;
 		}
 
+		StartTime = FPlatformTime::Cycles64();
+		Extension = FPaths::GetExtension( UsdFilePath.GetValue() );
+
 		UE::USDLayersTreeViewImpl::Private::ExportLayerToPath( LayersToExport[ 0 ], UsdFilePath.GetValue() );
 	}
 	// Multiple layers -> Pick folder and export them with the same name
-	if ( LayersToExport.Num() > 1 )
+	else if ( LayersToExport.Num() > 1 )
 	{
 		IDesktopPlatform* DesktopPlatform = FDesktopPlatformModule::Get();
 		if ( !DesktopPlatform )
@@ -479,6 +487,9 @@ void SUsdLayersTreeView::OnExportSelectedLayers() const
 		}
 		TargetFolderPath = FPaths::ConvertRelativePathToFull( TargetFolderPath );
 
+		StartTime = FPlatformTime::Cycles64();
+		Extension = FPaths::GetExtension( LayersToExport[ 0 ].GetDisplayName() );
+
 		if ( FPaths::DirectoryExists( TargetFolderPath ) )
 		{
 			for ( const UE::FSdfLayer& LayerToExport : LayersToExport )
@@ -497,6 +508,25 @@ void SUsdLayersTreeView::OnExportSelectedLayers() const
 			}
 		}
 	}
+
+	// Send analytics
+	if ( FEngineAnalytics::IsAvailable() )
+	{
+		TArray<FAnalyticsEventAttribute> EventAttributes;
+
+		EventAttributes.Emplace( TEXT( "NumLayersExported" ), LayersToExport.Num() );
+
+		bool bAutomated = false;
+		double ElapsedSeconds = FPlatformTime::ToSeconds64( FPlatformTime::Cycles64() - StartTime );
+		IUsdClassesModule::SendAnalytics(
+			MoveTemp( EventAttributes ),
+			TEXT( "ExportSelectedLayers" ),
+			bAutomated,
+			ElapsedSeconds,
+			LayersToExport.Num() > 0 ? UsdUtils::GetSdfLayerNumFrames( LayersToExport[ 0 ] ) : 0,
+			Extension
+		);
+	}
 }
 
 bool SUsdLayersTreeView::CanAddSubLayer() const
diff --git a/Engine/Plugins/Importers/USDImporter/Source/USDStageEditor/Private/SUSDStage.cpp b/Engine/Plugins/Importers/USDImporter/Source/USDStageEditor/Private/SUSDStage.cpp
--- a/Engine/Plugins/Importers/USDImporter/Source/USDStageEditor/Private/SUSDStage.cpp
+++ b/Engine/Plugins/Importers/USDImporter/Source/USDStageEditor/Private/SUSDStage.cpp
@@ -6,6 +6,7 @@
 #include "SUSDPrimInfo.h"
 #include "SUSDStageTreeView.h"
 #include "UnrealUSDWrapper.h"
+#include "USDClassesModule.h"
 #include "USDConversionUtils.h"
 #include "USDErrorUtils.h"
 #include "USDLayerUtils.h"
@@ -30,6 +31,7 @@
 #include "EditorStyleSet.h"
 #include "Engine/Selection.h"
 #include "Engine/World.h"
+#include "EngineAnalytics.h"
 #include "Framework/MultiBox/MultiBoxBuilder.h"
 #include "LevelEditor.h"
 #include "Modules/ModuleManager.h"
@@ -1403,6 +1405,8 @@ void SUsdStage::OnExportAll()
 	}
 	TargetFolderPath = FPaths::ConvertRelativePathToFull( TargetFolderPath );
 
+	double StartTime = FPlatformTime::Cycles64();
+
 	// Manually load any layer that is referenced by any of the stage's layers, but not currently loaded.
 	// This can happen if e.g. an unselected variant appends a reference or has a payload.
 	// We will need to actually load these as opposed to just copy-paste the files as we may need to update references/paths
@@ -1493,6 +1497,26 @@ void SUsdStage::OnExportAll()
 			OutputLayer.Save( bForce );
 		}
 	}
+
+	// Send analytics
+	if ( FEngineAnalytics::IsAvailable() )
+	{
+		TArray<FAnalyticsEventAttribute> EventAttributes;
+
+		EventAttributes.Emplace( TEXT( "NumLayersExported" ), LayerStack.Num() );
+
+		bool bAutomated = false;
+		double ElapsedSeconds = FPlatformTime::ToSeconds64( FPlatformTime::Cycles64() - StartTime );
+		FString Extension = FPaths::GetExtension( UsdStage.GetRootLayer().GetDisplayName() );
+		IUsdClassesModule::SendAnalytics(
+			MoveTemp( EventAttributes ),
+			TEXT( "ExportStageLayers" ),
+			bAutomated,
+			ElapsedSeconds,
+			UsdUtils::GetUsdStageNumFrames( UsdStage ),
+			Extension
+		);
+	}
 }
 
 void SUsdStage::OnExportFlattened()
@@ -1515,7 +1539,25 @@ void SUsdStage::OnExportFlattened()
 		return;
 	}
 
+	double StartTime = FPlatformTime::Cycles64();
+
 	UsdStage.Export( *( UsdFilePath.GetValue() ) );
+
+	// Send analytics
+	if ( FEngineAnalytics::IsAvailable() )
+	{
+		bool bAutomated = false;
+		double ElapsedSeconds = FPlatformTime::ToSeconds64( FPlatformTime::Cycles64() - StartTime );
+		FString Extension = FPaths::GetExtension( UsdFilePath.GetValue() );
+		IUsdClassesModule::SendAnalytics(
+			{},
+			TEXT( "ExportStageFlattened" ),
+			bAutomated,
+			ElapsedSeconds,
+			UsdUtils::GetUsdStageNumFrames( UsdStage ),
+			Extension
+		);
+	}
 }
 
 void SUsdStage::OnReloadStage()
diff --git a/Engine/Plugins/Importers/USDImporter/Source/USDStageEditor/USDStageEditor.Build.cs b/Engine/Plugins/Importers/USDImporter/Source/USDStageEditor/USDStageEditor.Build.cs
--- a/Engine/Plugins/Importers/USDImporter/Source/USDStageEditor/USDStageEditor.Build.cs
+++ b/Engine/Plugins/Importers/USDImporter/Source/USDStageEditor/USDStageEditor.Build.cs
@@ -29,6 +29,7 @@ namespace UnrealBuildTool.Rules
 					"SlateCore",
 					"UnrealEd",
 					"UnrealUSDWrapper",
+					"USDClasses",
 					"USDSchemas",
 					"USDStage",
 					"USDStageEditorViewModels",
diff --git a/Engine/Plugins/Importers/USDImporter/Source/USDStageImporter/Private/USDStageImporter.cpp b/Engine/Plugins/Importers/USDImporter/Source/USDStageImporter/Private/USDStageImporter.cpp
--- a/Engine/Plugins/Importers/USDImporter/Source/USDStageImporter/Private/USDStageImporter.cpp
+++ b/Engine/Plugins/Importers/USDImporter/Source/USDStageImporter/Private/USDStageImporter.cpp
@@ -1344,12 +1344,8 @@ namespace UsdStageImporterImpl
 			EventAttributes.Emplace( TEXT( "NumGeometryCaches" ), LexToString( NumGeometryCaches ) );
 
 			FString RootLayerIdentifier = ImportContext.FilePath;
-			double NumberOfFrames = 0;
-
 			if ( ImportContext.Stage )
 			{
-				NumberOfFrames = ImportContext.Stage.GetEndTimeCode() - ImportContext.Stage.GetStartTimeCode();
-
 				if ( RootLayerIdentifier.IsEmpty() )
 				{
 					RootLayerIdentifier = ImportContext.Stage.GetRootLayer().GetIdentifier();
@@ -1361,7 +1357,7 @@ namespace UsdStageImporterImpl
 				EventName,
 				ImportContext.bIsAutomated,
 				ElapsedSeconds,
-				NumberOfFrames,
+				UsdUtils::GetUsdStageNumFrames( ImportContext.Stage ),
 				FPaths::GetExtension( RootLayerIdentifier )
 			);
 		}
diff --git a/Engine/Plugins/Importers/USDImporter/Source/USDUtilities/Private/USDConversionUtils.cpp b/Engine/Plugins/Importers/USDImporter/Source/USDUtilities/Private/USDConversionUtils.cpp
--- a/Engine/Plugins/Importers/USDImporter/Source/USDUtilities/Private/USDConversionUtils.cpp
+++ b/Engine/Plugins/Importers/USDImporter/Source/USDUtilities/Private/USDConversionUtils.cpp
@@ -275,6 +275,12 @@ void UsdUtils::SetUsdStageMetersPerUnit( const pxr::UsdStageRefPtr& Stage, float
 	pxr::UsdGeomSetStageMetersPerUnit( Stage, MetersPerUnit );
 }
 
+int32 UsdUtils::GetUsdStageNumFrames( const pxr::UsdStageRefPtr& Stage )
+{
+	// USD time code range is inclusive on both ends
+	return FMath::Abs( FMath::CeilToInt32( Stage->GetEndTimeCode() ) - FMath::FloorToInt32( Stage->GetStartTimeCode() ) + 1 );
+}
+
 bool UsdUtils::HasCompositionArcs( const pxr::UsdPrim& Prim )
 {
 	if ( !Prim || !Prim.IsActive() )
diff --git a/Engine/Plugins/Importers/USDImporter/Source/USDUtilities/Private/USDLayerUtils.cpp b/Engine/Plugins/Importers/USDImporter/Source/USDUtilities/Private/USDLayerUtils.cpp
--- a/Engine/Plugins/Importers/USDImporter/Source/USDUtilities/Private/USDLayerUtils.cpp
+++ b/Engine/Plugins/Importers/USDImporter/Source/USDUtilities/Private/USDLayerUtils.cpp
@@ -736,6 +736,17 @@ void UsdUtils::ConvertAssetRelativePathsToAbsolute( UE::FSdfLayer& LayerToConver
 	);
 }
 
+int32 UsdUtils::GetSdfLayerNumFrames( const pxr::SdfLayerRefPtr& Layer )
+{
+	if ( !Layer )
+	{
+		return 0;
+	}
+
+	// USD time code range is inclusive on both ends
+	return FMath::Abs( FMath::CeilToInt32( Layer->GetEndTimeCode() ) - FMath::FloorToInt32( Layer->GetStartTimeCode() ) + 1 );
+}
+
 #undef LOCTEXT_NAMESPACE
 
 #endif // #if USE_USD_SDK
diff --git a/Engine/Plugins/Importers/USDImporter/Source/USDUtilities/Public/USDConversionUtils.h b/Engine/Plugins/Importers/USDImporter/Source/USDUtilities/Public/USDConversionUtils.h
--- a/Engine/Plugins/Importers/USDImporter/Source/USDUtilities/Public/USDConversionUtils.h
+++ b/Engine/Plugins/Importers/USDImporter/Source/USDUtilities/Public/USDConversionUtils.h
@@ -98,6 +98,8 @@ namespace UsdUtils
 	USDUTILITIES_API float GetUsdStageMetersPerUnit( const pxr::UsdStageRefPtr& Stage );
 	USDUTILITIES_API void SetUsdStageMetersPerUnit( const pxr::UsdStageRefPtr& Stage, float MetersPerUnit );
 
+	USDUTILITIES_API int32 GetUsdStageNumFrames( const pxr::UsdStageRefPtr& Stage );
+
 	USDUTILITIES_API bool HasCompositionArcs( const pxr::UsdPrim& Prim );
 
 	USDUTILITIES_API UClass* GetActorTypeForPrim( const pxr::UsdPrim& Prim );
diff --git a/Engine/Plugins/Importers/USDImporter/Source/USDUtilities/Public/USDLayerUtils.h b/Engine/Plugins/Importers/USDImporter/Source/USDUtilities/Public/USDLayerUtils.h
--- a/Engine/Plugins/Importers/USDImporter/Source/USDUtilities/Public/USDLayerUtils.h
+++ b/Engine/Plugins/Importers/USDImporter/Source/USDUtilities/Public/USDLayerUtils.h
@@ -115,6 +115,8 @@ namespace UsdUtils
 	 * updates them to be absolute with respect to AnchorLayer. LayerToConvert and AnchorLayer can be the same layer, but they don't have to
 	 */
 	USDUTILITIES_API void ConvertAssetRelativePathsToAbsolute( UE::FSdfLayer& LayerToConvert, const UE::FSdfLayer& AnchorLayer );
+
+	USDUTILITIES_API int32 GetSdfLayerNumFrames( const pxr::SdfLayerRefPtr& Layer );
 }
 
 #endif // #if USE_USD_SDK
