USD: Fix build errors when building monolithic editor targets without the ansi allocator.

diff --git a/Engine/Plugins/Importers/USDImporter/Source/USDStageEditor/Private/Widgets/SUSDIntegrationsPanel.cpp b/Engine/Plugins/Importers/USDImporter/Source/USDStageEditor/Private/Widgets/SUSDIntegrationsPanel.cpp
--- a/Engine/Plugins/Importers/USDImporter/Source/USDStageEditor/Private/Widgets/SUSDIntegrationsPanel.cpp
+++ b/Engine/Plugins/Importers/USDImporter/Source/USDStageEditor/Private/Widgets/SUSDIntegrationsPanel.cpp
@@ -2,6 +2,8 @@
 
 #include "SUSDIntegrationsPanel.h"
 
+#if USE_USD_SDK
+
 #include "SUSDStageEditorStyle.h"
 #include "USDIntegrationUtils.h"
 #include "USDTypesConversion.h"
@@ -316,4 +318,5 @@ TSharedRef< ITableRow > SUsdIntegrationsPanel::OnGenerateRow( TSharedPtr<UE::FUs
 	return SNew( SUsdIntegrationsPanelRow, InAttr, OwnerTable );
 }
 
-#undef LOCTEXT_NAMESPACE
\ No newline at end of file
+#undef LOCTEXT_NAMESPACE
+#endif // USE_USD_SDK
\ No newline at end of file
diff --git a/Engine/Plugins/Importers/USDImporter/Source/USDStageEditor/Private/Widgets/SUSDIntegrationsPanel.h b/Engine/Plugins/Importers/USDImporter/Source/USDStageEditor/Private/Widgets/SUSDIntegrationsPanel.h
--- a/Engine/Plugins/Importers/USDImporter/Source/USDStageEditor/Private/Widgets/SUSDIntegrationsPanel.h
+++ b/Engine/Plugins/Importers/USDImporter/Source/USDStageEditor/Private/Widgets/SUSDIntegrationsPanel.h
@@ -28,6 +28,8 @@ public:
 	UObject* AnimBPProperty;
 };
 
+#if USE_USD_SDK
+
 class SUsdIntegrationsPanelRow : public SMultiColumnTableRow< TSharedPtr<UE::FUsdAttribute> >
 {
 public:
@@ -65,3 +67,5 @@ private:
 	UE::FUsdStageWeak UsdStage;
 	FString PrimPath;
 };
+
+#endif // USE_USD_SDK
\ No newline at end of file
diff --git a/Engine/Plugins/Importers/USDImporter/Source/USDStageEditorViewModels/Private/USDPrimViewModel.cpp b/Engine/Plugins/Importers/USDImporter/Source/USDStageEditorViewModels/Private/USDPrimViewModel.cpp
--- a/Engine/Plugins/Importers/USDImporter/Source/USDStageEditorViewModels/Private/USDPrimViewModel.cpp
+++ b/Engine/Plugins/Importers/USDImporter/Source/USDStageEditorViewModels/Private/USDPrimViewModel.cpp
@@ -308,6 +308,8 @@ bool FUsdPrimViewModel::CanSetUpLiveLink() const
 	}
 
 	return !Prim.HasAPI( Schema ) && Prim.CanApplyAPI( Schema );
+#else
+	return false;
 #endif // #if USE_USD_SDK
 }
 
@@ -341,6 +343,8 @@ bool FUsdPrimViewModel::CanRemoveLiveLink() const
 	}
 
 	return Prim.HasAPI( Schema );
+#else
+	return false;
 #endif // #if USE_USD_SDK
 }
 
diff --git a/Engine/Plugins/Importers/USDImporter/Source/USDStageImporter/Private/USDStageImporter.cpp b/Engine/Plugins/Importers/USDImporter/Source/USDStageImporter/Private/USDStageImporter.cpp
--- a/Engine/Plugins/Importers/USDImporter/Source/USDStageImporter/Private/USDStageImporter.cpp
+++ b/Engine/Plugins/Importers/USDImporter/Source/USDStageImporter/Private/USDStageImporter.cpp
@@ -1258,6 +1258,7 @@ namespace UsdStageImporterImpl
 
 	void SendAnalytics( FUsdStageImportContext& ImportContext, UObject* Asset, const FString& Operation, const TSet<UObject*>& ImportedAssets, double ElapsedSeconds )
 	{
+#if USE_USD_SDK
 		if ( FEngineAnalytics::IsAvailable() )
 		{
 			TArray<FAnalyticsEventAttribute> EventAttributes;
@@ -1361,6 +1362,7 @@ namespace UsdStageImporterImpl
 				FPaths::GetExtension( RootLayerIdentifier )
 			);
 		}
+#endif // USE_USD_SDK
 	}
 
 	// Removes from AssetsToImport assets that are unwanted according to our import options, and adds entries to
