USD: Fix broken build on non-editor targets.

diff --git a/Engine/Plugins/Importers/USDImporter/Source/USDUtilities/Private/USDIntegrationUtils.cpp b/Engine/Plugins/Importers/USDImporter/Source/USDUtilities/Private/USDIntegrationUtils.cpp
--- a/Engine/Plugins/Importers/USDImporter/Source/USDUtilities/Private/USDIntegrationUtils.cpp
+++ b/Engine/Plugins/Importers/USDImporter/Source/USDUtilities/Private/USDIntegrationUtils.cpp
@@ -74,8 +74,6 @@ void UnrealToUsd::ConvertLiveLinkProperties( const UActorComponent* InComponent,
 			Attr.Set( InController->bEvaluateLiveLink, pxr::UsdTimeCode::Default() );
 		}
 	}
-#else
-	return false;
 #endif // WITH_EDITOR
 }
 
