USD: Prevented empty USD change notices from being added to the list of info/resync changes.

diff --git a/Engine/Plugins/Importers/USDImporter/Source/UnrealUSDWrapper/Private/USDListener.cpp b/Engine/Plugins/Importers/USDImporter/Source/UnrealUSDWrapper/Private/USDListener.cpp
--- a/Engine/Plugins/Importers/USDImporter/Source/UnrealUSDWrapper/Private/USDListener.cpp
+++ b/Engine/Plugins/Importers/USDImporter/Source/UnrealUSDWrapper/Private/USDListener.cpp
@@ -37,6 +37,12 @@ namespace UsdToUnreal
 
 		for ( UsdNotice::ObjectsChanged::PathRange::const_iterator It = PathRange.begin(); It != PathRange.end(); ++It )
 		{
+			const std::vector<const SdfChangeList::Entry*>& Changes = It.base()->second;
+			if ( Changes.empty() )
+			{
+				continue;
+			}
+
 			const FString PrimPath = ANSI_TO_TCHAR( It->GetAbsoluteRootOrPrimPath().GetAsString().c_str() );
 
 			// Something like "/Root/Prim.some_field", or "/"
@@ -44,7 +50,6 @@ namespace UsdToUnreal
 
 			TArray<UsdUtils::FObjectChangeNotice>& ConvertedChanges = OutChanges.FindOrAdd( PrimPath );
 
-			const std::vector<const SdfChangeList::Entry*>& Changes = It.base()->second;
 			for ( const SdfChangeList::Entry* Entry : Changes )
 			{
 				UsdUtils::FObjectChangeNotice& ConvertedEntry = ConvertedChanges.Emplace_GetRef();
