USD: Add analytics to material and LevelSequence export;

Add missing data to analytics of other workflows;
Make some analytics event values consistent between Python and C++ workflows;
Don't submit unused options/data to analytics events.

diff --git a/Engine/Plugins/Importers/USDImporter/Source/USDExporter/Private/AnimSequenceExporterUSD.cpp b/Engine/Plugins/Importers/USDImporter/Source/USDExporter/Private/AnimSequenceExporterUSD.cpp
--- a/Engine/Plugins/Importers/USDImporter/Source/USDExporter/Private/AnimSequenceExporterUSD.cpp
+++ b/Engine/Plugins/Importers/USDImporter/Source/USDExporter/Private/AnimSequenceExporterUSD.cpp
@@ -43,11 +43,37 @@ namespace UE
 
 					if ( Options )
 					{
+						EventAttributes.Emplace( TEXT( "MetersPerUnit" ), LexToString( Options->StageOptions.MetersPerUnit ) );
+						EventAttributes.Emplace( TEXT( "UpAxis" ), Options->StageOptions.UpAxis == EUsdUpAxis::YAxis ? TEXT( "Y" ) : TEXT( "Z" ) );
+
 						EventAttributes.Emplace( TEXT( "ExportPreviewMesh" ), LexToString( Options->bExportPreviewMesh ) );
 						if ( Options->bExportPreviewMesh )
 						{
 							EventAttributes.Emplace( TEXT( "UsePayload" ), LexToString( Options->PreviewMeshOptions.bUsePayload ) );
-							EventAttributes.Emplace( TEXT( "PayloadFormat" ), Options->PreviewMeshOptions.PayloadFormat );
+							if ( Options->PreviewMeshOptions.bUsePayload )
+							{
+								EventAttributes.Emplace( TEXT( "PayloadFormat" ), Options->PreviewMeshOptions.PayloadFormat );
+							}
+							EventAttributes.Emplace( TEXT( "BakeMaterials" ), Options->PreviewMeshOptions.bBakeMaterials );
+							if ( Options->PreviewMeshOptions.bBakeMaterials )
+							{
+								FString BakedPropertiesString;
+								{
+									const UEnum* PropertyEnum = StaticEnum<EMaterialProperty>();
+									for ( const FPropertyEntry& PropertyEntry : Options->PreviewMeshOptions.MaterialBakingOptions.Properties )
+									{
+										FString PropertyString = PropertyEnum->GetNameByValue( PropertyEntry.Property ).ToString();
+										PropertyString.RemoveFromStart( TEXT( "MP_" ) );
+										BakedPropertiesString += PropertyString + TEXT( ", " );
+									}
+
+									BakedPropertiesString.RemoveFromEnd( TEXT( ", " ) );
+								}
+
+								EventAttributes.Emplace( TEXT( "RemoveUnrealMaterials" ), Options->PreviewMeshOptions.bRemoveUnrealMaterials );
+								EventAttributes.Emplace( TEXT( "BakedProperties" ), BakedPropertiesString );
+								EventAttributes.Emplace( TEXT( "DefaultTextureSize" ), Options->PreviewMeshOptions.MaterialBakingOptions.DefaultTextureSize.ToString() );
+							}
 							EventAttributes.Emplace( TEXT( "LowestMeshLOD" ), LexToString( Options->PreviewMeshOptions.LowestMeshLOD ) );
 							EventAttributes.Emplace( TEXT( "HighestMeshLOD" ), LexToString( Options->PreviewMeshOptions.HighestMeshLOD ) );
 						}
diff --git a/Engine/Plugins/Importers/USDImporter/Source/USDExporter/Private/LevelSequenceExporterUSD.cpp b/Engine/Plugins/Importers/USDImporter/Source/USDExporter/Private/LevelSequenceExporterUSD.cpp
--- a/Engine/Plugins/Importers/USDImporter/Source/USDExporter/Private/LevelSequenceExporterUSD.cpp
+++ b/Engine/Plugins/Importers/USDImporter/Source/USDExporter/Private/LevelSequenceExporterUSD.cpp
@@ -4,6 +4,7 @@
 
 #include "LevelSequenceExporterUSDOptions.h"
 #include "UnrealUSDWrapper.h"
+#include "USDClassesModule.h"
 #include "USDConversionUtils.h"
 #include "USDLayerUtils.h"
 #include "USDLog.h"
@@ -20,6 +21,7 @@
 #include "Compilation/MovieSceneCompiledDataManager.h"
 #include "Components/SkeletalMeshComponent.h"
 #include "Editor.h"
+#include "EngineAnalytics.h"
 #include "Evaluation/MovieSceneSequenceHierarchy.h"
 #include "ISequencer.h"
 #include "ISequencerModule.h"
@@ -46,6 +48,107 @@ namespace UE
 	{
 		namespace Private
 		{
+			void SendAnalytics(
+				const ULevelSequence* LevelSequence,
+				const ULevelSequenceExporterUsdOptions* Options,
+				const TArray<UE::FUsdStage>& ExportedStages,
+				bool bAutomated,
+				double ElapsedSeconds,
+				const FString& Extension
+			)
+			{
+				if ( !LevelSequence || ExportedStages.Num() == 0 )
+				{
+					return;
+				}
+
+				if ( FEngineAnalytics::IsAvailable() )
+				{
+					FString ClassName = LevelSequence->GetClass()->GetName();
+
+					TArray<FAnalyticsEventAttribute> EventAttributes;
+
+					EventAttributes.Emplace( TEXT( "AssetType" ), ClassName );
+
+					if ( Options )
+					{
+						// Stage options
+						EventAttributes.Emplace( TEXT( "MetersPerUnit" ), LexToString( Options->StageOptions.MetersPerUnit ) );
+						EventAttributes.Emplace( TEXT( "UpAxis" ), Options->StageOptions.UpAxis == EUsdUpAxis::YAxis ? TEXT( "Y" ) : TEXT( "Z" ) );
+
+						// LevelSequence options
+						EventAttributes.Emplace( TEXT( "TimeCodesPerSecond" ), LexToString( Options->TimeCodesPerSecond ) );
+						EventAttributes.Emplace( TEXT( "OverrideExportRange" ), Options->bOverrideExportRange );
+						if( Options->bOverrideExportRange )
+						{
+							EventAttributes.Emplace( TEXT( "StartFrame" ), LexToString( Options->StartFrame ) );
+							EventAttributes.Emplace( TEXT( "EndFrame" ), LexToString( Options->EndFrame ) );
+						}
+						EventAttributes.Emplace( TEXT( "ExportSubsequencesAsLayers" ), Options->bExportSubsequencesAsLayers );
+						EventAttributes.Emplace( TEXT( "ExportLevel" ), Options->bExportLevel );
+						if( Options->bExportLevel )
+						{
+							EventAttributes.Emplace( TEXT( "UseExportedLevelAsSublayer" ), Options->bUseExportedLevelAsSublayer );
+						}
+						EventAttributes.Emplace( TEXT( "NumExportedLevelSequenceLayers" ), ExportedStages.Num() );
+
+						if( Options->bExportLevel )
+						{
+							// Level options
+							EventAttributes.Emplace( TEXT( "SelectionOnly" ), Options->LevelExportOptions.bSelectionOnly );
+							EventAttributes.Emplace( TEXT( "ExportActorFolders" ), Options->LevelExportOptions.bExportActorFolders );
+							EventAttributes.Emplace( TEXT( "IgnoreSequencerAnimations" ), Options->LevelExportOptions.bIgnoreSequencerAnimations );
+							EventAttributes.Emplace( TEXT( "ExportFoliageOnActorsLayer" ), Options->LevelExportOptions.bExportFoliageOnActorsLayer );
+							EventAttributes.Emplace( TEXT( "LowestLandscapeLOD" ), LexToString( Options->LevelExportOptions.LowestLandscapeLOD ) );
+							EventAttributes.Emplace( TEXT( "HighestLandscapeLOD" ), LexToString( Options->LevelExportOptions.HighestLandscapeLOD ) );
+							EventAttributes.Emplace( TEXT( "LandscapeBakeResolution" ), Options->LevelExportOptions.LandscapeBakeResolution.ToString() );
+							EventAttributes.Emplace( TEXT( "ExportSublayers" ), LexToString( Options->LevelExportOptions.bExportSublayers ) );
+							EventAttributes.Emplace( TEXT( "NumLevelsToIgnore" ), LexToString( Options->LevelExportOptions.LevelsToIgnore.Num() ) );
+						}
+
+						// Asset options
+						EventAttributes.Emplace( TEXT( "UsePayload" ), LexToString( Options->LevelExportOptions.AssetOptions.bUsePayload ) );
+						if( Options->LevelExportOptions.AssetOptions.bUsePayload )
+						{
+							EventAttributes.Emplace( TEXT( "PayloadFormat" ), Options->LevelExportOptions.AssetOptions.PayloadFormat );
+						}
+						EventAttributes.Emplace( TEXT( "BakeMaterials" ), Options->LevelExportOptions.AssetOptions.bBakeMaterials );
+						EventAttributes.Emplace( TEXT( "LowestMeshLOD" ), LexToString( Options->LevelExportOptions.AssetOptions.LowestMeshLOD ) );
+						EventAttributes.Emplace( TEXT( "HighestMeshLOD" ), LexToString( Options->LevelExportOptions.AssetOptions.HighestMeshLOD ) );
+
+						// Material baking options
+						if( Options->LevelExportOptions.AssetOptions.bBakeMaterials )
+						{
+							FString BakedPropertiesString;
+							{
+								const UEnum* PropertyEnum = StaticEnum<EMaterialProperty>();
+								for ( const FPropertyEntry& PropertyEntry : Options->LevelExportOptions.AssetOptions.MaterialBakingOptions.Properties )
+								{
+									FString PropertyString = PropertyEnum->GetNameByValue( PropertyEntry.Property ).ToString();
+									PropertyString.RemoveFromStart( TEXT( "MP_" ) );
+									BakedPropertiesString += PropertyString + TEXT( ", " );
+								}
+
+								BakedPropertiesString.RemoveFromEnd( TEXT( ", " ) );
+							}
+
+							EventAttributes.Emplace( TEXT( "RemoveUnrealMaterials" ), Options->LevelExportOptions.AssetOptions.bRemoveUnrealMaterials );
+							EventAttributes.Emplace( TEXT( "BakedProperties" ), BakedPropertiesString );
+							EventAttributes.Emplace( TEXT( "DefaultTextureSize" ), Options->LevelExportOptions.AssetOptions.MaterialBakingOptions.DefaultTextureSize.ToString() );
+						}
+					}
+
+					IUsdClassesModule::SendAnalytics(
+						MoveTemp( EventAttributes ),
+						FString::Printf( TEXT( "Export.%s" ), *ClassName ),
+						bAutomated,
+						ElapsedSeconds,
+						( Options->EndFrame - Options->StartFrame ),
+						Extension
+					);
+				}
+			}
+
 #if USE_USD_SDK
 			// Custom spawn register so that when DestroySpawnedObject is called while bDestroyingJustHides is true we
 			// actually just hide the objects, so that we can keep a live reference to components within the bakers.
@@ -819,7 +922,8 @@ namespace UE
 			void ExportMovieSceneSequence(
 				FLevelSequenceExportContext& Context,
 				UMovieSceneSequence& MovieSceneSequence,
-				const FString& FilePath
+				const FString& FilePath,
+				TArray<UE::FUsdStage>& InOutExportedStages
 			)
 			{
 				if ( FilePath.IsEmpty() || Context.ExportedMovieScenes.Contains( &MovieSceneSequence ) )
@@ -919,7 +1023,7 @@ namespace UE
 
 								FString SubSequencePath = FPaths::Combine( Directory, FString::Printf( TEXT( "%s.%s" ), *SubSequence->GetName(), *Extension ) );
 
-								ExportMovieSceneSequence( Context, *SubSequence, SubSequencePath );
+								ExportMovieSceneSequence( Context, *SubSequence, SubSequencePath, InOutExportedStages );
 
 								// For now we don't want to actually add the subsequence layers as sublayers since each exported level sequence
 								// contains the full baked result anyway, but this is how we'd do it:
@@ -953,6 +1057,8 @@ namespace UE
 				Context.UsedFilePaths.Add( UniqueFilePath );
 
 				UsdStage.GetRootLayer().Save();
+
+				InOutExportedStages.Add( UsdStage );
 			}
 #endif // USE_USD_SDK
 		}
@@ -1023,6 +1129,8 @@ bool ULevelSequenceExporterUsd::ExportBinary( UObject* Object, const TCHAR* Type
 		}
 	}
 
+	double StartTime = FPlatformTime::Cycles64();
+
 	TSharedPtr<LevelSequenceExporterImpl::FLevelSequenceHidingSpawnRegister> SpawnRegister = MakeShared<LevelSequenceExporterImpl::FLevelSequenceHidingSpawnRegister>();
 	if ( !SpawnRegister.IsValid() )
 	{
@@ -1150,7 +1258,8 @@ bool ULevelSequenceExporterUsd::ExportBinary( UObject* Object, const TCHAR* Type
 		}
 	}
 
-	LevelSequenceExporterImpl::ExportMovieSceneSequence( Context, *LevelSequence, TargetFileName );
+	TArray<UE::FUsdStage> ExportedStages;
+	LevelSequenceExporterImpl::ExportMovieSceneSequence( Context, *LevelSequence, TargetFileName, ExportedStages );
 
 	// Set this back to Stopped or else it will keep the editor viewport controls permanently hidden
 	TempSequencer->SetPlaybackStatus( EMovieScenePlayerStatus::Stopped );
@@ -1160,6 +1269,15 @@ bool ULevelSequenceExporterUsd::ExportBinary( UObject* Object, const TCHAR* Type
 		AssetEditorSubsystem->OpenEditorForAssets( AssetsToReopenEditorsFor );
 	}
 
+	// Analytics
+	{
+		bool bAutomated = ExportTask ? ExportTask->bAutomated : false;
+		double ElapsedSeconds = FPlatformTime::ToSeconds64( FPlatformTime::Cycles64() - StartTime );
+		FString Extension = FPaths::GetExtension( UExporter::CurrentFilename );
+
+		UE::LevelSequenceExporterUSD::Private::SendAnalytics( LevelSequence, Options, ExportedStages, bAutomated, ElapsedSeconds, Extension );
+	}
+
 	return true;
 #else
 	return false;
diff --git a/Engine/Plugins/Importers/USDImporter/Source/USDExporter/Private/MaterialExporterUSD.cpp b/Engine/Plugins/Importers/USDImporter/Source/USDExporter/Private/MaterialExporterUSD.cpp
--- a/Engine/Plugins/Importers/USDImporter/Source/USDExporter/Private/MaterialExporterUSD.cpp
+++ b/Engine/Plugins/Importers/USDImporter/Source/USDExporter/Private/MaterialExporterUSD.cpp
@@ -4,6 +4,7 @@
 
 #include "MaterialExporterUSDOptions.h"
 #include "UnrealUSDWrapper.h"
+#include "USDClassesModule.h"
 #include "USDConversionUtils.h"
 #include "USDGeomMeshConversion.h"
 #include "USDMemory.h"
@@ -16,13 +17,53 @@
 #include "UsdWrappers/UsdStage.h"
 
 #include "AssetExportTask.h"
+#include "EngineAnalytics.h"
 #include "HAL/FileManager.h"
 #include "IMaterialBakingModule.h"
 #include "MaterialOptions.h"
 #include "Materials/MaterialInterface.h"
+#include "Misc/Paths.h"
 #include "Modules/ModuleManager.h"
 #include "UObject/UObjectGlobals.h"
-#include "Misc/Paths.h"
+
+namespace UE
+{
+	namespace MaterialExporterUSD
+	{
+		namespace Private
+		{
+			void SendAnalytics( const UMaterialInterface& Material, const FUsdMaterialBakingOptions& Options, bool bAutomated, double ElapsedSeconds, double NumberOfFrames, const FString& Extension )
+			{
+				if ( FEngineAnalytics::IsAvailable() )
+				{
+					FString ClassName = Material.GetClass()->GetName();
+
+					TArray<FAnalyticsEventAttribute> EventAttributes;
+
+					EventAttributes.Emplace( TEXT( "AssetType" ), ClassName );
+
+					FString BakedPropertiesString;
+					{
+						const UEnum* PropertyEnum = StaticEnum<EMaterialProperty>();
+						for ( const FPropertyEntry& PropertyEntry : Options.Properties )
+						{
+							FString PropertyString = PropertyEnum->GetNameByValue( PropertyEntry.Property ).ToString();
+							PropertyString.RemoveFromStart( TEXT( "MP_" ) );
+							BakedPropertiesString += PropertyString + TEXT( ", " );
+						}
+
+						BakedPropertiesString.RemoveFromEnd( TEXT( ", " ) );
+					}
+
+					EventAttributes.Emplace( TEXT( "BakedProperties" ), BakedPropertiesString );
+					EventAttributes.Emplace( TEXT( "DefaultTextureSize" ), Options.DefaultTextureSize.ToString() );
+
+					IUsdClassesModule::SendAnalytics( MoveTemp( EventAttributes ), FString::Printf( TEXT( "Export.%s" ), *ClassName ), bAutomated, ElapsedSeconds, NumberOfFrames, Extension );
+				}
+			}
+		}
+	}
+}
 
 UMaterialExporterUsd::UMaterialExporterUsd()
 {
@@ -82,6 +123,8 @@ bool UMaterialExporterUsd::ExportBinary( UObject* Object, const TCHAR* Type, FAr
 bool UMaterialExporterUsd::ExportMaterial( const UMaterialInterface& Material, const FUsdMaterialBakingOptions& Options, const FFilePath& FilePath )
 {
 #if USE_USD_SDK
+	double StartTime = FPlatformTime::Cycles64();
+
 	UE::FUsdStage UsdStage = UnrealUSDWrapper::NewStage( *FilePath.FilePath );
 	if ( !UsdStage )
 	{
@@ -110,6 +153,16 @@ bool UMaterialExporterUsd::ExportMaterial( const UMaterialInterface& Material, c
 
 	UsdStage.GetRootLayer().Save();
 
+	// Analytics
+	{
+		bool bAutomated = false;  // TODO: Can't properly fetch this yet as it would change the function signature
+		double ElapsedSeconds = FPlatformTime::ToSeconds64( FPlatformTime::Cycles64() - StartTime );
+		FString Extension = FPaths::GetExtension( FilePath.FilePath );
+		double NumberOfFrames = 1 + UsdStage.GetEndTimeCode() - UsdStage.GetStartTimeCode();
+
+		UE::MaterialExporterUSD::Private::SendAnalytics( Material, Options, bAutomated, ElapsedSeconds, NumberOfFrames, Extension );
+	}
+
 	return true;
 #else
 	return false;
diff --git a/Engine/Plugins/Importers/USDImporter/Source/USDExporter/Private/SkeletalMeshExporterUSD.cpp b/Engine/Plugins/Importers/USDImporter/Source/USDExporter/Private/SkeletalMeshExporterUSD.cpp
--- a/Engine/Plugins/Importers/USDImporter/Source/USDExporter/Private/SkeletalMeshExporterUSD.cpp
+++ b/Engine/Plugins/Importers/USDImporter/Source/USDExporter/Private/SkeletalMeshExporterUSD.cpp
@@ -43,8 +43,33 @@ namespace UE
 
 					if ( Options )
 					{
+						EventAttributes.Emplace( TEXT( "MetersPerUnit" ), LexToString( Options->StageOptions.MetersPerUnit ) );
+						EventAttributes.Emplace( TEXT( "UpAxis" ), Options->StageOptions.UpAxis == EUsdUpAxis::YAxis ? TEXT( "Y" ) : TEXT( "Z" ) );
 						EventAttributes.Emplace( TEXT( "UsePayload" ), LexToString( Options->MeshAssetOptions.bUsePayload ) );
-						EventAttributes.Emplace( TEXT( "PayloadFormat" ), Options->MeshAssetOptions.PayloadFormat );
+						if ( Options->MeshAssetOptions.bUsePayload )
+						{
+							EventAttributes.Emplace( TEXT( "PayloadFormat" ), Options->MeshAssetOptions.PayloadFormat );
+						}
+						EventAttributes.Emplace( TEXT( "BakeMaterials" ), Options->MeshAssetOptions.bBakeMaterials );
+						if ( Options->MeshAssetOptions.bBakeMaterials )
+						{
+							FString BakedPropertiesString;
+							{
+								const UEnum* PropertyEnum = StaticEnum<EMaterialProperty>();
+								for ( const FPropertyEntry& PropertyEntry : Options->MeshAssetOptions.MaterialBakingOptions.Properties )
+								{
+									FString PropertyString = PropertyEnum->GetNameByValue( PropertyEntry.Property ).ToString();
+									PropertyString.RemoveFromStart( TEXT( "MP_" ) );
+									BakedPropertiesString += PropertyString + TEXT( ", " );
+								}
+
+								BakedPropertiesString.RemoveFromEnd( TEXT( ", " ) );
+							}
+
+							EventAttributes.Emplace( TEXT( "RemoveUnrealMaterials" ), Options->MeshAssetOptions.bRemoveUnrealMaterials );
+							EventAttributes.Emplace( TEXT( "BakedProperties" ), BakedPropertiesString );
+							EventAttributes.Emplace( TEXT( "DefaultTextureSize" ), Options->MeshAssetOptions.MaterialBakingOptions.DefaultTextureSize.ToString() );
+						}
 						EventAttributes.Emplace( TEXT( "LowestMeshLOD" ), LexToString( Options->MeshAssetOptions.LowestMeshLOD ) );
 						EventAttributes.Emplace( TEXT( "HighestMeshLOD" ), LexToString( Options->MeshAssetOptions.HighestMeshLOD ) );
 					}
diff --git a/Engine/Plugins/Importers/USDImporter/Source/USDExporter/Private/StaticMeshExporterUSD.cpp b/Engine/Plugins/Importers/USDImporter/Source/USDExporter/Private/StaticMeshExporterUSD.cpp
--- a/Engine/Plugins/Importers/USDImporter/Source/USDExporter/Private/StaticMeshExporterUSD.cpp
+++ b/Engine/Plugins/Importers/USDImporter/Source/USDExporter/Private/StaticMeshExporterUSD.cpp
@@ -44,8 +44,33 @@ namespace UE
 
 					if ( Options )
 					{
+						EventAttributes.Emplace( TEXT( "MetersPerUnit" ), LexToString( Options->StageOptions.MetersPerUnit ) );
+						EventAttributes.Emplace( TEXT( "UpAxis" ), Options->StageOptions.UpAxis == EUsdUpAxis::YAxis ? TEXT( "Y" ) : TEXT( "Z" ) );
 						EventAttributes.Emplace( TEXT( "UsePayload" ), LexToString( Options->MeshAssetOptions.bUsePayload ) );
-						EventAttributes.Emplace( TEXT( "PayloadFormat" ), Options->MeshAssetOptions.PayloadFormat );
+						if ( Options->MeshAssetOptions.bUsePayload )
+						{
+							EventAttributes.Emplace( TEXT( "PayloadFormat" ), Options->MeshAssetOptions.PayloadFormat );
+						}
+						EventAttributes.Emplace( TEXT( "BakeMaterials" ), Options->MeshAssetOptions.bBakeMaterials );
+						if ( Options->MeshAssetOptions.bBakeMaterials )
+						{
+							FString BakedPropertiesString;
+							{
+								const UEnum* PropertyEnum = StaticEnum<EMaterialProperty>();
+								for ( const FPropertyEntry& PropertyEntry : Options->MeshAssetOptions.MaterialBakingOptions.Properties )
+								{
+									FString PropertyString = PropertyEnum->GetNameByValue( PropertyEntry.Property ).ToString();
+									PropertyString.RemoveFromStart( TEXT( "MP_" ) );
+									BakedPropertiesString += PropertyString + TEXT( ", " );
+								}
+
+								BakedPropertiesString.RemoveFromEnd( TEXT( ", " ) );
+							}
+
+							EventAttributes.Emplace( TEXT( "RemoveUnrealMaterials" ), Options->MeshAssetOptions.bRemoveUnrealMaterials );
+							EventAttributes.Emplace( TEXT( "BakedProperties" ), BakedPropertiesString );
+							EventAttributes.Emplace( TEXT( "DefaultTextureSize" ), Options->MeshAssetOptions.MaterialBakingOptions.DefaultTextureSize.ToString() );
+						}
 						EventAttributes.Emplace( TEXT( "LowestMeshLOD" ), LexToString( Options->MeshAssetOptions.LowestMeshLOD ) );
 						EventAttributes.Emplace( TEXT( "HighestMeshLOD" ), LexToString( Options->MeshAssetOptions.HighestMeshLOD ) );
 					}
diff --git a/Engine/Plugins/Importers/USDImporter/Source/USDStageImporter/Private/USDStageImporter.cpp b/Engine/Plugins/Importers/USDImporter/Source/USDStageImporter/Private/USDStageImporter.cpp
--- a/Engine/Plugins/Importers/USDImporter/Source/USDStageImporter/Private/USDStageImporter.cpp
+++ b/Engine/Plugins/Importers/USDImporter/Source/USDStageImporter/Private/USDStageImporter.cpp
@@ -1283,6 +1283,12 @@ namespace UsdStageImporterImpl
 				EventAttributes.Emplace( TEXT( "PurposesToImport" ), LexToString( ImportContext.ImportOptions->PurposesToImport ) );
 				EventAttributes.Emplace( TEXT( "NaniteTriangleThreshold" ), LexToString( ImportContext.ImportOptions->NaniteTriangleThreshold ) );
 				EventAttributes.Emplace( TEXT( "RenderContextToImport" ), ImportContext.ImportOptions->RenderContextToImport.ToString() );
+				EventAttributes.Emplace( TEXT( "OverrideStageOptions" ), ImportContext.ImportOptions->bOverrideStageOptions );
+				if( ImportContext.ImportOptions->bOverrideStageOptions )
+				{
+					EventAttributes.Emplace( TEXT( "MetersPerUnit" ), ImportContext.ImportOptions->StageOptions.MetersPerUnit );
+					EventAttributes.Emplace( TEXT( "UpAxis" ), ImportContext.ImportOptions->StageOptions.UpAxis == EUsdUpAxis::YAxis ? TEXT( "Y" ) : TEXT( "Z" ) );
+				}
 				EventAttributes.Emplace( TEXT( "ReuseIdenticalAssets" ), ImportContext.ImportOptions->bReuseIdenticalAssets );
 				EventAttributes.Emplace( TEXT( "ReplaceActorPolicy" ), LexToString( (uint8)ImportContext.ImportOptions->ExistingActorPolicy ) );
 				EventAttributes.Emplace( TEXT( "ReplaceAssetPolicy" ), LexToString( (uint8)ImportContext.ImportOptions->ExistingAssetPolicy ) );
