USD: Relax detection of textures being used as normal maps to allow "float3" attributes if they're named as "normal".

diff --git a/Engine/Plugins/Importers/USDImporter/Source/USDUtilities/Private/USDShadeConversion.cpp b/Engine/Plugins/Importers/USDImporter/Source/USDUtilities/Private/USDShadeConversion.cpp
--- a/Engine/Plugins/Importers/USDImporter/Source/USDUtilities/Private/USDShadeConversion.cpp
+++ b/Engine/Plugins/Importers/USDImporter/Source/USDUtilities/Private/USDShadeConversion.cpp
@@ -482,7 +482,10 @@ namespace UE
 			{
 				FScopedUsdAllocs UsdAllocs;
 
-				const bool bIsNormalInput = ( ShadeInput.GetTypeName() == pxr::SdfValueTypeNames->Normal3f || ShadeInput.GetTypeName() == pxr::SdfValueTypeNames->Normal3fArray );
+				const bool bIsNormalInput = ( ShadeInput.GetTypeName() == pxr::SdfValueTypeNames->Normal3f
+					|| ShadeInput.GetTypeName() == pxr::SdfValueTypeNames->Normal3fArray
+					|| ShadeInput.GetBaseName() == UnrealIdentifiers::Normal
+				);
 
 				pxr::UsdShadeConnectableAPI Source;
 				pxr::TfToken SourceName;
