Fix static analysis warning.

diff --git a/Engine/Plugins/Importers/USDImporter/Source/USDExporter/Private/LevelSequenceExporterUSD.cpp b/Engine/Plugins/Importers/USDImporter/Source/USDExporter/Private/LevelSequenceExporterUSD.cpp
--- a/Engine/Plugins/Importers/USDImporter/Source/USDExporter/Private/LevelSequenceExporterUSD.cpp
+++ b/Engine/Plugins/Importers/USDImporter/Source/USDExporter/Private/LevelSequenceExporterUSD.cpp
@@ -136,16 +136,16 @@ namespace UE
 							EventAttributes.Emplace( TEXT( "BakedProperties" ), BakedPropertiesString );
 							EventAttributes.Emplace( TEXT( "DefaultTextureSize" ), Options->LevelExportOptions.AssetOptions.MaterialBakingOptions.DefaultTextureSize.ToString() );
 						}
-					}
 
-					IUsdClassesModule::SendAnalytics(
-						MoveTemp( EventAttributes ),
-						FString::Printf( TEXT( "Export.%s" ), *ClassName ),
-						bAutomated,
-						ElapsedSeconds,
-						( Options->EndFrame - Options->StartFrame ),
-						Extension
-					);
+						IUsdClassesModule::SendAnalytics(
+							MoveTemp( EventAttributes ),
+							FString::Printf( TEXT( "Export.%s" ), *ClassName ),
+							bAutomated,
+							ElapsedSeconds,
+							( Options->EndFrame - Options->StartFrame ),
+							Extension
+						);
+					}
 				}
 			}
 
