USD Groom: Added support to import groom Alembic referenced from USD.

diff --git a/Engine/Plugins/Importers/USDImporter/Source/USDSchemas/Private/USDGroomTranslator.cpp b/Engine/Plugins/Importers/USDImporter/Source/USDSchemas/Private/USDGroomTranslator.cpp
new file mode 100644
--- /dev/null
+++ b/Engine/Plugins/Importers/USDImporter/Source/USDSchemas/Private/USDGroomTranslator.cpp
@@ -0,0 +1,221 @@
+// Copyright Epic Games, Inc. All Rights Reserved.
+
+#include "USDGroomTranslator.h"
+
+#if USE_USD_SDK && WITH_EDITOR
+
+#include "GroomAsset.h"
+#include "GroomBuilder.h"
+#include "GroomImportOptions.h"
+#include "HairDescription.h"
+#include "HairStrandsImporter.h"
+#include "Misc/ArchiveMD5.h"
+
+#include "USDAssetImportData.h"
+#include "USDTypesConversion.h"
+#include "USDGroomConversion.h"
+
+#include "USDIncludesStart.h"
+#include "pxr/base/tf/type.h"
+#include "pxr/usd/usd/prim.h"
+#include "pxr/usd/usd/timeCode.h"
+#include "pxr/usd/usdGeom/curves.h"
+#include "USDIncludesEnd.h"
+
+namespace UE::UsdGroomTranslator::Private
+{
+	UGroomImportOptions* CreateGroomImportOptions(const FHairDescriptionGroups& GroupsDescription, const TArray<FHairGroupsInterpolation>& BuildSettings)
+	{
+		// Create a new groom import options and populate the interpolation settings based on the group count
+		UGroomImportOptions* ImportOptions = NewObject<UGroomImportOptions>();
+		const uint32 GroupCount = GroupsDescription.HairGroups.Num();
+		if (GroupCount != uint32(ImportOptions->InterpolationSettings.Num()))
+		{
+			ImportOptions->InterpolationSettings.Init(FHairGroupsInterpolation(), GroupCount);
+		}
+
+		TOptional<FHairGroupsInterpolation> LastBuildSettings;
+		if (BuildSettings.Num() > 0)
+		{
+			LastBuildSettings = BuildSettings.Last();
+		}
+
+		// If there are less build settings than groups, use the last one that was specified by the user
+		for (uint32 Index = 0; Index < GroupCount; ++Index)
+		{
+			if (BuildSettings.IsValidIndex(Index))
+			{
+				ImportOptions->InterpolationSettings[Index] = BuildSettings[Index];
+			}
+			else if (LastBuildSettings.IsSet())
+			{
+				ImportOptions->InterpolationSettings[Index] = LastBuildSettings.GetValue();
+			}
+		}
+		return ImportOptions;
+	}
+
+	FSHAHash ComputeHairDescriptionHash(FHairDescription& HairDescription, TArray<FHairGroupsInterpolation>& BuildSettings)
+	{
+		// The computed hash takes the hair description and group build settings into account since the groom builder
+		// will use those settings to build the groom asset from the description
+		FArchiveMD5 ArMD5;
+		HairDescription.Serialize(ArMD5);
+
+		for (FHairGroupsInterpolation& GroupSettings : BuildSettings)
+		{
+			GroupSettings.BuildDDCKey(ArMD5);
+		}
+
+		FMD5Hash MD5Hash;
+		ArMD5.GetHash(MD5Hash);
+
+		FSHA1 SHA1;
+		SHA1.Update(MD5Hash.GetBytes(), MD5Hash.GetSize());
+		SHA1.Final();
+
+		FSHAHash SHAHash;
+		SHA1.GetHash(SHAHash.Hash);
+
+		return SHAHash;
+	}
+}
+
+class FUsdGroomCreateAssetsTaskChain : public FUsdSchemaTranslatorTaskChain
+{
+public:
+	explicit FUsdGroomCreateAssetsTaskChain(const TSharedRef<FUsdSchemaTranslationContext>& InContext, const UE::FSdfPath& InPrimPath)
+		: PrimPath(InPrimPath)
+		, Context(InContext)
+	{
+		SetupTasks();
+	}
+
+protected:
+
+	UE::FSdfPath PrimPath;
+	TSharedRef<FUsdSchemaTranslationContext> Context;
+
+	FHairDescription HairDescription;
+
+protected:
+	UE::FUsdPrim GetPrim() const { return Context->Stage.GetPrimAtPath(PrimPath); }
+
+protected:
+	void SetupTasks()
+	{
+		FScopedUnrealAllocs UnrealAllocs;
+
+		// Create hair description (Async)
+		Do(ESchemaTranslationLaunchPolicy::Async,
+			[this]() -> bool
+			{
+				const bool bSuccess = UsdToUnreal::ConvertGroomHierarchy(GetPrim(), pxr::UsdTimeCode::EarliestTime(), FTransform::Identity, HairDescription);
+
+				return bSuccess && HairDescription.IsValid();
+			});
+		// Build groom asset from hair description (Sync)
+		Then(ESchemaTranslationLaunchPolicy::Sync,
+			[this]()
+			{
+				TRACE_CPUPROFILER_EVENT_SCOPE(FUsdGroomCreateAssetsTaskChain::Build);
+
+				// Extract the groom groups info from the hair description to get the number of groups
+				FHairDescriptionGroups GroupsDescription;
+				FGroomBuilder::BuildHairDescriptionGroups(HairDescription, GroupsDescription);
+
+				UGroomImportOptions* ImportOptions = UE::UsdGroomTranslator::Private::CreateGroomImportOptions(GroupsDescription, Context->GroomInterpolationSettings);
+
+				FSHAHash SHAHash = UE::UsdGroomTranslator::Private::ComputeHairDescriptionHash(HairDescription, ImportOptions->InterpolationSettings);
+
+				const FString PrimPathString = PrimPath.GetString();
+				UGroomAsset* GroomAsset = Cast<UGroomAsset>(Context->AssetCache->GetCachedAsset(SHAHash.ToString()));
+				if (!GroomAsset)
+				{
+					FName AssetName = MakeUniqueObjectName(GetTransientPackage(), UGroomAsset::StaticClass(), *FPaths::GetBaseFilename(PrimPathString));
+
+					FHairImportContext HairImportContext(ImportOptions, GetTransientPackage(), UGroomAsset::StaticClass(), AssetName, Context->ObjectFlags | EObjectFlags::RF_Public);
+					UGroomAsset* ExistingAsset = nullptr;
+					GroomAsset = FHairStrandsImporter::ImportHair(HairImportContext, HairDescription, ExistingAsset);
+					if (GroomAsset)
+					{
+						Context->AssetCache->CacheAsset(SHAHash.ToString(), GroomAsset);
+#if WITH_EDITOR
+						UUsdAssetImportData* ImportData = NewObject<UUsdAssetImportData>(GroomAsset, TEXT("UUSDAssetImportData"));
+						ImportData->PrimPath = PrimPathString;
+						ImportData->ImportOptions = ImportOptions;
+
+						GroomAsset->AssetImportData = ImportData;
+#endif // WITH_EDITOR
+					}
+				}
+
+				if (GroomAsset)
+				{
+					Context->AssetCache->LinkAssetToPrim(PrimPathString, GroomAsset);
+				}
+
+				return true;
+			});
+	}
+};
+
+bool FUsdGroomTranslator::IsGroomPrim() const
+{
+	FScopedUsdAllocs UsdAllocs;
+
+	// #ueent_todo: Better validation pf prim being a groom like using custom groom schema
+	TArray<TUsdStore<pxr::UsdPrim>> ChildCurvesPrims = UsdUtils::GetAllPrimsOfType(GetPrim(), pxr::TfType::Find<pxr::UsdGeomCurves>());
+
+	return ChildCurvesPrims.Num() > 0;
+}
+
+void FUsdGroomTranslator::CreateAssets()
+{
+	if (!IsGroomPrim())
+	{
+		return Super::CreateAssets();
+	}
+
+	Context->TranslatorTasks.Add(MakeShared<FUsdGroomCreateAssetsTaskChain>(Context, PrimPath));
+}
+
+USceneComponent* FUsdGroomTranslator::CreateComponents()
+{
+	if (!IsGroomPrim())
+	{
+		return Super::CreateComponents();
+	}
+
+	return nullptr;
+}
+
+void FUsdGroomTranslator::UpdateComponents(USceneComponent* SceneComponent)
+{
+	if (!IsGroomPrim())
+	{
+		Super::UpdateComponents(SceneComponent);
+	}
+}
+
+bool FUsdGroomTranslator::CollapsesChildren(ECollapsingType CollapsingType) const
+{
+	if (!IsGroomPrim())
+	{
+		return Super::CollapsesChildren(CollapsingType);
+	}
+
+	return true;
+}
+
+bool FUsdGroomTranslator::CanBeCollapsed(ECollapsingType CollapsingType) const
+{
+	if (!IsGroomPrim())
+	{
+		return Super::CanBeCollapsed(CollapsingType);
+	}
+
+	return true;
+}
+
+#endif // #if USE_USD_SDK
\ No newline at end of file
diff --git a/Engine/Plugins/Importers/USDImporter/Source/USDSchemas/Private/USDSchemasModule.cpp b/Engine/Plugins/Importers/USDImporter/Source/USDSchemas/Private/USDSchemasModule.cpp
--- a/Engine/Plugins/Importers/USDImporter/Source/USDSchemas/Private/USDSchemasModule.cpp
+++ b/Engine/Plugins/Importers/USDImporter/Source/USDSchemas/Private/USDSchemasModule.cpp
@@ -15,6 +15,7 @@
 #include "USDGeomMeshTranslator.h"
 #include "USDGeomPointInstancerTranslator.h"
 #include "USDGeomXformableTranslator.h"
+#include "USDGroomTranslator.h"
 #include "USDLuxLightTranslator.h"
 #include "USDMemory.h"
 #include "USDShadeMaterialTranslator.h"
@@ -48,6 +49,7 @@ public:
 		if ( GIsEditor )
 		{
 			UsdSkelRootTranslatorHandle = GetTranslatorRegistry().Register< FUsdSkelRootTranslator >( TEXT("UsdSkelRoot") );
+			//UsdGroomTranslatorHandle = GetTranslatorRegistry().Register< FUsdGroomTranslator >( TEXT("UsdGeomXformable") ); // #ueent_todo: Enable it when the groom schema is implemented
 
 			if ( IMDLImporterModule* MDLImporterModule = FModuleManager::Get().LoadModulePtr< IMDLImporterModule >( TEXT("MDLImporter") ) )
 			{
@@ -70,6 +72,7 @@ public:
 		if ( GIsEditor )
 		{
 			GetTranslatorRegistry().Unregister( UsdSkelRootTranslatorHandle );
+			GetTranslatorRegistry().Unregister( UsdGroomTranslatorHandle );
 
 			GetTranslatorRegistry().Unregister( MdlUsdShadeMaterialTranslatorHandle );
 			GetRenderContextRegistry().Unregister( FMdlUsdShadeMaterialTranslator::MdlRenderContext );
@@ -108,6 +111,7 @@ protected:
 	FRegisteredSchemaTranslatorHandle UsdSkelRootTranslatorHandle;
 	FRegisteredSchemaTranslatorHandle UsdGeomXformableTranslatorHandle;
 	FRegisteredSchemaTranslatorHandle UsdShadeMaterialTranslatorHandle;
+	FRegisteredSchemaTranslatorHandle UsdGroomTranslatorHandle;
 
 #if defined(HAS_USDLUX_LIGHTAPI)
 	FRegisteredSchemaTranslatorHandle UsdLuxBoundableLightBaseTranslatorHandle;
diff --git a/Engine/Plugins/Importers/USDImporter/Source/USDSchemas/Public/USDGroomTranslator.h b/Engine/Plugins/Importers/USDImporter/Source/USDSchemas/Public/USDGroomTranslator.h
new file mode 100644
--- /dev/null
+++ b/Engine/Plugins/Importers/USDImporter/Source/USDSchemas/Public/USDGroomTranslator.h
@@ -0,0 +1,27 @@
+// Copyright Epic Games, Inc. All Rights Reserved.
+
+#pragma once
+
+#if USE_USD_SDK && WITH_EDITOR
+
+#include "USDGeomXformableTranslator.h"
+
+class USDSCHEMAS_API FUsdGroomTranslator : public FUsdGeomXformableTranslator
+{
+public:
+	using Super = FUsdGeomXformableTranslator;
+	using FUsdGeomXformableTranslator::FUsdGeomXformableTranslator;
+
+	virtual void CreateAssets() override;
+
+	virtual USceneComponent* CreateComponents() override;
+	virtual void UpdateComponents(USceneComponent* SceneComponent) override;
+
+	virtual bool CollapsesChildren(ECollapsingType CollapsingType) const override;
+	virtual bool CanBeCollapsed(ECollapsingType CollapsingType) const override;
+
+private:
+	bool IsGroomPrim() const;
+};
+
+#endif // #if USE_USD_SDK
\ No newline at end of file
diff --git a/Engine/Plugins/Importers/USDImporter/Source/USDSchemas/Public/USDSchemaTranslator.h b/Engine/Plugins/Importers/USDImporter/Source/USDSchemas/Public/USDSchemaTranslator.h
--- a/Engine/Plugins/Importers/USDImporter/Source/USDSchemas/Public/USDSchemaTranslator.h
+++ b/Engine/Plugins/Importers/USDImporter/Source/USDSchemas/Public/USDSchemaTranslator.h
@@ -16,6 +16,7 @@
 #include "UsdWrappers/UsdTyped.h"
 
 #include "Async/Future.h"
+#include "GroomAssetInterpolation.h"
 #include "HAL/ThreadSafeBool.h"
 #include "Misc/Optional.h"
 #include "Templates/SubclassOf.h"
@@ -194,6 +195,9 @@ struct USDSCHEMAS_API FUsdSchemaTranslationContext : public TSharedFromThis< FUs
 	/** If true, we will also try creating UAnimSequence skeletal animation assets when parsing SkelRoot prims */
 	bool bAllowParsingSkeletalAnimations = true;
 
+	/** Groom group interpolation settings */
+	TArray<FHairGroupsInterpolation> GroomInterpolationSettings;
+
 	bool IsValid() const
 	{
 		return Level != nullptr;
diff --git a/Engine/Plugins/Importers/USDImporter/Source/USDSchemas/USDSchemas.Build.cs b/Engine/Plugins/Importers/USDImporter/Source/USDSchemas/USDSchemas.Build.cs
--- a/Engine/Plugins/Importers/USDImporter/Source/USDSchemas/USDSchemas.Build.cs
+++ b/Engine/Plugins/Importers/USDImporter/Source/USDSchemas/USDSchemas.Build.cs
@@ -35,6 +35,13 @@ namespace UnrealBuildTool.Rules
 				}
 			);
 
+			PublicDependencyModuleNames.AddRange(
+				new string[]
+				{
+					"HairStrandsCore",
+				}
+			);
+
 			if (Target.bBuildEditor)
 			{
 				PrivateDependencyModuleNames.AddRange(
@@ -42,6 +49,7 @@ namespace UnrealBuildTool.Rules
 					{
 						"BlueprintGraph",
 						"GeometryCacheUSD",
+						"HairStrandsEditor",
 						"Kismet",
 						"LiveLinkGraphNode",
 						"MaterialEditor",
diff --git a/Engine/Plugins/Importers/USDImporter/Source/USDStageImporter/Private/USDStageImporter.cpp b/Engine/Plugins/Importers/USDImporter/Source/USDStageImporter/Private/USDStageImporter.cpp
--- a/Engine/Plugins/Importers/USDImporter/Source/USDStageImporter/Private/USDStageImporter.cpp
+++ b/Engine/Plugins/Importers/USDImporter/Source/USDStageImporter/Private/USDStageImporter.cpp
@@ -375,14 +375,17 @@ namespace UsdStageImporterImpl
 		FString AssetSuffix;
 		FString AssetPath = Asset->GetName();
 
+		if (UUsdAssetImportData* AssetImportData = UsdUtils::GetAssetImportData(Asset))
+		{
+			AssetPath = AssetImportData->PrimPath;
+		}
+
 		if (UStaticMesh* Mesh = Cast<UStaticMesh>(Asset))
 		{
 			AssetPrefix = TEXT("SM_");
 
 			if (UUsdAssetImportData* AssetImportData = Cast<UUsdAssetImportData>(Mesh->AssetImportData))
 			{
-				AssetPath = AssetImportData->PrimPath;
-
 				// If we have multiple LODs here we must have parsed the LOD variant set pattern. If our prims were named with the LOD
 				// pattern, go from e.g. '/Root/MyMesh/LOD0' to '/Root/MyMesh', or else every single LOD mesh will be named "SM_LOD0_X".
 				// We'll actually check though because if the user set a custom name for his prim other than LOD0 then we'll keep that
@@ -403,33 +406,16 @@ namespace UsdStageImporterImpl
 		{
 			AssetPrefix = TEXT("SK_");
 
-			if (UUsdAssetImportData* AssetImportData = Cast<UUsdAssetImportData>(SkMesh->GetAssetImportData()))
-			{
-				AssetPath = AssetImportData->PrimPath;
-			}
 		}
 		else if (USkeleton* Skeleton = Cast<USkeleton>(Asset))
 		{
 			AssetSuffix = TEXT("_Skeleton");
 
-			// We always set the corresponding mesh as preview mesh on import. Fetching the name here is really important
-			// as it can determine the destination path and how the asset conflicts are resolved
-			if (USkeletalMesh* SkeletalMesh = Skeleton->GetPreviewMesh())
-			{
-				if (UUsdAssetImportData* AssetImportData = Cast<UUsdAssetImportData>(SkeletalMesh->GetAssetImportData()))
-				{
-					AssetPath = AssetImportData->PrimPath;
-				}
-			}
+			// Note that UsdUtils::GetAssetImportData fetches the AssetImportData of the PreviewMesh
 		}
 		else if ( UAnimSequence* AnimSequence = Cast<UAnimSequence>( Asset ) )
 		{
 			AssetPrefix = TEXT( "Anim_" );
-
-			if (UUsdAssetImportData* AssetImportData = Cast<UUsdAssetImportData>(AnimSequence->AssetImportData))
-			{
-				AssetPath = AssetImportData->PrimPath;
-			}
 		}
 		else if (UMaterialInterface* Material = Cast<UMaterialInterface>(Asset))
 		{
@@ -444,11 +430,6 @@ namespace UsdStageImporterImpl
 		else if (UTexture* Texture = Cast<UTexture>(Asset))
 		{
 			AssetPrefix = TEXT("T_");
-
-			if (UUsdAssetImportData* AssetImportData = Cast<UUsdAssetImportData>(Texture->AssetImportData))
-			{
-				AssetPath = AssetImportData->GetFirstFilename();
-			}
 		}
 
 		FString FinalName = FPaths::GetBaseFilename(AssetPath);
@@ -1509,6 +1490,7 @@ void UUsdStageImporter::ImportFromFile(FUsdStageImportContext& ImportContext)
 	TranslationContext->MaterialToPrimvarToUVIndex = &ImportContext.MaterialToPrimvarToUVIndex;
 	TranslationContext->InfoCache = ImportContext.InfoCache;
 	TranslationContext->BlendShapesByPath = &BlendShapesByPath;
+	TranslationContext->GroomInterpolationSettings = ImportContext.ImportOptions->GroomInterpolationSettings;
 	{
 		UsdStageImporterImpl::CacheCollapsingState( TranslationContext.Get() );
 		UsdStageImporterImpl::ImportMaterials( ImportContext, TranslationContext.Get() );
diff --git a/Engine/Plugins/Importers/USDImporter/Source/USDStageImporter/Public/USDStageImportOptions.h b/Engine/Plugins/Importers/USDImporter/Source/USDStageImporter/Public/USDStageImportOptions.h
--- a/Engine/Plugins/Importers/USDImporter/Source/USDStageImporter/Public/USDStageImportOptions.h
+++ b/Engine/Plugins/Importers/USDImporter/Source/USDStageImporter/Public/USDStageImportOptions.h
@@ -8,6 +8,7 @@
 #include "CoreMinimal.h"
 #include "Engine/EngineTypes.h"
 #include "Factories/MaterialImportHelpers.h"
+#include "GroomAssetInterpolation.h"
 #include "UObject/ObjectMacros.h"
 
 #include "USDStageImportOptions.generated.h"
@@ -86,6 +87,9 @@ public:
 	FUsdStageOptions StageOptions;
 
 
+	/** Groom group interpolation settings */
+	UPROPERTY(EditAnywhere, config, BlueprintReadWrite, Category = "Groom")
+	TArray<FHairGroupsInterpolation> GroomInterpolationSettings;
 
 	/**
 	 * If enabled, whenever two different prims import into identical assets, only one of those assets will be kept and reused.
diff --git a/Engine/Plugins/Importers/USDImporter/Source/USDStageImporter/USDStageImporter.Build.cs b/Engine/Plugins/Importers/USDImporter/Source/USDStageImporter/USDStageImporter.Build.cs
--- a/Engine/Plugins/Importers/USDImporter/Source/USDStageImporter/USDStageImporter.Build.cs
+++ b/Engine/Plugins/Importers/USDImporter/Source/USDStageImporter/USDStageImporter.Build.cs
@@ -36,6 +36,13 @@ namespace UnrealBuildTool.Rules
 					"USDUtilities",
 				}
 			);
+
+			PublicDependencyModuleNames.AddRange(
+				new string[]
+				{
+					"HairStrandsCore",
+				}
+			);
 		}
 	}
 }
diff --git a/Engine/Plugins/Importers/USDImporter/Source/USDUtilities/Private/USDConversionUtils.cpp b/Engine/Plugins/Importers/USDImporter/Source/USDUtilities/Private/USDConversionUtils.cpp
--- a/Engine/Plugins/Importers/USDImporter/Source/USDUtilities/Private/USDConversionUtils.cpp
+++ b/Engine/Plugins/Importers/USDImporter/Source/USDUtilities/Private/USDConversionUtils.cpp
@@ -36,6 +36,7 @@
 #include "Engine/StaticMesh.h"
 #include "Engine/Texture.h"
 #include "GeometryCache.h"
+#include "GroomAsset.h"
 #include "InstancedFoliageActor.h"
 #include "LandscapeProxy.h"
 
@@ -1086,6 +1087,11 @@ UUsdAssetImportData* UsdUtils::GetAssetImportData( UObject* Asset )
 	{
 		ImportData = Cast<UUsdAssetImportData>( GeometryCache->AssetImportData );
 	}
+	else if ( UGroomAsset* Groom = Cast<UGroomAsset>( Asset ) )
+	{
+		ImportData = Cast<UUsdAssetImportData>( Groom->AssetImportData );
+	}
+
 #endif
 	return ImportData;
 }
diff --git a/Engine/Plugins/Importers/USDImporter/Source/USDUtilities/Private/USDGroomConversion.cpp b/Engine/Plugins/Importers/USDImporter/Source/USDUtilities/Private/USDGroomConversion.cpp
new file mode 100644
--- /dev/null
+++ b/Engine/Plugins/Importers/USDImporter/Source/USDUtilities/Private/USDGroomConversion.cpp
@@ -0,0 +1,343 @@
+// Copyright Epic Games, Inc. All Rights Reserved.
+
+#if USE_USD_SDK
+
+#include "USDGroomConversion.h"
+
+#include "HairDescription.h"
+#include "USDPrimConversion.h"
+#include "USDTypesConversion.h"
+#include "UsdWrappers/UsdPrim.h"
+
+#include "USDIncludesStart.h"
+#include "pxr/usd/usdGeom/basisCurves.h"
+#include "pxr/usd/usdGeom/primvar.h"
+#include "pxr/usd/usdGeom/primvarsAPI.h"
+#include "pxr/usd/usdGeom/xformable.h"
+#include "USDIncludesEnd.h"
+
+namespace UE::UsdGroomConversion::Private
+{
+	bool IsAttributeValid(const FString& AttributeName)
+	{
+		// Ignore the attributes that aren't prefixed with groom_
+		return AttributeName.StartsWith(TEXT("groom_"));
+	}
+
+	// Default USD to UE type converter
+	template<typename TAttributeType>
+	struct TConverter
+	{
+		template<typename InType>
+		static TAttributeType ConvertType(const InType& In)
+		{
+			return In;
+		}
+
+		static TAttributeType DefaultValue()
+		{
+			return TAttributeType();
+		}
+	};
+
+	// String to FName converter
+	template<>
+	struct TConverter<FName>
+	{
+		template<typename InType>
+		static FName ConvertType(const InType& In)
+		{
+			return UsdToUnreal::ConvertName(In);
+		}
+
+		static FName DefaultValue()
+		{
+			return NAME_None;
+		}
+	};
+
+	// GfVec2f/Float2 to FVector2f converter
+	template<>
+	struct TConverter<FVector2f>
+	{
+		template<typename InType>
+		static FVector2f ConvertType(const InType& In)
+		{
+			return FVector2f(In[0], In[1]);
+		}
+
+		static FVector2f DefaultValue()
+		{
+			return FVector2f::ZeroVector;
+		}
+	};
+
+	// GfVec3f/Float3 to FVector3f converter
+	template<>
+	struct TConverter<FVector3f>
+	{
+		template<typename InType>
+		static FVector3f ConvertType(const InType& In)
+		{
+			return FVector3f(In[0], In[1], In[2]);
+		}
+
+		static FVector3f DefaultValue()
+		{
+			return FVector3f::ZeroVector;
+		}
+	};
+
+	/** Convert the given primvar to hair attribute in the proper scope */
+	template <typename USDType, typename UEType>
+	void ConvertPrimvar(const pxr::UsdGeomPrimvar& Primvar, const pxr::UsdTimeCode& TimeCode, int32 StartStrandID, int32 NumStrands, int32 StartVertexID, int32 NumVertices, FHairDescription& HairDescription)
+	{
+		FString PrimvarName = UsdToUnreal::ConvertToken(Primvar.GetPrimvarName());
+		FName AttributeName(*PrimvarName);
+
+		pxr::VtArray<USDType> Values;
+		Primvar.Get(&Values, TimeCode);
+		int32 NumValues = Values.size();
+
+		if (NumValues == 0)
+		{
+			return;
+		}
+
+		// Check the supported attribute scopes by checking the primvar interpolation
+		// and the number of values as a fallback
+		pxr::TfToken Interpolation = Primvar.GetInterpolation();
+		if (Interpolation == pxr::UsdGeomTokens->uniform || Interpolation == pxr::UsdGeomTokens->constant || NumValues == NumStrands)
+		{
+			if (Interpolation == pxr::UsdGeomTokens->uniform && NumValues != NumStrands)
+			{
+				// Attribute has a mismatch between the expected and actual number of values so skip it
+				return;
+			}
+
+			TStrandAttributesRef<UEType> StrandAttributeRef = HairDescription.StrandAttributes().GetAttributesRef<UEType>(AttributeName);
+			if (!StrandAttributeRef.IsValid())
+			{
+				HairDescription.StrandAttributes().RegisterAttribute<UEType>(AttributeName, 1, TConverter<UEType>::DefaultValue());
+				StrandAttributeRef = HairDescription.StrandAttributes().GetAttributesRef<UEType>(AttributeName);
+			}
+
+			// Constant scope is converted to uniform scope by setting the single value over all strands
+			bool bIsConstantValue = Interpolation == pxr::UsdGeomTokens->constant;
+			for (int32 StrandIndex = 0; StrandIndex < NumStrands; ++StrandIndex)
+			{
+				StrandAttributeRef[FStrandID(StartStrandID + StrandIndex)] = TConverter<UEType>::ConvertType(Values[bIsConstantValue ? 0 : StrandIndex]);
+			}
+		}
+		else if (Interpolation == pxr::UsdGeomTokens->vertex || NumValues == NumVertices)
+		{
+			TVertexAttributesRef<UEType> VertexAttributeRef = HairDescription.VertexAttributes().GetAttributesRef<UEType>(AttributeName);
+			if (!VertexAttributeRef.IsValid())
+			{
+				HairDescription.VertexAttributes().RegisterAttribute<UEType>(AttributeName, 1, TConverter<UEType>::DefaultValue());
+				VertexAttributeRef = HairDescription.VertexAttributes().GetAttributesRef<UEType>(AttributeName);
+			}
+
+			for (int32 VertexIndex = 0; VertexIndex < NumVertices; ++VertexIndex)
+			{
+				VertexAttributeRef[FVertexID(StartVertexID + VertexIndex)] = TConverter<UEType>::ConvertType(Values[VertexIndex]);
+			}
+		}
+	}
+
+	/** Convert the given primvars to hair attributes */
+	void ConvertPrimvars(const std::vector<pxr::UsdGeomPrimvar>& Attributes, const pxr::UsdTimeCode& TimeCode, int32 StartStrandID, int32 NumStrands, int32 StartVertexID, int32 NumVertices, FHairDescription& HairDescription)
+	{
+		for (const pxr::UsdGeomPrimvar& Primvar : Attributes)
+		{
+			// Process only groom attributes
+			FString UEPrimvarName = UsdToUnreal::ConvertToken(Primvar.GetPrimvarName());
+			if (!IsAttributeValid(UEPrimvarName))
+			{
+				continue;
+			}
+
+			// #ueent_todo: Check if the groom_color attribute is animated (for GroomCache)
+
+			// Dispatch the USD to UE conversion by the primvar scalar type
+			pxr::SdfValueTypeName ScalarTypeName = Primvar.GetTypeName().GetScalarType();
+			if (ScalarTypeName == pxr::SdfValueTypeNames->Int)
+			{
+				ConvertPrimvar<int, int>(Primvar, TimeCode, StartStrandID, NumStrands, StartVertexID, NumVertices, HairDescription);
+			}
+			else if (ScalarTypeName == pxr::SdfValueTypeNames->Float)
+			{
+				ConvertPrimvar<float, float>(Primvar, TimeCode, StartStrandID, NumStrands, StartVertexID, NumVertices, HairDescription);
+			}
+			else if (ScalarTypeName == pxr::SdfValueTypeNames->Vector3f || ScalarTypeName == pxr::SdfValueTypeNames->Float3)
+			{
+				ConvertPrimvar<pxr::GfVec3f, FVector3f>(Primvar, TimeCode, StartStrandID, NumStrands, StartVertexID, NumVertices, HairDescription);
+			}
+			else if (ScalarTypeName == pxr::SdfValueTypeNames->TexCoord2f || ScalarTypeName == pxr::SdfValueTypeNames->Float2)
+			{
+				ConvertPrimvar<pxr::GfVec2f, FVector2f>(Primvar, TimeCode, StartStrandID, NumStrands, StartVertexID, NumVertices, HairDescription);
+			}
+			else if (ScalarTypeName == pxr::SdfValueTypeNames->Double2)
+			{
+				ConvertPrimvar<pxr::GfVec2d, FVector2f>(Primvar, TimeCode, StartStrandID, NumStrands, StartVertexID, NumVertices, HairDescription);
+			}
+			else if (ScalarTypeName == pxr::SdfValueTypeNames->String)
+			{
+				ConvertPrimvar<std::string, FName>(Primvar, TimeCode, StartStrandID, NumStrands, StartVertexID, NumVertices, HairDescription);
+			}
+		}
+	}
+}
+
+namespace UsdToUnreal
+{
+	bool ConvertGroomHierarchy(const pxr::UsdPrim& Prim, const pxr::UsdTimeCode& TimeCode, const FTransform& ParentTransform, FHairDescription& HairDescription)
+	{
+		TRACE_CPUPROFILER_EVENT_SCOPE(UsdToUnreal::ConvertGroomHierarchy);
+
+		FScopedUsdAllocs Allocs;
+
+		// ref. AlembicHairTranslator ParseObject
+
+		bool bSuccess = true;
+		FTransform TransformToPropagate = ParentTransform;
+		if (pxr::UsdGeomBasisCurves Curves = pxr::UsdGeomBasisCurves(Prim))
+		{
+			// #ueent_todo: Add support for GroomCache import (needs AnimInfo)
+			// There are also curve attributes that are not currently queried (knots, orders, basis, type, etc.)
+
+			// Get curve and curve vextex count and vertex positions, the minimal data set required
+			const int32 NumCurves = Curves.GetCurveCount(TimeCode);
+			if (NumCurves == 0)
+			{
+				return false;
+			}
+
+			pxr::UsdAttribute NumVerticesAttr = Curves.GetCurveVertexCountsAttr();
+			pxr::VtArray<int> NumVertices;
+			if (!NumVerticesAttr || !NumVerticesAttr.Get(&NumVertices, TimeCode) || NumVertices.empty())
+			{
+				return false;
+			}
+
+			pxr::UsdAttribute PointsAttr = Curves.GetPointsAttr();
+			pxr::VtArray<pxr::GfVec3f> Points;
+			if (!PointsAttr || !PointsAttr.Get(&Points, TimeCode) || Points.empty())
+			{
+				return false;
+			}
+
+			// Get the starting strand and vertex IDs for this group of Curves
+			int32 StartStrandID = HairDescription.GetNumStrands();
+			int32 StartVertexID = HairDescription.GetNumVertices();
+
+			bool bResetXformStack = false;
+			FTransform PrimTransform = FTransform::Identity;
+			bool bConverted = UsdToUnreal::ConvertXformable(Prim.GetStage(), Curves, PrimTransform, TimeCode.GetValue(), &bResetXformStack);
+			if (bConverted)
+			{
+				TransformToPropagate = bResetXformStack ? PrimTransform : ParentTransform * PrimTransform;
+			}
+
+			uint32 PointIndex = 0;
+			const FUsdStageInfo StageInfo(Prim.GetStage());
+			for (int32 CurveIndex = 0; CurveIndex < NumCurves; ++CurveIndex)
+			{
+				const int CurveNumVertices = NumVertices[CurveIndex];
+
+				FStrandID StrandID = HairDescription.AddStrand();
+
+				SetHairStrandAttribute(HairDescription, StrandID, HairAttribute::Strand::VertexCount, CurveNumVertices);
+
+				for (int32 CurveVertexIndex = 0; CurveVertexIndex < CurveNumVertices; ++CurveVertexIndex, ++PointIndex)
+				{
+					FVertexID VertexID = HairDescription.AddVertex();
+
+					const pxr::GfVec3f& Point = Points[PointIndex];
+
+					FVector Position = UsdToUnreal::ConvertVector(StageInfo, Point);
+					FVector ConvertedPosition = TransformToPropagate.TransformPosition(Position);
+					SetHairVertexAttribute(HairDescription, VertexID, HairAttribute::Vertex::Position, FVector3f(ConvertedPosition));
+				}
+			}
+
+			// Set width values. Width attribute is not a generic primvar so must be processed separately
+			int NumPoints = Points.size();
+			pxr::UsdAttribute Widths = Curves.GetWidthsAttr();
+			pxr::VtArray<float> WidthsArray;
+			if (Widths && Widths.Get(&WidthsArray, TimeCode) && !WidthsArray.empty())
+			{
+				// Determine the scope of the Widths attribute
+				pxr::TfToken Interpolation = Curves.GetWidthsInterpolation();
+
+				const float UEMetersPerUnit = 0.01f;
+				const float Scale = !FMath::IsNearlyEqual(StageInfo.MetersPerUnit, UEMetersPerUnit) ? StageInfo.MetersPerUnit / UEMetersPerUnit : 1.0f;
+				if (Interpolation == pxr::UsdGeomTokens->uniform || Interpolation == pxr::UsdGeomTokens->constant)
+				{
+					const float ConstWidth = WidthsArray[0] * Scale;
+					TStrandAttributesRef<float> WidthStrandAttributeRef = HairDescription.StrandAttributes().GetAttributesRef<float>(HairAttribute::Strand::Width);
+					if (!WidthStrandAttributeRef.IsValid())
+					{
+						HairDescription.StrandAttributes().RegisterAttribute<float>(HairAttribute::Strand::Width);
+						WidthStrandAttributeRef = HairDescription.StrandAttributes().GetAttributesRef<float>(HairAttribute::Strand::Width);
+					}
+
+					for (int32 Index = 0; Index < NumCurves; ++Index)
+					{
+						WidthStrandAttributeRef[FStrandID(StartStrandID + Index)] = Interpolation == pxr::UsdGeomTokens->constant ? ConstWidth : WidthsArray[Index] * Scale;
+					}
+				}
+				else if (Interpolation == pxr::UsdGeomTokens->vertex)
+				{
+					TVertexAttributesRef<float> VertexAttributeRef = HairDescription.VertexAttributes().GetAttributesRef<float>(HairAttribute::Vertex::Width);
+					if (!VertexAttributeRef.IsValid())
+					{
+						HairDescription.VertexAttributes().RegisterAttribute<float>(HairAttribute::Vertex::Width);
+						VertexAttributeRef = HairDescription.VertexAttributes().GetAttributesRef<float>(HairAttribute::Vertex::Width);
+					}
+
+					for (int32 Index = 0; Index < NumPoints; ++Index)
+					{
+						VertexAttributeRef[FVertexID(StartVertexID + Index)] = WidthsArray[Index] * Scale;
+					}
+				}
+			}
+
+			// Extract the groom attributes (at vertex and strand scope) from the generic primvars
+			pxr::UsdGeomPrimvarsAPI PrimvarsAPI(Prim);
+			std::vector<pxr::UsdGeomPrimvar> Primvars = PrimvarsAPI.GetPrimvars();
+			if (!Primvars.empty())
+			{
+				UE::UsdGroomConversion::Private::ConvertPrimvars(Primvars, TimeCode, StartStrandID, NumCurves, StartVertexID, NumPoints, HairDescription);
+			}
+
+			// Groom-scope attributes are not currently extracted since they are not strictly necessary
+
+			// Following the USD recommendation that gprims be not nested, a Curves prim is handled as a leaf
+		}
+		else if (pxr::UsdGeomXformable Xform = pxr::UsdGeomXformable(Prim))
+		{
+			// Propagate the prim transform to the children
+			bool bResetXformStack = false;
+			FTransform PrimTransform = FTransform::Identity;
+			bool bConverted = UsdToUnreal::ConvertXformable(Prim.GetStage(), Xform, PrimTransform, TimeCode.GetValue(), &bResetXformStack);
+			if (bConverted)
+			{
+				TransformToPropagate = bResetXformStack ? PrimTransform : ParentTransform * PrimTransform;
+			}
+
+			for (const pxr::UsdPrim& Child : Prim.GetChildren())
+			{
+				bSuccess &= ConvertGroomHierarchy(Child, TimeCode, TransformToPropagate, HairDescription);
+				if (!bSuccess)
+				{
+					break;
+				}
+			}
+		}
+
+		return bSuccess;
+	}
+}
+
+#endif // #if USE_USD_SDK
diff --git a/Engine/Plugins/Importers/USDImporter/Source/USDUtilities/Public/USDGroomConversion.h b/Engine/Plugins/Importers/USDImporter/Source/USDUtilities/Public/USDGroomConversion.h
new file mode 100644
--- /dev/null
+++ b/Engine/Plugins/Importers/USDImporter/Source/USDUtilities/Public/USDGroomConversion.h
@@ -0,0 +1,34 @@
+// Copyright Epic Games, Inc. All Rights Reserved.
+
+#pragma once
+
+#if USE_USD_SDK
+
+#include "CoreMinimal.h"
+
+#include "USDIncludesStart.h"
+#include "pxr/pxr.h"
+#include "USDIncludesEnd.h"
+
+PXR_NAMESPACE_OPEN_SCOPE
+class UsdPrim;
+class UsdTimeCode;
+PXR_NAMESPACE_CLOSE_SCOPE
+
+class FHairDescription;
+
+namespace UsdToUnreal
+{
+	/**
+	 * Converts an UsdPrim groom hierarchy and appends the converted data to HairDescription
+	 *
+	 * @param Prim - Prim to convert
+	 * @param TimeCode - The timecode at which to query the prim attributes
+	 * @param ParentTransform - The prim transform to propagate to the children
+	 * @param HairDescription - The hair description where the groom prim data is outputted
+	 * @return true if the conversion was successful; false otherwise
+	 */
+	USDUTILITIES_API bool ConvertGroomHierarchy(const pxr::UsdPrim& Prim, const pxr::UsdTimeCode& TimeCode, const FTransform& ParentTransform, FHairDescription& HairDescription);
+}
+
+#endif // #if USE_USD_SDK
\ No newline at end of file
diff --git a/Engine/Plugins/Importers/USDImporter/Source/USDUtilities/USDUtilities.Build.cs b/Engine/Plugins/Importers/USDImporter/Source/USDUtilities/USDUtilities.Build.cs
--- a/Engine/Plugins/Importers/USDImporter/Source/USDUtilities/USDUtilities.Build.cs
+++ b/Engine/Plugins/Importers/USDImporter/Source/USDUtilities/USDUtilities.Build.cs
@@ -28,6 +28,7 @@ namespace UnrealBuildTool.Rules
 					"Engine",
 					"Foliage",
 					"GeometryCache", // Just so that we can fetch its AssetImportData
+					"HairStrandsCore",
 					"IntelTBB",
 					"Landscape", // So that GetSchemaNameForComponent knows what to do with landscape proxies
 					"LiveLinkComponents", // For converting LiveLinkComponentController properties to USD
diff --git a/Engine/Plugins/Importers/USDImporter/USDImporter.uplugin b/Engine/Plugins/Importers/USDImporter/USDImporter.uplugin
--- a/Engine/Plugins/Importers/USDImporter/USDImporter.uplugin
+++ b/Engine/Plugins/Importers/USDImporter/USDImporter.uplugin
@@ -110,6 +110,10 @@
 		{
 			"Name": "SequencerScripting",
 			"Enabled": true
+		},
+		{
+			"Name": "HairStrands",
+			"Enabled": true
 		}
 	]
 }
diff --git a/Engine/Source/Runtime/Core/Private/Misc/ArchiveMD5.cpp b/Engine/Source/Runtime/Core/Private/Misc/ArchiveMD5.cpp
new file mode 100644
--- /dev/null
+++ b/Engine/Source/Runtime/Core/Private/Misc/ArchiveMD5.cpp
@@ -0,0 +1,11 @@
+// Copyright Epic Games, Inc. All Rights Reserved.
+
+#include "Misc/ArchiveMD5.h"
+
+// Force a virtual function to be defined out-of-line so that the vtable
+// is emitted in this module instead of everywhere FArchiveMD5 is used.
+// Otherwise, it's causing a linking issue when used in a RTTI module
+FString FArchiveMD5::GetArchiveName() const
+{
+	return TEXT("FArchiveMD5");
+}
diff --git a/Engine/Source/Runtime/Core/Public/Misc/ArchiveMD5.h b/Engine/Source/Runtime/Core/Public/Misc/ArchiveMD5.h
--- a/Engine/Source/Runtime/Core/Public/Misc/ArchiveMD5.h
+++ b/Engine/Source/Runtime/Core/Public/Misc/ArchiveMD5.h
@@ -8,7 +8,7 @@
 /**
  * FArchive adapter for FMD5
  */
-class FArchiveMD5 : public FArchive
+class CORE_API FArchiveMD5 : public FArchive
 {
 public:
 	inline FArchiveMD5()
@@ -18,10 +18,7 @@ public:
 		SetIsPersistent(false);
 	}
 
-	virtual FString GetArchiveName() const
-	{
-		return TEXT("FArchiveMD5");
-	}
+	virtual FString GetArchiveName() const;
 
 	void Serialize(void* Data, int64 Num) override
 	{
