USD: Update AnimBP path to the new full object path convention.

diff --git a/Engine/Plugins/Importers/USDImporter/Source/USDStageEditor/Private/Widgets/SUSDIntegrationsPanel.h b/Engine/Plugins/Importers/USDImporter/Source/USDStageEditor/Private/Widgets/SUSDIntegrationsPanel.h
--- a/Engine/Plugins/Importers/USDImporter/Source/USDStageEditor/Private/Widgets/SUSDIntegrationsPanel.h
+++ b/Engine/Plugins/Importers/USDImporter/Source/USDStageEditor/Private/Widgets/SUSDIntegrationsPanel.h
@@ -24,7 +24,7 @@ public:
 
 	// Ideally this would be an FSoftObjectPath, but FPropertyEditorModule doesn't support generating
 	// widgets for those properties yet
-	UPROPERTY(EditAnywhere, Transient, Category=Dummy, meta=(AllowedClasses="AnimBlueprint"))
+	UPROPERTY(EditAnywhere, Transient, Category=Dummy, meta=(AllowedClasses="/Script/Engine.AnimBlueprint" ))
 	UObject* AnimBPProperty;
 };
 
