USD: Add custom schemas for Control Rig connection persistence;

Tweak Live Link schema attributes a bit for clarity.

diff --git a/Engine/Plugins/Importers/USDImporter/Resources/UsdResources/Linux/plugins/unreal/resources/plugInfo.json b/Engine/Plugins/Importers/USDImporter/Resources/UsdResources/Linux/plugins/unreal/resources/plugInfo.json
--- a/Engine/Plugins/Importers/USDImporter/Resources/UsdResources/Linux/plugins/unreal/resources/plugInfo.json
+++ b/Engine/Plugins/Importers/USDImporter/Resources/UsdResources/Linux/plugins/unreal/resources/plugInfo.json
@@ -6,22 +6,31 @@
         {
             "Info": {
                 "Types": {
+                    "unrealControlRigAPI": {
+                        "alias": {
+                            "UsdSchemaBase": "ControlRigAPI"
+                        },
+                        "autoGenerated": true,
+                        "bases": [
+                            "UsdAPISchemaBase"
+                        ],
+                        "schemaKind": "singleApplyAPI"
+                    },
                     "unrealLiveLinkAPI": {
                         "alias": {
                             "UsdSchemaBase": "LiveLinkAPI"
-                        }, 
-                        "autoGenerated": true, 
+                        },
+                        "autoGenerated": true,
                         "bases": [
                             "UsdAPISchemaBase"
-                        ], 
+                        ],
                         "schemaKind": "singleApplyAPI"
                     }
                 }
-            }, 
-            "LibraryPath": "@PLUG_INFO_LIBRARY_PATH@", 
-            "Name": "unreal", 
-            "ResourcePath": "@PLUG_INFO_RESOURCE_PATH@", 
-            "Root": "@PLUG_INFO_ROOT@", 
+            },
+            "Name": "unreal",
+            "ResourcePath": "resources",
+            "Root": "..",
             "Type": "resource"
         }
     ]
diff --git a/Engine/Plugins/Importers/USDImporter/Resources/UsdResources/Mac/plugins/unreal/resources/plugInfo.json b/Engine/Plugins/Importers/USDImporter/Resources/UsdResources/Mac/plugins/unreal/resources/plugInfo.json
--- a/Engine/Plugins/Importers/USDImporter/Resources/UsdResources/Mac/plugins/unreal/resources/plugInfo.json
+++ b/Engine/Plugins/Importers/USDImporter/Resources/UsdResources/Mac/plugins/unreal/resources/plugInfo.json
@@ -6,22 +6,31 @@
         {
             "Info": {
                 "Types": {
+                    "unrealControlRigAPI": {
+                        "alias": {
+                            "UsdSchemaBase": "ControlRigAPI"
+                        },
+                        "autoGenerated": true,
+                        "bases": [
+                            "UsdAPISchemaBase"
+                        ],
+                        "schemaKind": "singleApplyAPI"
+                    },
                     "unrealLiveLinkAPI": {
                         "alias": {
                             "UsdSchemaBase": "LiveLinkAPI"
-                        }, 
-                        "autoGenerated": true, 
+                        },
+                        "autoGenerated": true,
                         "bases": [
                             "UsdAPISchemaBase"
-                        ], 
+                        ],
                         "schemaKind": "singleApplyAPI"
                     }
                 }
-            }, 
-            "LibraryPath": "@PLUG_INFO_LIBRARY_PATH@", 
-            "Name": "unreal", 
-            "ResourcePath": "@PLUG_INFO_RESOURCE_PATH@", 
-            "Root": "@PLUG_INFO_ROOT@", 
+            },
+            "Name": "unreal",
+            "ResourcePath": "resources",
+            "Root": "..",
             "Type": "resource"
         }
     ]
diff --git a/Engine/Plugins/Importers/USDImporter/Resources/UsdResources/Win64/plugins/unreal/resources/plugInfo.json b/Engine/Plugins/Importers/USDImporter/Resources/UsdResources/Win64/plugins/unreal/resources/plugInfo.json
--- a/Engine/Plugins/Importers/USDImporter/Resources/UsdResources/Win64/plugins/unreal/resources/plugInfo.json
+++ b/Engine/Plugins/Importers/USDImporter/Resources/UsdResources/Win64/plugins/unreal/resources/plugInfo.json
@@ -6,22 +6,31 @@
         {
             "Info": {
                 "Types": {
+                    "unrealControlRigAPI": {
+                        "alias": {
+                            "UsdSchemaBase": "ControlRigAPI"
+                        },
+                        "autoGenerated": true,
+                        "bases": [
+                            "UsdAPISchemaBase"
+                        ],
+                        "schemaKind": "singleApplyAPI"
+                    },
                     "unrealLiveLinkAPI": {
                         "alias": {
                             "UsdSchemaBase": "LiveLinkAPI"
-                        }, 
-                        "autoGenerated": true, 
+                        },
+                        "autoGenerated": true,
                         "bases": [
                             "UsdAPISchemaBase"
-                        ], 
+                        ],
                         "schemaKind": "singleApplyAPI"
                     }
                 }
-            }, 
-            "LibraryPath": "@PLUG_INFO_LIBRARY_PATH@", 
-            "Name": "unreal", 
-            "ResourcePath": "@PLUG_INFO_RESOURCE_PATH@", 
-            "Root": "@PLUG_INFO_ROOT@", 
+            },
+            "Name": "unreal",
+            "ResourcePath": "resources",
+            "Root": "..",
             "Type": "resource"
         }
     ]
