USD: Revert cl 20351709 because root resync notices are how USD signals layer add/remove/mutes.

diff --git a/Engine/Plugins/Importers/USDImporter/Source/UnrealUSDWrapper/Private/USDListener.cpp b/Engine/Plugins/Importers/USDImporter/Source/UnrealUSDWrapper/Private/USDListener.cpp
--- a/Engine/Plugins/Importers/USDImporter/Source/UnrealUSDWrapper/Private/USDListener.cpp
+++ b/Engine/Plugins/Importers/USDImporter/Source/UnrealUSDWrapper/Private/USDListener.cpp
@@ -37,12 +37,6 @@ namespace UsdToUnreal
 
 		for ( UsdNotice::ObjectsChanged::PathRange::const_iterator It = PathRange.begin(); It != PathRange.end(); ++It )
 		{
-			const std::vector<const SdfChangeList::Entry*>& Changes = It.base()->second;
-			if ( Changes.empty() )
-			{
-				continue;
-			}
-
 			const FString PrimPath = ANSI_TO_TCHAR( It->GetAbsoluteRootOrPrimPath().GetAsString().c_str() );
 
 			// Something like "/Root/Prim.some_field", or "/"
@@ -50,6 +44,9 @@ namespace UsdToUnreal
 
 			TArray<UsdUtils::FObjectChangeNotice>& ConvertedChanges = OutChanges.FindOrAdd( PrimPath );
 
+			// Changes may be empty, but we should still pass along this overall notice because sending a root
+			// resync notice with no actual change item inside is how USD signals that a layer has been added/removed/resynced
+			const std::vector<const SdfChangeList::Entry*>& Changes = It.base()->second;
 			for ( const SdfChangeList::Entry* Entry : Changes )
 			{
 				UsdUtils::FObjectChangeNotice& ConvertedEntry = ConvertedChanges.Emplace_GetRef();
