USDImporter: consider render context in UsdUtils::HashShadeMaterial()

This addresses an issue where UsdUtils::HashShadeMaterial() would only ever
produce a hash for the universal render context surface output of a
UsdShadeMaterial (or no hash if there was no such output). This could be
problematic if you are translating materials in any non-universal render
context (e.g. "mdl" or "mtlx") since we would produce an incorrect or empty
hash for the surface shader being translated, which could cause distinct
shaders in USD to all map incorrectly to the same material asset in Unreal.

diff --git a/Engine/Plugins/Importers/USDImporter/Source/USDSchemas/Private/USDShadeMaterialTranslator.cpp b/Engine/Plugins/Importers/USDImporter/Source/USDSchemas/Private/USDShadeMaterialTranslator.cpp
--- a/Engine/Plugins/Importers/USDImporter/Source/USDSchemas/Private/USDShadeMaterialTranslator.cpp
+++ b/Engine/Plugins/Importers/USDImporter/Source/USDSchemas/Private/USDShadeMaterialTranslator.cpp
@@ -18,7 +18,9 @@
 #if USE_USD_SDK
 
 #include "USDIncludesStart.h"
+	#include "pxr/base/tf/token.h"
 	#include "pxr/usd/usdShade/material.h"
+	#include "pxr/usd/usdShade/tokens.h"
 #include "USDIncludesEnd.h"
 
 
@@ -33,7 +35,11 @@ void FUsdShadeMaterialTranslator::CreateAssets()
 
 	TRACE_CPUPROFILER_EVENT_SCOPE( FUsdShadeMaterialTranslator::CreateAssets );
 
-	FString MaterialHashString = UsdUtils::HashShadeMaterial( ShadeMaterial ).ToString();
+	const pxr::TfToken RenderContextToken =
+		Context->RenderContext.IsNone() ?
+			pxr::UsdShadeTokens->universalRenderContext :
+			UnrealToUsd::ConvertToken( *Context->RenderContext.ToString() ).Get();
+	FString MaterialHashString = UsdUtils::HashShadeMaterial( ShadeMaterial, RenderContextToken ).ToString();
 
 	UMaterialInterface* ConvertedMaterial = nullptr;
 
diff --git a/Engine/Plugins/Importers/USDImporter/Source/USDUtilities/Private/USDShadeConversion.cpp b/Engine/Plugins/Importers/USDImporter/Source/USDUtilities/Private/USDShadeConversion.cpp
--- a/Engine/Plugins/Importers/USDImporter/Source/USDUtilities/Private/USDShadeConversion.cpp
+++ b/Engine/Plugins/Importers/USDImporter/Source/USDUtilities/Private/USDShadeConversion.cpp
@@ -44,6 +44,7 @@
 #endif // WITH_EDITOR
 
 #include "USDIncludesStart.h"
+	#include "pxr/base/tf/token.h"
 	#include "pxr/usd/ar/ar.h"
 	#include "pxr/usd/ar/asset.h"
 	#include "pxr/usd/ar/resolvedPath.h"
@@ -2528,11 +2529,11 @@ bool UsdUtils::IsMaterialUsingUDIMs( const pxr::UsdShadeMaterial& UsdShadeMateri
 	return false;
 }
 
-FSHAHash UsdUtils::HashShadeMaterial( const pxr::UsdShadeMaterial& UsdShadeMaterial )
+FSHAHash UsdUtils::HashShadeMaterial( const pxr::UsdShadeMaterial& UsdShadeMaterial, const pxr::TfToken& RenderContext )
 {
 	FScopedUsdAllocs UsdAllocs;
 
-	pxr::UsdShadeShader SurfaceShader = UsdShadeMaterial.ComputeSurfaceSource();
+	pxr::UsdShadeShader SurfaceShader = UsdShadeMaterial.ComputeSurfaceSource( { RenderContext } );
 	if ( !SurfaceShader )
 	{
 		return {};
diff --git a/Engine/Plugins/Importers/USDImporter/Source/USDUtilities/Public/USDShadeConversion.h b/Engine/Plugins/Importers/USDImporter/Source/USDUtilities/Public/USDShadeConversion.h
--- a/Engine/Plugins/Importers/USDImporter/Source/USDUtilities/Public/USDShadeConversion.h
+++ b/Engine/Plugins/Importers/USDImporter/Source/USDUtilities/Public/USDShadeConversion.h
@@ -12,7 +12,9 @@
 #include "UsdWrappers/SdfLayer.h"
 
 #include "USDIncludesStart.h"
+	#include "pxr/base/tf/token.h"
 	#include "pxr/usd/usd/timeCode.h"
+	#include "pxr/usd/usdShade/tokens.h"
 #include "USDIncludesEnd.h"
 
 PXR_NAMESPACE_OPEN_SCOPE
@@ -130,7 +132,7 @@ namespace UsdUtils
 	*/
 	USDUTILITIES_API bool IsMaterialUsingUDIMs( const pxr::UsdShadeMaterial& UsdShadeMaterial );
 
-	USDUTILITIES_API FSHAHash HashShadeMaterial( const pxr::UsdShadeMaterial& UsdShadeMaterial );
+	USDUTILITIES_API FSHAHash HashShadeMaterial( const pxr::UsdShadeMaterial& UsdShadeMaterial, const pxr::TfToken& RenderContext = pxr::UsdShadeTokens->universalRenderContext );
 
 	/** Returns the resolved path from a pxr::SdfAssetPath attribute. For UDIMs path, returns the path to the 1001 tile. */
 	USDUTILITIES_API FString GetResolvedTexturePath( const pxr::UsdAttribute& TextureAssetPathAttr );
