USD: Utility functions and wrappers made for Control Rig support.

diff --git a/Engine/Plugins/Importers/USDImporter/Source/USDUtilities/Private/USDConversionUtils.cpp b/Engine/Plugins/Importers/USDImporter/Source/USDUtilities/Private/USDConversionUtils.cpp
--- a/Engine/Plugins/Importers/USDImporter/Source/USDUtilities/Private/USDConversionUtils.cpp
+++ b/Engine/Plugins/Importers/USDImporter/Source/USDUtilities/Private/USDConversionUtils.cpp
@@ -1383,6 +1383,26 @@ FString UsdUtils::GetUniqueName( FString Name, const TSet<FString>& UsedNames )
 	return Result;
 }
 
+#if USE_USD_SDK
+FString UsdUtils::GetValidChildName( FString InName, const pxr::UsdPrim& ParentPrim )
+{
+	if ( !ParentPrim )
+	{
+		return {};
+	}
+
+	FScopedUsdAllocs Allocs;
+
+	TSet<FString> UsedNames;
+	for ( const pxr::UsdPrim& Child : ParentPrim.GetChildren() )
+	{
+		UsedNames.Add( UsdToUnreal::ConvertToken( Child.GetName() ) );
+	}
+
+	return GetUniqueName( SanitizeUsdIdentifier( *InName ), UsedNames );
+}
+#endif // USE_USD_SDK
+
 FString UsdUtils::SanitizeUsdIdentifier( const TCHAR* InIdentifier )
 {
 #if USE_USD_SDK
diff --git a/Engine/Plugins/Importers/USDImporter/Source/USDUtilities/Private/USDTypesConversion.cpp b/Engine/Plugins/Importers/USDImporter/Source/USDUtilities/Private/USDTypesConversion.cpp
--- a/Engine/Plugins/Importers/USDImporter/Source/USDUtilities/Private/USDTypesConversion.cpp
+++ b/Engine/Plugins/Importers/USDImporter/Source/USDUtilities/Private/USDTypesConversion.cpp
@@ -61,6 +61,19 @@ namespace UsdUtils
 
 		return FTransform( Rotation, Translation, Scale );
 	}
+
+	FTransform ConvertTransformToUsdSpace( const FUsdStageInfo& StageInfo, const FTransform& TransformInUESpace )
+	{
+		FTransform TransformInUsdSpace = UsdUtils::ConvertAxes( StageInfo.UpAxis == EUsdUpAxis::ZAxis, TransformInUESpace );
+
+		const float UEMetersPerUnit = 0.01f;
+		if ( !FMath::IsNearlyEqual( StageInfo.MetersPerUnit, UEMetersPerUnit ) )
+		{
+			TransformInUsdSpace.ScaleTranslation( UEMetersPerUnit / StageInfo.MetersPerUnit );
+		}
+
+		return TransformInUsdSpace;
+	}
 }
 
 namespace UsdToUnreal
@@ -280,13 +293,7 @@ namespace UnrealToUsd
 
 	pxr::GfMatrix4d ConvertTransform( const FUsdStageInfo& StageInfo, const FTransform& Transform )
 	{
-		FTransform TransformInUsdSpace = UsdUtils::ConvertAxes( StageInfo.UpAxis == EUsdUpAxis::ZAxis, Transform );
-
-		const float UEMetersPerUnit = 0.01f;
-		if ( !FMath::IsNearlyEqual( StageInfo.MetersPerUnit, UEMetersPerUnit ) )
-		{
-			TransformInUsdSpace.ScaleTranslation( UEMetersPerUnit / StageInfo.MetersPerUnit );
-		}
+		FTransform TransformInUsdSpace = UsdUtils::ConvertTransformToUsdSpace( StageInfo, Transform );
 
 		return ConvertMatrix( TransformInUsdSpace.ToMatrixWithScale() );
 	}
diff --git a/Engine/Plugins/Importers/USDImporter/Source/USDUtilities/Private/USDValueConversion.cpp b/Engine/Plugins/Importers/USDImporter/Source/USDUtilities/Private/USDValueConversion.cpp
--- a/Engine/Plugins/Importers/USDImporter/Source/USDUtilities/Private/USDValueConversion.cpp
+++ b/Engine/Plugins/Importers/USDImporter/Source/USDUtilities/Private/USDValueConversion.cpp
@@ -853,6 +853,53 @@ namespace UsdUtils
 #endif // USE_USD_SDK
 	}
 
+	template<typename T>
+	TOptional<T> GetUnderlyingValue( const UE::FVtValue& InValue )
+	{
+#if USE_USD_SDK
+		const pxr::VtValue& UsdValue = InValue.GetUsdValue();
+		if ( UsdValue.IsHolding<T>() )
+		{
+			return { UsdValue.UncheckedGet<T>() };
+		}
+#endif // USE_USD_SDK
+
+		return {};
+	}
+    template USDUTILITIES_API TOptional<bool> GetUnderlyingValue<bool>( const UE::FVtValue& InValue );
+	template USDUTILITIES_API TOptional<float> GetUnderlyingValue<float>( const UE::FVtValue& InValue );
+	template USDUTILITIES_API TOptional<uint8_t> GetUnderlyingValue<uint8_t>( const UE::FVtValue& InValue );
+	template USDUTILITIES_API TOptional<int32_t> GetUnderlyingValue<int32_t>( const UE::FVtValue& InValue );
+	template USDUTILITIES_API TOptional<uint32_t> GetUnderlyingValue<uint32_t>( const UE::FVtValue& InValue );
+	template USDUTILITIES_API TOptional<int64_t> GetUnderlyingValue<int64_t>( const UE::FVtValue& InValue );
+	template USDUTILITIES_API TOptional<uint64_t> GetUnderlyingValue<uint64_t>( const UE::FVtValue& InValue );
+	template USDUTILITIES_API TOptional<std::string> GetUnderlyingValue<std::string>( const UE::FVtValue& InValue );
+
+	template<typename T>
+	bool SetUnderlyingValue( UE::FVtValue& InValue, const T& UnderlyingValue )
+	{
+#if USE_USD_SDK
+		FScopedUsdAllocs Allocs;
+
+		pxr::VtValue& UsdValue = InValue.GetUsdValue();
+		if ( UsdValue.IsHolding<T>() || UsdValue.IsEmpty() )
+		{
+			UsdValue = UnderlyingValue;
+			return true;
+		}
+#endif // USE_USD_SDK
+
+		return false;
+	}
+    template bool USDUTILITIES_API SetUnderlyingValue( UE::FVtValue& InValue, const bool& UnderlyingValue );
+	template bool USDUTILITIES_API SetUnderlyingValue( UE::FVtValue& InValue, const float& UnderlyingValue );
+	template bool USDUTILITIES_API SetUnderlyingValue<uint8_t>( UE::FVtValue& InValue, const uint8_t& UnderlyingValue );
+	template bool USDUTILITIES_API SetUnderlyingValue<int32_t>( UE::FVtValue& InValue, const int32_t& UnderlyingValue );
+	template bool USDUTILITIES_API SetUnderlyingValue<uint32_t>( UE::FVtValue& InValue, const uint32_t& UnderlyingValue );
+	template bool USDUTILITIES_API SetUnderlyingValue<int64_t>( UE::FVtValue& InValue, const int64_t& UnderlyingValue );
+	template bool USDUTILITIES_API SetUnderlyingValue<uint64_t>( UE::FVtValue& InValue, const uint64_t& UnderlyingValue );
+	template bool USDUTILITIES_API SetUnderlyingValue( UE::FVtValue& InValue, const std::string& UnderlyingValue );
+
 	FString GetImpliedTypeName( const UE::FVtValue& Value )
 	{
 #if USE_USD_SDK
diff --git a/Engine/Plugins/Importers/USDImporter/Source/USDUtilities/Public/USDConversionUtils.h b/Engine/Plugins/Importers/USDImporter/Source/USDUtilities/Public/USDConversionUtils.h
--- a/Engine/Plugins/Importers/USDImporter/Source/USDUtilities/Public/USDConversionUtils.h
+++ b/Engine/Plugins/Importers/USDImporter/Source/USDUtilities/Public/USDConversionUtils.h
@@ -213,6 +213,14 @@ namespace UsdUtils
 	 */
 	USDUTILITIES_API FString GetUniqueName( FString Name, const TSet<FString>& UsedNames );
 
+#if USE_USD_SDK
+	/**
+	 * Returns a sanitized, valid version of 'InName' that can be used as the name for a child prim of 'ParentPrim'.
+	 * Returns the empty string in case of an error.
+	 */
+	USDUTILITIES_API FString GetValidChildName( FString InName, const pxr::UsdPrim& ParentPrim );
+#endif // #if USE_USD_SDK
+
 	/**
 	 * Returns a modified version of InIdentifier that can be used as a USD prim or property name.
 	 * This means only allowing letters, numbers and the underscore character. All others are replaced with underscores.
diff --git a/Engine/Plugins/Importers/USDImporter/Source/USDUtilities/Public/USDTypesConversion.h b/Engine/Plugins/Importers/USDImporter/Source/USDUtilities/Public/USDTypesConversion.h
--- a/Engine/Plugins/Importers/USDImporter/Source/USDUtilities/Public/USDTypesConversion.h
+++ b/Engine/Plugins/Importers/USDImporter/Source/USDUtilities/Public/USDTypesConversion.h
@@ -97,6 +97,8 @@ namespace UnrealToUsd
 
 namespace UsdUtils
 {
+	USDUTILITIES_API FTransform ConvertTransformToUsdSpace( const FUsdStageInfo& StageInfo, const FTransform& TransformInUESpace );
+
 	FTransform ConvertAxes( const bool bZUp, const FTransform Transform );
 }
 #endif // #if USE_USD_SDK
diff --git a/Engine/Plugins/Importers/USDImporter/Source/USDUtilities/Public/USDValueConversion.h b/Engine/Plugins/Importers/USDImporter/Source/USDUtilities/Public/USDValueConversion.h
--- a/Engine/Plugins/Importers/USDImporter/Source/USDUtilities/Public/USDValueConversion.h
+++ b/Engine/Plugins/Importers/USDImporter/Source/USDUtilities/Public/USDValueConversion.h
@@ -2,6 +2,8 @@
 
 #pragma once
 
+#include <string>
+
 #include "CoreMinimal.h"
 #include "Misc/TVariant.h"
 
@@ -97,6 +99,24 @@ namespace UsdUtils
 	/** Uses USD to stringify the underlying pxr::VtValue */
 	USDUTILITIES_API FString Stringify( const UE::FVtValue& Value );
 
+	/**
+	 * Uses USD to quickly fetch the underlying value of the wrapped pxr::VtValue for fundamental data types.
+	 * For complex vector/array types ConvertValue must be used instead.
+	 * Template implementation must be hidden on cpp as we can't expose USD implementation on the header files.
+	 */
+	template<typename T>
+	USDUTILITIES_API TOptional<T> GetUnderlyingValue( const UE::FVtValue& InValue );
+
+	/**
+	 * Uses USD to quickly set the underlying value of the wrapped pxr::VtValue for fundamental data types.
+	 * This will not allow changing the underlying type of InValue, i.e. you can only set a float if InValue
+	 * already contains a float or if its otherwise empty.
+	 * For complex vector/array types ConvertValue must be used instead.
+	 * Template implementation must be hidden on cpp as we can't expose USD implementation on the header files.
+	 */
+	template<typename T>
+	USDUTILITIES_API bool SetUnderlyingValue( UE::FVtValue& InValue, const T& UnderlyingValue );
+
 	/**
 	 * Returns the name of the SdfValueTypeName object for the type of Value
 	 *
diff --git a/Engine/Plugins/Importers/USDImporter/Source/UnrealUSDWrapper/Private/UsdWrappers/UsdEditContext.cpp b/Engine/Plugins/Importers/USDImporter/Source/UnrealUSDWrapper/Private/UsdWrappers/UsdEditContext.cpp
new file mode 100644
--- /dev/null
+++ b/Engine/Plugins/Importers/USDImporter/Source/UnrealUSDWrapper/Private/UsdWrappers/UsdEditContext.cpp
@@ -0,0 +1,59 @@
+// Copyright Epic Games, Inc. All Rights Reserved.
+
+#include "UsdWrappers/UsdEditContext.h"
+
+#include "UsdWrappers/SdfLayer.h"
+#include "UsdWrappers/UsdStage.h"
+
+#if USE_USD_SDK
+
+#include "USDIncludesStart.h"
+	#include "pxr/usd/sdf/layer.h"
+	#include "pxr/usd/usd/editContext.h"
+#include "USDIncludesEnd.h"
+
+#endif // #if USE_USD_SDK
+
+namespace UE
+{
+	namespace Internal
+	{
+		class FUsdEditContextImpl
+		{
+#if USE_USD_SDK
+		public:
+			explicit FUsdEditContextImpl( const pxr::UsdStageWeakPtr& StageWeakPtr, const pxr::UsdEditTarget& EditTarget = pxr::UsdEditTarget() )
+				: UsdEditContext( StageWeakPtr, EditTarget )
+			{
+			}
+
+		private:
+			pxr::UsdEditContext UsdEditContext;
+#endif // #if USE_USD_SDK
+		};
+	}
+
+	FUsdEditContext::FUsdEditContext( const FUsdStageWeak& Stage )
+	{
+#if USE_USD_SDK
+		FScopedUsdAllocs Allocs;
+		Impl = MakeUnique< Internal::FUsdEditContextImpl >(Stage);
+#endif // #if USE_USD_SDK
+	}
+
+	FUsdEditContext::FUsdEditContext( const FUsdStageWeak& Stage, const FSdfLayer& EditTarget )
+	{
+#if USE_USD_SDK
+		FScopedUsdAllocs Allocs;
+		Impl = MakeUnique< Internal::FUsdEditContextImpl >( Stage, pxr::SdfLayerWeakPtr{ EditTarget } );
+#endif // #if USE_USD_SDK
+	}
+
+	FUsdEditContext::~FUsdEditContext()
+	{
+#if USE_USD_SDK
+		FScopedUsdAllocs Allocs;
+		Impl.Reset();
+#endif // #if USE_USD_SDK
+	}
+}
\ No newline at end of file
diff --git a/Engine/Plugins/Importers/USDImporter/Source/UnrealUSDWrapper/Public/UsdWrappers/UsdEditContext.h b/Engine/Plugins/Importers/USDImporter/Source/UnrealUSDWrapper/Public/UsdWrappers/UsdEditContext.h
new file mode 100644
--- /dev/null
+++ b/Engine/Plugins/Importers/USDImporter/Source/UnrealUSDWrapper/Public/UsdWrappers/UsdEditContext.h
@@ -0,0 +1,32 @@
+// Copyright Epic Games, Inc. All Rights Reserved.
+
+#pragma once
+
+#include "UsdWrappers/ForwardDeclarations.h"
+
+#include "Templates/UniquePtr.h"
+
+namespace UE
+{
+	namespace Internal
+	{
+		class FUsdEditContextImpl;
+	}
+
+	/**
+	 * Minimal pxr::UsdEditContext wrapper for Unreal that can be used from no-rtti modules.
+	 */
+	class UNREALUSDWRAPPER_API FUsdEditContext final
+	{
+		FUsdEditContext( const FUsdEditContext& Other ) = delete;
+		FUsdEditContext& operator=( const FUsdEditContext& Other ) = delete;
+
+	public:
+		explicit FUsdEditContext( const FUsdStageWeak& Stage );
+		FUsdEditContext( const FUsdStageWeak& Stage, const FSdfLayer& EditTarget );
+		~FUsdEditContext();
+
+	private:
+		TUniquePtr< Internal::FUsdEditContextImpl > Impl;
+	};
+}
