USD: Implement support for applying/removing the Control Rig schema.

Automatically write back to USD whenever the Control Rig tracks are updated.

diff --git a/Engine/Plugins/Importers/USDImporter/Source/USDSchemas/Private/USDSkelRootTranslator.cpp b/Engine/Plugins/Importers/USDImporter/Source/USDSchemas/Private/USDSkelRootTranslator.cpp
--- a/Engine/Plugins/Importers/USDImporter/Source/USDSchemas/Private/USDSkelRootTranslator.cpp
+++ b/Engine/Plugins/Importers/USDImporter/Source/USDSchemas/Private/USDSkelRootTranslator.cpp
@@ -16,6 +16,7 @@
 #include "USDMemory.h"
 #include "USDSkeletalDataConversion.h"
 #include "USDTypesConversion.h"
+#include "USDValueConversion.h"
 
 #include "UsdWrappers/SdfLayer.h"
 #include "UsdWrappers/SdfPath.h"
@@ -985,7 +986,7 @@ namespace UsdSkelRootTranslatorImpl
 				}
 			}
 
-			if ( pxr::UsdAttribute Attr = Prim.GetAttribute( UnrealIdentifiers::UnrealLiveLinkAnimBlueprintPath ) )
+			if ( pxr::UsdAttribute Attr = Prim.GetAttribute( UnrealIdentifiers::UnrealAnimBlueprintPath ) )
 			{
 				std::string PathString;
 				if ( Attr.Get( &PathString ) )
@@ -1227,6 +1228,8 @@ void FUsdSkelRootTranslator::UpdateComponents( USceneComponent* SceneComponent )
 	UE::FUsdPrim Prim = GetPrim();
 
 	const bool bPrimHasLiveLinkSchema = UsdUtils::PrimHasLiveLinkSchema( Prim );
+	const bool bPrimHasControlRigSchema = UsdUtils::PrimHasControlRigSchema( Prim );
+
 	bool bPrimHasLiveLinkEnabled = bPrimHasLiveLinkSchema;
 	if ( bPrimHasLiveLinkSchema )
 	{
@@ -1242,18 +1245,18 @@ void FUsdSkelRootTranslator::UpdateComponents( USceneComponent* SceneComponent )
 		}
 	}
 
-	UE::FUsdPrim SkelAnimPrim;
-	if ( SkeletalMeshComponent->AnimationData.AnimToPlay == nullptr )
+	UE::FUsdPrim SkelAnimPrim = UsdUtils::FindFirstAnimationSource( Prim );
+	if ( SkelAnimPrim )
 	{
-		SkelAnimPrim = UsdUtils::FindFirstAnimationSource( Prim );
-		if ( SkelAnimPrim )
+		UAnimSequence* TargetAnimSequence = Cast< UAnimSequence >( Context->AssetCache->GetAssetForPrim( SkelAnimPrim.GetPrimPath().GetString() ) );
+		if ( TargetAnimSequence != SkeletalMeshComponent->AnimationData.AnimToPlay )
 		{
-			if ( UAnimSequence* TargetAnimSequence = Cast< UAnimSequence >( Context->AssetCache->GetAssetForPrim( SkelAnimPrim.GetPrimPath().GetString() ) ) )
+			SkeletalMeshComponent->AnimationData.AnimToPlay = TargetAnimSequence;
+			SkeletalMeshComponent->AnimationData.bSavedLooping = false;
+			SkeletalMeshComponent->AnimationData.bSavedPlaying = false;
+
+			if ( SkeletalMeshComponent->GetAnimationMode() == EAnimationMode::AnimationSingleNode )
 			{
-				SkeletalMeshComponent->AnimationData.AnimToPlay = TargetAnimSequence;
-				SkeletalMeshComponent->AnimationData.bSavedLooping = false;
-				SkeletalMeshComponent->AnimationData.bSavedPlaying = false;
-				SkeletalMeshComponent->SetUpdateAnimationInEditor( true );
 				SkeletalMeshComponent->SetAnimation( TargetAnimSequence );
 			}
 		}
@@ -1300,7 +1303,9 @@ void FUsdSkelRootTranslator::UpdateComponents( USceneComponent* SceneComponent )
 	}
 
 	// Update the animation state
-	if ( SkeletalMeshComponent->SkeletalMesh && !bPrimHasLiveLinkEnabled )
+	// Don't try animating ourselves if the sequencer is animating as it will just overwrite the animation state on next
+	// tick anyway, and all this would do is lead to flickering and other issues
+	if ( !Context->bSequencerIsAnimating && SkeletalMeshComponent->SkeletalMesh && !bPrimHasLiveLinkEnabled )
 	{
 		if ( UAnimSequence* AnimSequence = Cast<UAnimSequence>( SkeletalMeshComponent->AnimationData.AnimToPlay.Get() ) )
 		{
diff --git a/Engine/Plugins/Importers/USDImporter/Source/USDSchemas/Public/USDSchemaTranslator.h b/Engine/Plugins/Importers/USDImporter/Source/USDSchemas/Public/USDSchemaTranslator.h
--- a/Engine/Plugins/Importers/USDImporter/Source/USDSchemas/Public/USDSchemaTranslator.h
+++ b/Engine/Plugins/Importers/USDImporter/Source/USDSchemas/Public/USDSchemaTranslator.h
@@ -198,6 +198,17 @@ struct USDSCHEMAS_API FUsdSchemaTranslationContext : public TSharedFromThis< FUs
 	/** Groom group interpolation settings */
 	TArray<FHairGroupsInterpolation> GroomInterpolationSettings;
 
+	/**
+	 * True if the Sequencer is currently opened and animating the stage level sequence.
+	 * Its relevant to know this because some translator ::UpdateComponents overloads may try to animate their components
+	 * by themselves, which could be wasteful and glitchy in case the sequencer is opened: It will likely also have an
+	 * animation track for that component and on next editor tick would override the animation with what is sampled from
+	 * the track.
+	 * In the future we'll likely get rid of the "Time" track on the generated LevelSequence, at which point we can
+	 * remove this
+	 */
+	bool bSequencerIsAnimating = false;
+
 	bool IsValid() const
 	{
 		return Level != nullptr;
diff --git a/Engine/Plugins/Importers/USDImporter/Source/USDStage/Private/USDLevelSequenceHelper.cpp b/Engine/Plugins/Importers/USDImporter/Source/USDStage/Private/USDLevelSequenceHelper.cpp
--- a/Engine/Plugins/Importers/USDImporter/Source/USDStage/Private/USDLevelSequenceHelper.cpp
+++ b/Engine/Plugins/Importers/USDImporter/Source/USDStage/Private/USDLevelSequenceHelper.cpp
@@ -6,6 +6,7 @@
 #include "USDAssetImportData.h"
 #include "USDAttributeUtils.h"
 #include "USDConversionUtils.h"
+#include "USDIntegrationUtils.h"
 #include "USDLayerUtils.h"
 #include "USDListener.h"
 #include "USDLog.h"
@@ -18,6 +19,7 @@
 #include "UsdWrappers/SdfChangeBlock.h"
 #include "UsdWrappers/SdfLayer.h"
 #include "UsdWrappers/UsdAttribute.h"
+#include "UsdWrappers/UsdEditContext.h"
 #include "UsdWrappers/UsdGeomXformable.h"
 #include "UsdWrappers/UsdPrim.h"
 #include "UsdWrappers/UsdStage.h"
@@ -36,13 +38,19 @@
 #include "Components/SceneComponent.h"
 #include "Components/SkeletalMeshComponent.h"
 #include "Components/SpotLightComponent.h"
+#include "ControlRigObjectBinding.h"
 #include "CoreMinimal.h"
 #include "LevelSequence.h"
+#include "LevelSequenceActor.h"
+#include "LevelSequencePlayer.h"
 #include "MovieScene.h"
 #include "MovieSceneTimeHelpers.h"
 #include "MovieSceneTrack.h"
+#include "Rigs/FKControlRig.h"
 #include "Sections/MovieSceneFloatSection.h"
 #include "Sections/MovieSceneSubSection.h"
+#include "Sequencer/MovieSceneControlRigParameterSection.h"
+#include "Sequencer/MovieSceneControlRigParameterTrack.h"
 #include "Templates/SharedPointer.h"
 #include "Tracks/MovieScene3DTransformTrack.h"
 #include "Tracks/MovieSceneBoolTrack.h"
@@ -55,10 +63,15 @@
 #include "UObject/UObjectGlobals.h"
 
 #if WITH_EDITOR
+#include "ControlRigBlueprint.h"
 #include "Editor.h"
 #include "Editor/TransBuffer.h"
+#include "EditorModeManager.h"
+#include "EditorModes.h"
+#include "Exporters/AnimSeqExportOption.h"
 #include "ILevelSequenceEditorToolkit.h"
 #include "ISequencer.h"
+#include "MovieSceneToolHelpers.h"
 #include "Subsystems/AssetEditorSubsystem.h"
 #endif // WITH_EDITOR
 
@@ -182,6 +195,238 @@
 
 		return nullptr;
 	}
+
+	void MuteTrack( UMovieSceneTrack* Track, UMovieScene* MovieScene, const FString& ComponentBindingString, const FString& TrackName, bool bMute )
+	{
+		if ( !Track || !MovieScene )
+		{
+			return;
+		}
+
+		if ( Track->IsEvalDisabled() == bMute )
+		{
+			return;
+		}
+
+#if WITH_EDITOR
+		// We need to update the MovieScene too, because if MuteNodes disagrees with Track->IsEvalDisabled() the sequencer
+		// will chose in favor of MuteNodes
+		MovieScene->Modify();
+
+		const FString MuteNode = FString::Printf( TEXT( "%s.%s" ), *ComponentBindingString, *TrackName );
+		if ( bMute )
+		{
+			MovieScene->GetMuteNodes().AddUnique( MuteNode );
+		}
+		else
+		{
+			MovieScene->GetMuteNodes().Remove( MuteNode );
+		}
+#endif // WITH_EDITOR
+
+		Track->Modify();
+		Track->SetEvalDisabled( bMute );
+	}
+
+	TSharedPtr< ISequencer > GetOpenedSequencerForLevelSequence( ULevelSequence* LevelSequence )
+	{
+		const bool bFocusIfOpen = false;
+		IAssetEditorInstance* AssetEditor = GEditor->GetEditorSubsystem<UAssetEditorSubsystem>()->FindEditorForAsset( LevelSequence, bFocusIfOpen );
+		ILevelSequenceEditorToolkit* LevelSequenceEditor = static_cast< ILevelSequenceEditorToolkit* >( AssetEditor );
+		return LevelSequenceEditor ? LevelSequenceEditor->GetSequencer() : nullptr;
+	}
+
+#if WITH_EDITOR
+	// Rough copy of UControlRigSequencerEditorLibrary::BakeToControlRig, except that it allows us to control which
+	// sequence player is used, lets us use our own existing AnimSequence for the ControlRig track, doesn't force
+	// the control rig editor mode to open and doesn't crash itself when changing the edit mode away from the
+	// control rig
+	bool BakeToControlRig(
+		UWorld* World,
+		ULevelSequence* LevelSequence,
+		UClass* InClass,
+		UAnimSequence* AnimSequence,
+		USkeletalMeshComponent* SkeletalMeshComp,
+		UAnimSeqExportOption* ExportOptions,
+		bool bReduceKeys,
+		float Tolerance,
+		const FGuid& ComponentBinding
+	)
+	{
+		UMovieScene* MovieScene = LevelSequence->GetMovieScene();
+		if ( !MovieScene || !SkeletalMeshComp || !SkeletalMeshComp->SkeletalMesh || !SkeletalMeshComp->SkeletalMesh->GetSkeleton() )
+		{
+			return false;
+		}
+
+		bool bResult = false;
+		bool bCreatedTempSequence = false;
+		ALevelSequenceActor* OutActor = nullptr;
+		UMovieSceneControlRigParameterTrack* Track = nullptr;
+		FMovieSceneSequencePlaybackSettings Settings;
+
+		// Always use a hidden player for this so that we don't affect/are affected by any Sequencer the user
+		// may have opened. Plus, if we have sublayers and subsequences its annoying to managed the Sequencer
+		// currently focused LevelSequence
+		IMovieScenePlayer* Player = nullptr;
+		ULevelSequencePlayer* LevelPlayer = nullptr;
+		{
+			Player = LevelPlayer = ULevelSequencePlayer::CreateLevelSequencePlayer( World, LevelSequence, Settings, OutActor );
+			if ( !Player || !LevelPlayer )
+			{
+				goto cleanup;
+			}
+
+			FLevelSequenceCameraSettings CameraSettings;
+			LevelPlayer->Initialize( LevelSequence, World->PersistentLevel, Settings, CameraSettings );
+			LevelPlayer->State.AssignSequence( MovieSceneSequenceID::Root, *LevelSequence, *Player );
+
+			// Evaluate at the beginning of the subscene time to ensure that spawnables are created before export
+			LevelPlayer->SetPlaybackPosition( FMovieSceneSequencePlaybackParams( UE::MovieScene::DiscreteInclusiveLower( MovieScene->GetPlaybackRange() ).Value, EUpdatePositionMethod::Play ) );
+		}
+
+		MovieScene->Modify();
+
+		// We allow baking with no AnimSequence (to allow rigging with no previous animation), so if we don't
+		// have an AnimSequence yet we need to bake a temp one
+		if( !AnimSequence )
+		{
+			bCreatedTempSequence = true;
+			AnimSequence = NewObject<UAnimSequence>();
+			AnimSequence->SetSkeleton( SkeletalMeshComp->SkeletalMesh->GetSkeleton() );
+
+			FMovieSceneSequenceIDRef Template = MovieSceneSequenceID::Root;
+			FMovieSceneSequenceTransform RootToLocalTransform;
+			bResult = MovieSceneToolHelpers::ExportToAnimSequence( AnimSequence, ExportOptions, MovieScene, Player, SkeletalMeshComp, Template, RootToLocalTransform );
+			if ( !bResult )
+			{
+				goto cleanup;
+			}
+		}
+
+		// Disable any extra existing control rig tracks for this binding.
+		// Reuse one of the control rig parameter tracks if we can
+		{
+			TArray<UMovieSceneTrack*> Tracks = MovieScene->FindTracks( UMovieSceneControlRigParameterTrack::StaticClass(), ComponentBinding, NAME_None );
+			for ( UMovieSceneTrack* AnyOleTrack : Tracks )
+			{
+				UMovieSceneControlRigParameterTrack* ValidTrack = Cast<UMovieSceneControlRigParameterTrack>( AnyOleTrack );
+				if ( ValidTrack )
+				{
+					Track = ValidTrack;
+					Track->Modify();
+					for ( UMovieSceneSection* Section : Track->GetAllSections() )
+					{
+						Section->SetIsActive( false );
+					}
+				}
+			}
+
+			if ( !Track )
+			{
+				Track = Cast<UMovieSceneControlRigParameterTrack>( MovieScene->AddTrack( UMovieSceneControlRigParameterTrack::StaticClass(), ComponentBinding ) );
+				Track->Modify();
+			}
+		}
+
+		if ( Track )
+		{
+			FString ObjectName = InClass->GetName();
+			ObjectName.RemoveFromEnd( TEXT( "_C" ) );
+			UControlRig* ControlRig = NewObject<UControlRig>( Track, InClass, FName( *ObjectName ), RF_Transactional );
+			if ( InClass != UFKControlRig::StaticClass() && !ControlRig->SupportsEvent( TEXT( "Backwards Solve" ) ) )
+			{
+				MovieScene->RemoveTrack( *Track );
+				goto cleanup;
+			}
+
+			ControlRig->Modify();
+			ControlRig->SetObjectBinding( MakeShared<FControlRigObjectBinding>() );
+			ControlRig->GetObjectBinding()->BindToObject( SkeletalMeshComp );
+			ControlRig->GetDataSourceRegistry()->RegisterDataSource( UControlRig::OwnerComponent, ControlRig->GetObjectBinding()->GetBoundObject() );
+			ControlRig->Initialize();
+			ControlRig->RequestInit();
+			ControlRig->SetBoneInitialTransformsFromSkeletalMeshComponent( SkeletalMeshComp, true );
+			ControlRig->Evaluate_AnyThread();
+
+			// Find the animation section's start frame, or else the baked control rig tracks will always be
+			// placed at the start of the movie scene playback range, instead of following where the actual
+			// animation section is
+			bool bFoundAtLeastOneSection = false;
+			FFrameNumber ControlRigSectionStartFrame = TNumericLimits<int32>::Max();
+			UMovieSceneSkeletalAnimationTrack* SkelTrack = Cast<UMovieSceneSkeletalAnimationTrack>(
+				MovieScene->FindTrack( UMovieSceneSkeletalAnimationTrack::StaticClass(), ComponentBinding, NAME_None )
+			);
+			if ( SkelTrack )
+			{
+				for ( const UMovieSceneSection* Section : SkelTrack->GetAllSections() )
+				{
+					if ( const UMovieSceneSkeletalAnimationSection* SkelSection = Cast<UMovieSceneSkeletalAnimationSection>( Section ) )
+					{
+						if ( SkelSection->Params.Animation == AnimSequence )
+						{
+							TRange<FFrameNumber> Range = SkelSection->ComputeEffectiveRange();
+							if ( Range.HasLowerBound() )
+							{
+								bFoundAtLeastOneSection = true;
+								ControlRigSectionStartFrame = FMath::Min( ControlRigSectionStartFrame, Range.GetLowerBoundValue() );
+								break;
+							}
+						}
+					}
+				}
+			}
+			if ( !bFoundAtLeastOneSection )
+			{
+				ControlRigSectionStartFrame = 0;
+			}
+
+			bool bSequencerOwnsControlRig = true;
+			UMovieSceneSection* NewSection = Track->CreateControlRigSection( ControlRigSectionStartFrame, ControlRig, bSequencerOwnsControlRig );
+			UMovieSceneControlRigParameterSection* ParamSection = Cast<UMovieSceneControlRigParameterSection>( NewSection );
+
+			Track->SetTrackName( FName( *ObjectName ) );
+			Track->SetDisplayName( FText::FromString( ObjectName ) );
+
+			ParamSection->LoadAnimSequenceIntoThisSection( AnimSequence, MovieScene, SkeletalMeshComp, bReduceKeys, Tolerance, ControlRigSectionStartFrame );
+
+			// Disable Skeletal Animation Tracks
+			if ( SkelTrack )
+			{
+				SkelTrack->Modify();
+
+				for ( UMovieSceneSection* Section : SkelTrack->GetAllSections() )
+				{
+					if ( Section )
+					{
+						Section->TryModify();
+						Section->SetIsActive( false );
+					}
+				}
+			}
+
+			bResult = true;
+		}
+
+	cleanup:
+		if ( bCreatedTempSequence && AnimSequence )
+		{
+			AnimSequence->MarkAsGarbage();
+		}
+
+		if ( LevelPlayer )
+		{
+			LevelPlayer->Stop();
+		}
+
+		if ( OutActor && World )
+		{
+			World->DestroyActor( OutActor );
+		}
+
+		return bResult;
+	}
+#endif // WITH_EDITOR
 }
 
 class FUsdLevelSequenceHelperImpl : private FGCObject
@@ -245,6 +490,11 @@
 		return SubSequences;
 	}
 
+	FUsdLevelSequenceHelper::FOnSkelAnimationBaked& GetOnSkelAnimationBaked()
+	{
+		return OnSkelAnimationBaked;
+	}
+
 private:
 	ULevelSequence* FindSequenceForAttribute( const UE::FUsdAttribute& Attribute );
 	ULevelSequence* FindOrAddSequenceForAttribute( const UE::FUsdAttribute& Attribute );
@@ -297,6 +547,9 @@
 	void AddPrim( UUsdPrimTwin& PrimTwin, bool bForceVisibilityTracks = false );
 	void RemovePrim(const UUsdPrimTwin& PrimTwin);
 
+	// These functions assume the skeletal animation tracks (if any) were already added to the level sequence
+	void UpdateControlRigTracks( UUsdPrimTwin& PrimTwin );
+
 private:
 	// Sequence Name to Prim Path. Relationship: 1 Sequence to N Prim Path.
 	TMultiMap<FName, FString> PrimPathByLevelSequenceName;
@@ -323,6 +576,8 @@
 	double GetFramesPerSecond() const;
 	double GetTimeCodesPerSecond() const;
 
+	FGuid GetOrCreateComponentBinding( const UUsdPrimTwin& PrimTwin, USceneComponent& ComponentToBind, ULevelSequence& Sequence );
+
 	TMap<FString, FLayerTimeInfo> LayerTimeInfosByLayerIdentifier; // Maps a LayerTimeInfo to a given Layer through its identifier
 
 // Changes handling
@@ -343,6 +598,7 @@
 	void HandleTransactionStateChanged( const FTransactionContext& InTransactionContext, const ETransactionStateEventType InTransactionState );
 	void HandleMovieSceneChange(UMovieScene& MovieScene);
 	void HandleSubSectionChange(UMovieSceneSubSection& Section);
+	void HandleControlRigSectionChange( UMovieSceneControlRigParameterSection& Section );
 	void HandleTrackChange( const UMovieSceneTrack& Track, bool bIsMuteChange );
 	void HandleDisplayRateChange(const double DisplayRate);
 
@@ -364,6 +620,8 @@
 	static const TCHAR* TimeTrackName;
 	static const double EmptySubSectionRange; // How many frames should an empty subsection cover, only needed so that the subsection is visible and the user can edit it
 
+	FUsdLevelSequenceHelper::FOnSkelAnimationBaked OnSkelAnimationBaked;
+
 	TWeakObjectPtr<AUsdStageActor> StageActor = nullptr;
 	UUsdAssetCache* AssetCache = nullptr;  // We keep a pointer to this directly because we may be called via the USDStageImporter directly, when we don't have an available actor
 	FGuid StageActorBinding;
@@ -1656,61 +1914,7 @@
 		return nullptr;
 	}
 
-	const FGuid ComponentBinding = [&]() -> FGuid
-	{
-		if ( TPair< ULevelSequence*, FGuid >* SceneComponentBinding = SceneComponentsBindings.Find( &PrimTwin ) )
-		{
-			return SceneComponentBinding->Value;
-		}
-		else
-		{
-			FGuid ComponentBinding;
-			FGuid ActorBinding;
-			UObject* ComponentContext = ComponentToBind.GetWorld();
-
-			FString PrimName = FPaths::GetBaseFilename( PrimTwin.PrimPath );
-
-			// Make sure we always bind the parent actor too
-			if ( AActor* Actor = ComponentToBind.GetOwner() )
-			{
-				ActorBinding = Sequence.FindBindingFromObject( Actor, Actor->GetWorld() );
-				if ( !ActorBinding.IsValid() )
-				{
-					// We use the label here because that will always be named after the prim that caused the actor
-					// to be generated. If we just used our own PrimName in here we may run into situations where a child Camera prim
-					// of a decomposed camera ends up naming the actor binding after itself, even though the parent Xform prim, and the
-					// actor on the level, maybe named something else
-					ActorBinding = MovieScene->AddPossessable(
-#if WITH_EDITOR
-						Actor->GetActorLabel(),
-#else
-						Actor->GetName(),
-#endif // WITH_EDITOR
-						Actor->GetClass()
-					);
-					Sequence.BindPossessableObject( ActorBinding, *Actor, Actor->GetWorld() );
-				}
-
-				ComponentContext = Actor;
-			}
-
-			ComponentBinding = MovieScene->AddPossessable( PrimName, ComponentToBind.GetClass() );
-
-			if ( ActorBinding.IsValid() && ComponentBinding.IsValid() )
-			{
-				if ( FMovieScenePossessable* ComponentPossessable = MovieScene->FindPossessable( ComponentBinding ) )
-				{
-					ComponentPossessable->SetParent( ActorBinding, MovieScene );
-				}
-			}
-
-			// Bind component
-			Sequence.BindPossessableObject( ComponentBinding, ComponentToBind, ComponentContext );
-
-			SceneComponentsBindings.Emplace( &PrimTwin ) = TPair< ULevelSequence*, FGuid >( &Sequence, ComponentBinding );
-			return ComponentBinding;
-		}
-	}( );
+	const FGuid ComponentBinding = GetOrCreateComponentBinding( PrimTwin, ComponentToBind, Sequence );
 
 	const bool bReadonly = false;
 	UsdLevelSequenceHelperImpl::FMovieSceneReadonlyGuard MovieSceneReadonlyGuard{ *MovieScene, bReadonly };
@@ -1740,19 +1944,8 @@
 #endif // WITH_EDITOR
 	}
 
-	if ( bIsMuted )
-	{
-#if WITH_EDITOR
-		// We need to update the MovieScene too, because if MuteNodes disagrees with Track->IsEvalDisabled() the sequencer
-		// will chose in favor of MuteNodes
-		MovieScene->Modify();
-		MovieScene->GetMuteNodes().AddUnique( FString::Printf( TEXT( "%s.%s" ), *ComponentBinding.ToString(), *Track->GetName() ) );
-#endif // WITH_EDITOR
+	UsdLevelSequenceHelperImpl::MuteTrack( Track, MovieScene, ComponentBinding.ToString(), Track->GetName(), bIsMuted );
 
-		Track->Modify();
-		Track->SetEvalDisabled( bIsMuted );
-	}
-
 	return Track;
 }
 
@@ -1801,6 +1994,193 @@
 	RefreshSequencer();
 }
 
+void FUsdLevelSequenceHelperImpl::UpdateControlRigTracks( UUsdPrimTwin& PrimTwin )
+{
+#if WITH_EDITOR
+	if ( !UsdStage )
+	{
+		return;
+	}
+
+	UE::FSdfPath PrimPath( *PrimTwin.PrimPath );
+	UE::FUsdPrim UsdPrim( UsdStage.GetPrimAtPath( PrimPath ) );
+	if ( !UsdPrim )
+	{
+		return;
+	}
+
+	USkeletalMeshComponent* ComponentToBind = Cast<USkeletalMeshComponent>( PrimTwin.GetSceneComponent() );
+	if ( !ComponentToBind )
+	{
+		return;
+	}
+
+	// Block here because USD needs to fire and respond to notices for the DefinePrim call to work,
+	// but we need UsdUtils::BindAnimationSource to run before we get in here again or else we'll
+	// repeatedly create Animation prims
+	FScopedBlockNoticeListening BlockNotices( StageActor.Get() );
+
+	UE::FSdfLayer SkelAnimationLayer;
+
+	// We'll place the skeletal animation track wherever the SkelAnimation prim is defined (not necessarily the
+	// same layer as the skel root)
+	UE::FUsdPrim SkelAnimationPrim = UsdUtils::FindFirstAnimationSource( UsdPrim );
+	if ( SkelAnimationPrim )
+	{
+		SkelAnimationLayer = UsdUtils::FindLayerForPrim( SkelAnimationPrim );
+	}
+	else
+	{
+		// If this SkelRoot doesn't have any animation, lets create a new one on the current edit target
+		SkelAnimationLayer = UsdStage.GetEditTarget();
+
+		ensure( UsdPrim.IsA( TEXT( "SkelRoot" ) ) );
+
+		FString UniqueChildName = UsdUtils::GetValidChildName( TEXT( "Animation" ), UsdPrim );
+		SkelAnimationPrim = UsdStage.DefinePrim(
+			UsdPrim.GetPrimPath().AppendChild( *UniqueChildName ),
+			TEXT( "SkelAnimation" )
+		);
+		if ( !SkelAnimationPrim )
+		{
+			return;
+		}
+
+		UsdUtils::BindAnimationSource( UsdPrim, SkelAnimationPrim );
+	}
+
+	// Fetch the UAnimSequence asset from the asset cache. Ideally we'd call AUsdStageActor::GetGeneratedAssets,
+	// but we may belong to a FUsdStageImportContext, and so there's no AUsdStageActor at all to use.
+	// At this point it doesn't matter much though, because we shouldn't need to uncollapse a SkelAnimation prim path anyway
+	FString SkelAnimationPrimPath = SkelAnimationPrim.GetPrimPath().GetString();
+	UAnimSequence* AnimSequence = Cast<UAnimSequence>( AssetCache->GetAssetForPrim( SkelAnimationPrimPath ) );
+
+	if ( !SkelAnimationLayer )
+	{
+		return;
+	}
+	UE::FUsdEditContext EditContext{ UsdStage, SkelAnimationLayer };
+	FString Identifier = SkelAnimationLayer.GetIdentifier();
+
+	// Force-create these because these are mandatory anyway (https://graphics.pixar.com/usd/release/api/_usd_skel__schemas.html#UsdSkel_SkelAnimation)
+	UE::FUsdAttribute JointsAttr = SkelAnimationPrim.CreateAttribute( TEXT( "joints" ), TEXT( "token[]" ) );
+	UE::FUsdAttribute TranslationsAttr = SkelAnimationPrim.CreateAttribute( TEXT( "translations" ), TEXT( "float3[]" ) );
+	UE::FUsdAttribute RotationsAttr = SkelAnimationPrim.CreateAttribute( TEXT( "rotations" ), TEXT( "quatf[]" ) );
+	UE::FUsdAttribute ScalesAttr = SkelAnimationPrim.CreateAttribute( TEXT( "scales" ), TEXT( "half3[]" ) );
+	UE::FUsdAttribute BlendShapeWeightsAttr = SkelAnimationPrim.GetAttribute( TEXT( "blendShapeWeights" ) );
+
+	ULevelSequence* SkelAnimationSequence = FindOrAddSequenceForLayer( SkelAnimationLayer, SkelAnimationLayer.GetIdentifier(), SkelAnimationLayer.GetDisplayName() );
+	if ( !SkelAnimationSequence )
+	{
+		return;
+	}
+
+	UMovieScene* MovieScene = SkelAnimationSequence->GetMovieScene();
+	if ( !MovieScene )
+	{
+		return;
+	}
+
+	const bool bReadonly = false;
+	UsdLevelSequenceHelperImpl::FMovieSceneReadonlyGuard MovieSceneReadonlyGuard{ *MovieScene, bReadonly };
+
+	const FGuid ComponentBinding = GetOrCreateComponentBinding( PrimTwin, *ComponentToBind, *SkelAnimationSequence );
+
+	// NOTE: We are fetching the first skel track we find, since we can't actually use SkelAnimationPrim.GetName() here at all!
+	// The property tracks do derive GetTrackName(), but the skeletal track doesn't, so FindTrack will never find them.
+	// This likely has no effect since we only ever spawn a single skeletal track per prim anyway, but its worth to keep in mind!
+	UMovieSceneSkeletalAnimationTrack* SkelTrack = MovieScene->FindTrack< UMovieSceneSkeletalAnimationTrack >( ComponentBinding );
+	UMovieSceneControlRigParameterTrack* ControlRigTrack = MovieScene->FindTrack< UMovieSceneControlRigParameterTrack >( ComponentBinding );
+
+	// We should be in control rig track mode but don't have any tracks yet --> Setup for Control Rig
+	if ( !ControlRigTrack )
+	{
+		bool bControlRigReduceKeys = false;
+		if ( UE::FUsdAttribute Attr = UsdPrim.GetAttribute( *UsdToUnreal::ConvertToken( UnrealIdentifiers::UnrealControlRigReduceKeys ) ) )
+		{
+			UE::FVtValue Value;
+			if ( Attr.Get( Value ) && !Value.IsEmpty() )
+			{
+				if ( TOptional<bool> UnderlyingValue = UsdUtils::GetUnderlyingValue<bool>( Value ) )
+				{
+					bControlRigReduceKeys = UnderlyingValue.GetValue();
+				}
+			}
+		}
+
+		float ControlRigReduceTolerance = 0.001f;
+		if ( UE::FUsdAttribute Attr = UsdPrim.GetAttribute( *UsdToUnreal::ConvertToken( UnrealIdentifiers::UnrealControlRigReductionTolerance ) ) )
+		{
+			UE::FVtValue Value;
+			if ( Attr.Get( Value ) && !Value.IsEmpty() )
+			{
+				if ( TOptional<float> UnderlyingValue = UsdUtils::GetUnderlyingValue<float>( Value ) )
+				{
+					ControlRigReduceTolerance = UnderlyingValue.GetValue();
+				}
+			}
+		}
+
+		bool bIsFKControlRig = false;
+		if ( UE::FUsdAttribute Attr = UsdPrim.GetAttribute( *UsdToUnreal::ConvertToken( UnrealIdentifiers::UnrealUseFKControlRig ) ) )
+		{
+			UE::FVtValue Value;
+			if ( Attr.Get( Value ) )
+			{
+				if (TOptional<bool> UseFKOptional = UsdUtils::GetUnderlyingValue<bool>( Value ) )
+				{
+					bIsFKControlRig = UseFKOptional.GetValue();
+				}
+			}
+		}
+
+		UClass* ControlRigClass = nullptr;
+		if ( bIsFKControlRig )
+		{
+			ControlRigClass = UFKControlRig::StaticClass();
+		}
+		else
+		{
+			FString ControlRigBPPath;
+			if ( UE::FUsdAttribute Attr = UsdPrim.GetAttribute( *UsdToUnreal::ConvertToken( UnrealIdentifiers::UnrealControlRigPath ) ) )
+			{
+				UE::FVtValue Value;
+				if ( Attr.Get( Value ) && !Value.IsEmpty() )
+				{
+					ControlRigBPPath = UsdUtils::Stringify( Value );
+				}
+			}
+
+			if ( UControlRigBlueprint* BP = Cast<UControlRigBlueprint>( FSoftObjectPath( ControlRigBPPath ).TryLoad() ) )
+			{
+				ControlRigClass = BP->GetControlRigClass();
+			}
+		}
+
+		if ( ControlRigClass )
+		{
+			UAnimSeqExportOption* NewOptions = NewObject< UAnimSeqExportOption >();
+
+			UsdLevelSequenceHelperImpl::BakeToControlRig(
+				ComponentToBind->GetWorld(),
+				SkelAnimationSequence,
+				ControlRigClass,
+				AnimSequence,
+				ComponentToBind,
+				NewOptions,
+				bControlRigReduceKeys,
+				ControlRigReduceTolerance,
+				ComponentBinding
+			);
+
+			RefreshSequencer();
+		}
+#endif // WITH_EDITOR
+	}
+
+	PrimPathByLevelSequenceName.AddUnique( SkelAnimationSequence->GetFName(), PrimPath.GetString() );
+}
+
 void FUsdLevelSequenceHelperImpl::RemoveSequenceForPrim( ULevelSequence& Sequence, const UUsdPrimTwin& PrimTwin )
 {
 	TArray< FString > PrimPathsForSequence;
@@ -1897,19 +2277,13 @@
 		return;
 	}
 
-	const bool bFocusIfOpen = false;
-	IAssetEditorInstance* AssetEditor = GEditor->GetEditorSubsystem<UAssetEditorSubsystem>()->FindEditorForAsset(MainLevelSequence, bFocusIfOpen);
-	ILevelSequenceEditorToolkit* LevelSequenceEditor = static_cast<ILevelSequenceEditorToolkit*>(AssetEditor);
-	TWeakPtr<ISequencer> WeakSequencer = LevelSequenceEditor ? LevelSequenceEditor->GetSequencer() : nullptr;
-
-	if ( TSharedPtr< ISequencer > Sequencer = WeakSequencer.Pin() )
+	if ( TSharedPtr< ISequencer > Sequencer = UsdLevelSequenceHelperImpl::GetOpenedSequencerForLevelSequence( MainLevelSequence ) )
 	{
 		// Don't try refreshing the sequencer if its displaying a stale sequence (e.g. during busy transitions like import) as it
 		// can crash
 		if ( UMovieSceneSequence* FocusedSequence = Sequencer->GetFocusedMovieSceneSequence() )
 		{
-			Sequencer->RefreshTree();
-			Sequencer->NotifyMovieSceneDataChanged( EMovieSceneDataChangeType::TrackValueChanged );
+			Sequencer->NotifyMovieSceneDataChanged( EMovieSceneDataChangeType::RefreshTree );
 		}
 	}
 #endif // WITH_EDITOR
@@ -2131,7 +2505,7 @@
 	}
 	else if ( UMovieSceneSubSection* SubSection = Cast<UMovieSceneSubSection>(Object) )
 	{
-		HandleSubSectionChange(*SubSection);
+		HandleSubSectionChange( *SubSection );
 	}
 	else if ( UMovieSceneTrack* Track = Cast<UMovieSceneTrack>(Object) )
 	{
@@ -2140,11 +2514,32 @@
 	}
 	else if ( UMovieSceneSection* Section = Cast<UMovieSceneSection>(Object) )
 	{
+		const bool bIsMuteChange = Event.GetChangedProperties().Contains( TEXT( "bIsActive" ) );
+
 		if ( UMovieSceneTrack* ParentTrack = Section->GetTypedOuter<UMovieSceneTrack>() )
 		{
-			const bool bIsMuteChange = Event.GetChangedProperties().Contains( TEXT( "bIsActive" ) );
 			HandleTrackChange( *ParentTrack, bIsMuteChange );
 		}
+
+		if ( !bIsMuteChange )
+		{
+			if ( UMovieSceneControlRigParameterSection* CRSection = Cast<UMovieSceneControlRigParameterSection>( Section ) )
+			{
+				if ( GEditor )
+				{
+					// We have to do this on next tick because HandleControlRigSectionChange will internally bake
+					// the sequence, repeteadly updating the ControlRig hierarchy. There is no way to silence
+					// FControlRigEditMode from here, and FControlRigEditMode::OnHierarchyModified ends up creating
+					// a brand new scoped transaction, which asserts inside UTransBuffer::CheckState when if finds
+					// out that the previous transaction wasn't fully complete (OnObjectTransacted gets called before
+					// the current transaction is fully done).
+					GEditor->GetTimerManager()->SetTimerForNextTick( [this, CRSection]( )
+					{
+						HandleControlRigSectionChange( *CRSection );
+					});
+				}
+			}
+		}
 	}
 }
 
@@ -2221,6 +2616,68 @@
 	return FMath::IsNearlyZero( StageTimeCodesPerSecond ) ? DefaultFramerate : StageTimeCodesPerSecond;
 }
 
+FGuid FUsdLevelSequenceHelperImpl::GetOrCreateComponentBinding( const UUsdPrimTwin& PrimTwin, USceneComponent& ComponentToBind, ULevelSequence& Sequence )
+{
+	UMovieScene* MovieScene = Sequence.GetMovieScene();
+	if ( !MovieScene )
+	{
+		return {};
+	}
+
+	if ( TPair< ULevelSequence*, FGuid >* SceneComponentBinding = SceneComponentsBindings.Find( &PrimTwin ) )
+	{
+		return SceneComponentBinding->Value;
+	}
+	else
+	{
+		FGuid ComponentBinding;
+		FGuid ActorBinding;
+		UObject* ComponentContext = ComponentToBind.GetWorld();
+
+		FString PrimName = FPaths::GetBaseFilename( PrimTwin.PrimPath );
+
+		// Make sure we always bind the parent actor too
+		if ( AActor* Actor = ComponentToBind.GetOwner() )
+		{
+			ActorBinding = Sequence.FindBindingFromObject( Actor, Actor->GetWorld() );
+			if ( !ActorBinding.IsValid() )
+			{
+				// We use the label here because that will always be named after the prim that caused the actor
+				// to be generated. If we just used our own PrimName in here we may run into situations where a child Camera prim
+				// of a decomposed camera ends up naming the actor binding after itself, even though the parent Xform prim, and the
+				// actor on the level, maybe named something else
+				ActorBinding = MovieScene->AddPossessable(
+#if WITH_EDITOR
+					Actor->GetActorLabel(),
+#else
+					Actor->GetName(),
+#endif // WITH_EDITOR
+					Actor->GetClass()
+				);
+				Sequence.BindPossessableObject( ActorBinding, *Actor, Actor->GetWorld() );
+			}
+
+			ComponentContext = Actor;
+		}
+
+		ComponentBinding = MovieScene->AddPossessable( PrimName, ComponentToBind.GetClass() );
+
+		if ( ActorBinding.IsValid() && ComponentBinding.IsValid() )
+		{
+			if ( FMovieScenePossessable* ComponentPossessable = MovieScene->FindPossessable( ComponentBinding ) )
+			{
+				ComponentPossessable->SetParent( ActorBinding, MovieScene );
+			}
+		}
+
+		// Bind component
+		Sequence.BindPossessableObject( ComponentBinding, ComponentToBind, ComponentContext );
+
+		SceneComponentsBindings.Emplace( &PrimTwin ) = TPair< ULevelSequence*, FGuid >( &Sequence, ComponentBinding );
+		return ComponentBinding;
+	}
+}
+
 void FUsdLevelSequenceHelperImpl::HandleMovieSceneChange( UMovieScene& MovieScene )
 {
 	// It's possible to get this called when the actor and it's level sequences are being all destroyed in one go.
@@ -2410,6 +2867,187 @@
 	UpdateUsdLayerOffsetFromSection(ParentSequence, &Section);
 }
 
+void FUsdLevelSequenceHelperImpl::HandleControlRigSectionChange( UMovieSceneControlRigParameterSection& Section )
+{
+	AUsdStageActor* StageActorValue = StageActor.Get();
+	if ( !StageActorValue )
+	{
+		return;
+	}
+
+	UWorld* World = StageActorValue->GetWorld();
+	if ( !World )
+	{
+		return;
+	}
+
+	ULevelSequence* LevelSequence = Section.GetTypedOuter<ULevelSequence>();
+	if ( !LevelSequence )
+	{
+		return;
+	}
+
+	UMovieScene* MovieScene = LevelSequence->GetMovieScene();
+	if ( !MovieScene )
+	{
+		return;
+	}
+
+	UMovieSceneTrack* ParentTrack = Section.GetTypedOuter<UMovieSceneTrack>();
+	if ( !ParentTrack )
+	{
+		return;
+	}
+
+	FGuid PossessableGuid;
+	const bool bFound = MovieScene->FindTrackBinding( *ParentTrack, PossessableGuid );
+	if ( !bFound )
+	{
+		return;
+	}
+
+	FMovieScenePossessable* Possessable = MovieScene->FindPossessable( PossessableGuid );
+	if ( !Possessable )
+	{
+		return;
+	}
+
+	USkeletalMeshComponent* BoundComponent = Cast< USkeletalMeshComponent >(
+		UsdLevelSequenceHelperImpl::LocateBoundObject( *LevelSequence, *Possessable )
+	);
+	if ( !BoundComponent )
+	{
+		return;
+	}
+
+	USkeleton* Skeleton = BoundComponent->SkeletalMesh ? BoundComponent->SkeletalMesh->GetSkeleton() : nullptr;
+	if ( !Skeleton )
+	{
+		return;
+	}
+
+	UUsdPrimTwin* PrimTwin = StageActorValue->RootUsdTwin->Find( BoundComponent );
+	if ( !PrimTwin )
+	{
+		return;
+	}
+
+	UE::FUsdPrim UsdPrim = UsdStage.GetPrimAtPath( UE::FSdfPath( *PrimTwin->PrimPath ) );
+	if ( !UsdPrim )
+	{
+		return;
+	}
+
+	// We'll place the skeletal animation track wherever the SkelAnimation prim is defined (not necessarily the
+	// same layer as the skel root)
+	UE::FUsdPrim SkelAnimationPrim = UsdUtils::FindFirstAnimationSource( UsdPrim );
+	if ( !SkelAnimationPrim )
+	{
+		return;
+	}
+
+	TSharedPtr<ISequencer> PinnedSequencer = UsdLevelSequenceHelperImpl::GetOpenedSequencerForLevelSequence( MainLevelSequence );
+
+	// Fetch a sequence player we can use. We'll almost always have the sequencer opened here (we are responding to a transaction
+	// where the section was changed after all), but its possible to have a fallback too
+	IMovieScenePlayer* Player = nullptr;
+	ULevelSequencePlayer* LevelPlayer = nullptr;
+	{
+		if ( PinnedSequencer )
+		{
+			Player = PinnedSequencer.Get();
+		}
+		else
+		{
+			ALevelSequenceActor* OutActor = nullptr;
+			FMovieSceneSequencePlaybackSettings Settings;
+			Player = LevelPlayer = ULevelSequencePlayer::CreateLevelSequencePlayer( World, LevelSequence, Settings, OutActor);
+
+			FLevelSequenceCameraSettings CameraSettings;
+			LevelPlayer->Initialize( LevelSequence, World->PersistentLevel, Settings, CameraSettings );
+
+			LevelPlayer->State.AssignSequence( MovieSceneSequenceID::Root, *LevelSequence, *Player );
+		}
+
+		if ( !Player )
+		{
+			return;
+		}
+	}
+
+	// We obviously don't want to respond to the fact that the stage will be modified since we're the
+	// ones actually modifying it already
+	FScopedBlockNoticeListening BlockNotices( StageActorValue );
+
+	// Prepare for baking
+	{
+		if ( PinnedSequencer )
+		{
+			PinnedSequencer->EnterSilentMode();
+		}
+
+		FSpawnableRestoreState SpawnableRestoreState( MovieScene );
+		if ( LevelPlayer && SpawnableRestoreState.bWasChanged )
+		{
+			// Evaluate at the beginning of the subscene time to ensure that spawnables are created before export
+			// Note that we never actually generate spawnables on our LevelSequence, but its a common pattern to
+			// do this and the user may have added them manually
+			LevelPlayer->SetPlaybackPosition(
+				FMovieSceneSequencePlaybackParams(
+					UE::MovieScene::DiscreteInclusiveLower( MovieScene->GetPlaybackRange() ).Value,
+					EUpdatePositionMethod::Play
+				)
+			);
+		}
+	}
+
+	FMovieSceneSequenceTransform SequenceTransform;
+	FMovieSceneSequenceID SequenceID = SequencesID.FindRef( LevelSequence );
+	if ( FMovieSceneSubSequenceData* SubSequenceData = SequenceHierarchyCache.FindSubData( SequenceID ) )
+	{
+		SequenceTransform = SubSequenceData->RootToSequenceTransform;
+	}
+
+	// Actually bake inside the UsdUtilities module as we need to manipulate USD arrays a lot
+	const UsdUtils::FBlendShapeMap& BlendShapeMap = StageActorValue->GetBlendShapeMap();
+	bool bBaked = UnrealToUsd::ConvertControlRigSection(
+		&Section,
+		SequenceTransform.InverseLinearOnly(),
+		MovieScene,
+		Player,
+		Skeleton->GetReferenceSkeleton(),
+		UsdPrim,
+		SkelAnimationPrim,
+		&BlendShapeMap
+	);
+
+	// Cleanup after baking
+	{
+		if ( LevelPlayer )
+		{
+			LevelPlayer->Stop();
+		}
+
+		if ( PinnedSequencer )
+		{
+			PinnedSequencer->ExitSilentMode();
+			PinnedSequencer->RequestEvaluate();
+		}
+	}
+
+	if ( bBaked )
+	{
+		// After we bake, both the sequencer and the USD stage have our updated tracks, but we still have the old
+		// AnimSequence asset on the component. If we closed the Sequencer and just animated via the Time attribute,
+		// we would see the old animation.
+		// This event is mostly used to have the stage actor quickly regenerate the assets and components for the
+		// skel root. Sadly we do need to regenerate the skeletal mesh too, since we may need to affect blend shapes
+		// for the correct bake. The user can disable this behavior (e.g. for costly skeletal meshes) by setting
+		// USD.RegenerateSkeletalAssetsOnControlRigBake to false.
+		GetOnSkelAnimationBaked().Broadcast( PrimTwin->PrimPath );
+	}
+}
+
 void FUsdLevelSequenceHelperImpl::HandleTrackChange( const UMovieSceneTrack& Track, bool bIsMuteChange )
 {
 	if ( !StageActor.IsValid() )
@@ -2711,6 +3349,8 @@
 	void AddPrim( UUsdPrimTwin& PrimTwin, bool bForceVisibilityTracks ) {}
 	void RemovePrim(const UUsdPrimTwin& PrimTwin) {}
 
+	void UpdateControlRigTracks( UUsdPrimTwin& PrimTwin ) {}
+
 	void StartMonitoringChanges() {}
 	void StopMonitoringChanges() {}
 	void BlockMonitoringChangesForThisTransaction() {}
@@ -2841,6 +3481,14 @@
 	}
 }
 
+void FUsdLevelSequenceHelper::UpdateControlRigTracks( UUsdPrimTwin& PrimTwin )
+{
+	if ( UsdSequencerImpl.IsValid() )
+	{
+		UsdSequencerImpl->UpdateControlRigTracks( PrimTwin );
+	}
+}
+
 void FUsdLevelSequenceHelper::StartMonitoringChanges()
 {
 	if (UsdSequencerImpl.IsValid())
@@ -2889,6 +3537,21 @@
 	}
 }
 
+FUsdLevelSequenceHelper::FOnSkelAnimationBaked& FUsdLevelSequenceHelper::GetOnSkelAnimationBaked()
+{
+#if USE_USD_SDK
+	if ( UsdSequencerImpl.IsValid() )
+	{
+		return UsdSequencerImpl->GetOnSkelAnimationBaked();
+	}
+	else
+#endif // USE_USD_SDK
+	{
+		static FOnSkelAnimationBaked DefaultHandler;
+		return DefaultHandler;
+	}
+}
+
 FScopedBlockMonitoringChangesForTransaction::FScopedBlockMonitoringChangesForTransaction( FUsdLevelSequenceHelper& InHelper )
 	: FScopedBlockMonitoringChangesForTransaction( *InHelper.UsdSequencerImpl.Get() )
 {
diff --git a/Engine/Plugins/Importers/USDImporter/Source/USDStage/Private/USDStageActor.cpp b/Engine/Plugins/Importers/USDImporter/Source/USDStage/Private/USDStageActor.cpp
--- a/Engine/Plugins/Importers/USDImporter/Source/USDStage/Private/USDStageActor.cpp
+++ b/Engine/Plugins/Importers/USDImporter/Source/USDStage/Private/USDStageActor.cpp
@@ -70,6 +70,7 @@
 #include "Editor/TransBuffer.h"
 #include "Editor/UnrealEdEngine.h"
 #include "Engine/Selection.h"
+#include "ILevelSequenceEditorToolkit.h"
 #include "LevelEditor.h"
 #include "PropertyEditorModule.h"
 #include "Subsystems/AssetEditorSubsystem.h"
@@ -78,6 +79,12 @@
 
 #define LOCTEXT_NAMESPACE "USDStageActor"
 
+static bool GRegenerateSkeletalAssetsOnControlRigBake = true;
+static FAutoConsoleVariableRef CVarRegenerateSkeletalAssetsOnControlRigBake(
+	TEXT( "USD.RegenerateSkeletalAssetsOnControlRigBake" ),
+	GRegenerateSkeletalAssetsOnControlRigBake,
+	TEXT( "Whether to regenerate the assets associated with a SkelRoot (mesh, skeleton, anim sequence, etc.) whenever we modify Control Rig tracks. The USD Stage itself is always updated however." ) );
+
 static const EObjectFlags DefaultObjFlag = EObjectFlags::RF_Transactional | EObjectFlags::RF_Transient;
 
 AUsdStageActor::FOnActorLoaded AUsdStageActor::OnActorLoaded;
@@ -712,6 +719,8 @@ AUsdStageActor::AUsdStageActor()
 			}
 		}
 
+		LevelSequenceHelper.GetOnSkelAnimationBaked().AddUObject( this, &AUsdStageActor::OnSkelAnimationBaked );
+
 #endif // WITH_EDITOR
 
 		OnTimeChanged.AddUObject( this, &AUsdStageActor::AnimatePrims );
@@ -753,6 +762,7 @@ AUsdStageActor::AUsdStageActor()
 
 void AUsdStageActor::OnUsdObjectsChanged( const UsdUtils::FObjectChangesByPath& InfoChanges, const UsdUtils::FObjectChangesByPath& ResyncChanges )
 {
+#if USE_USD_SDK
 	if ( !IsListeningToUsdNotices() )
 	{
 		return;
@@ -807,20 +817,33 @@ void AUsdStageActor::OnUsdObjectsChanged( const UsdUtils::FObjectChangesByPath&
 			continue;
 		}
 
-		// Some stage info should trigger some resyncs (even though technically info changes) because they should trigger reparsing of geometry
+		// Upgrade some info changes into resync changes
 		bool bIsResync = false;
-		if ( PrimPath.IsAbsoluteRootPath() )
+		for ( const UsdUtils::FObjectChangeNotice& ObjectChange : InfoChange.Value )
 		{
-			for ( const UsdUtils::FObjectChangeNotice& ObjectChange : InfoChange.Value )
+			for ( const UsdUtils::FAttributeChange& AttributeChange : ObjectChange.AttributeChanges )
 			{
-				for ( const UsdUtils::FAttributeChange& AttributeChange : ObjectChange.AttributeChanges )
+				static const TSet<FString> StageResyncProperties = {
+					TEXT( "metersPerUnit" ),
+					TEXT( "upAxis" )
+				};
+
+				// Upgrade these to resync so that the prim twins are regenerated, which clears all the existing
+				// animation tracks and adds new ones, automatically re-baking to control rig
+				static const TSet<FString> PrimResyncProperties = {
+					*UsdToUnreal::ConvertToken( UnrealIdentifiers::UnrealControlRigPath ),
+					*UsdToUnreal::ConvertToken( UnrealIdentifiers::UnrealUseFKControlRig ),
+					*UsdToUnreal::ConvertToken( UnrealIdentifiers::UnrealControlRigReduceKeys ),
+					*UsdToUnreal::ConvertToken( UnrealIdentifiers::UnrealControlRigReductionTolerance )
+				};
+
+				// Some stage info should trigger some resyncs because they should trigger reparsing of geometry
+				if ( ( PrimPath.IsAbsoluteRootPath() && StageResyncProperties.Contains( AttributeChange.PropertyName ) )
+					|| PrimResyncProperties.Contains( AttributeChange.PropertyName ) )
 				{
-					static const TSet<FString> ResyncProperties = { TEXT( "metersPerUnit" ), TEXT( "upAxis" ) };
-					if ( ResyncProperties.Contains( AttributeChange.PropertyName ) )
-					{
-						bIsResync = true;
-						bHasResync = true;
-					}
+					bIsResync = true;
+					bHasResync = true;
+					break;
 				}
 			}
 		}
@@ -998,7 +1021,6 @@ void AUsdStageActor::OnUsdObjectsChanged( const UsdUtils::FObjectChangesByPath&
 
 					UMaterialInterface* NewMaterial = Cast<UMaterialInterface>( AssetCache->GetAssetForPrim( AssetsPrimPath.GetString() ) );
 
-#if USE_USD_SDK
 					// For UE-120185: If we recreated a material for a prim path we also need to update all components that were using it.
 					// This could be fleshed out further if other asset types require this refresh of "dependent components" but materials
 					// seem to be the only ones that do at the moment
@@ -1010,7 +1032,6 @@ void AUsdStageActor::OnUsdObjectsChanged( const UsdUtils::FObjectChangesByPath&
 							UpdateComponents( UE::FSdfPath{ *MaterialUserPrim }, bResyncComponent );
 						}
 					}
-#endif // USE_USD_SDK
 
 					// Resyncing also includes "updating" the prim
 					UpdatedAssets.Add( AssetsPrimPath );
@@ -1032,6 +1053,7 @@ void AUsdStageActor::OnUsdObjectsChanged( const UsdUtils::FObjectChangesByPath&
 			OnPrimChanged.Broadcast( PrimChangedInfo.Key, PrimChangedInfo.Value );
 		}
 	}
+#endif // USE_USD_SDK
 }
 
 USDSTAGE_API void AUsdStageActor::Reset()
@@ -1107,6 +1129,8 @@ UUsdPrimTwin* AUsdStageActor::GetOrCreatePrimTwin( const UE::FSdfPath& UsdPrimPa
 
 UUsdPrimTwin* AUsdStageActor::ExpandPrim( const UE::FUsdPrim& Prim, FUsdSchemaTranslationContext& TranslationContext )
 {
+	UUsdPrimTwin* UsdPrimTwin = nullptr;
+#if USE_USD_SDK
 	// "Active" is the non-destructive deletion used in USD. Sometimes when we rename/remove a prim in a complex stage it may remain in
 	// an inactive state, but its otherwise effectively deleted
 	if ( !Prim || !Prim.IsActive() )
@@ -1116,7 +1140,7 @@ UUsdPrimTwin* AUsdStageActor::ExpandPrim( const UE::FUsdPrim& Prim, FUsdSchemaTr
 
 	TRACE_CPUPROFILER_EVENT_SCOPE( AUsdStageActor::ExpandPrim );
 
-	UUsdPrimTwin* UsdPrimTwin = GetOrCreatePrimTwin( Prim.GetPrimPath() );
+	UsdPrimTwin = GetOrCreatePrimTwin( Prim.GetPrimPath() );
 
 	if ( !UsdPrimTwin )
 	{
@@ -1177,8 +1201,8 @@ UUsdPrimTwin* AUsdStageActor::ExpandPrim( const UE::FUsdPrim& Prim, FUsdSchemaTr
 	}
 
 	// Update the prim animated status
-#if USE_USD_SDK
-	if ( UsdUtils::IsAnimated( Prim ) )
+	const bool bIsAnimated = UsdUtils::IsAnimated( Prim );
+	if ( bIsAnimated )
 	{
 		if ( !PrimsToAnimate.Contains( UsdPrimTwin->PrimPath ) )
 		{
@@ -1213,7 +1237,6 @@ UUsdPrimTwin* AUsdStageActor::ExpandPrim( const UE::FUsdPrim& Prim, FUsdSchemaTr
 		}
 	}
 	else
-#endif // USE_USD_SDK
 	{
 		if ( PrimsToAnimate.Contains( UsdPrimTwin->PrimPath ) )
 		{
@@ -1222,6 +1245,31 @@ UUsdPrimTwin* AUsdStageActor::ExpandPrim( const UE::FUsdPrim& Prim, FUsdSchemaTr
 		}
 	}
 
+	// Setup Control Rig tracks if we need to. This must be done after adding regular skeletal animation tracks
+	// if we have any as if will properly deactivate them like the usual "Bake to Control Rig" workflow.
+	if ( Prim.IsA( TEXT( "SkelRoot" ) ) )
+	{
+		if ( UsdUtils::PrimHasControlRigSchema( Prim ) )
+		{
+			LevelSequenceHelper.UpdateControlRigTracks( *UsdPrimTwin );
+
+			// If our prim wasn't originally considered animated and we just added a new track, it should be
+			// considered animated too, so lets add it to the proper locations. This will also ensure that
+			// we can close the sequencer after creating a new animation in this way and see it animate on
+			// the level
+			if ( !bIsAnimated )
+			{
+				PrimsToAnimate.Add( UsdPrimTwin->PrimPath );
+				LevelSequenceHelper.AddPrim( *UsdPrimTwin );
+
+				// Prevent register/unregister spam when calling FUsdGeomXformableTranslator::UpdateComponents later
+				// during sequencer animation (which can cause the Sequencer UI to glitch out a bit)
+				UsdPrimTwin->SceneComponent->SetMobility( EComponentMobility::Movable );
+			}
+		}
+	}
+
+#endif // USE_USD_SDK
 	return UsdPrimTwin;
 }
 
@@ -2696,6 +2744,60 @@ bool AUsdStageActor::HasAuthorityOverStage() const
 	return !IsTemplate();
 }
 
+void AUsdStageActor::OnSkelAnimationBaked( const FString& SkelRootPrimPath )
+{
+#if USE_USD_SDK
+	if ( !UsdStage || !GRegenerateSkeletalAssetsOnControlRigBake )
+	{
+		return;
+	}
+
+	UE::FUsdPrim SkelRootPrim = UsdStage.GetPrimAtPath( UE::FSdfPath{ *SkelRootPrimPath } );
+	if ( !SkelRootPrim || !SkelRootPrim.IsA( TEXT( "SkelRoot" ) ) )
+	{
+		return;
+	}
+
+	UUsdPrimTwin* RootTwin = GetRootPrimTwin();
+	if ( !RootTwin )
+	{
+		return;
+	}
+
+	UUsdPrimTwin* Twin = RootTwin->Find(SkelRootPrimPath);
+	if ( !Twin )
+	{
+		return;
+	}
+
+	USkeletalMeshComponent* SkeletalMeshComponent = Cast<USkeletalMeshComponent>( Twin->GetSceneComponent() );
+	if ( !SkeletalMeshComponent )
+	{
+		return;
+	}
+
+	TSharedRef< FUsdSchemaTranslationContext > TranslationContext = FUsdStageActorImpl::CreateUsdSchemaTranslationContext( this, SkelRootPrimPath );
+	// The only way we could have baked a skel animation is via the sequencer, so we know its playing
+	TranslationContext->bSequencerIsAnimating = true;
+
+	IUsdSchemasModule& UsdSchemasModule = FModuleManager::Get().LoadModuleChecked< IUsdSchemasModule >( TEXT( "USDSchemas" ) );
+	if ( TSharedPtr< FUsdSchemaTranslator > SchemaTranslator = UsdSchemasModule.GetTranslatorRegistry().CreateTranslatorForSchema( TranslationContext, UE::FUsdTyped( SkelRootPrim ) ) )
+	{
+		if ( TSharedPtr< FUsdSkelRootTranslator > SkelRootTranslator = StaticCastSharedPtr<FUsdSkelRootTranslator>( SchemaTranslator ) )
+		{
+			// For now we're regenerating all asset types (including skeletal meshes) but we could
+			// eventually just split off the anim sequence generation and call exclusively that from
+			// here
+			SkelRootTranslator->CreateAssets();
+			TranslationContext->CompleteTasks();
+
+			// Have to update the components to assign the new assets
+			SkelRootTranslator->UpdateComponents( SkeletalMeshComponent );
+		}
+	}
+#endif // #if USE_USD_SDK
+}
+
 void AUsdStageActor::LoadAsset( FUsdSchemaTranslationContext& TranslationContext, const UE::FUsdPrim& Prim )
 {
 	TRACE_CPUPROFILER_EVENT_SCOPE( AUsdStageActor::LoadAsset );
@@ -2804,6 +2906,16 @@ void AUsdStageActor::AnimatePrims()
 
 	TSharedRef< FUsdSchemaTranslationContext > TranslationContext = FUsdStageActorImpl::CreateUsdSchemaTranslationContext( this, GetRootPrimTwin()->PrimPath );
 
+	// c.f. comment on bSequencerIsAnimating's declaration
+#if WITH_EDITOR
+	const bool bFocusIfOpen = false;
+	IAssetEditorInstance* AssetEditor = GEditor->GetEditorSubsystem<UAssetEditorSubsystem>()->FindEditorForAsset( LevelSequence, bFocusIfOpen );
+	if ( ILevelSequenceEditorToolkit* LevelSequenceEditor = static_cast< ILevelSequenceEditorToolkit* >( AssetEditor ) )
+	{
+		TranslationContext->bSequencerIsAnimating = true;
+	}
+#endif // WITH_EDITOR
+
 	for ( const FString& PrimToAnimate : PrimsToAnimate )
 	{
 		UE::FSdfPath PrimPath( *PrimToAnimate );
@@ -2826,7 +2938,7 @@ void AUsdStageActor::AnimatePrims()
 		GEditor->BroadcastLevelActorListChanged();
 		GEditor->RedrawLevelEditingViewports();
 	}
-#endif
+#endif // WITH_EDITOR
 }
 
 FScopedBlockNoticeListening::FScopedBlockNoticeListening( AUsdStageActor* InStageActor )
diff --git a/Engine/Plugins/Importers/USDImporter/Source/USDStage/Public/USDLevelSequenceHelper.h b/Engine/Plugins/Importers/USDImporter/Source/USDStage/Public/USDLevelSequenceHelper.h
--- a/Engine/Plugins/Importers/USDImporter/Source/USDStage/Public/USDLevelSequenceHelper.h
+++ b/Engine/Plugins/Importers/USDImporter/Source/USDStage/Public/USDLevelSequenceHelper.h
@@ -66,6 +66,8 @@ public:
 	/** Removes any track associated with this prim */
 	void RemovePrim(const UUsdPrimTwin& PrimTwin);
 
+	void UpdateControlRigTracks( UUsdPrimTwin& PrimTwin );
+
 	/** Blocks updating the level sequences & tracks from object changes. */
 	void StartMonitoringChanges();
 	void StopMonitoringChanges();
@@ -74,6 +76,9 @@ public:
 	ULevelSequence* GetMainLevelSequence() const;
 	TArray< ULevelSequence* > GetSubSequences() const;
 
+	DECLARE_EVENT_OneParam( FUsdLevelSequenceHelper, FOnSkelAnimationBaked, const FString& /*SkelRootPrimPath*/ );
+	FOnSkelAnimationBaked& GetOnSkelAnimationBaked();
+
 private:
 	friend class FScopedBlockMonitoringChangesForTransaction;
 	TUniquePtr<FUsdLevelSequenceHelperImpl> UsdSequencerImpl;
diff --git a/Engine/Plugins/Importers/USDImporter/Source/USDStage/Public/USDStageActor.h b/Engine/Plugins/Importers/USDImporter/Source/USDStage/Public/USDStageActor.h
--- a/Engine/Plugins/Importers/USDImporter/Source/USDStage/Public/USDStageActor.h
+++ b/Engine/Plugins/Importers/USDImporter/Source/USDStage/Public/USDStageActor.h
@@ -188,6 +188,7 @@ public:
 	UUsdAssetCache* GetAssetCache() { return AssetCache; }
 	TSharedPtr<FUsdInfoCache> GetInfoCache() { return InfoCache; }
 	TMap< FString, TMap< FString, int32 > > GetMaterialToPrimvarToUVIndex() { return MaterialToPrimvarToUVIndex; }
+	const UsdUtils::FBlendShapeMap& GetBlendShapeMap() { return BlendShapesByPath; }
 
 public:
 #if WITH_EDITOR
@@ -235,6 +236,7 @@ private:
 	void OnObjectPropertyChanged( UObject* ObjectBeingModified, FPropertyChangedEvent& PropertyChangedEvent );
 	void HandlePropertyChangedEvent( FPropertyChangedEvent& PropertyChangedEvent );
 	bool HasAuthorityOverStage() const;
+	void OnSkelAnimationBaked( const FString& SkelRootPrimPath );
 
 private:
 	// Note: This cannot be instanced: Read the comment in the constructor
diff --git a/Engine/Plugins/Importers/USDImporter/Source/USDStage/USDStage.Build.cs b/Engine/Plugins/Importers/USDImporter/Source/USDStage/USDStage.Build.cs
--- a/Engine/Plugins/Importers/USDImporter/Source/USDStage/USDStage.Build.cs
+++ b/Engine/Plugins/Importers/USDImporter/Source/USDStage/USDStage.Build.cs
@@ -19,7 +19,9 @@ namespace UnrealBuildTool.Rules
 			PrivateDependencyModuleNames.AddRange(
 				new string[]
 				{
+					"AnimationCore",
 					"CinematicCamera",
+					"ControlRig",
 					"Core",
 					"CoreUObject",
 					"Engine",
@@ -44,10 +46,12 @@ namespace UnrealBuildTool.Rules
 				PrivateDependencyModuleNames.AddRange(
 					new string[]
 					{
+						"ControlRigDeveloper",
 						"DeveloperToolSettings",
 						"EditorStyle", // For the font style on the stage actor customization
 						"InputCore", // For keyboard control on the widget in the stage actor customization
 						"LevelSequenceEditor",
+						"MovieSceneTools",
 						"PropertyEditor", // For the stage actor's details customization
 						"Sequencer",
 						"UnrealEd",
diff --git a/Engine/Plugins/Importers/USDImporter/Source/USDStageEditor/Private/SUSDStageTreeView.cpp b/Engine/Plugins/Importers/USDImporter/Source/USDStageEditor/Private/SUSDStageTreeView.cpp
--- a/Engine/Plugins/Importers/USDImporter/Source/USDStageEditor/Private/SUSDStageTreeView.cpp
+++ b/Engine/Plugins/Importers/USDImporter/Source/USDStageEditor/Private/SUSDStageTreeView.cpp
@@ -690,6 +690,30 @@ TSharedPtr< SWidget > SUsdStageTreeView::ConstructPrimContextMenu()
 			NAME_None,
 			EUserInterfaceActionType::Button
 		);
+
+		PrimOptions.AddMenuEntry(
+			LOCTEXT( "SetUpControlRig", "Set up Control Rig" ),
+			LOCTEXT( "SetUpControlRig_ToolTip", "Sets up the generated component for Control Rig integration and store the connection details to the USD Stage" ),
+			FSlateIcon(),
+			FUIAction(
+				FExecuteAction::CreateSP( this, &SUsdStageTreeView::OnSetUpControlRig ),
+				FCanExecuteAction::CreateSP( this, &SUsdStageTreeView::CanSetUpControlRig )
+			),
+			NAME_None,
+			EUserInterfaceActionType::Button
+		);
+
+		PrimOptions.AddMenuEntry(
+			LOCTEXT( "RemoveControlRig", "Remove Control Rig" ),
+			LOCTEXT( "RemoveControlRig_ToolTip", "Reverses the Control Rig configuration on the component and removes the connection details from the USD Stage" ),
+			FSlateIcon(),
+			FUIAction(
+				FExecuteAction::CreateSP( this, &SUsdStageTreeView::OnRemoveControlRig ),
+				FCanExecuteAction::CreateSP( this, &SUsdStageTreeView::CanRemoveControlRig )
+			),
+			NAME_None,
+			EUserInterfaceActionType::Button
+		);
 	}
 	PrimOptions.EndSection();
 
@@ -855,6 +879,44 @@ void SUsdStageTreeView::OnRemoveLiveLink()
 	}
 }
 
+void SUsdStageTreeView::OnSetUpControlRig()
+{
+	if ( !UsdStageActor.IsValid() )
+	{
+		return;
+	}
+
+	FScopedTransaction Transaction( LOCTEXT( "SetUpControlRigTransaction", "Set up Control Rig for selected prims" ) );
+
+	TArray< FUsdPrimViewModelRef > MySelectedItems = GetSelectedItems();
+
+	UE::FSdfChangeBlock Block;
+
+	for ( FUsdPrimViewModelRef SelectedItem : MySelectedItems )
+	{
+		SelectedItem->SetUpControlRig();
+	}
+}
+
+void SUsdStageTreeView::OnRemoveControlRig()
+{
+	if ( !UsdStageActor.IsValid() )
+	{
+		return;
+	}
+
+	FScopedTransaction Transaction( LOCTEXT( "RemoveControlRigTransaction", "Reverses the Live Link configuration for selected prims" ) );
+
+	TArray< FUsdPrimViewModelRef > MySelectedItems = GetSelectedItems();
+
+	UE::FSdfChangeBlock Block;
+
+	for ( FUsdPrimViewModelRef SelectedItem : MySelectedItems )
+	{
+		SelectedItem->RemoveControlRig();
+	}
+}
+
 bool SUsdStageTreeView::CanAddPrim() const
 {
 	if ( !UsdStageActor.IsValid() )
@@ -953,6 +1015,58 @@ bool SUsdStageTreeView::CanRemoveLiveLink() const
 	return false;
 }
 
+bool SUsdStageTreeView::CanSetUpControlRig() const
+{
+	if ( !UsdStageActor.IsValid() )
+	{
+		return false;
+	}
+
+	UE::FUsdStage UsdStage = UsdStageActor->GetOrLoadUsdStage();
+	if ( !UsdStage || !UsdStage.IsEditTargetValid() )
+	{
+		return false;
+	}
+
+	TArray< FUsdPrimViewModelRef > MySelectedItems = GetSelectedItems();
+
+	for ( FUsdPrimViewModelRef SelectedItem : MySelectedItems )
+	{
+		if ( SelectedItem->CanSetUpControlRig() )
+		{
+			return true;
+		}
+	}
+
+	return false;
+}
+
+bool SUsdStageTreeView::CanRemoveControlRig() const
+{
+	if ( !UsdStageActor.IsValid() )
+	{
+		return false;
+	}
+
+	UE::FUsdStage UsdStage = UsdStageActor->GetOrLoadUsdStage();
+	if ( !UsdStage || !UsdStage.IsEditTargetValid() )
+	{
+		return false;
+	}
+
+	TArray< FUsdPrimViewModelRef > MySelectedItems = GetSelectedItems();
+
+	for ( FUsdPrimViewModelRef SelectedItem : MySelectedItems )
+	{
+		if ( SelectedItem->CanRemoveControlRig() )
+		{
+			return true;
+		}
+	}
+
+	return false;
+}
+
 void SUsdStageTreeView::RequestListRefresh()
 {
 	SUsdTreeView< FUsdPrimViewModelRef >::RequestListRefresh();
diff --git a/Engine/Plugins/Importers/USDImporter/Source/USDStageEditor/Private/SUSDStageTreeView.h b/Engine/Plugins/Importers/USDImporter/Source/USDStageEditor/Private/SUSDStageTreeView.h
--- a/Engine/Plugins/Importers/USDImporter/Source/USDStageEditor/Private/SUSDStageTreeView.h
+++ b/Engine/Plugins/Importers/USDImporter/Source/USDStageEditor/Private/SUSDStageTreeView.h
@@ -56,11 +56,15 @@ private:
 	void OnClearReferences();
 	void OnSetUpLiveLink();
 	void OnRemoveLiveLink();
+	void OnSetUpControlRig();
+	void OnRemoveControlRig();
 
 	bool CanAddPrim() const;
 	bool CanExecutePrimAction() const;
 	bool CanSetUpLiveLink() const;
 	bool CanRemoveLiveLink() const;
+	bool CanSetUpControlRig() const;
+	bool CanRemoveControlRig() const;
 
 	/** Uses TreeItemExpansionStates to travel the tree and call SetItemExpansion */
 	void RestoreExpansionStates();
diff --git a/Engine/Plugins/Importers/USDImporter/Source/USDStageEditor/Private/Widgets/SUSDIntegrationsPanel.cpp b/Engine/Plugins/Importers/USDImporter/Source/USDStageEditor/Private/Widgets/SUSDIntegrationsPanel.cpp
--- a/Engine/Plugins/Importers/USDImporter/Source/USDStageEditor/Private/Widgets/SUSDIntegrationsPanel.cpp
+++ b/Engine/Plugins/Importers/USDImporter/Source/USDStageEditor/Private/Widgets/SUSDIntegrationsPanel.cpp
@@ -9,21 +9,25 @@
 #include "USDTypesConversion.h"
 #include "USDValueConversion.h"
 
+#include "UsdWrappers/SdfChangeBlock.h"
 #include "UsdWrappers/SdfPath.h"
 #include "UsdWrappers/UsdAttribute.h"
 #include "UsdWrappers/UsdPrim.h"
 #include "UsdWrappers/UsdStage.h"
 #include "UsdWrappers/VtValue.h"
 
-#include "Styling/AppStyle.h"
 #include "ISinglePropertyView.h"
 #include "Modules/ModuleManager.h"
 #include "PropertyEditorModule.h"
 #include "Roles/LiveLinkAnimationRole.h"
 #include "Roles/LiveLinkTransformRole.h"
+#include "ScopedTransaction.h"
 #include "SLiveLinkSubjectRepresentationPicker.h"
+#include "Styling/AppStyle.h"
 #include "Widgets/Input/SCheckBox.h"
 #include "Widgets/Input/SEditableTextBox.h"
+#include "Widgets/Input/SNumericEntryBox.h"
+#include "Widgets/Input/STextComboBox.h"
 #include "Widgets/Views/SHeaderRow.h"
 
 #define LOCTEXT_NAMESPACE "SUSDIntegrationsPanel"
@@ -55,21 +59,43 @@ TSharedRef< SWidget > SUsdIntegrationsPanelRow::GenerateWidgetForColumn( const F
 
 	FName AttributeName = Attribute->GetName();
 
-	if ( ColumnName == TEXT("PropertyName") )
+	using DisplayTextForPropertiesEntry = TPairInitializer<const FName&, const FText&>;
+	const static TMap<FName, FText> DisplayTextForProperties
+	({
+		DisplayTextForPropertiesEntry( *UsdToUnreal::ConvertToken( UnrealIdentifiers::UnrealAnimBlueprintPath ), LOCTEXT( "AnimBlueprintPathText", "Anim Blueprint asset" ) ),
+		DisplayTextForPropertiesEntry( *UsdToUnreal::ConvertToken( UnrealIdentifiers::UnrealControlRigPath ), LOCTEXT( "ControlRigPathText", "Control Rig asset" ) ),
+		DisplayTextForPropertiesEntry( *UsdToUnreal::ConvertToken( UnrealIdentifiers::UnrealUseFKControlRig ), LOCTEXT( "UseFKControlRigText", "Use FKControlRig" ) ),
+		DisplayTextForPropertiesEntry( *UsdToUnreal::ConvertToken( UnrealIdentifiers::UnrealControlRigReduceKeys ), LOCTEXT( "ControlRigReduceKeysText", "Control Rig key reduction" ) ),
+		DisplayTextForPropertiesEntry( *UsdToUnreal::ConvertToken( UnrealIdentifiers::UnrealControlRigReductionTolerance ), LOCTEXT( "ControlRigReduceToleranceText", "Control Rig key reduction tolerance" ) ),
+		DisplayTextForPropertiesEntry( *UsdToUnreal::ConvertToken( UnrealIdentifiers::UnrealLiveLinkSubjectName ), LOCTEXT( "SubjectNameText", "Live Link subject name" ) ),
+		DisplayTextForPropertiesEntry( *UsdToUnreal::ConvertToken( UnrealIdentifiers::UnrealLiveLinkEnabled ), LOCTEXT( "LiveLinkEnabledText", "Enable LiveLink" ) ),
+	});
+
+	const static TMap<FName, FText> ToolTipTextForProperties
+	({
+		DisplayTextForPropertiesEntry( *UsdToUnreal::ConvertToken( UnrealIdentifiers::UnrealAnimBlueprintPath ), LOCTEXT( "AnimBlueprintPathToolTip", "Anim Blueprint asset to use on the component to connect with Live Link" ) ),
+		DisplayTextForPropertiesEntry( *UsdToUnreal::ConvertToken( UnrealIdentifiers::UnrealControlRigPath ), LOCTEXT( "ControlRigPathToolTip", "Control Rig Blueprint asset to use" ) ),
+		DisplayTextForPropertiesEntry( *UsdToUnreal::ConvertToken( UnrealIdentifiers::UnrealUseFKControlRig ), LOCTEXT( "UseFKControlRigToolTip", "Whether to use a generated FKControlRig instead of the Control Rig Blueprint asset" ) ),
+		DisplayTextForPropertiesEntry( *UsdToUnreal::ConvertToken( UnrealIdentifiers::UnrealControlRigReduceKeys ), LOCTEXT( "ControlRigReduceKeysToolTip", "Whether to enable key reduction when generating Control Rig tracks from the USD animation data" ) ),
+		DisplayTextForPropertiesEntry( *UsdToUnreal::ConvertToken( UnrealIdentifiers::UnrealControlRigReductionTolerance ), LOCTEXT( "ControlRigReduceToleranceToolTip", "Tolerance to use for the key reduction pass" ) ),
+		DisplayTextForPropertiesEntry( *UsdToUnreal::ConvertToken( UnrealIdentifiers::UnrealLiveLinkSubjectName ), LOCTEXT( "SubjectNameToolTip", "Which Live Link subject to use for this component" ) ),
+		DisplayTextForPropertiesEntry( *UsdToUnreal::ConvertToken( UnrealIdentifiers::UnrealLiveLinkEnabled ), LOCTEXT( "LiveLinkEnabledToolTip", "If checked will cause the component to be animated with the pose data coming from the Anim Blueprint. If unchecked will cause the component to be animated with the regular skeleton animation data coming from the USD Stage" ) ),
+	});
+
+	const FText* ToolTipText = ToolTipTextForProperties.Find( AttributeName );
+	if ( !ToolTipText )
 	{
-		using DisplayTextForPropertiesEntry = TPairInitializer<const FName&, const FText&>;
-		const static TMap<FName, FText> DisplayTextForProperties
-		({
-			DisplayTextForPropertiesEntry( *UsdToUnreal::ConvertToken( UnrealIdentifiers::UnrealLiveLinkAnimBlueprintPath ), LOCTEXT( "AnimBlueprintPathText", "AnimBlueprint asset" ) ),
-			DisplayTextForPropertiesEntry( *UsdToUnreal::ConvertToken( UnrealIdentifiers::UnrealLiveLinkSubjectName ), LOCTEXT( "SubjectNameText", "LiveLink subject name" ) ),
-			DisplayTextForPropertiesEntry( *UsdToUnreal::ConvertToken( UnrealIdentifiers::UnrealLiveLinkEnabled ), LOCTEXT( "LiveLinkEnabledText", "Enable LiveLink" ) ),
-		});
+		ToolTipText = &FText::GetEmpty();
+	}
 
+	if ( ColumnName == TEXT("PropertyName") )
+	{
 		if ( const FText* TextToDisplay = DisplayTextForProperties.Find( AttributeName ) )
 		{
 			SAssignNew( ColumnWidget, STextBlock )
 			.Text( *TextToDisplay )
-			.Font( FAppStyle::GetFontStyle( UE::SUsdIntergrationsPanel::Private::NormalFont ) );
+			.Font( FAppStyle::GetFontStyle( UE::SUsdIntergrationsPanel::Private::NormalFont ) )
+			.ToolTipText( *ToolTipText );
 		}
 		else
 		{
@@ -84,7 +110,7 @@ TSharedRef< SWidget > SUsdIntegrationsPanelRow::GenerateWidgetForColumn( const F
 		FName TypeName = Attribute->GetTypeName();
 		TSharedPtr<UE::FUsdAttribute> AttributeCopy = Attribute;
 
-		if ( AttributeName == *UsdToUnreal::ConvertToken( UnrealIdentifiers::UnrealLiveLinkAnimBlueprintPath ) )
+		if ( AttributeName == *UsdToUnreal::ConvertToken( UnrealIdentifiers::UnrealAnimBlueprintPath ) )
 		{
 			if ( UUsdIntegrationsPanelPropertyDummy* Dummy = GetMutableDefault<UUsdIntegrationsPanelPropertyDummy>() )
 			{
@@ -117,18 +143,74 @@ TSharedRef< SWidget > SUsdIntegrationsPanelRow::GenerateWidgetForColumn( const F
 					UUsdIntegrationsPanelPropertyDummy* Dummy = GetMutableDefault<UUsdIntegrationsPanelPropertyDummy>();
 
 					FString Path = ( Dummy && Dummy->AnimBPProperty ) ? Dummy->AnimBPProperty->GetPathName() : FString{};
+					std::string UsdPath = UnrealToUsd::ConvertString(*Path).Get();
+
+					// We need transactions for these because the notices emitted for them are upgraded to resyncs by the stage
+					// actor, which means they may generate new assets and components
+					FScopedTransaction Transaction(
+						FText::Format( LOCTEXT( "AnimBluepringTransaction", "Changed AnimBluepring path to '{0}'" ),
+							FText::FromString( Path )
+						)
+					);
+
+					UE::FVtValue Value;
+					if ( UsdUtils::SetUnderlyingValue( Value, UsdPath ) && !Value.IsEmpty() )
+					{
+						AttributeCopy->Set( Value );
+					}
+				} );
+				PropertyView->SetOnPropertyValueChanged( PropertyChanged );
+
+				ColumnWidget = PropertyView.ToSharedRef();
+			}
+			else
+			{
+				ColumnWidget = SNullWidget::NullWidget;
+			}
+		}
+		else if ( AttributeName == *UsdToUnreal::ConvertToken( UnrealIdentifiers::UnrealControlRigPath ) )
+		{
+			if ( UUsdIntegrationsPanelPropertyDummy* Dummy = GetMutableDefault<UUsdIntegrationsPanelPropertyDummy>() )
+			{
+				// Let the object picker row be as tall as it wants
+				RowHeight = FOptionalSize();
+
+				FSinglePropertyParams Params;
+				Params.Font = FAppStyle::GetFontStyle( UE::SUsdIntergrationsPanel::Private::NormalFont );
+				Params.NamePlacement = EPropertyNamePlacement::Hidden;
 
-					TArray< UsdUtils::FConvertedVtValueComponent > Components;
-					Components.Add( UsdUtils::FConvertedVtValueComponent( TInPlaceType<FString>(), Path ) );
+				UE::FVtValue Value;
+				if ( AttributeCopy->Get( Value ) && !Value.IsEmpty() )
+				{
+					Dummy->ControlRigProperty = FSoftObjectPath( UsdUtils::Stringify( Value ) ).TryLoad();
+				}
+				else
+				{
+					Dummy->ControlRigProperty = nullptr;
+				}
 
-					UsdUtils::FConvertedVtValue Converted;
-					Converted.SourceType = UsdUtils::EUsdBasicDataTypes::String;
-					Converted.bIsArrayValued = false;
-					Converted.bIsEmpty = false;
-					Converted.Entries.Add( Components );
+				FPropertyEditorModule& PropertyEditor = FModuleManager::LoadModuleChecked<FPropertyEditorModule>( TEXT( "PropertyEditor" ) );
+				TSharedPtr<class ISinglePropertyView> PropertyView = PropertyEditor.CreateSingleProperty(
+					Dummy,
+					GET_MEMBER_NAME_CHECKED( UUsdIntegrationsPanelPropertyDummy, ControlRigProperty ),
+					Params
+				);
+
+				FSimpleDelegate PropertyChanged = FSimpleDelegate::CreateLambda( [AttributeCopy]()
+				{
+					UUsdIntegrationsPanelPropertyDummy* Dummy = GetMutableDefault<UUsdIntegrationsPanelPropertyDummy>();
+
+					FString Path = ( Dummy && Dummy->ControlRigProperty ) ? Dummy->ControlRigProperty->GetPathName() : FString{};
+					std::string UsdPath = UnrealToUsd::ConvertString( *Path ).Get();
+
+					FScopedTransaction Transaction(
+						FText::Format( LOCTEXT( "ControlRigPathTransaction", "Changed ControlRig blueprint path to '{0}'" ),
+							FText::FromString( Path )
+						)
+					);
 
 					UE::FVtValue Value;
-					if ( UnrealToUsd::ConvertValue( Converted, Value ) && !Value.IsEmpty() )
+					if ( UsdUtils::SetUnderlyingValue( Value, UsdPath ) && !Value.IsEmpty() )
 					{
 						AttributeCopy->Set( Value );
 					}
@@ -136,12 +218,99 @@ TSharedRef< SWidget > SUsdIntegrationsPanelRow::GenerateWidgetForColumn( const F
 				PropertyView->SetOnPropertyValueChanged( PropertyChanged );
 
 				ColumnWidget = PropertyView.ToSharedRef();
+
+				// Disable the widget if we're using an FKControlRig instead
+				if ( UE::FUsdAttribute FKEnabledAttr = Attribute->GetPrim().GetAttribute( *UsdToUnreal::ConvertToken( UnrealIdentifiers::UnrealUseFKControlRig ) ) )
+				{
+					UE::FVtValue FKEnabledValue;
+					if ( FKEnabledAttr.Get( FKEnabledValue ) )
+					{
+						const bool bDefaultValue = false;
+						if ( UsdUtils::GetUnderlyingValue<bool>( FKEnabledValue ).Get( bDefaultValue ) )
+						{
+							ColumnWidget->SetEnabled( false );
+						}
+					}
+				}
 			}
 			else
 			{
 				ColumnWidget = SNullWidget::NullWidget;
 			}
 		}
+		else if ( AttributeName == *UsdToUnreal::ConvertToken( UnrealIdentifiers::UnrealControlRigReduceKeys ) )
+		{
+			SAssignNew( ColumnWidget, SBox )
+			.HeightOverride( FUsdStageEditorStyle::Get()->GetFloat( "UsdStageEditor.ListItemHeight" ) )
+			.VAlign( VAlign_Center )
+			[
+				SNew(SCheckBox)
+				.IsChecked_Lambda( [AttributeCopy]() ->ECheckBoxState
+				{
+					UE::FVtValue Value;
+					if ( AttributeCopy->Get( Value ) )
+					{
+						if ( TOptional<bool> UnderlyingValue = UsdUtils::GetUnderlyingValue<bool>( Value ) )
+						{
+							return UnderlyingValue.GetValue() ? ECheckBoxState::Checked : ECheckBoxState::Unchecked;
+						}
+					}
+
+					return ECheckBoxState::Undetermined;
+				})
+				.OnCheckStateChanged_Lambda( [AttributeCopy]( ECheckBoxState NewValue )
+				{
+					FScopedTransaction Transaction(
+						FText::Format( LOCTEXT( "ControlRigReduceKeysTransaction", "Changed ControlRig reduce keys property to '{0}'" ),
+							FText::FromString( NewValue == ECheckBoxState::Checked ? TEXT( "true" ) : TEXT( "false" ) )
+						)
+					);
+
+					UE::FVtValue Value;
+					if ( UsdUtils::SetUnderlyingValue( Value, NewValue == ECheckBoxState::Checked ) && !Value.IsEmpty() )
+					{
+						AttributeCopy->Set( Value );
+					}
+				})
+			];
+		}
+		else if ( AttributeName == *UsdToUnreal::ConvertToken( UnrealIdentifiers::UnrealControlRigReductionTolerance ) )
+		{
+			SAssignNew( ColumnWidget, SBox )
+			.HeightOverride( FUsdStageEditorStyle::Get()->GetFloat( "UsdStageEditor.ListItemHeight" ) )
+			.VAlign( VAlign_Center )
+			[
+				SNew(SNumericEntryBox<float>)
+				.Font( FAppStyle::GetFontStyle( UE::SUsdIntergrationsPanel::Private::NormalFont ) )
+				.OnValueCommitted_Lambda( [AttributeCopy]( float NewValue, ETextCommit::Type CommitType )
+				{
+					FScopedTransaction Transaction(
+						FText::Format( LOCTEXT( "ControlRigReduceToleranceTransaction", "Changed ControlRig reduce tolerance to '{0}'" ),
+							NewValue
+						)
+					);
+
+					UE::FVtValue Value;
+					if ( UsdUtils::SetUnderlyingValue( Value, NewValue ) && !Value.IsEmpty() )
+					{
+						AttributeCopy->Set( Value );
+					}
+				})
+				.Value_Lambda( [AttributeCopy]() -> float
+				{
+					UE::FVtValue Value;
+					if ( AttributeCopy->Get( Value ) )
+					{
+						if ( TOptional<float> UnderlyingValue = UsdUtils::GetUnderlyingValue<float>( Value ) )
+						{
+							return UnderlyingValue.GetValue();
+						}
+					}
+
+					return 0.0f;
+				})
+			];
+		}
 		else if ( AttributeName == *UsdToUnreal::ConvertToken( UnrealIdentifiers::UnrealLiveLinkSubjectName ) )
 		{
 			// SLiveLinkSubjectRepresentationPicker is also a little bit larger than one of our rows, so let it expand a bit
@@ -159,31 +328,42 @@ TSharedRef< SWidget > SUsdIntegrationsPanelRow::GenerateWidgetForColumn( const F
 					? ULiveLinkAnimationRole::StaticClass()
 					: ULiveLinkTransformRole::StaticClass();
 
+				FString SubjectName;
+
 				UE::FVtValue Value;
-				AttributeCopy->Get( Value );
-				Result.Subject = FLiveLinkSubjectName{ FName{*UsdUtils::Stringify( Value )} };
+				if ( AttributeCopy->Get( Value ) )
+				{
+					if ( TOptional<std::string> UnderlyingValue = UsdUtils::GetUnderlyingValue<std::string>( Value ) )
+					{
+						SubjectName = UsdToUnreal::ConvertString( UnderlyingValue.GetValue() );
+					}
+				}
+
+				Result.Subject = FLiveLinkSubjectName{ FName{*SubjectName} };
 
 				return Result;
 			})
 			.OnValueChanged_Lambda( [AttributeCopy]( SLiveLinkSubjectRepresentationPicker::FLiveLinkSourceSubjectRole NewValue )
 			{
-				TArray< UsdUtils::FConvertedVtValueComponent > Components;
-				Components.Add( UsdUtils::FConvertedVtValueComponent( TInPlaceType<FString>(), NewValue.Subject.ToString() ) );
+				FScopedUsdAllocs UsdAllocs;
 
-				UsdUtils::FConvertedVtValue Converted;
-				Converted.SourceType = UsdUtils::EUsdBasicDataTypes::String;
-				Converted.bIsArrayValued = false;
-				Converted.bIsEmpty = false;
-				Converted.Entries.Add( Components );
+				std::string UsdValue = UnrealToUsd::ConvertString( *( NewValue.Subject.ToString() ) ).Get();
+
+				FScopedTransaction Transaction(
+					FText::Format( LOCTEXT( "SubjectNameTransaction", "Changed Live Link subject name to '{0}'" ),
+						FText::FromString( *( NewValue.Subject.ToString() ) )
+					)
+				);
 
 				UE::FVtValue Value;
-				if ( UnrealToUsd::ConvertValue( Converted, Value ) && !Value.IsEmpty() )
+				if ( UsdUtils::SetUnderlyingValue( Value, UsdValue ) && !Value.IsEmpty() )
 				{
 					AttributeCopy->Set( Value );
 				}
 			});
 		}
-		else if ( AttributeName == *UsdToUnreal::ConvertToken( UnrealIdentifiers::UnrealLiveLinkEnabled ) )
+		else if ( AttributeName == *UsdToUnreal::ConvertToken( UnrealIdentifiers::UnrealLiveLinkEnabled )
+				|| AttributeName == *UsdToUnreal::ConvertToken( UnrealIdentifiers::UnrealUseFKControlRig ) )
 		{
 			SAssignNew( ColumnWidget, SBox )
 			.HeightOverride( FUsdStageEditorStyle::Get()->GetFloat( "UsdStageEditor.ListItemHeight" ) )
@@ -193,33 +373,27 @@ TSharedRef< SWidget > SUsdIntegrationsPanelRow::GenerateWidgetForColumn( const F
 				.IsChecked_Lambda( [AttributeCopy]()
 				{
 					UE::FVtValue Value;
-					AttributeCopy->Get( Value );
-
-					UsdUtils::FConvertedVtValue Converted;
-
-					if ( UsdToUnreal::ConvertValue( Value, Converted ) && Converted.SourceType == UsdUtils::EUsdBasicDataTypes::Bool && !Converted.bIsEmpty )
+					if ( AttributeCopy->Get( Value ) )
 					{
-						if ( bool* bBoolValue = Converted.Entries[ 0 ][ 0 ].TryGet<bool>() )
+						if ( TOptional<bool> UnderlyingValue = UsdUtils::GetUnderlyingValue<bool>( Value ) )
 						{
-							return *bBoolValue ? ECheckBoxState::Checked : ECheckBoxState::Unchecked;
+							return UnderlyingValue.GetValue() ? ECheckBoxState::Checked : ECheckBoxState::Unchecked;
 						}
 					}
 
 					return ECheckBoxState::Undetermined;
 				})
-				.OnCheckStateChanged_Lambda( [AttributeCopy]( ECheckBoxState NewState )
+				.OnCheckStateChanged_Lambda( [AttributeCopy, AttributeName]( ECheckBoxState NewState )
 				{
-					TArray< UsdUtils::FConvertedVtValueComponent > Components;
-					Components.Add( UsdUtils::FConvertedVtValueComponent( TInPlaceType<bool>(), NewState == ECheckBoxState::Checked ) );
-
-					UsdUtils::FConvertedVtValue Converted;
-					Converted.SourceType = UsdUtils::EUsdBasicDataTypes::Bool;
-					Converted.bIsArrayValued = false;
-					Converted.bIsEmpty = false;
-					Converted.Entries.Add( Components );
+					FScopedTransaction Transaction(
+						FText::Format( LOCTEXT( "IntegrationBoolTransaction", "Changed attribute '{0}' to " ),
+							FText::FromName( AttributeName ),
+							FText::FromString( NewState == ECheckBoxState::Checked ? TEXT( "true" ) : TEXT( "false" ) )
+						)
+					);
 
 					UE::FVtValue Value;
-					if ( UnrealToUsd::ConvertValue( Converted, Value ) && !Value.IsEmpty() )
+					if ( UsdUtils::SetUnderlyingValue( Value, NewState == ECheckBoxState::Checked ) && !Value.IsEmpty() )
 					{
 						AttributeCopy->Set( Value );
 					}
@@ -277,39 +451,59 @@ void SUsdIntegrationsPanel::SetPrimPath( const UE::FUsdStageWeak& InUsdStage, co
 
 	Attributes.Reset();
 
-	EVisibility NewVisibility = EVisibility::Collapsed;
+	EVisibility PanelVisibility = EVisibility::Collapsed;
+
+	// Important as this function may be reentrant in case our ApplyXSchema calls create attributes
+	UE::FSdfChangeBlock Block;
 
 	if ( InUsdStage )
 	{
 		if ( UE::FUsdPrim Prim = InUsdStage.GetPrimAtPath( UE::FSdfPath{ InPrimPath } ) )
 		{
-			if ( UsdUtils::PrimHasLiveLinkSchema( Prim ) )
+			if ( !Prim.IsPseudoRoot() )
 			{
-				NewVisibility = EVisibility::Visible;
+				const bool bHasLiveLink = UsdUtils::PrimHasLiveLinkSchema( Prim );
+				const bool bHasControlRig = UsdUtils::PrimHasControlRigSchema( Prim );
+				const bool bIsSkelRoot = Prim.IsA( TEXT( "SkelRoot" ) );
 
-				// Only show the dedicated AnimBlueprint picker for SkelRoots
-				if ( Prim.IsA( TEXT( "SkelRoot" ) ) )
+				TArray<FString> AttributeNames;
+
+				if ( bHasLiveLink )
 				{
-					if ( UE::FUsdAttribute Attr = Prim.GetAttribute( *UsdToUnreal::ConvertToken( UnrealIdentifiers::UnrealLiveLinkAnimBlueprintPath ) ) )
+					if ( bIsSkelRoot )
 					{
-						Attributes.Add( MakeShared<UE::FUsdAttribute>( Attr ) );
+						AttributeNames.Add( *UsdToUnreal::ConvertToken( UnrealIdentifiers::UnrealAnimBlueprintPath ) );
 					}
+
+					AttributeNames.Add( *UsdToUnreal::ConvertToken( UnrealIdentifiers::UnrealLiveLinkSubjectName ) );
+					AttributeNames.Add( *UsdToUnreal::ConvertToken( UnrealIdentifiers::UnrealLiveLinkEnabled ) );
+				}
+
+				if( bHasControlRig && bIsSkelRoot )
+				{
+					AttributeNames.Add( *UsdToUnreal::ConvertToken( UnrealIdentifiers::UnrealControlRigPath ) );
+					AttributeNames.Add( *UsdToUnreal::ConvertToken( UnrealIdentifiers::UnrealUseFKControlRig ) );
+					AttributeNames.Add( *UsdToUnreal::ConvertToken( UnrealIdentifiers::UnrealControlRigReduceKeys ) );
+					AttributeNames.Add( *UsdToUnreal::ConvertToken( UnrealIdentifiers::UnrealControlRigReductionTolerance ) );
 				}
 
-				if ( UE::FUsdAttribute Attr = Prim.GetAttribute( *UsdToUnreal::ConvertToken( UnrealIdentifiers::UnrealLiveLinkSubjectName ) ) )
+				if ( bHasLiveLink || bHasControlRig )
 				{
-					Attributes.Add( MakeShared<UE::FUsdAttribute>( Attr ) );
+					PanelVisibility = EVisibility::Visible;
 				}
 
-				if ( UE::FUsdAttribute Attr = Prim.GetAttribute( *UsdToUnreal::ConvertToken( UnrealIdentifiers::UnrealLiveLinkEnabled ) ) )
+				for ( const FString& AttributeName : AttributeNames )
 				{
-					Attributes.Add( MakeShared<UE::FUsdAttribute>( Attr ) );
+					if ( UE::FUsdAttribute Attr = Prim.GetAttribute( *AttributeName ) )
+					{
+						Attributes.Add( MakeShared<UE::FUsdAttribute>( Attr ) );
+					}
 				}
 			}
 		}
 	}
 
-	SetVisibility( NewVisibility );
+	SetVisibility( PanelVisibility );
 	RequestListRefresh();
 }
 
diff --git a/Engine/Plugins/Importers/USDImporter/Source/USDStageEditor/Private/Widgets/SUSDIntegrationsPanel.h b/Engine/Plugins/Importers/USDImporter/Source/USDStageEditor/Private/Widgets/SUSDIntegrationsPanel.h
--- a/Engine/Plugins/Importers/USDImporter/Source/USDStageEditor/Private/Widgets/SUSDIntegrationsPanel.h
+++ b/Engine/Plugins/Importers/USDImporter/Source/USDStageEditor/Private/Widgets/SUSDIntegrationsPanel.h
@@ -22,10 +22,13 @@ class UUsdIntegrationsPanelPropertyDummy : public UObject
 public:
 	GENERATED_BODY()
 
-	// Ideally this would be an FSoftObjectPath, but FPropertyEditorModule doesn't support generating
+	// Ideally these would be FSoftObjectPaths, but FPropertyEditorModule doesn't support generating
 	// widgets for those properties yet
-	UPROPERTY(EditAnywhere, Transient, Category=Dummy, meta=(AllowedClasses="/Script/Engine.AnimBlueprint" ))
+	UPROPERTY( EditAnywhere, Transient, Category = Dummy, meta = ( AllowedClasses = "/Script/Engine.AnimBlueprint" ) )
 	UObject* AnimBPProperty;
+
+	UPROPERTY( EditAnywhere, Transient, Category = Dummy, meta = ( AllowedClasses = "/Script/ControlRigDeveloper.ControlRigBlueprint" ) )
+	UObject* ControlRigProperty;
 };
 
 #if USE_USD_SDK
diff --git a/Engine/Plugins/Importers/USDImporter/Source/USDStageEditorViewModels/Private/USDPrimViewModel.cpp b/Engine/Plugins/Importers/USDImporter/Source/USDStageEditorViewModels/Private/USDPrimViewModel.cpp
--- a/Engine/Plugins/Importers/USDImporter/Source/USDStageEditorViewModels/Private/USDPrimViewModel.cpp
+++ b/Engine/Plugins/Importers/USDImporter/Source/USDStageEditorViewModels/Private/USDPrimViewModel.cpp
@@ -3,6 +3,7 @@
 #include "USDPrimViewModel.h"
 
 #include "USDConversionUtils.h"
+#include "USDIntegrationUtils.h"
 #include "USDLog.h"
 #include "USDMemory.h"
 #include "USDTypesConversion.h"
@@ -24,11 +25,97 @@
 	#include "pxr/usd/usd/modelAPI.h"
 	#include "pxr/usd/usd/prim.h"
 	#include "pxr/usd/usd/references.h"
+	#include "pxr/usd/usd/tokens.h"
 	#include "pxr/usd/usdGeom/xform.h"
 #include "USDIncludesEnd.h"
 
 #endif // #if USE_USD_SDK
 
+namespace UE::USDPrimViewModel::Private
+{
+#if USE_USD_SDK
+	bool CanAddOrRemoveCustomSchema( pxr::UsdPrim Prim, pxr::TfToken SchemaName, bool bAddSchema )
+	{
+		if ( !Prim )
+		{
+			return false;
+		}
+
+		FScopedUsdAllocs UsdAllocs;
+
+		pxr::TfType Schema = pxr::UsdSchemaRegistry::GetTypeFromSchemaTypeName( SchemaName );
+		if ( !ensure( static_cast< bool >( Schema ) ) )
+		{
+			return false;
+		}
+
+		// Check if the schema is compatible with this prim
+		if ( bAddSchema && !Prim.CanApplyAPI( Schema ) )
+		{
+			return false;
+		}
+
+		pxr::UsdStageRefPtr Stage = Prim.GetStage();
+		pxr::SdfLayerRefPtr EditTarget = Stage->GetEditTarget().GetLayer();
+		if ( !Stage || !EditTarget )
+		{
+			return false;
+		}
+
+		bool bAlreadyHasSchema = false;
+		if ( EditTarget == Stage->GetRootLayer() )
+		{
+			bAlreadyHasSchema = Prim.HasAPI( Schema );
+		}
+		else
+		{
+			// In the future we'll have better layer editing facilities. For now, lets make sure that
+			// we only show that we can setup/remove a schema for a prim if it can be done *at that
+			// particular edit target*, which is what will happen when clicking the button anyway.
+			// For example, a user could:
+			//  - Add the schema on a sublayer;
+			//  - Remove the schema on the root layer;
+			//  - Select the sublayer as the edit target;
+			// If we always just queried the existing stage, the user would see that he can add the schema
+			// (because the composed prim does not have the schema) but clicking the button would do nothing,
+			// as the sublayer prim spec already has the schema...
+
+			pxr::SdfPrimSpecHandle PrimSpecOnLayer = EditTarget->GetPrimAtPath( Prim.GetPath() );
+			if( !PrimSpecOnLayer )
+			{
+				// We can always just make a new 'over' on this layer
+				return true;
+			}
+
+			// We can only add the schema if its not present on this layer, but we should let USD sort the
+			// Ops out because we could have Ops to add/delete/prepend/append the same schema on the same prim...
+			pxr::SdfListOp<pxr::TfToken> Ops = PrimSpecOnLayer->GetInfo( pxr::UsdTokens->apiSchemas ).Get<pxr::SdfTokenListOp>();
+			std::vector<pxr::TfToken> AppliedOps;
+			Ops.ApplyOperations(&AppliedOps);
+			bAlreadyHasSchema = std::find( AppliedOps.begin(), AppliedOps.end(), SchemaName ) != AppliedOps.end();
+		}
+
+		return bAddSchema != bAlreadyHasSchema;
+	}
+
+	void RemoveCustomSchema( pxr::UsdPrim Prim, pxr::TfToken SchemaName )
+	{
+		if ( !Prim )
+		{
+			return;
+		}
+
+		FScopedUsdAllocs UsdAllocs;
+
+		pxr::TfType Schema = pxr::UsdSchemaRegistry::GetTypeFromSchemaTypeName( SchemaName );
+		ensure( static_cast< bool >( Schema ) );
+
+		ensure( Prim.RemoveAPI( Schema ) );
+	}
+
+#endif // #if USE_USD_SDK
+}
+
 FUsdPrimViewModel::FUsdPrimViewModel( FUsdPrimViewModel* InParentItem, const UE::FUsdStageWeak& InUsdStage, const UE::FUsdPrim& InUsdPrim )
 	: FUsdPrimViewModel( InParentItem, InUsdStage )
 {
@@ -240,74 +327,17 @@ void FUsdPrimViewModel::TogglePayload()
 void FUsdPrimViewModel::SetUpLiveLink()
 {
 #if USE_USD_SDK
-
-	FScopedUsdAllocs Allocs;
-
-	pxr::UsdPrim Prim{ UsdPrim };
-	if ( !Prim )
-	{
-		return;
-	}
-
-	pxr::TfType Schema = pxr::UsdSchemaRegistry::GetTypeFromSchemaTypeName( UnrealIdentifiers::LiveLinkAPI );
-	ensure( static_cast< bool >( Schema ) );
-
-	if ( Prim.HasAPI( Schema ) )
-	{
-		return;
-	}
-
-	ensure( Prim.ApplyAPI( Schema ) );
-
-	// It seems codeless schemas can't have fallback values, so just create the attributes here, and set some
-	// sensible defaults
-	const bool bCustom = false;
-
-	if ( pxr::UsdAttribute SubjectNameAttr = Prim.CreateAttribute( UnrealIdentifiers::UnrealLiveLinkSubjectName, pxr::SdfValueTypeNames->String, bCustom, pxr::SdfVariabilityUniform ) )
-	{
-		if ( !SubjectNameAttr.HasAuthoredValue() )
-		{
-			SubjectNameAttr.Set( UnrealToUsd::ConvertString( TEXT( "root" ) ).Get(), pxr::UsdTimeCode::Default() );
-		}
-	}
-
-	if ( pxr::UsdAttribute AnimBPPathAttr = Prim.CreateAttribute( UnrealIdentifiers::UnrealLiveLinkAnimBlueprintPath, pxr::SdfValueTypeNames->String, bCustom, pxr::SdfVariabilityUniform ) )
-	{
-		if ( !AnimBPPathAttr.HasAuthoredValue() )
-		{
-			AnimBPPathAttr.Set( UnrealToUsd::ConvertString( TEXT( "/USDImporter/Blueprint/DefaultLiveLinkAnimBP.DefaultLiveLinkAnimBP" ) ).Get(), pxr::UsdTimeCode::Default() );
-		}
-	}
-
-	if ( pxr::UsdAttribute EnabledAttr = Prim.CreateAttribute( UnrealIdentifiers::UnrealLiveLinkEnabled, pxr::SdfValueTypeNames->Bool, bCustom, pxr::SdfVariabilityUniform ) )
-	{
-		if ( !EnabledAttr.HasAuthoredValue() )
-		{
-			EnabledAttr.Set( true, pxr::UsdTimeCode::Default() );
-		}
-	}
-
+	UsdUtils::ApplyLiveLinkSchema( UsdPrim );
 #endif // #if USE_USD_SDK
 }
 
 bool FUsdPrimViewModel::CanSetUpLiveLink() const
 {
 #if USE_USD_SDK
-	FScopedUsdAllocs UsdAllocs;
-
-	pxr::UsdPrim Prim{ UsdPrim };
-	if ( !Prim )
-	{
-		return false;
-	}
-
-	pxr::TfType Schema = pxr::UsdSchemaRegistry::GetTypeFromSchemaTypeName( UnrealIdentifiers::LiveLinkAPI );
-	if ( !ensure( static_cast< bool >( Schema ) ) )
-	{
-		return false;
-	}
-
-	return !Prim.HasAPI( Schema ) && Prim.CanApplyAPI( Schema );
+	const bool bAddSchema = true;
+	return UsdPrim
+		&& !UsdPrim.IsPseudoRoot()
+		&& UE::USDPrimViewModel::Private::CanAddOrRemoveCustomSchema( UsdPrim, UnrealIdentifiers::LiveLinkAPI, bAddSchema );
 #else
 	return false;
 #endif // #if USE_USD_SDK
@@ -316,33 +346,55 @@ bool FUsdPrimViewModel::CanSetUpLiveLink() const
 void FUsdPrimViewModel::RemoveLiveLink()
 {
 #if USE_USD_SDK
-	FScopedUsdAllocs UsdAllocs;
-
-	pxr::TfType Schema = pxr::UsdSchemaRegistry::GetTypeFromSchemaTypeName( UnrealIdentifiers::LiveLinkAPI );
-	ensure( static_cast< bool >( Schema ) );
-
-	ensure( pxr::UsdPrim( UsdPrim ).RemoveAPI( Schema ) );
+	UE::USDPrimViewModel::Private::RemoveCustomSchema( UsdPrim, UnrealIdentifiers::LiveLinkAPI );
 #endif // #if USE_USD_SDK
 }
 
 bool FUsdPrimViewModel::CanRemoveLiveLink() const
 {
 #if USE_USD_SDK
-	FScopedUsdAllocs UsdAllocs;
+	const bool bAddSchema = false;
+	return UsdPrim
+		&& !UsdPrim.IsPseudoRoot()
+		&& UE::USDPrimViewModel::Private::CanAddOrRemoveCustomSchema( UsdPrim, UnrealIdentifiers::LiveLinkAPI, bAddSchema );
+#else
+	return false;
+#endif // #if USE_USD_SDK
+}
 
-	pxr::UsdPrim Prim{ UsdPrim };
-	if ( !Prim )
-	{
-		return false;
-	}
+void FUsdPrimViewModel::SetUpControlRig()
+{
+#if USE_USD_SDK
+	UsdUtils::ApplyControlRigSchema( UsdPrim );
+#endif // #if USE_USD_SDK
+}
 
-	pxr::TfType Schema = pxr::UsdSchemaRegistry::GetTypeFromSchemaTypeName( UnrealIdentifiers::LiveLinkAPI );
-	if ( !ensure( static_cast< bool >( Schema ) ) )
-	{
-		return false;
-	}
+bool FUsdPrimViewModel::CanSetUpControlRig() const
+{
+#if USE_USD_SDK
+	const bool bAddSchema = true;
+	return UsdPrim
+		&& !UsdPrim.IsPseudoRoot()
+		&& UE::USDPrimViewModel::Private::CanAddOrRemoveCustomSchema( UsdPrim, UnrealIdentifiers::ControlRigAPI, bAddSchema );
+#else
+	return false;
+#endif // #if USE_USD_SDK
+}
 
-	return Prim.HasAPI( Schema );
+void FUsdPrimViewModel::RemoveControlRig()
+{
+#if USE_USD_SDK
+	UE::USDPrimViewModel::Private::RemoveCustomSchema( UsdPrim, UnrealIdentifiers::ControlRigAPI );
+#endif // #if USE_USD_SDK
+}
+
+bool FUsdPrimViewModel::CanRemoveControlRig() const
+{
+#if USE_USD_SDK
+	const bool bAddSchema = false;
+	return UsdPrim
+		&& !UsdPrim.IsPseudoRoot()
+		&& UE::USDPrimViewModel::Private::CanAddOrRemoveCustomSchema( UsdPrim, UnrealIdentifiers::ControlRigAPI, bAddSchema );
 #else
 	return false;
 #endif // #if USE_USD_SDK
diff --git a/Engine/Plugins/Importers/USDImporter/Source/USDStageEditorViewModels/Public/USDPrimViewModel.h b/Engine/Plugins/Importers/USDImporter/Source/USDStageEditorViewModels/Public/USDPrimViewModel.h
--- a/Engine/Plugins/Importers/USDImporter/Source/USDStageEditorViewModels/Public/USDPrimViewModel.h
+++ b/Engine/Plugins/Importers/USDImporter/Source/USDStageEditorViewModels/Public/USDPrimViewModel.h
@@ -54,6 +54,12 @@ public:
 	void RemoveLiveLink();
 	bool CanRemoveLiveLink() const;
 
+	void SetUpControlRig();
+	bool CanSetUpControlRig() const;
+
+	void RemoveControlRig();
+	bool CanRemoveControlRig() const;
+
 	void DefinePrim( const TCHAR* PrimName );
 
 	/** Delegate for hooking up an inline editable text block to be notified that a rename is requested. */
diff --git a/Engine/Plugins/Importers/USDImporter/Source/USDUtilities/Private/USDIntegrationUtils.cpp b/Engine/Plugins/Importers/USDImporter/Source/USDUtilities/Private/USDIntegrationUtils.cpp
--- a/Engine/Plugins/Importers/USDImporter/Source/USDUtilities/Private/USDIntegrationUtils.cpp
+++ b/Engine/Plugins/Importers/USDImporter/Source/USDUtilities/Private/USDIntegrationUtils.cpp
@@ -31,6 +31,51 @@ bool UsdUtils::PrimHasLiveLinkSchema( const pxr::UsdPrim& Prim )
     return Prim.HasAPI( Schema );
 }
 
+bool UsdUtils::PrimHasControlRigSchema( const pxr::UsdPrim& Prim )
+{
+	if ( !Prim )
+	{
+		return false;
+	}
+
+	FScopedUsdAllocs Allocs;
+
+	pxr::TfType Schema = pxr::UsdSchemaRegistry::GetTypeFromSchemaTypeName( UnrealIdentifiers::ControlRigAPI );
+	return Prim.HasAPI( Schema );
+}
+
+bool UsdUtils::ApplyControlRigSchema( const pxr::UsdPrim& Prim )
+{
+	if ( !Prim )
+	{
+		return false;
+	}
+
+	FScopedUsdAllocs Allocs;
+
+	pxr::TfType Schema = pxr::UsdSchemaRegistry::GetTypeFromSchemaTypeName( UnrealIdentifiers::ControlRigAPI );
+	ensure( static_cast< bool >( Schema ) );
+
+	ensure( Prim.ApplyAPI( Schema ) );
+	return true;
+}
+
+bool UsdUtils::ApplyLiveLinkSchema( const pxr::UsdPrim& Prim )
+{
+	if ( !Prim )
+	{
+		return false;
+	}
+
+	FScopedUsdAllocs Allocs;
+
+	pxr::TfType Schema = pxr::UsdSchemaRegistry::GetTypeFromSchemaTypeName( UnrealIdentifiers::LiveLinkAPI );
+	ensure( static_cast< bool >( Schema ) );
+
+	ensure( Prim.ApplyAPI( Schema ) );
+	return true;
+}
+
 void UnrealToUsd::ConvertLiveLinkProperties( const UActorComponent* InComponent, pxr::UsdPrim& InOutPrim )
 {
 #if WITH_EDITOR
@@ -46,7 +91,7 @@ void UnrealToUsd::ConvertLiveLinkProperties( const UActorComponent* InComponent,
 	// Skeletal LiveLink case
 	if ( const USkeletalMeshComponent* SkeletalComponent = Cast<USkeletalMeshComponent>( InComponent ) )
 	{
-		if ( pxr::UsdAttribute Attr = InOutPrim.CreateAttribute( UnrealIdentifiers::UnrealLiveLinkAnimBlueprintPath, pxr::SdfValueTypeNames->String ) )
+		if ( pxr::UsdAttribute Attr = InOutPrim.CreateAttribute( UnrealIdentifiers::UnrealAnimBlueprintPath, pxr::SdfValueTypeNames->String ) )
 		{
 			FString AnimBPPath = SkeletalComponent->AnimClass && SkeletalComponent->AnimClass->ClassGeneratedBy
 				? SkeletalComponent->AnimClass->ClassGeneratedBy->GetPathName()
@@ -78,4 +123,4 @@ void UnrealToUsd::ConvertLiveLinkProperties( const UActorComponent* InComponent,
 }
 
 #undef LOCTEXT_NAMESPACE
-#endif // USE_USD_SDK
\ No newline at end of file
+#endif // USE_USD_SDK
diff --git a/Engine/Plugins/Importers/USDImporter/Source/USDUtilities/Private/USDSkeletalDataConversion.cpp b/Engine/Plugins/Importers/USDImporter/Source/USDUtilities/Private/USDSkeletalDataConversion.cpp
--- a/Engine/Plugins/Importers/USDImporter/Source/USDUtilities/Private/USDSkeletalDataConversion.cpp
+++ b/Engine/Plugins/Importers/USDImporter/Source/USDUtilities/Private/USDSkeletalDataConversion.cpp
@@ -20,16 +20,25 @@
 #include "Animation/AnimCurveTypes.h"
 #include "AnimationRuntime.h"
 #include "AnimEncoding.h"
+#include "ControlRig.h"
+#include "Evaluation/MovieSceneSequenceTransform.h"
+#include "IMovieScenePlayer.h"
 #include "Interfaces/ITargetPlatform.h"
 #include "Interfaces/ITargetPlatformManagerModule.h"
+#include "LevelSequencePlayer.h"
 #include "Misc/CoreMisc.h"
 #include "Misc/MemStack.h"
 #include "Modules/ModuleManager.h"
+#include "MovieScene.h"
+#include "MovieSceneTimeHelpers.h"
+#include "MovieSceneTrack.h"
 #include "Rendering/SkeletalMeshLODImporterData.h"
 #include "Rendering/SkeletalMeshLODModel.h"
 #include "Rendering/SkeletalMeshLODRenderData.h"
 #include "Rendering/SkeletalMeshModel.h"
 #include "Rendering/SkeletalMeshRenderData.h"
+#include "Rigs/RigHierarchyElements.h"
+#include "Sequencer/MovieSceneControlRigParameterSection.h"
 
 #if WITH_EDITOR
 #include "Animation/DebugSkelMeshComponent.h"
@@ -3123,6 +3132,383 @@ bool UnrealToUsd::ConvertAnimSequence( UAnimSequence* AnimSequence, pxr::UsdPrim
 	return true;
 }
 
+bool UnrealToUsd::ConvertControlRigSection(
+	UMovieSceneControlRigParameterSection* InSection,
+	const FMovieSceneSequenceTransform& InTransform,
+	UMovieScene* InMovieScene,
+	IMovieScenePlayer* InPlayer,
+	const FReferenceSkeleton& InRefSkeleton,
+	pxr::UsdPrim& InSkelRoot,
+	pxr::UsdPrim& OutSkelAnimPrim,
+	const UsdUtils::FBlendShapeMap* InBlendShapeMap
+)
+{
+	if ( !InSection || !InPlayer || !OutSkelAnimPrim )
+	{
+		return false;
+	}
+
+	UControlRig* ControlRig = InSection->GetControlRig();
+	if ( !ControlRig )
+	{
+		return false;
+	}
+
+	FScopedUsdAllocs UsdAllocs;
+
+	pxr::UsdSkelAnimation SkelAnim{ OutSkelAnimPrim };
+	pxr::UsdStageRefPtr UsdStage = OutSkelAnimPrim.GetStage();
+	if ( !UsdStage || !SkelAnim )
+	{
+		return false;
+	}
+
+	FUsdStageInfo StageInfo( UsdStage );
+
+	double StartTime = FPlatformTime::Cycles64();
+
+	// TODO: This does not seem necessary
+	ControlRig->Initialize();
+	ControlRig->RequestInit();
+
+	// Record how the topology looks while we setup our arrays and maps. If this changes during
+	// baking we'll just drop everything and return
+	URigHierarchy* InitialHierarchy = ControlRig->GetHierarchy();
+	if ( !InitialHierarchy )
+	{
+		return false;
+	}
+	uint16 TopologyVersion = InitialHierarchy->GetTopologyVersion();
+
+	// Prepare to remap from Rig joint order to USkeleton/Skeleton prim joint order.
+	// This works because the topology won't change in here, and bone names are unique across the entire skeleton
+	// Its possible we'll be putting INDEX_NONEs into RigIndexToRefSkeletonIndex, but that's alright.
+	TArray<int32> RigJointIndexToRefSkeletonIndex;
+	TArray<FRigBoneElement*> InitialBoneElements = InitialHierarchy->GetBones();
+	for ( FRigBoneElement* RigBone : InitialBoneElements )
+	{
+		RigJointIndexToRefSkeletonIndex.Add( InRefSkeleton.FindBoneIndex( RigBone->GetName() ) );
+	}
+
+	pxr::UsdAttribute JointsAttr = SkelAnim.CreateJointsAttr();
+	UnrealToUsd::ConvertJointsAttribute( InRefSkeleton, JointsAttr );
+
+	TArray<FTransform> GlobalUEJointTransformsForFrame;
+	GlobalUEJointTransformsForFrame.SetNum( InRefSkeleton.GetNum() );
+
+	pxr::UsdAttribute TranslationsAttr = SkelAnim.CreateTranslationsAttr();
+	pxr::UsdAttribute RotationsAttr = SkelAnim.CreateRotationsAttr();
+	pxr::UsdAttribute ScalesAttr = SkelAnim.CreateScalesAttr();
+	pxr::UsdAttribute BlendShapeWeightsAttr = SkelAnim.CreateBlendShapeWeightsAttr();
+	pxr::UsdAttribute BlendShapesAttr = SkelAnim.CreateBlendShapesAttr();
+
+	TranslationsAttr.Clear();
+	RotationsAttr.Clear();
+	ScalesAttr.Clear();
+	pxr::VtVec3fArray Translations;
+	pxr::VtQuatfArray Rotations;
+	pxr::VtVec3hArray Scales;
+	Translations.resize( InRefSkeleton.GetNum() );
+	Rotations.resize( InRefSkeleton.GetNum() );
+	Scales.resize( InRefSkeleton.GetNum() );
+
+	FFrameRate TickResolution = InMovieScene->GetTickResolution();
+	FFrameRate DisplayRate = InMovieScene->GetDisplayRate();
+
+	const double StageTimeCodesPerSecond = UsdStage->GetTimeCodesPerSecond();
+	const FFrameRate StageFrameRate( StageTimeCodesPerSecond, 1 );
+
+	TRange<FFrameNumber> PlaybackRange = InMovieScene->GetPlaybackRange();
+	TRange<FFrameNumber> BakeTickRange = InSection->ComputeEffectiveRange();
+
+	// Try our best to find the section start/end inclusive frames
+	FFrameNumber StartInclTickFrame;
+	FFrameNumber EndInclTickFrame;
+	{
+		TOptional<TRangeBound<FFrameNumber>> LowerBoundToUse;
+		if ( BakeTickRange.HasLowerBound() )
+		{
+			TRangeBound<FFrameNumber> SectionLowerBound = BakeTickRange.GetLowerBound();
+			if ( !SectionLowerBound.IsOpen() )
+			{
+				LowerBoundToUse = SectionLowerBound;
+			}
+		}
+		if ( !LowerBoundToUse.IsSet() && PlaybackRange.HasLowerBound() )
+		{
+			TRangeBound<FFrameNumber> PlaybackLowerBound = PlaybackRange.GetLowerBound();
+			if ( !PlaybackLowerBound.IsOpen() )
+			{
+				LowerBoundToUse = PlaybackLowerBound;
+			}
+		}
+		if ( !LowerBoundToUse.IsSet() )
+		{
+			return false;
+		}
+		StartInclTickFrame = LowerBoundToUse.GetValue().GetValue() + ( LowerBoundToUse.GetValue().IsInclusive() ? 0 : 1 );
+	}
+	{
+		TOptional<TRangeBound<FFrameNumber>> UpperBoundToUse;
+		if ( BakeTickRange.HasUpperBound() )
+		{
+			TRangeBound<FFrameNumber> SectionUpperBound = BakeTickRange.GetUpperBound();
+			if ( !SectionUpperBound.IsOpen() )
+			{
+				UpperBoundToUse = SectionUpperBound;
+			}
+		}
+		if ( !UpperBoundToUse.IsSet() && PlaybackRange.HasUpperBound() )
+		{
+			TRangeBound<FFrameNumber> PlaybackUpperBound = PlaybackRange.GetUpperBound();
+			if ( !PlaybackUpperBound.IsOpen() )
+			{
+				UpperBoundToUse = PlaybackUpperBound;
+			}
+		}
+		if ( !UpperBoundToUse.IsSet() )
+		{
+			return false;
+		}
+		EndInclTickFrame = UpperBoundToUse.GetValue().GetValue() + ( UpperBoundToUse.GetValue().IsInclusive() ? 0 : -1 );
+	}
+
+	pxr::VtArray<pxr::TfToken> CurveNames;
+	pxr::VtArray<float> CurvesValuesForTime;
+
+	// Prepare blend shape baking
+	// So far there doesn't seem to be any good way of handling the baking into blend shapes with inbetweens:
+	//	- We can't just pretend the Mesh prims have the flattened inbetween blend shapes (like we'd get if they were
+	//    exported) because we'd get warnings by having blend shape targets to blend shape prims that don't exist;
+	//  - We could flatten the actual BlendShape on the Mesh prim here, but that may be a bit too bold as the user likely
+	//    wants to keep his Mesh asset more or less intact when just baking out an animation section. If users do want
+	//    this behaviour we can later add it though;
+	//  - An alternative would have been to try to collect all the primary+inbetween weights, combine them back into a
+	//    single weight value, and write them back. That would work, but it would be incredibly hard to tell what is
+	//    going on from the users' perspective because that weight conversion is lossy and imperfect. Not to mention we'd
+	//    have this tricky code to test/maintain that slows down the baking process as a whole, and everything would break
+	//    if e.g. the curves were renamed;
+	//  - A slightly different approach to above would be to have the Mesh prims listen to the flattened inbetween blend
+	//    shape channels, but map them all to the single blend shape: This is not allowed in USD though, and its enough
+	//    to crash usdview. Besides, it wouldn't have added a lot of value as it would be impossible to comprehend what
+	//    was going on.
+	// The best we can do at the moment is to make one channel for each curve on the SkelAnimation prim, but maintain
+	// each Mesh prim connected only to the primary blend shape channel, if it was originally. We'll show a warning
+	// explaining the situation though.
+	if ( InBlendShapeMap && InSkelRoot )
+	{
+		TArray<FRigCurveElement*> CurveElements = InitialHierarchy->GetCurves();
+
+		CurveNames.reserve( CurveElements.Num() );
+		for ( int32 Index = 0; Index < CurveElements.Num(); ++Index )
+		{
+			FRigCurveElement* Element = CurveElements[Index];
+			const FString& CurveNameString = Element->GetNameString();
+
+			CurveNames.push_back( UnrealToUsd::ConvertToken( *CurveNameString ).Get() );
+		}
+
+		// Check if the blend shape channels on skel animation are the same names as morph target curves.
+		// Note that the actual order of the channel names within BlendShapesAttr is not important, as we'll always
+		// write out a new order that matches the rig anyway. We just want to know if all consumers of this SkelAnimation
+		// already have the processed, "one per morph target" channels
+		bool bNeedChannelUpdate = true;
+		pxr::VtArray<pxr::TfToken> SkelAnimBlendShapeChannels;
+		if ( BlendShapesAttr && BlendShapesAttr.Get( &SkelAnimBlendShapeChannels ) )
+		{
+			if ( SkelAnimBlendShapeChannels.size() == CurveNames.size() )
+			{
+				std::unordered_set<pxr::TfToken, pxr::TfToken::HashFunctor> ExistingCurveNames;
+				for ( const pxr::TfToken& Channel : SkelAnimBlendShapeChannels )
+				{
+					ExistingCurveNames.insert( Channel );
+				}
+
+				bool bFoundAllCurves = true;
+				for ( const pxr::TfToken& CurveName : CurveNames )
+				{
+					if ( ExistingCurveNames.count( CurveName ) == 0 )
+					{
+						bFoundAllCurves = false;
+						break;
+					}
+				}
+
+				bNeedChannelUpdate = !bFoundAllCurves;
+			}
+		}
+
+		// We haven't processed this SkelAnimation before, so we need to do it now.
+		// The summary is that since each MorphTarget/BlendShape has an independet curve in UE, but can share curves
+		// arbitrarily in USD, we need to replace the existing SkelAnimation channels with ones that are unique for
+		// each blend shape. This is not ideal, but the alternatives would be to: Not handle blend shape curves via
+		// control rigs; Have some morph target curves unintuitively "mirror each other" in UE, if at all possible;
+		// Try to keep the channels shared on USD's side, which would desync USD/UE and show a different result when
+		// reloading.
+		if ( bNeedChannelUpdate )
+		{
+			// We'll change the blend shape channel names, so we need to update all meshes that were using them too.
+			// For now we'll assume that they're all inside the same skel root. We could upgrade this for the stage
+			// later too, if needed
+			for ( UE::FUsdPrim& MeshPrim : UsdUtils::GetAllPrimsOfType( UE::FUsdPrim{ InSkelRoot }, TEXT( "UsdGeomMesh" ) ) )
+			{
+				pxr::UsdSkelBindingAPI SkelBindingAPI{ MeshPrim };
+				if ( !SkelBindingAPI )
+				{
+					continue;
+				}
+
+				pxr::UsdRelationship TargetsRel = SkelBindingAPI.GetBlendShapeTargetsRel();
+				pxr::UsdAttribute ChannelsAttr = SkelBindingAPI.GetBlendShapesAttr();
+
+				if ( TargetsRel && ChannelsAttr )
+				{
+					pxr::SdfPathVector BlendShapeTargets;
+					if ( TargetsRel.GetTargets( &BlendShapeTargets ) )
+					{
+						pxr::VtArray<pxr::TfToken> BlendShapeChannels;
+						ChannelsAttr.Get( &BlendShapeChannels );
+
+						BlendShapeChannels.resize( BlendShapeTargets.size() );
+
+						pxr::SdfPath MeshPath{ MeshPrim.GetPrimPath() };
+
+						bool bRenamedAChannel = false;
+						for ( int32 BlendShapeIndex = 0; BlendShapeIndex < BlendShapeTargets.size(); ++BlendShapeIndex )
+						{
+							const pxr::SdfPath& BlendShapePath = BlendShapeTargets[ BlendShapeIndex ];
+							FString PrimaryBlendShapePath = UsdToUnreal::ConvertPath( BlendShapePath.MakeAbsolutePath( MeshPath ) );
+
+							// Mesh had <blendshape1> target on channel "C" -> We have a morph target called "blendshape1" already, and
+							// we'll create a new channel on SkelAnimation called "blendshape1" -> Let's replace channel "C" with channel
+							// "blendshape1"
+							if ( const UsdUtils::FUsdBlendShape* FoundBlendShape = InBlendShapeMap->Find( PrimaryBlendShapePath ) )
+							{
+								bRenamedAChannel = true;
+								BlendShapeChannels[ BlendShapeIndex ] = UnrealToUsd::ConvertToken( *FoundBlendShape->Name ).Get();
+								UE_LOG( LogUsd, Log, TEXT( "Updating Mesh '%s' to bind BlendShape target '%s' to SkelAnimation curve '%s'" ),
+									*MeshPrim.GetPrimPath().GetString(),
+									*PrimaryBlendShapePath,
+									*FoundBlendShape->Name
+								);
+
+								if ( FoundBlendShape->Inbetweens.Num() > 0 )
+								{
+									UE_LOG( LogUsd, Warning, TEXT( "Baking Control Rig parameter sections for BlendShapes with inbetweens (like '%s') is not currently supported, so animation for mesh '%s' may look incorrect! Please flatten the inbetweens into separate BlendShapes beforehand (importing and exporting will do that)." ),
+										*FoundBlendShape->Name,
+										*MeshPrim.GetPrimPath().GetString()
+									);
+								}
+							}
+						}
+
+						if ( bRenamedAChannel )
+						{
+							ChannelsAttr.Set( BlendShapeChannels );
+						}
+					}
+				}
+			}
+		}
+
+		// Now that we updated the channel names we need to make sure we clear the previous weights as they'll
+		// make no sense
+		BlendShapeWeightsAttr.Clear();
+		BlendShapesAttr.Set( CurveNames );
+		CurvesValuesForTime.resize( CurveNames.size() );
+	}
+
+	FFrameTime TickIncr = FFrameRate::TransformTime( 1, DisplayRate, TickResolution );
+	for ( FFrameTime FrameTickTime = StartInclTickFrame; FrameTickTime <= EndInclTickFrame; FrameTickTime += TickIncr )
+	{
+		FFrameTime TransformedFrameTickTime = FrameTickTime * InTransform;
+
+		double UsdTimeCode = FFrameRate::TransformTime( TransformedFrameTickTime, TickResolution, StageFrameRate ).AsDecimal();
+
+		FMovieSceneContext Context = FMovieSceneContext(
+			FMovieSceneEvaluationRange(
+				TransformedFrameTickTime,
+				TickResolution
+			),
+			InPlayer->GetPlaybackStatus()
+		).SetHasJumped( true );
+
+		InPlayer->GetEvaluationTemplate().Evaluate( Context, *InPlayer );
+		ControlRig->Evaluate_AnyThread();
+
+		URigHierarchy* Hierarchy = ControlRig->GetHierarchy();
+		if ( !Hierarchy || Hierarchy->GetTopologyVersion() != TopologyVersion )
+		{
+			UE_LOG( LogUsd, Error, TEXT( "Baking Control Rig tracks for rig '%s' failed because the rig changed topology during baking process" ), *ControlRig->GetPathName() );
+			return false;
+		}
+
+		// Sadly we have to fetch these each frame as these are regenerated on each evaluation of the Sequencer
+		// (c.f. FControlRigBindingHelper::BindToSequencerInstance, URigHierarchy::CopyHierarchy)
+		TArray<FRigBoneElement*> BoneElements = Hierarchy->GetBones();
+
+		if ( CurvesValuesForTime.size() > 0 )
+		{
+			TArray<FRigCurveElement*> CurveElements = Hierarchy->GetCurves();
+			for ( int32 ElementIndex = 0; ElementIndex < CurveElements.Num(); ++ElementIndex )
+			{
+				FRigCurveElement* Element = CurveElements[ ElementIndex ];
+				CurvesValuesForTime[ ElementIndex ] = Hierarchy->GetCurveValue( Element );
+			}
+
+			BlendShapeWeightsAttr.Set( CurvesValuesForTime, pxr::UsdTimeCode( UsdTimeCode ) );
+		}
+
+		for ( int32 RigBoneIndex = 0; RigBoneIndex < BoneElements.Num(); ++RigBoneIndex )
+		{
+			FRigBoneElement* El = BoneElements[ RigBoneIndex ];
+
+			// Our skeleton doesn't have this rig bone
+			int32 RefSkeletonBoneIndex = RigJointIndexToRefSkeletonIndex[ RigBoneIndex ];
+			if ( RefSkeletonBoneIndex == INDEX_NONE )
+			{
+				continue;
+			}
+
+			GlobalUEJointTransformsForFrame[ RefSkeletonBoneIndex ] = Hierarchy->GetTransform( El, ERigTransformType::CurrentGlobal );
+
+			FTransform UsdTransform;
+
+			// We have to calculate the local transforms ourselves since the parent element could be a control
+			int32 RefSkeletonParentBoneIndex = InRefSkeleton.GetParentIndex( RefSkeletonBoneIndex );
+			if ( RefSkeletonParentBoneIndex == INDEX_NONE )
+			{
+				UsdTransform = UsdUtils::ConvertTransformToUsdSpace( StageInfo, GlobalUEJointTransformsForFrame[ RefSkeletonBoneIndex ] );
+			}
+			else
+			{
+				const FTransform& ChildGlobal = GlobalUEJointTransformsForFrame[ RefSkeletonBoneIndex ];
+				const FTransform& ParentGlobal = GlobalUEJointTransformsForFrame[ RefSkeletonParentBoneIndex ];
+				UsdTransform = UsdUtils::ConvertTransformToUsdSpace( StageInfo, ChildGlobal.GetRelativeTransform( ParentGlobal ) );
+			}
+
+			Translations[ RefSkeletonBoneIndex ] = UnrealToUsd::ConvertVector( UsdTransform.GetTranslation() );
+			Rotations[ RefSkeletonBoneIndex ] = UnrealToUsd::ConvertQuat( UsdTransform.GetRotation() ).GetNormalized();
+			Scales[ RefSkeletonBoneIndex ] = pxr::GfVec3h( UnrealToUsd::ConvertVector( UsdTransform.GetScale3D() ) );
+		}
+
+		TranslationsAttr.Set( Translations, pxr::UsdTimeCode( UsdTimeCode ) );
+		RotationsAttr.Set( Rotations, pxr::UsdTimeCode( UsdTimeCode ) );
+		ScalesAttr.Set( Scales, pxr::UsdTimeCode( UsdTimeCode ) );
+	}
+
+	double ElapsedSeconds = FPlatformTime::ToSeconds64( FPlatformTime::Cycles64() - StartTime );
+	int ElapsedMin = int( ElapsedSeconds / 60.0 );
+	ElapsedSeconds -= 60.0 * ( double ) ElapsedMin;
+	UE_LOG( LogUsd, Log, TEXT( "Baked new animation for prim '%s' in [%d min %.3f s]" ),
+		*UsdToUnreal::ConvertPath( OutSkelAnimPrim.GetPrimPath() ),
+		ElapsedMin,
+		ElapsedSeconds
+	);
+
+	return true;
+}
+
 #undef LOCTEXT_NAMESPACE
 
 #endif // #if USE_USD_SDK && WITH_EDITOR
diff --git a/Engine/Plugins/Importers/USDImporter/Source/USDUtilities/Public/USDIntegrationUtils.h b/Engine/Plugins/Importers/USDImporter/Source/USDUtilities/Public/USDIntegrationUtils.h
--- a/Engine/Plugins/Importers/USDImporter/Source/USDUtilities/Public/USDIntegrationUtils.h
+++ b/Engine/Plugins/Importers/USDImporter/Source/USDUtilities/Public/USDIntegrationUtils.h
@@ -25,6 +25,17 @@ namespace UsdUtils
 	// Returns true in case Prim has our custom "LiveLinkAPI" schema
 	USDUTILITIES_API bool PrimHasLiveLinkSchema( const pxr::UsdPrim& Prim );
 
+	// Returns true in case Prim has our custom "ControlRigAPI" schema
+	USDUTILITIES_API bool PrimHasControlRigSchema( const pxr::UsdPrim& Prim );
+
+	// Apply the control rig schema to the prim, creates all the required attributes and sets default values
+	// if they're otherwise not authored
+	USDUTILITIES_API bool ApplyControlRigSchema( const pxr::UsdPrim& Prim );
+
+	// Apply the live link schema to the prim, creates all the required attributes and sets default values
+	// if they're otherwise not authored
+	USDUTILITIES_API bool ApplyLiveLinkSchema( const pxr::UsdPrim& Prim );
+
 #endif // USE_USD_SDK
 }
 
diff --git a/Engine/Plugins/Importers/USDImporter/Source/USDUtilities/Public/USDSkeletalDataConversion.h b/Engine/Plugins/Importers/USDImporter/Source/USDUtilities/Public/USDSkeletalDataConversion.h
--- a/Engine/Plugins/Importers/USDImporter/Source/USDUtilities/Public/USDSkeletalDataConversion.h
+++ b/Engine/Plugins/Importers/USDImporter/Source/USDUtilities/Public/USDSkeletalDataConversion.h
@@ -29,9 +29,13 @@ PXR_NAMESPACE_OPEN_SCOPE
 PXR_NAMESPACE_CLOSE_SCOPE
 
 class FSkeletalMeshImportData;
+class IMovieScenePlayer;
 class UAnimSequence;
+class UMovieScene;
+class UMovieSceneControlRigParameterSection;
 class USkeletalMesh;
 struct FUsdStageInfo;
+struct FMovieSceneSequenceTransform;
 namespace SkeletalMeshImportData
 {
 	struct FBone;
@@ -233,6 +237,35 @@ namespace UnrealToUsd
 	 * @return Whether the conversion was successful or not.
 	 */
 	USDUTILITIES_API bool ConvertAnimSequence( UAnimSequence* AnimSequence, pxr::UsdPrim& SkelAnimPrim );
+
+	/**
+	 * Plays the provided Section in the background, driving its ControlRig and baking to USD the animated bones and
+	 * curves end result.
+	 * @param InSection - The section to bake;
+	 * @param InTransform - A time transform to apply to the baked data before writing out;
+	 * @param InMovieScene - The MovieScene that contains InSection (we can't access via the Outer chain from this RTTI
+	 *                       module);
+	 * @param IMovieScenePlayer - The player to drive the InSection with (e.g. a Sequencer instance);
+	 * @param InRefSkeleton - Skeleton that describes the target Skeleton prim bone structure to use (the control rig
+	 *                        may have an arbitrary bone hierarchy, but we always want to write bones compatible with
+	 *                        the target UsdSkelRoot and its child prims);
+	 * @param InSkelRoot - The SkelRoot that we're writing to (in some cases we need to update Mesh prims, and only
+	 *                     the meshes within this SkelRoot will be updated);
+	 * @param OutSkelAnimPrim - The SkelAnimation prim to receive the baked data;
+	 * @param InBlendShapeMap - Optional map describing the stage blend shapes, in case baking of the rig's animation
+	 *                          curves is desired;
+	 * @return Whether the conversion was successful or not.
+	 */
+	USDUTILITIES_API bool ConvertControlRigSection(
+		UMovieSceneControlRigParameterSection* InSection,
+		const FMovieSceneSequenceTransform& InTransform,
+		UMovieScene* InMovieScene,
+		IMovieScenePlayer* InPlayer,
+		const FReferenceSkeleton& InRefSkeleton,
+		pxr::UsdPrim& InSkelRoot,
+		pxr::UsdPrim& OutSkelAnimPrim,
+		const UsdUtils::FBlendShapeMap* InBlendShapeMap = nullptr
+	);
 }
 
 #endif // #if USE_USD_SDK && WITH_EDITOR
diff --git a/Engine/Plugins/Importers/USDImporter/Source/USDUtilities/USDUtilities.Build.cs b/Engine/Plugins/Importers/USDImporter/Source/USDUtilities/USDUtilities.Build.cs
--- a/Engine/Plugins/Importers/USDImporter/Source/USDUtilities/USDUtilities.Build.cs
+++ b/Engine/Plugins/Importers/USDImporter/Source/USDUtilities/USDUtilities.Build.cs
@@ -25,6 +25,7 @@ namespace UnrealBuildTool.Rules
 				new string[]
 				{
 					"CinematicCamera",
+					"ControlRig",
 					"Engine",
 					"Foliage",
 					"GeometryCache", // Just so that we can fetch its AssetImportData
diff --git a/Engine/Plugins/Importers/USDImporter/Source/UnrealUSDWrapper/Private/UnrealUSDWrapper.cpp b/Engine/Plugins/Importers/USDImporter/Source/UnrealUSDWrapper/Private/UnrealUSDWrapper.cpp
--- a/Engine/Plugins/Importers/USDImporter/Source/UnrealUSDWrapper/Private/UnrealUSDWrapper.cpp
+++ b/Engine/Plugins/Importers/USDImporter/Source/UnrealUSDWrapper/Private/UnrealUSDWrapper.cpp
@@ -107,10 +107,14 @@ namespace UnrealIdentifiers
 	const TfToken UnrealNaniteOverrideDisable = TfToken("disable");
 
 	const TfToken LiveLinkAPI = TfToken("LiveLinkAPI");
-	const TfToken UnrealLiveLinkAnimBlueprintPath = TfToken("unrealLiveLink:animBlueprintPath");
-	const TfToken UnrealLiveLinkSubjectName = TfToken("unrealLiveLink:subjectName");
-	const TfToken UnrealLiveLinkEnabled = TfToken("unrealLiveLink:enabled");
-
+	const TfToken ControlRigAPI = TfToken("ControlRigAPI");
+	const TfToken UnrealAnimBlueprintPath = TfToken("unreal:liveLink:animBlueprintPath");
+	const TfToken UnrealLiveLinkSubjectName = TfToken("unreal:liveLink:subjectName");
+	const TfToken UnrealLiveLinkEnabled = TfToken("unreal:liveLink:enabled");
+	const TfToken UnrealControlRigPath = TfToken("unreal:controlRig:controlRigPath");
+	const TfToken UnrealUseFKControlRig = TfToken("unreal:controlRig:useFKControlRig");
+	const TfToken UnrealControlRigReduceKeys = TfToken("unreal:controlRig:reduceKeys");
+	const TfToken UnrealControlRigReductionTolerance = TfToken("unreal:controlRig:reductionTolerance");
 	const TfToken DiffuseColor = TfToken("diffuseColor");
 	const TfToken EmissiveColor = TfToken("emissiveColor");
 	const TfToken Metallic = TfToken("metallic");
diff --git a/Engine/Plugins/Importers/USDImporter/Source/UnrealUSDWrapper/Public/UnrealUSDWrapper.h b/Engine/Plugins/Importers/USDImporter/Source/UnrealUSDWrapper/Public/UnrealUSDWrapper.h
--- a/Engine/Plugins/Importers/USDImporter/Source/UnrealUSDWrapper/Public/UnrealUSDWrapper.h
+++ b/Engine/Plugins/Importers/USDImporter/Source/UnrealUSDWrapper/Public/UnrealUSDWrapper.h
@@ -348,9 +348,14 @@ namespace UnrealIdentifiers
 	extern UNREALUSDWRAPPER_API const pxr::TfToken UnrealNaniteOverrideDisable;
 
 	extern UNREALUSDWRAPPER_API const pxr::TfToken LiveLinkAPI;
-	extern UNREALUSDWRAPPER_API const pxr::TfToken UnrealLiveLinkAnimBlueprintPath;
+	extern UNREALUSDWRAPPER_API const pxr::TfToken ControlRigAPI;
+	extern UNREALUSDWRAPPER_API const pxr::TfToken UnrealAnimBlueprintPath;
 	extern UNREALUSDWRAPPER_API const pxr::TfToken UnrealLiveLinkSubjectName;
 	extern UNREALUSDWRAPPER_API const pxr::TfToken UnrealLiveLinkEnabled;
+	extern UNREALUSDWRAPPER_API const pxr::TfToken UnrealControlRigPath;
+	extern UNREALUSDWRAPPER_API const pxr::TfToken UnrealUseFKControlRig;
+	extern UNREALUSDWRAPPER_API const pxr::TfToken UnrealControlRigReduceKeys;
+	extern UNREALUSDWRAPPER_API const pxr::TfToken UnrealControlRigReductionTolerance;
 
 	extern UNREALUSDWRAPPER_API const pxr::TfToken DiffuseColor;
 	extern UNREALUSDWRAPPER_API const pxr::TfToken EmissiveColor;
