USD: New sample script for adding Control Rig integration via Python.

Also updates the Live Link equivalent because the schema tokens have changed.

diff --git a/Engine/Plugins/Importers/USDImporter/Resources/PythonExamples/enable_control_rig.py b/Engine/Plugins/Importers/USDImporter/Resources/PythonExamples/enable_control_rig.py
new file mode 100644
--- /dev/null
+++ b/Engine/Plugins/Importers/USDImporter/Resources/PythonExamples/enable_control_rig.py
@@ -0,0 +1,41 @@
+"""
+This script describes how to enable Control Rig for the components generated for a particular prim when
+imported in UE.
+
+After the script is executed, the stage can be saved and the connection details will persist on the USD
+files themselves. Opening the stage again will automatically re-establish the connection.
+
+To use this example, follow these steps:
+- Enable the "USD Importer" and "Live Link" plugins in UE;
+- Open the USD Stage editor window by going to Window -> Virtual Production -> USD Stage;
+- Open your desired USD Stage by going to File -> Open on the USD Stage editor and picking a USD file;
+- Edit the script below, replacing the prim path, ControlRig path and other attributes with your desired values;
+- Run the script, either by copy-pasting it into the UE Python console, or by pasting the path to this
+  file onto the same console.
+"""
+
+from pxr import Usd, UsdUtils, Sdf
+
+stage = UsdUtils.StageCache().Get().GetAllStages()[0]
+prim = stage.GetPrimAtPath("/MySkelRootPrim")
+
+schema = Usd.SchemaRegistry.GetTypeFromSchemaTypeName("ControlRigAPI")
+prim.ApplyAPI(schema)
+
+# Optional: If you want to also set the schema attribute values, you can do this:
+with Sdf.ChangeBlock():
+    path_attr = prim.GetAttribute("unreal:controlRig:controlRigPath")
+    path_attr.Set("/Game/SkeletalMeshes/MyRig.MyRig")
+
+    # Set this to True if you want to generate an FKControlRig instead of using an existing
+    # Control Rig blueprint asset
+    use_fk_attr = prim.GetAttribute("unreal:controlRig:useFKControlRig")
+    use_fk_attr.Set(False)
+
+    # Enable this to have the Control Rig baking process automatically prune unnecessary keys
+    reduce_keys_attr = prim.GetAttribute("unreal:controlRig:reduceKeys")
+    reduce_keys_attr.Set(False)
+
+    tolerance_attr = prim.GetAttribute("unreal:controlRig:reductionTolerance")
+    tolerance_attr.Set(0.001)
+
diff --git a/Engine/Plugins/Importers/USDImporter/Resources/PythonExamples/enable_live_link.py b/Engine/Plugins/Importers/USDImporter/Resources/PythonExamples/enable_live_link.py
--- a/Engine/Plugins/Importers/USDImporter/Resources/PythonExamples/enable_live_link.py
+++ b/Engine/Plugins/Importers/USDImporter/Resources/PythonExamples/enable_live_link.py
@@ -1,5 +1,5 @@
 """
-This script describes how to enable LiveLink for the components generated for a particular prim when
+This script describes how to enable Live Link for the components generated for a particular prim when
 imported in UE.
 
 After the script is executed, the stage can be saved and the connection details will persist on the USD
@@ -23,13 +23,14 @@ prim = stage.GetPrimAtPath("/cube_default")
 schema = Usd.SchemaRegistry.GetTypeFromSchemaTypeName("LiveLinkAPI")
 prim.ApplyAPI(schema)
 
+# Optional: If you want to also set the schema attribute values, you can do this:
 with Sdf.ChangeBlock():
-    anim_bp_attr = prim.CreateAttribute("unrealLiveLink:animBlueprintPath", Sdf.ValueTypeNames.String)
-    anim_bp_attr.Set("/USDImporter/Blueprint/DefaultLiveLinkAnimBP.DefaultLiveLinkAnimBP")
+    anim_bp_attr = prim.GetAttribute("unreal:liveLink:animBlueprintPath")
+    anim_bp_attr.Set("/Game/MyBlueprints/MyAnimBP.MyAnimBP")
 
-    subject_name_attr = prim.CreateAttribute("unrealLiveLink:subjectName", Sdf.ValueTypeNames.String)
+    subject_name_attr = prim.GetAttribute("unreal:liveLink:subjectName")
     subject_name_attr.Set("cube")
 
-    enabled_attr = prim.CreateAttribute("unrealLiveLink:enabled", Sdf.ValueTypeNames.Bool)
+    enabled_attr = prim.GetAttribute("unreal:liveLink:enabled")
     enabled_attr.Set(True)
 
