Fix warning C4100: 'Other': unreferenced formal parameter
#rnx

diff --git a/Engine/Plugins/Importers/USDImporter/Source/UnrealUSDWrapper/Public/UsdWrappers/ForwardDeclarations.h b/Engine/Plugins/Importers/USDImporter/Source/UnrealUSDWrapper/Public/UsdWrappers/ForwardDeclarations.h
--- a/Engine/Plugins/Importers/USDImporter/Source/UnrealUSDWrapper/Public/UsdWrappers/ForwardDeclarations.h
+++ b/Engine/Plugins/Importers/USDImporter/Source/UnrealUSDWrapper/Public/UsdWrappers/ForwardDeclarations.h
@@ -46,19 +46,19 @@ namespace UE
 	{
 	public:
 		FDummyRefPtrType(){};
-		FDummyRefPtrType( const FDummyWeakPtrType& Other ) {};
+		FDummyRefPtrType( const FDummyWeakPtrType& ) {};
 		explicit operator bool() const { return false; }
-		bool operator==( const FDummyRefPtrType& Other ) const { return false; }
-		bool operator!=( const FDummyRefPtrType& Other ) const { return false; }
+		bool operator==( const FDummyRefPtrType& ) const { return false; }
+		bool operator!=( const FDummyRefPtrType& ) const { return false; }
 	};
 	class FDummyWeakPtrType
 	{
 	public:
 		FDummyWeakPtrType(){};
-		FDummyWeakPtrType( const FDummyRefPtrType& Other ) {};
+		FDummyWeakPtrType( const FDummyRefPtrType& ) {};
 		explicit operator bool() const { return false; }
-		bool operator==( const FDummyWeakPtrType& Other ) const { return false; }
-		bool operator!=( const FDummyWeakPtrType& Other ) const { return false; }
+		bool operator==( const FDummyWeakPtrType& ) const { return false; }
+		bool operator!=( const FDummyWeakPtrType& ) const { return false; }
 	};
 
 	template<typename PtrType> class FUsdStageBase;
