USD: Fixed prototype prim being instantiated as actor. Empty USD change notices are needed after all to handle add, remove and mute/unmute of sublayers so this change reverts the previous fix and adds a better check instead.

diff --git a/Engine/Plugins/Importers/USDImporter/Source/UnrealUSDWrapper/Private/USDListener.cpp b/Engine/Plugins/Importers/USDImporter/Source/UnrealUSDWrapper/Private/USDListener.cpp
--- a/Engine/Plugins/Importers/USDImporter/Source/UnrealUSDWrapper/Private/USDListener.cpp
+++ b/Engine/Plugins/Importers/USDImporter/Source/UnrealUSDWrapper/Private/USDListener.cpp
@@ -5,6 +5,7 @@
 #include "USDLog.h"
 #include "USDMemory.h"
 
+#include "UsdWrappers/SdfPath.h"
 #include "UsdWrappers/UsdStage.h"
 #include "UsdWrappers/VtValue.h"
 
@@ -17,6 +18,7 @@
 #include "pxr/usd/sdf/notice.h"
 #include "pxr/usd/sdf/path.h"
 #include "pxr/usd/usd/notice.h"
+#include "pxr/usd/usd/prim.h"
 #include "pxr/usd/usd/stage.h"
 #include "pxr/usd/usdGeom/tokens.h"
 
@@ -39,6 +41,7 @@ namespace UsdToUnreal
 		{
 			const FString PrimPath = ANSI_TO_TCHAR( It->GetAbsoluteRootOrPrimPath().GetAsString().c_str() );
 
+			if (pxr::UsdPrim::IsPathInPrototype(UE::FSdfPath(*PrimPath)))
 			// Something like "/Root/Prim.some_field", or "/"
 			const FString FullFieldPath = ANSI_TO_TCHAR( It->GetAsString().c_str() );
 
