USD: Remove "Reuse identical assets" import option.

diff --git a/Engine/Plugins/Importers/USDImporter/Source/USDStageImporter/Private/USDStageImportOptions.cpp b/Engine/Plugins/Importers/USDImporter/Source/USDStageImporter/Private/USDStageImportOptions.cpp
--- a/Engine/Plugins/Importers/USDImporter/Source/USDStageImporter/Private/USDStageImportOptions.cpp
+++ b/Engine/Plugins/Importers/USDImporter/Source/USDStageImporter/Private/USDStageImportOptions.cpp
@@ -25,7 +25,6 @@ UUsdStageImportOptions::UUsdStageImportOptions(const FObjectInitializer& ObjectI
 	StageOptions.MetersPerUnit = 0.01f;
 	StageOptions.UpAxis = EUsdUpAxis::ZAxis;
 
-	bReuseIdenticalAssets = true;
 	ExistingActorPolicy = EReplaceActorPolicy::Replace;
 	ExistingAssetPolicy = EReplaceAssetPolicy::Replace;
 
diff --git a/Engine/Plugins/Importers/USDImporter/Source/USDStageImporter/Private/USDStageImporter.cpp b/Engine/Plugins/Importers/USDImporter/Source/USDStageImporter/Private/USDStageImporter.cpp
--- a/Engine/Plugins/Importers/USDImporter/Source/USDStageImporter/Private/USDStageImporter.cpp
+++ b/Engine/Plugins/Importers/USDImporter/Source/USDStageImporter/Private/USDStageImporter.cpp
@@ -1270,7 +1270,6 @@ namespace UsdStageImporterImpl
 					EventAttributes.Emplace( TEXT( "MetersPerUnit" ), ImportContext.ImportOptions->StageOptions.MetersPerUnit );
 					EventAttributes.Emplace( TEXT( "UpAxis" ), ImportContext.ImportOptions->StageOptions.UpAxis == EUsdUpAxis::YAxis ? TEXT( "Y" ) : TEXT( "Z" ) );
 				}
-				EventAttributes.Emplace( TEXT( "ReuseIdenticalAssets" ), ImportContext.ImportOptions->bReuseIdenticalAssets );
 				EventAttributes.Emplace( TEXT( "ReplaceActorPolicy" ), LexToString( (uint8)ImportContext.ImportOptions->ExistingActorPolicy ) );
 				EventAttributes.Emplace( TEXT( "ReplaceAssetPolicy" ), LexToString( (uint8)ImportContext.ImportOptions->ExistingAssetPolicy ) );
 				EventAttributes.Emplace( TEXT( "PrimPathFolderStructure" ), LexToString( ImportContext.ImportOptions->bPrimPathFolderStructure ) );
diff --git a/Engine/Plugins/Importers/USDImporter/Source/USDStageImporter/Public/USDStageImportOptions.h b/Engine/Plugins/Importers/USDImporter/Source/USDStageImporter/Public/USDStageImportOptions.h
--- a/Engine/Plugins/Importers/USDImporter/Source/USDStageImporter/Public/USDStageImportOptions.h
+++ b/Engine/Plugins/Importers/USDImporter/Source/USDStageImporter/Public/USDStageImportOptions.h
@@ -91,11 +91,6 @@ public:
 	UPROPERTY(EditAnywhere, config, BlueprintReadWrite, Category = "Groom")
 	TArray<FHairGroupsInterpolation> GroomInterpolationSettings;
 
-	/**
-	 * If enabled, whenever two different prims import into identical assets, only one of those assets will be kept and reused.
-	 */
-	UPROPERTY(BlueprintReadWrite, config, EditAnywhere, Category = "Collision", meta=(EditCondition=bImportGeometry))
-	bool bReuseIdenticalAssets;
 
 	/** What should happen when imported actors and components try to overwrite existing actors and components */
 	UPROPERTY(BlueprintReadWrite, config, EditAnywhere, Category = "Collision", meta=(EditCondition=bImportActors))
