USD: Fix usage of editor-only objects outside WITH_EDITOR tags.

diff --git a/Engine/Plugins/Importers/USDImporter/Source/USDStage/Private/USDLevelSequenceHelper.cpp b/Engine/Plugins/Importers/USDImporter/Source/USDStage/Private/USDLevelSequenceHelper.cpp
--- a/Engine/Plugins/Importers/USDImporter/Source/USDStage/Private/USDLevelSequenceHelper.cpp
+++ b/Engine/Plugins/Importers/USDImporter/Source/USDStage/Private/USDLevelSequenceHelper.cpp
@@ -228,6 +228,7 @@ namespace UsdLevelSequenceHelperImpl
 		Track->SetEvalDisabled( bMute );
 	}
 
+#if WITH_EDITOR
 	TSharedPtr< ISequencer > GetOpenedSequencerForLevelSequence( ULevelSequence* LevelSequence )
 	{
 		const bool bFocusIfOpen = false;
@@ -236,7 +237,6 @@ namespace UsdLevelSequenceHelperImpl
 		return LevelSequenceEditor ? LevelSequenceEditor->GetSequencer() : nullptr;
 	}
 
-#if WITH_EDITOR
 	// Rough copy of UControlRigSequencerEditorLibrary::BakeToControlRig, except that it allows us to control which
 	// sequence player is used, lets us use our own existing AnimSequence for the ControlRig track, doesn't force
 	// the control rig editor mode to open and doesn't crash itself when changing the edit mode away from the
@@ -2175,10 +2175,10 @@ void FUsdLevelSequenceHelperImpl::UpdateControlRigTracks( UUsdPrimTwin& PrimTwin
 
 			RefreshSequencer();
 		}
-#endif // WITH_EDITOR
 	}
 
 	PrimPathByLevelSequenceName.AddUnique( SkelAnimationSequence->GetFName(), PrimPath.GetString() );
+#endif // WITH_EDITOR
 }
 
 void FUsdLevelSequenceHelperImpl::RemoveSequenceForPrim( ULevelSequence& Sequence, const UUsdPrimTwin& PrimTwin )
@@ -2521,6 +2521,7 @@ void FUsdLevelSequenceHelperImpl::OnObjectTransacted(UObject* Object, const clas
 			HandleTrackChange( *ParentTrack, bIsMuteChange );
 		}
 
+#if WITH_EDITOR
 		if ( !bIsMuteChange )
 		{
 			if ( UMovieSceneControlRigParameterSection* CRSection = Cast<UMovieSceneControlRigParameterSection>( Section ) )
@@ -2540,6 +2541,7 @@ void FUsdLevelSequenceHelperImpl::OnObjectTransacted(UObject* Object, const clas
 				}
 			}
 		}
+#endif // WITH_EDITOR
 	}
 }
 
@@ -2869,6 +2871,7 @@ void FUsdLevelSequenceHelperImpl::HandleSubSectionChange( UMovieSceneSubSection&
 
 void FUsdLevelSequenceHelperImpl::HandleControlRigSectionChange( UMovieSceneControlRigParameterSection& Section )
 {
+#if WITH_EDITOR
 	AUsdStageActor* StageActorValue = StageActor.Get();
 	if ( !StageActorValue )
 	{
@@ -3046,6 +3049,7 @@ void FUsdLevelSequenceHelperImpl::HandleControlRigSectionChange( UMovieSceneCont
 		// USD.RegenerateSkeletalAssetsOnControlRigBake to false.
 		GetOnSkelAnimationBaked().Broadcast( PrimTwin->PrimPath );
 	}
+#endif // WITH_EDITOR
 }
 
 void FUsdLevelSequenceHelperImpl::HandleTrackChange( const UMovieSceneTrack& Track, bool bIsMuteChange )
