USD: Fix static analysis warning.

diff --git a/Engine/Plugins/Importers/USDImporter/Source/USDStage/Private/USDStageActor.cpp b/Engine/Plugins/Importers/USDImporter/Source/USDStage/Private/USDStageActor.cpp
--- a/Engine/Plugins/Importers/USDImporter/Source/USDStage/Private/USDStageActor.cpp
+++ b/Engine/Plugins/Importers/USDImporter/Source/USDStage/Private/USDStageActor.cpp
@@ -2908,12 +2908,15 @@ void AUsdStageActor::AnimatePrims()
 
 	// c.f. comment on bSequencerIsAnimating's declaration
 #if WITH_EDITOR
-	const bool bFocusIfOpen = false;
-	IAssetEditorInstance* AssetEditor = GEditor->GetEditorSubsystem<UAssetEditorSubsystem>()->FindEditorForAsset( LevelSequence, bFocusIfOpen );
-	if ( ILevelSequenceEditorToolkit* LevelSequenceEditor = static_cast< ILevelSequenceEditorToolkit* >( AssetEditor ) )
-	{
-		TranslationContext->bSequencerIsAnimating = true;
-	}
+    if ( GEditor )
+    {
+        const bool bFocusIfOpen = false;
+        IAssetEditorInstance* AssetEditor = GEditor->GetEditorSubsystem<UAssetEditorSubsystem>()->FindEditorForAsset( LevelSequence, bFocusIfOpen );
+        if ( ILevelSequenceEditorToolkit* LevelSequenceEditor = static_cast< ILevelSequenceEditorToolkit* >( AssetEditor ) )
+        {
+            TranslationContext->bSequencerIsAnimating = true;
+        }
+    }
 #endif // WITH_EDITOR
 
 	for ( const FString& PrimToAnimate : PrimsToAnimate )
