USD: Added custom groom schema that can be applied and removed through the USD Stage. The groom translator now uses this schema to interpret the prim as a groom.

diff --git a/Engine/Plugins/Importers/USDImporter/Resources/UsdResources/Win64/plugins/unreal/resources/plugInfo.json b/Engine/Plugins/Importers/USDImporter/Resources/UsdResources/Win64/plugins/unreal/resources/plugInfo.json
--- a/Engine/Plugins/Importers/USDImporter/Resources/UsdResources/Win64/plugins/unreal/resources/plugInfo.json
+++ b/Engine/Plugins/Importers/USDImporter/Resources/UsdResources/Win64/plugins/unreal/resources/plugInfo.json
@@ -9,28 +9,42 @@
                     "unrealControlRigAPI": {
                         "alias": {
                             "UsdSchemaBase": "ControlRigAPI"
-                        },
-                        "autoGenerated": true,
+                        }, 
+                        "autoGenerated": true, 
                         "bases": [
                             "UsdAPISchemaBase"
-                        ],
+                        ], 
                         "schemaKind": "singleApplyAPI"
-                    },
+                    }, 
+                    "unrealGroomAPI": {
+                        "alias": {
+                            "UsdSchemaBase": "GroomAPI"
+                        }, 
+                        "apiSchemaCanOnlyApplyTo": [
+                            "Curves", 
+                            "Xform"
+                        ], 
+                        "autoGenerated": true, 
+                        "bases": [
+                            "UsdAPISchemaBase"
+                        ], 
+                        "schemaKind": "singleApplyAPI"
+                    }, 
                     "unrealLiveLinkAPI": {
                         "alias": {
                             "UsdSchemaBase": "LiveLinkAPI"
-                        },
-                        "autoGenerated": true,
+                        }, 
+                        "autoGenerated": true, 
                         "bases": [
                             "UsdAPISchemaBase"
-                        ],
+                        ], 
                         "schemaKind": "singleApplyAPI"
                     }
                 }
-            },
-            "Name": "unreal",
-            "ResourcePath": "resources",
-            "Root": "..",
+            }, 
+            "Name": "unreal", 
+            "ResourcePath": "resources", 
+            "Root": "..", 
             "Type": "resource"
         }
     ]
diff --git a/Engine/Plugins/Importers/USDImporter/Source/USDSchemas/Private/USDGroomTranslator.cpp b/Engine/Plugins/Importers/USDImporter/Source/USDSchemas/Private/USDGroomTranslator.cpp
--- a/Engine/Plugins/Importers/USDImporter/Source/USDSchemas/Private/USDGroomTranslator.cpp
+++ b/Engine/Plugins/Importers/USDImporter/Source/USDSchemas/Private/USDGroomTranslator.cpp
@@ -12,8 +12,9 @@
 #include "Misc/ArchiveMD5.h"
 
 #include "USDAssetImportData.h"
-#include "USDTypesConversion.h"
 #include "USDGroomConversion.h"
+#include "USDIntegrationUtils.h"
+#include "USDTypesConversion.h"
 
 #include "USDIncludesStart.h"
 #include "pxr/base/tf/type.h"
@@ -162,12 +163,7 @@ protected:
 
 bool FUsdGroomTranslator::IsGroomPrim() const
 {
-	FScopedUsdAllocs UsdAllocs;
-
-	// #ueent_todo: Better validation pf prim being a groom like using custom groom schema
-	TArray<TUsdStore<pxr::UsdPrim>> ChildCurvesPrims = UsdUtils::GetAllPrimsOfType(GetPrim(), pxr::TfType::Find<pxr::UsdGeomCurves>());
-
-	return ChildCurvesPrims.Num() > 0;
+	return UsdUtils::PrimHasGroomSchema(GetPrim());
 }
 
 void FUsdGroomTranslator::CreateAssets()
diff --git a/Engine/Plugins/Importers/USDImporter/Source/USDSchemas/Private/USDSchemasModule.cpp b/Engine/Plugins/Importers/USDImporter/Source/USDSchemas/Private/USDSchemasModule.cpp
--- a/Engine/Plugins/Importers/USDImporter/Source/USDSchemas/Private/USDSchemasModule.cpp
+++ b/Engine/Plugins/Importers/USDImporter/Source/USDSchemas/Private/USDSchemasModule.cpp
@@ -49,7 +49,7 @@ public:
 		if ( GIsEditor )
 		{
 			UsdSkelRootTranslatorHandle = GetTranslatorRegistry().Register< FUsdSkelRootTranslator >( TEXT("UsdSkelRoot") );
-			//UsdGroomTranslatorHandle = GetTranslatorRegistry().Register< FUsdGroomTranslator >( TEXT("UsdGeomXformable") ); // #ueent_todo: Enable it when the groom schema is implemented
+			UsdGroomTranslatorHandle = GetTranslatorRegistry().Register< FUsdGroomTranslator >( TEXT("UsdGeomXformable") );
 
 			if ( IMDLImporterModule* MDLImporterModule = FModuleManager::Get().LoadModulePtr< IMDLImporterModule >( TEXT("MDLImporter") ) )
 			{
diff --git a/Engine/Plugins/Importers/USDImporter/Source/USDStageEditor/Private/SUSDStageTreeView.cpp b/Engine/Plugins/Importers/USDImporter/Source/USDStageEditor/Private/SUSDStageTreeView.cpp
--- a/Engine/Plugins/Importers/USDImporter/Source/USDStageEditor/Private/SUSDStageTreeView.cpp
+++ b/Engine/Plugins/Importers/USDImporter/Source/USDStageEditor/Private/SUSDStageTreeView.cpp
@@ -714,6 +714,35 @@ TSharedPtr< SWidget > SUsdStageTreeView::ConstructPrimContextMenu()
 			NAME_None,
 			EUserInterfaceActionType::Button
 		);
+
+		PrimOptions.AddMenuEntry(
+			LOCTEXT( "ApplyGroomSchema", "Apply Groom schema" ),
+			LOCTEXT( "ApplyGroomSchema_ToolTip", "Applies the Groom schema to interpret the prim and its children as a groom" ),
+			FSlateIcon(),
+			FUIAction(
+				FExecuteAction::CreateSP( this, &SUsdStageTreeView::OnApplyGroomSchema ),
+				FCanExecuteAction::CreateSP( this, &SUsdStageTreeView::CanApplyGroomSchema ),
+				FIsActionChecked(),
+				FIsActionButtonVisible::CreateSP( this, &SUsdStageTreeView::CanApplyGroomSchema )
+			),
+			NAME_None,
+			EUserInterfaceActionType::Button
+		);
+
+		PrimOptions.AddMenuEntry(
+			LOCTEXT( "RemoveGroomSchema", "Remove Groom schema" ),
+			LOCTEXT( "RemoveGroomSchema_ToolTip", "Removes the Groom schema from the prim to stop interpreting it as a groom" ),
+			FSlateIcon(),
+			FUIAction(
+				FExecuteAction::CreateSP( this, &SUsdStageTreeView::OnRemoveGroomSchema ),
+				FCanExecuteAction::CreateSP( this, &SUsdStageTreeView::CanRemoveGroomSchema ),
+				FIsActionChecked(),
+				FIsActionButtonVisible::CreateSP( this, &SUsdStageTreeView::CanRemoveGroomSchema )
+			),
+			NAME_None,
+			EUserInterfaceActionType::Button
+		);
+
 	}
 	PrimOptions.EndSection();
 
@@ -917,6 +946,44 @@ void SUsdStageTreeView::OnRemoveControlRig()
 	}
 }
 
+void SUsdStageTreeView::OnApplyGroomSchema()
+{
+	if ( !UsdStageActor.IsValid() )
+	{
+		return;
+	}
+
+	FScopedTransaction Transaction( LOCTEXT( "ApplyGroomSchemaTransaction", "Apply Groom schema to selected prims" ) );
+
+	TArray< FUsdPrimViewModelRef > MySelectedItems = GetSelectedItems();
+
+	UE::FSdfChangeBlock Block;
+
+	for ( FUsdPrimViewModelRef SelectedItem : MySelectedItems )
+	{
+		SelectedItem->ApplyGroomSchema();
+	}
+}
+
+void SUsdStageTreeView::OnRemoveGroomSchema()
+{
+	if ( !UsdStageActor.IsValid() )
+	{
+		return;
+	}
+
+	FScopedTransaction Transaction( LOCTEXT( "RemoveGroomSchemaTransaction", "Remove Groom schema from selected prims" ) );
+
+	TArray< FUsdPrimViewModelRef > MySelectedItems = GetSelectedItems();
+
+	UE::FSdfChangeBlock Block;
+
+	for ( FUsdPrimViewModelRef SelectedItem : MySelectedItems )
+	{
+		SelectedItem->RemoveGroomSchema();
+	}
+}
+
 bool SUsdStageTreeView::CanAddPrim() const
 {
 	if ( !UsdStageActor.IsValid() )
@@ -1067,6 +1134,58 @@ bool SUsdStageTreeView::CanRemoveControlRig() const
 	return false;
 }
 
+bool SUsdStageTreeView::CanApplyGroomSchema() const
+{
+	if ( !UsdStageActor.IsValid() )
+	{
+		return false;
+	}
+
+	UE::FUsdStage UsdStage = UsdStageActor->GetOrLoadUsdStage();
+	if ( !UsdStage || !UsdStage.IsEditTargetValid() )
+	{
+		return false;
+	}
+
+	TArray< FUsdPrimViewModelRef > MySelectedItems = GetSelectedItems();
+
+	for ( FUsdPrimViewModelRef SelectedItem : MySelectedItems )
+	{
+		if ( SelectedItem->CanApplyGroomSchema() )
+		{
+			return true;
+		}
+	}
+
+	return false;
+}
+
+bool SUsdStageTreeView::CanRemoveGroomSchema() const
+{
+	if ( !UsdStageActor.IsValid() )
+	{
+		return false;
+	}
+
+	UE::FUsdStage UsdStage = UsdStageActor->GetOrLoadUsdStage();
+	if ( !UsdStage || !UsdStage.IsEditTargetValid() )
+	{
+		return false;
+	}
+
+	TArray< FUsdPrimViewModelRef > MySelectedItems = GetSelectedItems();
+
+	for ( FUsdPrimViewModelRef SelectedItem : MySelectedItems )
+	{
+		if ( SelectedItem->CanRemoveGroomSchema() )
+		{
+			return true;
+		}
+	}
+
+	return false;
+}
+
 void SUsdStageTreeView::RequestListRefresh()
 {
 	SUsdTreeView< FUsdPrimViewModelRef >::RequestListRefresh();
diff --git a/Engine/Plugins/Importers/USDImporter/Source/USDStageEditor/Private/SUSDStageTreeView.h b/Engine/Plugins/Importers/USDImporter/Source/USDStageEditor/Private/SUSDStageTreeView.h
--- a/Engine/Plugins/Importers/USDImporter/Source/USDStageEditor/Private/SUSDStageTreeView.h
+++ b/Engine/Plugins/Importers/USDImporter/Source/USDStageEditor/Private/SUSDStageTreeView.h
@@ -58,6 +58,8 @@ private:
 	void OnRemoveLiveLink();
 	void OnSetUpControlRig();
 	void OnRemoveControlRig();
+	void OnApplyGroomSchema();
+	void OnRemoveGroomSchema();
 
 	bool CanAddPrim() const;
 	bool CanExecutePrimAction() const;
@@ -65,6 +67,8 @@ private:
 	bool CanRemoveLiveLink() const;
 	bool CanSetUpControlRig() const;
 	bool CanRemoveControlRig() const;
+	bool CanApplyGroomSchema() const;
+	bool CanRemoveGroomSchema() const;
 
 	/** Uses TreeItemExpansionStates to travel the tree and call SetItemExpansion */
 	void RestoreExpansionStates();
diff --git a/Engine/Plugins/Importers/USDImporter/Source/USDStageEditorViewModels/Private/USDPrimViewModel.cpp b/Engine/Plugins/Importers/USDImporter/Source/USDStageEditorViewModels/Private/USDPrimViewModel.cpp
--- a/Engine/Plugins/Importers/USDImporter/Source/USDStageEditorViewModels/Private/USDPrimViewModel.cpp
+++ b/Engine/Plugins/Importers/USDImporter/Source/USDStageEditorViewModels/Private/USDPrimViewModel.cpp
@@ -400,6 +400,44 @@ bool FUsdPrimViewModel::CanRemoveControlRig() const
 #endif // #if USE_USD_SDK
 }
 
+void FUsdPrimViewModel::ApplyGroomSchema()
+{
+#if USE_USD_SDK
+	UsdUtils::ApplyGroomSchema( UsdPrim );
+#endif // #if USE_USD_SDK
+}
+
+bool FUsdPrimViewModel::CanApplyGroomSchema() const
+{
+#if USE_USD_SDK
+	const bool bAddSchema = true;
+	return UsdPrim
+		&& !UsdPrim.IsPseudoRoot()
+		&& UE::USDPrimViewModel::Private::CanAddOrRemoveCustomSchema( UsdPrim, UnrealIdentifiers::GroomAPI, bAddSchema );
+#else
+	return false;
+#endif // #if USE_USD_SDK
+}
+
+void FUsdPrimViewModel::RemoveGroomSchema()
+{
+#if USE_USD_SDK
+	UE::USDPrimViewModel::Private::RemoveCustomSchema( UsdPrim, UnrealIdentifiers::GroomAPI );
+#endif // #if USE_USD_SDK
+}
+
+bool FUsdPrimViewModel::CanRemoveGroomSchema() const
+{
+#if USE_USD_SDK
+	const bool bAddSchema = false;
+	return UsdPrim
+		&& !UsdPrim.IsPseudoRoot()
+		&& UE::USDPrimViewModel::Private::CanAddOrRemoveCustomSchema( UsdPrim, UnrealIdentifiers::GroomAPI, bAddSchema );
+#else
+	return false;
+#endif // #if USE_USD_SDK
+}
+
 void FUsdPrimViewModel::DefinePrim( const TCHAR* PrimName )
 {
 #if USE_USD_SDK
diff --git a/Engine/Plugins/Importers/USDImporter/Source/USDStageEditorViewModels/Public/USDPrimViewModel.h b/Engine/Plugins/Importers/USDImporter/Source/USDStageEditorViewModels/Public/USDPrimViewModel.h
--- a/Engine/Plugins/Importers/USDImporter/Source/USDStageEditorViewModels/Public/USDPrimViewModel.h
+++ b/Engine/Plugins/Importers/USDImporter/Source/USDStageEditorViewModels/Public/USDPrimViewModel.h
@@ -60,6 +60,12 @@ public:
 	void RemoveControlRig();
 	bool CanRemoveControlRig() const;
 
+	void ApplyGroomSchema();
+	bool CanApplyGroomSchema() const;
+
+	void RemoveGroomSchema();
+	bool CanRemoveGroomSchema() const;
+
 	void DefinePrim( const TCHAR* PrimName );
 
 	/** Delegate for hooking up an inline editable text block to be notified that a rename is requested. */
diff --git a/Engine/Plugins/Importers/USDImporter/Source/USDUtilities/Private/USDIntegrationUtils.cpp b/Engine/Plugins/Importers/USDImporter/Source/USDUtilities/Private/USDIntegrationUtils.cpp
--- a/Engine/Plugins/Importers/USDImporter/Source/USDUtilities/Private/USDIntegrationUtils.cpp
+++ b/Engine/Plugins/Importers/USDImporter/Source/USDUtilities/Private/USDIntegrationUtils.cpp
@@ -44,6 +44,19 @@ bool UsdUtils::PrimHasControlRigSchema( const pxr::UsdPrim& Prim )
 	return Prim.HasAPI( Schema );
 }
 
+bool UsdUtils::PrimHasGroomSchema( const pxr::UsdPrim& Prim )
+{
+	if ( !Prim )
+	{
+		return false;
+	}
+
+	FScopedUsdAllocs Allocs;
+
+	pxr::TfType Schema = pxr::UsdSchemaRegistry::GetTypeFromSchemaTypeName( UnrealIdentifiers::GroomAPI );
+	return Prim.HasAPI( Schema );
+}
+
 bool UsdUtils::ApplyControlRigSchema( const pxr::UsdPrim& Prim )
 {
 	if ( !Prim )
@@ -76,6 +89,22 @@ bool UsdUtils::ApplyLiveLinkSchema( const pxr::UsdPrim& Prim )
 	return true;
 }
 
+bool UsdUtils::ApplyGroomSchema( const pxr::UsdPrim& Prim )
+{
+	if ( !Prim )
+	{
+		return false;
+	}
+
+	FScopedUsdAllocs Allocs;
+
+	pxr::TfType Schema = pxr::UsdSchemaRegistry::GetTypeFromSchemaTypeName( UnrealIdentifiers::GroomAPI );
+	ensure( static_cast<bool>( Schema ) );
+
+	ensure( Prim.ApplyAPI( Schema ) );
+	return true;
+}
+
 void UnrealToUsd::ConvertLiveLinkProperties( const UActorComponent* InComponent, pxr::UsdPrim& InOutPrim )
 {
 #if WITH_EDITOR
diff --git a/Engine/Plugins/Importers/USDImporter/Source/USDUtilities/Public/USDIntegrationUtils.h b/Engine/Plugins/Importers/USDImporter/Source/USDUtilities/Public/USDIntegrationUtils.h
--- a/Engine/Plugins/Importers/USDImporter/Source/USDUtilities/Public/USDIntegrationUtils.h
+++ b/Engine/Plugins/Importers/USDImporter/Source/USDUtilities/Public/USDIntegrationUtils.h
@@ -28,6 +28,9 @@ namespace UsdUtils
 	// Returns true in case Prim has our custom "ControlRigAPI" schema
 	USDUTILITIES_API bool PrimHasControlRigSchema( const pxr::UsdPrim& Prim );
 
+	// Returns true if Prim has the custom "GroomAPI" schema
+	USDUTILITIES_API bool PrimHasGroomSchema( const pxr::UsdPrim& Prim );
+
 	// Apply the control rig schema to the prim, creates all the required attributes and sets default values
 	// if they're otherwise not authored
 	USDUTILITIES_API bool ApplyControlRigSchema( const pxr::UsdPrim& Prim );
@@ -36,6 +39,9 @@ namespace UsdUtils
 	// if they're otherwise not authored
 	USDUTILITIES_API bool ApplyLiveLinkSchema( const pxr::UsdPrim& Prim );
 
+	// Apply the groom schema to the prim
+	USDUTILITIES_API bool ApplyGroomSchema( const pxr::UsdPrim& Prim );
+
 #endif // USE_USD_SDK
 }
 
diff --git a/Engine/Plugins/Importers/USDImporter/Source/UnrealUSDWrapper/Private/UnrealUSDWrapper.cpp b/Engine/Plugins/Importers/USDImporter/Source/UnrealUSDWrapper/Private/UnrealUSDWrapper.cpp
--- a/Engine/Plugins/Importers/USDImporter/Source/UnrealUSDWrapper/Private/UnrealUSDWrapper.cpp
+++ b/Engine/Plugins/Importers/USDImporter/Source/UnrealUSDWrapper/Private/UnrealUSDWrapper.cpp
@@ -143,6 +143,8 @@ namespace UnrealIdentifiers
 	const TfToken UsdUVTexture = TfToken( "UsdUVTexture" );
 
 	const TfToken WorldSpaceNormals = TfToken( "worldSpaceNormals" );
+
+	const TfToken GroomAPI = TfToken( "GroomAPI" );
 }
 
 std::string FUsdAttribute::GetUnrealPropertyPath( const pxr::UsdAttribute& Attribute )
diff --git a/Engine/Plugins/Importers/USDImporter/Source/UnrealUSDWrapper/Public/UnrealUSDWrapper.h b/Engine/Plugins/Importers/USDImporter/Source/UnrealUSDWrapper/Public/UnrealUSDWrapper.h
--- a/Engine/Plugins/Importers/USDImporter/Source/UnrealUSDWrapper/Public/UnrealUSDWrapper.h
+++ b/Engine/Plugins/Importers/USDImporter/Source/UnrealUSDWrapper/Public/UnrealUSDWrapper.h
@@ -388,6 +388,8 @@ namespace UnrealIdentifiers
 
 	// Token used to indicate that a material parsed from a material prim should use world space normals
 	extern UNREALUSDWRAPPER_API const pxr::TfToken WorldSpaceNormals;
+
+	extern UNREALUSDWRAPPER_API const pxr::TfToken GroomAPI;
 #endif // #if USE_USD_SDK
 
 	extern UNREALUSDWRAPPER_API const TCHAR* Invisible;
