USD: Groom schema files for Mac/Linux.

diff --git a/Engine/Plugins/Importers/USDImporter/Resources/UsdResources/Linux/plugins/unreal/resources/plugInfo.json b/Engine/Plugins/Importers/USDImporter/Resources/UsdResources/Linux/plugins/unreal/resources/plugInfo.json
--- a/Engine/Plugins/Importers/USDImporter/Resources/UsdResources/Linux/plugins/unreal/resources/plugInfo.json
+++ b/Engine/Plugins/Importers/USDImporter/Resources/UsdResources/Linux/plugins/unreal/resources/plugInfo.json
@@ -9,28 +9,42 @@
                     "unrealControlRigAPI": {
                         "alias": {
                             "UsdSchemaBase": "ControlRigAPI"
-                        },
-                        "autoGenerated": true,
+                        }, 
+                        "autoGenerated": true, 
                         "bases": [
                             "UsdAPISchemaBase"
-                        ],
+                        ], 
                         "schemaKind": "singleApplyAPI"
-                    },
+                    }, 
+                    "unrealGroomAPI": {
+                        "alias": {
+                            "UsdSchemaBase": "GroomAPI"
+                        }, 
+                        "apiSchemaCanOnlyApplyTo": [
+                            "Curves", 
+                            "Xform"
+                        ], 
+                        "autoGenerated": true, 
+                        "bases": [
+                            "UsdAPISchemaBase"
+                        ], 
+                        "schemaKind": "singleApplyAPI"
+                    }, 
                     "unrealLiveLinkAPI": {
                         "alias": {
                             "UsdSchemaBase": "LiveLinkAPI"
-                        },
-                        "autoGenerated": true,
+                        }, 
+                        "autoGenerated": true, 
                         "bases": [
                             "UsdAPISchemaBase"
-                        ],
+                        ], 
                         "schemaKind": "singleApplyAPI"
                     }
                 }
-            },
-            "Name": "unreal",
-            "ResourcePath": "resources",
-            "Root": "..",
+            }, 
+            "Name": "unreal", 
+            "ResourcePath": "resources", 
+            "Root": "..", 
             "Type": "resource"
         }
     ]
diff --git a/Engine/Plugins/Importers/USDImporter/Resources/UsdResources/Mac/plugins/unreal/resources/plugInfo.json b/Engine/Plugins/Importers/USDImporter/Resources/UsdResources/Mac/plugins/unreal/resources/plugInfo.json
--- a/Engine/Plugins/Importers/USDImporter/Resources/UsdResources/Mac/plugins/unreal/resources/plugInfo.json
+++ b/Engine/Plugins/Importers/USDImporter/Resources/UsdResources/Mac/plugins/unreal/resources/plugInfo.json
@@ -9,28 +9,42 @@
                     "unrealControlRigAPI": {
                         "alias": {
                             "UsdSchemaBase": "ControlRigAPI"
-                        },
-                        "autoGenerated": true,
+                        }, 
+                        "autoGenerated": true, 
                         "bases": [
                             "UsdAPISchemaBase"
-                        ],
+                        ], 
                         "schemaKind": "singleApplyAPI"
-                    },
+                    }, 
+                    "unrealGroomAPI": {
+                        "alias": {
+                            "UsdSchemaBase": "GroomAPI"
+                        }, 
+                        "apiSchemaCanOnlyApplyTo": [
+                            "Curves", 
+                            "Xform"
+                        ], 
+                        "autoGenerated": true, 
+                        "bases": [
+                            "UsdAPISchemaBase"
+                        ], 
+                        "schemaKind": "singleApplyAPI"
+                    }, 
                     "unrealLiveLinkAPI": {
                         "alias": {
                             "UsdSchemaBase": "LiveLinkAPI"
-                        },
-                        "autoGenerated": true,
+                        }, 
+                        "autoGenerated": true, 
                         "bases": [
                             "UsdAPISchemaBase"
-                        ],
+                        ], 
                         "schemaKind": "singleApplyAPI"
                     }
                 }
-            },
-            "Name": "unreal",
-            "ResourcePath": "resources",
-            "Root": "..",
+            }, 
+            "Name": "unreal", 
+            "ResourcePath": "resources", 
+            "Root": "..", 
             "Type": "resource"
         }
     ]
