Fixed missing lines from bad robomerge from Release-5.0 to RES (cl 20649260).

diff --git a/Engine/Plugins/Importers/USDImporter/Source/UnrealUSDWrapper/Private/USDListener.cpp b/Engine/Plugins/Importers/USDImporter/Source/UnrealUSDWrapper/Private/USDListener.cpp
--- a/Engine/Plugins/Importers/USDImporter/Source/UnrealUSDWrapper/Private/USDListener.cpp
+++ b/Engine/Plugins/Importers/USDImporter/Source/UnrealUSDWrapper/Private/USDListener.cpp
@@ -42,6 +42,10 @@ namespace UsdToUnreal
 			const FString PrimPath = ANSI_TO_TCHAR( It->GetAbsoluteRootOrPrimPath().GetAsString().c_str() );
 
 			if (pxr::UsdPrim::IsPathInPrototype(UE::FSdfPath(*PrimPath)))
+			{
+				continue;
+			}
+
 			// Something like "/Root/Prim.some_field", or "/"
 			const FString FullFieldPath = ANSI_TO_TCHAR( It->GetAsString().c_str() );
 
