USD: Fix Live Link integration.

Update the default animBlueprintPath to point at the default blueprint now that we properly handle codeless schemas;
Update the DefaultLiveLinkAnimBP.uasset to have no SubjectName set by default (previously it had "root");
Fix user-made AnimBP not working for Live Link anymore due to skeleton IsCompatible call having been reversed;
Fix skeletal mesh components only following Live Link animation after a resync;
Simplify handling of AnimationMode;
Use a change block on integration attribute sets or else the first set fires two events (a resync and an info change, since creating the attribute itself is a resync).

diff --git a/Engine/Plugins/Importers/USDImporter/Source/USDSchemas/Private/USDSkelRootTranslator.cpp b/Engine/Plugins/Importers/USDImporter/Source/USDSchemas/Private/USDSkelRootTranslator.cpp
--- a/Engine/Plugins/Importers/USDImporter/Source/USDSchemas/Private/USDSkelRootTranslator.cpp
+++ b/Engine/Plugins/Importers/USDImporter/Source/USDSchemas/Private/USDSkelRootTranslator.cpp
@@ -1027,7 +1027,7 @@ namespace UsdSkelRootTranslatorImpl
 						}
 					}
 
-					FString CacheKey = UsdToUnreal::ConvertPath( Prim.GetPrimPath() ) + TEXT( "_DefaultAnimBlueprint" );
+					const FString CacheKey = UsdToUnreal::ConvertPath( Prim.GetPrimPath() ) + TEXT( "_DefaultAnimBlueprint" );
 
 					// Check if we can find an AnimBP for this prim in the asset cache (useful when doing Action->Import)
 					if ( !bAlreadyHasTransientAnimBP )
@@ -1072,8 +1072,9 @@ namespace UsdSkelRootTranslatorImpl
 				// Path is pointing to an existing, persistent AnimBP
 				else
 				{
-					// Force skeletons to be compatible
+					// Force skeletons to be compatible (we need both ways!)
 					AnimBP->TargetSkeleton->AddCompatibleSkeleton( Skeleton );
+					Skeleton->AddCompatibleSkeleton( AnimBP->TargetSkeleton );
 
 					// Check ExistingAnimBP so that we only emit this warning when we first set up the component (when we'll realistically not have an ExistingAnimBP yet),
 					// to try and prevent some warning spam
@@ -1190,9 +1191,20 @@ namespace UsdSkelRootTranslatorImpl
 				FCompilerResultsLog Results;
 				FBPCompileRequest Request( AnimBP, EBlueprintCompileOptions::None, &Results );
 				FBlueprintCompilationManager::CompileSynchronously( Request );
+
+				// We need to force the component to update its anim after we regenerate the blueprint class
+				Component->ClearAnimScriptInstance();
+				Component->InitAnim(true);
 			}
 
-			Component->SetAnimInstanceClass( AnimBP->GeneratedClass );
+			if ( AnimBP != ExistingAnimBP )
+			{
+				// This can internally change AnimationMode, but lets revert it to what it was so that we can control it from
+				// that single place in ::UpdateComponents
+				EAnimationMode::Type OldMode = Component->GetAnimationMode();
+				Component->SetAnimInstanceClass( AnimBP->GeneratedClass );
+				Component->SetAnimationMode( OldMode );
+			}
 		}
 	}
 
@@ -1245,6 +1257,14 @@ void FUsdSkelRootTranslator::UpdateComponents( USceneComponent* SceneComponent )
 		}
 	}
 
+	SkeletalMeshComponent->SetAnimationMode(
+		bPrimHasLiveLinkEnabled
+			? EAnimationMode::AnimationBlueprint
+			: Context->bSequencerIsAnimating
+				? EAnimationMode::AnimationCustomMode
+				: EAnimationMode::AnimationSingleNode
+	);
+
 	UE::FUsdPrim SkelAnimPrim = UsdUtils::FindFirstAnimationSource( Prim );
 	if ( SkelAnimPrim )
 	{
@@ -1254,11 +1274,7 @@ void FUsdSkelRootTranslator::UpdateComponents( USceneComponent* SceneComponent )
 			SkeletalMeshComponent->AnimationData.AnimToPlay = TargetAnimSequence;
 			SkeletalMeshComponent->AnimationData.bSavedLooping = false;
 			SkeletalMeshComponent->AnimationData.bSavedPlaying = false;
-
-			if ( SkeletalMeshComponent->GetAnimationMode() == EAnimationMode::AnimationSingleNode )
-			{
-				SkeletalMeshComponent->SetAnimation( TargetAnimSequence );
-			}
+			SkeletalMeshComponent->SetAnimation( TargetAnimSequence );
 		}
 	}
 
diff --git a/Engine/Plugins/Importers/USDImporter/Source/USDStageEditor/Private/Widgets/SUSDIntegrationsPanel.cpp b/Engine/Plugins/Importers/USDImporter/Source/USDStageEditor/Private/Widgets/SUSDIntegrationsPanel.cpp
--- a/Engine/Plugins/Importers/USDImporter/Source/USDStageEditor/Private/Widgets/SUSDIntegrationsPanel.cpp
+++ b/Engine/Plugins/Importers/USDImporter/Source/USDStageEditor/Private/Widgets/SUSDIntegrationsPanel.cpp
@@ -156,6 +156,7 @@ TSharedRef< SWidget > SUsdIntegrationsPanelRow::GenerateWidgetForColumn( const F
 					UE::FVtValue Value;
 					if ( UsdUtils::SetUnderlyingValue( Value, UsdPath ) && !Value.IsEmpty() )
 					{
+						UE::FSdfChangeBlock Block;
 						AttributeCopy->Set( Value );
 					}
 				} );
@@ -212,6 +213,7 @@ TSharedRef< SWidget > SUsdIntegrationsPanelRow::GenerateWidgetForColumn( const F
 					UE::FVtValue Value;
 					if ( UsdUtils::SetUnderlyingValue( Value, UsdPath ) && !Value.IsEmpty() )
 					{
+						UE::FSdfChangeBlock Block;
 						AttributeCopy->Set( Value );
 					}
 				} );
@@ -269,6 +271,7 @@ TSharedRef< SWidget > SUsdIntegrationsPanelRow::GenerateWidgetForColumn( const F
 					UE::FVtValue Value;
 					if ( UsdUtils::SetUnderlyingValue( Value, NewValue == ECheckBoxState::Checked ) && !Value.IsEmpty() )
 					{
+						UE::FSdfChangeBlock Block;
 						AttributeCopy->Set( Value );
 					}
 				})
@@ -293,6 +296,7 @@ TSharedRef< SWidget > SUsdIntegrationsPanelRow::GenerateWidgetForColumn( const F
 					UE::FVtValue Value;
 					if ( UsdUtils::SetUnderlyingValue( Value, NewValue ) && !Value.IsEmpty() )
 					{
+						UE::FSdfChangeBlock Block;
 						AttributeCopy->Set( Value );
 					}
 				})
@@ -358,6 +362,7 @@ TSharedRef< SWidget > SUsdIntegrationsPanelRow::GenerateWidgetForColumn( const F
 				UE::FVtValue Value;
 				if ( UsdUtils::SetUnderlyingValue( Value, UsdValue ) && !Value.IsEmpty() )
 				{
+					UE::FSdfChangeBlock Block;
 					AttributeCopy->Set( Value );
 				}
 			});
@@ -395,6 +400,7 @@ TSharedRef< SWidget > SUsdIntegrationsPanelRow::GenerateWidgetForColumn( const F
 					UE::FVtValue Value;
 					if ( UsdUtils::SetUnderlyingValue( Value, NewState == ECheckBoxState::Checked ) && !Value.IsEmpty() )
 					{
+						UE::FSdfChangeBlock Block;
 						AttributeCopy->Set( Value );
 					}
 				})
