USD: Fix automated tests crashing due to material changes.

diff --git a/Engine/Plugins/Importers/USDImporter/Source/USDStageImporter/Private/USDStageImportContext.cpp b/Engine/Plugins/Importers/USDImporter/Source/USDStageImporter/Private/USDStageImportContext.cpp
--- a/Engine/Plugins/Importers/USDImporter/Source/USDStageImporter/Private/USDStageImportContext.cpp
+++ b/Engine/Plugins/Importers/USDImporter/Source/USDStageImporter/Private/USDStageImportContext.cpp
@@ -99,4 +99,6 @@ void FUsdStageImportContext::Reset()
 	bStageWasOriginallyOpenInCache = false;
 	OriginalMetersPerUnit = 0.01f;
 	OriginalUpAxis = EUsdUpAxis::ZAxis;
+
+	bNeedsGarbageCollection = false;
 }
\ No newline at end of file
diff --git a/Engine/Plugins/Importers/USDImporter/Source/USDStageImporter/Private/USDStageImporter.cpp b/Engine/Plugins/Importers/USDImporter/Source/USDStageImporter/Private/USDStageImporter.cpp
--- a/Engine/Plugins/Importers/USDImporter/Source/USDStageImporter/Private/USDStageImporter.cpp
+++ b/Engine/Plugins/Importers/USDImporter/Source/USDStageImporter/Private/USDStageImporter.cpp
@@ -553,13 +553,14 @@ namespace UsdStageImporterImpl
 		UObject* MovedAsset = ExistingAsset;
 		if (ExistingAsset != nullptr && ExistingAsset != Asset && ReplacePolicy == EReplaceAssetPolicy::Replace)
 		{
-			MovedAsset = DuplicateObject<UObject>(Asset, Package, ExistingAsset->GetFName());
+			MovedAsset = DuplicateObject<UObject>( Asset, Package, *TargetAssetName );
 
-			// If mesh's label has changed, update its name
-			if (ExistingAsset->GetFName() != Asset->GetFName())
+			// If our DuplicateObject didn't stomp the old asset because TargetAssetName != ExistingAsset->GetName(),
+			// mark the existing asset as garbage so that we collect it when we run GC after the import
+			if ( ExistingAsset->GetName() != TargetAssetName )
 			{
-				// We can't dirty the package here. Read the comment around MarkPackageDirty, below
-				MovedAsset->Rename(*TargetAssetName, Package, REN_DontCreateRedirectors | REN_NonTransactional | REN_DoNotDirty);
+				ExistingAsset->MarkAsGarbage();
+				ImportContext.bNeedsGarbageCollection = true;
 			}
 
 			if (UStaticMesh* DestinationMesh = Cast< UStaticMesh >(MovedAsset))
@@ -1519,6 +1520,10 @@ void UUsdStageImporter::ImportFromFile(FUsdStageImportContext& ImportContext)
 
 	UsdStageImporterImpl::CloseStageIfNeeded( ImportContext );
 
+	if ( ImportContext.bNeedsGarbageCollection )
+	{
+		CollectGarbage( GARBAGE_COLLECTION_KEEPFLAGS );
+	}
 #endif // #if USE_USD_SDK
 }
 
@@ -1634,6 +1639,11 @@ bool UUsdStageImporter::ReimportSingleAsset(FUsdStageImportContext& ImportContex
 
 	UsdStageImporterImpl::CloseStageIfNeeded( ImportContext );
 
+	if ( ImportContext.bNeedsGarbageCollection )
+	{
+		CollectGarbage( GARBAGE_COLLECTION_KEEPFLAGS );
+	}
+
 #endif // #if USE_USD_SDK
 	return bSuccess;
 }
diff --git a/Engine/Plugins/Importers/USDImporter/Source/USDStageImporter/Public/USDStageImportContext.h b/Engine/Plugins/Importers/USDImporter/Source/USDStageImporter/Public/USDStageImportContext.h
--- a/Engine/Plugins/Importers/USDImporter/Source/USDStageImporter/Public/USDStageImportContext.h
+++ b/Engine/Plugins/Importers/USDImporter/Source/USDStageImporter/Public/USDStageImportContext.h
@@ -91,6 +91,9 @@ struct USDSTAGEIMPORTER_API FUsdStageImportContext
 	float OriginalMetersPerUnit;
 	EUsdUpAxis OriginalUpAxis;
 
+	/** If we need to run GC after the import is complete */
+	bool bNeedsGarbageCollection;
+
 public:
 	FUsdStageImportContext();
 
