USDImporter: remove support for USD versions earlier than 21.11 and Ar 1.0

A few adjustments that were made to ease the transition from USD
version 21.08 to versions 21.11 and later are no longer necessary
since settling on 22.05b. This includes completing the transition from Ar
1.0 to Ar 2.0. There is no intention to support building the USDImporter
plugin against older releases of USD or the older Ar API.

diff --git a/Engine/Plugins/Importers/USDImporter/Source/USDSchemas/Private/USDLuxLightTranslator.cpp b/Engine/Plugins/Importers/USDImporter/Source/USDSchemas/Private/USDLuxLightTranslator.cpp
--- a/Engine/Plugins/Importers/USDImporter/Source/USDSchemas/Private/USDLuxLightTranslator.cpp
+++ b/Engine/Plugins/Importers/USDImporter/Source/USDSchemas/Private/USDLuxLightTranslator.cpp
@@ -22,11 +22,7 @@
 	#include "pxr/usd/usdLux/diskLight.h"
 	#include "pxr/usd/usdLux/distantLight.h"
 	#include "pxr/usd/usdLux/domeLight.h"
-#if defined(HAS_USDLUX_LIGHTAPI)
 	#include "pxr/usd/usdLux/lightAPI.h"
-#else
-	#include "pxr/usd/usdLux/light.h"
-#endif // #if defined(HAS_USDLUX_LIGHTAPI)
 	#include "pxr/usd/usdLux/rectLight.h"
 	#include "pxr/usd/usdLux/shapingAPI.h"
 	#include "pxr/usd/usdLux/sphereLight.h"
@@ -55,11 +51,7 @@ void FUsdLuxLightTranslator::UpdateComponents( USceneComponent* SceneComponent )
 
 	pxr::UsdPrim Prim = GetPrim();
 
-#if defined(HAS_USDLUX_LIGHTAPI)
 	const pxr::UsdLuxLightAPI LightAPI( Prim );
-#else
-	const pxr::UsdLuxLight LightAPI( Prim );
-#endif // #if defined(HAS_USDLUX_LIGHTAPI)
 
 	if ( !LightAPI )
 	{
diff --git a/Engine/Plugins/Importers/USDImporter/Source/USDSchemas/Private/USDSchemasModule.cpp b/Engine/Plugins/Importers/USDImporter/Source/USDSchemas/Private/USDSchemasModule.cpp
--- a/Engine/Plugins/Importers/USDImporter/Source/USDSchemas/Private/USDSchemasModule.cpp
+++ b/Engine/Plugins/Importers/USDImporter/Source/USDSchemas/Private/USDSchemasModule.cpp
@@ -35,13 +35,8 @@ public:
 		UsdGeomPointInstancerTranslatorHandle = GetTranslatorRegistry().Register< FUsdGeomPointInstancerTranslator >( TEXT("UsdGeomPointInstancer") );
 		UsdGeomXformableTranslatorHandle = GetTranslatorRegistry().Register< FUsdGeomXformableTranslator >( TEXT("UsdGeomXformable") );
 		UsdShadeMaterialTranslatorHandle = GetTranslatorRegistry().Register< FUsdShadeMaterialTranslator >( TEXT("UsdShadeMaterial") );
-
-#if defined(HAS_USDLUX_LIGHTAPI)
 		UsdLuxBoundableLightBaseTranslatorHandle = GetTranslatorRegistry().Register< FUsdLuxLightTranslator >( TEXT("UsdLuxBoundableLightBase") );
 		UsdLuxNonboundableLightBaseTranslatorHandle = GetTranslatorRegistry().Register< FUsdLuxLightTranslator >( TEXT("UsdLuxNonboundableLightBase") );
-#else
-		UsdLuxLightTranslatorHandle = GetTranslatorRegistry().Register< FUsdLuxLightTranslator >( TEXT("UsdLuxLight") );
-#endif
 
 #if WITH_EDITOR
 		// Creating skeletal meshes technically works in Standalone mode, but by checking for this we artificially block it
@@ -80,14 +75,8 @@ public:
 #endif // WITH_EDITOR
 		GetTranslatorRegistry().Unregister( UsdGeomXformableTranslatorHandle );
 		GetTranslatorRegistry().Unregister( UsdShadeMaterialTranslatorHandle );
-
-#if defined(HAS_USDLUX_LIGHTAPI)
 		GetTranslatorRegistry().Unregister( UsdLuxBoundableLightBaseTranslatorHandle );
 		GetTranslatorRegistry().Unregister( UsdLuxNonboundableLightBaseTranslatorHandle );
-#else
-		GetTranslatorRegistry().Unregister( UsdLuxLightTranslatorHandle );
-#endif
-
 #endif // #if USE_USD_SDK
 	}
 
@@ -112,13 +101,8 @@ protected:
 	FRegisteredSchemaTranslatorHandle UsdGeomXformableTranslatorHandle;
 	FRegisteredSchemaTranslatorHandle UsdShadeMaterialTranslatorHandle;
 	FRegisteredSchemaTranslatorHandle UsdGroomTranslatorHandle;
-
-#if defined(HAS_USDLUX_LIGHTAPI)
 	FRegisteredSchemaTranslatorHandle UsdLuxBoundableLightBaseTranslatorHandle;
 	FRegisteredSchemaTranslatorHandle UsdLuxNonboundableLightBaseTranslatorHandle;
-#else
-	FRegisteredSchemaTranslatorHandle UsdLuxLightTranslatorHandle;
-#endif
 
 	// Custom schemas
 	FRegisteredSchemaTranslatorHandle MdlUsdShadeMaterialTranslatorHandle;
diff --git a/Engine/Plugins/Importers/USDImporter/Source/USDStageEditor/Private/SUSDLayersTreeView.cpp b/Engine/Plugins/Importers/USDImporter/Source/USDStageEditor/Private/SUSDLayersTreeView.cpp
--- a/Engine/Plugins/Importers/USDImporter/Source/USDStageEditor/Private/SUSDLayersTreeView.cpp
+++ b/Engine/Plugins/Importers/USDImporter/Source/USDStageEditor/Private/SUSDLayersTreeView.cpp
@@ -49,20 +49,11 @@ namespace UE::USDLayersTreeViewImpl::Private
 		// Convert layer references to absolute paths so that it still works at its target location
 		FString LayerPath = LayerToExport.GetRealPath();
 		FPaths::NormalizeFilename( LayerPath );
-		const TSet<FString> AssetDependencies =
-#if defined(PXR_VERSION) && PXR_VERSION >= 2111
-			OutputLayer.GetCompositionAssetDependencies();
-#else
-			OutputLayer.GetExternalReferences();
-#endif
+		const TSet<FString> AssetDependencies = OutputLayer.GetCompositionAssetDependencies();
 		for ( const FString& Ref : AssetDependencies )
 		{
 			FString AbsRef = FPaths::ConvertRelativePathToFull( FPaths::GetPath( LayerPath ), Ref ); // Relative to the original file
-#if defined(PXR_VERSION) && PXR_VERSION >= 2111
 			OutputLayer.UpdateCompositionAssetDependency( *Ref, *AbsRef );
-#else
-			OutputLayer.UpdateExternalReference( *Ref, *AbsRef );
-#endif
 		}
 
 		bool bForce = true;
diff --git a/Engine/Plugins/Importers/USDImporter/Source/USDStageEditor/Private/SUSDStage.cpp b/Engine/Plugins/Importers/USDImporter/Source/USDStageEditor/Private/SUSDStage.cpp
--- a/Engine/Plugins/Importers/USDImporter/Source/USDStageEditor/Private/SUSDStage.cpp
+++ b/Engine/Plugins/Importers/USDImporter/Source/USDStageEditor/Private/SUSDStage.cpp
@@ -96,11 +96,7 @@ namespace SUSDStageImpl
 			UE::FSdfLayer& Layer = AllLayers[ Index ];
 
 			TArray< UE::FSdfLayer > NewLayers;
-#if defined(PXR_VERSION) && PXR_VERSION >= 2111
 			for ( const FString& AssetDependency : Layer.GetCompositionAssetDependencies() )
-#else
-			for ( const FString& AssetDependency : Layer.GetExternalReferences() )
-#endif
 			{
 				FString AbsoluteReference = Layer.ComputeAbsolutePath( AssetDependency );
 				FPaths::NormalizeFilename( AbsoluteReference );
@@ -1474,22 +1470,14 @@ void SUsdStage::OnExportAll()
 
 			// Remap references to layers so that they point at the other newly saved files. Note that SUSDStageImpl::AppendAllExternalReferences
 			// will have collected all external references already, so OldPathToSavedPath should have entries for all references we'll find
-#if defined(PXR_VERSION) && PXR_VERSION >= 2111
 			for ( const FString& AssetDependency : OutputLayer.GetCompositionAssetDependencies() )
-#else
-			for ( const FString& AssetDependency : OutputLayer.GetExternalReferences() )
-#endif
 			{
 				FString AbsRef = FPaths::ConvertRelativePathToFull( FPaths::GetPath( LayerPath ), AssetDependency ); // Relative to the original file
 				FPaths::NormalizeFilename( AbsRef );
 
 				if ( FString* SavedReference = OldPathToSavedPath.Find( AbsRef ) )
 				{
-#if defined(PXR_VERSION) && PXR_VERSION >= 2111
 					OutputLayer.UpdateCompositionAssetDependency( *AssetDependency, **SavedReference );
-#else
-					OutputLayer.UpdateExternalReference( *AssetDependency, **SavedReference );
-#endif
 				}
 			}
 
diff --git a/Engine/Plugins/Importers/USDImporter/Source/USDUtilities/Private/USDLightConversion.cpp b/Engine/Plugins/Importers/USDImporter/Source/USDUtilities/Private/USDLightConversion.cpp
--- a/Engine/Plugins/Importers/USDImporter/Source/USDUtilities/Private/USDLightConversion.cpp
+++ b/Engine/Plugins/Importers/USDImporter/Source/USDUtilities/Private/USDLightConversion.cpp
@@ -29,11 +29,7 @@
 	#include "pxr/usd/usdLux/diskLight.h"
 	#include "pxr/usd/usdLux/distantLight.h"
 	#include "pxr/usd/usdLux/domeLight.h"
-#if defined(HAS_USDLUX_LIGHTAPI)
 	#include "pxr/usd/usdLux/lightAPI.h"
-#else
-	#include "pxr/usd/usdLux/light.h"
-#endif // #if defined(HAS_USDLUX_LIGHTAPI)
 	#include "pxr/usd/usdLux/rectLight.h"
 	#include "pxr/usd/usdLux/shapingAPI.h"
 	#include "pxr/usd/usdLux/sphereLight.h"
@@ -88,17 +84,10 @@ namespace LightConversionImpl
 	}
 }
 
-#if defined(HAS_USDLUX_LIGHTAPI)
 bool UsdToUnreal::ConvertLight( const pxr::UsdLuxLightAPI& LightAPI, ULightComponentBase& LightComponentBase, double TimeCode )
 {
 	return UsdToUnreal::ConvertLight( LightAPI.GetPrim(), LightComponentBase, TimeCode );
 }
-#else
-bool UsdToUnreal::ConvertLight( const pxr::UsdLuxLight& Light, ULightComponentBase& LightComponentBase, double TimeCode )
-{
-	return UsdToUnreal::ConvertLight( Light.GetPrim(), LightComponentBase, TimeCode );
-}
-#endif // #if defined(HAS_USDLUX_LIGHTAPI)
 
 bool UsdToUnreal::ConvertDistantLight( const pxr::UsdLuxDistantLight& DistantLight, UDirectionalLightComponent& LightComponent, double TimeCode )
 {
@@ -134,12 +123,7 @@ bool UsdToUnreal::ConvertLight( const pxr::UsdPrim& Prim, ULightComponentBase& L
 {
 	FScopedUsdAllocs Allocs;
 
-#if defined(HAS_USDLUX_LIGHTAPI)
 	const pxr::UsdLuxLightAPI LightAPI( Prim );
-#else
-	const pxr::UsdLuxLight LightAPI( Prim );
-#endif // #if defined(HAS_USDLUX_LIGHTAPI)
-
 	if ( !LightAPI )
 	{
 		return false;
@@ -420,12 +404,7 @@ bool UnrealToUsd::ConvertLightComponent( const ULightComponentBase& LightCompone
 		return false;
 	}
 
-#if defined(HAS_USDLUX_LIGHTAPI)
 	pxr::UsdLuxLightAPI LightAPI( Prim );
-#else
-	pxr::UsdLuxLight LightAPI( Prim );
-#endif // #if defined(HAS_USDLUX_LIGHTAPI)
-
 	if ( !LightAPI )
 	{
 		return false;
@@ -500,12 +479,7 @@ bool UnrealToUsd::ConvertRectLightComponent( const URectLightComponent& LightCom
 		return false;
 	}
 
-#if defined(HAS_USDLUX_LIGHTAPI)
 	pxr::UsdLuxLightAPI LightAPI( Prim );
-#else
-	pxr::UsdLuxLight LightAPI( Prim );
-#endif // #if defined(HAS_USDLUX_LIGHTAPI)
-
 	if ( !LightAPI )
 	{
 		return false;
diff --git a/Engine/Plugins/Importers/USDImporter/Source/USDUtilities/Private/USDPrimConversion.cpp b/Engine/Plugins/Importers/USDImporter/Source/USDUtilities/Private/USDPrimConversion.cpp
--- a/Engine/Plugins/Importers/USDImporter/Source/USDUtilities/Private/USDPrimConversion.cpp
+++ b/Engine/Plugins/Importers/USDImporter/Source/USDUtilities/Private/USDPrimConversion.cpp
@@ -71,11 +71,7 @@
 #include "pxr/usd/usdGeom/xformCommonAPI.h"
 #include "pxr/usd/usdLux/diskLight.h"
 #include "pxr/usd/usdLux/distantLight.h"
-#if defined(HAS_USDLUX_LIGHTAPI)
 #include "pxr/usd/usdLux/lightAPI.h"
-#else
-#include "pxr/usd/usdLux/light.h"
-#endif // #if defined(HAS_USDLUX_LIGHTAPI)
 #include "pxr/usd/usdLux/rectLight.h"
 #include "pxr/usd/usdLux/shapingAPI.h"
 #include "pxr/usd/usdLux/sphereLight.h"
@@ -336,12 +332,7 @@ bool UsdToUnreal::ConvertXformable( const pxr::UsdStageRefPtr& Stage, const pxr:
 	FUsdStageInfo StageInfo( Stage );
 	OutTransform = UsdToUnreal::ConvertMatrix( StageInfo, UsdMatrix );
 
-	const bool bPrimIsLight = 
-#if defined(HAS_USDLUX_LIGHTAPI)
-			Xformable.GetPrim().HasAPI< pxr::UsdLuxLightAPI >();
-#else
-			Xformable.GetPrim().IsA< pxr::UsdLuxLight >();
-#endif // #if defined(HAS_USDLUX_LIGHTAPI)
+	const bool bPrimIsLight = Xformable.GetPrim().HasAPI< pxr::UsdLuxLightAPI >();
 
 	// Extra rotation to match different camera facing direction convention
 	// Note: The camera space is always Y-up, yes, but this is not what this is: This is the camera's transform wrt the stage,
@@ -360,12 +351,7 @@ bool UsdToUnreal::ConvertXformable( const pxr::UsdStageRefPtr& Stage, const pxr:
 	// Invert the compensation applied to our parents, in case they're a camera or a light
 	if ( pxr::UsdPrim Parent = Xformable.GetPrim().GetParent() )
 	{
-		const bool bParentIsLight = 
-#if defined(HAS_USDLUX_LIGHTAPI)
-				Parent.HasAPI< pxr::UsdLuxLightAPI >();
-#else
-				Parent.IsA< pxr::UsdLuxLight >();
-#endif // #if defined(HAS_USDLUX_LIGHTAPI)
+		const bool bParentIsLight = Parent.HasAPI< pxr::UsdLuxLightAPI >();
 
 		// If bResetXFormStack is true, then the prim's local transform will be used directly as the world transform, and we will
 		// already invert the parent transform fully, regardless of what it is. This means it doesn't really matter if our parent
@@ -1042,11 +1028,7 @@ UsdToUnreal::FPropertyTrackReader UsdToUnreal::CreatePropertyTrackReader( const
 			}
 		}
 	}
-#if defined(HAS_USDLUX_LIGHTAPI)
 	else if ( const pxr::UsdLuxLightAPI LightAPI{ Prim } )
-#else
-	else if ( const pxr::UsdLuxLight LightAPI{ Prim } )
-#endif // #if defined(HAS_USDLUX_LIGHTAPI)
 	{
 		if ( PropertyPath == UnrealIdentifiers::LightColorPropertyName )
 		{
@@ -2066,12 +2048,7 @@ bool UnrealToUsd::ConvertSceneComponent( const pxr::UsdStageRefPtr& Stage, const
 	// In UE cameras shoot towards local + X, with + Z up.Lights also emit towards local + X, with + Z up
 	// Note that this wouldn't have worked in case we collapsed light and camera components, but these always get their own
 	// actors, so we know that we don't have a single component that represents a large collapsed prim hierarchy
-	if ( UsdPrim.IsA<pxr::UsdGeomCamera>() ||
-#if defined(HAS_USDLUX_LIGHTAPI)
-			UsdPrim.HasAPI<pxr::UsdLuxLightAPI>() )
-#else
-			UsdPrim.IsA<pxr::UsdLuxLight>() )
-#endif // #if defined(HAS_USDLUX_LIGHTAPI)
+	if ( UsdPrim.IsA<pxr::UsdGeomCamera>() || UsdPrim.HasAPI<pxr::UsdLuxLightAPI>() )
 	{
 		FTransform AdditionalRotation = FTransform( FRotator( 0.0f, 90.f, 0.0f ) );
 
@@ -2086,12 +2063,7 @@ bool UnrealToUsd::ConvertSceneComponent( const pxr::UsdStageRefPtr& Stage, const
 	// Invert compensation applied to parent if it's a light or camera component
 	if ( pxr::UsdPrim ParentPrim = UsdPrim.GetParent() )
 	{
-		if ( ParentPrim.IsA<pxr::UsdGeomCamera>() ||
-#if defined(HAS_USDLUX_LIGHTAPI)
-			ParentPrim.HasAPI<pxr::UsdLuxLightAPI>() )
-#else
-			ParentPrim.IsA<pxr::UsdLuxLight>() )
-#endif // #if defined(HAS_USDLUX_LIGHTAPI)
+		if ( ParentPrim.IsA<pxr::UsdGeomCamera>() || ParentPrim.HasAPI<pxr::UsdLuxLightAPI>() )
 		{
 			FTransform AdditionalRotation = FTransform( FRotator( 0.0f, 90.f, 0.0f ) );
 
@@ -2642,13 +2614,7 @@ bool UnrealToUsd::CreateComponentPropertyBaker( UE::FUsdPrim& Prim, const UScene
 
 			// Compensate different orientation for light or camera components
 			FTransform AdditionalRotation = FTransform::Identity;
-			if ( UsdPrim.IsA< pxr::UsdGeomCamera >() ||
-#if defined(HAS_USDLUX_LIGHTAPI)
-					UsdPrim.HasAPI< pxr::UsdLuxLightAPI >() )
-#else
-					UsdPrim.IsA< pxr::UsdLuxLight >() )
-#endif // #if defined(HAS_USDLUX_LIGHTAPI)
-
+			if ( UsdPrim.IsA< pxr::UsdGeomCamera >() || UsdPrim.HasAPI< pxr::UsdLuxLightAPI >() )
 			{
 				AdditionalRotation = FTransform( FRotator( 0.0f, 90.0f, 0.0f ) );
 
@@ -3036,12 +3002,7 @@ UnrealToUsd::FPropertyTrackWriter UnrealToUsd::CreatePropertyTrackWriter( const
 
 					// Compensate different orientation for light or camera components
 					FTransform AdditionalRotation = FTransform::Identity;
-					if ( UsdPrim.IsA< pxr::UsdGeomCamera >() ||
-#if defined(HAS_USDLUX_LIGHTAPI)
-							UsdPrim.HasAPI< pxr::UsdLuxLightAPI >() )
-#else
-							UsdPrim.IsA< pxr::UsdLuxLight >() )
-#endif // #if defined(HAS_USDLUX_LIGHTAPI)
+					if ( UsdPrim.IsA< pxr::UsdGeomCamera >() || UsdPrim.HasAPI< pxr::UsdLuxLightAPI >() )
 					{
 						AdditionalRotation = FTransform( FRotator( 0.0f, 90.0f, 0.0f ) );
 
@@ -3153,11 +3114,7 @@ UnrealToUsd::FPropertyTrackWriter UnrealToUsd::CreatePropertyTrackWriter( const
 			}
 		}
 	}
-#if defined(HAS_USDLUX_LIGHTAPI)
 	else if ( pxr::UsdLuxLightAPI LightAPI{ Prim } )
-#else
-	else if ( pxr::UsdLuxLight LightAPI{ Prim } )
-#endif // #if defined(HAS_USDLUX_LIGHTAPI)
 	{
 		if ( PropertyPath == UnrealIdentifiers::LightColorPropertyName )
 		{
@@ -3630,12 +3587,7 @@ bool UnrealToUsd::ConvertXformable( const UMovieScene3DTransformTrack& MovieScen
 
 		// Compensate different orientation for light or camera components
 		FTransform CameraCompensation = FTransform::Identity;
-		if ( UsdPrim.IsA< pxr::UsdGeomCamera >() ||
-#if defined(HAS_USDLUX_LIGHTAPI)
-				UsdPrim.HasAPI< pxr::UsdLuxLightAPI >() )
-#else
-				UsdPrim.IsA< pxr::UsdLuxLight >() )
-#endif // #if defined(HAS_USDLUX_LIGHTAPI)
+		if ( UsdPrim.IsA< pxr::UsdGeomCamera >() || UsdPrim.HasAPI< pxr::UsdLuxLightAPI >() )
 		{
 			CameraCompensation = FTransform( FRotator( 0.0f, 90.0f, 0.0f ) );
 
@@ -3649,12 +3601,7 @@ bool UnrealToUsd::ConvertXformable( const UMovieScene3DTransformTrack& MovieScen
 		FTransform InverseCameraCompensation = FTransform::Identity;
 		if ( pxr::UsdPrim ParentPrim = UsdPrim.GetParent() )
 		{
-			if ( ParentPrim.IsA<pxr::UsdGeomCamera>() ||
-#if defined(HAS_USDLUX_LIGHTAPI)
-					ParentPrim.HasAPI<pxr::UsdLuxLightAPI>() )
-#else
-					ParentPrim.IsA<pxr::UsdLuxLight>() )
-#endif // #if defined(HAS_USDLUX_LIGHTAPI)
+			if ( ParentPrim.IsA<pxr::UsdGeomCamera>() || ParentPrim.HasAPI<pxr::UsdLuxLightAPI>() )
 			{
 				InverseCameraCompensation = FTransform( FRotator( 0.0f, 90.f, 0.0f ) );
 
diff --git a/Engine/Plugins/Importers/USDImporter/Source/USDUtilities/Private/USDShadeConversion.cpp b/Engine/Plugins/Importers/USDImporter/Source/USDUtilities/Private/USDShadeConversion.cpp
--- a/Engine/Plugins/Importers/USDImporter/Source/USDUtilities/Private/USDShadeConversion.cpp
+++ b/Engine/Plugins/Importers/USDImporter/Source/USDUtilities/Private/USDShadeConversion.cpp
@@ -141,11 +141,7 @@ namespace UE
 				// Don't normalize an empty path as the result will be "."
 				if ( AssetIdentifier.size() > 0 )
 				{
-#if AR_VERSION >= 2
 					AssetIdentifier = Resolver.CreateIdentifier( AssetIdentifier );
-#else
-					AssetIdentifier = Resolver.ComputeNormalizedPath( AssetIdentifier );
-#endif // #if AR_VERSION >= 2
 				}
 
 				FString ResolvedAssetPath = UsdToUnreal::ConvertString( AssetIdentifier );
@@ -187,11 +183,7 @@ namespace UE
 			{
 				pxr::ArResolver& Resolver = pxr::ArGetResolver();
 				std::shared_ptr<pxr::ArAsset> Asset = Resolver.OpenAsset(
-#if AR_VERSION >= 2
 					pxr::ArResolvedPath( UnrealToUsd::ConvertString( *ResolvedTexturePath ).Get() ) );
-#else
-					UnrealToUsd::ConvertString( *ResolvedTexturePath ).Get() );
-#endif // #if AR_VERSION >= 2
 
 				TUsdStore<std::shared_ptr<const char>> Buffer;
 
@@ -2207,11 +2199,7 @@ FString UsdUtils::GetResolvedTexturePath( const pxr::UsdAttribute& TextureAssetP
 	// Don't normalize an empty path as the result will be "."
 	if ( AssetIdentifier.size() > 0 )
 	{
-#if AR_VERSION >= 2
 		AssetIdentifier = Resolver.CreateIdentifier( AssetIdentifier );
-#else
-		AssetIdentifier = Resolver.ComputeNormalizedPath( AssetIdentifier );
-#endif // #if AR_VERSION >= 2
 	}
 
 	FString ResolvedTexturePath = UsdToUnreal::ConvertString( AssetIdentifier );
diff --git a/Engine/Plugins/Importers/USDImporter/Source/USDUtilities/Public/USDLightConversion.h b/Engine/Plugins/Importers/USDImporter/Source/USDUtilities/Public/USDLightConversion.h
--- a/Engine/Plugins/Importers/USDImporter/Source/USDUtilities/Public/USDLightConversion.h
+++ b/Engine/Plugins/Importers/USDImporter/Source/USDUtilities/Public/USDLightConversion.h
@@ -17,11 +17,7 @@ PXR_NAMESPACE_OPEN_SCOPE
 class UsdLuxDiskLight;
 class UsdLuxDistantLight;
 class UsdLuxDomeLight;
-#if defined(HAS_USDLUX_LIGHTAPI)
 class UsdLuxLightAPI;
-#else
-class UsdLuxLight;
-#endif // #if defined(HAS_USDLUX_LIGHTAPI)
 class UsdLuxRectLight;
 class UsdLuxShapingAPI;
 class UsdLuxSphereLight;
@@ -48,7 +44,7 @@ struct FUsdStageInfo;
  *
  * Corresponding UsdLux light schema to Unreal component:
  *
- *	UsdLuxLight[API]	->	ULightComponent
+ *	UsdLuxLightAPI		->	ULightComponent
  *	UsdLuxDistantLight	->	UDirectionalLightComponent
  *	UsdLuxRectLight		->	URectLightComponent
  *	UsdLuxDiskLight		->	URectLightComponent
@@ -58,13 +54,8 @@ struct FUsdStageInfo;
  */
 namespace UsdToUnreal
 {
-#if defined(HAS_USDLUX_LIGHTAPI)
 	UE_DEPRECATED( 5.0, "Prefer the overload that receives a pxr::UsdPrim" )
 	USDUTILITIES_API bool ConvertLight( const pxr::UsdLuxLightAPI& LightAPI, ULightComponentBase& LightComponentBase, double TimeCode );
-#else
-	UE_DEPRECATED( 5.0, "Prefer the overload that receives a pxr::UsdPrim" )
-	USDUTILITIES_API bool ConvertLight( const pxr::UsdLuxLight& Light, ULightComponentBase& LightComponentBase, double TimeCode );
-#endif // #if defined(HAS_USDLUX_LIGHTAPI)
 	UE_DEPRECATED( 5.0, "Prefer the overload that receives a pxr::UsdPrim" )
 	USDUTILITIES_API bool ConvertDistantLight( const pxr::UsdLuxDistantLight& DistantLight, UDirectionalLightComponent& LightComponent, double TimeCode );
 	UE_DEPRECATED( 5.0, "Prefer the overload that receives a pxr::UsdPrim" )
diff --git a/Engine/Plugins/Importers/USDImporter/Source/UnrealUSDWrapper/Private/UsdWrappers/SdfLayer.cpp b/Engine/Plugins/Importers/USDImporter/Source/UnrealUSDWrapper/Private/UsdWrappers/SdfLayer.cpp
--- a/Engine/Plugins/Importers/USDImporter/Source/UnrealUSDWrapper/Private/UsdWrappers/SdfLayer.cpp
+++ b/Engine/Plugins/Importers/USDImporter/Source/UnrealUSDWrapper/Private/UsdWrappers/SdfLayer.cpp
@@ -271,42 +271,6 @@ namespace UE
 		return false;
 	}
 
-	template<typename PtrType>
-	TSet<FString> FSdfLayerBase<PtrType>::GetExternalReferences() const
-	{
-		TSet<FString> Result;
-#if USE_USD_SDK
-		if ( const PtrType& Ptr = Impl->GetInner() )
-		{
-			FScopedUsdAllocs UsdAllocs;
-			std::set<std::string> ExternalReferences = Ptr->GetExternalReferences();
-
-			Result.Reserve( ExternalReferences.size() );
-
-			for ( const std::string& Reference : ExternalReferences )
-			{
-				Result.Add( ANSI_TO_TCHAR( Reference.c_str() ) );
-			}
-		}
-#endif // #if USE_USD_SDK
-		return Result;
-	}
-
-	template<typename PtrType>
-	bool FSdfLayerBase<PtrType>::UpdateExternalReference( const TCHAR* OldReferencePath, const TCHAR* NewReferencePath )
-	{
-#if USE_USD_SDK
-		if ( const PtrType& Ptr = Impl->GetInner() )
-		{
-			FScopedUsdAllocs UsdAllocs;
-			return Ptr->UpdateExternalReference( TCHAR_TO_ANSI( OldReferencePath ), TCHAR_TO_ANSI( NewReferencePath ) );
-		}
-#endif // #if USE_USD_SDK
-
-		return false;
-	}
-
-#if defined(PXR_VERSION) && PXR_VERSION >= 2111
 	template<typename PtrType>
 	TSet<FString> FSdfLayerBase<PtrType>::GetCompositionAssetDependencies() const
 	{
@@ -343,7 +307,6 @@ namespace UE
 #endif // #if USE_USD_SDK
 		return false;
 	}
-#endif // #if defined(PXR_VERSION) && PXR_VERSION >= 2111
 
 	template<typename PtrType>
 	FString FSdfLayerBase<PtrType>::GetRealPath() const
diff --git a/Engine/Plugins/Importers/USDImporter/Source/UnrealUSDWrapper/Public/USDIncludesStart.h b/Engine/Plugins/Importers/USDImporter/Source/UnrealUSDWrapper/Public/USDIncludesStart.h
--- a/Engine/Plugins/Importers/USDImporter/Source/UnrealUSDWrapper/Public/USDIncludesStart.h
+++ b/Engine/Plugins/Importers/USDImporter/Source/UnrealUSDWrapper/Public/USDIncludesStart.h
@@ -63,10 +63,4 @@ THIRD_PARTY_INCLUDES_START
 
 #include "pxr/pxr.h"
 
-// In USD 21.11, the UsdLuxLight schema was removed and replaced with the
-// single-apply API schema UsdLuxLightAPI.
-#if defined(PXR_VERSION) && PXR_VERSION >= 2111
-#define HAS_USDLUX_LIGHTAPI
-#endif
-
 #endif // #if USE_USD_SDK
diff --git a/Engine/Plugins/Importers/USDImporter/Source/UnrealUSDWrapper/Public/UsdWrappers/SdfLayer.h b/Engine/Plugins/Importers/USDImporter/Source/UnrealUSDWrapper/Public/UsdWrappers/SdfLayer.h
--- a/Engine/Plugins/Importers/USDImporter/Source/UnrealUSDWrapper/Public/UsdWrappers/SdfLayer.h
+++ b/Engine/Plugins/Importers/USDImporter/Source/UnrealUSDWrapper/Public/UsdWrappers/SdfLayer.h
@@ -87,16 +87,11 @@ namespace UE
 
 		bool Save( bool bForce = false ) const;
 
-		TSet<FString> GetExternalReferences() const;
-		bool UpdateExternalReference( const TCHAR* OldReferencePath, const TCHAR* NewReferencePath );
-
-#if defined(PXR_VERSION) && PXR_VERSION >= 2111
 		TSet<FString> GetCompositionAssetDependencies() const;
 
 		bool UpdateCompositionAssetDependency(
 			const TCHAR* OldAssetPath,
 			const TCHAR* NewAssetPath = nullptr);
-#endif // #if defined(PXR_VERSION) && PXR_VERSION >= 2111
 
 		FString GetRealPath() const;
 		FString GetIdentifier() const;
