USD: Added GroomBinding schema to allow binding a groom to a SkelRoot or GeomMesh.

diff --git a/Engine/Plugins/Importers/USDImporter/Resources/UsdResources/Win64/plugins/unreal/resources/plugInfo.json b/Engine/Plugins/Importers/USDImporter/Resources/UsdResources/Win64/plugins/unreal/resources/plugInfo.json
--- a/Engine/Plugins/Importers/USDImporter/Resources/UsdResources/Win64/plugins/unreal/resources/plugInfo.json
+++ b/Engine/Plugins/Importers/USDImporter/Resources/UsdResources/Win64/plugins/unreal/resources/plugInfo.json
@@ -30,6 +30,20 @@
                         ], 
                         "schemaKind": "singleApplyAPI"
                     }, 
+                    "unrealGroomBindingAPI": {
+                        "alias": {
+                            "UsdSchemaBase": "GroomBindingAPI"
+                        }, 
+                        "apiSchemaCanOnlyApplyTo": [
+                            "Mesh", 
+                            "SkelRoot"
+                        ], 
+                        "autoGenerated": true, 
+                        "bases": [
+                            "UsdAPISchemaBase"
+                        ], 
+                        "schemaKind": "singleApplyAPI"
+                    }, 
                     "unrealLiveLinkAPI": {
                         "alias": {
                             "UsdSchemaBase": "LiveLinkAPI"
diff --git a/Engine/Plugins/Importers/USDImporter/Source/GeometryCacheUSD/Private/GeometryCacheTrackUSD.cpp b/Engine/Plugins/Importers/USDImporter/Source/GeometryCacheUSD/Private/GeometryCacheTrackUSD.cpp
--- a/Engine/Plugins/Importers/USDImporter/Source/GeometryCacheUSD/Private/GeometryCacheTrackUSD.cpp
+++ b/Engine/Plugins/Importers/USDImporter/Source/GeometryCacheUSD/Private/GeometryCacheTrackUSD.cpp
@@ -109,6 +109,13 @@ const FGeometryCacheTrackSampleInfo& UGeometryCacheTrackUsd::GetSampleInfo(float
 	return CurrentSampleInfo;
 }
 
+bool UGeometryCacheTrackUsd::GetMeshDataAtTime(float Time, FGeometryCacheMeshData& OutMeshData)
+{
+	const bool bLooping = true;
+	const int32 SampleIndex = FindSampleIndexFromTime(Time, bLooping);
+	return GetMeshData(SampleIndex, OutMeshData);
+}
+
 bool UGeometryCacheTrackUsd::GetMeshData(int32 SampleIndex, FGeometryCacheMeshData& OutMeshData)
 {
 	if (IGeometryCacheStreamer::Get().IsTrackRegistered(this))
diff --git a/Engine/Plugins/Importers/USDImporter/Source/GeometryCacheUSD/Public/GeometryCacheTrackUSD.h b/Engine/Plugins/Importers/USDImporter/Source/GeometryCacheUSD/Public/GeometryCacheTrackUSD.h
--- a/Engine/Plugins/Importers/USDImporter/Source/GeometryCacheUSD/Public/GeometryCacheTrackUSD.h
+++ b/Engine/Plugins/Importers/USDImporter/Source/GeometryCacheUSD/Public/GeometryCacheTrackUSD.h
@@ -38,6 +38,7 @@ public:
 	virtual const bool UpdateMeshData(const float Time, const bool bLooping, int32& InOutMeshSampleIndex, FGeometryCacheMeshData*& OutMeshData) override;
 	virtual const bool UpdateBoundsData(const float Time, const bool bLooping, const bool bIsPlayingBackward, int32& InOutBoundsSampleIndex, FBox& OutBounds) override;
 	virtual const FGeometryCacheTrackSampleInfo& GetSampleInfo(float Time, const bool bLooping) override;
+	virtual bool GetMeshDataAtTime(float Time, FGeometryCacheMeshData& OutMeshData) override;
 	virtual void UpdateTime(float Time, bool bLooping) override;
 	//~ End UGeometryCacheTrack Interface.
 
diff --git a/Engine/Plugins/Importers/USDImporter/Source/USDSchemas/Private/USDGeomMeshTranslator.cpp b/Engine/Plugins/Importers/USDImporter/Source/USDSchemas/Private/USDGeomMeshTranslator.cpp
--- a/Engine/Plugins/Importers/USDImporter/Source/USDSchemas/Private/USDGeomMeshTranslator.cpp
+++ b/Engine/Plugins/Importers/USDImporter/Source/USDSchemas/Private/USDGeomMeshTranslator.cpp
@@ -10,6 +10,8 @@
 #include "USDAssetImportData.h"
 #include "USDConversionUtils.h"
 #include "USDGeomMeshConversion.h"
+#include "USDGroomTranslatorUtils.h"
+#include "USDIntegrationUtils.h"
 #include "USDLog.h"
 #include "USDPrimConversion.h"
 #include "USDTypesConversion.h"
@@ -24,6 +26,7 @@
 #include "Engine/StaticMesh.h"
 #include "GeometryCache.h"
 #include "GeometryCacheComponent.h"
+#include "GroomComponent.h"
 #include "Interfaces/ITargetPlatform.h"
 #include "Interfaces/ITargetPlatformManagerModule.h"
 #include "Materials/Material.h"
@@ -1484,6 +1487,24 @@ USceneComponent* FUsdGeomMeshTranslator::CreateComponents()
 				bAllowInterpretingLODs,
 				Context->RenderContext
 			);
+
+#if WITH_EDITOR
+			// Check if the prim has the GroomBinding schema and setup the component and assets necessary to bind the groom to the GeometryCache
+			if ( UsdUtils::PrimHasSchema( GetPrim(), UnrealIdentifiers::GroomBindingAPI ) )
+			{
+				UsdGroomTranslatorUtils::CreateGroomBindingAsset( GetPrim(), *( Context->AssetCache ), Context->ObjectFlags );
+
+				// For the groom binding to work, the GroomComponent must be a child of the SceneComponent 
+				// so the Context ParentComponent is set to the SceneComponent temporarily
+				TGuardValue< USceneComponent* > ParentComponentGuard{ Context->ParentComponent, SceneComponent };
+				const bool bNeedsActor = false;
+				UGroomComponent* GroomComponent = Cast< UGroomComponent >( CreateComponentsEx( TSubclassOf< USceneComponent >( UGroomComponent::StaticClass() ), bNeedsActor ) );
+				if ( GroomComponent )
+				{
+					UpdateComponents( SceneComponent );
+				}
+			}
+#endif // !WITH_EDITOR
 		}
 	}
 
@@ -1555,6 +1576,12 @@ void FUsdGeomMeshTranslator::UpdateComponents( USceneComponent* SceneComponent )
 		{
 			GeometryCacheUsdComponent->RegisterComponent();
 		}
+
+		// If the prim has a GroomBinding schema, apply the target groom to its associated GroomComponent
+		if ( UsdUtils::PrimHasSchema( GetPrim(), UnrealIdentifiers::GroomBindingAPI ) )
+		{
+			UsdGroomTranslatorUtils::SetGroomFromPrim( GetPrim(), *Context->AssetCache, SceneComponent );
+		}
 	}
 #endif // WITH_EDITOR
 
diff --git a/Engine/Plugins/Importers/USDImporter/Source/USDSchemas/Private/USDGroomTranslator.cpp b/Engine/Plugins/Importers/USDImporter/Source/USDSchemas/Private/USDGroomTranslator.cpp
--- a/Engine/Plugins/Importers/USDImporter/Source/USDSchemas/Private/USDGroomTranslator.cpp
+++ b/Engine/Plugins/Importers/USDImporter/Source/USDSchemas/Private/USDGroomTranslator.cpp
@@ -6,11 +6,13 @@
 
 #include "GroomAsset.h"
 #include "GroomBuilder.h"
+#include "GroomComponent.h"
 #include "GroomImportOptions.h"
 #include "HairDescription.h"
 #include "HairStrandsImporter.h"
 #include "Misc/ArchiveMD5.h"
 
+#include "UnrealUSDWrapper.h"
 #include "USDAssetImportData.h"
 #include "USDGroomConversion.h"
 #include "USDIntegrationUtils.h"
@@ -163,7 +165,7 @@ protected:
 
 bool FUsdGroomTranslator::IsGroomPrim() const
 {
-	return UsdUtils::PrimHasGroomSchema(GetPrim());
+	return UsdUtils::PrimHasSchema(GetPrim(), UnrealIdentifiers::GroomAPI);
 }
 
 void FUsdGroomTranslator::CreateAssets()
@@ -183,7 +185,18 @@ USceneComponent* FUsdGroomTranslator::CreateComponents()
 		return Super::CreateComponents();
 	}
 
-	return nullptr;
+	// Display the groom as a standalone actor only if the stage loads the matching purpose.
+	// The groom asset is processed regardless of the purpose so that it can be bound to mesh prims.
+	if (!EnumHasAllFlags(Context->PurposesToLoad, IUsdPrim::GetPurpose(GetPrim())))
+	{
+		return nullptr;
+	}
+
+	bool bNeedsActor = true;
+	USceneComponent* Component = CreateComponentsEx(TSubclassOf<USceneComponent>(UGroomComponent::StaticClass()), bNeedsActor);
+	UpdateComponents(Component);
+
+	return Component;
 }
 
 void FUsdGroomTranslator::UpdateComponents(USceneComponent* SceneComponent)
@@ -192,6 +205,36 @@ void FUsdGroomTranslator::UpdateComponents(USceneComponent* SceneComponent)
 	{
 		Super::UpdateComponents(SceneComponent);
 	}
+
+	if (UGroomComponent* GroomComponent = Cast<UGroomComponent>(SceneComponent))
+	{
+		GroomComponent->Modify();
+
+		UGroomAsset* Groom = Cast<UGroomAsset>(Context->AssetCache->GetAssetForPrim(PrimPath.GetString()));
+
+		bool bShouldRegister = false;
+		if (Groom != GroomComponent->GroomAsset.Get())
+		{
+			bShouldRegister = true;
+
+			if (GroomComponent->IsRegistered())
+			{
+				GroomComponent->UnregisterComponent();
+			}
+
+			GroomComponent->SetGroomAsset(Groom);
+		}
+
+		// Use the prim purpose in conjunction with the prim's computed visibility to toggle the visibility of the groom component
+		// since the component itself cannot be removed if the groom shouldn't be displayed
+		const bool bShouldRender = UsdUtils::IsVisible(GetPrim()) && EnumHasAllFlags(Context->PurposesToLoad, IUsdPrim::GetPurpose(GetPrim()));
+		GroomComponent->SetVisibility(bShouldRender);
+
+		if (bShouldRegister && !GroomComponent->IsRegistered())
+		{
+			GroomComponent->RegisterComponent();
+		}
+	}
 }
 
 bool FUsdGroomTranslator::CollapsesChildren(ECollapsingType CollapsingType) const
diff --git a/Engine/Plugins/Importers/USDImporter/Source/USDSchemas/Private/USDGroomTranslatorUtils.cpp b/Engine/Plugins/Importers/USDImporter/Source/USDSchemas/Private/USDGroomTranslatorUtils.cpp
new file mode 100644
--- /dev/null
+++ b/Engine/Plugins/Importers/USDImporter/Source/USDSchemas/Private/USDGroomTranslatorUtils.cpp
@@ -0,0 +1,289 @@
+// Copyright Epic Games, Inc. All Rights Reserved.
+
+#include "USDGroomTranslatorUtils.h"
+
+#if USE_USD_SDK && WITH_EDITOR
+
+#include "Engine/SkeletalMesh.h"
+#include "GeometryCache.h"
+#include "GroomAsset.h"
+#include "GroomBindingAsset.h"
+#include "GroomComponent.h"
+#include "Misc/SecureHash.h"
+#include "UObject/UObjectGlobals.h"
+
+#include "UnrealUSDWrapper.h"
+#include "USDAssetCache.h"
+#include "USDGeomXformableTranslator.h"
+#include "USDIntegrationUtils.h"
+#include "USDSchemaTranslator.h"
+#include "USDTypesConversion.h"
+
+#include "USDIncludesStart.h"
+#include "pxr/usd/usd/prim.h"
+#include "pxr/usd/usd/relationship.h"
+#include "pxr/usd/usdGeom/mesh.h"
+#include "pxr/usd/usdSkel/root.h"
+#include "USDIncludesEnd.h"
+
+namespace UE::UsdGroomTranslatorUtils::Private
+{
+	struct FGroomBindingBuildSettings : public FGCObject
+	{
+		virtual void AddReferencedObjects(FReferenceCollector& Collector)  override
+		{
+			Collector.AddReferencedObject(Groom);
+			Collector.AddReferencedObject(SourceMesh);
+			Collector.AddReferencedObject(TargetMesh);
+		}
+
+		virtual FString GetReferencerName() const override
+		{
+			return TEXT("FGroomBindingBuildSettings");
+		}
+
+		EGroomBindingMeshType GroomBindingType = EGroomBindingMeshType::SkeletalMesh;
+		UGroomAsset* Groom = nullptr;
+		UObject* SourceMesh = nullptr;
+		UObject* TargetMesh = nullptr;
+		int32 NumInterpolationPoints = 100;
+		int32 MatchingSection = 0;
+	};
+
+	FSHAHash ComputeGroomBindingHash(const FGroomBindingBuildSettings& Settings)
+	{
+		// Ref. GroomBindingAsset BuildDerivedDataKeySuffix
+		FString BindingType;
+		FString SourceKey;
+		FString TargetKey;
+		if (Settings.GroomBindingType == EGroomBindingMeshType::SkeletalMesh)
+		{
+			// Binding type is implicitly SkeletalMesh so keep BindingType empty
+			SourceKey = Cast<USkeletalMesh>(Settings.SourceMesh) ? Cast<USkeletalMesh>(Settings.SourceMesh)->GetDerivedDataKey() : FString();
+			TargetKey = Cast<USkeletalMesh>(Settings.TargetMesh) ? Cast<USkeletalMesh>(Settings.TargetMesh)->GetDerivedDataKey() : FString();
+		}
+		else
+		{
+			BindingType = "GEOCACHE_";
+			SourceKey = Cast<UGeometryCache>(Settings.SourceMesh) ? Cast<UGeometryCache>(Settings.SourceMesh)->GetHash() : FString();
+			TargetKey = Cast<UGeometryCache>(Settings.TargetMesh) ? Cast<UGeometryCache>(Settings.TargetMesh)->GetHash() : FString();
+		}
+		const FString GroomKey = Settings.Groom ? Settings.Groom->GetDerivedDataKey() : FString();
+		const FString PointKey = FString::FromInt(Settings.NumInterpolationPoints);
+		const FString SectionKey = FString::FromInt(Settings.MatchingSection);
+
+		FSHA1 SHA1;
+		SHA1.UpdateWithString(*BindingType, BindingType.Len());
+		SHA1.UpdateWithString(*SourceKey, SourceKey.Len());
+		SHA1.UpdateWithString(*TargetKey, TargetKey.Len());
+		SHA1.UpdateWithString(*GroomKey, GroomKey.Len());
+		SHA1.UpdateWithString(*PointKey, PointKey.Len());
+		SHA1.UpdateWithString(*SectionKey, SectionKey.Len());
+		SHA1.Final();
+
+		FSHAHash SHAHash;
+		SHA1.GetHash(SHAHash.Hash);
+
+		return SHAHash;
+	}
+
+	UGroomBindingAsset* CreateGroomBindingAsset(FString GroomBindingPath, const FGroomBindingBuildSettings& Settings, EObjectFlags ObjectFlags)
+	{
+		// Need at least the groom and the target mesh to build a GroomBindingAsset; the source mesh is optional
+		if (GroomBindingPath.IsEmpty() || !Settings.Groom || !Settings.TargetMesh)
+		{
+			return nullptr;
+		}
+
+		const FName BindingAssetName = MakeUniqueObjectName(GetTransientPackage(), UGroomBindingAsset::StaticClass(), *FPaths::GetBaseFilename(GroomBindingPath));
+		UGroomBindingAsset* GroomBinding = NewObject<UGroomBindingAsset>(GetTransientPackage(), BindingAssetName, ObjectFlags | RF_Public);
+		if (GroomBinding)
+		{
+			GroomBinding->GroomBindingType = Settings.GroomBindingType;
+			GroomBinding->Groom = Settings.Groom;
+			if (GroomBinding->GroomBindingType == EGroomBindingMeshType::SkeletalMesh)
+			{
+				GroomBinding->SourceSkeletalMesh = Cast<USkeletalMesh>(Settings.SourceMesh);
+				GroomBinding->TargetSkeletalMesh = Cast<USkeletalMesh>(Settings.TargetMesh);
+			}
+			else
+			{
+				GroomBinding->SourceGeometryCache = Cast<UGeometryCache>(Settings.SourceMesh);
+				GroomBinding->TargetGeometryCache = Cast<UGeometryCache>(Settings.TargetMesh);
+			}
+			GroomBinding->HairGroupBulkDatas.Reserve(Settings.Groom->HairGroupsData.Num());
+			GroomBinding->NumInterpolationPoints = Settings.NumInterpolationPoints;
+			GroomBinding->MatchingSection = Settings.MatchingSection;
+
+			GroomBinding->Build();
+		}
+		return GroomBinding;
+	}
+
+	FString GetGroomPrimPath(const pxr::UsdPrim& Prim)
+	{
+		FScopedUsdAllocs Allocs;
+
+		// Get the groom prim path to bind from the GroomBindingAPI
+		if (pxr::UsdRelationship Relationship = Prim.GetRelationship(UnrealIdentifiers::UnrealGroomToBind))
+		{
+			pxr::SdfPathVector Targets;
+			Relationship.GetTargets(&Targets);
+
+			if (Targets.size() > 0)
+			{
+				// Validate that the target prim is in fact a groom prim
+				const pxr::SdfPath& TargetPrimPath = Targets[0];
+				pxr::UsdPrim TargetPrim = Prim.GetPrimAtPath(TargetPrimPath);
+				if (TargetPrim && UsdUtils::PrimHasSchema(TargetPrim, UnrealIdentifiers::GroomAPI))
+				{
+					return UsdToUnreal::ConvertPath(TargetPrimPath);
+				}
+			}
+		}
+		return {};
+	}
+
+	UObject* GetGroomBindingSourceMesh(const pxr::UsdPrim& Prim, const UUsdAssetCache& AssetCache, EGroomBindingMeshType BindingType)
+	{
+		FScopedUsdAllocs Allocs;
+
+		// Get the reference mesh asset from the GroomBindingAPI; this property is optional
+		if (pxr::UsdRelationship Relationship = Prim.GetRelationship(UnrealIdentifiers::UnrealGroomReferenceMesh))
+		{
+			pxr::SdfPathVector Targets;
+			Relationship.GetTargets(&Targets);
+
+			if (Targets.size() > 0)
+			{
+				const pxr::SdfPath& TargetPrimPath = Targets[0];
+				FString PrimPath(UsdToUnreal::ConvertPath(TargetPrimPath));
+				UObject* Mesh = AssetCache.GetAssetForPrim(PrimPath);
+
+				// Validate that the target prim and associated asset are of the expected type for the binding
+				pxr::UsdPrim TargetPrim = Prim.GetPrimAtPath(TargetPrimPath);
+				if (BindingType == EGroomBindingMeshType::SkeletalMesh && pxr::UsdSkelRoot(TargetPrim))
+				{
+					return Cast<USkeletalMesh>(Mesh);
+				}
+				else if (BindingType == EGroomBindingMeshType::GeometryCache && pxr::UsdGeomMesh(TargetPrim))
+				{
+					return Cast<UGeometryCache>(Mesh);
+				}
+			}
+		}
+		return nullptr;
+	}
+}
+
+namespace UsdGroomTranslatorUtils
+{
+	using namespace UE::UsdGroomTranslatorUtils::Private;
+
+	void CreateGroomBindingAsset(const pxr::UsdPrim& Prim, UUsdAssetCache& AssetCache, EObjectFlags ObjectFlags)
+	{
+		// At this point, the prim (SkelRoot or GeomMesh) has already been checked to have the GroomBindingAPI,
+		// so we need to set up the GroomComponent and the groom binding asset to be able to bind it to the mesh
+		FScopedUsdAllocs Allocs;
+
+		// The GroomBinding schema must specify a groom prim to bind to the mesh
+		const FString GroomPrimPath = GetGroomPrimPath(Prim);
+		if (GroomPrimPath.IsEmpty())
+		{
+			return;
+		}
+
+		// The GroomAsset should already be processed and cached by the USDGroomTranslator
+		UGroomAsset* GroomAsset = Cast<UGroomAsset>(AssetCache.GetAssetForPrim(GroomPrimPath));
+		if (!GroomAsset)
+		{
+			return;
+		}
+
+		// Determine the type of binding needed based on the prim mesh type
+		const FString PrimPath(UsdToUnreal::ConvertPath(Prim.GetPath()));
+		UObject* PrimAsset = AssetCache.GetAssetForPrim(PrimPath);
+
+		EGroomBindingMeshType GroomBindingType = EGroomBindingMeshType::SkeletalMesh;
+		UObject* TargetMesh = Cast<USkeletalMesh>(PrimAsset);
+		if (!TargetMesh)
+		{
+			TargetMesh = Cast<UGeometryCache>(PrimAsset);
+			if (!TargetMesh)
+			{
+				return;
+			}
+			GroomBindingType = EGroomBindingMeshType::GeometryCache;
+		}
+
+		FGroomBindingBuildSettings Settings;
+		Settings.GroomBindingType = GroomBindingType;
+		Settings.Groom = GroomAsset;
+		Settings.TargetMesh = TargetMesh;
+		Settings.SourceMesh = GetGroomBindingSourceMesh(Prim, AssetCache, GroomBindingType);
+
+		// Try to get the GroomBindingAsset from the cache
+		const FSHAHash SHAHash = ComputeGroomBindingHash(Settings);
+		UGroomBindingAsset* GroomBinding = Cast<UGroomBindingAsset>(AssetCache.GetCachedAsset(SHAHash.ToString()));
+
+		const FString GroomBindingPath = FString::Printf(TEXT("%s_groombinding"), *PrimPath);
+		if (!GroomBinding)
+		{
+			// Create and cache it, if it didn't exist already
+			GroomBinding = CreateGroomBindingAsset(GroomBindingPath, Settings, ObjectFlags);
+			if (GroomBinding)
+			{
+				AssetCache.CacheAsset(SHAHash.ToString(), GroomBinding);
+			}
+		}
+
+		if (GroomBinding)
+		{
+			AssetCache.LinkAssetToPrim(GroomBindingPath, GroomBinding);
+		}
+	}
+
+	void SetGroomFromPrim(const pxr::UsdPrim& Prim, const UUsdAssetCache& AssetCache, USceneComponent* SceneComponent)
+	{
+		if (!SceneComponent)
+		{
+			return;
+		}
+
+		FScopedUsdAllocs Allocs;
+
+		const FString GroomPrimPath = GetGroomPrimPath(Prim);
+		if (GroomPrimPath.IsEmpty())
+		{
+			return;
+		}
+
+		UGroomAsset* GroomAsset = Cast<UGroomAsset>(AssetCache.GetAssetForPrim(GroomPrimPath));
+		if (!GroomAsset)
+		{
+			return;
+		}
+
+		const FString PrimPath(UsdToUnreal::ConvertPath(Prim.GetPath()));
+		const FString GroomBindingPath = FString::Printf(TEXT("%s_groombinding"), *PrimPath);
+		UGroomBindingAsset* GroomBinding = Cast<UGroomBindingAsset>(AssetCache.GetAssetForPrim(GroomBindingPath));
+
+		// Set the GroomAsset and GroomBindingAsset on the child GroomComponent of SceneComponent that was set up in the translator
+		TArray<USceneComponent*> Children;
+		const bool bIncludeAllDescendants = false;
+		SceneComponent->GetChildrenComponents(bIncludeAllDescendants, Children);
+		for (USceneComponent* ChildComponent : Children)
+		{
+			if (UGroomComponent* GroomComponent = Cast<UGroomComponent>(ChildComponent))
+			{
+				if (GroomComponent->GroomAsset != GroomAsset || GroomComponent->BindingAsset != GroomBinding)
+				{
+					GroomComponent->SetGroomAsset(GroomAsset, GroomBinding);
+				}
+				break;
+			}
+		}
+	}
+}
+
+#endif // #if USE_USD_SDK && WITH_EDITOR
diff --git a/Engine/Plugins/Importers/USDImporter/Source/USDSchemas/Private/USDGroomTranslatorUtils.h b/Engine/Plugins/Importers/USDImporter/Source/USDSchemas/Private/USDGroomTranslatorUtils.h
new file mode 100644
--- /dev/null
+++ b/Engine/Plugins/Importers/USDImporter/Source/USDSchemas/Private/USDGroomTranslatorUtils.h
@@ -0,0 +1,30 @@
+// Copyright Epic Games, Inc. All Rights Reserved.
+
+#pragma once
+
+#if USE_USD_SDK && WITH_EDITOR
+
+#include "CoreMinimal.h"
+#include "UObject/ObjectMacros.h"
+
+#include "USDIncludesStart.h"
+#include "pxr/pxr.h"
+#include "USDIncludesEnd.h"
+
+PXR_NAMESPACE_OPEN_SCOPE
+class UsdPrim;
+PXR_NAMESPACE_CLOSE_SCOPE
+
+class USceneComponent;
+class UUsdAssetCache;
+
+namespace UsdGroomTranslatorUtils
+{
+	// Create the groom binding asset for the given Prim with GroomBindingAPI and cache it in the AssetCache
+	void CreateGroomBindingAsset(const pxr::UsdPrim& Prim, UUsdAssetCache& AssetCache, EObjectFlags ObjectFlags);
+
+	// Set the groom asset targeted by the given prim with GroomBindingAPI on the SceneComponent (with a GroomComponent as a direct child)
+	void SetGroomFromPrim(const pxr::UsdPrim& Prim, const UUsdAssetCache& AssetCache, USceneComponent* SceneComponent);
+}
+
+#endif // #if USE_USD_SDK && WITH_EDITOR
\ No newline at end of file
diff --git a/Engine/Plugins/Importers/USDImporter/Source/USDSchemas/Private/USDSkelRootTranslator.cpp b/Engine/Plugins/Importers/USDImporter/Source/USDSchemas/Private/USDSkelRootTranslator.cpp
--- a/Engine/Plugins/Importers/USDImporter/Source/USDSchemas/Private/USDSkelRootTranslator.cpp
+++ b/Engine/Plugins/Importers/USDImporter/Source/USDSchemas/Private/USDSkelRootTranslator.cpp
@@ -10,6 +10,7 @@
 #include "USDConversionUtils.h"
 #include "USDErrorUtils.h"
 #include "USDGeomMeshConversion.h"
+#include "USDGroomTranslatorUtils.h"
 #include "USDIntegrationUtils.h"
 #include "USDLayerUtils.h"
 #include "USDLog.h"
@@ -29,6 +30,7 @@
 #include "Components/SkeletalMeshComponent.h"
 #include "Components/SkinnedMeshComponent.h"
 #include "Engine/SkeletalMesh.h"
+#include "GroomComponent.h"
 #include "Materials/Material.h"
 #include "Materials/MaterialInstanceConstant.h"
 #include "Rendering/SkeletalMeshLODImporterData.h"
@@ -1225,8 +1227,27 @@ void FUsdSkelRootTranslator::CreateAssets()
 
 USceneComponent* FUsdSkelRootTranslator::CreateComponents()
 {
-	// Don't need to call UpdateComponents as FUsdGeomXformableTranslator will already do that for us
-	return FUsdGeomXformableTranslator::CreateComponents();
+	USceneComponent* SceneComponent = FUsdGeomXformableTranslator::CreateComponents();
+
+#if WITH_EDITOR
+	// Check if the prim has the GroomBinding schema and setup the component and assets necessary to bind the groom to the SkeletalMesh
+	if ( UsdUtils::PrimHasSchema( GetPrim(), UnrealIdentifiers::GroomBindingAPI ) )
+	{
+		UsdGroomTranslatorUtils::CreateGroomBindingAsset( GetPrim(), *( Context->AssetCache ), Context->ObjectFlags );
+
+		// For the groom binding to work, the GroomComponent must be a child of the SceneComponent 
+		// so the Context ParentComponent is set to the SceneComponent temporarily
+		TGuardValue< USceneComponent* > ParentComponentGuard{ Context->ParentComponent, SceneComponent };
+		const bool bNeedsActor = false;
+		UGroomComponent* GroomComponent = Cast< UGroomComponent >( CreateComponentsEx( TSubclassOf< USceneComponent >( UGroomComponent::StaticClass() ), bNeedsActor ) );
+		if ( GroomComponent )
+		{
+			UpdateComponents( SceneComponent );
+		}
+	}
+#endif // WITH_EDITOR
+
+	return SceneComponent;
 }
 
 void FUsdSkelRootTranslator::UpdateComponents( USceneComponent* SceneComponent )
@@ -1365,6 +1386,12 @@ void FUsdSkelRootTranslator::UpdateComponents( USceneComponent* SceneComponent )
 			SkeletalMeshComponent->MarkRenderDynamicDataDirty();
 		}
 	}
+
+	// If the prim has a GroomBinding schema, apply the target groom to its associated GroomComponent
+	if ( UsdUtils::PrimHasSchema( Prim, UnrealIdentifiers::GroomBindingAPI ) )
+	{
+		UsdGroomTranslatorUtils::SetGroomFromPrim( Prim, *Context->AssetCache, SceneComponent );
+	}
 #endif // WITH_EDITOR
 }
 
diff --git a/Engine/Plugins/Importers/USDImporter/Source/USDStageEditorViewModels/Private/USDPrimViewModel.cpp b/Engine/Plugins/Importers/USDImporter/Source/USDStageEditorViewModels/Private/USDPrimViewModel.cpp
--- a/Engine/Plugins/Importers/USDImporter/Source/USDStageEditorViewModels/Private/USDPrimViewModel.cpp
+++ b/Engine/Plugins/Importers/USDImporter/Source/USDStageEditorViewModels/Private/USDPrimViewModel.cpp
@@ -403,7 +403,7 @@ bool FUsdPrimViewModel::CanRemoveControlRig() const
 void FUsdPrimViewModel::ApplyGroomSchema()
 {
 #if USE_USD_SDK
-	UsdUtils::ApplyGroomSchema( UsdPrim );
+	UsdUtils::ApplySchema( UsdPrim, UnrealIdentifiers::GroomAPI );
 #endif // #if USE_USD_SDK
 }
 
diff --git a/Engine/Plugins/Importers/USDImporter/Source/USDUtilities/Private/USDGroomConversion.cpp b/Engine/Plugins/Importers/USDImporter/Source/USDUtilities/Private/USDGroomConversion.cpp
--- a/Engine/Plugins/Importers/USDImporter/Source/USDUtilities/Private/USDGroomConversion.cpp
+++ b/Engine/Plugins/Importers/USDImporter/Source/USDUtilities/Private/USDGroomConversion.cpp
@@ -10,10 +10,11 @@
 #include "UsdWrappers/UsdPrim.h"
 
 #include "USDIncludesStart.h"
+#include "pxr/usd/usd/prim.h"
 #include "pxr/usd/usdGeom/basisCurves.h"
 #include "pxr/usd/usdGeom/primvar.h"
 #include "pxr/usd/usdGeom/primvarsAPI.h"
-#include "pxr/usd/usdGeom/xformable.h"
+#include "pxr/usd/usdGeom/xform.h"
 #include "USDIncludesEnd.h"
 
 namespace UE::UsdGroomConversion::Private
@@ -315,15 +316,19 @@ namespace UsdToUnreal
 
 			// Following the USD recommendation that gprims be not nested, a Curves prim is handled as a leaf
 		}
-		else if (pxr::UsdGeomXformable Xform = pxr::UsdGeomXformable(Prim))
+		else if (pxr::UsdGeomImageable(Prim))
 		{
-			// Propagate the prim transform to the children
-			bool bResetXformStack = false;
-			FTransform PrimTransform = FTransform::Identity;
-			bool bConverted = UsdToUnreal::ConvertXformable(Prim.GetStage(), Xform, PrimTransform, TimeCode.GetValue(), &bResetXformStack);
-			if (bConverted)
+			// UsdGeomImageable includes UsdGeomXformable and UsdGeomScope
+			if (pxr::UsdGeomXform Xform = pxr::UsdGeomXform(Prim))
 			{
-				TransformToPropagate = bResetXformStack ? PrimTransform : ParentTransform * PrimTransform;
+				// Propagate the prim transform to the children
+				bool bResetXformStack = false;
+				FTransform PrimTransform = FTransform::Identity;
+				bool bConverted = UsdToUnreal::ConvertXformable(Prim.GetStage(), Xform, PrimTransform, TimeCode.GetValue(), &bResetXformStack);
+				if (bConverted)
+				{
+					TransformToPropagate = bResetXformStack ? PrimTransform : ParentTransform * PrimTransform;
+				}
 			}
 
 			for (const pxr::UsdPrim& Child : Prim.GetChildren())
diff --git a/Engine/Plugins/Importers/USDImporter/Source/USDUtilities/Private/USDIntegrationUtils.cpp b/Engine/Plugins/Importers/USDImporter/Source/USDUtilities/Private/USDIntegrationUtils.cpp
--- a/Engine/Plugins/Importers/USDImporter/Source/USDUtilities/Private/USDIntegrationUtils.cpp
+++ b/Engine/Plugins/Importers/USDImporter/Source/USDUtilities/Private/USDIntegrationUtils.cpp
@@ -18,20 +18,22 @@
 
 #define LOCTEXT_NAMESPACE "USDIntegrationUtils"
 
-bool UsdUtils::PrimHasLiveLinkSchema( const pxr::UsdPrim& Prim )
+bool UsdUtils::PrimHasSchema( const pxr::UsdPrim& Prim, const pxr::TfToken& SchemaToken )
 {
 	if ( !Prim )
 	{
 		return false;
 	}
 
-    FScopedUsdAllocs Allocs;
+	FScopedUsdAllocs Allocs;
 
-    pxr::TfType Schema = pxr::UsdSchemaRegistry::GetTypeFromSchemaTypeName( UnrealIdentifiers::LiveLinkAPI );
-    return Prim.HasAPI( Schema );
+	pxr::TfType Schema = pxr::UsdSchemaRegistry::GetTypeFromSchemaTypeName( SchemaToken );
+	ensure( static_cast<bool>( Schema ) );
+
+	return Prim.HasAPI( Schema );
 }
 
-bool UsdUtils::PrimHasControlRigSchema( const pxr::UsdPrim& Prim )
+bool UsdUtils::ApplySchema( const pxr::UsdPrim& Prim, const pxr::TfToken& SchemaToken )
 {
 	if ( !Prim )
 	{
@@ -40,24 +42,26 @@ bool UsdUtils::PrimHasControlRigSchema( const pxr::UsdPrim& Prim )
 
 	FScopedUsdAllocs Allocs;
 
-	pxr::TfType Schema = pxr::UsdSchemaRegistry::GetTypeFromSchemaTypeName( UnrealIdentifiers::ControlRigAPI );
-	return Prim.HasAPI( Schema );
+	pxr::TfType Schema = pxr::UsdSchemaRegistry::GetTypeFromSchemaTypeName( SchemaToken );
+	ensure( static_cast<bool>( Schema ) );
+
+	return Prim.ApplyAPI( Schema );
 }
 
-bool UsdUtils::PrimHasGroomSchema( const pxr::UsdPrim& Prim )
+bool UsdUtils::PrimHasLiveLinkSchema( const pxr::UsdPrim& Prim )
 {
 	if ( !Prim )
 	{
 		return false;
 	}
 
-	FScopedUsdAllocs Allocs;
+    FScopedUsdAllocs Allocs;
 
-	pxr::TfType Schema = pxr::UsdSchemaRegistry::GetTypeFromSchemaTypeName( UnrealIdentifiers::GroomAPI );
-	return Prim.HasAPI( Schema );
+    pxr::TfType Schema = pxr::UsdSchemaRegistry::GetTypeFromSchemaTypeName( UnrealIdentifiers::LiveLinkAPI );
+    return Prim.HasAPI( Schema );
 }
 
-bool UsdUtils::ApplyControlRigSchema( const pxr::UsdPrim& Prim )
+bool UsdUtils::PrimHasControlRigSchema( const pxr::UsdPrim& Prim )
 {
 	if ( !Prim )
 	{
@@ -67,13 +71,10 @@ bool UsdUtils::ApplyControlRigSchema( const pxr::UsdPrim& Prim )
 	FScopedUsdAllocs Allocs;
 
 	pxr::TfType Schema = pxr::UsdSchemaRegistry::GetTypeFromSchemaTypeName( UnrealIdentifiers::ControlRigAPI );
-	ensure( static_cast< bool >( Schema ) );
-
-	ensure( Prim.ApplyAPI( Schema ) );
-	return true;
+	return Prim.HasAPI( Schema );
 }
 
-bool UsdUtils::ApplyLiveLinkSchema( const pxr::UsdPrim& Prim )
+bool UsdUtils::ApplyControlRigSchema( const pxr::UsdPrim& Prim )
 {
 	if ( !Prim )
 	{
@@ -82,14 +83,14 @@ bool UsdUtils::ApplyLiveLinkSchema( const pxr::UsdPrim& Prim )
 
 	FScopedUsdAllocs Allocs;
 
-	pxr::TfType Schema = pxr::UsdSchemaRegistry::GetTypeFromSchemaTypeName( UnrealIdentifiers::LiveLinkAPI );
+	pxr::TfType Schema = pxr::UsdSchemaRegistry::GetTypeFromSchemaTypeName( UnrealIdentifiers::ControlRigAPI );
 	ensure( static_cast< bool >( Schema ) );
 
 	ensure( Prim.ApplyAPI( Schema ) );
 	return true;
 }
 
-bool UsdUtils::ApplyGroomSchema( const pxr::UsdPrim& Prim )
+bool UsdUtils::ApplyLiveLinkSchema( const pxr::UsdPrim& Prim )
 {
 	if ( !Prim )
 	{
@@ -98,8 +99,8 @@ bool UsdUtils::ApplyGroomSchema( const pxr::UsdPrim& Prim )
 
 	FScopedUsdAllocs Allocs;
 
-	pxr::TfType Schema = pxr::UsdSchemaRegistry::GetTypeFromSchemaTypeName( UnrealIdentifiers::GroomAPI );
-	ensure( static_cast<bool>( Schema ) );
+	pxr::TfType Schema = pxr::UsdSchemaRegistry::GetTypeFromSchemaTypeName( UnrealIdentifiers::LiveLinkAPI );
+	ensure( static_cast< bool >( Schema ) );
 
 	ensure( Prim.ApplyAPI( Schema ) );
 	return true;
diff --git a/Engine/Plugins/Importers/USDImporter/Source/USDUtilities/Public/USDIntegrationUtils.h b/Engine/Plugins/Importers/USDImporter/Source/USDUtilities/Public/USDIntegrationUtils.h
--- a/Engine/Plugins/Importers/USDImporter/Source/USDUtilities/Public/USDIntegrationUtils.h
+++ b/Engine/Plugins/Importers/USDImporter/Source/USDUtilities/Public/USDIntegrationUtils.h
@@ -11,6 +11,7 @@
 
 PXR_NAMESPACE_OPEN_SCOPE
 	class UsdPrim;
+	class TfToken;
 PXR_NAMESPACE_CLOSE_SCOPE
 
 #endif // #if USE_USD_SDK
@@ -22,15 +23,18 @@ namespace UsdUtils
 {
 #if USE_USD_SDK
 
+	// Returns true if Prim has the given Schema
+	USDUTILITIES_API bool PrimHasSchema( const pxr::UsdPrim& Prim, const pxr::TfToken& Schema );
+
+	// Apply the given Schema to the Prim. Returns true if the application was successful
+	USDUTILITIES_API bool ApplySchema( const pxr::UsdPrim& Prim, const pxr::TfToken& Schema );
+
 	// Returns true in case Prim has our custom "LiveLinkAPI" schema
 	USDUTILITIES_API bool PrimHasLiveLinkSchema( const pxr::UsdPrim& Prim );
 
 	// Returns true in case Prim has our custom "ControlRigAPI" schema
 	USDUTILITIES_API bool PrimHasControlRigSchema( const pxr::UsdPrim& Prim );
 
-	// Returns true if Prim has the custom "GroomAPI" schema
-	USDUTILITIES_API bool PrimHasGroomSchema( const pxr::UsdPrim& Prim );
-
 	// Apply the control rig schema to the prim, creates all the required attributes and sets default values
 	// if they're otherwise not authored
 	USDUTILITIES_API bool ApplyControlRigSchema( const pxr::UsdPrim& Prim );
@@ -39,9 +43,6 @@ namespace UsdUtils
 	// if they're otherwise not authored
 	USDUTILITIES_API bool ApplyLiveLinkSchema( const pxr::UsdPrim& Prim );
 
-	// Apply the groom schema to the prim
-	USDUTILITIES_API bool ApplyGroomSchema( const pxr::UsdPrim& Prim );
-
 #endif // USE_USD_SDK
 }
 
diff --git a/Engine/Plugins/Importers/USDImporter/Source/UnrealUSDWrapper/Private/UnrealUSDWrapper.cpp b/Engine/Plugins/Importers/USDImporter/Source/UnrealUSDWrapper/Private/UnrealUSDWrapper.cpp
--- a/Engine/Plugins/Importers/USDImporter/Source/UnrealUSDWrapper/Private/UnrealUSDWrapper.cpp
+++ b/Engine/Plugins/Importers/USDImporter/Source/UnrealUSDWrapper/Private/UnrealUSDWrapper.cpp
@@ -145,6 +145,9 @@ namespace UnrealIdentifiers
 	const TfToken WorldSpaceNormals = TfToken( "worldSpaceNormals" );
 
 	const TfToken GroomAPI = TfToken( "GroomAPI" );
+	const TfToken GroomBindingAPI = TfToken( "GroomBindingAPI" );
+	const TfToken UnrealGroomToBind = TfToken( "unreal:groomBinding:groom" );
+	const TfToken UnrealGroomReferenceMesh = TfToken( "unreal:groomBinding:referenceMesh" );
 }
 
 std::string FUsdAttribute::GetUnrealPropertyPath( const pxr::UsdAttribute& Attribute )
diff --git a/Engine/Plugins/Importers/USDImporter/Source/UnrealUSDWrapper/Public/UnrealUSDWrapper.h b/Engine/Plugins/Importers/USDImporter/Source/UnrealUSDWrapper/Public/UnrealUSDWrapper.h
--- a/Engine/Plugins/Importers/USDImporter/Source/UnrealUSDWrapper/Public/UnrealUSDWrapper.h
+++ b/Engine/Plugins/Importers/USDImporter/Source/UnrealUSDWrapper/Public/UnrealUSDWrapper.h
@@ -390,6 +390,9 @@ namespace UnrealIdentifiers
 	extern UNREALUSDWRAPPER_API const pxr::TfToken WorldSpaceNormals;
 
 	extern UNREALUSDWRAPPER_API const pxr::TfToken GroomAPI;
+	extern UNREALUSDWRAPPER_API const pxr::TfToken GroomBindingAPI;
+	extern UNREALUSDWRAPPER_API const pxr::TfToken UnrealGroomToBind;
+	extern UNREALUSDWRAPPER_API const pxr::TfToken UnrealGroomReferenceMesh;
 #endif // #if USE_USD_SDK
 
 	extern UNREALUSDWRAPPER_API const TCHAR* Invisible;
