USD: Fix crash when overwriting a SkelRoot import while appending actors and replacing assets.

diff --git a/Engine/Plugins/Importers/USDImporter/Source/USDStageImporter/Private/USDStageImporter.cpp b/Engine/Plugins/Importers/USDImporter/Source/USDStageImporter/Private/USDStageImporter.cpp
--- a/Engine/Plugins/Importers/USDImporter/Source/USDStageImporter/Private/USDStageImporter.cpp
+++ b/Engine/Plugins/Importers/USDImporter/Source/USDStageImporter/Private/USDStageImporter.cpp
@@ -553,6 +553,17 @@ namespace UsdStageImporterImpl
 		UObject* MovedAsset = ExistingAsset;
 		if (ExistingAsset != nullptr && ExistingAsset != Asset && ReplacePolicy == EReplaceAssetPolicy::Replace)
 		{
+			// We must do this here or else on the next engine tick the USkeletalMeshComponent::ParallelAnimationEvaluation
+			// will try to use the old anim instance, which will try using the old Skeleton's SmartNameContainer and crash.
+			// We do not need to recreate the anim instance afterwards as InitAnim is called on component registration anyway.
+			for ( TObjectIterator<USkeletalMeshComponent> It; It; ++It )
+			{
+				if ( It->SkeletalMesh == ExistingAsset )
+				{
+					It->ClearAnimScriptInstance();
+				}
+			}
+
 			MovedAsset = DuplicateObject<UObject>( Asset, Package, *TargetAssetName );
 
 			// If our DuplicateObject didn't stomp the old asset because TargetAssetName != ExistingAsset->GetName(),
