USD: Combine both USD blueprint libraries and deprecate one of them;

Have UUsdConversionBlueprintLibrary actually derive UBlueprintFunctionLibrary.

diff --git a/Engine/Plugins/Importers/USDImporter/Source/USDExporter/Private/USDConversionBlueprintLibrary.cpp b/Engine/Plugins/Importers/USDImporter/Source/USDExporter/Private/USDConversionBlueprintLibrary.cpp
--- a/Engine/Plugins/Importers/USDImporter/Source/USDExporter/Private/USDConversionBlueprintLibrary.cpp
+++ b/Engine/Plugins/Importers/USDImporter/Source/USDExporter/Private/USDConversionBlueprintLibrary.cpp
@@ -4,6 +4,7 @@
 
 #include "LevelExporterUSD.h"
 #include "UnrealUSDWrapper.h"
+#include "USDClassesModule.h"
 #include "USDConversionUtils.h"
 #include "USDLayerUtils.h"
 #include "USDLog.h"
@@ -13,12 +14,18 @@
 #include "UsdWrappers/UsdPrim.h"
 #include "UsdWrappers/UsdStage.h"
 
+#include "AnalyticsBlueprintLibrary.h"
+#include "AnalyticsEventAttribute.h"
+#include "CoreMinimal.h"
+#include "Editor.h"
 #include "Engine/Engine.h"
 #include "Engine/Level.h"
 #include "Engine/LevelStreaming.h"
 #include "Engine/World.h"
+#include "InstancedFoliageActor.h"
 #include "ISequencer.h"
 #include "LevelEditorSequencerIntegration.h"
+#include "UObject/ObjectMacros.h"
 
 namespace UE
 {
@@ -363,3 +370,101 @@ FString UUsdConversionBlueprintLibrary::GetSchemaNameForComponent( const USceneC
 
 	return {};
 }
+
+AInstancedFoliageActor* UUsdConversionBlueprintLibrary::GetInstancedFoliageActorForLevel( bool bCreateIfNone /*= false */, ULevel* Level /*= nullptr */ )
+{
+	if ( !Level )
+	{
+		const bool bEnsureIsGWorld = false;
+		UWorld* EditorWorld = GEditor ? GEditor->GetEditorWorldContext( bEnsureIsGWorld ).World() : nullptr;
+		if ( !EditorWorld )
+		{
+			return nullptr;
+		}
+
+		Level = EditorWorld->GetCurrentLevel();
+		if ( !Level )
+		{
+			return nullptr;
+		}
+	}
+
+	return AInstancedFoliageActor::GetInstancedFoliageActorForLevel( Level, bCreateIfNone );
+}
+
+TArray<UFoliageType*> UUsdConversionBlueprintLibrary::GetUsedFoliageTypes( AInstancedFoliageActor* Actor )
+{
+	TArray<UFoliageType*> Result;
+	if ( !Actor )
+	{
+		return Result;
+	}
+
+	for ( const TPair<UFoliageType*, TUniqueObj<FFoliageInfo>>& FoliagePair : Actor->GetFoliageInfos() )
+	{
+		Result.Add( FoliagePair.Key );
+	}
+
+	return Result;
+}
+
+UObject* UUsdConversionBlueprintLibrary::GetSource( UFoliageType* FoliageType )
+{
+	if ( FoliageType )
+	{
+		return FoliageType->GetSource();
+	}
+
+	return nullptr;
+}
+
+TArray<FTransform> UUsdConversionBlueprintLibrary::GetInstanceTransforms( AInstancedFoliageActor* Actor, UFoliageType* FoliageType, ULevel* InstancesLevel )
+{
+	TArray<FTransform> Result;
+	if ( !Actor || !FoliageType )
+	{
+		return Result;
+	}
+
+	// Modified from AInstancedFoliageActor::GetInstancesForComponent to limit traversal only to our FoliageType
+
+	if ( const TUniqueObj<FFoliageInfo>* FoundInfo = Actor->GetFoliageInfos().Find( FoliageType ) )
+	{
+		const FFoliageInfo& Info = ( *FoundInfo ).Get();
+
+		// Collect IDs of components that are on the same level as the actor's level. This because later on we'll have level-by-level
+		// export, and we'd want one point instancer per level
+		for ( const TPair<FFoliageInstanceBaseId, FFoliageInstanceBaseInfo>& FoliageInstancePair : Actor->InstanceBaseCache.InstanceBaseMap )
+		{
+			UActorComponent* Comp = FoliageInstancePair.Value.BasePtr.Get();
+			if ( !Comp || ( InstancesLevel && ( Comp->GetComponentLevel() != InstancesLevel ) ) )
+			{
+				continue;
+			}
+
+			if ( const auto* InstanceSet = Info.ComponentHash.Find( FoliageInstancePair.Key ) )
+			{
+				Result.Reserve( Result.Num() + InstanceSet->Num() );
+				for ( int32 InstanceIndex : *InstanceSet )
+				{
+					const FFoliageInstancePlacementInfo* Instance = &Info.Instances[ InstanceIndex ];
+					Result.Emplace( FQuat( Instance->Rotation ), Instance->Location, ( FVector ) Instance->DrawScale3D );
+				}
+			}
+		}
+	}
+
+	return Result;
+}
+
+void UUsdConversionBlueprintLibrary::SendAnalytics( const TArray<FAnalyticsEventAttr>& Attrs, const FString& EventName, bool bAutomated, double ElapsedSeconds, double NumberOfFrames, const FString& Extension )
+{
+	TArray<FAnalyticsEventAttribute> Converted;
+	Converted.Reserve( Attrs.Num() );
+	for ( const FAnalyticsEventAttr& Attr : Attrs )
+	{
+		Converted.Emplace( Attr.Name, Attr.Value );
+	}
+
+	IUsdClassesModule::SendAnalytics( MoveTemp( Converted ), EventName, bAutomated, ElapsedSeconds, NumberOfFrames, Extension );
+}
diff --git a/Engine/Plugins/Importers/USDImporter/Source/USDExporter/Private/USDExporterBlueprintLibrary.cpp b/Engine/Plugins/Importers/USDImporter/Source/USDExporter/Private/USDExporterBlueprintLibrary.cpp
--- a/Engine/Plugins/Importers/USDImporter/Source/USDExporter/Private/USDExporterBlueprintLibrary.cpp
+++ b/Engine/Plugins/Importers/USDImporter/Source/USDExporter/Private/USDExporterBlueprintLibrary.cpp
@@ -2,109 +2,30 @@
 
 #include "USDExporterBlueprintLibrary.h"
 
-#include "AnalyticsBlueprintLibrary.h"
-#include "AnalyticsEventAttribute.h"
-#include "CoreMinimal.h"
-#include "Editor.h"
-#include "InstancedFoliageActor.h"
-#include "UObject/ObjectMacros.h"
-#include "USDClassesModule.h"
+#include "USDConversionBlueprintLibrary.h"
 
 AInstancedFoliageActor* UUsdExporterBlueprintLibrary::GetInstancedFoliageActorForLevel( bool bCreateIfNone /*= false */, ULevel* Level /*= nullptr */ )
 {
-	if ( !Level )
-	{
-		const bool bEnsureIsGWorld = false;
-		UWorld* EditorWorld = GEditor ? GEditor->GetEditorWorldContext( bEnsureIsGWorld ).World() : nullptr;
-		if ( !EditorWorld )
-		{
-			return nullptr;
-		}
-
-		Level = EditorWorld->GetCurrentLevel();
-		if ( !Level )
-		{
-			return nullptr;
-		}
-	}
-
-	return AInstancedFoliageActor::GetInstancedFoliageActorForLevel( Level, bCreateIfNone );
+	return UUsdConversionBlueprintLibrary::GetInstancedFoliageActorForLevel( bCreateIfNone, Level );
 }
 
 TArray<UFoliageType*> UUsdExporterBlueprintLibrary::GetUsedFoliageTypes( AInstancedFoliageActor* Actor )
 {
-	TArray<UFoliageType*> Result;
-	if ( !Actor )
-	{
-		return Result;
-	}
-
-	for ( const TPair<UFoliageType*, TUniqueObj<FFoliageInfo>>& FoliagePair : Actor->GetFoliageInfos())
-	{
-		Result.Add( FoliagePair.Key );
-	}
-
-	return Result;
+	return UUsdConversionBlueprintLibrary::GetUsedFoliageTypes( Actor );
 }
 
 UObject* UUsdExporterBlueprintLibrary::GetSource( UFoliageType* FoliageType )
 {
-	if ( FoliageType )
-	{
-		return FoliageType->GetSource();
-	}
-
-	return nullptr;
+	return UUsdConversionBlueprintLibrary::GetSource( FoliageType );
 }
 
 TArray<FTransform> UUsdExporterBlueprintLibrary::GetInstanceTransforms( AInstancedFoliageActor* Actor, UFoliageType* FoliageType, ULevel* InstancesLevel )
 {
-	TArray<FTransform> Result;
-	if ( !Actor || !FoliageType )
-	{
-		return Result;
-	}
-
-	// Modified from AInstancedFoliageActor::GetInstancesForComponent to limit traversal only to our FoliageType
-
-	if ( const TUniqueObj<FFoliageInfo>* FoundInfo = Actor->GetFoliageInfos().Find( FoliageType ) )
-	{
-		const FFoliageInfo& Info = (*FoundInfo).Get();
-
-		// Collect IDs of components that are on the same level as the actor's level. This because later on we'll have level-by-level
-		// export, and we'd want one point instancer per level
-		for ( const TPair<FFoliageInstanceBaseId, FFoliageInstanceBaseInfo>& FoliageInstancePair : Actor->InstanceBaseCache.InstanceBaseMap )
-		{
-			UActorComponent* Comp = FoliageInstancePair.Value.BasePtr.Get();
-			if ( !Comp || ( InstancesLevel && ( Comp->GetComponentLevel() != InstancesLevel ) ) )
-			{
-				continue;
-			}
-
-			if ( const auto* InstanceSet = Info.ComponentHash.Find( FoliageInstancePair.Key ) )
-			{
-				Result.Reserve( Result.Num() + InstanceSet->Num() );
-				for ( int32 InstanceIndex : *InstanceSet )
-				{
-					const FFoliageInstancePlacementInfo* Instance = &Info.Instances[ InstanceIndex ];
-					Result.Emplace( FQuat(Instance->Rotation), Instance->Location, (FVector)Instance->DrawScale3D );
-				}
-			}
-		}
-	}
-
-	return Result;
+	return UUsdConversionBlueprintLibrary::GetInstanceTransforms( Actor, FoliageType, InstancesLevel );
 }
 
 void UUsdExporterBlueprintLibrary::SendAnalytics( const TArray<FAnalyticsEventAttr>& Attrs, const FString& EventName, bool bAutomated, double ElapsedSeconds, double NumberOfFrames, const FString& Extension )
 {
-	TArray<FAnalyticsEventAttribute> Converted;
-	Converted.Reserve( Attrs.Num() );
-	for ( const FAnalyticsEventAttr& Attr : Attrs )
-	{
-		Converted.Emplace( Attr.Name, Attr.Value );
-	}
-
-	IUsdClassesModule::SendAnalytics( MoveTemp( Converted ), EventName, bAutomated, ElapsedSeconds, NumberOfFrames, Extension );
+	return UUsdConversionBlueprintLibrary::SendAnalytics( Attrs, EventName, bAutomated, ElapsedSeconds, NumberOfFrames, Extension );
 }
 
diff --git a/Engine/Plugins/Importers/USDImporter/Source/USDExporter/Public/USDConversionBlueprintLibrary.h b/Engine/Plugins/Importers/USDImporter/Source/USDExporter/Public/USDConversionBlueprintLibrary.h
--- a/Engine/Plugins/Importers/USDImporter/Source/USDExporter/Public/USDConversionBlueprintLibrary.h
+++ b/Engine/Plugins/Importers/USDImporter/Source/USDExporter/Public/USDConversionBlueprintLibrary.h
@@ -2,13 +2,19 @@
 
 #pragma once
 
+#include "AnalyticsBlueprintLibrary.h"
 #include "Engine/EngineTypes.h"
+#include "Kismet/BlueprintFunctionLibrary.h"
 
 #include "USDConversionBlueprintLibrary.generated.h"
 
+class AInstancedFoliageActor;
+class UFoliageType;
+class ULevel;
+
 /** Wrapped static conversion functions from the UsdUtilities module, so that they can be used via scripting */
 UCLASS(meta=(ScriptName="UsdConversionLibrary"))
-class USDEXPORTER_API UUsdConversionBlueprintLibrary : public UObject
+class USDEXPORTER_API UUsdConversionBlueprintLibrary : public UBlueprintFunctionLibrary
 {
 	GENERATED_BODY()
 
@@ -66,4 +72,39 @@ public:
 
 	UFUNCTION( BlueprintCallable, Category = "Prim utils" )
 	static FString GetSchemaNameForComponent( const USceneComponent* Component );
+
+	/**
+	 * Wraps AInstancedFoliageActor::GetInstancedFoliageActorForLevel, and allows retrieving the current AInstancedFoliageActor
+	 * for a level. Will default to the current editor level if Level is left nullptr.
+	 * This function is useful because it's difficult to retrieve this actor otherwise, as it will be filtered from
+	 * the results of functions like EditorLevelLibrary.get_all_level_actors()
+	 */
+	UFUNCTION( BlueprintCallable, Category = "USD Foliage Exporter" )
+	static AInstancedFoliageActor* GetInstancedFoliageActorForLevel( bool bCreateIfNone = false, ULevel* Level = nullptr );
+
+	/**
+	 * Returns all the different types of UFoliageType assets that a particular AInstancedFoliageActor uses.
+	 * This function exists because we want to retrieve all instances of all foliage types on an actor, but we
+	 * can't return nested containers from UFUNCTIONs, so users of this API should call this, and then GetInstanceTransforms.
+	 */
+	UFUNCTION( BlueprintCallable, meta = ( ScriptMethod ), Category = "USD Foliage Exporter" )
+	static TArray<UFoliageType*> GetUsedFoliageTypes( AInstancedFoliageActor* Actor );
+
+	/**
+	 * Returns the source asset for a UFoliageType.
+	 * It can be a UStaticMesh in case we're dealing with a UFoliageType_InstancedStaticMesh, but it can be other types of objects.
+	 */
+	UFUNCTION( BlueprintCallable, meta = ( ScriptMethod ), Category = "USD Foliage Exporter" )
+	static UObject* GetSource( UFoliageType* FoliageType );
+
+	/**
+	 * Returns the transforms of all instances of a particular UFoliageType on a given level. If no level is provided all instances will be returned.
+	 * Use GetUsedFoliageTypes() to retrieve all foliage types managed by a particular actor.
+	 */
+	UFUNCTION( BlueprintCallable, meta = ( ScriptMethod ), Category = "USD Foliage Exporter" )
+	static TArray<FTransform> GetInstanceTransforms( AInstancedFoliageActor* Actor, UFoliageType* FoliageType, ULevel* InstancesLevel = nullptr );
+
+	/** Defer to the USDClasses module to actually send analytics information */
+	UFUNCTION( BlueprintCallable, Category = "Analytics" )
+	static void SendAnalytics( const TArray<FAnalyticsEventAttr>& Attrs, const FString& EventName, bool bAutomated, double ElapsedSeconds, double NumberOfFrames, const FString& Extension );
 };
diff --git a/Engine/Plugins/Importers/USDImporter/Source/USDExporter/Public/USDExporterBlueprintLibrary.h b/Engine/Plugins/Importers/USDImporter/Source/USDExporter/Public/USDExporterBlueprintLibrary.h
--- a/Engine/Plugins/Importers/USDImporter/Source/USDExporter/Public/USDExporterBlueprintLibrary.h
+++ b/Engine/Plugins/Importers/USDImporter/Source/USDExporter/Public/USDExporterBlueprintLibrary.h
@@ -28,6 +28,7 @@ public:
 	 * This function is useful because it's difficult to retrieve this actor otherwise, as it will be filtered from
 	 * the results of functions like EditorLevelLibrary.get_all_level_actors()
 	 */
+	UE_DEPRECATED( 5.1, "This function was moved to the UUsdConversionBlueprintLibrary C++ class (UsdConversionLibrary on Python/Blueprint)" )
 	UFUNCTION( BlueprintCallable, Category = "USD Foliage Exporter" )
 	static AInstancedFoliageActor* GetInstancedFoliageActorForLevel( bool bCreateIfNone = false, ULevel* Level = nullptr );
 
@@ -36,6 +37,7 @@ public:
 	 * This function exists because we want to retrieve all instances of all foliage types on an actor, but we
 	 * can't return nested containers from UFUNCTIONs, so users of this API should call this, and then GetInstanceTransforms.
 	 */
+	UE_DEPRECATED( 5.1, "This function was moved to the UUsdConversionBlueprintLibrary C++ class (UsdConversionLibrary on Python/Blueprint)" )
 	UFUNCTION( BlueprintCallable, meta = ( ScriptMethod ), Category = "USD Foliage Exporter" )
 	static TArray<UFoliageType*> GetUsedFoliageTypes( AInstancedFoliageActor* Actor );
 
@@ -43,6 +45,7 @@ public:
 	 * Returns the source asset for a UFoliageType.
 	 * It can be a UStaticMesh in case we're dealing with a UFoliageType_InstancedStaticMesh, but it can be other types of objects.
 	 */
+	UE_DEPRECATED( 5.1, "This function was moved to the UUsdConversionBlueprintLibrary C++ class (UsdConversionLibrary on Python/Blueprint)" )
 	UFUNCTION( BlueprintCallable, meta = ( ScriptMethod ), Category = "USD Foliage Exporter" )
 	static UObject* GetSource( UFoliageType* FoliageType );
 
@@ -50,10 +53,12 @@ public:
 	 * Returns the transforms of all instances of a particular UFoliageType on a given level. If no level is provided all instances will be returned.
 	 * Use GetUsedFoliageTypes() to retrieve all foliage types managed by a particular actor.
 	 */
+	UE_DEPRECATED( 5.1, "This function was moved to the UUsdConversionBlueprintLibrary C++ class (UsdConversionLibrary on Python/Blueprint)" )
 	UFUNCTION( BlueprintCallable, meta = ( ScriptMethod ), Category = "USD Foliage Exporter" )
 	static TArray<FTransform> GetInstanceTransforms( AInstancedFoliageActor* Actor, UFoliageType* FoliageType, ULevel* InstancesLevel = nullptr );
 
 	/** Defer to the USDClasses module to actually send analytics information */
+	UE_DEPRECATED( 5.1, "This function was moved to the UUsdConversionBlueprintLibrary C++ class (UsdConversionLibrary on Python/Blueprint)" )
 	UFUNCTION( BlueprintCallable, Category = "Analytics" )
 	static void SendAnalytics( const TArray<FAnalyticsEventAttr>& Attrs, const FString& EventName, bool bAutomated, double ElapsedSeconds, double NumberOfFrames, const FString& Extension );
 };
\ No newline at end of file
