USD: Fix deleting/renaming prims not updating the USD Stage editor.

diff --git a/Engine/Plugins/Importers/USDImporter/Source/USDStage/Private/USDStageActor.cpp b/Engine/Plugins/Importers/USDImporter/Source/USDStage/Private/USDStageActor.cpp
--- a/Engine/Plugins/Importers/USDImporter/Source/USDStage/Private/USDStageActor.cpp
+++ b/Engine/Plugins/Importers/USDImporter/Source/USDStage/Private/USDStageActor.cpp
@@ -811,12 +811,6 @@ void AUsdStageActor::OnUsdObjectsChanged( const UsdUtils::FObjectChangesByPath&
 	{
 		UE::FSdfPath PrimPath = UE::FSdfPath( *InfoChange.Key ).StripAllVariantSelections();
 
-		const bool bPrimExists = static_cast< bool >( Stage.GetPrimAtPath( PrimPath ) );
-		if ( !bPrimExists )
-		{
-			continue;
-		}
-
 		// Upgrade some info changes into resync changes
 		bool bIsResync = false;
 		for ( const UsdUtils::FObjectChangeNotice& ObjectChange : InfoChange.Value )
@@ -858,14 +852,6 @@ void AUsdStageActor::OnUsdObjectsChanged( const UsdUtils::FObjectChangesByPath&
 	{
 		UE::FSdfPath PrimPath = UE::FSdfPath( *ResyncChange.Key ).StripAllVariantSelections();
 
-		// In some scenarios (e.g. prim renaming) we'll get resync notices on prims that don't exist anymore. There's no point
-		// doing anything with these, so just ignore them
-		const bool bPrimExists = static_cast< bool >( Stage.GetPrimAtPath( PrimPath ) );
-		if ( !bPrimExists )
-		{
-			continue;
-		}
-
 		const bool bIsResync = true;
 		SortedPrimsChangedList.Add( PrimPath.GetString(), bIsResync );
 	}
@@ -976,9 +962,12 @@ void AUsdStageActor::OnUsdObjectsChanged( const UsdUtils::FObjectChangesByPath&
 			return !PathToProcess.IsEmpty() && PathsProcessed.Contains( PathToProcess );
 		};
 
-		auto UpdateComponents = [&]( const UE::FSdfPath& InPrimPath, const bool bInResync )
+		auto UpdateComponents = [&]( const UE::FUsdPrim& Prim, const UE::FSdfPath& InPrimPath, const bool bInResync )
 		{
-			UE::FSdfPath ComponentsPrimPath = InfoCache.IsValid()
+			// Don't query the InfoCache about prims that don't exist since that will ensure.
+			// We need to carry on though, because we must update components in case a prim is renamed or deleted,
+			// which will emit notices for prims that don't actually exist on the stage
+			UE::FSdfPath ComponentsPrimPath = ( Prim && InfoCache.IsValid() )
 				? InfoCache->UnwindToNonCollapsedPath( InPrimPath, ECollapsingType::Components )
 				: InPrimPath;
 
@@ -1000,9 +989,9 @@ void AUsdStageActor::OnUsdObjectsChanged( const UsdUtils::FObjectChangesByPath&
 			}
 		};
 
-		auto ReloadAssets = [&]( const UE::FSdfPath& InPrimPath, const bool bInResync )
+		auto ReloadAssets = [&]( const UE::FUsdPrim& Prim, const UE::FSdfPath& InPrimPath, const bool bInResync )
 		{
-			UE::FSdfPath AssetsPrimPath = InfoCache.IsValid()
+			UE::FSdfPath AssetsPrimPath = ( Prim && InfoCache.IsValid() )
 				? InfoCache->UnwindToNonCollapsedPath( InPrimPath, ECollapsingType::Assets )
 				: InPrimPath;
 
@@ -1029,7 +1018,9 @@ void AUsdStageActor::OnUsdObjectsChanged( const UsdUtils::FObjectChangesByPath&
 						for ( const FString& MaterialUserPrim : UsdUtils::GetMaterialUsers( PrimToResync ) )
 						{
 							const bool bResyncComponent = true; // We need to force resync to reassign materials
-							UpdateComponents( UE::FSdfPath{ *MaterialUserPrim }, bResyncComponent );
+							UE::FSdfPath MaterialPrimPath{ *MaterialUserPrim };
+							UE::FUsdPrim MaterialPrim = GetOrLoadUsdStage().GetPrimAtPath( MaterialPrimPath );
+							UpdateComponents( MaterialPrim, MaterialPrimPath, bResyncComponent );
 						}
 					}
 
@@ -1045,8 +1036,9 @@ void AUsdStageActor::OnUsdObjectsChanged( const UsdUtils::FObjectChangesByPath&
 			}
 		};
 
-		ReloadAssets( PrimPath, bIsResync );
-		UpdateComponents( PrimPath, bIsResync );
+		UE::FUsdPrim Prim = GetOrLoadUsdStage().GetPrimAtPath( PrimPath );
+		ReloadAssets( Prim, PrimPath, bIsResync );
+		UpdateComponents( Prim, PrimPath, bIsResync );
 
 		if ( HasAuthorityOverStage() )
 		{
