USD: Fixed schema translator task chain continuing even if a prior task failed.

diff --git a/Engine/Plugins/Importers/USDImporter/Source/USDSchemas/Private/USDSchemaTranslator.cpp b/Engine/Plugins/Importers/USDImporter/Source/USDSchemas/Private/USDSchemaTranslator.cpp
--- a/Engine/Plugins/Importers/USDImporter/Source/USDSchemas/Private/USDSchemaTranslator.cpp
+++ b/Engine/Plugins/Importers/USDImporter/Source/USDSchemas/Private/USDSchemaTranslator.cpp
@@ -326,7 +326,14 @@ ESchemaTranslationStatus FUsdSchemaTranslatorTaskChain::Execute(bool bExclusiveS
 	}
 	else
 	{
-		CurrentTask = CurrentTask->Continuation;
+		if ( CurrentTask->Result.IsSet() )
+		{
+			CurrentTask = CurrentTask->Result->Get() ? CurrentTask->Continuation : nullptr;
+		}
+		else
+		{
+			CurrentTask = CurrentTask->Continuation;
+		}
 
 		if (( TranslatorTask = CurrentTask.Get()) != nullptr )
 		{
