USD: Tweak level sequence export options for clarity.

diff --git a/Engine/Plugins/Importers/USDImporter/Source/USDExporter/Public/LevelExporterUSDOptions.h b/Engine/Plugins/Importers/USDImporter/Source/USDExporter/Public/LevelExporterUSDOptions.h
--- a/Engine/Plugins/Importers/USDImporter/Source/USDExporter/Public/LevelExporterUSDOptions.h
+++ b/Engine/Plugins/Importers/USDImporter/Source/USDExporter/Public/LevelExporterUSDOptions.h
@@ -61,7 +61,7 @@ struct USDEXPORTER_API FLevelExporterUSDOptionsInner
 	FIntPoint LandscapeBakeResolution = FIntPoint( 1024, 1024 );
 
 	/** If true, will export sub-levels as separate layers (referenced as sublayers). If false, will collapse all sub-levels in a single exported root layer */
-    UPROPERTY( EditAnywhere, config, BlueprintReadWrite, Category = "Sublayers" )
+    UPROPERTY( EditAnywhere, config, BlueprintReadWrite, Category = "Sublayers", meta = ( DisplayName = "Export Sublevels as Layers" ) )
     bool bExportSublayers = false;
 
 	/** Names of levels that should be ignored when collecting actors to export (e.g. "Persistent Level", "Level1", "MySubLevel", etc.) */
diff --git a/Engine/Plugins/Importers/USDImporter/Source/USDExporter/Public/LevelSequenceExporterUSDOptions.h b/Engine/Plugins/Importers/USDImporter/Source/USDExporter/Public/LevelSequenceExporterUSDOptions.h
--- a/Engine/Plugins/Importers/USDImporter/Source/USDExporter/Public/LevelSequenceExporterUSDOptions.h
+++ b/Engine/Plugins/Importers/USDImporter/Source/USDExporter/Public/LevelSequenceExporterUSDOptions.h
@@ -57,15 +57,16 @@ public:
 	UPROPERTY( EditAnywhere, config, BlueprintReadWrite, Category = "Level Sequence" )
 	bool bExportLevel = false;
 
+	/**
+	 * If checked this will also add the exported level as a sublayer to the USD files emitted for all exported level sequences
+	 */
+	UPROPERTY( EditAnywhere, config, BlueprintReadWrite, Category = "Level Sequence", meta = ( EditCondition = "bExportLevel" ) )
+	bool bUseExportedLevelAsSublayer = false;
+
 	UPROPERTY( EditAnywhere, BlueprintReadWrite, Category = "Level Export", meta = ( EditCondition = "bExportLevel" ) )
 	TWeakObjectPtr<UWorld> Level = nullptr;
 
 	UPROPERTY( EditAnywhere, config, BlueprintReadWrite, Category = "Level Export", meta = ( EditCondition = "bExportLevel" ) )
 	FLevelExporterUSDOptionsInner LevelExportOptions;
 
-	/**
-	 * If checked this will also add the exported level as a sublayer to the USD files emitted for all exported level sequences
-	 */
-	UPROPERTY( EditAnywhere, config, BlueprintReadWrite, Category = "Level Export", meta = ( EditCondition = "bExportLevel" ) )
-	bool bUseExportedLevelAsSublayer = false;
 };
