USD: Fix integration options being clickable for all prim types.

diff --git a/Engine/Plugins/Importers/USDImporter/Source/USDStageEditorViewModels/Private/USDPrimViewModel.cpp b/Engine/Plugins/Importers/USDImporter/Source/USDStageEditorViewModels/Private/USDPrimViewModel.cpp
--- a/Engine/Plugins/Importers/USDImporter/Source/USDStageEditorViewModels/Private/USDPrimViewModel.cpp
+++ b/Engine/Plugins/Importers/USDImporter/Source/USDStageEditorViewModels/Private/USDPrimViewModel.cpp
@@ -27,6 +27,7 @@
 	#include "pxr/usd/usd/references.h"
 	#include "pxr/usd/usd/tokens.h"
 	#include "pxr/usd/usdGeom/xform.h"
+	#include "pxr/usd/usdSkel/root.h"
 #include "USDIncludesEnd.h"
 
 #endif // #if USE_USD_SDK
@@ -337,6 +338,7 @@ bool FUsdPrimViewModel::CanSetUpLiveLink() const
 	const bool bAddSchema = true;
 	return UsdPrim
 		&& !UsdPrim.IsPseudoRoot()
+		&& pxr::UsdPrim{ UsdPrim }.IsA<pxr::UsdGeomXformable>() // SkelRoot is also an Xformable
 		&& UE::USDPrimViewModel::Private::CanAddOrRemoveCustomSchema( UsdPrim, UnrealIdentifiers::LiveLinkAPI, bAddSchema );
 #else
 	return false;
@@ -356,6 +358,7 @@ bool FUsdPrimViewModel::CanRemoveLiveLink() const
 	const bool bAddSchema = false;
 	return UsdPrim
 		&& !UsdPrim.IsPseudoRoot()
+		&& pxr::UsdPrim{ UsdPrim }.IsA<pxr::UsdGeomXformable>() // SkelRoot is also an Xformable
 		&& UE::USDPrimViewModel::Private::CanAddOrRemoveCustomSchema( UsdPrim, UnrealIdentifiers::LiveLinkAPI, bAddSchema );
 #else
 	return false;
@@ -375,6 +378,7 @@ bool FUsdPrimViewModel::CanSetUpControlRig() const
 	const bool bAddSchema = true;
 	return UsdPrim
 		&& !UsdPrim.IsPseudoRoot()
+		&& pxr::UsdPrim{ UsdPrim }.IsA<pxr::UsdSkelRoot>()
 		&& UE::USDPrimViewModel::Private::CanAddOrRemoveCustomSchema( UsdPrim, UnrealIdentifiers::ControlRigAPI, bAddSchema );
 #else
 	return false;
@@ -394,6 +398,7 @@ bool FUsdPrimViewModel::CanRemoveControlRig() const
 	const bool bAddSchema = false;
 	return UsdPrim
 		&& !UsdPrim.IsPseudoRoot()
+		&& pxr::UsdPrim{ UsdPrim }.IsA<pxr::UsdSkelRoot>()
 		&& UE::USDPrimViewModel::Private::CanAddOrRemoveCustomSchema( UsdPrim, UnrealIdentifiers::ControlRigAPI, bAddSchema );
 #else
 	return false;
@@ -413,6 +418,7 @@ bool FUsdPrimViewModel::CanApplyGroomSchema() const
 	const bool bAddSchema = true;
 	return UsdPrim
 		&& !UsdPrim.IsPseudoRoot()
+		&& pxr::UsdPrim{ UsdPrim }.IsA<pxr::UsdGeomXformable>() // UsdGeomCurves is also an Xformable
 		&& UE::USDPrimViewModel::Private::CanAddOrRemoveCustomSchema( UsdPrim, UnrealIdentifiers::GroomAPI, bAddSchema );
 #else
 	return false;
@@ -432,6 +438,7 @@ bool FUsdPrimViewModel::CanRemoveGroomSchema() const
 	const bool bAddSchema = false;
 	return UsdPrim
 		&& !UsdPrim.IsPseudoRoot()
+		&& pxr::UsdPrim{ UsdPrim }.IsA<pxr::UsdGeomXformable>() // UsdGeomCurves is also an Xformable
 		&& UE::USDPrimViewModel::Private::CanAddOrRemoveCustomSchema( UsdPrim, UnrealIdentifiers::GroomAPI, bAddSchema );
 #else
 	return false;
