USD: Fix two blueprint libraries injecting the same ScriptMethod.

diff --git a/Engine/Plugins/Importers/USDImporter/Source/USDExporter/Public/USDExporterBlueprintLibrary.h b/Engine/Plugins/Importers/USDImporter/Source/USDExporter/Public/USDExporterBlueprintLibrary.h
--- a/Engine/Plugins/Importers/USDImporter/Source/USDExporter/Public/USDExporterBlueprintLibrary.h
+++ b/Engine/Plugins/Importers/USDImporter/Source/USDExporter/Public/USDExporterBlueprintLibrary.h
@@ -38,7 +38,7 @@ public:
 	 * can't return nested containers from UFUNCTIONs, so users of this API should call this, and then GetInstanceTransforms.
 	 */
 	UE_DEPRECATED( 5.1, "This function was moved to the UUsdConversionBlueprintLibrary C++ class (UsdConversionLibrary on Python/Blueprint)" )
-	UFUNCTION( BlueprintCallable, meta = ( ScriptMethod ), Category = "USD Foliage Exporter" )
+	UFUNCTION( BlueprintCallable, Category = "USD Foliage Exporter" )
 	static TArray<UFoliageType*> GetUsedFoliageTypes( AInstancedFoliageActor* Actor );
 
 	/**
@@ -46,7 +46,7 @@ public:
 	 * It can be a UStaticMesh in case we're dealing with a UFoliageType_InstancedStaticMesh, but it can be other types of objects.
 	 */
 	UE_DEPRECATED( 5.1, "This function was moved to the UUsdConversionBlueprintLibrary C++ class (UsdConversionLibrary on Python/Blueprint)" )
-	UFUNCTION( BlueprintCallable, meta = ( ScriptMethod ), Category = "USD Foliage Exporter" )
+	UFUNCTION( BlueprintCallable, Category = "USD Foliage Exporter" )
 	static UObject* GetSource( UFoliageType* FoliageType );
 
 	/**
@@ -54,7 +54,7 @@ public:
 	 * Use GetUsedFoliageTypes() to retrieve all foliage types managed by a particular actor.
 	 */
 	UE_DEPRECATED( 5.1, "This function was moved to the UUsdConversionBlueprintLibrary C++ class (UsdConversionLibrary on Python/Blueprint)" )
-	UFUNCTION( BlueprintCallable, meta = ( ScriptMethod ), Category = "USD Foliage Exporter" )
+	UFUNCTION( BlueprintCallable, Category = "USD Foliage Exporter" )
 	static TArray<FTransform> GetInstanceTransforms( AInstancedFoliageActor* Actor, UFoliageType* FoliageType, ULevel* InstancesLevel = nullptr );
 
 	/** Defer to the USDClasses module to actually send analytics information */
