USD: Fix prim attributes panel not updating when clearing selection.

diff --git a/Engine/Plugins/Importers/USDImporter/Source/USDStageEditor/Private/SUSDStageTreeView.cpp b/Engine/Plugins/Importers/USDImporter/Source/USDStageEditor/Private/SUSDStageTreeView.cpp
--- a/Engine/Plugins/Importers/USDImporter/Source/USDStageEditor/Private/SUSDStageTreeView.cpp
+++ b/Engine/Plugins/Importers/USDImporter/Source/USDStageEditor/Private/SUSDStageTreeView.cpp
@@ -505,6 +505,10 @@ void SUsdStageTreeView::SelectPrims( const TArray<FString>& PrimPaths )
 	else
 	{
 		ClearSelection();
+
+		// ClearSelection is not going to fire the OnSelectionChanged event in case we have nothing selected, but we
+		// need to do that to refresh the prim properties panel to display the stage properties instead
+		OnPrimSelectionChanged.ExecuteIfBound({});
 	}
 }
 
