USD: Remove usage of unacceptable words.

diff --git a/Engine/Plugins/Importers/USDImporter/Source/USDExporter/Private/LevelSequenceExporterUSD.cpp b/Engine/Plugins/Importers/USDImporter/Source/USDExporter/Private/LevelSequenceExporterUSD.cpp
--- a/Engine/Plugins/Importers/USDImporter/Source/USDExporter/Private/LevelSequenceExporterUSD.cpp
+++ b/Engine/Plugins/Importers/USDImporter/Source/USDExporter/Private/LevelSequenceExporterUSD.cpp
@@ -813,9 +813,9 @@ namespace UE
 				// Also note that we can't share these bakers with our parent movie scenes in case we're a subsequence, unfortunately, because our
 				// bakers will contain lambdas that write directly to a given prim, and those prims are specific to each layer that we're exporting
 				// (e.g. our parent level sequence will export to different prims than this movie scene will)
-				for ( const UMovieSceneTrack* MasterTrack : MovieScene->GetMasterTracks() )
+				for ( const UMovieSceneTrack* Track : MovieScene->GetMasterTracks() )
 				{
-					const UMovieSceneSubTrack* SubTrack = Cast<const UMovieSceneSubTrack>( MasterTrack );
+					const UMovieSceneSubTrack* SubTrack = Cast<const UMovieSceneSubTrack>( Track );
 					if ( !SubTrack )
 					{
 						continue;
@@ -999,9 +999,9 @@ namespace UE
 						FString Extension;
 						FPaths::Split( UniqueFilePath, Directory, FileName, Extension );
 
-						for ( const UMovieSceneTrack* MasterTrack : MovieScene->GetMasterTracks() )
+						for ( const UMovieSceneTrack* Track : MovieScene->GetMasterTracks() )
 						{
-							const UMovieSceneSubTrack* SubTrack = Cast<const UMovieSceneSubTrack>( MasterTrack );
+							const UMovieSceneSubTrack* SubTrack = Cast<const UMovieSceneSubTrack>( Track );
 							if ( !SubTrack )
 							{
 								continue;
diff --git a/Engine/Plugins/Importers/USDImporter/Source/USDSchemas/Private/USDShadeMaterialTranslator.cpp b/Engine/Plugins/Importers/USDImporter/Source/USDSchemas/Private/USDShadeMaterialTranslator.cpp
--- a/Engine/Plugins/Importers/USDImporter/Source/USDSchemas/Private/USDShadeMaterialTranslator.cpp
+++ b/Engine/Plugins/Importers/USDImporter/Source/USDSchemas/Private/USDShadeMaterialTranslator.cpp
@@ -53,21 +53,21 @@ void FUsdShadeMaterialTranslator::CreateAssets()
 		const bool bIsTranslucent = UsdUtils::IsMaterialTranslucent( ShadeMaterial );
 		const bool bNeedsVirtualTextures = UsdUtils::IsMaterialUsingUDIMs( ShadeMaterial );
 
-		FString MasterMaterialName = TEXT("UsdPreviewSurface");
+		FString BaseMaterialName = TEXT("UsdPreviewSurface");
 
 		if ( bIsTranslucent )
 		{
-			MasterMaterialName += TEXT("Translucent");
+			BaseMaterialName += TEXT("Translucent");
 		}
 
 		if ( bNeedsVirtualTextures )
 		{
-			MasterMaterialName += TEXT("VT");
+			BaseMaterialName += TEXT("VT");
 		}
 
-		const FString MasterMaterialPath = FString::Printf( TEXT("Material'/USDImporter/Materials/%s.%s"), *MasterMaterialName, *MasterMaterialName );
+		const FString BaseMaterialPath = FString::Printf( TEXT("Material'/USDImporter/Materials/%s.%s"), *BaseMaterialName, *BaseMaterialName );
 
-		if ( UMaterialInterface* MasterMaterial = Cast< UMaterialInterface >( FSoftObjectPath( MasterMaterialPath ).TryLoad() ) )
+		if ( UMaterialInterface* BaseMaterial = Cast< UMaterialInterface >( FSoftObjectPath( BaseMaterialPath ).TryLoad() ) )
 		{
 			FName InstanceName = MakeUniqueObjectName( GetTransientPackage(), UMaterialInstance::StaticClass(), *FPaths::GetBaseFilename( PrimPath.GetString() ) );
 			if ( GIsEditor ) // Also have to prevent Standalone game from going with MaterialInstanceConstants
@@ -75,7 +75,7 @@ void FUsdShadeMaterialTranslator::CreateAssets()
 #if WITH_EDITOR
 				if ( UMaterialInstanceConstant* NewMaterial = NewObject<UMaterialInstanceConstant>( GetTransientPackage(), InstanceName, Context->ObjectFlags ) )
 				{
-					NewMaterial->SetParentEditorOnly( MasterMaterial );
+					NewMaterial->SetParentEditorOnly( BaseMaterial );
 
 					UUsdAssetImportData* ImportData = NewObject< UUsdAssetImportData >( NewMaterial, TEXT( "USDAssetImportData" ) );
 					ImportData->PrimPath = PrimPath.GetString();
@@ -105,7 +105,7 @@ void FUsdShadeMaterialTranslator::CreateAssets()
 				}
 #endif // WITH_EDITOR
 			}
-			else if ( UMaterialInstanceDynamic* NewMaterial = UMaterialInstanceDynamic::Create( MasterMaterial, GetTransientPackage(), InstanceName ) )
+			else if ( UMaterialInstanceDynamic* NewMaterial = UMaterialInstanceDynamic::Create( BaseMaterial, GetTransientPackage(), InstanceName ) )
 			{
 				TMap<FString, int32> Unused;
 				TMap<FString, int32>& PrimvarToUVIndex = Context->MaterialToPrimvarToUVIndex ? Context->MaterialToPrimvarToUVIndex->FindOrAdd( PrimPath.GetString() ) : Unused;
diff --git a/Engine/Plugins/Importers/USDImporter/Source/USDStage/USDStage.Build.cs b/Engine/Plugins/Importers/USDImporter/Source/USDStage/USDStage.Build.cs
--- a/Engine/Plugins/Importers/USDImporter/Source/USDStage/USDStage.Build.cs
+++ b/Engine/Plugins/Importers/USDImporter/Source/USDStage/USDStage.Build.cs
@@ -31,7 +31,7 @@ namespace UnrealBuildTool.Rules
 					"MeshDescription",
 					"MovieScene",
 					"MovieSceneTracks",
-					"Projects", // So that we can use the plugin manager to find out our content dir and cook the master materials
+					"Projects", // So that we can use the plugin manager to find out our content dir and cook the base materials
 					"Slate",
 					"SlateCore",
 					"StaticMeshDescription",
diff --git a/Engine/Plugins/Importers/USDImporter/Source/USDUtilities/Private/USDShadeConversion.cpp b/Engine/Plugins/Importers/USDImporter/Source/USDUtilities/Private/USDShadeConversion.cpp
--- a/Engine/Plugins/Importers/USDImporter/Source/USDUtilities/Private/USDShadeConversion.cpp
+++ b/Engine/Plugins/Importers/USDImporter/Source/USDUtilities/Private/USDShadeConversion.cpp
@@ -1080,7 +1080,7 @@ namespace UE
 			};
 
 			/**
-			 * Specialized version of FSetParameterValueVisitor for UE's UsdPreviewSurface master materials.
+			 * Specialized version of FSetParameterValueVisitor for UE's UsdPreviewSurface base materials.
 			 */
 			struct FSetPreviewSurfaceParameterValueVisitor : private FSetParameterValueVisitor
 			{
@@ -1756,7 +1756,7 @@ bool UsdToUnreal::ConvertMaterial( const pxr::UsdShadeMaterial& UsdShadeMaterial
 	bool bHasMaterialInfo = false;
 
 	// We assume that all samplers for a material instance that supports UDIMs would be virtual so we'll force all the textures to be virtual.
-	// The idea being that when building a master material, we can't know which textures will be virtual so it's safer to set the samplers as if they were all virtual.
+	// The idea being that when building a base material, we can't know which textures will be virtual so it's safer to set the samplers as if they were all virtual.
 	const bool bForceVirtualTextures = UsdUtils::IsMaterialUsingUDIMs( UsdShadeMaterial );
 
 	pxr::TfToken RenderContextToken = pxr::UsdShadeTokens->universalRenderContext;
diff --git a/Engine/Plugins/Importers/USDImporter/Source/USDUtilities/Public/USDGeomMeshConversion.h b/Engine/Plugins/Importers/USDImporter/Source/USDUtilities/Public/USDGeomMeshConversion.h
--- a/Engine/Plugins/Importers/USDImporter/Source/USDUtilities/Public/USDGeomMeshConversion.h
+++ b/Engine/Plugins/Importers/USDImporter/Source/USDUtilities/Public/USDGeomMeshConversion.h
@@ -258,7 +258,7 @@ namespace UsdUtils
 	/** Creates a FDisplayColorMaterial object describing the vertex color/opacity data from UsdMesh at time TimeCode */
 	USDUTILITIES_API TOptional<FDisplayColorMaterial> ExtractDisplayColorMaterial( const pxr::UsdGeomMesh& UsdMesh, const pxr::UsdTimeCode TimeCode = pxr::UsdTimeCode::EarliestTime() );
 
-	/** Creates a dynamic material instance using the right master material depending on the given description */
+	/** Creates a dynamic material instance using the right base material depending on the given description */
 	USDUTILITIES_API UMaterialInstanceDynamic* CreateDisplayColorMaterialInstanceDynamic( const UsdUtils::FDisplayColorMaterial& DisplayColorDescription );
 	USDUTILITIES_API UMaterialInstanceConstant* CreateDisplayColorMaterialInstanceConstant( const UsdUtils::FDisplayColorMaterial& DisplayColorDescription );
 
