USD: Improve support for prim edit operations.

Rename Add Prim to Add Child;
Add options to Cut/Copy/Paste/Duplicate prims;
Rename "Remove" to "Delete" to match the GenericCommands pattern;
Allow referencing specific prims, internal prims, and using offset/scale on prim references and payloads;
Allow adding/clearing Payloads like references;
Prevent attempting to rename multiple prims at once;
Fix being able to add references to multiple prims at once.

diff --git a/Engine/Plugins/Importers/USDImporter/Source/USDClasses/Public/USDDuplicateType.h b/Engine/Plugins/Importers/USDImporter/Source/USDClasses/Public/USDDuplicateType.h
new file mode 100644
--- /dev/null
+++ b/Engine/Plugins/Importers/USDImporter/Source/USDClasses/Public/USDDuplicateType.h
@@ -0,0 +1,19 @@
+// Copyright Epic Games, Inc. All Rights Reserved.
+
+#pragma once
+
+#include "USDDuplicateType.generated.h"
+
+/**
+ * Describes the different types of "prim duplication" operations we support with UsdUtils::DuplicatePrims
+ */
+UENUM( BlueprintType )
+enum class EUsdDuplicateType : uint8
+{
+	/** Generate a flattened duplicate of the composed prim onto the current edit target */
+	FlattenComposedPrim,
+	/** Duplicate the prim's specs on the current edit target only */
+	SingleLayerSpecs,
+	/** Duplicate each of the prim's specs across the entire stage */
+	AllLocalLayerSpecs,
+};
diff --git a/Engine/Plugins/Importers/USDImporter/Source/USDClasses/Public/USDReferenceOptions.h b/Engine/Plugins/Importers/USDImporter/Source/USDClasses/Public/USDReferenceOptions.h
new file mode 100644
--- /dev/null
+++ b/Engine/Plugins/Importers/USDImporter/Source/USDClasses/Public/USDReferenceOptions.h
@@ -0,0 +1,39 @@
+// Copyright Epic Games, Inc. All Rights Reserved.
+
+#pragma once
+
+#include "USDReferenceOptions.generated.h"
+
+/**
+ * Options to display when adding a reference or a payload for a prim
+ */
+UCLASS( Blueprintable )
+class USDCLASSES_API UUsdReferenceOptions : public UObject
+{
+	GENERATED_BODY()
+
+public:
+	// When enabled, the reference/payload will target a prim on this stage
+    UPROPERTY( EditAnywhere, BlueprintReadWrite, Category = "Options" )
+	bool bInternalReference = false;
+
+	// File to use as the reference
+    UPROPERTY( EditAnywhere, BlueprintReadWrite, Category = "Options", meta = (EditCondition = "!bInternalReference") )
+	FFilePath TargetFile;
+
+	// Use the default prim of the target stage as the referenced/payload prim
+	UPROPERTY( EditAnywhere, BlueprintReadWrite, Category = "Options" )
+	bool bUseDefaultPrim = true;
+
+	// Use a specific prim of the target stage as the referenced/payload prim
+    UPROPERTY( EditAnywhere, BlueprintReadWrite, Category = "Options", meta = (EditCondition = "!bUseDefaultPrim") )
+	FString TargetPrimPath;
+
+	// Offset to apply to the referenced/payload prim's time sampled attributes
+    UPROPERTY( EditAnywhere, BlueprintReadWrite, Category = "Options", AdvancedDisplay )
+    float TimeCodeOffset = 0.0f;
+
+	// TimeCode scaling factor to apply to the referenced/payload prim's time sampled attributes
+    UPROPERTY( EditAnywhere, BlueprintReadWrite, Category = "Options", AdvancedDisplay )
+    float TimeCodeScale = 1.0f;
+};
diff --git a/Engine/Plugins/Importers/USDImporter/Source/USDExporter/Private/AnimSequenceExporterUSD.cpp b/Engine/Plugins/Importers/USDImporter/Source/USDExporter/Private/AnimSequenceExporterUSD.cpp
--- a/Engine/Plugins/Importers/USDImporter/Source/USDExporter/Private/AnimSequenceExporterUSD.cpp
+++ b/Engine/Plugins/Importers/USDImporter/Source/USDExporter/Private/AnimSequenceExporterUSD.cpp
@@ -130,7 +130,7 @@ bool UAnimSequenceExporterUSD::ExportBinary( UObject* Object, const TCHAR* Type,
 			Options->PreviewMeshOptions.MaterialBakingOptions.TexturesDir.Path = FPaths::Combine( FPaths::GetPath( UExporter::CurrentFilename ), TEXT( "Textures" ) );
 
 			const bool bIsImport = false;
-			const bool bContinue = SUsdOptionsWindow::ShowOptions( *Options, bIsImport );
+			const bool bContinue = SUsdOptionsWindow::ShowImportExportOptions( *Options, bIsImport );
 			if ( !bContinue )
 			{
 				return false;
diff --git a/Engine/Plugins/Importers/USDImporter/Source/USDExporter/Private/LevelExporterUSD.cpp b/Engine/Plugins/Importers/USDImporter/Source/USDExporter/Private/LevelExporterUSD.cpp
--- a/Engine/Plugins/Importers/USDImporter/Source/USDExporter/Private/LevelExporterUSD.cpp
+++ b/Engine/Plugins/Importers/USDImporter/Source/USDExporter/Private/LevelExporterUSD.cpp
@@ -81,7 +81,7 @@ bool ULevelExporterUSD::ExportBinary( UObject* Object, const TCHAR* Type, FArchi
 		Options->Inner.AssetFolder.Path = FPaths::Combine( FPaths::GetPath( UExporter::CurrentFilename ), TEXT( "Assets" ) );
 
 		const bool bIsImport = false;
-		const bool bContinue = SUsdOptionsWindow::ShowOptions( *Options, bIsImport );
+		const bool bContinue = SUsdOptionsWindow::ShowImportExportOptions( *Options, bIsImport );
 		if ( !bContinue )
 		{
 			return false;
diff --git a/Engine/Plugins/Importers/USDImporter/Source/USDExporter/Private/LevelSequenceExporterUSD.cpp b/Engine/Plugins/Importers/USDImporter/Source/USDExporter/Private/LevelSequenceExporterUSD.cpp
--- a/Engine/Plugins/Importers/USDImporter/Source/USDExporter/Private/LevelSequenceExporterUSD.cpp
+++ b/Engine/Plugins/Importers/USDImporter/Source/USDExporter/Private/LevelSequenceExporterUSD.cpp
@@ -1122,7 +1122,7 @@ bool ULevelSequenceExporterUsd::ExportBinary( UObject* Object, const TCHAR* Type
 		Options->TimeCodesPerSecond = MovieScene->GetDisplayRate().AsDecimal();
 
 		const bool bIsImport = false;
-		const bool bContinue = SUsdOptionsWindow::ShowOptions( *Options, bIsImport );
+		const bool bContinue = SUsdOptionsWindow::ShowImportExportOptions( *Options, bIsImport );
 		if ( !bContinue )
 		{
 			return false;
diff --git a/Engine/Plugins/Importers/USDImporter/Source/USDExporter/Private/MaterialExporterUSD.cpp b/Engine/Plugins/Importers/USDImporter/Source/USDExporter/Private/MaterialExporterUSD.cpp
--- a/Engine/Plugins/Importers/USDImporter/Source/USDExporter/Private/MaterialExporterUSD.cpp
+++ b/Engine/Plugins/Importers/USDImporter/Source/USDExporter/Private/MaterialExporterUSD.cpp
@@ -106,7 +106,7 @@ bool UMaterialExporterUsd::ExportBinary( UObject* Object, const TCHAR* Type, FAr
 			Options->MaterialBakingOptions.TexturesDir.Path = FPaths::Combine( FPaths::GetPath( UExporter::CurrentFilename ), TEXT( "Textures" ) );
 
 			const bool bIsImport = false;
-			const bool bContinue = SUsdOptionsWindow::ShowOptions( *Options, bIsImport );
+			const bool bContinue = SUsdOptionsWindow::ShowImportExportOptions( *Options, bIsImport );
 			if ( !bContinue )
 			{
 				return false;
diff --git a/Engine/Plugins/Importers/USDImporter/Source/USDExporter/Private/SkeletalMeshExporterUSD.cpp b/Engine/Plugins/Importers/USDImporter/Source/USDExporter/Private/SkeletalMeshExporterUSD.cpp
--- a/Engine/Plugins/Importers/USDImporter/Source/USDExporter/Private/SkeletalMeshExporterUSD.cpp
+++ b/Engine/Plugins/Importers/USDImporter/Source/USDExporter/Private/SkeletalMeshExporterUSD.cpp
@@ -122,7 +122,7 @@ bool USkeletalMeshExporterUsd::ExportBinary( UObject* Object, const TCHAR* Type,
 			Options->MeshAssetOptions.MaterialBakingOptions.TexturesDir.Path = FPaths::Combine( FPaths::GetPath( UExporter::CurrentFilename ), TEXT( "Textures" ) );
 
 			const bool bIsImport = false;
-			const bool bContinue = SUsdOptionsWindow::ShowOptions( *Options, bIsImport );
+			const bool bContinue = SUsdOptionsWindow::ShowImportExportOptions( *Options, bIsImport );
 			if ( !bContinue )
 			{
 				return false;
diff --git a/Engine/Plugins/Importers/USDImporter/Source/USDExporter/Private/StaticMeshExporterUSD.cpp b/Engine/Plugins/Importers/USDImporter/Source/USDExporter/Private/StaticMeshExporterUSD.cpp
--- a/Engine/Plugins/Importers/USDImporter/Source/USDExporter/Private/StaticMeshExporterUSD.cpp
+++ b/Engine/Plugins/Importers/USDImporter/Source/USDExporter/Private/StaticMeshExporterUSD.cpp
@@ -132,7 +132,7 @@ bool UStaticMeshExporterUsd::ExportBinary( UObject* Object, const TCHAR* Type, F
 			Options->MeshAssetOptions.MaterialBakingOptions.TexturesDir.Path = FPaths::Combine( FPaths::GetPath( UExporter::CurrentFilename ), TEXT( "Textures" ) );
 
 			const bool bIsImport = false;
-			const bool bContinue = SUsdOptionsWindow::ShowOptions( *Options, bIsImport );
+			const bool bContinue = SUsdOptionsWindow::ShowImportExportOptions( *Options, bIsImport );
 			if ( !bContinue )
 			{
 				return false;
diff --git a/Engine/Plugins/Importers/USDImporter/Source/USDExporter/Private/USDConversionBlueprintLibrary.cpp b/Engine/Plugins/Importers/USDImporter/Source/USDExporter/Private/USDConversionBlueprintLibrary.cpp
--- a/Engine/Plugins/Importers/USDImporter/Source/USDExporter/Private/USDConversionBlueprintLibrary.cpp
+++ b/Engine/Plugins/Importers/USDImporter/Source/USDExporter/Private/USDConversionBlueprintLibrary.cpp
@@ -468,3 +468,138 @@ void UUsdConversionBlueprintLibrary::SendAnalytics( const TArray<FAnalyticsEvent
 
 	IUsdClassesModule::SendAnalytics( MoveTemp( Converted ), EventName, bAutomated, ElapsedSeconds, NumberOfFrames, Extension );
 }
+
+void UUsdConversionBlueprintLibrary::RemoveAllPrimSpecs( const FString& StageRootLayer, const FString& PrimPath, const FString& TargetLayer )
+{
+#if USE_USD_SDK
+	const bool bUseStageCache = true;
+	UE::FUsdStage Stage = UnrealUSDWrapper::OpenStage( *StageRootLayer, EUsdInitialLoadSet::LoadAll, bUseStageCache );
+	if ( !Stage )
+	{
+		return;
+	}
+
+	UsdUtils::RemoveAllPrimSpecs(
+		Stage.GetPrimAtPath( UE::FSdfPath{ *PrimPath } ),
+		UE::FSdfLayer::FindOrOpen( *TargetLayer )
+	);
+#endif // USE_USD_SDK
+}
+
+bool UUsdConversionBlueprintLibrary::CutPrims( const FString& StageRootLayer, const TArray<FString>& PrimPaths )
+{
+#if USE_USD_SDK
+	const bool bUseStageCache = true;
+	UE::FUsdStage Stage = UnrealUSDWrapper::OpenStage( *StageRootLayer, EUsdInitialLoadSet::LoadAll, bUseStageCache );
+	if ( !Stage )
+	{
+		return false;
+	}
+
+	TArray<UE::FUsdPrim> Prims;
+	Prims.Reserve( PrimPaths.Num() );
+
+	for ( const FString& PrimPath : PrimPaths )
+	{
+		Prims.Add( Stage.GetPrimAtPath( UE::FSdfPath{ *PrimPath } ) );
+	}
+
+	return UsdUtils::CutPrims( Prims );
+#else
+	return false;
+#endif // USE_USD_SDK
+}
+
+bool UUsdConversionBlueprintLibrary::CopyPrims( const FString& StageRootLayer, const TArray<FString>& PrimPaths )
+{
+#if USE_USD_SDK
+	const bool bUseStageCache = true;
+	UE::FUsdStage Stage = UnrealUSDWrapper::OpenStage( *StageRootLayer, EUsdInitialLoadSet::LoadAll, bUseStageCache );
+	if ( !Stage )
+	{
+		return false;
+	}
+
+	TArray<UE::FUsdPrim> Prims;
+	Prims.Reserve( PrimPaths.Num() );
+
+	for ( const FString& PrimPath : PrimPaths )
+	{
+		Prims.Add( Stage.GetPrimAtPath( UE::FSdfPath{ *PrimPath } ) );
+	}
+
+	return UsdUtils::CopyPrims( Prims );
+#else
+	return false;
+#endif // USE_USD_SDK
+}
+
+TArray<FString> UUsdConversionBlueprintLibrary::PastePrims( const FString& StageRootLayer, const FString& ParentPrimPath )
+{
+	TArray<FString> Result;
+
+#if USE_USD_SDK
+	const bool bUseStageCache = true;
+	UE::FUsdStage Stage = UnrealUSDWrapper::OpenStage( *StageRootLayer, EUsdInitialLoadSet::LoadAll, bUseStageCache );
+	if ( !Stage )
+	{
+		return Result;
+	}
+
+	TArray<UE::FSdfPath> PastedPrims = UsdUtils::PastePrims( Stage.GetPrimAtPath( UE::FSdfPath{ *ParentPrimPath } ) );
+
+	for ( const UE::FSdfPath& DuplicatePrim : PastedPrims )
+	{
+		Result.Add( DuplicatePrim.GetString() );
+	}
+#endif // USE_USD_SDK
+
+	return Result;
+}
+
+bool UUsdConversionBlueprintLibrary::CanPastePrims()
+{
+	return UsdUtils::CanPastePrims();
+}
+
+void UUsdConversionBlueprintLibrary::ClearPrimClipboard()
+{
+	UsdUtils::ClearPrimClipboard();
+}
+
+TArray<FString> UUsdConversionBlueprintLibrary::DuplicatePrims( const FString& StageRootLayer, const TArray<FString>& PrimPaths, EUsdDuplicateType DuplicateType, const FString& TargetLayer )
+{
+	TArray<FString> Result;
+	Result.SetNum( PrimPaths.Num() );
+
+#if USE_USD_SDK
+	const bool bUseStageCache = true;
+	UE::FUsdStage Stage = UnrealUSDWrapper::OpenStage( *StageRootLayer, EUsdInitialLoadSet::LoadAll, bUseStageCache );
+	if ( !Stage )
+	{
+		return Result;
+	}
+
+	TArray<UE::FUsdPrim> Prims;
+	Prims.Reserve( PrimPaths.Num() );
+
+	for ( const FString& PrimPath : PrimPaths )
+	{
+		Prims.Add( Stage.GetPrimAtPath( UE::FSdfPath{ *PrimPath } ) );
+	}
+
+	TArray<UE::FSdfPath> DuplicatedPrims = UsdUtils::DuplicatePrims(
+		Prims,
+		DuplicateType,
+		UE::FSdfLayer::FindOrOpen( *TargetLayer )
+	);
+
+	for ( int32 Index = 0; Index < DuplicatedPrims.Num(); ++Index )
+	{
+		const UE::FSdfPath& DuplicatePrim = DuplicatedPrims[Index];
+		Result[ Index ] = DuplicatePrim.GetString();
+	}
+#endif // USE_USD_SDK
+
+	return Result;
+}
\ No newline at end of file
diff --git a/Engine/Plugins/Importers/USDImporter/Source/USDExporter/Public/USDConversionBlueprintLibrary.h b/Engine/Plugins/Importers/USDImporter/Source/USDExporter/Public/USDConversionBlueprintLibrary.h
--- a/Engine/Plugins/Importers/USDImporter/Source/USDExporter/Public/USDConversionBlueprintLibrary.h
+++ b/Engine/Plugins/Importers/USDImporter/Source/USDExporter/Public/USDConversionBlueprintLibrary.h
@@ -107,4 +107,91 @@ public:
 	/** Defer to the USDClasses module to actually send analytics information */
 	UFUNCTION( BlueprintCallable, Category = "Analytics" )
 	static void SendAnalytics( const TArray<FAnalyticsEventAttr>& Attrs, const FString& EventName, bool bAutomated, double ElapsedSeconds, double NumberOfFrames, const FString& Extension );
+
+    /**
+	 * Removes all the prim specs for Prim on the given Layer.
+	 *
+	 * This function is useful in case the prim is inside a variant set: In that case, just calling FUsdStage::RemovePrim()
+	 * will attempt to remove the "/Root/Example/Child", which wouldn't remove the "/Root{Varset=Var}Example/Child" spec,
+	 * meaning the prim may still be left on the stage. Note that it's even possible to have both of those specs at the same time:
+	 * for example when we have a prim inside a variant set, but outside of it we have overrides to the same prim. This function
+	 * will remove both.
+	 *
+	 * @param StageRootLayer - Path to the root layer of the stage from which we should fetch the Prims
+	 * @param PrimPath - Prim to remove
+	 * @param Layer - Layer to remove prim specs from. This can be left with the invalid layer (default) in order to remove all
+	 *				  specs from the entire stage's local layer stack.
+	 */
+	UFUNCTION( BlueprintCallable, Category = "Prim utils" )
+	static void RemoveAllPrimSpecs( const FString& StageRootLayer, const FString& PrimPath, const FString& TargetLayer );
+
+	/**
+	 * Copies flattened versions of the input prims onto the clipboard stage and removes all the prim specs for Prims from their stages.
+	 * These cut prims can then be pasted with PastePrims.
+	 *
+	 * @param StageRootLayer - Path to the root layer of the stage from which we should fetch the Prims
+	 * @param PrimPaths - Prims to cut
+	 * @return True if we managed to cut
+	 */
+	UFUNCTION( BlueprintCallable, Category = "Prim utils" )
+	static bool CutPrims( const FString& StageRootLayer, const TArray<FString>& PrimPaths );
+
+	/**
+	 * Copies flattened versions of the input prims onto the clipboard stage.
+	 * These copied prims can then be pasted with PastePrims.
+	 *
+	 * @param StageRootLayer - Path to the root layer of the stage from which we should fetch the Prims
+	 * @param PrimPaths - Prims to copy
+	 * @return True if we managed to copy
+	 */
+	UFUNCTION( BlueprintCallable, Category = "Prim utils" )
+	static bool CopyPrims( const FString& StageRootLayer, const TArray<FString>& PrimPaths );
+
+	/**
+	 * Pastes the prims from the clipboard stage as children of ParentPrim.
+	 *
+	 * The pasted prims may be renamed in order to have valid names for the target location, which is why this function
+	 * returns the pasted prim paths.
+	 * This function returns just paths instead of actual prims because USD needs to respond to the notices about
+	 * the created prim specs before the prims are fully created, which means we wouldn't be able to return the
+	 * created prims yet, in case this function was called from within an SdfChangeBlock.
+	 *
+	 * @param StageRootLayer - Path to the root layer of the stage from which we should fetch the Prims
+	 * @param ParentPrimPath - Prim that will become parent to the pasted prims
+	 * @return Paths to the pasted prim specs, after they were added as children of ParentPrim
+	 */
+	UFUNCTION( BlueprintCallable, Category = "Prim utils" )
+	static TArray<FString> PastePrims( const FString& StageRootLayer, const FString& ParentPrimPath );
+
+	/** Returns true if we have prims that we can paste within our clipboard stage */
+	UFUNCTION( BlueprintCallable, Category = "Prim utils" )
+	static bool CanPastePrims();
+
+	/** Clears all prims from our clipboard stage */
+	UFUNCTION( BlueprintCallable, Category = "Prim utils" )
+	static void ClearPrimClipboard();
+
+	/**
+	 * Duplicates all provided Prims one-by-one, performing the requested DuplicateType.
+	 * See the documentation on EUsdDuplicateType for the different operation types.
+	 *
+	 * The duplicated prims may be renamed in order to have valid names for the target location, which is why this
+	 * function returns the pasted prim paths.
+	 * This function returns just paths instead of actual prims because USD needs to respond to the notices about
+	 * the created prim specs before the prims are fully created, which means we wouldn't be able to return the
+	 * created prims yet, in case this function was called from within an SdfChangeBlock.
+	 *
+	 * @param StageRootLayer - Path to the root layer of the stage from which we should fetch the Prims
+	 * @param PrimPaths - Prims to duplicate
+	 * @param DuplicateType - Type of prim duplication to perform
+	 * @param TargetLayer - Target layer to use when duplicating, if relevant for that duplication type
+	 * @return Paths to the duplicated prim specs, after they were added as children of ParentPrim.
+	 */
+	UFUNCTION( BlueprintCallable, Category = "Prim utils" )
+	static TArray<FString> DuplicatePrims(
+		const FString& StageRootLayer,
+		const TArray<FString>& PrimPaths,
+		EUsdDuplicateType DuplicateType,
+		const FString& TargetLayer
+	);
 };
diff --git a/Engine/Plugins/Importers/USDImporter/Source/USDStageEditor/Private/SUSDStageTreeView.cpp b/Engine/Plugins/Importers/USDImporter/Source/USDStageEditor/Private/SUSDStageTreeView.cpp
--- a/Engine/Plugins/Importers/USDImporter/Source/USDStageEditor/Private/SUSDStageTreeView.cpp
+++ b/Engine/Plugins/Importers/USDImporter/Source/USDStageEditor/Private/SUSDStageTreeView.cpp
@@ -5,7 +5,10 @@
 #include "SUSDStageEditorStyle.h"
 #include "UnrealUSDWrapper.h"
 #include "USDConversionUtils.h"
+#include "USDDuplicateType.h"
 #include "USDLayerUtils.h"
+#include "USDOptionsWindow.h"
+#include "USDReferenceOptions.h"
 #include "USDStageActor.h"
 #include "USDStageModule.h"
 #include "USDTypesConversion.h"
@@ -16,13 +19,15 @@
 #include "UsdWrappers/UsdPrim.h"
 #include "UsdWrappers/UsdStage.h"
 
-#include "Styling/AppStyle.h"
 #include "Engine/World.h"
 #include "Framework/Application/SlateApplication.h"
+#include "Framework/Commands/GenericCommands.h"
 #include "Framework/MultiBox/MultiBoxBuilder.h"
 #include "Modules/ModuleManager.h"
 #include "ScopedTransaction.h"
+#include "Styling/AppStyle.h"
 #include "Templates/Function.h"
+#include "UObject/StrongObjectPtr.h"
 #include "Widgets/Images/SImage.h"
 #include "Widgets/Input/SButton.h"
 #include "Widgets/Input/SCheckBox.h"
@@ -167,22 +172,20 @@ public:
 
 	virtual TSharedRef< SWidget > GenerateWidget( const TSharedPtr< IUsdTreeViewItem > InTreeItem, const TSharedPtr< ITableRow > TableRow ) override
 	{
-		bool bHasPayload = false;
-
-		if ( InTreeItem )
-		{
-			const TSharedPtr< FUsdPrimViewModel > TreeItem = StaticCastSharedPtr< FUsdPrimViewModel >( InTreeItem );
-			bHasPayload = TreeItem->RowData->HasPayload();
-		}
+		const TWeakPtr< FUsdPrimViewModel > TreeItem = StaticCastSharedPtr< FUsdPrimViewModel >( InTreeItem );
 
-		TSharedRef< SWidget > Item = SNullWidget::NullWidget;
+		 TSharedRef< SWidget > Item = SNew( SCheckBox )
+			.Visibility_Lambda( [this, TreeItem]() -> EVisibility
+			{
+				if ( const TSharedPtr< FUsdPrimViewModel > PinnedItem = TreeItem.Pin() )
+				{
+					return PinnedItem->RowData->HasPayload() ? EVisibility::Visible : EVisibility::Collapsed;
+				}
 
-		if ( bHasPayload )
-		{
-			Item = SNew( SCheckBox )
-					.IsChecked( this, &FUsdStagePayloadColumn::IsChecked, StaticCastSharedPtr< FUsdPrimViewModel >( InTreeItem ) )
-					.OnCheckStateChanged( this, &FUsdStagePayloadColumn::OnCheckedPayload, StaticCastSharedPtr< FUsdPrimViewModel >( InTreeItem ) );
-		}
+				return EVisibility::Collapsed;
+			})
+			.IsChecked( this, &FUsdStagePayloadColumn::IsChecked, StaticCastSharedPtr< FUsdPrimViewModel >( InTreeItem ) )
+			.OnCheckStateChanged( this, &FUsdStagePayloadColumn::OnCheckedPayload, StaticCastSharedPtr< FUsdPrimViewModel >( InTreeItem ) );
 
 		return Item;
 	}
@@ -341,6 +344,38 @@ void SUsdStageTreeView::Construct( const FArguments& InArgs, AUsdStageActor* InU
 
 	OnPrimSelectionChanged = InArgs._OnPrimSelectionChanged;
 
+	UICommandList = MakeShared<FUICommandList>();
+	UICommandList->MapAction(
+		FGenericCommands::Get().Cut,
+		FExecuteAction::CreateSP( this, &SUsdStageTreeView::OnCutPrim ),
+		FCanExecuteAction::CreateSP( this, &SUsdStageTreeView::DoesPrimExistOnStage )
+	);
+	UICommandList->MapAction(
+		FGenericCommands::Get().Copy,
+		FExecuteAction::CreateSP( this, &SUsdStageTreeView::OnCopyPrim ),
+		FCanExecuteAction::CreateSP( this, &SUsdStageTreeView::DoesPrimExistOnStage )
+	);
+	UICommandList->MapAction(
+		FGenericCommands::Get().Paste,
+		FExecuteAction::CreateSP( this, &SUsdStageTreeView::OnPastePrim ),
+		FCanExecuteAction::CreateSP( this, &SUsdStageTreeView::CanPastePrim )
+	);
+	UICommandList->MapAction(
+		FGenericCommands::Get().Duplicate,
+		FExecuteAction::CreateSP( this, &SUsdStageTreeView::OnDuplicatePrim, EUsdDuplicateType::AllLocalLayerSpecs ),
+		FCanExecuteAction::CreateSP( this, &SUsdStageTreeView::DoesPrimExistOnStage )
+	);
+	UICommandList->MapAction(
+		FGenericCommands::Get().Delete,
+		FExecuteAction::CreateSP( this, &SUsdStageTreeView::OnDeletePrim ),
+		FCanExecuteAction::CreateSP( this, &SUsdStageTreeView::DoesPrimExistOnStage )
+	);
+	UICommandList->MapAction(
+		FGenericCommands::Get().Rename,
+		FExecuteAction::CreateSP( this, &SUsdStageTreeView::OnRenamePrim ),
+		FCanExecuteAction::CreateLambda( [this]() { return SUsdStageTreeView::DoesPrimExistOnStage() && GetSelectedItems().Num() == 1; } )
+	);
+
 	Refresh( InUsdStageActor );
 }
 
@@ -560,43 +595,74 @@ TSharedPtr< SWidget > SUsdStageTreeView::ConstructPrimContextMenu()
 {
 	TSharedRef< SWidget > MenuWidget = SNullWidget::NullWidget;
 
-	FMenuBuilder PrimOptions( true, nullptr );
-	PrimOptions.BeginSection( "Prim", LOCTEXT("Prim", "Prim") );
+	const bool bInShouldCloseWindowAfterMenuSelection = true;
+	FMenuBuilder PrimOptions( bInShouldCloseWindowAfterMenuSelection, UICommandList );
+	PrimOptions.BeginSection( "Edit", LOCTEXT("EditText", "Edit") );
 	{
 		PrimOptions.AddMenuEntry(
-			LOCTEXT("AddPrim", "Add Prim"),
-			LOCTEXT("AddPrim_ToolTip", "Adds a new prim"),
-			FSlateIcon(),
+			TAttribute<FText>::Create( TAttribute<FText>::FGetter::CreateLambda( [this]()
+			{
+				return GetSelectedItems().Num() == 0
+					? LOCTEXT( "AddTopLevelPrim", "Add Prim" )
+					: LOCTEXT( "AddPrim", "Add Child" );
+			})),
+			TAttribute<FText>::Create( TAttribute<FText>::FGetter::CreateLambda( [this]()
+			{
+				return GetSelectedItems().Num() == 0
+					? LOCTEXT( "AddTopPrim_ToolTip", "Adds a new top-level prim" )
+					: LOCTEXT( "AddPrim_ToolTip", "Adds a new prim as a child of this one" );
+			})),
+			FSlateIcon( FAppStyle::GetAppStyleSetName(), "Icons.PlaceActors" ),
 			FUIAction(
-				FExecuteAction::CreateSP( this, &SUsdStageTreeView::OnAddPrim ),
-				FCanExecuteAction::CreateSP( this, &SUsdStageTreeView::CanAddPrim )
+				FExecuteAction::CreateSP( this, &SUsdStageTreeView::OnAddChildPrim ),
+				FCanExecuteAction::CreateSP( this, &SUsdStageTreeView::CanAddChildPrim )
 			),
 			NAME_None,
 			EUserInterfaceActionType::Button
 		);
 
 		PrimOptions.AddMenuEntry(
-			LOCTEXT( "RenamePrim", "Rename Prim" ),
-			LOCTEXT( "RenamePrim_ToolTip", "Renames the prim on all layers" ),
-			FSlateIcon(),
-			FUIAction(
-				FExecuteAction::CreateSP( this, &SUsdStageTreeView::OnRenamePrim ),
-				FCanExecuteAction::CreateSP( this, &SUsdStageTreeView::CanExecutePrimAction )
-			),
+			FGenericCommands::Get().Cut,
 			NAME_None,
-			EUserInterfaceActionType::Button
+			LOCTEXT( "Cut_Text", "Cut" ),
+			LOCTEXT( "Cut_ToolTip", "Cuts this prim selection from the stage" )
 		);
 
 		PrimOptions.AddMenuEntry(
-			LOCTEXT("RemovePrim", "Remove Prim"),
-			LOCTEXT("RemovePrim_ToolTip", "Removes the prim and its children from the current edit target"),
-			FSlateIcon(),
-			FUIAction(
-				FExecuteAction::CreateSP( this, &SUsdStageTreeView::OnRemovePrim ),
-				FCanExecuteAction::CreateSP( this, &SUsdStageTreeView::CanExecutePrimAction )
-			),
+			FGenericCommands::Get().Copy,
 			NAME_None,
-			EUserInterfaceActionType::Button
+			LOCTEXT( "Copy_Text", "Copy" ),
+			LOCTEXT( "Copy_ToolTip", "Copies the selected prims from the stage" )
+		);
+
+		PrimOptions.AddMenuEntry(
+			FGenericCommands::Get().Paste,
+			NAME_None,
+			LOCTEXT( "Paste_Text", "Paste" ),
+			LOCTEXT( "Paste_ToolTip", "Pastes a flattened representation of the cut/copied prims as children of this prim, on the current edit target" )
+		);
+
+		const bool bInOpenSubMenuOnClick = false;
+		PrimOptions.AddSubMenu(
+			LOCTEXT( "Duplicate_Text", "Duplicate..." ),
+			FText::GetEmpty(),
+			FNewMenuDelegate::CreateSP( this, &SUsdStageTreeView::FillDuplicateSubmenu ),
+			bInOpenSubMenuOnClick,
+			FSlateIcon( FAppStyle::GetAppStyleSetName(), "Icons.Duplicate" )
+		);
+
+		PrimOptions.AddMenuEntry(
+			FGenericCommands::Get().Delete,
+			NAME_None,
+			LOCTEXT( "Delete_Text", "Delete" ),
+			LOCTEXT( "Delete_ToolTip", "Deletes the selected prims" )
+		);
+
+		PrimOptions.AddMenuEntry(
+			FGenericCommands::Get().Rename,
+			NAME_None,
+			LOCTEXT( "Rename_Text", "Rename" ),
+			LOCTEXT( "Rename_ToolTip", "Renames the selected prim" )
 		);
 	}
 	PrimOptions.EndSection();
@@ -645,11 +711,11 @@ TSharedPtr< SWidget > SUsdStageTreeView::ConstructPrimContextMenu()
 	{
 		PrimOptions.AddMenuEntry(
 			LOCTEXT("AddReference", "Add Reference"),
-			LOCTEXT("AddReference_ToolTip", "Adds a reference to this prim"),
+			LOCTEXT("AddReference_ToolTip", "Adds a reference for this prim"),
 			FSlateIcon(),
 			FUIAction(
 				FExecuteAction::CreateSP( this, &SUsdStageTreeView::OnAddReference ),
-				FCanExecuteAction::CreateSP( this, &SUsdStageTreeView::CanExecutePrimAction )
+				FCanExecuteAction::CreateLambda( [this]() { return SUsdStageTreeView::DoesPrimExistOnEditTarget() && GetSelectedItems().Num() == 1; } )
 			),
 			NAME_None,
 			EUserInterfaceActionType::Button
@@ -661,7 +727,31 @@ TSharedPtr< SWidget > SUsdStageTreeView::ConstructPrimContextMenu()
 			FSlateIcon(),
 			FUIAction(
 				FExecuteAction::CreateSP( this, &SUsdStageTreeView::OnClearReferences ),
-				FCanExecuteAction::CreateSP( this, &SUsdStageTreeView::CanExecutePrimAction )
+				FCanExecuteAction::CreateSP( this, &SUsdStageTreeView::DoesPrimExistOnEditTarget )
+			),
+			NAME_None,
+			EUserInterfaceActionType::Button
+		);
+
+		PrimOptions.AddMenuEntry(
+			LOCTEXT( "AddPayload", "Add Payload" ),
+			LOCTEXT( "AddPayload_ToolTip", "Adds a payload for this prim" ),
+			FSlateIcon(),
+			FUIAction(
+				FExecuteAction::CreateSP( this, &SUsdStageTreeView::OnAddPayload ),
+				FCanExecuteAction::CreateLambda( [this]() { return SUsdStageTreeView::DoesPrimExistOnEditTarget() && GetSelectedItems().Num() == 1; } )
+			),
+			NAME_None,
+			EUserInterfaceActionType::Button
+		);
+
+		PrimOptions.AddMenuEntry(
+			LOCTEXT( "ClearPayloads", "Clear Payloads" ),
+			LOCTEXT( "ClearPayloads_ToolTip", "Clears the payloads for this prim" ),
+			FSlateIcon(),
+			FUIAction(
+				FExecuteAction::CreateSP( this, &SUsdStageTreeView::OnClearPayloads ),
+				FCanExecuteAction::CreateSP( this, &SUsdStageTreeView::DoesPrimExistOnEditTarget )
 			),
 			NAME_None,
 			EUserInterfaceActionType::Button
@@ -755,7 +845,7 @@ TSharedPtr< SWidget > SUsdStageTreeView::ConstructPrimContextMenu()
 	return MenuWidget;
 }
 
-void SUsdStageTreeView::OnAddPrim()
+void SUsdStageTreeView::OnAddChildPrim()
 {
 	if ( !UsdStageActor.IsValid() )
 	{
@@ -789,44 +879,169 @@ void SUsdStageTreeView::OnAddPrim()
 	RequestTreeRefresh();
 }
 
-void SUsdStageTreeView::OnRenamePrim()
+void SUsdStageTreeView::OnCutPrim()
 {
 	if ( !UsdStageActor.IsValid() )
 	{
 		return;
 	}
 
+	FScopedTransaction Transaction( LOCTEXT( "CutPrimTransaction", "Cut prims" ) );
+
+	UE::FSdfChangeBlock Block;
+
 	TArray< FUsdPrimViewModelRef > MySelectedItems = GetSelectedItems();
 
-	if ( MySelectedItems.Num() > 0 )
+	TArray<UE::FUsdPrim> Prims;
+	Prims.Reserve( MySelectedItems.Num() );
+
+	for ( FUsdPrimViewModelRef SelectedItem : MySelectedItems )
 	{
-		FUsdPrimViewModelRef TreeItem = MySelectedItems[ 0 ];
+		if ( SelectedItem->UsdPrim )
+		{
+			Prims.Add( SelectedItem->UsdPrim );
+		}
+	}
 
-		TreeItem->bIsRenamingExistingPrim = true;
-		PendingRenameItem = TreeItem;
-		RequestScrollIntoView( TreeItem );
+	UsdUtils::CutPrims( Prims );
+}
+
+void SUsdStageTreeView::OnCopyPrim()
+{
+	if ( !UsdStageActor.IsValid() )
+	{
+		return;
 	}
+
+	TArray< FUsdPrimViewModelRef > MySelectedItems = GetSelectedItems();
+
+	TArray<UE::FUsdPrim> Prims;
+	Prims.Reserve( MySelectedItems.Num() );
+
+	for ( FUsdPrimViewModelRef SelectedItem : MySelectedItems )
+	{
+		if ( SelectedItem->UsdPrim )
+		{
+			Prims.Add( SelectedItem->UsdPrim );
+		}
+	}
+
+	UsdUtils::CopyPrims( Prims );
 }
 
-void SUsdStageTreeView::OnRemovePrim()
+void SUsdStageTreeView::OnPastePrim()
 {
 	if ( !UsdStageActor.IsValid() )
 	{
 		return;
 	}
 
-	FScopedTransaction Transaction( LOCTEXT( "RemovePrimTransaction", "Remove prims'" ) );
+	FScopedTransaction Transaction( LOCTEXT( "PastePrimTransaction", "Paste prims" ) );
+
+	UE::FSdfChangeBlock Block;
 
 	TArray< FUsdPrimViewModelRef > MySelectedItems = GetSelectedItems();
 
-	for ( FUsdPrimViewModelRef SelectedItem : MySelectedItems )
+	TArray<UE::FUsdPrim> ParentPrims;
+
+	// This happens when right-clicking the background area without selecting any prim
+	if ( MySelectedItems.IsEmpty() )
 	{
 		UE::FUsdStage Stage = UsdStageActor->GetOrLoadUsdStage();
+		ParentPrims.Add( Stage.GetPseudoRoot() );
+	}
+	else
+	{
+		ParentPrims.Reserve( MySelectedItems.Num() );
+
+		// A bit unusual that we can paste to multiple locations at the same time, but why not?
+		for ( FUsdPrimViewModelRef SelectedItem : MySelectedItems )
+		{
+			ParentPrims.Add( SelectedItem->UsdPrim );
+		}
+	}
+
+	for ( const UE::FUsdPrim& ParentPrim : ParentPrims )
+	{
+		// Preemptively mark the parent prims as expanded so that we can always see what we pasted
+		const bool bIsExpanded = true;
+		TreeItemExpansionStates.Add( ParentPrim.GetPrimPath().GetString(), bIsExpanded );
+
+		UsdUtils::PastePrims( ParentPrim );
+	}
+}
+
+void SUsdStageTreeView::OnDuplicatePrim( EUsdDuplicateType DuplicateType )
+{
+	if ( !UsdStageActor.IsValid() )
+	{
+		return;
+	}
+
+	UE::FUsdStage Stage = UsdStageActor->GetOrLoadUsdStage();
+	if ( !Stage )
+	{
+		return;
+	}
+
+	FScopedTransaction Transaction( LOCTEXT( "DuplicatePrimTransaction", "Duplicate prims" ) );
+
+	UE::FSdfChangeBlock Block;
+
+	TArray< FUsdPrimViewModelRef > MySelectedItems = GetSelectedItems();
+
+	TArray<UE::FUsdPrim> Prims;
+	Prims.Reserve( MySelectedItems.Num() );
+
+	for ( FUsdPrimViewModelRef SelectedItem : MySelectedItems )
+	{
+		if ( SelectedItem->UsdPrim )
+		{
+			Prims.Add( SelectedItem->UsdPrim );
+		}
+	}
+
+	UsdUtils::DuplicatePrims( Prims, DuplicateType, Stage.GetEditTarget() );
+}
+
+void SUsdStageTreeView::OnDeletePrim()
+{
+	if ( !UsdStageActor.IsValid() )
+	{
+		return;
+	}
+
+	FScopedTransaction Transaction( LOCTEXT( "DeletePrimTransaction", "Delete prims" ) );
+
+	UE::FSdfChangeBlock Block;
+
+	TArray< FUsdPrimViewModelRef > MySelectedItems = GetSelectedItems();
 
+	for ( FUsdPrimViewModelRef SelectedItem : MySelectedItems )
+	{
 		UsdUtils::RemoveAllPrimSpecs( SelectedItem->UsdPrim );
 	}
 }
 
+void SUsdStageTreeView::OnRenamePrim()
+{
+	if ( !UsdStageActor.IsValid() )
+	{
+		return;
+	}
+
+	TArray< FUsdPrimViewModelRef > MySelectedItems = GetSelectedItems();
+
+	if ( MySelectedItems.Num() > 0 )
+	{
+		FUsdPrimViewModelRef TreeItem = MySelectedItems[ 0 ];
+
+		TreeItem->bIsRenamingExistingPrim = true;
+		PendingRenameItem = TreeItem;
+		RequestScrollIntoView( TreeItem );
+	}
+}
+
 void SUsdStageTreeView::OnAddReference()
 {
 	UE::FUsdStage& Stage = UsdStageActor->GetOrLoadUsdStage();
@@ -835,42 +1050,125 @@ void SUsdStageTreeView::OnAddReference()
 		return;
 	}
 
-	TOptional< FString > PickedFile = UsdUtils::BrowseUsdFile( UsdUtils::EBrowseFileMode::Composition, AsShared() );
+	TStrongObjectPtr<UUsdReferenceOptions> Options = TStrongObjectPtr<UUsdReferenceOptions>{ NewObject<UUsdReferenceOptions>() };
+	UUsdReferenceOptions* OptionsPtr = Options.Get();
+	if ( !OptionsPtr )
+	{
+		return;
+	}
 
-	if ( !PickedFile )
+	bool bContinue = SUsdOptionsWindow::ShowOptions(
+		*OptionsPtr,
+		LOCTEXT( "AddReferenceTitle", "Add reference" ),
+		LOCTEXT( "AddReferenceAccept", "OK" )
+	);
+	if ( !bContinue )
 	{
 		return;
 	}
 
+	TArray< FUsdPrimViewModelRef > MySelectedItems = GetSelectedItems();
+	if ( MySelectedItems.Num() != 1 )
+	{
+		return;
+	}
+	UE::FUsdPrim Referencer = MySelectedItems[ 0 ]->UsdPrim;
+
 	// This transaction is important as adding a reference may trigger the creation of new unreal assets, which need to be
 	// destroyed if we spam undo afterwards. Undoing won't remove the actual reference from the stage yet though, sadly...
 	FScopedTransaction Transaction( FText::Format(
-		LOCTEXT( "AddReferenceTransaction", "Add reference to file '{0}'" ),
-		FText::FromString( PickedFile.GetValue() )
+		LOCTEXT( "AddReferenceTransaction", "Add reference from prim '{0}'" ),
+		FText::FromString( Referencer.GetPrimPath().GetString() )
 	) );
 
+	UsdUtils::AddReference(
+		Referencer,
+		OptionsPtr->bInternalReference ? TEXT("") : *OptionsPtr->TargetFile.FilePath,
+		OptionsPtr->bUseDefaultPrim ? UE::FSdfPath{} : UE::FSdfPath{ *OptionsPtr->TargetPrimPath },
+		OptionsPtr->TimeCodeOffset,
+		OptionsPtr->TimeCodeScale
+	);
+}
+
+void SUsdStageTreeView::OnClearReferences()
+{
+	if ( !UsdStageActor.IsValid() )
+	{
+		return;
+	}
+
+	FScopedTransaction Transaction( LOCTEXT( "ClearReferenceTransaction", "Clear references to USD layers" ) );
+
 	TArray< FUsdPrimViewModelRef > MySelectedItems = GetSelectedItems();
 
 	for ( FUsdPrimViewModelRef SelectedItem : MySelectedItems )
 	{
-		SelectedItem->AddReference( *PickedFile.GetValue() );
+		SelectedItem->ClearReferences();
 	}
 }
 
-void SUsdStageTreeView::OnClearReferences()
+void SUsdStageTreeView::OnAddPayload()
+{
+	UE::FUsdStage& Stage = UsdStageActor->GetOrLoadUsdStage();
+	if ( !UsdStageActor.IsValid() || !Stage || !Stage.IsEditTargetValid() )
+	{
+		return;
+	}
+
+	TStrongObjectPtr<UUsdReferenceOptions> Options = TStrongObjectPtr<UUsdReferenceOptions>{ NewObject<UUsdReferenceOptions>() };
+	UUsdReferenceOptions* OptionsPtr = Options.Get();
+	if ( !OptionsPtr )
+	{
+		return;
+	}
+
+	bool bContinue = SUsdOptionsWindow::ShowOptions(
+		*OptionsPtr,
+		LOCTEXT( "AddPayloadTitle", "Add payload" ),
+		LOCTEXT( "AddPayloadAccept", "OK" )
+	);
+	if ( !bContinue )
+	{
+		return;
+	}
+
+	TArray< FUsdPrimViewModelRef > MySelectedItems = GetSelectedItems();
+	if ( MySelectedItems.Num() != 1 )
+	{
+		return;
+	}
+	UE::FUsdPrim Referencer = MySelectedItems[ 0 ]->UsdPrim;
+
+	// This transaction is important as adding a payload may trigger the creation of new unreal assets, which need to be
+	// destroyed if we spam undo afterwards. Undoing won't remove the actual payload from the stage yet though, sadly...
+	FScopedTransaction Transaction( FText::Format(
+		LOCTEXT( "AddPayloadTransaction", "Add payload from prim '{0}'" ),
+		FText::FromString( Referencer.GetPrimPath().GetString() )
+	) );
+
+	UsdUtils::AddPayload(
+		Referencer,
+		OptionsPtr->bInternalReference ? TEXT( "" ) : *OptionsPtr->TargetFile.FilePath,
+		OptionsPtr->bUseDefaultPrim ? UE::FSdfPath{} : UE::FSdfPath{ *OptionsPtr->TargetPrimPath },
+		OptionsPtr->TimeCodeOffset,
+		OptionsPtr->TimeCodeScale
+	);
+}
+
+void SUsdStageTreeView::OnClearPayloads()
 {
 	if ( !UsdStageActor.IsValid() )
 	{
 		return;
 	}
 
-	FScopedTransaction Transaction( LOCTEXT( "ClearReferenceTransaction", "Clear references to USD layers" ) );
+	FScopedTransaction Transaction( LOCTEXT( "ClearPayloadTransaction", "Clear payloads to USD layers" ) );
 
 	TArray< FUsdPrimViewModelRef > MySelectedItems = GetSelectedItems();
 
 	for ( FUsdPrimViewModelRef SelectedItem : MySelectedItems )
 	{
-		SelectedItem->ClearReferences();
+		SelectedItem->ClearPayloads();
 	}
 }
 
@@ -988,7 +1286,7 @@ void SUsdStageTreeView::OnRemoveGroomSchema()
 	}
 }
 
-bool SUsdStageTreeView::CanAddPrim() const
+bool SUsdStageTreeView::CanAddChildPrim() const
 {
 	if ( !UsdStageActor.IsValid() )
 	{
@@ -1001,37 +1299,99 @@ bool SUsdStageTreeView::CanAddPrim() const
 		return false;
 	}
 
+	TArray< FUsdPrimViewModelRef > MySelectedItems = GetSelectedItems();
+
+	// Allow adding a new top-level prim
+	if ( MySelectedItems.IsEmpty() )
+	{
+		return true;
+	}
+	else
+	{
+		// We use the "rename" text input workflow to specify the target name,
+		// so this doesn't work very well for multiple prims yet
+		if ( MySelectedItems.Num() > 1 )
+		{
+			return false;
+		}
+
+		// If we have something selected it must be valid
+		if ( !MySelectedItems[ 0 ]->UsdPrim.IsValid() )
+		{
+			return false;
+		}
+	}
+
 	return true;
 }
 
-bool SUsdStageTreeView::CanExecutePrimAction() const
+bool SUsdStageTreeView::CanPastePrim() const
 {
 	if ( !UsdStageActor.IsValid() )
 	{
 		return false;
 	}
 
-	UE::FUsdStage UsdStage =  UsdStageActor->GetOrLoadUsdStage();
-	if ( !UsdStage || !UsdStage.IsEditTargetValid() )
+	UE::FUsdStage UsdStage = UsdStageActor->GetOrLoadUsdStage();
+	if ( !UsdStage )
 	{
 		return false;
 	}
 
-	bool bHasPrimSpec = false;
+	return UsdUtils::CanPastePrims();
+}
+
+bool SUsdStageTreeView::DoesPrimExistOnStage() const
+{
+	if ( !UsdStageActor.IsValid() )
+	{
+		return false;
+	}
+
+	UE::FUsdStage UsdStage = UsdStageActor->GetOrLoadUsdStage();
+	if ( !UsdStage || !UsdStage.IsEditTargetValid() )
+	{
+		return false;
+	}
 
 	TArray< FUsdPrimViewModelRef > MySelectedItems = GetSelectedItems();
 
 	for ( FUsdPrimViewModelRef SelectedItem : MySelectedItems )
 	{
-		bHasPrimSpec = bHasPrimSpec || (bool)SelectedItem->CanExecutePrimAction();
+		if ( !SelectedItem->UsdPrim.IsPseudoRoot() && SelectedItem->UsdPrim.IsValid() )
+		{
+			return true;
+		}
+	}
+
+	return false;
+}
+
+bool SUsdStageTreeView::DoesPrimExistOnEditTarget() const
+{
+	if ( !UsdStageActor.IsValid() )
+	{
+		return false;
+	}
+
+	UE::FUsdStage UsdStage =  UsdStageActor->GetOrLoadUsdStage();
+	if ( !UsdStage || !UsdStage.IsEditTargetValid() )
+	{
+		return false;
+	}
 
-		if ( bHasPrimSpec )
+	TArray< FUsdPrimViewModelRef > MySelectedItems = GetSelectedItems();
+
+	for ( FUsdPrimViewModelRef SelectedItem : MySelectedItems )
+	{
+		UE::FSdfPath SpecPath = UsdUtils::GetPrimSpecPathForLayer( SelectedItem->UsdPrim, UsdStage.GetEditTarget() );
+		if ( !SpecPath.IsAbsoluteRootPath() && !SpecPath.IsEmpty() )
 		{
-			break;
+			return true;
 		}
 	}
 
-	return bHasPrimSpec;
+	return false;
 }
 
 bool SUsdStageTreeView::CanSetUpLiveLink() const
@@ -1296,6 +1656,55 @@ void SUsdStageTreeView::OnToggleAllPayloads( EPayloadsTrigger PayloadsTrigger )
 	}
 }
 
+void SUsdStageTreeView::FillDuplicateSubmenu( FMenuBuilder& MenuBuilder )
+{
+	MenuBuilder.AddMenuEntry(
+		LOCTEXT( "DuplicateFlattened_Text", "Flatten composed prim" ),
+		LOCTEXT( "DuplicateFlattened_ToolTip", "Generate a flattened duplicate of the composed prim onto the current edit target" ),
+		FSlateIcon(),
+		FUIAction(
+			FExecuteAction::CreateSP( this, &SUsdStageTreeView::OnDuplicatePrim, EUsdDuplicateType::FlattenComposedPrim ),
+			FCanExecuteAction::CreateSP( this, &SUsdStageTreeView::DoesPrimExistOnStage )
+		),
+		NAME_None,
+		EUserInterfaceActionType::Button
+	);
+
+	MenuBuilder.AddMenuEntry(
+		LOCTEXT( "DuplicateSingle_Text", "Single layer specs" ),
+		LOCTEXT( "DuplicateSingle_ToolTip", "Duplicate the prim's specs on the current edit target only" ),
+		FSlateIcon(),
+		FUIAction(
+			FExecuteAction::CreateSP( this, &SUsdStageTreeView::OnDuplicatePrim, EUsdDuplicateType::SingleLayerSpecs ),
+			FCanExecuteAction::CreateSP( this, &SUsdStageTreeView::DoesPrimExistOnStage )
+		),
+		NAME_None,
+		EUserInterfaceActionType::Button
+	);
+
+	MenuBuilder.AddMenuEntry(
+		LOCTEXT( "DuplicateAllLocal_Text", "All local layer specs" ),
+		LOCTEXT( "DuplicateAllLocal_ToolTip", "Duplicate each of the prim's specs across the entire stage" ),
+		FSlateIcon(),
+		FUIAction(
+			FExecuteAction::CreateSP( this, &SUsdStageTreeView::OnDuplicatePrim, EUsdDuplicateType::AllLocalLayerSpecs ),
+			FCanExecuteAction::CreateSP( this, &SUsdStageTreeView::DoesPrimExistOnStage )
+		),
+		NAME_None,
+		EUserInterfaceActionType::Button
+	);
+}
+
+FReply SUsdStageTreeView::OnKeyDown( const FGeometry& MyGeometry, const FKeyEvent& InKeyEvent )
+{
+	if ( UICommandList->ProcessCommandBindings( InKeyEvent ) )
+	{
+		return FReply::Handled();
+	}
+
+	return SUsdTreeView::OnKeyDown( MyGeometry, InKeyEvent );
+}
+
 void SUsdStageTreeView::ScrollItemIntoView( FUsdPrimViewModelRef TreeItem )
 {
 	FUsdPrimViewModel* Parent = TreeItem->ParentItem;
diff --git a/Engine/Plugins/Importers/USDImporter/Source/USDStageEditor/Private/SUSDStageTreeView.h b/Engine/Plugins/Importers/USDImporter/Source/USDStageEditor/Private/SUSDStageTreeView.h
--- a/Engine/Plugins/Importers/USDImporter/Source/USDStageEditor/Private/SUSDStageTreeView.h
+++ b/Engine/Plugins/Importers/USDImporter/Source/USDStageEditor/Private/SUSDStageTreeView.h
@@ -11,7 +11,9 @@
 #include "USDPrimViewModel.h"
 
 class AUsdStageActor;
+class FUICommandList;
 enum class EPayloadsTrigger;
+enum class EUsdDuplicateType : uint8;
 
 #if USE_USD_SDK
 
@@ -38,6 +40,10 @@ private:
 	virtual TSharedRef< ITableRow > OnGenerateRow( FUsdPrimViewModelRef InDisplayNode, const TSharedRef< STableViewBase >& OwnerTable ) override;
 	virtual void OnGetChildren( FUsdPrimViewModelRef InParent, TArray< FUsdPrimViewModelRef >& OutChildren ) const override;
 
+	// Required so that we can use the cut/copy/paste/etc. shortcuts
+	virtual bool SupportsKeyboardFocus() const override { return true; }
+	virtual FReply OnKeyDown( const FGeometry& MyGeometry, const FKeyEvent& InKeyEvent ) override;
+
 	void ScrollItemIntoView( FUsdPrimViewModelRef TreeItem );
 	virtual void OnTreeItemScrolledIntoView( FUsdPrimViewModelRef TreeItem, const TSharedPtr< ITableRow >& Widget ) override ;
 
@@ -49,11 +55,22 @@ private:
 
 	void OnToggleAllPayloads( EPayloadsTrigger PayloadsTrigger );
 
-	void OnAddPrim();
+	void FillDuplicateSubmenu( FMenuBuilder& MenuBuilder );
+
+	void OnAddChildPrim();
+	void OnCutPrim();
+	void OnCopyPrim();
+	void OnPastePrim();
+	void OnDuplicatePrim( EUsdDuplicateType DuplicateType );
+	void OnDeletePrim();
 	void OnRenamePrim();
-	void OnRemovePrim();
+
 	void OnAddReference();
 	void OnClearReferences();
+
+	void OnAddPayload();
+	void OnClearPayloads();
+
 	void OnSetUpLiveLink();
 	void OnRemoveLiveLink();
 	void OnSetUpControlRig();
@@ -61,8 +78,11 @@ private:
 	void OnApplyGroomSchema();
 	void OnRemoveGroomSchema();
 
-	bool CanAddPrim() const;
-	bool CanExecutePrimAction() const;
+	bool CanAddChildPrim() const;
+	bool CanPastePrim() const;
+	bool DoesPrimExistOnStage() const;
+	bool DoesPrimExistOnEditTarget() const;
+
 	bool CanSetUpLiveLink() const;
 	bool CanRemoveLiveLink() const;
 	bool CanSetUpControlRig() const;
@@ -82,6 +102,8 @@ private:
 	TMap< FString, bool > TreeItemExpansionStates;
 
 	FOnPrimSelectionChanged OnPrimSelectionChanged;
+
+	TSharedPtr<FUICommandList> UICommandList;
 };
 
 #endif // #if USE_USD_SDK
diff --git a/Engine/Plugins/Importers/USDImporter/Source/USDStageEditorViewModels/Private/USDPrimViewModel.cpp b/Engine/Plugins/Importers/USDImporter/Source/USDStageEditorViewModels/Private/USDPrimViewModel.cpp
--- a/Engine/Plugins/Importers/USDImporter/Source/USDStageEditorViewModels/Private/USDPrimViewModel.cpp
+++ b/Engine/Plugins/Importers/USDImporter/Source/USDStageEditorViewModels/Private/USDPrimViewModel.cpp
@@ -23,6 +23,7 @@
 	#include "pxr/usd/sdf/namespaceEdit.h"
 	#include "pxr/usd/sdf/path.h"
 	#include "pxr/usd/usd/modelAPI.h"
+	#include "pxr/usd/usd/payloads.h"
 	#include "pxr/usd/usd/prim.h"
 	#include "pxr/usd/usd/references.h"
 	#include "pxr/usd/usd/tokens.h"
@@ -254,25 +255,6 @@ void FUsdPrimViewModel::RefreshData( bool bRefreshChildren )
 #endif // #if USE_USD_SDK
 }
 
-bool FUsdPrimViewModel::CanExecutePrimAction() const
-{
-#if USE_USD_SDK
-	if ( UsdPrim.IsPseudoRoot() )
-	{
-		return false;
-	}
-
-	UE::FSdfPath SpecPath = UsdUtils::GetPrimSpecPathForLayer( UsdPrim, UsdStage.GetEditTarget() );
-	if ( !SpecPath.IsEmpty() )
-	{
-		return true;
-	}
-
-#endif // #if USE_USD_SDK
-
-	return false;
-}
-
 bool FUsdPrimViewModel::HasVisibilityAttribute() const
 {
 #if USE_USD_SDK
@@ -467,11 +449,6 @@ void FUsdPrimViewModel::DefinePrim( const TCHAR* PrimName )
 #endif // #if USE_USD_SDK
 }
 
-void FUsdPrimViewModel::AddReference( const TCHAR* AbsoluteFilePath )
-{
-	UsdUtils::AddReference( UsdPrim, AbsoluteFilePath );
-}
-
 void FUsdPrimViewModel::ClearReferences()
 {
 #if USE_USD_SDK
@@ -481,3 +458,13 @@ void FUsdPrimViewModel::ClearReferences()
 	References.ClearReferences();
 #endif // #if USE_USD_SDK
 }
+
+void FUsdPrimViewModel::ClearPayloads()
+{
+#if USE_USD_SDK
+	FScopedUsdAllocs UsdAllocs;
+
+	pxr::UsdPayloads Payloads = pxr::UsdPrim( UsdPrim ).GetPayloads();
+	Payloads.ClearPayloads();
+#endif // #if USE_USD_SDK
+}
diff --git a/Engine/Plugins/Importers/USDImporter/Source/USDStageEditorViewModels/Public/USDPrimViewModel.h b/Engine/Plugins/Importers/USDImporter/Source/USDStageEditorViewModels/Public/USDPrimViewModel.h
--- a/Engine/Plugins/Importers/USDImporter/Source/USDStageEditorViewModels/Public/USDPrimViewModel.h
+++ b/Engine/Plugins/Importers/USDImporter/Source/USDStageEditorViewModels/Public/USDPrimViewModel.h
@@ -43,7 +43,6 @@ public:
 
 	void RefreshData( bool bRefreshChildren );
 
-	bool CanExecutePrimAction() const;
 	bool HasVisibilityAttribute() const;
 	void ToggleVisibility();
 	void TogglePayload();
@@ -75,8 +74,8 @@ public:
 	FOnRenameRequest RenameRequestEvent;
 
 public:
-	void AddReference( const TCHAR* AbsoluteFilePath );
 	void ClearReferences();
+	void ClearPayloads();
 
 public:
 	UE::FUsdStageWeak UsdStage;
diff --git a/Engine/Plugins/Importers/USDImporter/Source/USDStageImporter/Private/USDOptionsWindow.cpp b/Engine/Plugins/Importers/USDImporter/Source/USDStageImporter/Private/USDOptionsWindow.cpp
--- a/Engine/Plugins/Importers/USDImporter/Source/USDStageImporter/Private/USDOptionsWindow.cpp
+++ b/Engine/Plugins/Importers/USDImporter/Source/USDStageImporter/Private/USDOptionsWindow.cpp
@@ -16,7 +16,27 @@
 
 #define LOCTEXT_NAMESPACE "USDOptionsWindow"
 
-bool SUsdOptionsWindow::ShowOptions( UObject& OptionsObject, bool bIsImport )
+bool SUsdOptionsWindow::ShowImportExportOptions( UObject& OptionsObject, bool bIsImport )
+{
+	if ( bIsImport )
+	{
+		return SUsdOptionsWindow::ShowOptions(
+			OptionsObject,
+			LOCTEXT( "USDImportOptionsTitle", "USD Import Options" ),
+			LOCTEXT( "USDOptionWindow_Import", "Import" )
+		);
+	}
+	else
+	{
+		return SUsdOptionsWindow::ShowOptions(
+			OptionsObject,
+			LOCTEXT( "USDExportOptionsTitle", "USD Export Options" ),
+			LOCTEXT( "USDOptionWindow_Export", "Export" )
+		);
+	}
+}
+
+bool SUsdOptionsWindow::ShowOptions( UObject& OptionsObject, const FText& WindowTitle, const FText& AcceptText )
 {
 	TSharedPtr<SWindow> ParentWindow;
 
@@ -27,10 +47,7 @@ bool SUsdOptionsWindow::ShowOptions( UObject& OptionsObject, bool bIsImport )
 	}
 
 	TSharedRef<SWindow> Window = SNew( SWindow )
-		.Title( bIsImport
-			? LOCTEXT( "USDImportOptionsTitle", "USD Import Options" )
-			: LOCTEXT( "USDExportOptionsTitle", "USD Export Options" )
-		)
+		.Title( WindowTitle )
 		.SizingRule( ESizingRule::Autosized );
 
 	TSharedPtr<SUsdOptionsWindow> OptionsWindow;
@@ -38,7 +55,7 @@ bool SUsdOptionsWindow::ShowOptions( UObject& OptionsObject, bool bIsImport )
 	(
 		SAssignNew( OptionsWindow, SUsdOptionsWindow )
 		.OptionsObject( &OptionsObject )
-		.IsImport( bIsImport )
+		.AcceptText( AcceptText )
 		.WidgetWindow( Window )
 	);
 
@@ -60,7 +77,7 @@ void SUsdOptionsWindow::Construct( const FArguments& InArgs )
 {
 	OptionsObject = InArgs._OptionsObject;
 	Window = InArgs._WidgetWindow;
-	bIsImport = InArgs._IsImport;
+	AcceptText = InArgs._AcceptText;
 	bAccepted = false;
 
 	TSharedPtr<SBox> DetailsViewBox;
@@ -89,10 +106,7 @@ void SUsdOptionsWindow::Construct( const FArguments& InArgs )
 			[
 				SNew( SButton )
 				.HAlign( HAlign_Center )
-				.Text( bIsImport
-					? LOCTEXT( "USDOptionWindow_Import", "Import" )
-					: LOCTEXT( "USDOptionWindow_Export", "Export" )
-				)
+				.Text( AcceptText )
 				.OnClicked( this, &SUsdOptionsWindow::OnAccept )
 			]
 
diff --git a/Engine/Plugins/Importers/USDImporter/Source/USDStageImporter/Private/USDStageImportContext.cpp b/Engine/Plugins/Importers/USDImporter/Source/USDStageImporter/Private/USDStageImportContext.cpp
--- a/Engine/Plugins/Importers/USDImporter/Source/USDStageImporter/Private/USDStageImportContext.cpp
+++ b/Engine/Plugins/Importers/USDImporter/Source/USDStageImporter/Private/USDStageImportContext.cpp
@@ -59,7 +59,8 @@ bool FUsdStageImportContext::Init(const FString& InName, const FString& InFilePa
 		ImportOptions->EnableActorImport( bAllowActorImport );
 
 		// Show dialog for import options
-		bool bProceedWithImport = SUsdOptionsWindow::ShowOptions( *ImportOptions );
+		const bool bIsImport = true;
+		bool bProceedWithImport = SUsdOptionsWindow::ShowImportExportOptions( *ImportOptions, bIsImport );
 		if (!bProceedWithImport)
 		{
 			return false;
diff --git a/Engine/Plugins/Importers/USDImporter/Source/USDStageImporter/Public/USDOptionsWindow.h b/Engine/Plugins/Importers/USDImporter/Source/USDStageImporter/Public/USDOptionsWindow.h
--- a/Engine/Plugins/Importers/USDImporter/Source/USDStageImporter/Public/USDOptionsWindow.h
+++ b/Engine/Plugins/Importers/USDImporter/Source/USDStageImporter/Public/USDOptionsWindow.h
@@ -14,11 +14,12 @@ public:
 	SLATE_BEGIN_ARGS( SUsdOptionsWindow ) : _OptionsObject( nullptr ) {}
 		SLATE_ARGUMENT( UObject*, OptionsObject )
 		SLATE_ARGUMENT( TSharedPtr<SWindow>, WidgetWindow )
-		SLATE_ARGUMENT( bool, IsImport )
+		SLATE_ARGUMENT( FText, AcceptText )
 	SLATE_END_ARGS()
 
 public:
-	static bool ShowOptions( UObject& OptionsObject, bool bIsImport = true );
+	static bool ShowImportExportOptions( UObject& OptionsObject, bool bIsImport );
+	static bool ShowOptions( UObject& OptionsObject, const FText& WindowTitle, const FText& AcceptText );
 
 	void Construct( const FArguments& InArgs );
 	virtual bool SupportsKeyboardFocus() const override;
@@ -32,6 +33,6 @@ public:
 private:
 	UObject* OptionsObject;
 	TWeakPtr< SWindow > Window;
-	bool bIsImport;
+	FText AcceptText;
 	bool bAccepted;
 };
diff --git a/Engine/Plugins/Importers/USDImporter/Source/USDUtilities/Private/USDConversionUtils.cpp b/Engine/Plugins/Importers/USDImporter/Source/USDUtilities/Private/USDConversionUtils.cpp
--- a/Engine/Plugins/Importers/USDImporter/Source/USDUtilities/Private/USDConversionUtils.cpp
+++ b/Engine/Plugins/Importers/USDImporter/Source/USDUtilities/Private/USDConversionUtils.cpp
@@ -3,6 +3,7 @@
 #include "USDConversionUtils.h"
 
 #include "USDAssetImportData.h"
+#include "USDDuplicateType.h"
 #include "USDErrorUtils.h"
 #include "USDGeomMeshConversion.h"
 #include "USDLayerUtils.h"
@@ -49,10 +50,12 @@
 	#include "pxr/base/tf/stringUtils.h"
 	#include "pxr/base/tf/token.h"
 	#include "pxr/usd/kind/registry.h"
+	#include "pxr/usd/sdf/copyUtils.h"
 	#include "pxr/usd/usd/attribute.h"
 	#include "pxr/usd/usd/editContext.h"
 	#include "pxr/usd/usd/modelAPI.h"
 	#include "pxr/usd/usd/payloads.h"
+	#include "pxr/usd/usd/primCompositionQuery.h"
 	#include "pxr/usd/usd/primRange.h"
 	#include "pxr/usd/usd/stage.h"
 	#include "pxr/usd/usd/variantSets.h"
@@ -1096,77 +1099,161 @@ UUsdAssetImportData* UsdUtils::GetAssetImportData( UObject* Asset )
 	return ImportData;
 }
 
-void UsdUtils::AddReference( UE::FUsdPrim& Prim, const TCHAR* AbsoluteFilePath )
+namespace UE::UsdConversionUtils::Private
 {
 #if USE_USD_SDK
-	if ( !Prim || !AbsoluteFilePath )
-	{
-		return;
-	}
+	bool PrepareToAddReferenceOrPayload(
+		const UE::FUsdPrim& Prim,
+		const TCHAR* AbsoluteFilePath,
+		const UE::FSdfPath& TargetPrimPath,
+		std::string& OutRelativeFilePath,
+		pxr::SdfPath& OutPrimPath
+	)
+	{
+		if ( !Prim || !AbsoluteFilePath )
+		{
+			return false;
+		}
 
-	FScopedUsdAllocs UsdAllocs;
+		FScopedUsdAllocs UsdAllocs;
 
-	pxr::UsdPrim UsdPrim( Prim );
+		pxr::UsdPrim UsdPrim( Prim );
 
-	const std::string UsdAbsoluteFilePath = UnrealToUsd::ConvertString( AbsoluteFilePath ).Get();
-	pxr::UsdReferences References = UsdPrim.GetReferences();
+		pxr::UsdStageRefPtr UsdStage = UsdPrim.GetStage();
+		if ( !UsdStage )
+		{
+			return false;
+		}
 
-	pxr::SdfLayerRefPtr ReferenceLayer = pxr::SdfLayer::FindOrOpen( UsdAbsoluteFilePath );
+		// Turn our layer path into a relative one
+		FString RelativePath = AbsoluteFilePath;
+		if ( !RelativePath.IsEmpty() )
+		{
+			pxr::SdfLayerHandle EditLayer = UsdPrim.GetStage()->GetEditTarget().GetLayer();
 
-	// Group updates or else the SetTypeName and AddReference calls below will both trigger separate resyncs of the same prim path
-	pxr::SdfChangeBlock ChangeBlock;
+			std::string RepositoryPath = EditLayer->GetRepositoryPath().empty() ? EditLayer->GetRealPath() : EditLayer->GetRepositoryPath();
 
-	if ( ReferenceLayer && !UsdPrim.GetTypeName().IsEmpty() )
-	{
-		pxr::SdfPrimSpecHandle DefaultPrimSpec = ReferenceLayer->GetPrimAtPath( pxr::SdfPath( ReferenceLayer->GetDefaultPrim() ) );
-		if ( DefaultPrimSpec )
+			// If we're editing an in-memory stage our root layer may not have a path yet
+			// Giving an empty InRelativeTo to MakePathRelativeTo causes it to use the engine binary
+			if ( !RepositoryPath.empty() )
+			{
+				FString LayerAbsolutePath = UsdToUnreal::ConvertString( RepositoryPath );
+				FPaths::MakePathRelativeTo( RelativePath, *LayerAbsolutePath );
+			}
+		}
+
+		// Get the target layer
+		pxr::SdfLayerRefPtr TargetLayer;
+		bool bIsInternalReference = false;
+		if ( RelativePath.IsEmpty() )
+		{
+			TargetLayer = UsdStage->GetRootLayer();
+			bIsInternalReference = true;
+		}
+		else
+		{
+			TargetLayer = pxr::SdfLayer::FindOrOpen( UnrealToUsd::ConvertString( AbsoluteFilePath ).Get() );
+		}
+		if ( !TargetLayer )
 		{
-			// Set the same prim type as its reference so that they are compatible
-			pxr::TfType DefaultPrimType = pxr::UsdSchemaRegistry::GetTypeFromName( DefaultPrimSpec->GetTypeName() );
-			if ( DefaultPrimType.IsUnknown() )
+			return false;
+		}
+
+		// Get the target prim spec we want to reference
+		pxr::SdfPrimSpecHandle TargetPrimSpec = TargetLayer->GetPrimAtPath( TargetPrimPath );
+		if ( ( TargetPrimPath.IsEmpty() || !TargetPrimSpec ) && TargetLayer->HasDefaultPrim() )
+		{
+			TargetPrimSpec = TargetLayer->GetPrimAtPath( pxr::SdfPath::AbsoluteRootPath().AppendChild( TargetLayer->GetDefaultPrim() ) );
+		}
+
+		// Update UsdPrim's type so that it can handle the reference and be parsed by the proper translator
+		if ( !UsdPrim.GetTypeName().IsEmpty() && TargetPrimSpec )
+		{
+			pxr::TfToken TargetTypeName = TargetPrimSpec->GetTypeName();
+			pxr::TfType TargetPrimType = pxr::UsdSchemaRegistry::GetTypeFromName( TargetTypeName );
+
+			if ( TargetPrimType.IsUnknown() )
 			{
 				UsdPrim.ClearTypeName();
 			}
-			else if ( !UsdPrim.IsA( DefaultPrimType ) )
+			else if ( !UsdPrim.IsA( TargetPrimType ) )
 			{
-				UsdPrim.SetTypeName( DefaultPrimSpec->GetTypeName() );
+				UsdPrim.SetTypeName( TargetTypeName );
 			}
 		}
+
+		// We want to output no path for the prim if we received it as such, even if we already know what the path to the
+		// default prim is, so that the authored reference doesn't actually specify any prim name and just refers to the
+		// default prim by default. Otherwise if the default prim of the layer changed, we wouldn't update to the new prim
+		OutPrimPath = TargetPrimSpec && !TargetPrimPath.IsEmpty() ? TargetPrimSpec->GetPath() : pxr::SdfPath{};
+		OutRelativeFilePath = bIsInternalReference ? std::string() : UnrealToUsd::ConvertString( *RelativePath ).Get();
+		return true;
 	}
+#endif // #if USE_USD_SDK
+}
 
-	FString RelativePath = AbsoluteFilePath;
+void UsdUtils::AddReference( UE::FUsdPrim& Prim, const TCHAR* AbsoluteFilePath, const UE::FSdfPath& TargetPrimPath, double TimeCodeOffset, double TimeCodeScale )
+{
+#if USE_USD_SDK
+	FScopedUsdAllocs UsdAllocs;
 
-	pxr::SdfLayerHandle EditLayer = UsdPrim.GetStage()->GetEditTarget().GetLayer();
+	// Group updates or else the SetTypeName (inside PrepareToAddReferenceOrPayload) and AddReference calls below will
+	// both trigger separate resyncs of the same prim path
+	pxr::SdfChangeBlock ChangeBlock;
 
-	std::string RepositoryPath = EditLayer->GetRepositoryPath().empty() ? EditLayer->GetRealPath() : EditLayer->GetRepositoryPath();
+	std::string RelativeLayerPath;
+	pxr::SdfPath FinalPrimPath;
+	bool bProceed = UE::UsdConversionUtils::Private::PrepareToAddReferenceOrPayload(
+		Prim,
+		AbsoluteFilePath,
+		TargetPrimPath,
+		RelativeLayerPath,
+		FinalPrimPath
+	);
 
-	// If we're editing an in-memory stage our root layer may not have a path yet
-	// Giving an empty InRelativeTo to MakePathRelativeTo causes it to use the engine binary
-	if ( !RepositoryPath.empty() )
+	if( !bProceed )
 	{
-		FString LayerAbsolutePath = UsdToUnreal::ConvertString( RepositoryPath );
-		FPaths::MakePathRelativeTo( RelativePath, *LayerAbsolutePath );
+		return;
 	}
 
-	References.AddReference( UnrealToUsd::ConvertString( *RelativePath ).Get() );
+	pxr::UsdReferences References = pxr::UsdPrim{ Prim }.GetReferences();
+	References.AddReference(
+		RelativeLayerPath,
+		FinalPrimPath,
+		pxr::SdfLayerOffset{ TimeCodeOffset, TimeCodeScale }
+	);
 #endif // #if USE_USD_SDK
 }
 
-void UsdUtils::AddPayload( UE::FUsdPrim& Prim, const TCHAR* AbsoluteFilePath )
+void UsdUtils::AddPayload( UE::FUsdPrim& Prim, const TCHAR* AbsoluteFilePath, const UE::FSdfPath& TargetPrimPath, double TimeCodeOffset, double TimeCodeScale )
 {
 #if USE_USD_SDK
 	FScopedUsdAllocs UsdAllocs;
 
-	pxr::UsdPrim UsdPrim( Prim );
+	// Group updates or else the SetTypeName (inside PrepareToAddReferenceOrPayload) and AddReference calls below will
+	// both trigger separate resyncs of the same prim path
+	pxr::SdfChangeBlock ChangeBlock;
 
-	pxr::SdfLayerHandle EditLayer = UsdPrim.GetStage()->GetEditTarget().GetLayer();
+	std::string RelativeLayerPath;
+	pxr::SdfPath FinalPrimPath;
+	bool bProceed = UE::UsdConversionUtils::Private::PrepareToAddReferenceOrPayload(
+		Prim,
+		AbsoluteFilePath,
+		TargetPrimPath,
+		RelativeLayerPath,
+		FinalPrimPath
+	);
 
-	FString RelativePath = AbsoluteFilePath;
-	MakePathRelativeToLayer( UE::FSdfLayer( EditLayer ), RelativePath );
+	if ( !bProceed )
+	{
+		return;
+	}
 
-	pxr::UsdPayloads Payloads = UsdPrim.GetPayloads();
+	pxr::UsdPayloads Payloads = pxr::UsdPrim{ Prim }.GetPayloads();
 	Payloads.AddPayload(
-		UnrealToUsd::ConvertString( *RelativePath ).Get()
+		RelativeLayerPath,
+		FinalPrimPath,
+		pxr::SdfLayerOffset{ TimeCodeOffset, TimeCodeScale }
 	);
 #endif // #if USE_USD_SDK
 }
@@ -1527,6 +1614,10 @@ UE::FSdfPath UsdUtils::GetPrimSpecPathForLayer( const UE::FUsdPrim& Prim, const
 
 	pxr::UsdPrim UsdPrim{ Prim };
 	pxr::SdfLayerRefPtr UsdLayer{ Layer };
+	if ( !UsdPrim || !UsdLayer )
+	{
+		return Result;
+	}
 
 	// We may have multiple specs in the same layer if we're within a variant set (e.g "/Root/Parent/Child" and
 	// "/Root{Varset=Var}Parent/Child{ChildSet=ChildVar}" and "/Root{Varset=Var}Parent/Child").
@@ -1617,6 +1708,485 @@ USDUTILITIES_API void UsdUtils::RemoveAllPrimSpecs( const UE::FUsdPrim& Prim, co
 #endif // USE_USD_SDK
 }
 
+bool UsdUtils::CutPrims( const TArray<UE::FUsdPrim>& Prims )
+{
+	bool bCopied = UsdUtils::CopyPrims( Prims );
+	if ( !bCopied )
+	{
+		return false;
+	}
+
+	for ( const UE::FUsdPrim& Prim : Prims )
+	{
+		UsdUtils::RemoveAllPrimSpecs( Prim );
+	}
+
+	return true;
+}
+
+bool UsdUtils::CopyPrims( const TArray<UE::FUsdPrim>& Prims )
+{
+	bool bCopiedSomething = false;
+
+#if USE_USD_SDK
+	FScopedUsdAllocs Allocs;
+
+	pxr::UsdStageRefPtr UsdStage;
+	for ( const UE::FUsdPrim& Prim : Prims )
+	{
+		if ( Prim )
+		{
+			UsdStage = pxr::UsdStageRefPtr{ Prim.GetStage() };
+			if ( UsdStage )
+			{
+				break;
+			}
+		}
+	}
+	if ( !UsdStage )
+	{
+		return false;
+	}
+
+	pxr::UsdStageRefPtr ClipboardStage = pxr::UsdStageRefPtr{ UnrealUSDWrapper::GetClipboardStage() };
+	if ( !ClipboardStage )
+	{
+		return false;
+	}
+
+	pxr::SdfLayerHandle ClipboardRoot = ClipboardStage->GetRootLayer();
+	if ( !ClipboardRoot )
+	{
+		return false;
+	}
+
+	pxr::UsdStagePopulationMask Mask;
+	for ( const UE::FUsdPrim& Prim : Prims )
+	{
+		if ( Prim )
+		{
+			Mask.Add( pxr::SdfPath{ Prim.GetPrimPath() } );
+		}
+	}
+	if ( Mask.IsEmpty() )
+	{
+		return false;
+	}
+
+	pxr::UsdStageRefPtr TempStage = pxr::UsdStage::OpenMasked( UsdStage->GetRootLayer(), Mask );
+	if ( !TempStage )
+	{
+		return false;
+	}
+
+	const bool bAddSourceFileComment = false;
+	pxr::SdfLayerRefPtr FlattenedLayer = TempStage->Flatten( bAddSourceFileComment );
+	if ( !FlattenedLayer )
+	{
+		return false;
+	}
+
+	ClipboardRoot->Clear();
+
+	TSet<FString> UsedNames;
+
+	for ( const UE::FUsdPrim& Prim : Prims )
+	{
+		pxr::SdfPrimSpecHandle FlattenedPrim = FlattenedLayer->GetPrimAtPath( pxr::SdfPath{ Prim.GetPrimPath() } );
+		if ( !FlattenedPrim )
+		{
+			continue;
+		}
+
+		// Have to ensure the selected prims can coexist as siblings on the clipboard until being pasted.
+		// Note how we don't use GetValidChildName here: That should work too, but it could fail if somebody ever
+		// calls this function within a SdfChangeBlock, given that GetValidChildName relies on USD's GetChildren,
+		// which could potentially yield stale results until USD actually emits the notices about these prims being
+		// added.
+		FString PrimName = Prim.GetName().ToString();
+		FString UniqueName = GetUniqueName( SanitizeUsdIdentifier( *PrimName ), UsedNames );
+		UsedNames.Add( UniqueName );
+
+		const bool bSuccess = pxr::SdfCopySpec(
+			FlattenedLayer,
+			FlattenedPrim->GetPath(),
+			ClipboardRoot,
+			pxr::SdfPath::AbsoluteRootPath().AppendChild( UnrealToUsd::ConvertToken( *UniqueName ).Get() )
+		);
+		if ( !bSuccess )
+		{
+			continue;
+		}
+
+		bCopiedSomething = true;
+		UE_LOG( LogUsd, Log, TEXT( "Copied prim '%s' into the clipboard" ),
+			*Prim.GetPrimPath().GetString()
+		);
+	}
+#endif // USE_USD_SDK
+
+	return bCopiedSomething;
+}
+
+TArray<UE::FSdfPath> UsdUtils::PastePrims( const UE::FUsdPrim& ParentPrim )
+{
+	TArray<UE::FSdfPath> Result;
+
+#if USE_USD_SDK
+	FScopedUsdAllocs Allocs;
+
+	pxr::UsdPrim UsdParentPrim{ ParentPrim };
+	if ( !UsdParentPrim )
+	{
+		return Result;
+	}
+
+	pxr::UsdStageRefPtr UsdStage = UsdParentPrim.GetStage();
+	if ( !UsdStage )
+	{
+		return Result;
+	}
+
+	pxr::UsdStageRefPtr ClipboardStage = pxr::UsdStageRefPtr{ UnrealUSDWrapper::GetClipboardStage() };
+	if ( !ClipboardStage )
+	{
+		return Result;
+	}
+
+	pxr::SdfLayerHandle ClipboardRoot = ClipboardStage->GetRootLayer();
+	if ( !ClipboardRoot )
+	{
+		return Result;
+	}
+
+	pxr::UsdPrimSiblingRange PrimChildren = ClipboardStage->GetPseudoRoot().GetChildren();
+	int32 NumPrimsToPaste = std::distance( PrimChildren.begin(), PrimChildren.end() );
+
+	TArray<pxr::UsdPrim> PrimsToPaste;
+	PrimsToPaste.Reserve( NumPrimsToPaste );
+	for ( const pxr::UsdPrim& ClipboardPrim : ClipboardStage->GetPseudoRoot().GetChildren() )
+	{
+		PrimsToPaste.Add( ClipboardPrim );
+	}
+
+	pxr::SdfLayerHandle EditTarget = UsdStage->GetEditTarget().GetLayer();
+	if ( !EditTarget )
+	{
+		return Result;
+	}
+
+	TSet<FString> UsedNames;
+	for ( const pxr::UsdPrim& Child : ParentPrim.GetChildren() )
+	{
+		UsedNames.Add( UsdToUnreal::ConvertToken( Child.GetName() ) );
+	}
+
+	Result.SetNum( NumPrimsToPaste );
+	for ( int32 Index = 0; Index < NumPrimsToPaste; ++Index )
+	{
+		const pxr::UsdPrim& ClipboardPrim = PrimsToPaste[ Index ];
+		if ( !ClipboardPrim )
+		{
+			continue;
+		}
+
+		const FString OriginalName = UsdToUnreal::ConvertToken( ClipboardPrim.GetName() );
+		FString ValidName = GetUniqueName( SanitizeUsdIdentifier( *OriginalName ), UsedNames );
+		UsedNames.Add( ValidName );
+
+		pxr::SdfPath TargetSpecPath = UsdParentPrim.GetPath().AppendChild( UnrealToUsd::ConvertToken( *ValidName ).Get() );
+
+		// Ensure our parent prim spec exists, otherwise pxr::SdfCopySpec will fail
+		if ( !pxr::SdfCreatePrimInLayer( EditTarget, TargetSpecPath ) )
+		{
+			continue;
+		}
+
+		if ( !pxr::SdfCopySpec( ClipboardRoot, ClipboardPrim.GetPath(), EditTarget, TargetSpecPath ) )
+		{
+			continue;
+		}
+
+		UE_LOG( LogUsd, Log, TEXT( "Pasted prim '%s' as a child of prim '%s' within the edit target '%s'" ),
+			*OriginalName,
+			*UsdToUnreal::ConvertPath( UsdParentPrim.GetPath() ),
+			*UsdToUnreal::ConvertString( EditTarget->GetIdentifier() )
+		);
+		Result[ Index ] = UE::FSdfPath{ TargetSpecPath };
+	}
+#endif // USE_USD_SDK
+
+	return Result;
+}
+
+bool UsdUtils::CanPastePrims()
+{
+#if USE_USD_SDK
+	pxr::UsdStageRefPtr ClipboardStage = pxr::UsdStageRefPtr{ UnrealUSDWrapper::GetClipboardStage() };
+	if ( !ClipboardStage )
+	{
+		return false;
+	}
+
+	for ( const pxr::UsdPrim& ClipboardPrim : ClipboardStage->GetPseudoRoot().GetChildren() )
+	{
+		if ( ClipboardPrim )
+		{
+			return true;
+		}
+	}
+#endif // USE_USD_SDK
+
+	return false;
+}
+
+void UsdUtils::ClearPrimClipboard()
+{
+#if USE_USD_SDK
+	pxr::UsdStageRefPtr ClipboardStage = pxr::UsdStageRefPtr{ UnrealUSDWrapper::GetClipboardStage() };
+	if ( !ClipboardStage )
+	{
+		return;
+	}
+
+	pxr::SdfLayerHandle ClipboardRoot = ClipboardStage->GetRootLayer();
+	if ( !ClipboardRoot )
+	{
+		return;
+	}
+
+	ClipboardRoot->Clear();
+#endif // USE_USD_SDK
+}
+
+TArray<UE::FSdfPath> UsdUtils::DuplicatePrims( const TArray<UE::FUsdPrim>& Prims, EUsdDuplicateType DuplicateType, const UE::FSdfLayer& TargetLayer )
+{
+	TArray<UE::FSdfPath> Result;
+	Result.SetNum( Prims.Num() );
+
+#if USE_USD_SDK
+	FScopedUsdAllocs Allocs;
+
+	pxr::UsdStageRefPtr UsdStage;
+	for ( const UE::FUsdPrim& Prim : Prims )
+	{
+		if ( Prim )
+		{
+			UsdStage = pxr::UsdStageRefPtr{ Prim.GetStage() };
+			if ( UsdStage )
+			{
+				break;
+			}
+		}
+	}
+	if ( !UsdStage )
+	{
+		return Result;
+	}
+
+	pxr::SdfLayerRefPtr UsdLayer{ TargetLayer };
+
+	// Figure out which layers we'll modify
+	std::unordered_set<pxr::SdfLayerHandle, pxr::TfHash> LayersThatCanBeAffected;
+	switch ( DuplicateType )
+	{
+		case EUsdDuplicateType::FlattenComposedPrim:
+		case EUsdDuplicateType::SingleLayerSpecs:
+		{
+			if ( !UsdLayer )
+			{
+				return Result;
+			}
+
+			LayersThatCanBeAffected.insert( UsdLayer );
+			break;
+		}
+		case EUsdDuplicateType::AllLocalLayerSpecs:
+		{
+			const bool bIncludeSessionLayers = true;
+			for ( const pxr::SdfLayerHandle& Handle : UsdStage->GetLayerStack( bIncludeSessionLayers ) )
+			{
+				LayersThatCanBeAffected.insert( Handle );
+			}
+			break;
+		}
+	}
+
+	// If we're going to need to flatten, just flatten the stage once for all prims we'll duplicate
+	pxr::SdfLayerRefPtr FlattenedLayer = nullptr;
+	if ( DuplicateType == EUsdDuplicateType::FlattenComposedPrim )
+	{
+		pxr::UsdStagePopulationMask Mask;
+		for ( int32 Index = 0; Index < Prims.Num(); ++Index )
+		{
+			pxr::UsdPrim UsdPrim{ Prims[ Index ] };
+			if ( UsdPrim )
+			{
+				Mask.Add( UsdPrim.GetPath() );
+			}
+		}
+
+		pxr::UsdStageRefPtr TempStage = pxr::UsdStage::OpenMasked( UsdStage->GetRootLayer(), Mask );
+		if ( !TempStage )
+		{
+			return Result;
+		}
+
+		const bool bAddSourceFileComment = false;
+		FlattenedLayer = TempStage->Flatten( bAddSourceFileComment );
+		if ( !FlattenedLayer )
+		{
+			return Result;
+		}
+	}
+
+	for ( int32 Index = 0; Index < Prims.Num(); ++Index )
+	{
+		pxr::UsdPrim UsdPrim{ Prims[ Index ] };
+		if ( !UsdPrim )
+		{
+			continue;
+		}
+
+		std::vector<pxr::SdfPrimSpecHandle> PrimSpecs = UsdPrim.GetPrimStack();
+
+		// Note: We won't actually use these in case we're flattening, but it makes the code a bit simpler to also
+		// do this while we're collecting LayersThatWillBeAffected below
+		std::vector<pxr::SdfPrimSpecHandle> SpecsToDuplicate;
+		SpecsToDuplicate.reserve( PrimSpecs.size() );
+
+		std::unordered_set<pxr::SdfLayerHandle, pxr::TfHash> LayersThatWillBeAffected;
+		LayersThatWillBeAffected.reserve( PrimSpecs.size() );
+
+		for ( const pxr::SdfPrimSpecHandle& Spec : PrimSpecs )
+		{
+			// For whatever reason sometimes there are invalid specs in the layer stack, so we need to be careful
+			if ( !Spec )
+			{
+				continue;
+			}
+
+			pxr::SdfPath SpecPath = Spec->GetPath();
+			if ( !SpecPath.IsPrimPath() )
+			{
+				continue;
+			}
+
+			pxr::SdfLayerHandle SpecLayerHandle = Spec->GetLayer();
+			if ( !SpecLayerHandle || LayersThatCanBeAffected.count( SpecLayerHandle ) == 0 )
+			{
+				continue;
+			}
+
+			SpecsToDuplicate.push_back( Spec );
+			LayersThatWillBeAffected.insert( SpecLayerHandle );
+		}
+
+		// Find a usable name for the new duplicate prim
+		pxr::SdfPath NewSpecPath;
+		{
+			const std::string SourcePrimName = UsdPrim.GetName().GetString();
+			const pxr::SdfPath ParentPath = UsdPrim.GetPath().GetParentPath();
+
+			int32 Suffix = -1;
+
+			bool bFoundName = false;
+			while ( !bFoundName )
+			{
+				NewSpecPath = ParentPath.AppendElementString( SourcePrimName + "_" + std::to_string( ++Suffix ) );
+
+				bFoundName = true;
+				for ( const pxr::SdfLayerHandle& Layer : LayersThatWillBeAffected )
+				{
+					if ( Layer->HasSpec( NewSpecPath ) )
+					{
+						bFoundName = false;
+						break;
+					}
+				}
+			}
+		}
+
+		// Actually do the duplication operation we chose
+		if ( DuplicateType == EUsdDuplicateType::FlattenComposedPrim && FlattenedLayer )
+		{
+			pxr::SdfPrimSpecHandle FlattenedPrim = FlattenedLayer->GetPrimAtPath( UsdPrim.GetPath() );
+			if ( !FlattenedPrim )
+			{
+				return Result;
+			}
+
+			if ( !pxr::SdfJustCreatePrimInLayer( UsdLayer, NewSpecPath ) )
+			{
+				UE_LOG( LogUsd, Warning, TEXT( "Failed to create prim and parent specs for path '%s' within layer '%s'" ),
+					*UsdToUnreal::ConvertPath( NewSpecPath ),
+					*UsdToUnreal::ConvertString( UsdLayer->GetIdentifier() )
+				);
+				return Result;
+			}
+
+			if ( !pxr::SdfCopySpec( FlattenedLayer, FlattenedPrim->GetPath(), UsdLayer, NewSpecPath ) )
+			{
+				UE_LOG( LogUsd, Warning, TEXT( "Failed to copy flattened prim spec from '%s' onto path '%s' within layer '%s'" ),
+					*UsdToUnreal::ConvertPath( UsdPrim.GetPath() ),
+					*UsdToUnreal::ConvertPath( NewSpecPath ),
+					*UsdToUnreal::ConvertString( UsdLayer->GetIdentifier() )
+				);
+				return Result;
+			}
+
+			UE_LOG( LogUsd, Log, TEXT( "Flattened prim '%s' onto spec '%s' at layer '%s'" ),
+				*UsdToUnreal::ConvertPath( UsdPrim.GetPath() ),
+				*UsdToUnreal::ConvertPath( NewSpecPath ),
+				*UsdToUnreal::ConvertString( UsdLayer->GetIdentifier() )
+			);
+		}
+		else
+		{
+			for ( const pxr::SdfPrimSpecHandle& Spec : SpecsToDuplicate )
+			{
+				pxr::SdfPath SpecPath = Spec->GetPath();
+				pxr::SdfLayerHandle SpecLayerHandle = Spec->GetLayer();
+
+				UE_LOG( LogUsd, Log, TEXT( "Duplicating prim spec '%s' within layer '%s'" ),
+					*UsdToUnreal::ConvertPath( SpecPath ),
+					*UsdToUnreal::ConvertString( SpecLayerHandle->GetIdentifier() )
+				);
+
+				// Technically we shouldn't need to do this since we'll already do our changes on the Sdf level, however the
+				// USDTransactor will record these notices as belonging to the current edit target, and if that is not in sync
+				// with the layer that is actually changing, we won't be able to undo/redo the duplicate operation
+				pxr::UsdEditContext Context{ UsdStage, SpecLayerHandle };
+
+				// Since we're duplicating a prim essentially as a sibling, parent specs should always exist.
+				// Let's ensure that though, just in case
+				if ( !pxr::SdfJustCreatePrimInLayer( SpecLayerHandle, NewSpecPath ) )
+				{
+					UE_LOG( LogUsd, Warning, TEXT( "Failed to create prim and parent specs for path '%s' within layer '%s'" ),
+						*UsdToUnreal::ConvertPath( NewSpecPath ),
+						*UsdToUnreal::ConvertString( SpecLayerHandle->GetIdentifier() )
+					);
+					continue;
+				}
+
+				if ( !pxr::SdfCopySpec( SpecLayerHandle, SpecPath, SpecLayerHandle, NewSpecPath ) )
+				{
+					UE_LOG( LogUsd, Warning, TEXT( "Failed to copy spec from path '%s' onto path '%s' within layer '%s'" ),
+						*UsdToUnreal::ConvertPath( SpecPath ),
+						*UsdToUnreal::ConvertPath( NewSpecPath ),
+						*UsdToUnreal::ConvertString( SpecLayerHandle->GetIdentifier() )
+					);
+				}
+			}
+		}
+
+		Result[ Index ] = UE::FSdfPath{ NewSpecPath };
+	}
+#endif // USE_USD_SDK
+
+	return Result;
+}
+
 #if USE_USD_SDK
 #undef LOCTEXT_NAMESPACE
 #endif
diff --git a/Engine/Plugins/Importers/USDImporter/Source/USDUtilities/Public/USDConversionUtils.h b/Engine/Plugins/Importers/USDImporter/Source/USDUtilities/Public/USDConversionUtils.h
--- a/Engine/Plugins/Importers/USDImporter/Source/USDUtilities/Public/USDConversionUtils.h
+++ b/Engine/Plugins/Importers/USDImporter/Source/USDUtilities/Public/USDConversionUtils.h
@@ -10,6 +10,7 @@
 #include "USDMemory.h"
 #include "UsdWrappers/ForwardDeclarations.h"
 #include "UsdWrappers/SdfLayer.h"
+#include "UsdWrappers/SdfPath.h"
 
 #if USE_USD_SDK
 #include "USDIncludesStart.h"
@@ -39,13 +40,13 @@ PXR_NAMESPACE_CLOSE_SCOPE
 
 #endif // #if USE_USD_SDK
 
-class UUsdAssetImportData;
 class USceneComponent;
+class UUsdAssetImportData;
+enum class EUsdDuplicateType : uint8;
 enum class EUsdUpAxis : uint8;
 namespace UE
 {
 	class FUsdPrim;
-	class FSdfPath;
 }
 namespace UsdUtils
 {
@@ -181,10 +182,10 @@ namespace UsdUtils
 	USDUTILITIES_API UUsdAssetImportData* GetAssetImportData( UObject* Asset );
 
 	/** Adds a reference on Prim to the layer at AbsoluteFilePath */
-	USDUTILITIES_API void AddReference( UE::FUsdPrim& Prim, const TCHAR* AbsoluteFilePath );
+	USDUTILITIES_API void AddReference( UE::FUsdPrim& Prim, const TCHAR* AbsoluteFilePath, const UE::FSdfPath& TargetPrimPath = {}, double TimeCodeOffset = 0.0, double TimeCodeScale = 1.0 );
 
 	/** Adds a payload on Prim pointing at the default prim of the layer at AbsoluteFilePath */
-	USDUTILITIES_API void AddPayload( UE::FUsdPrim& Prim, const TCHAR* AbsoluteFilePath );
+	USDUTILITIES_API void AddPayload( UE::FUsdPrim& Prim, const TCHAR* AbsoluteFilePath, const UE::FSdfPath& TargetPrimPath = {}, double TimeCodeOffset = 0.0, double TimeCodeScale = 1.0 );
 
 	/**
 	 * Renames a single prim to a new name
@@ -275,5 +276,64 @@ namespace UsdUtils
 	 *				  specs from the entire stage's local layer stack.
 	 */
 	USDUTILITIES_API void RemoveAllPrimSpecs( const UE::FUsdPrim& Prim, const UE::FSdfLayer& Layer = UE::FSdfLayer{} );
+
+	/**
+	 * Copies flattened versions of the input prims onto the clipboard stage and removes all the prim specs for Prims from their stages.
+	 * These cut prims can then be pasted with PastePrims.
+	 *
+	 * @param Prims - Prims to cut
+	 * @return True if we managed to cut
+	 */
+	USDUTILITIES_API bool CutPrims( const TArray<UE::FUsdPrim>& Prims );
+
+	/**
+	 * Copies flattened versions of the input prims onto the clipboard stage.
+	 * These copied prims can then be pasted with PastePrims.
+	 *
+	 * @param Prims - Prims to copy
+	 * @return True if we managed to copy
+	 */
+	USDUTILITIES_API bool CopyPrims( const TArray<UE::FUsdPrim>& Prims );
+
+	/**
+	 * Pastes the prims from the clipboard stage as children of ParentPrim.
+	 *
+	 * The pasted prims may be renamed in order to have valid names for the target location, which is why this function
+	 * returns the pasted prim paths.
+	 * This function returns just paths instead of actual prims because USD needs to respond to the notices about
+	 * the created prim specs before the prims are fully created, which means we wouldn't be able to return the
+	 * created prims yet, in case this function was called from within an SdfChangeBlock.
+	 *
+	 * @param ParentPrim - Prim that will become parent to the pasted prims
+	 * @return Paths to the pasted prim specs, after they were added as children of ParentPrim
+	 */
+	USDUTILITIES_API TArray<UE::FSdfPath> PastePrims( const UE::FUsdPrim& ParentPrim );
+
+	/** Returns true if we have prims that we can paste within our clipboard stage */
+	USDUTILITIES_API bool CanPastePrims();
+
+	/** Clears all prims from our clipboard stage */
+	USDUTILITIES_API void ClearPrimClipboard();
+
+	/**
+	 * Duplicates all provided Prims one-by-one, performing the requested DuplicateType.
+	 * See the documentation on EUsdDuplicateType for the different operation types.
+	 *
+	 * The duplicated prims may be renamed in order to have valid names for the target location, which is why this
+	 * function returns the pasted prim paths.
+	 * This function returns just paths instead of actual prims because USD needs to respond to the notices about
+	 * the created prim specs before the prims are fully created, which means we wouldn't be able to return the
+	 * created prims yet, in case this function was called from within an SdfChangeBlock.
+	 *
+	 * @param Prims - Prims to duplicate
+	 * @param DuplicateType - Type of prim duplication to perform
+	 * @param TargetLayer - Target layer to use when duplicating, if relevant for that duplication type
+	 * @return Paths to the duplicated prim specs, after they were added as children of ParentPrim.
+	 */
+	USDUTILITIES_API TArray<UE::FSdfPath> DuplicatePrims(
+		const TArray<UE::FUsdPrim>& Prims,
+		EUsdDuplicateType DuplicateType,
+		const UE::FSdfLayer& TargetLayer = UE::FSdfLayer{}
+	);
 }
 
diff --git a/Engine/Plugins/Importers/USDImporter/Source/UnrealUSDWrapper/Private/UnrealUSDWrapper.cpp b/Engine/Plugins/Importers/USDImporter/Source/UnrealUSDWrapper/Private/UnrealUSDWrapper.cpp
--- a/Engine/Plugins/Importers/USDImporter/Source/UnrealUSDWrapper/Private/UnrealUSDWrapper.cpp
+++ b/Engine/Plugins/Importers/USDImporter/Source/UnrealUSDWrapper/Private/UnrealUSDWrapper.cpp
@@ -984,6 +984,26 @@ UE::FUsdStage UnrealUSDWrapper::NewStage()
 #endif // #if USE_USD_SDK
 }
 
+UE::FUsdStage UnrealUSDWrapper::GetClipboardStage()
+{
+	static UE::FUsdStage ClipboardStage;
+
+#if USE_USD_SDK
+	FScopedUsdAllocs UsdAllocs;
+
+	if ( !ClipboardStage )
+	{
+		// Use this specific name so that we can check for it on our automated tests
+		pxr::UsdStageCacheContext StageCacheContext{ pxr::UsdUtilsStageCache::Get() };
+		ClipboardStage = UE::FUsdStage{ pxr::UsdStage::CreateInMemory( TCHAR_TO_ANSI( TEXT( "unreal_clipboard_layer" ) ) ) };
+
+		pxr::UsdGeomSetStageUpAxis( ClipboardStage, pxr::UsdGeomTokens->z );
+	}
+#endif // #if USE_USD_SDK
+
+	return ClipboardStage;
+}
+
 TArray< UE::FUsdStage > UnrealUSDWrapper::GetAllStagesFromCache()
 {
 	TArray< UE::FUsdStage > StagesInCache;
diff --git a/Engine/Plugins/Importers/USDImporter/Source/UnrealUSDWrapper/Public/UnrealUSDWrapper.h b/Engine/Plugins/Importers/USDImporter/Source/UnrealUSDWrapper/Public/UnrealUSDWrapper.h
--- a/Engine/Plugins/Importers/USDImporter/Source/UnrealUSDWrapper/Public/UnrealUSDWrapper.h
+++ b/Engine/Plugins/Importers/USDImporter/Source/UnrealUSDWrapper/Public/UnrealUSDWrapper.h
@@ -225,6 +225,13 @@ public:
 	/** Creates a new memory USD root layer, opens it as a new stage and returns that stage */
 	UNREALUSDWRAPPER_API static UE::FUsdStage NewStage();
 
+	/**
+	 * Get the singleton, persistent stage used as a clipboard for prim cut/copy/paste operations.
+	 * WARNING: This stage may remain open indefinitely! If you use this directly, be aware that there may be unintended
+	 * consequences (e.g. a sublayer added to this stage may never fully close)
+	 */
+	UNREALUSDWRAPPER_API static UE::FUsdStage GetClipboardStage();
+
 	/** Returns all the stages that are currently opened in the USD utils stage cache, shared between C++ and Python */
 	UNREALUSDWRAPPER_API static TArray< UE::FUsdStage > GetAllStagesFromCache();
 
