USD: Fix incorrect material asset path.

All this does is add the missing single quote at the end of the material path string. I'm not sure why this was accepted before, but its no longer valid and is causing us to fail to find the UsdPreviewSurface base material (and to fail a large part of the automated tests)

diff --git a/Engine/Plugins/Importers/USDImporter/Source/USDSchemas/Private/USDShadeMaterialTranslator.cpp b/Engine/Plugins/Importers/USDImporter/Source/USDSchemas/Private/USDShadeMaterialTranslator.cpp
--- a/Engine/Plugins/Importers/USDImporter/Source/USDSchemas/Private/USDShadeMaterialTranslator.cpp
+++ b/Engine/Plugins/Importers/USDImporter/Source/USDSchemas/Private/USDShadeMaterialTranslator.cpp
@@ -65,7 +65,7 @@ void FUsdShadeMaterialTranslator::CreateAssets()
 			BaseMaterialName += TEXT("VT");
 		}
 
-		const FString BaseMaterialPath = FString::Printf( TEXT("Material'/USDImporter/Materials/%s.%s"), *BaseMaterialName, *BaseMaterialName );
+		const FString BaseMaterialPath = FString::Printf( TEXT("Material'/USDImporter/Materials/%s.%s'"), *BaseMaterialName, *BaseMaterialName );
 
 		if ( UMaterialInterface* BaseMaterial = Cast< UMaterialInterface >( FSoftObjectPath( BaseMaterialPath ).TryLoad() ) )
 		{
