USD: Actually fix missing include from non-unity build due to transaction object refactor.

diff --git a/Engine/Plugins/Importers/USDImporter/Source/USDStage/Private/USDLevelSequenceHelper.cpp b/Engine/Plugins/Importers/USDImporter/Source/USDStage/Private/USDLevelSequenceHelper.cpp
--- a/Engine/Plugins/Importers/USDImporter/Source/USDStage/Private/USDLevelSequenceHelper.cpp
+++ b/Engine/Plugins/Importers/USDImporter/Source/USDStage/Private/USDLevelSequenceHelper.cpp
@@ -44,6 +44,7 @@
 #include "LevelSequenceActor.h"
 #include "LevelSequencePlayer.h"
 #include "Misc/ITransaction.h"
+#include "Misc/TransactionObjectEvent.h"
 #include "MovieScene.h"
 #include "MovieSceneTimeHelpers.h"
 #include "MovieSceneTrack.h"
