USD: Fix missing groom schemas on Linux and Mac.

diff --git a/Engine/Plugins/Importers/USDImporter/Resources/UsdResources/Linux/plugins/unreal/resources/plugInfo.json b/Engine/Plugins/Importers/USDImporter/Resources/UsdResources/Linux/plugins/unreal/resources/plugInfo.json
--- a/Engine/Plugins/Importers/USDImporter/Resources/UsdResources/Linux/plugins/unreal/resources/plugInfo.json
+++ b/Engine/Plugins/Importers/USDImporter/Resources/UsdResources/Linux/plugins/unreal/resources/plugInfo.json
@@ -30,6 +30,20 @@
                         ], 
                         "schemaKind": "singleApplyAPI"
                     }, 
+                    "unrealGroomBindingAPI": {
+                        "alias": {
+                            "UsdSchemaBase": "GroomBindingAPI"
+                        }, 
+                        "apiSchemaCanOnlyApplyTo": [
+                            "Mesh", 
+                            "SkelRoot"
+                        ], 
+                        "autoGenerated": true, 
+                        "bases": [
+                            "UsdAPISchemaBase"
+                        ], 
+                        "schemaKind": "singleApplyAPI"
+                    }, 
                     "unrealLiveLinkAPI": {
                         "alias": {
                             "UsdSchemaBase": "LiveLinkAPI"
diff --git a/Engine/Plugins/Importers/USDImporter/Resources/UsdResources/Mac/plugins/unreal/resources/plugInfo.json b/Engine/Plugins/Importers/USDImporter/Resources/UsdResources/Mac/plugins/unreal/resources/plugInfo.json
--- a/Engine/Plugins/Importers/USDImporter/Resources/UsdResources/Mac/plugins/unreal/resources/plugInfo.json
+++ b/Engine/Plugins/Importers/USDImporter/Resources/UsdResources/Mac/plugins/unreal/resources/plugInfo.json
@@ -30,6 +30,20 @@
                         ], 
                         "schemaKind": "singleApplyAPI"
                     }, 
+                    "unrealGroomBindingAPI": {
+                        "alias": {
+                            "UsdSchemaBase": "GroomBindingAPI"
+                        }, 
+                        "apiSchemaCanOnlyApplyTo": [
+                            "Mesh", 
+                            "SkelRoot"
+                        ], 
+                        "autoGenerated": true, 
+                        "bases": [
+                            "UsdAPISchemaBase"
+                        ], 
+                        "schemaKind": "singleApplyAPI"
+                    }, 
                     "unrealLiveLinkAPI": {
                         "alias": {
                             "UsdSchemaBase": "LiveLinkAPI"
