USD: Add options to clear and save individual layers.

diff --git a/Engine/Plugins/Importers/USDImporter/Source/USDClasses/Public/USDProjectSettings.h b/Engine/Plugins/Importers/USDImporter/Source/USDClasses/Public/USDProjectSettings.h
--- a/Engine/Plugins/Importers/USDImporter/Source/USDClasses/Public/USDProjectSettings.h
+++ b/Engine/Plugins/Importers/USDImporter/Source/USDClasses/Public/USDProjectSettings.h
@@ -16,4 +16,7 @@ class UUsdProjectSettings : public UDeveloperSettings
 public:
 	UPROPERTY(config, EditAnywhere, Category=General)
 	TArray<FDirectoryPath> AdditionalPluginDirectories;
+
+	UPROPERTY( config, EditAnywhere, Category = USD )
+	bool bShowConfirmationWhenClearingLayers = true;
 };
\ No newline at end of file
diff --git a/Engine/Plugins/Importers/USDImporter/Source/USDStageEditor/Private/SUSDLayersTreeView.cpp b/Engine/Plugins/Importers/USDImporter/Source/USDStageEditor/Private/SUSDLayersTreeView.cpp
--- a/Engine/Plugins/Importers/USDImporter/Source/USDStageEditor/Private/SUSDLayersTreeView.cpp
+++ b/Engine/Plugins/Importers/USDImporter/Source/USDStageEditor/Private/SUSDLayersTreeView.cpp
@@ -7,23 +7,27 @@
 #include "USDLayersViewModel.h"
 #include "USDLayerUtils.h"
 #include "USDMemory.h"
+#include "USDProjectSettings.h"
 #include "USDStageActor.h"
 #include "USDStageModule.h"
 #include "USDTypesConversion.h"
-
 #include "UsdWrappers/SdfLayer.h"
 #include "UsdWrappers/UsdStage.h"
-#include "Styling/AppStyle.h"
+
 #include "DesktopPlatformModule.h"
-#include "Styling/AppStyle.h"
 #include "Engine/World.h"
 #include "EngineAnalytics.h"
 #include "Framework/MultiBox/MultiBoxBuilder.h"
+#include "Framework/Notifications/NotificationManager.h"
 #include "IDesktopPlatform.h"
 #include "Modules/ModuleManager.h"
+#include "ScopedTransaction.h"
+#include "Styling/AppStyle.h"
 #include "Styling/SlateTypes.h"
+#include "Templates/SharedPointer.h"
 #include "Widgets/Images/SImage.h"
 #include "Widgets/Input/SButton.h"
+#include "Widgets/Notifications/SNotificationList.h"
 #include "Widgets/SToolTip.h"
 
 #if USE_USD_SDK
@@ -332,6 +336,30 @@ TSharedPtr< SWidget > SUsdLayersTreeView::ConstructLayerContextMenu()
 			EUserInterfaceActionType::Button
 		);
 
+		LayerOptions.AddMenuEntry(
+			LOCTEXT( "ClearLayer", "Clear" ),
+			LOCTEXT( "ClearLayer_ToolTip", "Clears this layer of all data" ),
+			FSlateIcon(),
+			FUIAction(
+				FExecuteAction::CreateSP( this, &SUsdLayersTreeView::OnClearSelectedLayers ),
+				FCanExecuteAction::CreateSP( this, &SUsdLayersTreeView::CanClearSelectedLayers )
+			),
+			NAME_None,
+			EUserInterfaceActionType::Button
+		);
+
+		LayerOptions.AddMenuEntry(
+			LOCTEXT( "SaveLayer", "Save" ),
+			LOCTEXT( "SaveLayer_ToolTip", "Saves the layer modifications to disk" ),
+			FSlateIcon(),
+			FUIAction(
+				FExecuteAction::CreateSP( this, &SUsdLayersTreeView::OnSaveSelectedLayers ),
+				FCanExecuteAction::CreateSP( this, &SUsdLayersTreeView::CanSaveSelectedLayers )
+			),
+			NAME_None,
+			EUserInterfaceActionType::Button
+		);
+
 		LayerOptions.AddMenuEntry(
 			LOCTEXT( "ExportLayer", "Export" ),
 			LOCTEXT( "Export_ToolTip", "Export the selected layers, having the exported layers reference the original stage's layers" ),
@@ -422,6 +450,164 @@ void SUsdLayersTreeView::OnEditSelectedLayer()
 	}
 }
 
+void SUsdLayersTreeView::OnClearSelectedLayers()
+{
+	TArray< FUsdLayerViewModelRef > MySelectedItems = GetSelectedItems();
+
+	// We'll show a confirmation toast, which is non-modal. So keep track of the original selected items so that
+	// if anything changes by the time we accept the dialog we'll still know which layers to clear
+	TArray<UE::FSdfLayerWeak> SelectedLayers;
+	SelectedLayers.Reserve( MySelectedItems.Num() );
+
+	FString LayerNames;
+	const FString Separator = TEXT(", ");
+	for ( FUsdLayerViewModelRef SelectedItem : MySelectedItems )
+	{
+		if ( UE::FSdfLayer Layer = SelectedItem->GetLayer() )
+		{
+			SelectedLayers.Add( Layer );
+			LayerNames += Layer.GetDisplayName() + Separator;
+		}
+	}
+	LayerNames.RemoveFromEnd( Separator );
+
+	TFunction<void()> ClearSelectedLayersInner = [this, SelectedLayers]()
+	{
+		// This doesn't work very well USD-wise yet, and I don't know how we can possibly undo/redo
+		// clearing a layer like this. We do need a transaction though, as this may trigger the creation
+		// or deletion of UObjects
+		FScopedTransaction Transaction( LOCTEXT( "ClearTransaction", "Clear selected layers" ) );
+
+		for ( UE::FSdfLayerWeak Layer : SelectedLayers )
+		{
+			if ( Layer )
+			{
+				Layer.Clear();
+			}
+		}
+	};
+
+	const UUsdProjectSettings* Settings = GetDefault<UUsdProjectSettings>();
+	if ( Settings && Settings->bShowConfirmationWhenClearingLayers )
+	{
+		static TWeakPtr<SNotificationItem> Notification;
+
+		FNotificationInfo Toast( LOCTEXT( "ConfirmClearingLayer", "Clearing cannot be undone" ) );
+		Toast.SubText = FText::Format(
+			LOCTEXT( "ConfirmClearingLayer_Subtext", "Clearing USD layers cannot be undone.\nDo you wish to proceed clearing {0}|plural(one=layer,other=layers) '{1}' ?" ),
+			SelectedLayers.Num(),
+			FText::FromString(LayerNames)
+		);
+		Toast.Image = FCoreStyle::Get().GetBrush( TEXT( "MessageLog.Warning" ) );
+		Toast.bUseLargeFont = false;
+		Toast.bFireAndForget = false;
+		Toast.FadeOutDuration = 0.0f;
+		Toast.ExpireDuration = 0.0f;
+		Toast.bUseThrobber = false;
+		Toast.bUseSuccessFailIcons = false;
+
+		Toast.ButtonDetails.Emplace(
+			LOCTEXT( "ConfirmClearingLayerOkAll", "Always proceed" ),
+			FText::GetEmpty(),
+			FSimpleDelegate::CreateLambda( [ClearSelectedLayersInner]() {
+			if ( TSharedPtr<SNotificationItem> PinnedNotification = Notification.Pin() )
+			{
+				PinnedNotification->SetCompletionState( SNotificationItem::CS_Success );
+				PinnedNotification->ExpireAndFadeout();
+
+				UUsdProjectSettings* Settings = GetMutableDefault<UUsdProjectSettings>();
+				Settings->bShowConfirmationWhenClearingLayers = false;
+
+				ClearSelectedLayersInner();
+			}
+		}));
+
+		Toast.ButtonDetails.Emplace(
+			LOCTEXT( "ConfirmClearingLayerOk", "Proceed" ),
+			FText::GetEmpty(),
+			FSimpleDelegate::CreateLambda( [ClearSelectedLayersInner]() {
+			if ( TSharedPtr<SNotificationItem> PinnedNotification = Notification.Pin() )
+			{
+				PinnedNotification->SetCompletionState( SNotificationItem::CS_Success );
+				PinnedNotification->ExpireAndFadeout();
+
+				ClearSelectedLayersInner();
+			}
+		}));
+
+		Toast.ButtonDetails.Emplace(
+			LOCTEXT( "ConfirmClearingLayerCancel", "Cancel" ),
+			FText::GetEmpty(),
+			FSimpleDelegate::CreateLambda( []() {
+			if ( TSharedPtr<SNotificationItem> PinnedNotification = Notification.Pin() )
+			{
+				PinnedNotification->SetCompletionState( SNotificationItem::CS_Fail );
+				PinnedNotification->ExpireAndFadeout();
+			}
+		}));
+
+		// Only show one at a time
+		if ( !Notification.IsValid() )
+		{
+			Notification = FSlateNotificationManager::Get().AddNotification( Toast );
+		}
+
+		if ( TSharedPtr<SNotificationItem> PinnedNotification = Notification.Pin() )
+		{
+			PinnedNotification->SetCompletionState( SNotificationItem::CS_Pending );
+		}
+	}
+	else
+	{
+		// Don't show prompt, just clear
+		ClearSelectedLayersInner();
+	}
+}
+
+bool SUsdLayersTreeView::CanClearSelectedLayers() const
+{
+	TArray< FUsdLayerViewModelRef > MySelectedItems = GetSelectedItems();
+
+	for ( FUsdLayerViewModelRef SelectedItem : MySelectedItems )
+	{
+		if ( !SelectedItem->GetLayer().IsEmpty() )
+		{
+			return true;
+		}
+	}
+
+	return false;
+}
+
+void SUsdLayersTreeView::OnSaveSelectedLayers()
+{
+	TArray< FUsdLayerViewModelRef > MySelectedItems = GetSelectedItems();
+
+	for ( FUsdLayerViewModelRef SelectedItem : MySelectedItems )
+	{
+		if ( SelectedItem->IsLayerDirty() )
+		{
+			const bool bForce = true;
+			SelectedItem->GetLayer().Save( bForce );
+		}
+	}
+}
+
+bool SUsdLayersTreeView::CanSaveSelectedLayers() const
+{
+	TArray< FUsdLayerViewModelRef > MySelectedItems = GetSelectedItems();
+
+	for ( FUsdLayerViewModelRef SelectedItem : MySelectedItems )
+	{
+		if ( SelectedItem->IsLayerDirty() )
+		{
+			return true;
+		}
+	}
+
+	return false;
+}
+
 void SUsdLayersTreeView::OnExportSelectedLayers() const
 {
 	TArray< FUsdLayerViewModelRef > MySelectedItems = GetSelectedItems();
diff --git a/Engine/Plugins/Importers/USDImporter/Source/USDStageEditor/Private/SUSDLayersTreeView.h b/Engine/Plugins/Importers/USDImporter/Source/USDStageEditor/Private/SUSDLayersTreeView.h
--- a/Engine/Plugins/Importers/USDImporter/Source/USDStageEditor/Private/SUSDLayersTreeView.h
+++ b/Engine/Plugins/Importers/USDImporter/Source/USDStageEditor/Private/SUSDLayersTreeView.h
@@ -33,6 +33,12 @@ private:
 	bool CanEditSelectedLayer() const;
 	void OnEditSelectedLayer();
 
+	void OnClearSelectedLayers();
+	bool CanClearSelectedLayers() const;
+
+	void OnSaveSelectedLayers();
+	bool CanSaveSelectedLayers() const;
+
 	void OnExportSelectedLayers() const;
 
 	bool CanAddSubLayer() const;
diff --git a/Engine/Plugins/Importers/USDImporter/Source/UnrealUSDWrapper/Private/UsdWrappers/SdfLayer.cpp b/Engine/Plugins/Importers/USDImporter/Source/UnrealUSDWrapper/Private/UsdWrappers/SdfLayer.cpp
--- a/Engine/Plugins/Importers/USDImporter/Source/UnrealUSDWrapper/Private/UsdWrappers/SdfLayer.cpp
+++ b/Engine/Plugins/Importers/USDImporter/Source/UnrealUSDWrapper/Private/UsdWrappers/SdfLayer.cpp
@@ -677,6 +677,17 @@ namespace UE
 		return false;
 	}
 
+	template<typename PtrType>
+	void FSdfLayerBase<PtrType>::Clear()
+	{
+#if USE_USD_SDK
+		if ( const PtrType& Ptr = Impl->GetInner() )
+		{
+			return Ptr->Clear();
+		}
+#endif // #if USE_USD_SDK
+	}
+
 	FString FSdfLayerUtils::SdfComputeAssetPathRelativeToLayer( const FSdfLayer& Anchor, const TCHAR* AssetPath )
 	{
 #if USE_USD_SDK
diff --git a/Engine/Plugins/Importers/USDImporter/Source/UnrealUSDWrapper/Public/UsdWrappers/SdfLayer.h b/Engine/Plugins/Importers/USDImporter/Source/UnrealUSDWrapper/Public/UsdWrappers/SdfLayer.h
--- a/Engine/Plugins/Importers/USDImporter/Source/UnrealUSDWrapper/Public/UsdWrappers/SdfLayer.h
+++ b/Engine/Plugins/Importers/USDImporter/Source/UnrealUSDWrapper/Public/UsdWrappers/SdfLayer.h
@@ -104,6 +104,8 @@ namespace UE
 
 		bool Export( const TCHAR* Filename ) const;
 
+		void Clear();
+
 		bool HasStartTimeCode() const;
 		double GetStartTimeCode() const;
 		void SetStartTimeCode( double TimeCode );
