USD: Notify users in case an authored opinion is overridden by an existing one.

diff --git a/Engine/Plugins/Importers/USDImporter/Source/USDClasses/Public/USDProjectSettings.h b/Engine/Plugins/Importers/USDImporter/Source/USDClasses/Public/USDProjectSettings.h
--- a/Engine/Plugins/Importers/USDImporter/Source/USDClasses/Public/USDProjectSettings.h
+++ b/Engine/Plugins/Importers/USDImporter/Source/USDClasses/Public/USDProjectSettings.h
@@ -14,9 +14,14 @@ class UUsdProjectSettings : public UDeveloperSettings
 	GENERATED_BODY()
 
 public:
-	UPROPERTY(config, EditAnywhere, Category=General)
+	// Additional paths to check for USD plugins
+	UPROPERTY( config, EditAnywhere, Category = USD )
 	TArray<FDirectoryPath> AdditionalPluginDirectories;
 
 	UPROPERTY( config, EditAnywhere, Category = USD )
 	bool bShowConfirmationWhenClearingLayers = true;
+
+	// Whether to show the warning dialog when authoring opinions that could have no effect on the composed stage
+	UPROPERTY( config, EditAnywhere, Category = USD )
+	bool bShowOverriddenOpinionsWarning = true;
 };
\ No newline at end of file
diff --git a/Engine/Plugins/Importers/USDImporter/Source/USDStageEditor/Private/Widgets/SUSDIntegrationsPanel.cpp b/Engine/Plugins/Importers/USDImporter/Source/USDStageEditor/Private/Widgets/SUSDIntegrationsPanel.cpp
--- a/Engine/Plugins/Importers/USDImporter/Source/USDStageEditor/Private/Widgets/SUSDIntegrationsPanel.cpp
+++ b/Engine/Plugins/Importers/USDImporter/Source/USDStageEditor/Private/Widgets/SUSDIntegrationsPanel.cpp
@@ -5,6 +5,7 @@
 #if USE_USD_SDK
 
 #include "SUSDStageEditorStyle.h"
+#include "USDAttributeUtils.h"
 #include "USDIntegrationUtils.h"
 #include "USDTypesConversion.h"
 #include "USDValueConversion.h"
@@ -109,6 +110,10 @@ TSharedRef< SWidget > SUsdIntegrationsPanelRow::GenerateWidgetForColumn( const F
 
 		FName TypeName = Attribute->GetTypeName();
 		TSharedPtr<UE::FUsdAttribute> AttributeCopy = Attribute;
+		if ( !AttributeCopy || !*AttributeCopy )
+		{
+			return ColumnWidget;
+		}
 
 		if ( AttributeName == *UsdToUnreal::ConvertToken( UnrealIdentifiers::UnrealAnimBlueprintPath ) )
 		{
@@ -158,6 +163,7 @@ TSharedRef< SWidget > SUsdIntegrationsPanelRow::GenerateWidgetForColumn( const F
 					{
 						UE::FSdfChangeBlock Block;
 						AttributeCopy->Set( Value );
+						UsdUtils::NotifyIfOverriddenOpinion( *AttributeCopy );
 					}
 				} );
 				PropertyView->SetOnPropertyValueChanged( PropertyChanged );
@@ -215,6 +221,7 @@ TSharedRef< SWidget > SUsdIntegrationsPanelRow::GenerateWidgetForColumn( const F
 					{
 						UE::FSdfChangeBlock Block;
 						AttributeCopy->Set( Value );
+						UsdUtils::NotifyIfOverriddenOpinion( *AttributeCopy );
 					}
 				} );
 				PropertyView->SetOnPropertyValueChanged( PropertyChanged );
@@ -273,6 +280,7 @@ TSharedRef< SWidget > SUsdIntegrationsPanelRow::GenerateWidgetForColumn( const F
 					{
 						UE::FSdfChangeBlock Block;
 						AttributeCopy->Set( Value );
+						UsdUtils::NotifyIfOverriddenOpinion( *AttributeCopy );
 					}
 				})
 			];
@@ -298,6 +306,7 @@ TSharedRef< SWidget > SUsdIntegrationsPanelRow::GenerateWidgetForColumn( const F
 					{
 						UE::FSdfChangeBlock Block;
 						AttributeCopy->Set( Value );
+						UsdUtils::NotifyIfOverriddenOpinion( *AttributeCopy );
 					}
 				})
 				.Value_Lambda( [AttributeCopy]() -> float
@@ -364,6 +373,7 @@ TSharedRef< SWidget > SUsdIntegrationsPanelRow::GenerateWidgetForColumn( const F
 				{
 					UE::FSdfChangeBlock Block;
 					AttributeCopy->Set( Value );
+					UsdUtils::NotifyIfOverriddenOpinion( *AttributeCopy );
 				}
 			});
 		}
@@ -402,6 +412,7 @@ TSharedRef< SWidget > SUsdIntegrationsPanelRow::GenerateWidgetForColumn( const F
 					{
 						UE::FSdfChangeBlock Block;
 						AttributeCopy->Set( Value );
+						UsdUtils::NotifyIfOverriddenOpinion( *AttributeCopy );
 					}
 				})
 			];
diff --git a/Engine/Plugins/Importers/USDImporter/Source/USDStageEditorViewModels/Private/USDPrimAttributesViewModel.cpp b/Engine/Plugins/Importers/USDImporter/Source/USDStageEditorViewModels/Private/USDPrimAttributesViewModel.cpp
--- a/Engine/Plugins/Importers/USDImporter/Source/USDStageEditorViewModels/Private/USDPrimAttributesViewModel.cpp
+++ b/Engine/Plugins/Importers/USDImporter/Source/USDStageEditorViewModels/Private/USDPrimAttributesViewModel.cpp
@@ -3,6 +3,7 @@
 #include "USDPrimAttributesViewModel.h"
 
 #include "UnrealUSDWrapper.h"
+#include "USDAttributeUtils.h"
 #include "USDErrorUtils.h"
 #include "USDLog.h"
 #include "USDMemory.h"
@@ -11,6 +12,7 @@
 
 #include "UsdWrappers/SdfLayer.h"
 #include "UsdWrappers/SdfPath.h"
+#include "UsdWrappers/UsdAttribute.h"
 #include "UsdWrappers/UsdPrim.h"
 #include "UsdWrappers/VtValue.h"
 
@@ -173,6 +175,7 @@ void FUsdPrimAttributesViewModel::SetPrimAttribute( const FString& AttributeName
 			else if ( UE::FUsdAttribute Attribute = UsdPrim.GetAttribute( *AttributeName ) )
 			{
 				bSuccess = Attribute.Set( VtValue );
+				UsdUtils::NotifyIfOverriddenOpinion( Attribute );
 			}
 		}
 	}
diff --git a/Engine/Plugins/Importers/USDImporter/Source/USDUtilities/Private/USDAttributeUtils.cpp b/Engine/Plugins/Importers/USDImporter/Source/USDUtilities/Private/USDAttributeUtils.cpp
--- a/Engine/Plugins/Importers/USDImporter/Source/USDUtilities/Private/USDAttributeUtils.cpp
+++ b/Engine/Plugins/Importers/USDImporter/Source/USDUtilities/Private/USDAttributeUtils.cpp
@@ -6,15 +6,20 @@
 #include "USDErrorUtils.h"
 #include "USDLayerUtils.h"
 #include "USDLog.h"
+#include "USDProjectSettings.h"
 #include "USDTypesConversion.h"
 
 #include "UsdWrappers/SdfLayer.h"
 #include "UsdWrappers/SdfPath.h"
 #include "UsdWrappers/UsdAttribute.h"
+#include "UsdWrappers/UsdPrim.h"
 #include "UsdWrappers/UsdStage.h"
 
-#if USE_USD_SDK
+#include "Framework/Notifications/NotificationManager.h"
+#include "Templates/SharedPointer.h"
+#include "Widgets/Notifications/SNotificationList.h"
 
+#if USE_USD_SDK
 #include "USDIncludesStart.h"
 	#include "pxr/usd/pcp/cache.h"
 	#include "pxr/usd/pcp/primIndex.h"
@@ -23,18 +28,25 @@
 
 	#include "pxr/usd/sdf/layer.h"
 	#include "pxr/usd/sdf/layerUtils.h"
+	#include "pxr/usd/sdf/spec.h"
 	#include "pxr/usd/usd/attribute.h"
 	#include "pxr/usd/usd/editContext.h"
 	#include "pxr/usd/usd/stage.h"
 #include "USDIncludesEnd.h"
+#endif // #if USE_USD_SDK
+
+#define LOCTEXT_NAMESPACE "UsdAttributeUtils"
 
 namespace UsdUtils
 {
+#if USE_USD_SDK
 	const pxr::TfToken MutedToken = UnrealToUsd::ConvertToken( TEXT( "UE:Muted" ) ).Get();
+#endif // #if USE_USD_SDK
 }
 
 bool UsdUtils::MuteAttribute( UE::FUsdAttribute& Attribute, const UE::FUsdStage& Stage )
 {
+#if USE_USD_SDK
 	FScopedUsdAllocs Allocs;
 
 	const pxr::UsdAttribute& UsdAttribute = static_cast< const pxr::UsdAttribute& >( Attribute );
@@ -88,10 +100,14 @@ bool UsdUtils::MuteAttribute( UE::FUsdAttribute& Attribute, const UE::FUsdStage&
 	}
 
 	return true;
+#else
+	return false;
+#endif // #if USE_USD_SDK
 }
 
 bool UsdUtils::UnmuteAttribute( UE::FUsdAttribute& Attribute, const UE::FUsdStage& Stage )
 {
+#if USE_USD_SDK
 	FScopedUsdAllocs Allocs;
 
 	pxr::UsdAttribute& UsdAttribute = static_cast< pxr::UsdAttribute& >( Attribute );
@@ -133,10 +149,14 @@ bool UsdUtils::UnmuteAttribute( UE::FUsdAttribute& Attribute, const UE::FUsdStag
 	}
 
 	return true;
+#else
+	return false;
+#endif // #if USE_USD_SDK
 }
 
 bool UsdUtils::IsAttributeMuted( const UE::FUsdAttribute& Attribute, const UE::FUsdStage& Stage )
 {
+#if USE_USD_SDK
 	FScopedUsdAllocs Allocs;
 
 	const pxr::UsdAttribute& UsdAttribute = static_cast< const pxr::UsdAttribute& >( Attribute );
@@ -150,8 +170,120 @@ bool UsdUtils::IsAttributeMuted( const UE::FUsdAttribute& Attribute, const UE::F
 	{
 		return Data.Get<bool>();
 	}
+#endif // #if USE_USD_SDK
 
 	return false;
 }
 
+void UsdUtils::NotifyIfOverriddenOpinion( const UE::FUsdAttribute& Attribute )
+{
+#if USE_USD_SDK
+	FScopedUsdAllocs Allocs;
+
+	pxr::UsdAttribute UsdAttribute{Attribute};
+	if ( !UsdAttribute )
+	{
+		return;
+	}
+
+	pxr::UsdStageRefPtr Stage = UsdAttribute.GetPrim().GetStage();
+	if ( !Stage )
+	{
+		return;
+	}
+
+	const pxr::UsdEditTarget& EditTarget = Stage->GetEditTarget();
+	const pxr::SdfLayerHandle& Layer = EditTarget.GetLayer();
+	if ( !Layer )
+	{
+		return;
+	}
+
+	// Currently this will only warn in case our opinion and the strongest one both come from the local layer stack.
+	// This is good enough for us at this point though, because we can't edit outside of the local layer stack anyway,
+	// which due to LIVRPS is always stronger than the other composition arc types like references, payloads, etc.
+	// Alternatively this also means that any other opinion on a non-local composition arc is never going to be strong
+	// enough to override anything that we can author.
+	// References:
+	// - https://graphics.pixar.com/usd/release/glossary.html#livrps-strength-ordering
+	// - https://groups.google.com/g/usd-interest/c/xTxFYQA_bRs/m/qbGkvx3yAgAJ
+	std::vector<pxr::SdfPropertySpecHandle> SpecStack = UsdAttribute.GetPropertyStack();
+	for ( const pxr::SdfPropertySpecHandle& Spec : SpecStack )
+	{
+		if ( !Spec )
+		{
+			continue;
+		}
+
+		const pxr::SdfLayerHandle& SpecLayer = Spec->GetLayer();
+		if ( SpecLayer != Layer )
+		{
+			const FText Text = LOCTEXT( "OverridenOpinionText", "USD: Overridden opinion" );
+
+			const FText SubText = FText::Format(
+				LOCTEXT( "OverridenOpinionSubText", "Opinion authored for this attribute:\n\n{0}\n\nAt this layer:\n\n{1}\n\nIs overridden by another spec at this layer:\n\n{2}\n\nAnd so may not be visible on the composed stage. This means this edit may not be visible once the stage is reloaded." ),
+				FText::FromString( UsdToUnreal::ConvertPath( Spec->GetPath() ) ),
+				FText::FromString( UsdToUnreal::ConvertString( Layer->GetIdentifier() ) ),
+				FText::FromString( UsdToUnreal::ConvertString( SpecLayer->GetIdentifier() ) )
+			);
+
+			UE_LOG( LogUsd, Warning, TEXT( "%s" ), *SubText.ToString().Replace( TEXT( "\n\n" ), TEXT( " " ) ) );
+
+			const UUsdProjectSettings* Settings = GetDefault<UUsdProjectSettings>();
+			if ( Settings && Settings->bShowOverriddenOpinionsWarning )
+			{
+				static TWeakPtr<SNotificationItem> Notification;
+
+				FNotificationInfo Toast( Text );
+				Toast.SubText = SubText;
+				Toast.Image = FCoreStyle::Get().GetBrush( TEXT( "MessageLog.Warning" ) );
+				Toast.CheckBoxText = LOCTEXT("DontAskAgain", "Don't prompt again");
+				Toast.bUseLargeFont = false;
+				Toast.bFireAndForget = false;
+				Toast.FadeOutDuration = 0.0f;
+				Toast.ExpireDuration = 0.0f;
+				Toast.bUseThrobber = false;
+				Toast.bUseSuccessFailIcons = false;
+				Toast.ButtonDetails.Emplace(
+					LOCTEXT( "OverridenOpinionMessageOk", "Ok" ),
+					FText::GetEmpty(),
+					FSimpleDelegate::CreateLambda([](){
+						if ( TSharedPtr<SNotificationItem> PinnedNotification = Notification.Pin() )
+						{
+							PinnedNotification->SetCompletionState( SNotificationItem::CS_Success );
+							PinnedNotification->ExpireAndFadeout();
+						}
+					})
+				);
+				// This is flipped because the default checkbox message is "Don't prompt again"
+				Toast.CheckBoxState = Settings->bShowOverriddenOpinionsWarning ? ECheckBoxState::Unchecked : ECheckBoxState::Checked;
+				Toast.CheckBoxStateChanged = FOnCheckStateChanged::CreateStatic([](ECheckBoxState NewState)
+				{
+					if ( UUsdProjectSettings* Settings = GetMutableDefault<UUsdProjectSettings>() )
+					{
+						// This is flipped because the default checkbox message is "Don't prompt again"
+						Settings->bShowOverriddenOpinionsWarning = NewState == ECheckBoxState::Unchecked;
+					}
+				});
+
+				// Only show one at a time
+				if ( !Notification.IsValid() )
+				{
+					Notification = FSlateNotificationManager::Get().AddNotification( Toast );
+				}
+
+				if ( TSharedPtr<SNotificationItem> PinnedNotification = Notification.Pin() )
+				{
+					PinnedNotification->SetCompletionState( SNotificationItem::CS_Pending );
+				}
+			}
+		}
+		else
+		{
+			break;
+		}
+	}
 #endif // #if USE_USD_SDK
+}
+
+#undef LOCTEXT_NAMESPACE
\ No newline at end of file
diff --git a/Engine/Plugins/Importers/USDImporter/Source/USDUtilities/Private/USDGeomMeshConversion.cpp b/Engine/Plugins/Importers/USDImporter/Source/USDUtilities/Private/USDGeomMeshConversion.cpp
--- a/Engine/Plugins/Importers/USDImporter/Source/USDUtilities/Private/USDGeomMeshConversion.cpp
+++ b/Engine/Plugins/Importers/USDImporter/Source/USDUtilities/Private/USDGeomMeshConversion.cpp
@@ -6,6 +6,7 @@
 
 #include "UnrealUSDWrapper.h"
 #include "USDAssetImportData.h"
+#include "USDAttributeUtils.h"
 #include "USDConversionUtils.h"
 #include "USDErrorUtils.h"
 #include "USDLog.h"
@@ -16,6 +17,7 @@
 
 #include "UsdWrappers/SdfLayer.h"
 #include "UsdWrappers/SdfPath.h"
+#include "UsdWrappers/UsdAttribute.h"
 #include "UsdWrappers/UsdPrim.h"
 #include "UsdWrappers/UsdStage.h"
 
@@ -1923,6 +1925,7 @@ bool UnrealToUsd::ConvertStaticMesh( const UStaticMesh* StaticMesh, pxr::UsdPrim
 		if ( pxr::UsdAttribute Attr = UsdPrim.CreateAttribute( UnrealIdentifiers::UnrealNaniteOverride, pxr::SdfValueTypeNames->Token ) )
 		{
 			Attr.Set( UnrealIdentifiers::UnrealNaniteOverrideEnable );
+			UsdUtils::NotifyIfOverriddenOpinion( UE::FUsdAttribute{ Attr } );
 		}
 	}
 #endif // WITH_EDITOR
diff --git a/Engine/Plugins/Importers/USDImporter/Source/USDUtilities/Private/USDIntegrationUtils.cpp b/Engine/Plugins/Importers/USDImporter/Source/USDUtilities/Private/USDIntegrationUtils.cpp
--- a/Engine/Plugins/Importers/USDImporter/Source/USDUtilities/Private/USDIntegrationUtils.cpp
+++ b/Engine/Plugins/Importers/USDImporter/Source/USDUtilities/Private/USDIntegrationUtils.cpp
@@ -3,8 +3,10 @@
 #include "USDIntegrationUtils.h"
 
 #include "UnrealUSDWrapper.h"
+#include "USDAttributeUtils.h"
 #include "USDTypesConversion.h"
 #include "UsdWrappers/SdfChangeBlock.h"
+#include "UsdWrappers/UsdAttribute.h"
 
 #include "Animation/AnimInstance.h"
 #include "Components/SkeletalMeshComponent.h"
@@ -128,12 +130,14 @@ void UnrealToUsd::ConvertLiveLinkProperties( const UActorComponent* InComponent,
 				: FString();
 
 			Attr.Set( UnrealToUsd::ConvertString( *AnimBPPath ).Get(), pxr::UsdTimeCode::Default() );
+			UsdUtils::NotifyIfOverriddenOpinion( UE::FUsdAttribute{ Attr } );
 		}
 
 		if ( pxr::UsdAttribute Attr = InOutPrim.CreateAttribute( UnrealIdentifiers::UnrealLiveLinkEnabled, pxr::SdfValueTypeNames->Bool ) )
 		{
 			const bool bEnabled = SkeletalComponent->GetAnimationMode() == EAnimationMode::AnimationBlueprint;
 			Attr.Set( bEnabled, pxr::UsdTimeCode::Default() );
+			UsdUtils::NotifyIfOverriddenOpinion( UE::FUsdAttribute{ Attr } );
 		}
 	}
 	// Non-skeletal LiveLink case
@@ -142,11 +146,13 @@ void UnrealToUsd::ConvertLiveLinkProperties( const UActorComponent* InComponent,
 		if ( pxr::UsdAttribute Attr = InOutPrim.CreateAttribute( UnrealIdentifiers::UnrealLiveLinkSubjectName, pxr::SdfValueTypeNames->String ) )
 		{
 			Attr.Set( UnrealToUsd::ConvertString( *InController->SubjectRepresentation.Subject.Name.ToString() ).Get(), pxr::UsdTimeCode::Default() );
+			UsdUtils::NotifyIfOverriddenOpinion( UE::FUsdAttribute{ Attr } );
 		}
 
 		if ( pxr::UsdAttribute Attr = InOutPrim.CreateAttribute( UnrealIdentifiers::UnrealLiveLinkEnabled, pxr::SdfValueTypeNames->Bool ) )
 		{
 			Attr.Set( InController->bEvaluateLiveLink, pxr::UsdTimeCode::Default() );
+			UsdUtils::NotifyIfOverriddenOpinion( UE::FUsdAttribute{ Attr } );
 		}
 	}
 #endif // WITH_EDITOR
diff --git a/Engine/Plugins/Importers/USDImporter/Source/USDUtilities/Private/USDLightConversion.cpp b/Engine/Plugins/Importers/USDImporter/Source/USDUtilities/Private/USDLightConversion.cpp
--- a/Engine/Plugins/Importers/USDImporter/Source/USDUtilities/Private/USDLightConversion.cpp
+++ b/Engine/Plugins/Importers/USDImporter/Source/USDUtilities/Private/USDLightConversion.cpp
@@ -3,6 +3,7 @@
 #include "USDLightConversion.h"
 
 #include "USDAssetCache.h"
+#include "USDAttributeUtils.h"
 #include "USDConversionUtils.h"
 #include "USDLayerUtils.h"
 #include "USDLog.h"
@@ -11,6 +12,7 @@
 
 #include "UsdWrappers/SdfLayer.h"
 #include "UsdWrappers/SdfPath.h"
+#include "UsdWrappers/UsdAttribute.h"
 #include "UsdWrappers/UsdPrim.h"
 #include "UsdWrappers/UsdStage.h"
 
@@ -413,6 +415,7 @@ bool UnrealToUsd::ConvertLightComponent( const ULightComponentBase& LightCompone
 	if ( pxr::UsdAttribute Attr = LightAPI.CreateIntensityAttr() )
 	{
 		Attr.Set<float>( LightComponent.Intensity, UsdTimeCode );
+		UsdUtils::NotifyIfOverriddenOpinion( UE::FUsdAttribute{ Attr } );
 	}
 
 	// When converting into UE we multiply intensity and exposure together, so when writing back we just
@@ -420,6 +423,7 @@ bool UnrealToUsd::ConvertLightComponent( const ULightComponentBase& LightCompone
 	if ( pxr::UsdAttribute Attr = LightAPI.CreateExposureAttr() )
 	{
 		Attr.Set<float>( 0.0f, UsdTimeCode );
+		UsdUtils::NotifyIfOverriddenOpinion( UE::FUsdAttribute{ Attr } );
 	}
 
 	if ( const ULightComponent* DerivedLightComponent = Cast<ULightComponent>( &LightComponent ) )
@@ -427,11 +431,13 @@ bool UnrealToUsd::ConvertLightComponent( const ULightComponentBase& LightCompone
 		if ( pxr::UsdAttribute Attr = LightAPI.CreateEnableColorTemperatureAttr() )
 		{
 			Attr.Set<bool>( DerivedLightComponent->bUseTemperature, UsdTimeCode );
+			UsdUtils::NotifyIfOverriddenOpinion( UE::FUsdAttribute{ Attr } );
 		}
 
 		if ( pxr::UsdAttribute Attr = LightAPI.CreateColorTemperatureAttr() )
 		{
 			Attr.Set<float>( DerivedLightComponent->Temperature, UsdTimeCode );
+			UsdUtils::NotifyIfOverriddenOpinion( UE::FUsdAttribute{ Attr } );
 		}
 	}
 
@@ -439,6 +445,7 @@ bool UnrealToUsd::ConvertLightComponent( const ULightComponentBase& LightCompone
 	{
 		pxr::GfVec4f LinearColor = UnrealToUsd::ConvertColor( LightComponent.LightColor );
 		Attr.Set<pxr::GfVec3f>( pxr::GfVec3f( LinearColor[0], LinearColor[1], LinearColor[2] ), UsdTimeCode );
+		UsdUtils::NotifyIfOverriddenOpinion( UE::FUsdAttribute{ Attr } );
 	}
 
 	return true;
@@ -462,6 +469,7 @@ bool UnrealToUsd::ConvertDirectionalLightComponent( const UDirectionalLightCompo
 	if ( pxr::UsdAttribute Attr = Light.CreateAngleAttr() )
 	{
 		Attr.Set<float>( LightComponent.LightSourceAngle, UsdTimeCode );
+		UsdUtils::NotifyIfOverriddenOpinion( UE::FUsdAttribute{ Attr } );
 	}
 
 	// USD intensity units should be in Nits == Lux / Steradian, but there is no
@@ -498,6 +506,7 @@ bool UnrealToUsd::ConvertRectLightComponent( const URectLightComponent& LightCom
 		if ( pxr::UsdAttribute Attr = DiskLight.CreateRadiusAttr() )
 		{
 			Attr.Set<float>( UnrealToUsd::ConvertDistance( StageInfo, Radius ), UsdTimeCode );
+			UsdUtils::NotifyIfOverriddenOpinion( UE::FUsdAttribute{ Attr } );
 		}
 	}
 	// Rect light
@@ -508,11 +517,13 @@ bool UnrealToUsd::ConvertRectLightComponent( const URectLightComponent& LightCom
 		if ( pxr::UsdAttribute Attr = RectLight.CreateWidthAttr() )
 		{
 			Attr.Set<float>( UnrealToUsd::ConvertDistance( StageInfo, LightComponent.SourceWidth ), UsdTimeCode );
+			UsdUtils::NotifyIfOverriddenOpinion( UE::FUsdAttribute{ Attr } );
 		}
 
 		if ( pxr::UsdAttribute Attr = RectLight.CreateHeightAttr() )
 		{
 			Attr.Set<float>( UnrealToUsd::ConvertDistance( StageInfo, LightComponent.SourceHeight ), UsdTimeCode );
+			UsdUtils::NotifyIfOverriddenOpinion( UE::FUsdAttribute{ Attr } );
 		}
 	}
 	else
@@ -558,6 +569,7 @@ bool UnrealToUsd::ConvertRectLightComponent( const URectLightComponent& LightCom
 		}
 
 		Attr.Set<float>( FinalIntensityNits, UsdTimeCode );
+		UsdUtils::NotifyIfOverriddenOpinion( UE::FUsdAttribute{ Attr } );
 	}
 
 	return true;
@@ -583,11 +595,13 @@ bool UnrealToUsd::ConvertPointLightComponent( const UPointLightComponent& LightC
 	if ( pxr::UsdAttribute Attr = Light.CreateRadiusAttr() )
 	{
 		Attr.Set<float>( UnrealToUsd::ConvertDistance( StageInfo, LightComponent.SourceRadius ), UsdTimeCode );
+		UsdUtils::NotifyIfOverriddenOpinion( UE::FUsdAttribute{ Attr } );
 	}
 
 	if ( pxr::UsdAttribute Attr = Light.CreateTreatAsPointAttr() )
 	{
 		Attr.Set<bool>( FMath::IsNearlyZero( LightComponent.SourceRadius ), UsdTimeCode );
+		UsdUtils::NotifyIfOverriddenOpinion( UE::FUsdAttribute{ Attr } );
 	}
 
 	float SolidAngle = 4.f * PI;
@@ -625,6 +639,7 @@ bool UnrealToUsd::ConvertPointLightComponent( const UPointLightComponent& LightC
 		}
 
 		Attr.Set<float>( FinalIntensityNits, UsdTimeCode );
+		UsdUtils::NotifyIfOverriddenOpinion( UE::FUsdAttribute{ Attr } );
 	}
 
 	return true;
@@ -666,6 +681,7 @@ bool UnrealToUsd::ConvertSkyLightComponent( const USkyLightComponent& LightCompo
 
 				UsdUtils::MakePathRelativeToLayer( UE::FSdfLayer{ Prim.GetStage()->GetEditTarget().GetLayer() }, FilePath );
 				Attr.Set<pxr::SdfAssetPath>( pxr::SdfAssetPath{ UnrealToUsd::ConvertString( *FilePath ).Get() }, UsdTimeCode );
+				UsdUtils::NotifyIfOverriddenOpinion( UE::FUsdAttribute{ Attr } );
 			}
 		}
 	}
@@ -687,6 +703,7 @@ bool UnrealToUsd::ConvertSpotLightComponent( const USpotLightComponent& LightCom
 	if ( pxr::UsdAttribute ConeAngleAttr = ShapingAPI.CreateShapingConeAngleAttr() )
 	{
 		ConeAngleAttr.Set<float>( LightComponent.OuterConeAngle, UsdTimeCode );
+		UsdUtils::NotifyIfOverriddenOpinion( UE::FUsdAttribute{ ConeAngleAttr } );
 	}
 
 	// As of March 2021 there doesn't seem to be a consensus on what softness means, according to https://groups.google.com/g/usd-interest/c/A6bc4OZjSB0
@@ -696,6 +713,7 @@ bool UnrealToUsd::ConvertSpotLightComponent( const USpotLightComponent& LightCom
 		// Keep in [0, 1] range, where 1 is maximum softness, i.e. inner cone angle is zero
 		const float Softness = FMath::IsNearlyZero( LightComponent.OuterConeAngle ) ? 0.0 : 1.0f - LightComponent.InnerConeAngle / LightComponent.OuterConeAngle;
 		SoftnessAttr.Set<float>( Softness, UsdTimeCode );
+		UsdUtils::NotifyIfOverriddenOpinion( UE::FUsdAttribute{ SoftnessAttr } );
 	}
 
 	return true;
diff --git a/Engine/Plugins/Importers/USDImporter/Source/USDUtilities/Private/USDPrimConversion.cpp b/Engine/Plugins/Importers/USDImporter/Source/USDUtilities/Private/USDPrimConversion.cpp
--- a/Engine/Plugins/Importers/USDImporter/Source/USDUtilities/Private/USDPrimConversion.cpp
+++ b/Engine/Plugins/Importers/USDImporter/Source/USDUtilities/Private/USDPrimConversion.cpp
@@ -3,6 +3,7 @@
 #include "USDPrimConversion.h"
 
 #include "UnrealUSDWrapper.h"
+#include "USDAttributeUtils.h"
 #include "USDConversionUtils.h"
 #include "USDLayerUtils.h"
 #include "USDLightConversion.h"
@@ -1525,26 +1526,31 @@ bool UnrealToUsd::ConvertCameraComponent( const UCineCameraComponent& CameraComp
 	if ( pxr::UsdAttribute Attr = GeomCamera.CreateFocalLengthAttr() )
 	{
 		Attr.Set<float>( UnrealToUsd::ConvertDistance( StageInfo, CameraComponent.CurrentFocalLength ), UsdTimeCode );
+		UsdUtils::NotifyIfOverriddenOpinion( UE::FUsdAttribute{ Attr } );
 	}
 
 	if ( pxr::UsdAttribute Attr = GeomCamera.CreateFocusDistanceAttr() )
 	{
 		Attr.Set<float>( UnrealToUsd::ConvertDistance( StageInfo, CameraComponent.FocusSettings.ManualFocusDistance ), UsdTimeCode );
+		UsdUtils::NotifyIfOverriddenOpinion( UE::FUsdAttribute{ Attr } );
 	}
 
 	if ( pxr::UsdAttribute Attr = GeomCamera.CreateFStopAttr() )
 	{
 		Attr.Set<float>( CameraComponent.CurrentAperture, UsdTimeCode );
+		UsdUtils::NotifyIfOverriddenOpinion( UE::FUsdAttribute{ Attr } );
 	}
 
 	if ( pxr::UsdAttribute Attr = GeomCamera.CreateHorizontalApertureAttr() )
 	{
 		Attr.Set<float>( UnrealToUsd::ConvertDistance( StageInfo, CameraComponent.Filmback.SensorWidth ), UsdTimeCode );
+		UsdUtils::NotifyIfOverriddenOpinion( UE::FUsdAttribute{ Attr } );
 	}
 
 	if ( pxr::UsdAttribute Attr = GeomCamera.CreateVerticalApertureAttr() )
 	{
 		Attr.Set<float>( UnrealToUsd::ConvertDistance( StageInfo, CameraComponent.Filmback.SensorHeight ), UsdTimeCode );
+		UsdUtils::NotifyIfOverriddenOpinion( UE::FUsdAttribute{ Attr } );
 	}
 
 	return true;
@@ -2095,6 +2101,7 @@ bool UnrealToUsd::ConvertSceneComponent( const pxr::UsdStageRefPtr& Stage, const
 		}
 
 		VisibilityAttr.Set<pxr::TfToken>( Value );
+		UsdUtils::NotifyIfOverriddenOpinion( UE::FUsdAttribute{ VisibilityAttr } );
 	}
 
 	return true;
@@ -2361,21 +2368,25 @@ bool UnrealToUsd::ConvertHierarchicalInstancedStaticMeshComponent( const UHierar
 	if ( UsdAttribute Attr = PointInstancer.CreateProtoIndicesAttr() )
 	{
 		Attr.Set( ProtoIndices, UsdTimeCode );
+		UsdUtils::NotifyIfOverriddenOpinion( UE::FUsdAttribute{ Attr } );
 	}
 
 	if ( UsdAttribute Attr = PointInstancer.CreatePositionsAttr() )
 	{
 		Attr.Set( Positions, UsdTimeCode );
+		UsdUtils::NotifyIfOverriddenOpinion( UE::FUsdAttribute{ Attr } );
 	}
 
 	if ( UsdAttribute Attr = PointInstancer.CreateOrientationsAttr() )
 	{
 		Attr.Set( Orientations, UsdTimeCode );
+		UsdUtils::NotifyIfOverriddenOpinion( UE::FUsdAttribute{ Attr } );
 	}
 
 	if ( UsdAttribute Attr = PointInstancer.CreateScalesAttr() )
 	{
 		Attr.Set( Scales, UsdTimeCode );
+		UsdUtils::NotifyIfOverriddenOpinion( UE::FUsdAttribute{ Attr } );
 	}
 
 	return true;
@@ -2415,6 +2426,9 @@ bool UnrealToUsd::ConvertXformable( const FTransform& RelativeTransform, pxr::Us
 	if ( pxr::UsdGeomXformOp MatrixXform = UE::USDPrimConversionImpl::Private::ForceMatrixXform( XForm ) )
 	{
 		MatrixXform.Set( UsdTransform, UsdTimeCode );
+
+		UsdUtils::NotifyIfOverriddenOpinion( UE::FUsdAttribute{ MatrixXform.GetAttr() } );
+		UsdUtils::NotifyIfOverriddenOpinion( UE::FUsdAttribute{ XForm.GetXformOpOrderAttr() } );
 	}
 
 	return true;
@@ -2551,21 +2565,25 @@ bool UnrealToUsd::ConvertInstancedFoliageActor( const AInstancedFoliageActor& Ac
 	if ( UsdAttribute Attr = PointInstancer.CreateProtoIndicesAttr() )
 	{
 		Attr.Set( ProtoIndices, UsdTimeCode );
+		UsdUtils::NotifyIfOverriddenOpinion( UE::FUsdAttribute{ Attr } );
 	}
 
 	if ( UsdAttribute Attr = PointInstancer.CreatePositionsAttr() )
 	{
 		Attr.Set( Positions, UsdTimeCode );
+		UsdUtils::NotifyIfOverriddenOpinion( UE::FUsdAttribute{ Attr } );
 	}
 
 	if ( UsdAttribute Attr = PointInstancer.CreateOrientationsAttr() )
 	{
 		Attr.Set( Orientations, UsdTimeCode );
+		UsdUtils::NotifyIfOverriddenOpinion( UE::FUsdAttribute{ Attr } );
 	}
 
 	if ( UsdAttribute Attr = PointInstancer.CreateScalesAttr() )
 	{
 		Attr.Set( Scales, UsdTimeCode );
+		UsdUtils::NotifyIfOverriddenOpinion( UE::FUsdAttribute{ Attr } );
 	}
 
 	return true;
@@ -3431,6 +3449,8 @@ UnrealToUsd::FPropertyTrackWriter UnrealToUsd::CreatePropertyTrackWriter( const
 		{
 			Attr.ClearAtTime( TimeSample );
 		}
+
+		UsdUtils::NotifyIfOverriddenOpinion( UE::FUsdAttribute{ Attr } );
 	}
 
 	return Result;
diff --git a/Engine/Plugins/Importers/USDImporter/Source/USDUtilities/Private/USDShadeConversion.cpp b/Engine/Plugins/Importers/USDImporter/Source/USDUtilities/Private/USDShadeConversion.cpp
--- a/Engine/Plugins/Importers/USDImporter/Source/USDUtilities/Private/USDShadeConversion.cpp
+++ b/Engine/Plugins/Importers/USDImporter/Source/USDUtilities/Private/USDShadeConversion.cpp
@@ -6,6 +6,7 @@
 
 #include "USDAssetCache.h"
 #include "USDAssetImportData.h"
+#include "USDAttributeUtils.h"
 #include "USDConversionUtils.h"
 #include "USDErrorUtils.h"
 #include "USDLayerUtils.h"
@@ -14,6 +15,7 @@
 #include "USDTypesConversion.h"
 
 #include "UsdWrappers/SdfLayer.h"
+#include "UsdWrappers/UsdAttribute.h"
 #include "UsdWrappers/UsdPrim.h"
 
 #include "Algo/AllOf.h"
@@ -2348,6 +2350,7 @@ bool UsdUtils::MarkMaterialPrimWithWorldSpaceNormals( const UE::FUsdPrim& Materi
 	}
 
 	Attr.Set<bool>( true );
+	UsdUtils::NotifyIfOverriddenOpinion( UE::FUsdAttribute{ Attr } );
 	return true;
 }
 
diff --git a/Engine/Plugins/Importers/USDImporter/Source/USDUtilities/Private/USDSkeletalDataConversion.cpp b/Engine/Plugins/Importers/USDImporter/Source/USDUtilities/Private/USDSkeletalDataConversion.cpp
--- a/Engine/Plugins/Importers/USDImporter/Source/USDUtilities/Private/USDSkeletalDataConversion.cpp
+++ b/Engine/Plugins/Importers/USDImporter/Source/USDUtilities/Private/USDSkeletalDataConversion.cpp
@@ -4,6 +4,7 @@
 
 #include "UnrealUSDWrapper.h"
 #include "USDAssetImportData.h"
+#include "USDAttributeUtils.h"
 #include "USDConversionUtils.h"
 #include "USDErrorUtils.h"
 #include "USDGeomMeshConversion.h"
@@ -14,6 +15,7 @@
 
 #include "UsdWrappers/SdfLayer.h"
 #include "UsdWrappers/SdfPath.h"
+#include "UsdWrappers/UsdAttribute.h"
 #include "UsdWrappers/UsdPrim.h"
 #include "UsdWrappers/UsdStage.h"
 
@@ -3269,6 +3271,11 @@ bool UnrealToUsd::ConvertControlRigSection(
 		EndInclTickFrame = UpperBoundToUse.GetValue().GetValue() + ( UpperBoundToUse.GetValue().IsInclusive() ? 0 : -1 );
 	}
 
+	UsdUtils::NotifyIfOverriddenOpinion( UE::FUsdAttribute{ BlendShapeWeightsAttr } );
+	UsdUtils::NotifyIfOverriddenOpinion( UE::FUsdAttribute{ TranslationsAttr } );
+	UsdUtils::NotifyIfOverriddenOpinion( UE::FUsdAttribute{ RotationsAttr } );
+	UsdUtils::NotifyIfOverriddenOpinion( UE::FUsdAttribute{ ScalesAttr } );
+
 	pxr::VtArray<pxr::TfToken> CurveNames;
 	pxr::VtArray<float> CurvesValuesForTime;
 
@@ -3401,6 +3408,7 @@ bool UnrealToUsd::ConvertControlRigSection(
 						if ( bRenamedAChannel )
 						{
 							ChannelsAttr.Set( BlendShapeChannels );
+							UsdUtils::NotifyIfOverriddenOpinion( UE::FUsdAttribute{ ChannelsAttr } );
 						}
 					}
 				}
@@ -3412,6 +3420,8 @@ bool UnrealToUsd::ConvertControlRigSection(
 		BlendShapeWeightsAttr.Clear();
 		BlendShapesAttr.Set( CurveNames );
 		CurvesValuesForTime.resize( CurveNames.size() );
+
+		UsdUtils::NotifyIfOverriddenOpinion( UE::FUsdAttribute{ BlendShapesAttr } );
 	}
 
 	FFrameTime TickIncr = FFrameRate::TransformTime( 1, DisplayRate, TickResolution );
diff --git a/Engine/Plugins/Importers/USDImporter/Source/USDUtilities/Public/USDAttributeUtils.h b/Engine/Plugins/Importers/USDImporter/Source/USDUtilities/Public/USDAttributeUtils.h
--- a/Engine/Plugins/Importers/USDImporter/Source/USDUtilities/Public/USDAttributeUtils.h
+++ b/Engine/Plugins/Importers/USDImporter/Source/USDUtilities/Public/USDAttributeUtils.h
@@ -6,8 +6,6 @@
 
 #include "UsdWrappers/ForwardDeclarations.h"
 
-#if USE_USD_SDK
-
 namespace UE
 {
 	class FUsdAttribute;
@@ -23,6 +21,15 @@ namespace UsdUtils
 
 	/** Returns whether this attribute has the "Muted" CustomData on the stage's UE state sublayer, meaning it shouldn't be animated when loaded into UE */
 	USDUTILITIES_API bool IsAttributeMuted( const UE::FUsdAttribute& Attribute, const UE::FUsdStage& Stage );
+
+	/**
+	 * If Attribute has opinions authored on layers stronger than the current edit target this will emit a warning to
+	 * the user, indicating that the new value may not be visible on the USD Stage, and may cease to be visible on
+	 * the UE level after a reload.
+	 *
+	 * Use this after setting any attribute: This function does nothing in case the Stage's current edit target has
+	 * the strongest opinion for the attribute already
+	 */
+	USDUTILITIES_API void NotifyIfOverriddenOpinion( const UE::FUsdAttribute& Attribute );
 }
 
-#endif // #if USE_USD_SDK
