Build fix: Missing updates from CL 21410131.

#rnx

diff --git a/Engine/Plugins/Importers/USDImporter/Source/USDSchemas/Private/USDSkelRootTranslator.cpp b/Engine/Plugins/Importers/USDImporter/Source/USDSchemas/Private/USDSkelRootTranslator.cpp
--- a/Engine/Plugins/Importers/USDImporter/Source/USDSchemas/Private/USDSkelRootTranslator.cpp
+++ b/Engine/Plugins/Importers/USDImporter/Source/USDSchemas/Private/USDSkelRootTranslator.cpp
@@ -288,20 +288,16 @@ namespace UsdSkelRootTranslatorImpl
 			return;
 		}
 
-		UMorphTarget* MorphTarget = SkeletalMeshComponent.GetSkeletalMeshAsset()->GetMorphTargets()[ IndexInSkeletalMesh ];
+		const UMorphTarget* MorphTarget = SkeletalMeshComponent.GetSkeletalMeshAsset()->GetMorphTargets()[ IndexInSkeletalMesh ];
 		if ( !MorphTarget )
 		{
 			return;
 		}
 
 		int32 WeightIndex = INDEX_NONE;
-		if ( SkeletalMeshComponent.ActiveMorphTargets.IsValidIndex( IndexInSkeletalMesh ) )
+		if ( SkeletalMeshComponent.ActiveMorphTargets.Contains( MorphTarget ) )
 		{
-			FActiveMorphTarget& ActiveMorphTarget = SkeletalMeshComponent.ActiveMorphTargets[ IndexInSkeletalMesh ];
-			if ( ActiveMorphTarget.MorphTarget == MorphTarget )
-			{
-				WeightIndex = ActiveMorphTarget.WeightIndex;
-			}
+			WeightIndex = SkeletalMeshComponent.ActiveMorphTargets[ MorphTarget ];
 		}
 
 		// Morph target is not at expected location (i.e. after CreateComponents, duplicate for PIE or undo/redo) --> Rebuild ActiveMorphTargets
@@ -313,11 +309,7 @@ namespace UsdSkelRootTranslatorImpl
 			TArray<UMorphTarget*>& MorphTargets = SkeletalMesh->GetMorphTargets();
 			for ( int32 MorphTargetIndex = 0; MorphTargetIndex < MorphTargets.Num(); ++MorphTargetIndex )
 			{
-				FActiveMorphTarget ActiveMorphTarget;
-				ActiveMorphTarget.MorphTarget = MorphTargets[ MorphTargetIndex ];
-				ActiveMorphTarget.WeightIndex = MorphTargetIndex;
-
-				SkeletalMeshComponent.ActiveMorphTargets.Add( ActiveMorphTarget );
+				SkeletalMeshComponent.ActiveMorphTargets.Add( MorphTargets[ MorphTargetIndex ], MorphTargetIndex );
 				SkeletalMeshComponent.MorphTargetWeights.Add( 0.0f ); // We'll update these right afterwards when we call UpdateComponents
 			}
 
diff --git a/Engine/Plugins/Importers/USDImporter/Source/USDUtilities/Private/USDPrimConversion.cpp b/Engine/Plugins/Importers/USDImporter/Source/USDUtilities/Private/USDPrimConversion.cpp
--- a/Engine/Plugins/Importers/USDImporter/Source/USDUtilities/Private/USDPrimConversion.cpp
+++ b/Engine/Plugins/Importers/USDImporter/Source/USDUtilities/Private/USDPrimConversion.cpp
@@ -2878,11 +2878,7 @@ bool UnrealToUsd::CreateSkeletalAnimationBaker( UE::FUsdPrim& SkelRoot, UE::FUsd
 	{
 		for ( int32 MorphTargetIndex = 0; MorphTargetIndex < MorphTargets.Num(); ++MorphTargetIndex )
 		{
-			FActiveMorphTarget ActiveMorphTarget;
-			ActiveMorphTarget.MorphTarget = MorphTargets[ MorphTargetIndex ];
-			ActiveMorphTarget.WeightIndex = MorphTargetIndex;
-
-			Component.ActiveMorphTargets.Add( ActiveMorphTarget );
+			Component.ActiveMorphTargets.Add( MorphTargets[ MorphTargetIndex ], MorphTargetIndex );
 		}
 	}
 
@@ -2898,19 +2894,19 @@ bool UnrealToUsd::CreateSkeletalAnimationBaker( UE::FUsdPrim& SkelRoot, UE::FUsd
 		BlendShapeWeightsAttr = UsdSkelAnimation.CreateBlendShapeWeightsAttr();
 		BlendShapesAttr = UsdSkelAnimation.CreateBlendShapesAttr();
 
-		TArray<FActiveMorphTarget> SortedMorphTargets = Component.ActiveMorphTargets;
-		SortedMorphTargets.Sort([](const FActiveMorphTarget& Left, const FActiveMorphTarget& Right)
+		TArray<FMorphTargetWeightMap::ElementType> SortedMorphTargets = Component.ActiveMorphTargets.Array();
+		SortedMorphTargets.Sort([](const FMorphTargetWeightMap::ElementType& Left, const FMorphTargetWeightMap::ElementType& Right)
 		{
-			return Left.WeightIndex < Right.WeightIndex;
+			return Left.Value < Right.Value;
 		});
 
 		pxr::VtArray< pxr::TfToken > BlendShapeNames;
 		BlendShapeNames.reserve(SortedMorphTargets.Num());
 
-		for ( const FActiveMorphTarget& ActiveMorphTarget : SortedMorphTargets )
+		for ( const FMorphTargetWeightMap::ElementType& ActiveMorphTarget : SortedMorphTargets )
 		{
 			FString BlendShapeName;
-			if ( UMorphTarget* MorphTarget = ActiveMorphTarget.MorphTarget )
+			if ( const UMorphTarget* MorphTarget = ActiveMorphTarget.Key )
 			{
 				BlendShapeName = MorphTarget->GetFName().ToString();
 			}
