USD: Prompt user if they want to save dirty USD layers when hitting Save All or Close.

diff --git a/Engine/Plugins/Importers/USDImporter/Source/USDClasses/Public/USDProjectSettings.h b/Engine/Plugins/Importers/USDImporter/Source/USDClasses/Public/USDProjectSettings.h
--- a/Engine/Plugins/Importers/USDImporter/Source/USDClasses/Public/USDProjectSettings.h
+++ b/Engine/Plugins/Importers/USDImporter/Source/USDClasses/Public/USDProjectSettings.h
@@ -8,6 +8,14 @@
 
 #include "USDProjectSettings.generated.h"
 
+UENUM( BlueprintType )
+enum class EUsdSaveDialogBehavior : uint8
+{
+	NeverSave,
+	AlwaysSave,
+	ShowPrompt
+};
+
 UCLASS(config=Editor, meta=(DisplayName=USDImporter), MinimalAPI)
 class UUsdProjectSettings : public UDeveloperSettings
 {
@@ -18,10 +26,18 @@ public:
 	UPROPERTY( config, EditAnywhere, Category = USD )
 	TArray<FDirectoryPath> AdditionalPluginDirectories;
 
-	UPROPERTY( config, EditAnywhere, Category = USD )
+	UPROPERTY( config, EditAnywhere, Category = "USD|Dialogs" )
 	bool bShowConfirmationWhenClearingLayers = true;
 
 	// Whether to show the warning dialog when authoring opinions that could have no effect on the composed stage
-	UPROPERTY( config, EditAnywhere, Category = USD )
+	UPROPERTY( config, EditAnywhere, Category = "USD|Dialogs" )
 	bool bShowOverriddenOpinionsWarning = true;
+
+	// Whether to display the pop up dialog asking what to do about dirty USD layers when saving the UE level
+	UPROPERTY(config, EditAnywhere, Category = "USD|Dialogs" )
+	EUsdSaveDialogBehavior ShowSaveLayersDialogWhenSaving = EUsdSaveDialogBehavior::ShowPrompt;
+
+	// Whether to display the pop up dialog asking what to do about dirty USD layers when closing USD stages
+	UPROPERTY(config, EditAnywhere, Category = "USD|Dialogs" )
+	EUsdSaveDialogBehavior ShowSaveLayersDialogWhenClosing = EUsdSaveDialogBehavior::ShowPrompt;
 };
\ No newline at end of file
diff --git a/Engine/Plugins/Importers/USDImporter/Source/USDStage/Private/USDStageActor.cpp b/Engine/Plugins/Importers/USDImporter/Source/USDStage/Private/USDStageActor.cpp
--- a/Engine/Plugins/Importers/USDImporter/Source/USDStage/Private/USDStageActor.cpp
+++ b/Engine/Plugins/Importers/USDImporter/Source/USDStage/Private/USDStageActor.cpp
@@ -800,6 +800,11 @@ void AUsdStageActor::OnUsdObjectsChanged( const UsdUtils::FObjectChangesByPath&
 		}
 	}
 
+	// Mark the level as dirty since we received a notice about our stage having changed in some way.
+	// The main goal of this is to trigger the "save layers" dialog if we then save the UE level
+	const bool bAlwaysMarkDirty = true;
+	Modify( bAlwaysMarkDirty );
+
 	// We may update our levelsequence objects (tracks, moviescene, sections, etc.) due to these changes. We definitely don't want to write anything
 	// back to USD when these objects change though.
 	FScopedBlockMonitoringChangesForTransaction BlockMonitoring{ LevelSequenceHelper };
@@ -1590,6 +1595,8 @@ void AUsdStageActor::OpenUsdStage()
 		AbsPath = RootLayer.FilePath;
 	}
 
+	OnPreStageChanged.Broadcast();
+
 	UsdStage = UnrealUSDWrapper::OpenStage( *AbsPath, InitialLoadSet );
 	if ( UsdStage )
 	{
@@ -1613,6 +1620,8 @@ void AUsdStageActor::OpenUsdStage()
 
 void AUsdStageActor::CloseUsdStage()
 {
+	OnPreStageChanged.Broadcast();
+
 	FUsdStageActorImpl::DiscardStage( UsdStage, this );
 	UsdStage = UE::FUsdStage();
 	LevelSequenceHelper.Init( UE::FUsdStage() ); // Drop the helper's reference to the stage
@@ -1722,6 +1731,7 @@ void AUsdStageActor::OnObjectsReplaced( const TMap<UObject*, UObject*>& ObjectRe
 			NewActor->OnActorDestroyed = OnActorDestroyed;
 			NewActor->OnActorLoaded = OnActorLoaded;
 			NewActor->OnStageChanged = OnStageChanged;
+			NewActor->OnPreStageChanged = OnPreStageChanged;
 			NewActor->OnPrimChanged = OnPrimChanged;
 
 			// UEngine::CopyPropertiesForUnrelatedObjects won't copy over the cache's transient assets, but we still
@@ -1770,6 +1780,8 @@ void AUsdStageActor::LoadUsdStage()
 	FScopedSlowTask SlowTask( 1.f, LOCTEXT( "LoadingUDStage", "Loading USD Stage") );
 	SlowTask.MakeDialog();
 
+	OnPreStageChanged.Broadcast();
+
 	if ( !AssetCache )
 	{
 		AssetCache = NewObject< UUsdAssetCache >( this, TEXT("AssetCache"), GetMaskedFlags( RF_PropagateToSubObjects ) );
diff --git a/Engine/Plugins/Importers/USDImporter/Source/USDStage/Public/USDStageActor.h b/Engine/Plugins/Importers/USDImporter/Source/USDStage/Public/USDStageActor.h
--- a/Engine/Plugins/Importers/USDImporter/Source/USDStage/Public/USDStageActor.h
+++ b/Engine/Plugins/Importers/USDImporter/Source/USDStage/Public/USDStageActor.h
@@ -170,6 +170,7 @@ public:
 	USDSTAGE_API static FOnActorLoaded OnActorLoaded;
 
 	DECLARE_EVENT( AUsdStageActor, FOnStageActorEvent );
+	FOnStageActorEvent OnPreStageChanged;
 	FOnStageActorEvent OnStageChanged;
 	FOnStageActorEvent OnActorDestroyed;
 
diff --git a/Engine/Plugins/Importers/USDImporter/Source/USDStageEditor/Private/SUSDLayersTreeView.cpp b/Engine/Plugins/Importers/USDImporter/Source/USDStageEditor/Private/SUSDLayersTreeView.cpp
--- a/Engine/Plugins/Importers/USDImporter/Source/USDStageEditor/Private/SUSDLayersTreeView.cpp
+++ b/Engine/Plugins/Importers/USDImporter/Source/USDStageEditor/Private/SUSDLayersTreeView.cpp
@@ -632,7 +632,7 @@ void SUsdLayersTreeView::OnExportSelectedLayers() const
 	// Single layer -> Allow picking the target layer filename
 	if ( LayersToExport.Num() == 1 )
 	{
-		TOptional< FString > UsdFilePath = UsdUtils::BrowseUsdFile( UsdUtils::EBrowseFileMode::Save, AsShared() );
+		TOptional< FString > UsdFilePath = UsdUtils::BrowseUsdFile( UsdUtils::EBrowseFileMode::Save );
 		if ( !UsdFilePath.IsSet() )
 		{
 			return;
@@ -713,7 +713,7 @@ bool SUsdLayersTreeView::CanAddSubLayer() const
 
 void SUsdLayersTreeView::OnAddSubLayer()
 {
-	TOptional< FString > SubLayerFile = UsdUtils::BrowseUsdFile( UsdUtils::EBrowseFileMode::Composition, AsShared() );
+	TOptional< FString > SubLayerFile = UsdUtils::BrowseUsdFile( UsdUtils::EBrowseFileMode::Composition );
 
 	if ( !SubLayerFile )
 	{
@@ -733,7 +733,7 @@ void SUsdLayersTreeView::OnAddSubLayer()
 
 void SUsdLayersTreeView::OnNewSubLayer()
 {
-	TOptional< FString > SubLayerFile = UsdUtils::BrowseUsdFile( UsdUtils::EBrowseFileMode::Save, AsShared() );
+	TOptional< FString > SubLayerFile = UsdUtils::BrowseUsdFile( UsdUtils::EBrowseFileMode::Save );
 
 	if ( !SubLayerFile )
 	{
diff --git a/Engine/Plugins/Importers/USDImporter/Source/USDStageEditor/Private/SUSDSaveDialog.cpp b/Engine/Plugins/Importers/USDImporter/Source/USDStageEditor/Private/SUSDSaveDialog.cpp
new file mode 100644
--- /dev/null
+++ b/Engine/Plugins/Importers/USDImporter/Source/USDStageEditor/Private/SUSDSaveDialog.cpp
@@ -0,0 +1,388 @@
+// Copyright Epic Games, Inc. All Rights Reserved.
+
+#include "SUSDSaveDialog.h"
+
+#include "USDStageActor.h"
+
+#include "CoreMinimal.h"
+#include "Editor.h"
+#include "Framework/Application/SlateApplication.h"
+#include "Interfaces/IMainFrameModule.h"
+#include "Misc/ScopedSlowTask.h"
+#include "Modules/ModuleManager.h"
+#include "PropertyEditorModule.h"
+#include "Widgets/Input/SButton.h"
+#include "Widgets/Input/SCheckBox.h"
+#include "Widgets/Layout/SBox.h"
+#include "Widgets/Layout/SSpacer.h"
+#include "Widgets/Layout/SUniformGridPanel.h"
+#include "Widgets/SBoxPanel.h"
+#include "Widgets/SWindow.h"
+
+#define LOCTEXT_NAMESPACE "SUsdSaveDialog"
+
+namespace UE::UsdSaveDialog::Private
+{
+	FName CheckboxColumn = TEXT( "Checkbox" );
+	FName IdentifierColumn = TEXT( "Identifier" );
+	FName UsedByStagesColumn = TEXT( "Used by stages" );
+	FName UsedByActorsColumn = TEXT( "Used by actors" );
+}
+
+TArray<FUsdSaveDialogRowData> SUsdSaveDialog::ShowDialog(
+	const TArray<FUsdSaveDialogRowData>& InRows,
+	const FText& WindowTitle,
+	const FText& DescriptionText,
+	bool* OutShouldSave,
+	bool* OutShouldPromptAgain
+)
+{
+	TSharedRef< SWindow > Window = SNew( SWindow )
+		.Title( WindowTitle )
+		.ClientSize( FVector2D{ 900, 500 } )
+		.SizingRule( ESizingRule::UserSized );
+
+	TSharedPtr<SUsdSaveDialog> SaveDialogWidget;
+	Window->SetContent(
+		SAssignNew( SaveDialogWidget, SUsdSaveDialog )
+		.Rows( InRows )
+		.WidgetWindow( Window )
+		.DescriptionText( DescriptionText )
+	);
+	Window->SetWidgetToFocusOnActivate( SaveDialogWidget );
+
+	GEditor->EditorAddModalWindow( Window );
+
+	TArray<FUsdSaveDialogRowData> Result;
+
+	if ( SaveDialogWidget->ShouldSave() )
+	{
+		const TArray<TSharedPtr<FUsdSaveDialogRowData>>& SharedRows = SaveDialogWidget->Rows;
+		Result.Reserve( SharedRows.Num() );
+		for ( const TSharedPtr<FUsdSaveDialogRowData>& SharedRow : SharedRows )
+		{
+			Result.Add( *SharedRow );
+		}
+	}
+
+	if ( OutShouldSave )
+	{
+		*OutShouldSave = SaveDialogWidget->ShouldSave();
+	}
+
+	if ( OutShouldPromptAgain )
+	{
+		*OutShouldPromptAgain = SaveDialogWidget->ShouldPromptAgain();
+	}
+
+	return Result;
+}
+
+void SUsdSaveDialog::Construct( const FArguments& InArgs )
+{
+	// Copy into shared pointers because of the list view
+	TArray<TSharedPtr<FUsdSaveDialogRowData>> SharedRows;
+	SharedRows.Reserve( InArgs._Rows.Num() );
+	for ( const FUsdSaveDialogRowData& RawRow : InArgs._Rows )
+	{
+		SharedRows.Add( MakeShared<FUsdSaveDialogRowData>( RawRow ) );
+	}
+
+	Rows = SharedRows;
+	Window = InArgs._WidgetWindow;
+
+	TSharedRef< SHeaderRow > HeaderRowWidget = SNew( SHeaderRow );
+	HeaderRowWidget->AddColumn(
+		SHeaderRow::Column( UE::UsdSaveDialog::Private::CheckboxColumn )
+		[
+			SNew(SBox)
+			.Padding(FMargin(6,3,6,3))
+			.HAlign(HAlign_Center)
+			[
+				SNew( SCheckBox )
+				.IsChecked_Lambda( [this]() -> ECheckBoxState
+				{
+					ECheckBoxState Result = ECheckBoxState::Checked;
+
+					for ( TSharedPtr<FUsdSaveDialogRowData>& Row : Rows )
+					{
+						if ( !Row->bSaveLayer )
+						{
+							// Follow logic within SPackagesDialog::GetToggleSelectedState:
+							// If any of them is unchecked, treat them all as unchecked so that the first
+							// click on the toggle all checkbox checks them
+							Result = ECheckBoxState::Unchecked;
+						}
+					}
+
+					return Result;
+				})
+				.OnCheckStateChanged_Lambda( [this]( ECheckBoxState NewState )
+				{
+					for ( TSharedPtr<FUsdSaveDialogRowData>& Row : Rows )
+					{
+						Row->bSaveLayer = NewState == ECheckBoxState::Checked;
+					}
+				})
+			]
+		]
+		.FixedWidth( 38.0f )
+	);
+	HeaderRowWidget->AddColumn(
+		SHeaderRow::Column( UE::UsdSaveDialog::Private::IdentifierColumn )
+		.DefaultLabel( LOCTEXT( "IdentifierColumnLabel", "Identifier" ) )
+		.FillWidth( 4.0f )
+	);
+	HeaderRowWidget->AddColumn(
+		SHeaderRow::Column( UE::UsdSaveDialog::Private::UsedByStagesColumn )
+		.DefaultLabel( LOCTEXT( "UsedByStagesColumnLabel", "Used by stages" ) )
+		.FillWidth( 1.0f )
+	);
+	HeaderRowWidget->AddColumn(
+		SHeaderRow::Column( UE::UsdSaveDialog::Private::UsedByActorsColumn )
+		.DefaultLabel( LOCTEXT( "UsedByActorsColumnLabel", "Used by actors" ) )
+		.FillWidth( 1.0f )
+	);
+
+	const bool bAccept = true;
+	const bool bCancel = false;
+
+	TSharedPtr< SHorizontalBox > ButtonsBox = SNew( SHorizontalBox )
+		+SHorizontalBox::Slot()
+		.AutoWidth()
+		.Padding( 5, 0 )
+		[
+			SNew( SCheckBox )
+			.IsChecked( false ) // If we're showing the prompt we know we have this unchecked
+			.OnCheckStateChanged_Lambda( [this]( ECheckBoxState NewState )
+			{
+				bPromptAgain = ( NewState == ECheckBoxState::Unchecked );
+			})
+			[
+				SNew( STextBlock )
+				.Text( LOCTEXT( "DontPromptAgainText", "Don't prompt again" ) )
+				.AutoWrapText( true )
+			]
+		]
+		+SHorizontalBox::Slot()
+		.AutoWidth()
+		.Padding( 5, 0 )
+		[
+			SNew( SButton )
+			.ButtonStyle( &FAppStyle::Get(), "PrimaryButton" )
+			.TextStyle( &FAppStyle::Get(), "PrimaryButtonText" )
+			.Text( LOCTEXT( "SaveSelectedText", "Save selected" ) )
+			.ToolTipText( LOCTEXT( "SaveSelectedToolTip", "Saves the selected layers to disk" ) )
+			.HAlign( HAlign_Right )
+			.VAlign( VAlign_Center )
+			.OnClicked( this, &SUsdSaveDialog::Close, bAccept )
+		]
+		+SHorizontalBox::Slot()
+		.AutoWidth()
+		.Padding( 5, 0 )
+		[
+			SNew( SButton )
+			.ButtonStyle( &FAppStyle::Get(), "Button" )
+			.TextStyle( &FAppStyle::Get(), "ButtonText" )
+			.Text( LOCTEXT( "SaveNoneText", "Don't save layers" ) )
+			.ToolTipText( LOCTEXT( "SaveNoneToolTip", "Proceed with the Unreal save, but don't save any USD layer" ) )
+			.HAlign( HAlign_Right )
+			.VAlign( VAlign_Center )
+			.OnClicked( this, &SUsdSaveDialog::Close, bCancel )
+		];
+
+	ChildSlot
+	[
+		SNew(SBorder)
+		.BorderImage(FAppStyle::Get().GetBrush("Brushes.Panel"))
+		.Padding(FMargin(16))
+		[
+			SNew(SVerticalBox)
+			+SVerticalBox::Slot()
+			.AutoHeight()
+			.Padding(0.0f,0.0f,0.0f,8.0f)
+			[
+				SNew(STextBlock)
+				.Text( InArgs._DescriptionText )
+				.AutoWrapText(true)
+			]
+			+SVerticalBox::Slot()
+			.FillHeight(0.8)
+			[
+				SNew( SListView< TSharedPtr<FUsdSaveDialogRowData> > )
+				.ListItemsSource(&Rows)
+				.OnGenerateRow( this, &SUsdSaveDialog::OnGenerateListRow )
+				.ItemHeight(20)
+				.HeaderRow( HeaderRowWidget )
+				.SelectionMode( ESelectionMode::None )
+			]
+			+SVerticalBox::Slot()
+			.AutoHeight()
+			.Padding(0.0f, 16.0f, 0.0f, 0.0f)
+			.HAlign(HAlign_Right)
+			.VAlign(VAlign_Bottom)
+			[
+				ButtonsBox.ToSharedRef()
+			]
+		]
+	];
+}
+
+bool SUsdSaveDialog::SupportsKeyboardFocus() const
+{
+	return true;
+}
+
+FReply SUsdSaveDialog::OnKeyDown( const FGeometry& MyGeometry, const FKeyEvent& InKeyEvent )
+{
+	if ( InKeyEvent.GetKey() == EKeys::Escape )
+	{
+		// If we're exiting via Esc, don't assume we want to commit our choice of bPromptAgain
+		bPromptAgain = true;
+
+		const bool bShouldProceed = false;
+		return Close( bShouldProceed );
+	}
+
+	return FReply::Unhandled();
+}
+
+class SUsdSaveDialogRow : public SMultiColumnTableRow< TSharedPtr< FUsdSaveDialogRowData > >
+{
+public:
+	SLATE_BEGIN_ARGS( SUsdSaveDialogRow ) {}
+		SLATE_ARGUMENT( TSharedPtr< FUsdSaveDialogRowData >, Item )
+	SLATE_END_ARGS()
+
+	void Construct( const FArguments& InArgs, const TSharedRef<STableViewBase>& InOwnerTableView )
+	{
+		Item = InArgs._Item;
+
+		SMultiColumnTableRow< TSharedPtr< FUsdSaveDialogRowData > >::Construct(
+			FSuperRowType::FArguments(),
+			InOwnerTableView
+		);
+	}
+
+	virtual TSharedRef<SWidget> GenerateWidgetForColumn( const FName& ColumnName ) override
+	{
+		TSharedRef<SWidget> ItemContentWidget = SNullWidget::NullWidget;
+
+		if ( ColumnName == UE::UsdSaveDialog::Private::CheckboxColumn )
+		{
+			ItemContentWidget = SNew(SHorizontalBox)
+				+SHorizontalBox::Slot()
+				.Padding(FMargin(10, 3, 6, 3))
+				[
+					SNew(SCheckBox)
+					.IsChecked_Lambda( [this]() -> ECheckBoxState
+					{
+						if ( Item )
+						{
+							return Item->bSaveLayer
+								? ECheckBoxState::Checked
+								: ECheckBoxState::Unchecked;
+						}
+
+						return ECheckBoxState::Undetermined;
+					})
+					.OnCheckStateChanged_Lambda( [this]( ECheckBoxState NewState )
+					{
+						if ( Item )
+						{
+							Item->bSaveLayer = NewState == ECheckBoxState::Checked;
+						}
+					})
+				];
+		}
+		else if ( ColumnName == UE::UsdSaveDialog::Private::IdentifierColumn )
+		{
+			if ( Item && Item->Layer )
+			{
+				ItemContentWidget = SNew( SHorizontalBox )
+				+SHorizontalBox::Slot()
+				.Padding( 3, 3, 3, 3 )
+				[
+					SNew( STextBlock )
+					.Text( FText::FromString( Item->Layer.GetIdentifier() ) )
+				];
+			}
+		}
+		else if ( ColumnName == UE::UsdSaveDialog::Private::UsedByStagesColumn )
+		{
+			if ( Item )
+			{
+				FString Text;
+				for ( const UE::FUsdStageWeak& Stage : Item->ConsumerStages )
+				{
+					if ( Stage )
+					{
+						Text += Stage.GetRootLayer().GetDisplayName() + TEXT( ", " );
+					}
+				}
+				Text.RemoveFromEnd( TEXT( ", " ) );
+
+				ItemContentWidget = SNew( SHorizontalBox )
+				+SHorizontalBox::Slot()
+				.Padding( 3, 3, 3, 3 )
+				[
+					SNew( STextBlock )
+					.Text( FText::FromString( Text ) )
+				];
+			}
+		}
+		else if ( ColumnName == UE::UsdSaveDialog::Private::UsedByActorsColumn )
+		{
+			if ( Item )
+			{
+				FString Text;
+				FString Delimiter = TEXT(", ");
+
+				for ( const AUsdStageActor* Actor : Item->ConsumerActors )
+				{
+					if ( Actor )
+					{
+						Text += Actor->GetActorLabel() + Delimiter;
+					}
+				}
+				Text.RemoveFromEnd( Delimiter );
+
+				ItemContentWidget = SNew( SHorizontalBox )
+				+SHorizontalBox::Slot()
+				.Padding( 3, 3, 3, 3 )
+				[
+					SNew( STextBlock )
+					.Text( FText::FromString( Text ) )
+				];
+			}
+		}
+
+		return ItemContentWidget;
+	}
+
+private:
+	TSharedPtr< FUsdSaveDialogRowData > Item;
+};
+
+
+TSharedRef<ITableRow> SUsdSaveDialog::OnGenerateListRow(
+	TSharedPtr<FUsdSaveDialogRowData> Item,
+	const TSharedRef<STableViewBase>& OwnerTable
+)
+{
+	return SNew( SUsdSaveDialogRow, OwnerTable ).Item( Item );
+}
+
+FReply SUsdSaveDialog::Close( bool bInProceed )
+{
+	bProceed = bInProceed;
+
+	if ( Window.IsValid() )
+	{
+		Window.Pin()->RequestDestroyWindow();
+	}
+
+	return FReply::Handled();
+}
+
+#undef LOCTEXT_NAMESPACE
diff --git a/Engine/Plugins/Importers/USDImporter/Source/USDStageEditor/Private/SUSDSaveDialog.h b/Engine/Plugins/Importers/USDImporter/Source/USDStageEditor/Private/SUSDSaveDialog.h
new file mode 100644
--- /dev/null
+++ b/Engine/Plugins/Importers/USDImporter/Source/USDStageEditor/Private/SUSDSaveDialog.h
@@ -0,0 +1,64 @@
+// Copyright Epic Games, Inc. All Rights Reserved.
+
+#pragma once
+
+#include "UsdWrappers/ForwardDeclarations.h"
+#include "UsdWrappers/SdfLayer.h"
+#include "UsdWrappers/UsdStage.h"
+
+#include "CoreMinimal.h"
+#include "Widgets/SCompoundWidget.h"
+
+class AUsdStageActor;
+class ITableRow;
+class STableViewBase;
+
+struct FUsdSaveDialogRowData
+{
+	bool bSaveLayer = true;
+	UE::FSdfLayerWeak Layer;
+	TArray<UE::FUsdStageWeak> ConsumerStages;
+	TSet<AUsdStageActor*> ConsumerActors;
+};
+
+/** Dialog shown to ask if the user wishes to save any dirty USD layer to disk */
+class SUsdSaveDialog : public SCompoundWidget
+{
+public:
+	SLATE_BEGIN_ARGS( SUsdSaveDialog ){}
+		SLATE_ARGUMENT( TArray<FUsdSaveDialogRowData>, Rows )
+		SLATE_ARGUMENT( FText, DescriptionText )
+		SLATE_ARGUMENT( TSharedPtr<SWindow>, WidgetWindow )
+	SLATE_END_ARGS()
+
+	// Returns the data describing which layers should be saved
+	static TArray<FUsdSaveDialogRowData> ShowDialog(
+		const TArray<FUsdSaveDialogRowData>& InRows,
+		const FText& WindowTitle,
+		const FText& DescriptionText,
+		bool* OutShouldSave = nullptr,
+		bool* OutShouldPromptAgain = nullptr
+	);
+
+	void Construct( const FArguments& InArgs );
+	bool ShouldSave() const { return bProceed; }
+	bool ShouldPromptAgain() const { return bPromptAgain; }
+
+private:
+    virtual bool SupportsKeyboardFocus() const override;
+	virtual FReply OnKeyDown( const FGeometry& MyGeometry, const FKeyEvent& InKeyEvent ) override;
+
+	TSharedRef<ITableRow> OnGenerateListRow(
+		TSharedPtr<FUsdSaveDialogRowData> Item,
+		const TSharedRef<STableViewBase>& OwnerTable
+	);
+
+	FReply Close( bool bProceed );
+
+private:
+	TWeakPtr< SWindow > Window;
+	TArray<TSharedPtr<FUsdSaveDialogRowData>> Rows;
+
+	bool bPromptAgain = true;
+	bool bProceed = false;
+};
diff --git a/Engine/Plugins/Importers/USDImporter/Source/USDStageEditor/Private/SUSDStage.cpp b/Engine/Plugins/Importers/USDImporter/Source/USDStageEditor/Private/SUSDStage.cpp
--- a/Engine/Plugins/Importers/USDImporter/Source/USDStageEditor/Private/SUSDStage.cpp
+++ b/Engine/Plugins/Importers/USDImporter/Source/USDStageEditor/Private/SUSDStage.cpp
@@ -329,12 +329,21 @@ void SUsdStage::SetupStageActorDelegates()
 					const bool bEvenIfPendingKill = true;
 					if ( ViewModel.UsdStageActor.IsValid( bEvenIfPendingKill ) )
 					{
+						// Reset our selection to the stage root
+						SelectedPrimPath = TEXT( "/" );
+
 						if ( this->UsdPrimInfoWidget )
 						{
 							// The cast here forces us to use the const version of GetUsdStage, that won't force-load the stage in case it isn't opened yet
 							const UE::FUsdStage& UsdStage = static_cast< const AUsdStageActor* >( ViewModel.UsdStageActor.Get( bEvenIfPendingKill ) )->GetUsdStage();
 							this->UsdPrimInfoWidget->SetPrimPath( UsdStage, TEXT("/") );
 						}
+
+						if ( this->UsdStageTreeView )
+						{
+							this->UsdStageTreeView->ClearSelection();
+							this->UsdStageTreeView->RequestTreeRefresh();
+						}
 					}
 
 					this->Refresh();
@@ -1325,7 +1334,7 @@ void SUsdStage::OnNew()
 
 void SUsdStage::OnOpen()
 {
-	TOptional< FString > UsdFilePath = UsdUtils::BrowseUsdFile( UsdUtils::EBrowseFileMode::Open, AsShared() );
+	TOptional< FString > UsdFilePath = UsdUtils::BrowseUsdFile( UsdUtils::EBrowseFileMode::Open );
 
 	if ( UsdFilePath )
 	{
@@ -1352,7 +1361,7 @@ void SUsdStage::OnSave()
 			}
 			else
 			{
-				TOptional< FString > UsdFilePath = UsdUtils::BrowseUsdFile( UsdUtils::EBrowseFileMode::Save, AsShared() );
+				TOptional< FString > UsdFilePath = UsdUtils::BrowseUsdFile( UsdUtils::EBrowseFileMode::Save );
 				if ( UsdFilePath )
 				{
 					ViewModel.SaveStageAs( *UsdFilePath.GetValue() );
@@ -1521,7 +1530,7 @@ void SUsdStage::OnExportFlattened()
 		return;
 	}
 
-	TOptional< FString > UsdFilePath = UsdUtils::BrowseUsdFile( UsdUtils::EBrowseFileMode::Save, AsShared() );
+	TOptional< FString > UsdFilePath = UsdUtils::BrowseUsdFile( UsdUtils::EBrowseFileMode::Save );
 	if ( !UsdFilePath.IsSet() )
 	{
 		return;
@@ -1604,8 +1613,19 @@ void SUsdStage::OnPrimSelectionChanged( const TArray<FString>& PrimPaths )
 		UsdPrimInfoWidget->SetPrimPath( UsdStage, *SelectedPrimPath );
 	}
 
+	// We don't need to check for actors and components in case we're just selecting the stage root
+	// Note: This is also a temp fix that prevents an ensure: When we're switching to a new stage,
+	// AUsdStageActor::LoadUsdStage calls AUsdStageActor::OpenUsdStage, which triggers OnStageChanged
+	// and gets us here. Ideally this wouldn't happen at the same time as AUsdStageActor::LoadUsdStage
+	// because we need the InfoCache to be complete at this point. Strangely though, the SlowTask.EnterProgressFrame
+	// within AUsdStageActor::LoadUsdStage triggers a full slate tick right there, that calls us before
+	// LoadUsdStage is complete... in the future I think we can fix this by changing our delegates from
+	// AsyncTask( ENamedThreads::GameThread ) to finishing up on the next tick instead, but for now
+	// this will do
+	const bool bIsRootSelection = PrimPaths.Num() == 1 && PrimPaths[0] == TEXT("/");
+
 	const UUsdStageEditorSettings* Settings = GetDefault<UUsdStageEditorSettings>();
-	if ( Settings && Settings->bSelectionSynced && GEditor )
+	if ( Settings && Settings->bSelectionSynced && GEditor && !bIsRootSelection )
 	{
 		TGuardValue<bool> SelectionLoopGuard( bUpdatingViewportSelection, true );
 
diff --git a/Engine/Plugins/Importers/USDImporter/Source/USDStageEditor/Private/USDStageEditorModule.cpp b/Engine/Plugins/Importers/USDImporter/Source/USDStageEditor/Private/USDStageEditorModule.cpp
--- a/Engine/Plugins/Importers/USDImporter/Source/USDStageEditor/Private/USDStageEditorModule.cpp
+++ b/Engine/Plugins/Importers/USDImporter/Source/USDStageEditor/Private/USDStageEditorModule.cpp
@@ -2,23 +2,216 @@
 
 #include "USDStageEditorModule.h"
 
+#include "SUSDSaveDialog.h"
 #include "SUSDStage.h"
 #include "SUSDStageEditorStyle.h"
+#include "USDLayerUtils.h"
+#include "USDLog.h"
 #include "USDMemory.h"
+#include "USDProjectSettings.h"
 #include "USDStageActor.h"
 
-#include "Styling/AppStyle.h"
 #include "EngineUtils.h"
 #include "Framework/Docking/TabManager.h"
 #include "IAssetTools.h"
+#include "Interfaces/IMainFrameModule.h"
 #include "LevelEditor.h"
 #include "Modules/ModuleManager.h"
+#include "Styling/AppStyle.h"
+#include "UObject/ObjectSaveContext.h"
 #include "Widgets/Docking/SDockTab.h"
 #include "WorkspaceMenuStructure.h"
 #include "WorkspaceMenuStructureModule.h"
 
 #define LOCTEXT_NAMESPACE "UsdStageEditorModule"
 
+namespace UE::UsdStageEditorModule::Private
+{
+	void SaveStageActorLayersForWorld( UWorld* World, bool bForClosing, AUsdStageActor* TargetStageActor = nullptr )
+	{
+		if ( !World )
+		{
+			return;
+		}
+
+		UUsdProjectSettings* Settings = GetMutableDefault<UUsdProjectSettings>();
+		if ( !Settings )
+		{
+			return;
+		}
+
+		// Reentrant guard here because if we ever save an anonymous layer we'll update the stage actors that use it
+		// to point to the new (saved) layer, which will internally close the anonymous stage and get us back in here
+		static bool bIsReentrant = false;
+		if ( bIsReentrant )
+		{
+			return;
+		}
+		TGuardValue<bool> ReentrantGuard( bIsReentrant, true );
+
+		bool bPrompt = false;
+		switch ( bForClosing ? Settings->ShowSaveLayersDialogWhenClosing : Settings->ShowSaveLayersDialogWhenSaving )
+		{
+			case EUsdSaveDialogBehavior::NeverSave:
+			{
+				// Don't even do anything if we're not going to save anyway
+				return;
+				break;
+			}
+			case EUsdSaveDialogBehavior::AlwaysSave:
+			{
+				break;
+			}
+			default:
+			case EUsdSaveDialogBehavior::ShowPrompt:
+			{
+				bPrompt = true;
+				break;
+			}
+		}
+
+		TArray<AUsdStageActor*> StageActorsToVisit;
+		if ( TargetStageActor )
+		{
+			if ( TargetStageActor->GetWorld() == World )
+			{
+				StageActorsToVisit.Add( TargetStageActor );
+			}
+		}
+		else
+		{
+			for ( TActorIterator<AUsdStageActor> It{ World }; It; ++It )
+			{
+				StageActorsToVisit.Add( *It );
+			}
+		}
+
+		// For now lets only care about stages opened on stage actors. The user could have additional stages,
+		// like opened via Python or custom C++ plugins, but lets ignore those
+		TMap<FString, FUsdSaveDialogRowData> RowsByIdentifier;
+		for ( AUsdStageActor* StageActor : StageActorsToVisit )
+		{
+			if ( !StageActor )
+			{
+				continue;
+			}
+
+			UE::FUsdStage UsdStage = static_cast< const AUsdStageActor* >( StageActor )->GetUsdStage();
+			if ( !UsdStage )
+			{
+				continue;
+			}
+
+			TArray<UE::FSdfLayer> UsedLayers = UsdStage.GetUsedLayers();
+			RowsByIdentifier.Reserve( RowsByIdentifier.Num() + UsedLayers.Num() );
+
+			for ( const UE::FSdfLayer& UsedLayer : UsedLayers )
+			{
+				if ( UsedLayer.IsDirty() )
+				{
+					FUsdSaveDialogRowData& RowData = RowsByIdentifier.FindOrAdd( UsedLayer.GetIdentifier() );
+					RowData.Layer = UsedLayer;
+					RowData.ConsumerStages.AddUnique( UsdStage );
+					RowData.ConsumerActors.Add( StageActor );
+				}
+			}
+		}
+
+		if ( RowsByIdentifier.Num() == 0 )
+		{
+			return;
+		}
+
+		TArray< FUsdSaveDialogRowData > Rows;
+		RowsByIdentifier.GenerateValueArray( Rows );
+
+		if ( bPrompt )
+		{
+			Rows.Sort( []( const FUsdSaveDialogRowData& Left, const FUsdSaveDialogRowData& Right )
+			{
+				return ( Left.Layer && Right.Layer )
+					? Left.Layer.GetIdentifier() < Right.Layer.GetIdentifier()
+					// This shouldn't ever happen but just do something consistent here instead anyway
+					: Left.ConsumerStages.Num() < Right.ConsumerStages.Num();
+			});
+
+			const FText WindowTitle = LOCTEXT( "SaveDialogTitle", "Save USD Layers" );
+			const FText DescriptionText = bForClosing
+				? LOCTEXT( "CloseDialogDescTextText", "Before closing these USD Stages, do you want to save these USD layers to disk?" )
+				: LOCTEXT( "SaveDialogDescTextText", "Since you're saving the Level, do you want to save these USD layers to disk?" );
+
+			bool bShouldSave = true;
+			bool bShouldPromptAgain = true;
+			Rows = SUsdSaveDialog::ShowDialog(
+				Rows,
+				WindowTitle,
+				DescriptionText,
+				&bShouldSave,
+				&bShouldPromptAgain
+			);
+
+			EUsdSaveDialogBehavior* bSetting = bForClosing
+				? &Settings->ShowSaveLayersDialogWhenClosing
+				: &Settings->ShowSaveLayersDialogWhenSaving;
+
+			*bSetting = bShouldPromptAgain
+				? EUsdSaveDialogBehavior::ShowPrompt
+				: bShouldSave
+					? EUsdSaveDialogBehavior::AlwaysSave
+					: EUsdSaveDialogBehavior::NeverSave;
+		}
+
+		for ( const FUsdSaveDialogRowData& ReturnedRow : Rows )
+		{
+			UE::FSdfLayer PinnedLayer = ReturnedRow.Layer;
+			if ( ReturnedRow.bSaveLayer && PinnedLayer )
+			{
+				bool bSaved = false;
+
+				if ( PinnedLayer.IsAnonymous() )
+				{
+					// For now we only allow anonymous stages, and not individual layers, so we don't have to
+					// patch up anything
+					TOptional< FString > UsdFilePath = UsdUtils::BrowseUsdFile( UsdUtils::EBrowseFileMode::Save );
+					if ( UsdFilePath )
+					{
+						bSaved = PinnedLayer.Export( *UsdFilePath.GetValue() );
+
+						// If any stage actors were pointing at the in-memory versions of these stages, update them
+						// to point to the saved versions
+						if ( bSaved && !bForClosing )
+						{
+							// Even though we're potentially going to load actors and assets here, we don't need to
+							// use a scoped transaction as the Save All command already clears the transaction buffer
+							// anyway, and we won't get in here when closing.
+
+							FString ExpectedIdentifier = UnrealIdentifiers::IdentifierPrefix + PinnedLayer.GetIdentifier();
+							for ( AUsdStageActor* StageActor : ReturnedRow.ConsumerActors )
+							{
+								if ( StageActor->RootLayer.FilePath == ExpectedIdentifier
+									&& static_cast<const AUsdStageActor*>(StageActor)->GetUsdStage() )
+								{
+									StageActor->SetRootLayer( *UsdFilePath.GetValue() );
+								}
+							}
+						}
+					}
+				}
+				else
+				{
+					const bool bForce = true;
+					bSaved = PinnedLayer.Save( bForce );
+				}
+
+				if ( !bSaved )
+				{
+					UE_LOG( LogUsd, Warning, TEXT( "Failed to save layer '%s'" ), *PinnedLayer.GetIdentifier() );
+				}
+			}
+		}
+	}
+}
+
 class FUsdStageEditorModule : public IUsdStageEditorModule
 {
 public:
@@ -40,6 +233,8 @@ public:
 
 	virtual void StartupModule() override
 	{
+		FUsdStageEditorStyle::Initialize();
+
 		FLevelEditorModule& LevelEditorModule = FModuleManager::GetModuleChecked< FLevelEditorModule >( "LevelEditor" );
 		LevelEditorTabManagerChangedHandle = LevelEditorModule.OnTabManagerChanged().AddLambda(
 			[]()
@@ -56,11 +251,70 @@ public:
 					.SetIcon( LayersIcon );
 			});
 
-		FUsdStageEditorStyle::Initialize();
+		// Prompt to save modified USD layers when closing the editor
+		IMainFrameModule& MainFrame = FModuleManager::LoadModuleChecked<IMainFrameModule>( "MainFrame" );
+		EditorCanCloseDelegate = MainFrame.RegisterCanCloseEditor( IMainFrameModule::FMainFrameCanCloseEditor::CreateLambda( []() -> bool
+		{
+			if ( GEditor && GEngine )
+			{
+				UWorld* EditorWorld = nullptr;
+				for ( const FWorldContext& Context : GEngine->GetWorldContexts() )
+				{
+					if ( Context.WorldType == EWorldType::Editor )
+					{
+						EditorWorld = Context.World();
+					}
+				}
+
+				const bool bForClosing = true;
+				UE::UsdStageEditorModule::Private::SaveStageActorLayersForWorld( EditorWorld, bForClosing );
+			}
+
+			// We won't actually ever block the save
+			return true;
+		}));
+
+		// Prompt to save modified USD Layers when closing stageactor stages
+		StageActorLoadedHandle = AUsdStageActor::OnActorLoaded.AddLambda([](AUsdStageActor* StageActor)
+		{
+			if ( !StageActor )
+			{
+				return;
+			}
+
+			StageActor->OnPreStageChanged.AddLambda( [StageActor]()
+			{
+				if ( StageActor && static_cast< const AUsdStageActor* >( StageActor )->GetUsdStage() )
+				{
+					const bool bForClosing = true;
+					UE::UsdStageEditorModule::Private::SaveStageActorLayersForWorld( StageActor->GetWorld(), bForClosing, StageActor );
+				}
+			});
+		});
+
+		// Prompt to save modified USD layers when saving the world
+		PreSaveWorldEditorDelegateHandle = FEditorDelegates::PreSaveWorldWithContext.AddLambda( []( UWorld* World, FObjectPreSaveContext InContext )
+		{
+			// Detect if we should actually do anything (check for autosaves, cooking, etc.)
+			if ( InContext.GetSaveFlags() & ESaveFlags::SAVE_FromAutosave || InContext.IsProceduralSave() )
+			{
+				return;
+			}
+
+			const bool bForClosing = false;
+			UE::UsdStageEditorModule::Private::SaveStageActorLayersForWorld( World, bForClosing );
+		});
 	}
 
 	virtual void ShutdownModule() override
 	{
+		FEditorDelegates::PreSaveWorldWithContext.Remove( PreSaveWorldEditorDelegateHandle );
+
+		AUsdStageActor::OnActorLoaded.Remove( StageActorLoadedHandle );
+
+		IMainFrameModule& MainFrame = FModuleManager::LoadModuleChecked<IMainFrameModule>( "MainFrame" );
+		MainFrame.UnregisterCanCloseEditor( EditorCanCloseDelegate );
+
 		if ( LevelEditorTabManagerChangedHandle.IsValid() && FModuleManager::Get().IsModuleLoaded( "LevelEditor" ) )
 		{
 			FLevelEditorModule& LevelEditorModule = FModuleManager::GetModuleChecked<FLevelEditorModule>( "LevelEditor" );
@@ -72,6 +326,9 @@ public:
 
 private:
 	FDelegateHandle LevelEditorTabManagerChangedHandle;
+	FDelegateHandle PreSaveWorldEditorDelegateHandle;
+	FDelegateHandle EditorCanCloseDelegate;
+	FDelegateHandle StageActorLoadedHandle;
 #endif // #if USE_USD_SDK
 };
 
diff --git a/Engine/Plugins/Importers/USDImporter/Source/USDUtilities/Private/USDLayerUtils.cpp b/Engine/Plugins/Importers/USDImporter/Source/USDUtilities/Private/USDLayerUtils.cpp
--- a/Engine/Plugins/Importers/USDImporter/Source/USDUtilities/Private/USDLayerUtils.cpp
+++ b/Engine/Plugins/Importers/USDImporter/Source/USDUtilities/Private/USDLayerUtils.cpp
@@ -117,7 +117,7 @@ bool UsdUtils::InsertSubLayer( const pxr::SdfLayerRefPtr& ParentLayer, const TCH
 }
 
 #if WITH_EDITOR
-TOptional< FString > UsdUtils::BrowseUsdFile( EBrowseFileMode Mode, TSharedRef< const SWidget > OriginatingWidget )
+TOptional< FString > UsdUtils::BrowseUsdFile( EBrowseFileMode Mode )
 {
 	IDesktopPlatform* DesktopPlatform = FDesktopPlatformModule::Get();
 
@@ -126,12 +126,6 @@ TOptional< FString > UsdUtils::BrowseUsdFile( EBrowseFileMode Mode, TSharedRef<
 		return {};
 	}
 
-	// show the file browse dialog
-	TSharedPtr< SWindow > ParentWindow = FSlateApplication::Get().FindWidgetWindow( OriginatingWidget );
-	void* ParentWindowHandle = ( ParentWindow.IsValid() && ParentWindow->GetNativeWindow().IsValid() )
-		? ParentWindow->GetNativeWindow()->GetOSWindowHandle()
-		: nullptr;
-
 	TArray< FString > OutFiles;
 
 	// When browsing files for the purposes of opening a stage or saving layers,
@@ -168,7 +162,15 @@ TOptional< FString > UsdUtils::BrowseUsdFile( EBrowseFileMode Mode, TSharedRef<
 		case EBrowseFileMode::Open :
 		case EBrowseFileMode::Composition :
 		{
-			if ( !DesktopPlatform->OpenFileDialog( ParentWindowHandle, LOCTEXT( "ChooseFile", "Choose file").ToString(), TEXT(""), TEXT(""), *FileTypes, EFileDialogFlags::None, OutFiles ) )
+			if ( !DesktopPlatform->OpenFileDialog(
+				FSlateApplication::Get().FindBestParentWindowHandleForDialogs( nullptr ),
+				LOCTEXT( "ChooseFile", "Choose file").ToString(),
+				TEXT(""),
+				TEXT(""),
+				*FileTypes,
+				EFileDialogFlags::None,
+				OutFiles
+			) )
 			{
 				return {};
 			}
@@ -176,7 +178,15 @@ TOptional< FString > UsdUtils::BrowseUsdFile( EBrowseFileMode Mode, TSharedRef<
 		}
 		case EBrowseFileMode::Save :
 		{
-			if ( !DesktopPlatform->SaveFileDialog( ParentWindowHandle, LOCTEXT( "ChooseFile", "Choose file").ToString(), TEXT(""), TEXT(""), *FileTypes, EFileDialogFlags::None, OutFiles ) )
+			if ( !DesktopPlatform->SaveFileDialog(
+				FSlateApplication::Get().FindBestParentWindowHandleForDialogs( nullptr ),
+				LOCTEXT( "ChooseFile", "Choose file").ToString(),
+				TEXT(""),
+				TEXT(""),
+				*FileTypes,
+				EFileDialogFlags::None,
+				OutFiles
+			) )
 			{
 				return {};
 			}
diff --git a/Engine/Plugins/Importers/USDImporter/Source/USDUtilities/Public/USDLayerUtils.h b/Engine/Plugins/Importers/USDImporter/Source/USDUtilities/Public/USDLayerUtils.h
--- a/Engine/Plugins/Importers/USDImporter/Source/USDUtilities/Public/USDLayerUtils.h
+++ b/Engine/Plugins/Importers/USDImporter/Source/USDUtilities/Public/USDLayerUtils.h
@@ -59,7 +59,7 @@ namespace UsdUtils
 
 #if WITH_EDITOR
 	/** Opens a file dialog to open or save a USD file. The returned file path will always be absolute */
-	USDUTILITIES_API TOptional< FString > BrowseUsdFile( EBrowseFileMode Mode, TSharedRef< const SWidget > OriginatingWidget );
+	USDUTILITIES_API TOptional< FString > BrowseUsdFile( EBrowseFileMode Mode );
 #endif // #if WITH_EDITOR
 
 	/**
diff --git a/Engine/Plugins/Importers/USDImporter/Source/UnrealUSDWrapper/Private/UsdWrappers/SdfLayer.cpp b/Engine/Plugins/Importers/USDImporter/Source/UnrealUSDWrapper/Private/UsdWrappers/SdfLayer.cpp
--- a/Engine/Plugins/Importers/USDImporter/Source/UnrealUSDWrapper/Private/UsdWrappers/SdfLayer.cpp
+++ b/Engine/Plugins/Importers/USDImporter/Source/UnrealUSDWrapper/Private/UsdWrappers/SdfLayer.cpp
@@ -638,6 +638,19 @@ namespace UE
 #endif // #if USE_USD_SDK
 	}
 
+	template<typename PtrType>
+	bool FSdfLayerBase<PtrType>::IsDirty() const
+	{
+#if USE_USD_SDK
+		if ( const PtrType& Ptr = Impl->GetInner() )
+		{
+			return Ptr->IsDirty();
+		}
+#endif // #if USE_USD_SDK
+
+		return false;
+	}
+
 	template<typename PtrType>
 	bool FSdfLayerBase<PtrType>::IsEmpty() const
 	{
diff --git a/Engine/Plugins/Importers/USDImporter/Source/UnrealUSDWrapper/Public/UsdWrappers/SdfLayer.h b/Engine/Plugins/Importers/USDImporter/Source/UnrealUSDWrapper/Public/UsdWrappers/SdfLayer.h
--- a/Engine/Plugins/Importers/USDImporter/Source/UnrealUSDWrapper/Public/UsdWrappers/SdfLayer.h
+++ b/Engine/Plugins/Importers/USDImporter/Source/UnrealUSDWrapper/Public/UsdWrappers/SdfLayer.h
@@ -99,6 +99,7 @@ namespace UE
 
 		FString ComputeAbsolutePath(const FString& AssetPath) const;
 
+		bool IsDirty() const;
 		bool IsEmpty() const;
 		bool IsAnonymous() const;
 
