USDImporter: remove use of deprecated UsdGeomImageable::CreatePrimvar()

The primvar-related methods on UsdGeomImageable were deprecated way back in
USD 19.03, but they were only just removed with the release of USD 22.08:

[REDACTED_PROFILE_URL]/USD/commit/592b4d39edf5daf0534d467e970c95462a65d44b

This removes the handful of remaining usages of the CreatePrimvar() method in favor
of using the supported UsdGeomPrimvarsAPI directly.

diff --git a/Engine/Plugins/Importers/USDImporter/Source/USDUtilities/Private/USDGeomMeshConversion.cpp b/Engine/Plugins/Importers/USDImporter/Source/USDUtilities/Private/USDGeomMeshConversion.cpp
--- a/Engine/Plugins/Importers/USDImporter/Source/USDUtilities/Private/USDGeomMeshConversion.cpp
+++ b/Engine/Plugins/Importers/USDImporter/Source/USDUtilities/Private/USDGeomMeshConversion.cpp
@@ -177,7 +177,7 @@ namespace UE::UsdGeomMeshConversion::Private
 				{
 					pxr::TfToken UsdUVSetName = UsdUtils::GetUVSetName( TexCoordSourceIndex ).Get();
 
-					pxr::UsdGeomPrimvar PrimvarST = UsdMesh.CreatePrimvar( UsdUVSetName, pxr::SdfValueTypeNames->TexCoord2fArray, pxr::UsdGeomTokens->vertex );
+					pxr::UsdGeomPrimvar PrimvarST = pxr::UsdGeomPrimvarsAPI(MeshPrim).CreatePrimvar( UsdUVSetName, pxr::SdfValueTypeNames->TexCoord2fArray, pxr::UsdGeomTokens->vertex );
 
 					if ( PrimvarST )
 					{
@@ -435,7 +435,7 @@ namespace UE::UsdGeomMeshConversion::Private
 			{
 				pxr::TfToken UsdUVSetName = UsdUtils::GetUVSetName( UVIndex ).Get();
 
-				pxr::UsdGeomPrimvar PrimvarST = UsdMesh.CreatePrimvar( UsdUVSetName, pxr::SdfValueTypeNames->TexCoord2fArray, pxr::UsdGeomTokens->vertex );
+				pxr::UsdGeomPrimvar PrimvarST = pxr::UsdGeomPrimvarsAPI(MeshPrim).CreatePrimvar( UsdUVSetName, pxr::SdfValueTypeNames->TexCoord2fArray, pxr::UsdGeomTokens->vertex );
 				if ( PrimvarST )
 				{
 					pxr::VtVec2fArray UVs;
diff --git a/Engine/Plugins/Importers/USDImporter/Source/USDUtilities/Private/USDSkeletalDataConversion.cpp b/Engine/Plugins/Importers/USDImporter/Source/USDUtilities/Private/USDSkeletalDataConversion.cpp
--- a/Engine/Plugins/Importers/USDImporter/Source/USDUtilities/Private/USDSkeletalDataConversion.cpp
+++ b/Engine/Plugins/Importers/USDImporter/Source/USDUtilities/Private/USDSkeletalDataConversion.cpp
@@ -734,7 +734,7 @@ namespace UnrealToUsdImpl
 				{
 					pxr::TfToken UsdUVSetName = UsdUtils::GetUVSetName( TexCoordSourceIndex ).Get();
 
-					pxr::UsdGeomPrimvar PrimvarST = UsdLODPrimGeomMesh.CreatePrimvar( UsdUVSetName, pxr::SdfValueTypeNames->TexCoord2fArray, pxr::UsdGeomTokens->vertex );
+					pxr::UsdGeomPrimvar PrimvarST = pxr::UsdGeomPrimvarsAPI(MeshPrim).CreatePrimvar( UsdUVSetName, pxr::SdfValueTypeNames->TexCoord2fArray, pxr::UsdGeomTokens->vertex );
 
 					if ( PrimvarST )
 					{
