USD: Fix build when USD SDK is not available.

diff --git a/Engine/Plugins/Importers/USDImporter/Source/USDStageEditor/Private/USDStageEditorModule.cpp b/Engine/Plugins/Importers/USDImporter/Source/USDStageEditor/Private/USDStageEditorModule.cpp
--- a/Engine/Plugins/Importers/USDImporter/Source/USDStageEditor/Private/USDStageEditorModule.cpp
+++ b/Engine/Plugins/Importers/USDImporter/Source/USDStageEditor/Private/USDStageEditorModule.cpp
@@ -29,6 +29,7 @@ namespace UE::UsdStageEditorModule::Private
 {
 	void SaveStageActorLayersForWorld( UWorld* World, bool bForClosing, AUsdStageActor* TargetStageActor = nullptr )
 	{
+#if USE_USD_SDK
 		if ( !World )
 		{
 			return;
@@ -209,6 +210,7 @@ namespace UE::UsdStageEditorModule::Private
 				}
 			}
 		}
+#endif // USE_USD_SDK
 	}
 }
 
