USD: Add support for material purposes.

diff --git a/Engine/Plugins/Importers/USDImporter/Source/USDClasses/Public/USDProjectSettings.h b/Engine/Plugins/Importers/USDImporter/Source/USDClasses/Public/USDProjectSettings.h
--- a/Engine/Plugins/Importers/USDImporter/Source/USDClasses/Public/USDProjectSettings.h
+++ b/Engine/Plugins/Importers/USDImporter/Source/USDClasses/Public/USDProjectSettings.h
@@ -26,6 +26,10 @@ public:
 	UPROPERTY( config, EditAnywhere, Category = USD )
 	TArray<FDirectoryPath> AdditionalPluginDirectories;
 
+	// Material purposes to show on drop-downs in addition to the standard "preview" and "full"
+	UPROPERTY( config, EditAnywhere, Category = USD )
+	TArray<FName> AdditionalMaterialPurposes;
+
 	UPROPERTY( config, EditAnywhere, Category = "USD|Dialogs" )
 	bool bShowConfirmationWhenClearingLayers = true;
 
diff --git a/Engine/Plugins/Importers/USDImporter/Source/USDSchemas/Private/MeshTranslationImpl.cpp b/Engine/Plugins/Importers/USDImporter/Source/USDSchemas/Private/MeshTranslationImpl.cpp
--- a/Engine/Plugins/Importers/USDImporter/Source/USDSchemas/Private/MeshTranslationImpl.cpp
+++ b/Engine/Plugins/Importers/USDImporter/Source/USDSchemas/Private/MeshTranslationImpl.cpp
@@ -138,7 +138,17 @@ TMap<const UsdUtils::FUsdPrimMaterialSlot*, UMaterialInterface*> MeshTranslation
 	return ResolvedMaterials;
 }
 
-void MeshTranslationImpl::SetMaterialOverrides( const pxr::UsdPrim& Prim, const TArray<UMaterialInterface*>& ExistingAssignments, UMeshComponent& MeshComponent, UUsdAssetCache& AssetCache, float Time, EObjectFlags Flags, bool bInterpretLODs, const FName& RenderContext )
+void MeshTranslationImpl::SetMaterialOverrides(
+	const pxr::UsdPrim& Prim,
+	const TArray<UMaterialInterface*>& ExistingAssignments,
+	UMeshComponent& MeshComponent,
+	UUsdAssetCache& AssetCache,
+	float Time,
+	EObjectFlags Flags,
+	bool bInterpretLODs,
+	const FName& RenderContext,
+	const FName& MaterialPurpose
+)
 {
 	FScopedUsdAllocs Allocs;
 
@@ -151,12 +161,17 @@ void MeshTranslationImpl::SetMaterialOverrides( const pxr::UsdPrim& Prim, const
 	pxr::UsdStageRefPtr Stage = Prim.GetStage();
 
 	pxr::TfToken RenderContextToken = pxr::UsdShadeTokens->universalRenderContext;
-
 	if ( !RenderContext.IsNone() )
 	{
 		RenderContextToken = UnrealToUsd::ConvertToken( *RenderContext.ToString() ).Get();
 	}
 
+	pxr::TfToken MaterialPurposeToken = pxr::UsdShadeTokens->allPurpose;
+	if ( !MaterialPurpose.IsNone() )
+	{
+		MaterialPurposeToken = UnrealToUsd::ConvertToken( *MaterialPurpose.ToString() ).Get();
+	}
+
 	TArray<UsdUtils::FUsdPrimMaterialAssignmentInfo> LODIndexToAssignments;
 	const bool bProvideMaterialIndices = false; // We have no use for material indices and it can be slow to retrieve, as it will iterate all faces
 
@@ -167,7 +182,13 @@ void MeshTranslationImpl::SetMaterialOverrides( const pxr::UsdPrim& Prim, const
 		TMap<int32, UsdUtils::FUsdPrimMaterialAssignmentInfo> LODIndexToAssignmentsMap;
 		TFunction<bool( const pxr::UsdGeomMesh&, int32 )> IterateLODs = [ & ]( const pxr::UsdGeomMesh& LODMesh, int32 LODIndex )
 		{
-			UsdUtils::FUsdPrimMaterialAssignmentInfo LODInfo = UsdUtils::GetPrimMaterialAssignments( LODMesh.GetPrim(), pxr::UsdTimeCode( Time ), bProvideMaterialIndices, RenderContextToken );
+			UsdUtils::FUsdPrimMaterialAssignmentInfo LODInfo = UsdUtils::GetPrimMaterialAssignments(
+				LODMesh.GetPrim(),
+				pxr::UsdTimeCode( Time ),
+				bProvideMaterialIndices,
+				RenderContextToken,
+				MaterialPurposeToken
+			);
 			LODIndexToAssignmentsMap.Add( LODIndex, LODInfo );
 			return true;
 		};
@@ -191,7 +212,15 @@ void MeshTranslationImpl::SetMaterialOverrides( const pxr::UsdPrim& Prim, const
 	// Extract material assignment info from prim if its *not* a LOD mesh, or if we failed to parse LODs
 	if ( !bInterpretedLODs )
 	{
-		LODIndexToAssignments = { UsdUtils::GetPrimMaterialAssignments( ValidPrim, pxr::UsdTimeCode( Time ), bProvideMaterialIndices, RenderContextToken ) };
+		LODIndexToAssignments = {
+			UsdUtils::GetPrimMaterialAssignments(
+				ValidPrim,
+				pxr::UsdTimeCode( Time ),
+				bProvideMaterialIndices,
+				RenderContextToken,
+				MaterialPurposeToken
+			)
+		};
 	}
 
 	TMap<const UsdUtils::FUsdPrimMaterialSlot*, UMaterialInterface*> ResolvedMaterials = MeshTranslationImpl::ResolveMaterialAssignmentInfo(
diff --git a/Engine/Plugins/Importers/USDImporter/Source/USDSchemas/Private/MeshTranslationImpl.h b/Engine/Plugins/Importers/USDImporter/Source/USDSchemas/Private/MeshTranslationImpl.h
--- a/Engine/Plugins/Importers/USDImporter/Source/USDSchemas/Private/MeshTranslationImpl.h
+++ b/Engine/Plugins/Importers/USDImporter/Source/USDSchemas/Private/MeshTranslationImpl.h
@@ -49,7 +49,8 @@ namespace MeshTranslationImpl
 		float Time,
 		EObjectFlags Flags,
 		bool bInterpretLODs,
-		const FName& RenderContext
+		const FName& RenderContext,
+		const FName& MaterialPurpose
 	);
 }
 
diff --git a/Engine/Plugins/Importers/USDImporter/Source/USDSchemas/Private/USDGeomMeshTranslator.cpp b/Engine/Plugins/Importers/USDImporter/Source/USDSchemas/Private/USDGeomMeshTranslator.cpp
--- a/Engine/Plugins/Importers/USDImporter/Source/USDSchemas/Private/USDGeomMeshTranslator.cpp
+++ b/Engine/Plugins/Importers/USDImporter/Source/USDSchemas/Private/USDGeomMeshTranslator.cpp
@@ -1466,7 +1466,8 @@ USceneComponent* FUsdGeomMeshTranslator::CreateComponents()
 				Context->Time,
 				Context->ObjectFlags,
 				Context->bAllowInterpretingLODs,
-				Context->RenderContext
+				Context->RenderContext,
+				Context->MaterialPurpose
 			);
 		}
 	}
@@ -1485,7 +1486,8 @@ USceneComponent* FUsdGeomMeshTranslator::CreateComponents()
 				Context->Time,
 				Context->ObjectFlags,
 				bAllowInterpretingLODs,
-				Context->RenderContext
+				Context->RenderContext,
+				Context->MaterialPurpose
 			);
 
 #if WITH_EDITOR
@@ -1494,7 +1496,7 @@ USceneComponent* FUsdGeomMeshTranslator::CreateComponents()
 			{
 				UsdGroomTranslatorUtils::CreateGroomBindingAsset( GetPrim(), *( Context->AssetCache ), Context->ObjectFlags );
 
-				// For the groom binding to work, the GroomComponent must be a child of the SceneComponent 
+				// For the groom binding to work, the GroomComponent must be a child of the SceneComponent
 				// so the Context ParentComponent is set to the SceneComponent temporarily
 				TGuardValue< USceneComponent* > ParentComponentGuard{ Context->ParentComponent, SceneComponent };
 				const bool bNeedsActor = false;
diff --git a/Engine/Plugins/Importers/USDImporter/Source/USDSchemas/Private/USDGeomPointInstancerTranslator.cpp b/Engine/Plugins/Importers/USDImporter/Source/USDSchemas/Private/USDGeomPointInstancerTranslator.cpp
--- a/Engine/Plugins/Importers/USDImporter/Source/USDSchemas/Private/USDGeomPointInstancerTranslator.cpp
+++ b/Engine/Plugins/Importers/USDImporter/Source/USDSchemas/Private/USDGeomPointInstancerTranslator.cpp
@@ -351,7 +351,8 @@ USceneComponent* FUsdGeomPointInstancerTranslator::CreateComponents()
 					Context->Time,
 					Context->ObjectFlags,
 					Context->bAllowInterpretingLODs,
-					Context->RenderContext
+					Context->RenderContext,
+					Context->MaterialPurpose
 				);
 			}
 		}
diff --git a/Engine/Plugins/Importers/USDImporter/Source/USDSchemas/Private/USDInfoCache.cpp b/Engine/Plugins/Importers/USDImporter/Source/USDSchemas/Private/USDInfoCache.cpp
--- a/Engine/Plugins/Importers/USDImporter/Source/USDSchemas/Private/USDInfoCache.cpp
+++ b/Engine/Plugins/Importers/USDImporter/Source/USDSchemas/Private/USDInfoCache.cpp
@@ -140,12 +140,19 @@ namespace UE::USDInfoCacheImpl::Private
 				RenderContextToken = UnrealToUsd::ConvertToken( *Context.RenderContext.ToString() ).Get();
 			}
 
+			pxr::TfToken MaterialPurposeToken = pxr::UsdShadeTokens->allPurpose;
+			if ( !Context.MaterialPurpose.IsNone() )
+			{
+				MaterialPurposeToken = UnrealToUsd::ConvertToken( *Context.MaterialPurpose.ToString() ).Get();
+			}
+
 			const bool bProvideMaterialIndices = false;
 			UsdUtils::FUsdPrimMaterialAssignmentInfo LocalInfo = UsdUtils::GetPrimMaterialAssignments(
 				UsdPrim,
 				Context.Time,
 				bProvideMaterialIndices,
-				RenderContextToken
+				RenderContextToken,
+				MaterialPurposeToken
 			);
 
 			OutMaterialSlots.Append( MoveTemp( LocalInfo.Slots ) );
diff --git a/Engine/Plugins/Importers/USDImporter/Source/USDSchemas/Private/USDSkelRootTranslator.cpp b/Engine/Plugins/Importers/USDImporter/Source/USDSchemas/Private/USDSkelRootTranslator.cpp
--- a/Engine/Plugins/Importers/USDImporter/Source/USDSchemas/Private/USDSkelRootTranslator.cpp
+++ b/Engine/Plugins/Importers/USDImporter/Source/USDSchemas/Private/USDSkelRootTranslator.cpp
@@ -331,7 +331,8 @@ namespace UsdSkelRootTranslatorImpl
 		float InTime,
 		const UUsdAssetCache& AssetCache,
 		bool bInInterpretLODs,
-		const pxr::TfToken& RenderContext
+		const pxr::TfToken& RenderContext,
+		const pxr::TfToken& MaterialPurpose
 	)
 	{
 		if ( !InSkeletonRoot )
@@ -396,7 +397,18 @@ namespace UsdSkelRootTranslatorImpl
 		}
 
 		TFunction<bool( const pxr::UsdGeomMesh&, int32 )> ConvertLOD =
-		[ &LODIndexToSkeletalMeshImportDataMap, &LODIndexToMaterialInfoMap, &SkelQuery, InTime, &InMaterialToPrimvarsUVSetNames, &InOutUsedMorphTargetNames, OutBlendShapes, &StageInfo, RenderContext ]
+		[
+			&LODIndexToSkeletalMeshImportDataMap,
+			&LODIndexToMaterialInfoMap,
+			&SkelQuery,
+			InTime,
+			&InMaterialToPrimvarsUVSetNames,
+			&InOutUsedMorphTargetNames,
+			OutBlendShapes,
+			&StageInfo,
+			RenderContext,
+			MaterialPurpose
+		]
 		( const pxr::UsdGeomMesh& LODMesh, int32 LODIndex )
 		{
 			pxr::UsdSkelSkinningQuery SkinningQuery = UsdUtils::CreateSkinningQuery( LODMesh, SkelQuery );
@@ -417,7 +429,15 @@ namespace UsdSkelRootTranslatorImpl
 			// into one FSkeletalMeshImportData per LOD, so we need to remap the indices using this
 			uint32 NumPointsBeforeThisMesh = static_cast< uint32 >( LODImportData.Points.Num() );
 
-			bool bSuccess = UsdToUnreal::ConvertSkinnedMesh( SkinningQuery, SkelQuery, LODImportData, LODSlots, InMaterialToPrimvarsUVSetNames, RenderContext );
+			bool bSuccess = UsdToUnreal::ConvertSkinnedMesh(
+				SkinningQuery,
+				SkelQuery,
+				LODImportData,
+				LODSlots,
+				InMaterialToPrimvarsUVSetNames,
+				RenderContext,
+				MaterialPurpose
+			);
 			if ( !bSuccess )
 			{
 				return true;
@@ -527,7 +547,17 @@ namespace UsdSkelRootTranslatorImpl
 	}
 
 	/** Warning: This function will temporarily switch the active LOD variant if one exists, so it's *not* thread safe! */
-	void SetMaterialOverrides( const pxr::UsdPrim& SkelRootPrim, const TArray<UMaterialInterface*>& ExistingAssignments, UMeshComponent& MeshComponent, UUsdAssetCache& AssetCache, float Time, EObjectFlags Flags, bool bInterpretLODs, const FName& RenderContext )
+	void SetMaterialOverrides(
+		const pxr::UsdPrim& SkelRootPrim,
+		const TArray<UMaterialInterface*>& ExistingAssignments,
+		UMeshComponent& MeshComponent,
+		UUsdAssetCache& AssetCache,
+		float Time,
+		EObjectFlags Flags,
+		bool bInterpretLODs,
+		const FName& RenderContext,
+		const FName& MaterialPurpose
+	)
 	{
 		FScopedUsdAllocs Allocs;
 
@@ -545,9 +575,23 @@ namespace UsdSkelRootTranslatorImpl
 			RenderContextToken = UnrealToUsd::ConvertToken( *RenderContext.ToString() ).Get();
 		}
 
+		pxr::TfToken MaterialPurposeToken = pxr::UsdShadeTokens->allPurpose;
+		if ( !MaterialPurpose.IsNone() )
+		{
+			MaterialPurposeToken = UnrealToUsd::ConvertToken( *MaterialPurpose.ToString() ).Get();
+		}
+
 		TMap<int32, UsdUtils::FUsdPrimMaterialAssignmentInfo> LODIndexToMaterialInfoMap;
 		TMap<int32, TSet<UsdUtils::FUsdPrimMaterialSlot>> CombinedSlotsForLODIndex;
-		TFunction<bool( const pxr::UsdGeomMesh&, int32 )> IterateLODsLambda = [ &LODIndexToMaterialInfoMap, &CombinedSlotsForLODIndex, Time, RenderContextToken ]( const pxr::UsdGeomMesh& LODMesh, int32 LODIndex )
+		TFunction<bool( const pxr::UsdGeomMesh&, int32 )> IterateLODsLambda =
+		[
+			&LODIndexToMaterialInfoMap,
+			&CombinedSlotsForLODIndex,
+			Time,
+			RenderContextToken,
+			MaterialPurposeToken
+		]
+		( const pxr::UsdGeomMesh& LODMesh, int32 LODIndex )
 		{
 			if ( LODMesh && LODMesh.ComputeVisibility() == pxr::UsdGeomTokens->invisible )
 			{
@@ -558,7 +602,13 @@ namespace UsdSkelRootTranslatorImpl
 			TSet<UsdUtils::FUsdPrimMaterialSlot>& CombinedLODSlotsSet = CombinedSlotsForLODIndex.FindOrAdd( LODIndex );
 
 			const bool bProvideMaterialIndices = false; // We have no use for material indices and it can be slow to retrieve, as it will iterate all faces
-			UsdUtils::FUsdPrimMaterialAssignmentInfo LocalInfo = UsdUtils::GetPrimMaterialAssignments( LODMesh.GetPrim(), pxr::UsdTimeCode( Time ), bProvideMaterialIndices, RenderContextToken );
+			UsdUtils::FUsdPrimMaterialAssignmentInfo LocalInfo = UsdUtils::GetPrimMaterialAssignments(
+				LODMesh.GetPrim(),
+				pxr::UsdTimeCode( Time ),
+				bProvideMaterialIndices,
+				RenderContextToken,
+				MaterialPurposeToken
+			);
 
 			// Combine material slots in the same order that UsdToUnreal::ConvertSkinnedMesh does
 			for ( UsdUtils::FUsdPrimMaterialSlot& LocalSlot : LocalInfo.Slots )
@@ -755,6 +805,12 @@ namespace UsdSkelRootTranslatorImpl
 					RenderContextToken = UnrealToUsd::ConvertToken( *Context->RenderContext.ToString() ).Get();
 				}
 
+				pxr::TfToken MaterialPurposeToken = pxr::UsdShadeTokens->allPurpose;
+				if ( !Context->MaterialPurpose.IsNone() )
+				{
+					MaterialPurposeToken = UnrealToUsd::ConvertToken( *Context->MaterialPurpose.ToString() ).Get();
+				}
+
 				const bool bContinueTaskChain = UsdSkelRootTranslatorImpl::LoadAllSkeletalData(
 					SkeletonCache.Get(),
 					pxr::UsdSkelRoot( GetPrim() ),
@@ -767,7 +823,8 @@ namespace UsdSkelRootTranslatorImpl
 					Context->Time,
 					*Context->AssetCache.Get(),
 					Context->bAllowInterpretingLODs,
-					RenderContextToken
+					RenderContextToken,
+					MaterialPurposeToken
 				);
 
 				return bContinueTaskChain;
@@ -1227,7 +1284,7 @@ USceneComponent* FUsdSkelRootTranslator::CreateComponents()
 	{
 		UsdGroomTranslatorUtils::CreateGroomBindingAsset( GetPrim(), *( Context->AssetCache ), Context->ObjectFlags );
 
-		// For the groom binding to work, the GroomComponent must be a child of the SceneComponent 
+		// For the groom binding to work, the GroomComponent must be a child of the SceneComponent
 		// so the Context ParentComponent is set to the SceneComponent temporarily
 		TGuardValue< USceneComponent* > ParentComponentGuard{ Context->ParentComponent, SceneComponent };
 		const bool bNeedsActor = false;
@@ -1321,7 +1378,8 @@ void FUsdSkelRootTranslator::UpdateComponents( USceneComponent* SceneComponent )
 				Context->Time,
 				Context->ObjectFlags,
 				Context->bAllowInterpretingLODs,
-				Context->RenderContext
+				Context->RenderContext,
+				Context->MaterialPurpose
 			);
 		}
 	}
diff --git a/Engine/Plugins/Importers/USDImporter/Source/USDSchemas/Public/USDSchemaTranslator.h b/Engine/Plugins/Importers/USDImporter/Source/USDSchemas/Public/USDSchemaTranslator.h
--- a/Engine/Plugins/Importers/USDImporter/Source/USDSchemas/Public/USDSchemaTranslator.h
+++ b/Engine/Plugins/Importers/USDImporter/Source/USDSchemas/Public/USDSchemaTranslator.h
@@ -149,6 +149,9 @@ struct USDSCHEMAS_API FUsdSchemaTranslationContext : public TSharedFromThis< FUs
 	/** The render context to use when translating materials */
 	FName RenderContext;
 
+	/** The material purpose to use when translating material bindings */
+	FName MaterialPurpose;
+
 	/** If a generated UStaticMesh has at least this many triangles we will attempt to enable Nanite */
 	int32 NaniteTriangleThreshold;
 
diff --git a/Engine/Plugins/Importers/USDImporter/Source/USDStage/Private/USDStageActor.cpp b/Engine/Plugins/Importers/USDImporter/Source/USDStage/Private/USDStageActor.cpp
--- a/Engine/Plugins/Importers/USDImporter/Source/USDStage/Private/USDStageActor.cpp
+++ b/Engine/Plugins/Importers/USDImporter/Source/USDStage/Private/USDStageActor.cpp
@@ -104,6 +104,7 @@ struct FUsdStageActorImpl
 		TranslationContext->PurposesToLoad = (EUsdPurpose) StageActor->PurposesToLoad;
 		TranslationContext->NaniteTriangleThreshold = StageActor->NaniteTriangleThreshold;
 		TranslationContext->RenderContext = StageActor->RenderContext;
+		TranslationContext->MaterialPurpose = StageActor->MaterialPurpose;
 		TranslationContext->MaterialToPrimvarToUVIndex = &StageActor->MaterialToPrimvarToUVIndex;
 		TranslationContext->BlendShapesByPath = &StageActor->BlendShapesByPath;
 		TranslationContext->InfoCache = StageActor->InfoCache;
@@ -537,6 +538,7 @@ struct FUsdStageActorImpl
 			EventAttributes.Emplace( TEXT( "PurposesToLoad" ), LexToString( StageActor->PurposesToLoad ) );
 			EventAttributes.Emplace( TEXT( "NaniteTriangleThreshold" ), LexToString( StageActor->NaniteTriangleThreshold ) );
 			EventAttributes.Emplace( TEXT( "RenderContext" ), StageActor->RenderContext.ToString() );
+			EventAttributes.Emplace( TEXT( "MaterialPurpose" ), StageActor->MaterialPurpose.ToString() );
 
 			const bool bAutomated = false;
 			IUsdClassesModule::SendAnalytics( MoveTemp( EventAttributes ), TEXT( "Open" ), bAutomated, ElapsedSeconds, NumberOfFrames, Extension );
@@ -621,6 +623,7 @@ AUsdStageActor::AUsdStageActor()
 	, bCollapseTopLevelPointInstancers( false )
 	, PurposesToLoad( (int32) EUsdPurpose::Proxy )
 	, NaniteTriangleThreshold( (uint64) 1000000 )
+	, MaterialPurpose( *UnrealIdentifiers::MaterialAllPurpose )
 	, Time( 0.0f )
 	, bIsTransitioningIntoPIE( false )
 	, bIsModifyingAProperty( false )
@@ -1021,7 +1024,7 @@ void AUsdStageActor::OnUsdObjectsChanged( const UsdUtils::FObjectChangesByPath&
 					// seem to be the only ones that do at the moment
 					if ( ExistingMaterial && ( ExistingMaterial != NewMaterial ) )
 					{
-						for ( const FString& MaterialUserPrim : UsdUtils::GetMaterialUsers( PrimToResync ) )
+						for ( const FString& MaterialUserPrim : UsdUtils::GetMaterialUsers( PrimToResync, MaterialPurpose ) )
 						{
 							const bool bResyncComponent = true; // We need to force resync to reassign materials
 							UE::FSdfPath MaterialPrimPath{ *MaterialUserPrim };
@@ -1445,6 +1448,15 @@ void AUsdStageActor::SetRenderContext( const FName& NewRenderContext )
 	LoadUsdStage();
 }
 
+void AUsdStageActor::SetMaterialPurpose( const FName& NewMaterialPurpose )
+{
+	const bool bMarkDirty = false;
+	Modify( bMarkDirty );
+
+	MaterialPurpose = NewMaterialPurpose;
+	LoadUsdStage();
+}
+
 void AUsdStageActor::SetTime(float InTime)
 {
 	const bool bMarkDirty = false;
@@ -2732,6 +2744,10 @@ void AUsdStageActor::HandlePropertyChangedEvent( FPropertyChangedEvent& Property
 	{
 		SetRenderContext( RenderContext );
 	}
+	else if ( PropertyName == GET_MEMBER_NAME_CHECKED( AUsdStageActor, MaterialPurpose ) )
+	{
+		SetMaterialPurpose( MaterialPurpose );
+	}
 
 	bIsModifyingAProperty = false;
 }
diff --git a/Engine/Plugins/Importers/USDImporter/Source/USDStage/Private/USDStageActorCustomization.cpp b/Engine/Plugins/Importers/USDImporter/Source/USDStage/Private/USDStageActorCustomization.cpp
--- a/Engine/Plugins/Importers/USDImporter/Source/USDStage/Private/USDStageActorCustomization.cpp
+++ b/Engine/Plugins/Importers/USDImporter/Source/USDStage/Private/USDStageActorCustomization.cpp
@@ -4,6 +4,7 @@
 
 #include "USDStageActorCustomization.h"
 
+#include "USDProjectSettings.h"
 #include "USDSchemasModule.h"
 #include "USDSchemaTranslator.h"
 #include "USDStageActor.h"
@@ -13,10 +14,10 @@
 #include "DetailWidgetRow.h"
 #include "Modules/ModuleManager.h"
 #include "ScopedTransaction.h"
+#include "Styling/AppStyle.h"
 #include "UObject/Object.h"
 #include "Widgets/DeclarativeSyntaxSupport.h"
-#include "Styling/AppStyle.h"
-
+#include "Widgets/Input/SEditableTextBox.h"
 
 #define LOCTEXT_NAMESPACE "UsdStageActorCustomization"
 
@@ -53,7 +54,7 @@ void FUsdStageActorCustomization::CustomizeDetails( IDetailLayoutBuilder& Detail
 
 	IUsdSchemasModule& UsdSchemasModule = FModuleManager::Get().LoadModuleChecked< IUsdSchemasModule >( TEXT( "USDSchemas" ) );
 
-	ComboBoxItems.Reset();
+	RenderContextComboBoxItems.Reset();
 	TSharedPtr<FString> InitiallySelectedContext;
 	for ( const FName& Context : UsdSchemasModule.GetRenderContextRegistry().GetRenderContexts() )
 	{
@@ -72,7 +73,7 @@ void FUsdStageActorCustomization::CustomizeDetails( IDetailLayoutBuilder& Detail
 			InitiallySelectedContext = ContextStr;
 		}
 
-		ComboBoxItems.Add( ContextStr );
+		RenderContextComboBoxItems.Add( ContextStr );
 	}
 
 	IDetailCategoryBuilder& CatBuilder = DetailLayoutBuilder.EditCategory( TEXT( "USD" ) );
@@ -85,14 +86,14 @@ void FUsdStageActorCustomization::CustomizeDetails( IDetailLayoutBuilder& Detail
 		.NameContent()
 		[
 			SNew( STextBlock )
-			.Text( FText::FromString( TEXT( "Render Context" ) ) )
+			.Text( FText::FromString( TEXT( "Render context" ) ) )
 			.Font( FAppStyle::GetFontStyle( TEXT( "PropertyWindow.NormalFont" ) ) )
 			.ToolTipText( RenderContextProperty->GetToolTipText() )
 		]
 		.ValueContent()
 		[
-			SAssignNew( ComboBox, SComboBox<TSharedPtr<FString>> )
-			.OptionsSource( &ComboBoxItems )
+			SAssignNew( RenderContextComboBox, SComboBox<TSharedPtr<FString>> )
+			.OptionsSource( &RenderContextComboBoxItems )
 			.InitiallySelectedItem( InitiallySelectedContext )
 			.OnSelectionChanged( this, &FUsdStageActorCustomization::OnComboBoxSelectionChanged )
 			.OnGenerateWidget_Lambda( []( TSharedPtr<FString> Item )
@@ -110,6 +111,131 @@ void FUsdStageActorCustomization::CustomizeDetails( IDetailLayoutBuilder& Detail
 		];
 	}
 
+	if ( TSharedPtr<IPropertyHandle> MaterialPurposeProperty = DetailLayoutBuilder.GetProperty( GET_MEMBER_NAME_CHECKED( AUsdStageActor, MaterialPurpose ) ) )
+	{
+		DetailLayoutBuilder.HideProperty( MaterialPurposeProperty );
+
+		CatBuilder.AddCustomRow( FText::FromString( TEXT( "MaterialPurposeCustomization" ) ) )
+		.NameContent()
+		[
+			SNew( STextBlock )
+			.Text( FText::FromString( TEXT( "Material purpose" ) ) )
+			.Font( FAppStyle::GetFontStyle( TEXT( "PropertyWindow.NormalFont" ) ) )
+			.ToolTipText( MaterialPurposeProperty->GetToolTipText() )
+		]
+		.ValueContent()
+		[
+			SNew( SBox )
+			.VAlign( VAlign_Center )
+			[
+				SNew( SComboBox< TSharedPtr<FString> > )
+				.OptionsSource( &MaterialPurposeComboBoxItems )
+				.OnComboBoxOpening_Lambda([this]()
+				{
+					MaterialPurposeComboBoxItems = {
+						MakeShared<FString>( UnrealIdentifiers::MaterialAllPurpose ),
+						MakeShared<FString>( UnrealIdentifiers::MaterialPreviewPurpose ),
+						MakeShared<FString>( UnrealIdentifiers::MaterialFullPurpose )
+					};
+
+					// Add additional purposes from project settings
+					if ( const UUsdProjectSettings* ProjectSettings = GetDefault<UUsdProjectSettings>() )
+					{
+						MaterialPurposeComboBoxItems.Reserve( MaterialPurposeComboBoxItems.Num() + ProjectSettings->AdditionalMaterialPurposes.Num() );
+
+						TSet<FString> ExistingEntries = {
+							UnrealIdentifiers::MaterialAllPurpose,
+							UnrealIdentifiers::MaterialPreviewPurpose,
+							UnrealIdentifiers::MaterialFullPurpose
+						};
+
+						for ( const FName& AdditionalPurpose : ProjectSettings->AdditionalMaterialPurposes )
+						{
+							FString AdditionalPurposeStr = AdditionalPurpose.ToString();
+
+							if ( !ExistingEntries.Contains( AdditionalPurposeStr ) )
+							{
+								ExistingEntries.Add( AdditionalPurposeStr );
+								MaterialPurposeComboBoxItems.AddUnique( MakeShared<FString>( AdditionalPurposeStr ) );
+							}
+						}
+					}
+				})
+				.OnGenerateWidget_Lambda( [ & ]( TSharedPtr<FString> Option )
+				{
+					TSharedPtr<SWidget> Widget = SNullWidget::NullWidget;
+					if ( Option )
+					{
+						Widget = SNew( STextBlock )
+							.Text( FText::FromString( ( *Option ) == UnrealIdentifiers::MaterialAllPurpose
+								? UnrealIdentifiers::MaterialAllPurposeText
+								: *Option
+							) )
+							.Font( FAppStyle::GetFontStyle( "PropertyWindow.NormalFont" ) );
+					}
+
+					return Widget.ToSharedRef();
+				})
+				.OnSelectionChanged_Lambda([this]( TSharedPtr<FString> ChosenOption, ESelectInfo::Type SelectInfo )
+				{
+					if ( ChosenOption && CurrentActor )
+					{
+						CurrentActor->SetMaterialPurpose( **ChosenOption );
+					}
+				})
+				[
+					SNew( SEditableTextBox )
+					.Text_Lambda([this]() -> FText
+					{
+						if ( CurrentActor )
+						{
+							return FText::FromString( CurrentActor->MaterialPurpose == *UnrealIdentifiers::MaterialAllPurpose
+								? UnrealIdentifiers::MaterialAllPurposeText
+								: CurrentActor->MaterialPurpose.ToString()
+							);
+						}
+
+						return FText::GetEmpty();
+					})
+					.Font( FAppStyle::GetFontStyle( "PropertyWindow.NormalFont" ) )
+					.OnTextCommitted_Lambda( [this]( const FText& NewText, ETextCommit::Type CommitType )
+					{
+						if ( CommitType != ETextCommit::OnEnter )
+						{
+							return;
+						}
+
+						FString NewPurposeString = NewText.ToString();
+						FName NewPurpose = *NewPurposeString;
+
+						bool bIsNew = true;
+						for ( const TSharedPtr<FString>& Purpose : MaterialPurposeComboBoxItems )
+						{
+							if ( Purpose && *Purpose == NewPurposeString )
+							{
+								bIsNew = false;
+								break;
+							}
+						}
+
+						if ( bIsNew )
+						{
+							if ( UUsdProjectSettings* ProjectSettings = GetMutableDefault<UUsdProjectSettings>() )
+							{
+								ProjectSettings->AdditionalMaterialPurposes.AddUnique( NewPurpose );
+							}
+						}
+
+						if ( CurrentActor )
+						{
+							CurrentActor->SetMaterialPurpose( NewPurpose );
+						}
+					})
+				]
+			]
+		];
+	}
+
 	// Add/remove properties so that they retain their usual order
 	if ( TSharedPtr<IPropertyHandle> TimeProperty = DetailLayoutBuilder.GetProperty( TEXT( "Time" ) ) )
 	{
@@ -147,7 +273,7 @@ void FUsdStageActorCustomization::OnComboBoxSelectionChanged( TSharedPtr<FString
 
 FText FUsdStageActorCustomization::GetComboBoxSelectedOptionText() const
 {
-	TSharedPtr<FString> SelectedItem = ComboBox->GetSelectedItem();
+	TSharedPtr<FString> SelectedItem = RenderContextComboBox->GetSelectedItem();
 	if ( SelectedItem.IsValid() )
 	{
 		return FText::FromString( *SelectedItem );
diff --git a/Engine/Plugins/Importers/USDImporter/Source/USDStage/Private/USDStageActorCustomization.h b/Engine/Plugins/Importers/USDImporter/Source/USDStage/Private/USDStageActorCustomization.h
--- a/Engine/Plugins/Importers/USDImporter/Source/USDStage/Private/USDStageActorCustomization.h
+++ b/Engine/Plugins/Importers/USDImporter/Source/USDStage/Private/USDStageActorCustomization.h
@@ -23,15 +23,17 @@ public:
 	// End of IDetailCustomization interface
 
 private:
-	void OnComboBoxOpening();
 	void OnComboBoxSelectionChanged( TSharedPtr<FString> NewContext, ESelectInfo::Type SelectType );
 	FText GetComboBoxSelectedOptionText() const;
 	void ForceRefreshDetails();
 
 	AUsdStageActor* CurrentActor;
 	TWeakPtr<IDetailLayoutBuilder> DetailBuilderWeakPtr;
-	TSharedPtr<SComboBox<TSharedPtr<FString>>> ComboBox;
-	TArray<TSharedPtr<FString>> ComboBoxItems;
+
+	TSharedPtr<SComboBox<TSharedPtr<FString>>> RenderContextComboBox;
+	TArray<TSharedPtr<FString>> RenderContextComboBoxItems;
+
+	TArray<TSharedPtr<FString>> MaterialPurposeComboBoxItems;
 };
 
 #endif // WITH_EDITOR
\ No newline at end of file
diff --git a/Engine/Plugins/Importers/USDImporter/Source/USDStage/Public/USDStageActor.h b/Engine/Plugins/Importers/USDImporter/Source/USDStage/Public/USDStageActor.h
--- a/Engine/Plugins/Importers/USDImporter/Source/USDStage/Public/USDStageActor.h
+++ b/Engine/Plugins/Importers/USDImporter/Source/USDStage/Public/USDStageActor.h
@@ -83,6 +83,10 @@ public:
 	UPROPERTY(EditAnywhere, BlueprintReadOnly, Category = "USD")
 	FName RenderContext;
 
+	/** Specifies which material purbose to use when parsing USD material bindings, in addition to the "allPurpose" fallback. */
+	UPROPERTY( EditAnywhere, BlueprintReadOnly, Category = "USD" )
+	FName MaterialPurpose;
+
 public:
 	UFUNCTION(BlueprintCallable, Category = "USD", meta = (CallInEditor = "true"))
 	USDSTAGE_API void SetRootLayer(const FString& RootFilePath );
@@ -111,6 +115,9 @@ public:
 	UFUNCTION(BlueprintCallable, Category = "USD", meta = (CallInEditor = "true"))
 	USDSTAGE_API void SetRenderContext( const FName& NewRenderContext );
 
+	UFUNCTION( BlueprintCallable, Category = "USD", meta = ( CallInEditor = "true" ) )
+	USDSTAGE_API void SetMaterialPurpose( const FName& NewMaterialPurpose );
+
 	UFUNCTION(BlueprintCallable, Category = "USD", meta = (CallInEditor = "true"))
 	USDSTAGE_API float GetTime() const { return Time; }
 
diff --git a/Engine/Plugins/Importers/USDImporter/Source/USDStageEditor/Private/SUSDStage.cpp b/Engine/Plugins/Importers/USDImporter/Source/USDStageEditor/Private/SUSDStage.cpp
--- a/Engine/Plugins/Importers/USDImporter/Source/USDStageEditor/Private/SUSDStage.cpp
+++ b/Engine/Plugins/Importers/USDImporter/Source/USDStageEditor/Private/SUSDStage.cpp
@@ -10,6 +10,7 @@
 #include "USDConversionUtils.h"
 #include "USDErrorUtils.h"
 #include "USDLayerUtils.h"
+#include "USDProjectSettings.h"
 #include "USDSchemasModule.h"
 #include "USDSchemaTranslator.h"
 #include "USDStageActor.h"
@@ -28,7 +29,6 @@
 #include "Async/Async.h"
 #include "DesktopPlatformModule.h"
 #include "Dialogs/DlgPickPath.h"
-#include "Styling/AppStyle.h"
 #include "Engine/Selection.h"
 #include "Engine/World.h"
 #include "EngineAnalytics.h"
@@ -37,7 +37,9 @@
 #include "Modules/ModuleManager.h"
 #include "SceneOutlinerModule.h"
 #include "ScopedTransaction.h"
+#include "Styling/AppStyle.h"
 #include "UObject/StrongObjectPtr.h"
+#include "Widgets/Input/SEditableTextBox.h"
 #include "Widgets/Input/SSpinBox.h"
 #include "Widgets/Layout/SSplitter.h"
 
@@ -713,10 +715,15 @@ void SUsdStage::FillOptionsMenu(FMenuBuilder& MenuBuilder)
 			FNewMenuDelegate::CreateSP(this, &SUsdStage::FillPurposesToLoadSubMenu));
 
 		MenuBuilder.AddSubMenu(
-			LOCTEXT("RenderContext", "Render Context"),
+			LOCTEXT("RenderContext", "Render context"),
 			LOCTEXT("RenderContext_ToolTip", "Choose which render context to use when parsing materials"),
 			FNewMenuDelegate::CreateSP(this, &SUsdStage::FillRenderContextSubMenu));
 
+		MenuBuilder.AddSubMenu(
+			LOCTEXT( "MaterialPurpose", "Material purpose" ),
+			LOCTEXT( "MaterialPurpose_ToolTip", "Material purpose to use when parsing material bindings in addition to the \"allPurpose\" fallback" ),
+			FNewMenuDelegate::CreateSP( this, &SUsdStage::FillMaterialPurposeSubMenu ) );
+
 		MenuBuilder.AddSubMenu(
 			LOCTEXT( "Collapsing", "Collapsing" ),
 			LOCTEXT( "Collapsing_ToolTip", "Whether to try to combine individual assets and components of the same type on a Kind-per-Kind basis, like multiple Mesh prims into a single Static Mesh" ),
@@ -984,6 +991,126 @@ void SUsdStage::FillRenderContextSubMenu( FMenuBuilder& MenuBuilder )
 	}
 }
 
+void SUsdStage::FillMaterialPurposeSubMenu( FMenuBuilder& MenuBuilder )
+{
+	MaterialPurposes = {
+		MakeShared<FString>( UnrealIdentifiers::MaterialAllPurpose ),
+		MakeShared<FString>( UnrealIdentifiers::MaterialPreviewPurpose ),
+		MakeShared<FString>( UnrealIdentifiers::MaterialFullPurpose )
+	};
+
+	// Add additional purposes from project settings
+	if ( const UUsdProjectSettings* ProjectSettings = GetDefault<UUsdProjectSettings>() )
+	{
+		MaterialPurposes.Reserve( MaterialPurposes.Num() + ProjectSettings->AdditionalMaterialPurposes.Num() );
+
+		TSet<FString> ExistingEntries = {
+			UnrealIdentifiers::MaterialAllPurpose,
+			UnrealIdentifiers::MaterialPreviewPurpose,
+			UnrealIdentifiers::MaterialFullPurpose
+		};
+
+		for (const FName& AdditionalPurpose : ProjectSettings->AdditionalMaterialPurposes )
+		{
+			FString AdditionalPurposeStr = AdditionalPurpose.ToString();
+
+			if ( !ExistingEntries.Contains( AdditionalPurposeStr ) )
+			{
+				ExistingEntries.Add( AdditionalPurposeStr );
+				MaterialPurposes.AddUnique( MakeShared<FString>( AdditionalPurposeStr ) );
+			}
+		}
+	}
+
+	TSharedRef<SBox> Box = SNew( SBox )
+	.Padding( FMargin( 8.0f, 0.0f ) )
+	.VAlign( VAlign_Center )
+	[
+		// We have to use TSharedPtr<FString> here as the combobox actually contains a list view.
+		// Also, a regular SComboBox<FName> doesn't even call our OnGenerateWidget function in case the item is NAME_None,
+		// and we do need that case for allPurpose.
+		SNew( SComboBox< TSharedPtr<FString> > )
+		.OptionsSource( &MaterialPurposes )
+		.OnGenerateWidget_Lambda( [ & ]( TSharedPtr<FString> Option )
+		{
+			TSharedPtr<SWidget> Widget = SNullWidget::NullWidget;
+			if ( Option )
+			{
+				Widget = SNew( STextBlock )
+					.Text( FText::FromString( (*Option) == UnrealIdentifiers::MaterialAllPurpose
+						? UnrealIdentifiers::MaterialAllPurposeText
+						: *Option
+					))
+					.Font( FAppStyle::GetFontStyle( "PropertyWindow.NormalFont" ) );
+			}
+
+			return Widget.ToSharedRef();
+		})
+		.OnSelectionChanged_Lambda([this]( TSharedPtr<FString> ChosenOption, ESelectInfo::Type SelectInfo )
+		{
+			if ( ChosenOption )
+			{
+				if ( AUsdStageActor* StageActor = ViewModel.UsdStageActor.Get() )
+				{
+					StageActor->SetMaterialPurpose( **ChosenOption );
+				}
+			}
+		})
+		[
+			SNew( SEditableTextBox )
+			.Text_Lambda([this]() -> FText
+			{
+				if ( AUsdStageActor* StageActor = ViewModel.UsdStageActor.Get() )
+				{
+					return FText::FromString( StageActor->MaterialPurpose == *UnrealIdentifiers::MaterialAllPurpose
+						? UnrealIdentifiers::MaterialAllPurposeText
+						: StageActor->MaterialPurpose.ToString()
+					);
+				}
+
+				return FText::GetEmpty();
+			})
+			.Font( FAppStyle::GetFontStyle( "PropertyWindow.NormalFont" ) )
+			.OnTextCommitted_Lambda( [this]( const FText& NewText, ETextCommit::Type CommitType )
+			{
+				if ( CommitType != ETextCommit::OnEnter )
+				{
+					return;
+				}
+
+				FString NewPurposeString = NewText.ToString();
+				FName NewPurpose = *NewPurposeString;
+
+				bool bIsNew = true;
+				for ( const TSharedPtr<FString>& Purpose : MaterialPurposes )
+				{
+					if ( Purpose && *Purpose == NewPurposeString )
+					{
+						bIsNew = false;
+						break;
+					}
+				}
+
+				if ( bIsNew )
+				{
+					if ( UUsdProjectSettings* ProjectSettings = GetMutableDefault<UUsdProjectSettings>() )
+					{
+						ProjectSettings->AdditionalMaterialPurposes.AddUnique( NewPurpose );
+					}
+				}
+
+				if ( AUsdStageActor* StageActor = ViewModel.UsdStageActor.Get() )
+				{
+					StageActor->SetMaterialPurpose( NewPurpose );
+				}
+			})
+		]
+	];
+
+	const bool bNoIndent = true;
+	MenuBuilder.AddWidget( Box, FText::GetEmpty(), bNoIndent );
+}
+
 void SUsdStage::FillCollapsingSubMenu( FMenuBuilder& MenuBuilder )
 {
 	MenuBuilder.AddMenuEntry(
diff --git a/Engine/Plugins/Importers/USDImporter/Source/USDStageEditor/Private/SUSDStage.h b/Engine/Plugins/Importers/USDImporter/Source/USDStageEditor/Private/SUSDStage.h
--- a/Engine/Plugins/Importers/USDImporter/Source/USDStageEditor/Private/SUSDStage.h
+++ b/Engine/Plugins/Importers/USDImporter/Source/USDStageEditor/Private/SUSDStage.h
@@ -46,6 +46,7 @@ protected:
 	void FillPayloadsSubMenu( FMenuBuilder& MenuBuilder );
 	void FillPurposesToLoadSubMenu( FMenuBuilder& MenuBuilder );
 	void FillRenderContextSubMenu( FMenuBuilder& MenuBuilder );
+	void FillMaterialPurposeSubMenu( FMenuBuilder& MenuBuilder );
 	void FillCollapsingSubMenu( FMenuBuilder& MenuBuilder );
 	void FillInterpolationTypeSubMenu( FMenuBuilder& MenuBuilder );
 	void FillSelectionSubMenu( FMenuBuilder& MenuBuilder );
@@ -111,6 +112,8 @@ protected:
 	TSharedPtr<ISceneOutliner> ActorPickerMenu;
 
 	int32 CurrentNaniteThreshold = INT32_MAX;
+
+	TArray<TSharedPtr<FString>> MaterialPurposes;
 };
 
 #endif // #if USE_USD_SDK
diff --git a/Engine/Plugins/Importers/USDImporter/Source/USDStageEditorViewModels/Private/USDStageViewModel.cpp b/Engine/Plugins/Importers/USDImporter/Source/USDStageEditorViewModels/Private/USDStageViewModel.cpp
--- a/Engine/Plugins/Importers/USDImporter/Source/USDStageEditorViewModels/Private/USDStageViewModel.cpp
+++ b/Engine/Plugins/Importers/USDImporter/Source/USDStageEditorViewModels/Private/USDStageViewModel.cpp
@@ -286,6 +286,7 @@ void FUsdStageViewModel::ImportStage()
 		// Preload some settings according to USDStage options. These will overwrite whatever is loaded from config
 		ImportContext.ImportOptions->PurposesToImport = StageActor->PurposesToLoad;
 		ImportContext.ImportOptions->RenderContextToImport = StageActor->RenderContext;
+		ImportContext.ImportOptions->MaterialPurpose = StageActor->MaterialPurpose;
 		ImportContext.ImportOptions->StageOptions.MetersPerUnit = UsdUtils::GetUsdStageMetersPerUnit( UsdStage );
 		ImportContext.ImportOptions->StageOptions.UpAxis = UsdUtils::GetUsdStageUpAxisAsEnum( UsdStage );
 		ImportContext.bReadFromStageCache = true; // So that we import whatever the user has open right now, even if the file has changes
diff --git a/Engine/Plugins/Importers/USDImporter/Source/USDStageImporter/Private/USDStageImportOptions.cpp b/Engine/Plugins/Importers/USDImporter/Source/USDStageImporter/Private/USDStageImportOptions.cpp
--- a/Engine/Plugins/Importers/USDImporter/Source/USDStageImporter/Private/USDStageImportOptions.cpp
+++ b/Engine/Plugins/Importers/USDImporter/Source/USDStageImporter/Private/USDStageImportOptions.cpp
@@ -21,6 +21,7 @@ UUsdStageImportOptions::UUsdStageImportOptions(const FObjectInitializer& ObjectI
 	NaniteTriangleThreshold = INT32_MAX;
 	IUsdSchemasModule& UsdSchemasModule = FModuleManager::Get().LoadModuleChecked< IUsdSchemasModule >( TEXT("USDSchemas") );
 	RenderContextToImport = UsdSchemasModule.GetRenderContextRegistry().GetUnrealRenderContext();
+	MaterialPurpose = *UnrealIdentifiers::MaterialAllPurpose;
 	bOverrideStageOptions = false;
 	StageOptions.MetersPerUnit = 0.01f;
 	StageOptions.UpAxis = EUsdUpAxis::ZAxis;
diff --git a/Engine/Plugins/Importers/USDImporter/Source/USDStageImporter/Private/USDStageImportOptionsCustomization.cpp b/Engine/Plugins/Importers/USDImporter/Source/USDStageImporter/Private/USDStageImportOptionsCustomization.cpp
--- a/Engine/Plugins/Importers/USDImporter/Source/USDStageImporter/Private/USDStageImportOptionsCustomization.cpp
+++ b/Engine/Plugins/Importers/USDImporter/Source/USDStageImporter/Private/USDStageImportOptionsCustomization.cpp
@@ -2,6 +2,7 @@
 
 #include "USDStageImportOptionsCustomization.h"
 
+#include "USDProjectSettings.h"
 #include "USDSchemasModule.h"
 #include "USDSchemaTranslator.h"
 #include "USDStageImportOptions.h"
@@ -10,10 +11,10 @@
 #include "DetailLayoutBuilder.h"
 #include "DetailWidgetRow.h"
 #include "Modules/ModuleManager.h"
+#include "Styling/AppStyle.h"
 #include "UObject/Object.h"
 #include "Widgets/DeclarativeSyntaxSupport.h"
-#include "Styling/AppStyle.h"
-
+#include "Widgets/Input/SEditableTextBox.h"
 
 #define LOCTEXT_NAMESPACE "UsdStageImportOptionsCustomization"
 
@@ -48,7 +49,7 @@ void FUsdStageImportOptionsCustomization::CustomizeDetails(IDetailLayoutBuilder&
 
 	IUsdSchemasModule& UsdSchemasModule = FModuleManager::Get().LoadModuleChecked< IUsdSchemasModule >( TEXT( "USDSchemas" ) );
 
-	ComboBoxItems.Reset();
+	RenderContextComboBoxItems.Reset();
 	TSharedPtr<FString> InitiallySelectedContext;
 	for ( const FName& Context : UsdSchemasModule.GetRenderContextRegistry().GetRenderContexts() )
 	{
@@ -67,7 +68,7 @@ void FUsdStageImportOptionsCustomization::CustomizeDetails(IDetailLayoutBuilder&
 			InitiallySelectedContext = ContextStr;
 		}
 
-		ComboBoxItems.Add( ContextStr );
+		RenderContextComboBoxItems.Add( ContextStr );
 	}
 
 	IDetailCategoryBuilder& CatBuilder = DetailLayoutBuilder.EditCategory( TEXT( "USD options" ) );
@@ -86,8 +87,8 @@ void FUsdStageImportOptionsCustomization::CustomizeDetails(IDetailLayoutBuilder&
 		]
 		.ValueContent()
 		[
-			SAssignNew( ComboBox, SComboBox<TSharedPtr<FString>> )
-			.OptionsSource( &ComboBoxItems )
+			SAssignNew( RenderContextComboBox, SComboBox<TSharedPtr<FString>> )
+			.OptionsSource( &RenderContextComboBoxItems )
 			.InitiallySelectedItem( InitiallySelectedContext )
 			.OnSelectionChanged( this, &FUsdStageImportOptionsCustomization::OnComboBoxSelectionChanged )
 			.OnGenerateWidget_Lambda( []( TSharedPtr<FString> Item )
@@ -105,6 +106,132 @@ void FUsdStageImportOptionsCustomization::CustomizeDetails(IDetailLayoutBuilder&
 		];
 	}
 
+
+	if ( TSharedPtr<IPropertyHandle> MaterialPurposeProperty = DetailLayoutBuilder.GetProperty( GET_MEMBER_NAME_CHECKED( UUsdStageImportOptions, MaterialPurpose ) ) )
+	{
+		DetailLayoutBuilder.HideProperty( MaterialPurposeProperty );
+
+		CatBuilder.AddCustomRow( FText::FromString( TEXT( "MaterialPurposeCustomization" ) ) )
+		.NameContent()
+		[
+			SNew( STextBlock )
+			.Text( FText::FromString( TEXT( "Material purpose" ) ) )
+			.Font( FAppStyle::GetFontStyle( TEXT( "PropertyWindow.NormalFont" ) ) )
+			.ToolTipText( MaterialPurposeProperty->GetToolTipText() )
+		]
+		.ValueContent()
+		[
+			SNew( SBox )
+			.VAlign( VAlign_Center )
+			[
+				SNew( SComboBox< TSharedPtr<FString> > )
+				.OptionsSource( &MaterialPurposeComboBoxItems )
+				.OnComboBoxOpening_Lambda([this]()
+				{
+					MaterialPurposeComboBoxItems = {
+						MakeShared<FString>( UnrealIdentifiers::MaterialAllPurpose ),
+						MakeShared<FString>( UnrealIdentifiers::MaterialPreviewPurpose ),
+						MakeShared<FString>( UnrealIdentifiers::MaterialFullPurpose )
+					};
+
+					// Add additional purposes from project settings
+					if ( const UUsdProjectSettings* ProjectSettings = GetDefault<UUsdProjectSettings>() )
+					{
+						MaterialPurposeComboBoxItems.Reserve( MaterialPurposeComboBoxItems.Num() + ProjectSettings->AdditionalMaterialPurposes.Num() );
+
+						TSet<FString> ExistingEntries = {
+							UnrealIdentifiers::MaterialAllPurpose,
+							UnrealIdentifiers::MaterialPreviewPurpose,
+							UnrealIdentifiers::MaterialFullPurpose
+						};
+
+						for ( const FName& AdditionalPurpose : ProjectSettings->AdditionalMaterialPurposes )
+						{
+							FString AdditionalPurposeStr = AdditionalPurpose.ToString();
+
+							if ( !ExistingEntries.Contains( AdditionalPurposeStr ) )
+							{
+								ExistingEntries.Add( AdditionalPurposeStr );
+								MaterialPurposeComboBoxItems.AddUnique( MakeShared<FString>( AdditionalPurposeStr ) );
+							}
+						}
+					}
+				})
+				.OnGenerateWidget_Lambda( [ & ]( TSharedPtr<FString> Option )
+				{
+					TSharedPtr<SWidget> Widget = SNullWidget::NullWidget;
+					if ( Option )
+					{
+						Widget = SNew( STextBlock )
+							.Text( FText::FromString( ( *Option ) == UnrealIdentifiers::MaterialAllPurpose
+								? UnrealIdentifiers::MaterialAllPurposeText
+								: *Option
+							) )
+							.Font( FAppStyle::GetFontStyle( "PropertyWindow.NormalFont" ) );
+					}
+
+					return Widget.ToSharedRef();
+				})
+				.OnSelectionChanged_Lambda([this]( TSharedPtr<FString> ChosenOption, ESelectInfo::Type SelectInfo )
+				{
+					if ( CurrentOptions && ChosenOption )
+					{
+						CurrentOptions->MaterialPurpose = **ChosenOption;
+					}
+				})
+				[
+					SNew( SEditableTextBox )
+					.Text_Lambda([this]() -> FText
+					{
+						if ( CurrentOptions )
+						{
+							return FText::FromString( CurrentOptions->MaterialPurpose == *UnrealIdentifiers::MaterialAllPurpose
+								? UnrealIdentifiers::MaterialAllPurposeText
+								: CurrentOptions->MaterialPurpose.ToString()
+							);
+						}
+
+						return FText::GetEmpty();
+					})
+					.Font( FAppStyle::GetFontStyle( "PropertyWindow.NormalFont" ) )
+					.OnTextCommitted_Lambda( [this]( const FText& NewText, ETextCommit::Type CommitType )
+					{
+						if ( CommitType != ETextCommit::OnEnter )
+						{
+							return;
+						}
+
+						FString NewPurposeString = NewText.ToString();
+						FName NewPurpose = *NewPurposeString;
+
+						bool bIsNew = true;
+						for ( const TSharedPtr<FString>& Purpose : MaterialPurposeComboBoxItems )
+						{
+							if ( Purpose && *Purpose == NewPurposeString )
+							{
+								bIsNew = false;
+								break;
+							}
+						}
+
+						if ( bIsNew )
+						{
+							if ( UUsdProjectSettings* ProjectSettings = GetMutableDefault<UUsdProjectSettings>() )
+							{
+								ProjectSettings->AdditionalMaterialPurposes.AddUnique( NewPurpose );
+							}
+						}
+
+						if ( CurrentOptions )
+						{
+							CurrentOptions->MaterialPurpose = NewPurpose;
+						}
+					})
+				]
+			]
+		];
+	}
+
 	// Add/remove properties so that they retain their usual order
 	if ( TSharedPtr<IPropertyHandle> OverrideStageOptionsProperty = DetailLayoutBuilder.GetProperty( GET_MEMBER_NAME_CHECKED( UUsdStageImportOptions, bOverrideStageOptions ) ) )
 	{
@@ -135,7 +262,7 @@ void FUsdStageImportOptionsCustomization::OnComboBoxSelectionChanged( TSharedPtr
 
 FText FUsdStageImportOptionsCustomization::GetComboBoxSelectedOptionText() const
 {
-	TSharedPtr<FString> SelectedItem = ComboBox->GetSelectedItem();
+	TSharedPtr<FString> SelectedItem = RenderContextComboBox->GetSelectedItem();
 	if ( SelectedItem.IsValid() )
 	{
 		return FText::FromString( *SelectedItem );
diff --git a/Engine/Plugins/Importers/USDImporter/Source/USDStageImporter/Private/USDStageImportOptionsCustomization.h b/Engine/Plugins/Importers/USDImporter/Source/USDStageImporter/Private/USDStageImportOptionsCustomization.h
--- a/Engine/Plugins/Importers/USDImporter/Source/USDStageImporter/Private/USDStageImportOptionsCustomization.h
+++ b/Engine/Plugins/Importers/USDImporter/Source/USDStageImporter/Private/USDStageImportOptionsCustomization.h
@@ -29,8 +29,11 @@ private:
 	void ForceRefreshDetails();
 
 	UUsdStageImportOptions* CurrentOptions;
-	TSharedPtr<SComboBox<TSharedPtr<FString>>> ComboBox;
-	TArray<TSharedPtr<FString>> ComboBoxItems;
+
+	TSharedPtr<SComboBox<TSharedPtr<FString>>> RenderContextComboBox;
+	TArray<TSharedPtr<FString>> RenderContextComboBoxItems;
+
+	TArray<TSharedPtr<FString>> MaterialPurposeComboBoxItems;
 };
 
 #endif // WITH_EDITOR
\ No newline at end of file
diff --git a/Engine/Plugins/Importers/USDImporter/Source/USDStageImporter/Private/USDStageImporter.cpp b/Engine/Plugins/Importers/USDImporter/Source/USDStageImporter/Private/USDStageImporter.cpp
--- a/Engine/Plugins/Importers/USDImporter/Source/USDStageImporter/Private/USDStageImporter.cpp
+++ b/Engine/Plugins/Importers/USDImporter/Source/USDStageImporter/Private/USDStageImporter.cpp
@@ -1276,6 +1276,7 @@ namespace UsdStageImporterImpl
 				EventAttributes.Emplace( TEXT( "PurposesToImport" ), LexToString( ImportContext.ImportOptions->PurposesToImport ) );
 				EventAttributes.Emplace( TEXT( "NaniteTriangleThreshold" ), LexToString( ImportContext.ImportOptions->NaniteTriangleThreshold ) );
 				EventAttributes.Emplace( TEXT( "RenderContextToImport" ), ImportContext.ImportOptions->RenderContextToImport.ToString() );
+				EventAttributes.Emplace( TEXT( "MaterialPurpose" ), ImportContext.ImportOptions->MaterialPurpose.ToString() );
 				EventAttributes.Emplace( TEXT( "OverrideStageOptions" ), ImportContext.ImportOptions->bOverrideStageOptions );
 				if( ImportContext.ImportOptions->bOverrideStageOptions )
 				{
@@ -1492,6 +1493,7 @@ void UUsdStageImporter::ImportFromFile(FUsdStageImportContext& ImportContext)
 	TranslationContext->PurposesToLoad = ( EUsdPurpose ) ImportContext.ImportOptions->PurposesToImport;
 	TranslationContext->NaniteTriangleThreshold = ImportContext.ImportOptions->NaniteTriangleThreshold;
 	TranslationContext->RenderContext = ImportContext.ImportOptions->RenderContextToImport;
+	TranslationContext->MaterialPurpose = ImportContext.ImportOptions->MaterialPurpose;
 	TranslationContext->ParentComponent = ImportContext.SceneActor ? ImportContext.SceneActor->GetRootComponent() : nullptr;
 	TranslationContext->KindsToCollapse = ( EUsdDefaultKind ) ImportContext.ImportOptions->KindsToCollapse;
 	TranslationContext->bMergeIdenticalMaterialSlots = ImportContext.ImportOptions->bMergeIdenticalMaterialSlots;
@@ -1583,8 +1585,10 @@ bool UUsdStageImporter::ReimportSingleAsset(FUsdStageImportContext& ImportContex
 	TranslationContext->Level = ImportContext.World->GetCurrentLevel();
 	TranslationContext->ObjectFlags = ImportContext.ImportObjectFlags;
 	TranslationContext->Time = static_cast< float >( UsdUtils::GetDefaultTimeCode() );
-	TranslationContext->PurposesToLoad = (EUsdPurpose) ImportContext.ImportOptions->PurposesToImport;
+	TranslationContext->PurposesToLoad = ( EUsdPurpose ) ImportContext.ImportOptions->PurposesToImport;
 	TranslationContext->NaniteTriangleThreshold = ImportContext.ImportOptions->NaniteTriangleThreshold;
+	TranslationContext->RenderContext = ImportContext.ImportOptions->RenderContextToImport;
+	TranslationContext->MaterialPurpose = ImportContext.ImportOptions->MaterialPurpose;
 	TranslationContext->KindsToCollapse = ( EUsdDefaultKind ) ImportContext.ImportOptions->KindsToCollapse;
 	TranslationContext->bMergeIdenticalMaterialSlots = ImportContext.ImportOptions->bMergeIdenticalMaterialSlots;
 	TranslationContext->bCollapseTopLevelPointInstancers = ImportContext.ImportOptions->bCollapseTopLevelPointInstancers;
@@ -1593,6 +1597,7 @@ bool UUsdStageImporter::ReimportSingleAsset(FUsdStageImportContext& ImportContex
 	TranslationContext->MaterialToPrimvarToUVIndex = &ImportContext.MaterialToPrimvarToUVIndex;
 	TranslationContext->InfoCache = ImportContext.InfoCache;
 	TranslationContext->BlendShapesByPath = &BlendShapesByPath;
+	TranslationContext->GroomInterpolationSettings = ImportContext.ImportOptions->GroomInterpolationSettings;
 	{
 		UsdStageImporterImpl::CacheCollapsingState( TranslationContext.Get() );
 
diff --git a/Engine/Plugins/Importers/USDImporter/Source/USDStageImporter/Public/USDStageImportOptions.h b/Engine/Plugins/Importers/USDImporter/Source/USDStageImporter/Public/USDStageImportOptions.h
--- a/Engine/Plugins/Importers/USDImporter/Source/USDStageImporter/Public/USDStageImportOptions.h
+++ b/Engine/Plugins/Importers/USDImporter/Source/USDStageImporter/Public/USDStageImportOptions.h
@@ -78,6 +78,10 @@ public:
 	UPROPERTY(BlueprintReadWrite, config, EditAnywhere, Category= "USD options")
 	FName RenderContextToImport;
 
+	/** Specifies which material purpose to use when parsing USD material bindings, in addition to the "allPurpose" fallback */
+	UPROPERTY( BlueprintReadWrite, config, EditAnywhere, Category = "USD options" )
+	FName MaterialPurpose;
+
 	/** Whether to use the specified StageOptions instead of the stage's own settings */
 	UPROPERTY( BlueprintReadWrite, config, EditAnywhere, Category = "USD options" )
 	bool bOverrideStageOptions;
diff --git a/Engine/Plugins/Importers/USDImporter/Source/USDUtilities/Private/USDConversionUtils.cpp b/Engine/Plugins/Importers/USDImporter/Source/USDUtilities/Private/USDConversionUtils.cpp
--- a/Engine/Plugins/Importers/USDImporter/Source/USDUtilities/Private/USDConversionUtils.cpp
+++ b/Engine/Plugins/Importers/USDImporter/Source/USDUtilities/Private/USDConversionUtils.cpp
@@ -583,7 +583,12 @@ TArray< TUsdStore< pxr::UsdGeomPrimvar > > UsdUtils::GetUVSetPrimvars( const pxr
 	return UsdUtils::GetUVSetPrimvars(UsdMesh, {});
 }
 
-TArray< TUsdStore< pxr::UsdGeomPrimvar > > UsdUtils::GetUVSetPrimvars( const pxr::UsdGeomMesh& UsdMesh, const TMap< FString, TMap< FString, int32 > >& MaterialToPrimvarsUVSetNames, const pxr::TfToken& RenderContext )
+TArray< TUsdStore< pxr::UsdGeomPrimvar > > UsdUtils::GetUVSetPrimvars(
+	const pxr::UsdGeomMesh& UsdMesh,
+	const TMap< FString, TMap< FString, int32 > >& MaterialToPrimvarsUVSetNames,
+	const pxr::TfToken& RenderContext,
+	const pxr::TfToken& MaterialPurpose
+)
 {
 	if ( !UsdMesh )
 	{
@@ -591,7 +596,14 @@ TArray< TUsdStore< pxr::UsdGeomPrimvar > > UsdUtils::GetUVSetPrimvars( const pxr
 	}
 
 	const bool bProvideMaterialIndices = false;
-	UsdUtils::FUsdPrimMaterialAssignmentInfo Info = UsdUtils::GetPrimMaterialAssignments( UsdMesh.GetPrim(), pxr::UsdTimeCode( 0.0 ), bProvideMaterialIndices, RenderContext );
+	UsdUtils::FUsdPrimMaterialAssignmentInfo Info = UsdUtils::GetPrimMaterialAssignments(
+		UsdMesh.GetPrim(),
+		pxr::UsdTimeCode( 0.0 ),
+		bProvideMaterialIndices,
+		RenderContext,
+		MaterialPurpose
+	);
+
 	return UsdUtils::GetUVSetPrimvars( UsdMesh, MaterialToPrimvarsUVSetNames, Info );
 }
 
diff --git a/Engine/Plugins/Importers/USDImporter/Source/USDUtilities/Private/USDGeomMeshConversion.cpp b/Engine/Plugins/Importers/USDImporter/Source/USDUtilities/Private/USDGeomMeshConversion.cpp
--- a/Engine/Plugins/Importers/USDImporter/Source/USDUtilities/Private/USDGeomMeshConversion.cpp
+++ b/Engine/Plugins/Importers/USDImporter/Source/USDUtilities/Private/USDGeomMeshConversion.cpp
@@ -625,6 +625,7 @@ namespace UsdToUnreal
 		: AdditionalTransform( FTransform::Identity )
 		, PurposesToLoad( EUsdPurpose::Render )
 		, RenderContext( pxr::UsdShadeTokens->universalRenderContext )
+		, MaterialPurpose( pxr::UsdShadeTokens->allPurpose )
 		, TimeCode( pxr::UsdTimeCode::EarliestTime() )
 		, MaterialToPrimvarToUVIndex( nullptr )
 		, bMergeIdenticalMaterialSlots( true )
@@ -662,7 +663,8 @@ bool UsdToUnreal::ConvertGeomMesh(
 		UsdPrim,
 		Options.TimeCode,
 		bProvideMaterialIndices,
-		Options.RenderContext
+		Options.RenderContext,
+		Options.MaterialPurpose
 	);
 	TArray< UsdUtils::FUsdPrimMaterialSlot >& LocalMaterialSlots = LocalInfo.Slots;
 	TArray< int32 >& FaceMaterialIndices = LocalInfo.MaterialIndices;
@@ -1465,7 +1467,13 @@ UMaterialInstanceConstant* UsdUtils::CreateDisplayColorMaterialInstanceConstant(
 	return nullptr;
 }
 
-UsdUtils::FUsdPrimMaterialAssignmentInfo UsdUtils::GetPrimMaterialAssignments( const pxr::UsdPrim& UsdPrim, const pxr::UsdTimeCode TimeCode, bool bProvideMaterialIndices, const pxr::TfToken& RenderContext )
+UsdUtils::FUsdPrimMaterialAssignmentInfo UsdUtils::GetPrimMaterialAssignments(
+	const pxr::UsdPrim& UsdPrim,
+	const pxr::UsdTimeCode TimeCode,
+	bool bProvideMaterialIndices,
+	const pxr::TfToken& RenderContext,
+	const pxr::TfToken& MaterialPurpose
+)
 {
 	if ( !UsdPrim )
 	{
@@ -1522,10 +1530,10 @@ UsdUtils::FUsdPrimMaterialAssignmentInfo UsdUtils::GetPrimMaterialAssignments( c
 		return {};
 	};
 
-	auto FetchMaterialByComputingBoundMaterial = [ &RenderContext ]( const pxr::UsdPrim& UsdPrim ) -> TOptional<FString>
+	auto FetchMaterialByComputingBoundMaterial = [ &RenderContext, &MaterialPurpose ]( const pxr::UsdPrim& UsdPrim ) -> TOptional<FString>
 	{
 		pxr::UsdShadeMaterialBindingAPI BindingAPI( UsdPrim );
-		pxr::UsdShadeMaterial ShadeMaterial = BindingAPI.ComputeBoundMaterial();
+		pxr::UsdShadeMaterial ShadeMaterial = BindingAPI.ComputeBoundMaterial( MaterialPurpose );
 		if ( !ShadeMaterial )
 		{
 			return {};
@@ -1642,7 +1650,7 @@ UsdUtils::FUsdPrimMaterialAssignmentInfo UsdUtils::GetPrimMaterialAssignments( c
 		pxr::UsdShadeMaterialBindingAPI BindingAPI( UsdPrim );
 		if ( BindingAPI )
 		{
-			if ( pxr::UsdShadeMaterial ShadeMaterial = BindingAPI.ComputeBoundMaterial() )
+			if ( pxr::UsdShadeMaterial ShadeMaterial = BindingAPI.ComputeBoundMaterial( MaterialPurpose ) )
 			{
 				if ( TOptional<FString> UnrealMaterial = UsdUtils::GetUnrealSurfaceOutput( ShadeMaterial.GetPrim() ) )
 				{
@@ -1717,7 +1725,7 @@ UsdUtils::FUsdPrimMaterialAssignmentInfo UsdUtils::GetPrimMaterialAssignments( c
 					pxr::UsdShadeMaterialBindingAPI BindingAPI( GeomSubset.GetPrim() );
 					if ( BindingAPI )
 					{
-						if ( pxr::UsdShadeMaterial ShadeMaterial = BindingAPI.ComputeBoundMaterial() )
+						if ( pxr::UsdShadeMaterial ShadeMaterial = BindingAPI.ComputeBoundMaterial( MaterialPurpose ) )
 						{
 							if ( TOptional<FString> UnrealMaterial = UsdUtils::GetUnrealSurfaceOutput( ShadeMaterial.GetPrim() ) )
 							{
@@ -1833,7 +1841,7 @@ UsdUtils::FUsdPrimMaterialAssignmentInfo UsdUtils::GetPrimMaterialAssignments( c
 	return Result;
 }
 
-TArray<FString> UsdUtils::GetMaterialUsers( const UE::FUsdPrim& MaterialPrim )
+TArray<FString> UsdUtils::GetMaterialUsers( const UE::FUsdPrim& MaterialPrim, FName MaterialPurpose )
 {
 	TArray<FString> Result;
 
@@ -1845,6 +1853,12 @@ TArray<FString> UsdUtils::GetMaterialUsers( const UE::FUsdPrim& MaterialPrim )
 		return Result;
 	}
 
+	pxr::TfToken MaterialPurposeToken = pxr::UsdShadeTokens->allPurpose;
+	if ( !MaterialPurpose.IsNone() )
+	{
+		MaterialPurposeToken = UnrealToUsd::ConvertToken( *MaterialPurpose.ToString() ).Get();
+	}
+
 	pxr::UsdStageRefPtr UsdStage = UsdMaterialPrim.GetStage();
 
 	pxr::UsdPrimRange PrimRange = pxr::UsdPrimRange::Stage( UsdStage, pxr::UsdTraverseInstanceProxies() );
@@ -1858,7 +1872,7 @@ TArray<FString> UsdUtils::GetMaterialUsers( const UE::FUsdPrim& MaterialPrim )
 		}
 
 		pxr::UsdShadeMaterialBindingAPI BindingAPI( Prim );
-		pxr::UsdShadeMaterial ShadeMaterial = BindingAPI.ComputeBoundMaterial();
+		pxr::UsdShadeMaterial ShadeMaterial = BindingAPI.ComputeBoundMaterial( MaterialPurposeToken );
 		if ( !ShadeMaterial )
 		{
 			continue;
@@ -2476,6 +2490,8 @@ void UsdUtils::ReplaceUnrealMaterialsWithBaked(
 			pxr::UsdShadeMaterialBindingAPI MaterialBindingAPI{ Prim };
 			if ( MaterialBindingAPI )
 			{
+				// We always emit UnrealMaterials with allpurpose bindings, so we can use default arguments for
+				// ComputeBoundMaterial
 				if ( pxr::UsdShadeMaterial BoundMaterial = MaterialBindingAPI.ComputeBoundMaterial() )
 				{
 					UnrealMaterial = BoundMaterial;
@@ -2752,6 +2768,8 @@ FString UsdUtils::HashGeomMeshPrim( const UE::FUsdStage& Stage, const FString& P
 		MD5.Update( ( uint8* ) UsdOpacities.cdata(), UsdOpacities.size() * sizeof( float ) );
 	}
 
+	// TODO: This is not providing render context or material purpose, so it will never consider float2f primvars
+	// for the hash, which could be an issue in very exotic cases
 	TArray< TUsdStore< UsdGeomPrimvar > > PrimvarsByUVIndex = UsdUtils::GetUVSetPrimvars( UsdMesh );
 	for ( int32 UVChannelIndex = 0; UVChannelIndex < PrimvarsByUVIndex.Num(); ++UVChannelIndex )
 	{
diff --git a/Engine/Plugins/Importers/USDImporter/Source/USDUtilities/Private/USDSkeletalDataConversion.cpp b/Engine/Plugins/Importers/USDImporter/Source/USDUtilities/Private/USDSkeletalDataConversion.cpp
--- a/Engine/Plugins/Importers/USDImporter/Source/USDUtilities/Private/USDSkeletalDataConversion.cpp
+++ b/Engine/Plugins/Importers/USDImporter/Source/USDUtilities/Private/USDSkeletalDataConversion.cpp
@@ -1136,7 +1136,15 @@ bool UsdToUnreal::ConvertSkeleton(const pxr::UsdSkelSkeletonQuery& SkeletonQuery
 	return true;
 }
 
-bool UsdToUnreal::ConvertSkinnedMesh(const pxr::UsdSkelSkinningQuery& SkinningQuery, const pxr::UsdSkelSkeletonQuery& SkeletonQuery, FSkeletalMeshImportData& SkelMeshImportData, TArray< UsdUtils::FUsdPrimMaterialSlot >& MaterialAssignments, const TMap< FString, TMap< FString, int32 > >& MaterialToPrimvarsUVSetNames, const pxr::TfToken& RenderContext )
+bool UsdToUnreal::ConvertSkinnedMesh(
+	const pxr::UsdSkelSkinningQuery& SkinningQuery,
+	const pxr::UsdSkelSkeletonQuery& SkeletonQuery,
+	FSkeletalMeshImportData& SkelMeshImportData,
+	TArray< UsdUtils::FUsdPrimMaterialSlot >& MaterialAssignments,
+	const TMap< FString, TMap< FString, int32 > >& MaterialToPrimvarsUVSetNames,
+	const pxr::TfToken& RenderContext,
+	const pxr::TfToken& MaterialPurpose
+)
 {
 	TRACE_CPUPROFILER_EVENT_SCOPE( UsdToUnreal::ConvertSkinnedMesh );
 
@@ -1286,7 +1294,13 @@ bool UsdToUnreal::ConvertSkinnedMesh(const pxr::UsdSkelSkinningQuery& SkinningQu
 
 	// Material assignments
 	const bool bProvideMaterialIndices = true;
-	UsdUtils::FUsdPrimMaterialAssignmentInfo LocalInfo = UsdUtils::GetPrimMaterialAssignments( SkinningPrim, pxr::UsdTimeCode::EarliestTime(), bProvideMaterialIndices, RenderContext );
+	UsdUtils::FUsdPrimMaterialAssignmentInfo LocalInfo = UsdUtils::GetPrimMaterialAssignments(
+		SkinningPrim,
+		pxr::UsdTimeCode::EarliestTime(),
+		bProvideMaterialIndices,
+		RenderContext,
+		MaterialPurpose
+	);
 	TArray< UsdUtils::FUsdPrimMaterialSlot >& LocalMaterialSlots = LocalInfo.Slots;
 	TArray< int32 >& FaceMaterialIndices = LocalInfo.MaterialIndices;
 
@@ -1451,7 +1465,12 @@ bool UsdToUnreal::ConvertSkinnedMesh(const pxr::UsdSkelSkinningQuery& SkinningQu
 
 	TArray< FUVSet > UVSets;
 
-	TArray< TUsdStore< UsdGeomPrimvar > > PrimvarsByUVIndex = UsdUtils::GetUVSetPrimvars( UsdMesh, MaterialToPrimvarsUVSetNames, RenderContext );
+	TArray< TUsdStore< UsdGeomPrimvar > > PrimvarsByUVIndex = UsdUtils::GetUVSetPrimvars(
+		UsdMesh,
+		MaterialToPrimvarsUVSetNames,
+		RenderContext,
+		MaterialPurpose
+	);
 
 	int32 UVChannelIndex = 0;
 	while ( true )
@@ -1729,7 +1748,15 @@ bool UsdToUnreal::ConvertSkinnedMesh(const pxr::UsdSkelSkinningQuery& SkinningQu
 	return true;
 }
 
-bool UsdToUnreal::ConvertSkinnedMesh( const pxr::UsdSkelSkinningQuery& UsdSkinningQuery, const FTransform& AdditionalTransform, FSkeletalMeshImportData& SkelMeshImportData, TArray< UsdUtils::FUsdPrimMaterialSlot >& MaterialAssignments, const TMap< FString, TMap< FString, int32 > >& MaterialToPrimvarsUVSetNames, const pxr::TfToken& RenderContext )
+bool UsdToUnreal::ConvertSkinnedMesh(
+	const pxr::UsdSkelSkinningQuery& UsdSkinningQuery,
+	const FTransform& AdditionalTransform,
+	FSkeletalMeshImportData& SkelMeshImportData,
+	TArray< UsdUtils::FUsdPrimMaterialSlot >& MaterialAssignments,
+	const TMap< FString, TMap< FString, int32 > >& MaterialToPrimvarsUVSetNames,
+	const pxr::TfToken& RenderContext,
+	const pxr::TfToken& MaterialPurpose
+)
 {
 	FScopedUsdAllocs UEAllocs;
 
@@ -1801,7 +1828,8 @@ bool UsdToUnreal::ConvertSkinnedMesh( const pxr::UsdSkelSkinningQuery& UsdSkinni
 			SkelMeshImportData,
 			MaterialAssignments,
 			MaterialToPrimvarsUVSetNames,
-			RenderContext
+			RenderContext,
+			MaterialPurpose
 		);
 	}
 
diff --git a/Engine/Plugins/Importers/USDImporter/Source/USDUtilities/Public/USDConversionUtils.h b/Engine/Plugins/Importers/USDImporter/Source/USDUtilities/Public/USDConversionUtils.h
--- a/Engine/Plugins/Importers/USDImporter/Source/USDUtilities/Public/USDConversionUtils.h
+++ b/Engine/Plugins/Importers/USDImporter/Source/USDUtilities/Public/USDConversionUtils.h
@@ -149,9 +149,15 @@ namespace UsdUtils
 	 * @param MaterialToPrimvarsUVSetNames - Maps from a material prim path, to pairs indicating which primvar names are used as 'st' coordinates, and which UVIndex the imported material will sample from (e.g. ["st0", 0], ["myUvSet2", 2], etc). These are supposed to be the materials used by the mesh, and we do this because it helps identify which primvars are valid/used as texture coordinates, as the user may have these named as 'myUvSet2' and still expect it to work
 	 * @param UsdMeshMaterialAssignmentInfo - Result of calling GetPrimMaterialAssignments on UsdMesh's Prim. This can be provided or will be retrieved on-demand using RenderContext
 	 * @param RenderContext - Render context to use when traversing through material shaders looking for used primvars
+	 * @param MaterialPurpose - Specific material purpose to use when parsing the UsdMesh's material bindings
 	 * @return Array where each index gives the primvar that should be used for that UV index
 	 */
-	USDUTILITIES_API TArray< TUsdStore< pxr::UsdGeomPrimvar > > GetUVSetPrimvars( const pxr::UsdGeomMesh& UsdMesh, const TMap< FString, TMap< FString, int32 > >& MaterialToPrimvarsUVSetNames, const pxr::TfToken& RenderContext = pxr::UsdShadeTokens->universalRenderContext );
+	USDUTILITIES_API TArray< TUsdStore< pxr::UsdGeomPrimvar > > GetUVSetPrimvars(
+		const pxr::UsdGeomMesh& UsdMesh,
+		const TMap< FString, TMap< FString, int32 > >& MaterialToPrimvarsUVSetNames,
+		const pxr::TfToken& RenderContext = pxr::UsdShadeTokens->universalRenderContext,
+		const pxr::TfToken& MaterialPurpose = pxr::UsdShadeTokens->allPurpose
+	);
 	USDUTILITIES_API TArray< TUsdStore< pxr::UsdGeomPrimvar > > GetUVSetPrimvars( const pxr::UsdGeomMesh& UsdMesh, const TMap< FString, TMap< FString, int32 > >& MaterialToPrimvarsUVSetNames, const UsdUtils::FUsdPrimMaterialAssignmentInfo& UsdMeshMaterialAssignmentInfo );
 
 	USDUTILITIES_API bool IsAnimated( const pxr::UsdPrim& Prim );
diff --git a/Engine/Plugins/Importers/USDImporter/Source/USDUtilities/Public/USDGeomMeshConversion.h b/Engine/Plugins/Importers/USDImporter/Source/USDUtilities/Public/USDGeomMeshConversion.h
--- a/Engine/Plugins/Importers/USDImporter/Source/USDUtilities/Public/USDGeomMeshConversion.h
+++ b/Engine/Plugins/Importers/USDImporter/Source/USDUtilities/Public/USDGeomMeshConversion.h
@@ -57,6 +57,9 @@ namespace UsdToUnreal
 		// Render context to use when parsing Mesh materials
 		pxr::TfToken RenderContext;
 
+		// Material purpose to use when parsing Mesh materials
+		pxr::TfToken MaterialPurpose;
+
 		// USD timecode when the prims should be sampled for mesh data and converted
 		pxr::UsdTimeCode TimeCode;
 
@@ -269,16 +272,22 @@ namespace UsdUtils
 	 * @param TimeCode - Instant where the material data is sampled
 	 * @param bProvideMaterialIndices - Whether to fill out the material index information for the assignment info (which can be expensive). If this is false, MaterialIndices on the result Struct will have zero values
 	 * @param RenderContext - Which render context to get the materials for. Defaults to universal.
+	 * @param MaterialPurpose - Which material purpose to use when retrieving material bindings
 	 * @return Struct containing an array of material assignments, and a corresponding array of material indices for all polygons of the prim, matching the ordering of the material assignments.
 	 */
-	USDUTILITIES_API FUsdPrimMaterialAssignmentInfo GetPrimMaterialAssignments( const pxr::UsdPrim& UsdPrim, const pxr::UsdTimeCode TimeCode = pxr::UsdTimeCode::EarliestTime(),
-		bool bProvideMaterialIndices = true, const pxr::TfToken& RenderContext = pxr::UsdShadeTokens->universalRenderContext );
+	USDUTILITIES_API FUsdPrimMaterialAssignmentInfo GetPrimMaterialAssignments(
+		const pxr::UsdPrim& UsdPrim,
+		const pxr::UsdTimeCode TimeCode = pxr::UsdTimeCode::EarliestTime(),
+		bool bProvideMaterialIndices = true,
+		const pxr::TfToken& RenderContext = pxr::UsdShadeTokens->universalRenderContext,
+		const pxr::TfToken& MaterialPurpose = pxr::UsdShadeTokens->allPurpose
+	);
 
 	/**
 	 * Returns an array of prim paths to prims within MaterialPrim's stage that have a UsdShadeMaterialBindingAPI, and whose computed material binding points to MaterialPrim.
 	 * The user prim's schema is not even checked, so this list will naturally include regular UsdGeomMesh prims, UsdGeomSubset prims, and any other with the UsdShadeMaterialBindingAPI.
 	 */
-	USDUTILITIES_API TArray<FString> GetMaterialUsers( const UE::FUsdPrim& MaterialPrim );
+	USDUTILITIES_API TArray<FString> GetMaterialUsers( const UE::FUsdPrim& MaterialPrim, FName MaterialPurpose = *UnrealIdentifiers::MaterialAllPurpose );
 
 	/** Returns whether this prim can be interpreted as describing a static mesh with multiple LODs (i.e. if the prim holds the LOD variant set itself) */
 	USDUTILITIES_API bool DoesPrimContainMeshLODs( const pxr::UsdPrim& Prim );
diff --git a/Engine/Plugins/Importers/USDImporter/Source/USDUtilities/Public/USDSkeletalDataConversion.h b/Engine/Plugins/Importers/USDImporter/Source/USDUtilities/Public/USDSkeletalDataConversion.h
--- a/Engine/Plugins/Importers/USDImporter/Source/USDUtilities/Public/USDSkeletalDataConversion.h
+++ b/Engine/Plugins/Importers/USDImporter/Source/USDUtilities/Public/USDSkeletalDataConversion.h
@@ -165,11 +165,28 @@ namespace UsdToUnreal
 	 * @param MaterialAssignments - Output parameter that will be filled with the material assignment data extracted from UsdSkinningQuery
 	 * @param MaterialToPrimvarsUVSetNames - Maps from a material prim path, to pairs indicating which primvar names are used as 'st' coordinates for this mesh, and which UVIndex materials will sample from (e.g. ["st0", 0], ["myUvSet2", 2], etc). This is used to pick which primvars will become UV sets.
 	 * @param RenderContext - Render context to use when parsing the skinned mesh's materials (e.g. '' for universal, or 'mdl', or 'unreal', etc.)
+	 * @param MaterialPurpose - Material purpose to use when parsing the skinned Mesh prim's material bindings
 	 * @return Whether the conversion was successful or not.
 	 */
-	USDUTILITIES_API bool ConvertSkinnedMesh( const pxr::UsdSkelSkinningQuery& UsdSkinningQuery, const pxr::UsdSkelSkeletonQuery& SkeletonQuery, FSkeletalMeshImportData& SkelMeshImportData, TArray< UsdUtils::FUsdPrimMaterialSlot >& MaterialAssignments, const TMap< FString, TMap< FString, int32 > >& MaterialToPrimvarsUVSetNames, const pxr::TfToken& RenderContext = pxr::UsdShadeTokens->universalRenderContext );
+	USDUTILITIES_API bool ConvertSkinnedMesh(
+		const pxr::UsdSkelSkinningQuery& UsdSkinningQuery,
+		const pxr::UsdSkelSkeletonQuery& SkeletonQuery,
+		FSkeletalMeshImportData& SkelMeshImportData,
+		TArray< UsdUtils::FUsdPrimMaterialSlot >& MaterialAssignments,
+		const TMap< FString, TMap< FString, int32 > >& MaterialToPrimvarsUVSetNames,
+		const pxr::TfToken& RenderContext = pxr::UsdShadeTokens->universalRenderContext,
+		const pxr::TfToken& MaterialPurpose = pxr::UsdShadeTokens->allPurpose
+	);
 	UE_DEPRECATED( 5.1, "Please use the other overload that also receives the relevant UsdSkelSkeletonQuery object." )
-	USDUTILITIES_API bool ConvertSkinnedMesh( const pxr::UsdSkelSkinningQuery& UsdSkinningQuery, const FTransform& AdditionalTransform, FSkeletalMeshImportData& SkelMeshImportData, TArray< UsdUtils::FUsdPrimMaterialSlot >& MaterialAssignments, const TMap< FString, TMap< FString, int32 > >& MaterialToPrimvarsUVSetNames, const pxr::TfToken& RenderContext = pxr::UsdShadeTokens->universalRenderContext );
+	USDUTILITIES_API bool ConvertSkinnedMesh(
+		const pxr::UsdSkelSkinningQuery& UsdSkinningQuery,
+		const FTransform& AdditionalTransform,
+		FSkeletalMeshImportData& SkelMeshImportData,
+		TArray< UsdUtils::FUsdPrimMaterialSlot >& MaterialAssignments,
+		const TMap< FString, TMap< FString, int32 > >& MaterialToPrimvarsUVSetNames,
+		const pxr::TfToken& RenderContext = pxr::UsdShadeTokens->universalRenderContext,
+		const pxr::TfToken& MaterialPurpose = pxr::UsdShadeTokens->allPurpose
+	);
 
 	/**
 	 * Will extract animation data from the animation source of InUsdSkeletonQuery's skeleton, and populate OutSkeletalAnimationAsset with the data.
diff --git a/Engine/Plugins/Importers/USDImporter/Source/UnrealUSDWrapper/Private/UnrealUSDWrapper.cpp b/Engine/Plugins/Importers/USDImporter/Source/UnrealUSDWrapper/Private/UnrealUSDWrapper.cpp
--- a/Engine/Plugins/Importers/USDImporter/Source/UnrealUSDWrapper/Private/UnrealUSDWrapper.cpp
+++ b/Engine/Plugins/Importers/USDImporter/Source/UnrealUSDWrapper/Private/UnrealUSDWrapper.cpp
@@ -762,6 +762,17 @@ FName UnrealIdentifiers::OuterConeAnglePropertyName = GET_MEMBER_NAME_CHECKED( U
 FName UnrealIdentifiers::InnerConeAnglePropertyName = GET_MEMBER_NAME_CHECKED( USpotLightComponent, InnerConeAngle );
 FName UnrealIdentifiers::LightSourceAnglePropertyName = GET_MEMBER_NAME_CHECKED( UDirectionalLightComponent, LightSourceAngle );
 
+FString UnrealIdentifiers::MaterialAllPurposeText = TEXT( "allPurpose" );
+#if USE_USD_SDK
+FString UnrealIdentifiers::MaterialAllPurpose = ANSI_TO_TCHAR( pxr::UsdShadeTokens->allPurpose.GetString().c_str() );
+FString UnrealIdentifiers::MaterialPreviewPurpose = ANSI_TO_TCHAR( pxr::UsdShadeTokens->preview.GetString().c_str() );
+FString UnrealIdentifiers::MaterialFullPurpose = ANSI_TO_TCHAR( pxr::UsdShadeTokens->full.GetString().c_str() );
+#else
+FString UnrealIdentifiers::MaterialAllPurpose = TEXT( "" );
+FString UnrealIdentifiers::MaterialPreviewPurpose = TEXT( "preview" );
+FString UnrealIdentifiers::MaterialFullPurpose = TEXT( "full" );
+#endif // USE_USD_SDK
+
 FUsdDelegates::FUsdImportDelegate FUsdDelegates::OnPreUsdImport;
 FUsdDelegates::FUsdImportDelegate FUsdDelegates::OnPostUsdImport;
 
diff --git a/Engine/Plugins/Importers/USDImporter/Source/UnrealUSDWrapper/Public/UnrealUSDWrapper.h b/Engine/Plugins/Importers/USDImporter/Source/UnrealUSDWrapper/Public/UnrealUSDWrapper.h
--- a/Engine/Plugins/Importers/USDImporter/Source/UnrealUSDWrapper/Public/UnrealUSDWrapper.h
+++ b/Engine/Plugins/Importers/USDImporter/Source/UnrealUSDWrapper/Public/UnrealUSDWrapper.h
@@ -439,6 +439,12 @@ namespace UnrealIdentifiers
 
 	// UDirectionalLightComponent properties
 	extern UNREALUSDWRAPPER_API FName LightSourceAnglePropertyName;
+
+	// Material purpose tokens that we convert from USD, if available
+	extern UNREALUSDWRAPPER_API FString MaterialAllPurpose;
+	extern UNREALUSDWRAPPER_API FString MaterialAllPurposeText; // Text to show on UI for "allPurpose", as its value is actually the empty string
+	extern UNREALUSDWRAPPER_API FString MaterialPreviewPurpose;
+	extern UNREALUSDWRAPPER_API FString MaterialFullPurpose;
 }
 
 struct UNREALUSDWRAPPER_API FUsdDelegates
