USD: Show warning when duplicating all local layer specs and the prim has some non-local specs;

Also disable the options to duplicate all local layer specs and single layer specs whenever they would lead to no action.

diff --git a/Engine/Plugins/Importers/USDImporter/Source/USDClasses/Public/USDProjectSettings.h b/Engine/Plugins/Importers/USDImporter/Source/USDClasses/Public/USDProjectSettings.h
--- a/Engine/Plugins/Importers/USDImporter/Source/USDClasses/Public/USDProjectSettings.h
+++ b/Engine/Plugins/Importers/USDImporter/Source/USDClasses/Public/USDProjectSettings.h
@@ -37,6 +37,11 @@ public:
 	UPROPERTY( config, EditAnywhere, Category = "USD|Dialogs" )
 	bool bShowOverriddenOpinionsWarning = true;
 
+	// Whether to show a warning whenever the "Duplicate All Local Layer Specs" option is picked, and the duplicated
+	// prim has some specs outside the local layer stack that will not be duplicated.
+	UPROPERTY( config, EditAnywhere, Category = "USD|Dialogs" )
+	bool bShowWarningOnIncompleteDuplication = true;
+
 	// Whether to display the pop up dialog asking what to do about dirty USD layers when saving the UE level
 	UPROPERTY(config, EditAnywhere, Category = "USD|Dialogs" )
 	EUsdSaveDialogBehavior ShowSaveLayersDialogWhenSaving = EUsdSaveDialogBehavior::ShowPrompt;
diff --git a/Engine/Plugins/Importers/USDImporter/Source/USDStageEditor/Private/SUSDStageTreeView.cpp b/Engine/Plugins/Importers/USDImporter/Source/USDStageEditor/Private/SUSDStageTreeView.cpp
--- a/Engine/Plugins/Importers/USDImporter/Source/USDStageEditor/Private/SUSDStageTreeView.cpp
+++ b/Engine/Plugins/Importers/USDImporter/Source/USDStageEditor/Private/SUSDStageTreeView.cpp
@@ -1394,6 +1394,21 @@ bool SUsdStageTreeView::DoesPrimExistOnEditTarget() const
 	return false;
 }
 
+bool SUsdStageTreeView::CanDuplicateAllLocalLayerSpecs() const
+{
+	TArray< FUsdPrimViewModelRef > MySelectedItems = GetSelectedItems();
+
+	for ( FUsdPrimViewModelRef SelectedItem : MySelectedItems )
+	{
+		if ( SelectedItem->HasSpecsOnLocalLayer() )
+		{
+			return true;
+		}
+	}
+
+	return false;
+}
+
 bool SUsdStageTreeView::CanSetUpLiveLink() const
 {
 	if ( !UsdStageActor.IsValid() )
@@ -1672,11 +1687,16 @@ void SUsdStageTreeView::FillDuplicateSubmenu( FMenuBuilder& MenuBuilder )
 
 	MenuBuilder.AddMenuEntry(
 		LOCTEXT( "DuplicateSingle_Text", "Single layer specs" ),
-		LOCTEXT( "DuplicateSingle_ToolTip", "Duplicate the prim's specs on the current edit target only" ),
+		TAttribute<FText>::Create( TAttribute<FText>::FGetter::CreateLambda( [this]()
+		{
+			return DoesPrimExistOnEditTarget()
+				? LOCTEXT( "DuplicateSingleValid_ToolTip", "Duplicate the prim's specs on the current edit target only" )
+				: LOCTEXT( "DuplicateSingleInvalid_ToolTip", "This prim needs a spec in the local edit target for this option to be usable!" );
+		})),
 		FSlateIcon(),
 		FUIAction(
 			FExecuteAction::CreateSP( this, &SUsdStageTreeView::OnDuplicatePrim, EUsdDuplicateType::SingleLayerSpecs ),
-			FCanExecuteAction::CreateSP( this, &SUsdStageTreeView::DoesPrimExistOnStage )
+			FCanExecuteAction::CreateSP( this, &SUsdStageTreeView::DoesPrimExistOnEditTarget )
 		),
 		NAME_None,
 		EUserInterfaceActionType::Button
@@ -1684,11 +1704,16 @@ void SUsdStageTreeView::FillDuplicateSubmenu( FMenuBuilder& MenuBuilder )
 
 	MenuBuilder.AddMenuEntry(
 		LOCTEXT( "DuplicateAllLocal_Text", "All local layer specs" ),
-		LOCTEXT( "DuplicateAllLocal_ToolTip", "Duplicate each of the prim's specs across the entire stage" ),
+		TAttribute<FText>::Create( TAttribute<FText>::FGetter::CreateLambda( [this]()
+		{
+			return CanDuplicateAllLocalLayerSpecs()
+				? LOCTEXT( "DuplicateAllLocalValid_ToolTip", "Duplicate each of the prim's specs across the entire stage" )
+				: LOCTEXT( "DuplicateAllLocalInvalid_ToolTip", "This prim needs at least one spec on the stage's local layer stack for this option to be usable!" );
+		})),
 		FSlateIcon(),
 		FUIAction(
 			FExecuteAction::CreateSP( this, &SUsdStageTreeView::OnDuplicatePrim, EUsdDuplicateType::AllLocalLayerSpecs ),
-			FCanExecuteAction::CreateSP( this, &SUsdStageTreeView::DoesPrimExistOnStage )
+			FCanExecuteAction::CreateSP( this, &SUsdStageTreeView::CanDuplicateAllLocalLayerSpecs )
 		),
 		NAME_None,
 		EUserInterfaceActionType::Button
diff --git a/Engine/Plugins/Importers/USDImporter/Source/USDStageEditor/Private/SUSDStageTreeView.h b/Engine/Plugins/Importers/USDImporter/Source/USDStageEditor/Private/SUSDStageTreeView.h
--- a/Engine/Plugins/Importers/USDImporter/Source/USDStageEditor/Private/SUSDStageTreeView.h
+++ b/Engine/Plugins/Importers/USDImporter/Source/USDStageEditor/Private/SUSDStageTreeView.h
@@ -82,6 +82,7 @@ private:
 	bool CanPastePrim() const;
 	bool DoesPrimExistOnStage() const;
 	bool DoesPrimExistOnEditTarget() const;
+	bool CanDuplicateAllLocalLayerSpecs() const;
 
 	bool CanSetUpLiveLink() const;
 	bool CanRemoveLiveLink() const;
diff --git a/Engine/Plugins/Importers/USDImporter/Source/USDStageEditorViewModels/Private/USDPrimViewModel.cpp b/Engine/Plugins/Importers/USDImporter/Source/USDStageEditorViewModels/Private/USDPrimViewModel.cpp
--- a/Engine/Plugins/Importers/USDImporter/Source/USDStageEditorViewModels/Private/USDPrimViewModel.cpp
+++ b/Engine/Plugins/Importers/USDImporter/Source/USDStageEditorViewModels/Private/USDPrimViewModel.cpp
@@ -427,6 +427,30 @@ bool FUsdPrimViewModel::CanRemoveGroomSchema() const
 #endif // #if USE_USD_SDK
 }
 
+bool FUsdPrimViewModel::HasSpecsOnLocalLayer() const
+{
+#if USE_USD_SDK
+	FScopedUsdAllocs UsdAllocs;
+
+	pxr::UsdPrim PxrUsdPrim{ UsdPrim };
+	if ( PxrUsdPrim )
+	{
+		if ( pxr::UsdStageRefPtr PrimUsdStage = PxrUsdPrim.GetStage() )
+		{
+			for ( const pxr::SdfPrimSpecHandle& Spec : PxrUsdPrim.GetPrimStack() )
+			{
+				if ( Spec && PrimUsdStage->HasLocalLayer( Spec->GetLayer() ) )
+				{
+					return true;
+				}
+			}
+		}
+	}
+#endif // #if USE_USD_SDK
+
+	return false;
+}
+
 void FUsdPrimViewModel::DefinePrim( const TCHAR* PrimName )
 {
 #if USE_USD_SDK
diff --git a/Engine/Plugins/Importers/USDImporter/Source/USDStageEditorViewModels/Public/USDPrimViewModel.h b/Engine/Plugins/Importers/USDImporter/Source/USDStageEditorViewModels/Public/USDPrimViewModel.h
--- a/Engine/Plugins/Importers/USDImporter/Source/USDStageEditorViewModels/Public/USDPrimViewModel.h
+++ b/Engine/Plugins/Importers/USDImporter/Source/USDStageEditorViewModels/Public/USDPrimViewModel.h
@@ -65,6 +65,9 @@ public:
 	void RemoveGroomSchema();
 	bool CanRemoveGroomSchema() const;
 
+	// Returns true if this prim has at least one spec on the stage's local layer stack
+	bool HasSpecsOnLocalLayer() const;
+
 	void DefinePrim( const TCHAR* PrimName );
 
 	/** Delegate for hooking up an inline editable text block to be notified that a rename is requested. */
diff --git a/Engine/Plugins/Importers/USDImporter/Source/USDUtilities/Private/USDConversionUtils.cpp b/Engine/Plugins/Importers/USDImporter/Source/USDUtilities/Private/USDConversionUtils.cpp
--- a/Engine/Plugins/Importers/USDImporter/Source/USDUtilities/Private/USDConversionUtils.cpp
+++ b/Engine/Plugins/Importers/USDImporter/Source/USDUtilities/Private/USDConversionUtils.cpp
@@ -8,6 +8,7 @@
 #include "USDGeomMeshConversion.h"
 #include "USDLayerUtils.h"
 #include "USDLog.h"
+#include "USDProjectSettings.h"
 #include "USDTypesConversion.h"
 
 #include "UsdWrappers/SdfLayer.h"
@@ -36,10 +37,12 @@
 #include "Engine/SpotLight.h"
 #include "Engine/StaticMesh.h"
 #include "Engine/Texture.h"
+#include "Framework/Notifications/NotificationManager.h"
 #include "GeometryCache.h"
 #include "GroomAsset.h"
 #include "InstancedFoliageActor.h"
 #include "LandscapeProxy.h"
+#include "Widgets/Notifications/SNotificationList.h"
 
 #if WITH_EDITOR
 #include "ObjectTools.h"
@@ -217,6 +220,85 @@ namespace USDConversionUtilsImpl
 			}
 		}
 	}
+
+	// Shows a notification saying that some specs of the provided prims won't be duplicated due to being on external
+	// layers
+	void NotifySpecsWontBeDuplicated( const TArray<UE::FUsdPrim>& Prims )
+	{
+		if ( Prims.Num() == 0 )
+		{
+			return;
+		}
+
+		const FText Text = LOCTEXT( "IncompleteDuplicationText", "USD: Incomplete duplication" );
+
+		FString PrimNamesString;
+		const static FString Delimiter = TEXT(", ");
+		for( const UE::FUsdPrim& Prim : Prims )
+		{
+			PrimNamesString += Prim.GetName().ToString() + Delimiter;
+		}
+		PrimNamesString.RemoveFromEnd( Delimiter );
+
+		const int32 NumPrims = Prims.Num();
+
+		const FText SubText = FText::Format(
+			LOCTEXT( "IncompleteDuplicationSubText", "{0}|plural(one=This,other=These) duplicated {0}|plural(one=prim,other=prims):\n\n{1}\n\n{0}|plural(one=Has,other=Have) some specs within layers that are outside of the stage's local layer stack, and so will not be duplicated.\n\nIf you wish to modify referenced or payload layers, please open those layers as USD stages directly." ),
+			NumPrims,
+			FText::FromString( PrimNamesString )
+		);
+
+		UE_LOG( LogUsd, Warning, TEXT( "%s" ), *SubText.ToString().Replace( TEXT( "\n\n" ), TEXT( " " ) ) );
+
+		const UUsdProjectSettings* Settings = GetDefault<UUsdProjectSettings>();
+		if ( Settings && Settings->bShowWarningOnIncompleteDuplication )
+		{
+			static TWeakPtr<SNotificationItem> Notification;
+
+			FNotificationInfo Toast( Text );
+			Toast.SubText = SubText;
+			Toast.Image = FCoreStyle::Get().GetBrush( TEXT( "MessageLog.Warning" ) );
+			Toast.CheckBoxText = LOCTEXT( "DontAskAgain", "Don't prompt again" );
+			Toast.bUseLargeFont = false;
+			Toast.bFireAndForget = false;
+			Toast.FadeOutDuration = 0.0f;
+			Toast.ExpireDuration = 0.0f;
+			Toast.bUseThrobber = false;
+			Toast.bUseSuccessFailIcons = false;
+			Toast.ButtonDetails.Emplace(
+				LOCTEXT( "OverridenOpinionMessageOk", "Ok" ),
+				FText::GetEmpty(),
+				FSimpleDelegate::CreateLambda( []() {
+				if ( TSharedPtr<SNotificationItem> PinnedNotification = Notification.Pin() )
+				{
+					PinnedNotification->SetCompletionState( SNotificationItem::CS_Success );
+					PinnedNotification->ExpireAndFadeout();
+				}
+			} )
+			);
+			// This is flipped because the default checkbox message is "Don't prompt again"
+			Toast.CheckBoxState = Settings->bShowWarningOnIncompleteDuplication ? ECheckBoxState::Unchecked : ECheckBoxState::Checked;
+			Toast.CheckBoxStateChanged = FOnCheckStateChanged::CreateStatic( []( ECheckBoxState NewState )
+			{
+				if ( UUsdProjectSettings* Settings = GetMutableDefault<UUsdProjectSettings>() )
+				{
+					// This is flipped because the default checkbox message is "Don't prompt again"
+					Settings->bShowWarningOnIncompleteDuplication = NewState == ECheckBoxState::Unchecked;
+				}
+			} );
+
+			// Only show one at a time
+			if ( !Notification.IsValid() )
+			{
+				Notification = FSlateNotificationManager::Get().AddNotification( Toast );
+			}
+
+			if ( TSharedPtr<SNotificationItem> PinnedNotification = Notification.Pin() )
+			{
+				PinnedNotification->SetCompletionState( SNotificationItem::CS_Pending );
+			}
+		}
+	}
 }
 
 template< typename ValueType >
@@ -2020,6 +2102,30 @@ TArray<UE::FSdfPath> UsdUtils::DuplicatePrims( const TArray<UE::FUsdPrim>& Prims
 			{
 				LayersThatCanBeAffected.insert( Handle );
 			}
+
+			// If any of our prims has specs on layers that are used by the stage but are not within the local layer
+			// stack, then warn the user that some of these specs will not be duplicated
+			{
+				TArray<UE::FUsdPrim> PrimsWithExternalSpecs;
+				for ( const UE::FUsdPrim& Prim : Prims )
+				{
+					pxr::UsdPrim UsdPrim{ Prim };
+					if ( !UsdPrim )
+					{
+						continue;
+					}
+
+					for ( const pxr::SdfPrimSpecHandle& Spec : UsdPrim.GetPrimStack() )
+					{
+						if ( Spec && LayersThatCanBeAffected.count( Spec->GetLayer() ) == 0 )
+						{
+							PrimsWithExternalSpecs.Add(Prim);
+							break;
+						}
+					}
+				}
+				USDConversionUtilsImpl::NotifySpecsWontBeDuplicated( PrimsWithExternalSpecs );
+			}
 			break;
 		}
 	}
