USD: Add support for using xformable transform animations as root motion within AnimSequences.

diff --git a/Engine/Plugins/Importers/USDImporter/Source/USDSchemas/Private/USDGeomXformableTranslator.cpp b/Engine/Plugins/Importers/USDImporter/Source/USDSchemas/Private/USDGeomXformableTranslator.cpp
--- a/Engine/Plugins/Importers/USDImporter/Source/USDSchemas/Private/USDGeomXformableTranslator.cpp
+++ b/Engine/Plugins/Importers/USDImporter/Source/USDSchemas/Private/USDGeomXformableTranslator.cpp
@@ -597,7 +597,14 @@ void FUsdGeomXformableTranslator::UpdateComponents( USceneComponent* SceneCompon
 		// into the mesh at that case, as each LOD could technically have a separate transform
 		if ( !Context->bAllowInterpretingLODs || !bHasMultipleLODs )
 		{
-			UsdToUnreal::ConvertXformable( Context->Stage, pxr::UsdGeomXformable( Prim ), *SceneComponent, Context->Time );
+			// Don't update the component's transform if this is already factored in as root motion within the AnimSequence
+			bool bConvertTransform = true;
+			if ( Prim.IsA( TEXT( "SkelRoot" ) ) && Context->RootMotionHandling == EUsdRootMotionHandling::UseMotionFromSkelRoot )
+			{
+				bConvertTransform = false;
+			}
+
+			UsdToUnreal::ConvertXformable( Context->Stage, pxr::UsdGeomXformable( Prim ), *SceneComponent, Context->Time, bConvertTransform );
 		}
 
 		// Note how we should only register if we unregistered ourselves: If we did this every time we would
diff --git a/Engine/Plugins/Importers/USDImporter/Source/USDSchemas/Private/USDSkelRootTranslator.cpp b/Engine/Plugins/Importers/USDImporter/Source/USDSchemas/Private/USDSkelRootTranslator.cpp
--- a/Engine/Plugins/Importers/USDImporter/Source/USDSchemas/Private/USDSkelRootTranslator.cpp
+++ b/Engine/Plugins/Importers/USDImporter/Source/USDSchemas/Private/USDSkelRootTranslator.cpp
@@ -205,7 +205,7 @@ namespace UsdSkelRootTranslatorImpl
 		return OutHash;
 	}
 
-	FSHAHash ComputeSHAHash( const pxr::UsdSkelSkeletonQuery& InUsdSkeletonQuery )
+	FSHAHash ComputeSHAHash( const pxr::UsdSkelSkeletonQuery& InUsdSkeletonQuery, const pxr::UsdPrim& RootMotionPrim )
 	{
 		TRACE_CPUPROFILER_EVENT_SCOPE( UsdSkelRootTranslatorImpl::ComputeSHAHash_SkelQuery );
 
@@ -270,6 +270,24 @@ namespace UsdSkelRootTranslatorImpl
 			HashState.Update( ( uint8* ) WeightsForSample.data(), WeightsForSample.size() * sizeof( float ) );
 		}
 
+		// If we're pulling root motion from anywhere, hash that too because if it changes we'll need to rebake
+		// the AnimSequence asset
+		if ( pxr::UsdGeomXformable Xformable{ RootMotionPrim } )
+		{
+			std::vector<double> TimeSamples;
+			Xformable.GetTimeSamples( &TimeSamples );
+
+			pxr::GfMatrix4d Transform{};
+			bool bResetsXformSack = false;
+			for ( double TimeSample : TimeSamples )
+			{
+				Xformable.GetLocalTransformation( &Transform, &bResetsXformSack, TimeSample );
+				HashState.Update( ( uint8* ) Transform.data(), 16 * sizeof( double ) );
+			}
+
+			HashState.Update( ( uint8* ) &bResetsXformSack, sizeof( bool ) );
+		}
+
 		HashState.Final();
 		HashState.GetHash( &OutHash.Hash[ 0 ] );
 
@@ -952,7 +970,13 @@ namespace UsdSkelRootTranslatorImpl
 							continue;
 						}
 
-						FSHAHash Hash = UsdSkelRootTranslatorImpl::ComputeSHAHash( SkelQuery );
+						pxr::UsdPrim RootMotionPrim = Context->RootMotionHandling == EUsdRootMotionHandling::UseMotionFromSkelRoot
+							? SkeletonRoot.GetPrim()
+							: Context->RootMotionHandling == EUsdRootMotionHandling::UseMotionFromSkeleton
+								? Skeleton.GetPrim()
+								: pxr::UsdPrim{};
+
+						FSHAHash Hash = UsdSkelRootTranslatorImpl::ComputeSHAHash( SkelQuery, RootMotionPrim );
 						FString HashString = Hash.ToString();
 						UAnimSequence* AnimSequence = Cast< UAnimSequence >( Context->AssetCache->GetCachedAsset( HashString ) );
 
@@ -972,7 +996,15 @@ namespace UsdSkelRootTranslatorImpl
 
 							TUsdStore<pxr::VtArray<pxr::UsdSkelSkinningQuery>> SkinningTargets = Binding.GetSkinningTargets();
 							float LayerStartOffsetSeconds = 0.0f;
-							UsdToUnreal::ConvertSkelAnim( SkelQuery, &SkinningTargets.Get(), &NewBlendShapes, Context->bAllowInterpretingLODs, AnimSequence, &LayerStartOffsetSeconds );
+							UsdToUnreal::ConvertSkelAnim(
+								SkelQuery,
+								&SkinningTargets.Get(),
+								&NewBlendShapes,
+								Context->bAllowInterpretingLODs,
+								RootMotionPrim,
+								AnimSequence,
+								&LayerStartOffsetSeconds
+							);
 
 							if (AnimSequence->GetDataModel()->GetNumBoneTracks() != 0 || AnimSequence->GetDataModel()->GetNumberOfFloatCurves() != 0 )
 							{
diff --git a/Engine/Plugins/Importers/USDImporter/Source/USDSchemas/Public/USDSchemaTranslator.h b/Engine/Plugins/Importers/USDImporter/Source/USDSchemas/Public/USDSchemaTranslator.h
--- a/Engine/Plugins/Importers/USDImporter/Source/USDSchemas/Public/USDSchemaTranslator.h
+++ b/Engine/Plugins/Importers/USDImporter/Source/USDSchemas/Public/USDSchemaTranslator.h
@@ -152,6 +152,9 @@ struct USDSCHEMAS_API FUsdSchemaTranslationContext : public TSharedFromThis< FUs
 	/** The material purpose to use when translating material bindings */
 	FName MaterialPurpose;
 
+	/** Describes what to add to the root bone animation within generated AnimSequences, if anything */
+	EUsdRootMotionHandling RootMotionHandling = EUsdRootMotionHandling::NoAdditionalRootMotion;
+
 	/** If a generated UStaticMesh has at least this many triangles we will attempt to enable Nanite */
 	int32 NaniteTriangleThreshold;
 
diff --git a/Engine/Plugins/Importers/USDImporter/Source/USDStage/Private/USDLevelSequenceHelper.cpp b/Engine/Plugins/Importers/USDImporter/Source/USDStage/Private/USDLevelSequenceHelper.cpp
--- a/Engine/Plugins/Importers/USDImporter/Source/USDStage/Private/USDLevelSequenceHelper.cpp
+++ b/Engine/Plugins/Importers/USDImporter/Source/USDStage/Private/USDLevelSequenceHelper.cpp
@@ -1381,6 +1381,20 @@ void FUsdLevelSequenceHelperImpl::AddCommonTracks( const UUsdPrimTwin& PrimTwin,
 		}
 	}
 
+	// Check whether we should ignore the prim's local transform or not. We only do this for SkelRoots, in case we
+	// already appended their transform animations to the root bone as additional root motion
+	bool bIgnorePrimLocalTransform = false;
+	if ( AUsdStageActor* StageActorPtr = StageActor.Get() )
+	{
+		if ( StageActorPtr->RootMotionHandling == EUsdRootMotionHandling::UseMotionFromSkelRoot )
+		{
+			if ( Prim.IsA( TEXT( "SkelRoot" ) ) )
+			{
+				bIgnorePrimLocalTransform = true;
+			}
+		}
+	}
+
 	// Test that transform might be time varying and not TransformAttribute since it will check each xform ops
 	if ( Xformable.TransformMightBeTimeVarying() || bNeedTrackToCompensateResetXformOp )
 	{
@@ -1416,7 +1430,7 @@ void FUsdLevelSequenceHelperImpl::AddCommonTracks( const UUsdPrimTwin& PrimTwin,
 
 							if ( UMovieScene3DTransformTrack* TransformTrack = AddTrack<UMovieScene3DTransformTrack>( UnrealIdentifiers::TransformPropertyName, PrimTwin, *ComponentToBind, *AttributeSequence, bIsMuted ) )
 							{
-								UsdToUnreal::FPropertyTrackReader Reader = UsdToUnreal::CreatePropertyTrackReader( Prim, UnrealIdentifiers::TransformPropertyName );
+								UsdToUnreal::FPropertyTrackReader Reader = UsdToUnreal::CreatePropertyTrackReader( Prim, UnrealIdentifiers::TransformPropertyName, bIgnorePrimLocalTransform );
 								UsdToUnreal::ConvertTransformTimeSamples( UsdStage, TimeSamples, Reader.TransformReader, *TransformTrack, SequenceTransform );
 							}
 
diff --git a/Engine/Plugins/Importers/USDImporter/Source/USDStage/Private/USDStageActor.cpp b/Engine/Plugins/Importers/USDImporter/Source/USDStage/Private/USDStageActor.cpp
--- a/Engine/Plugins/Importers/USDImporter/Source/USDStage/Private/USDStageActor.cpp
+++ b/Engine/Plugins/Importers/USDImporter/Source/USDStage/Private/USDStageActor.cpp
@@ -105,6 +105,7 @@ struct FUsdStageActorImpl
 		TranslationContext->NaniteTriangleThreshold = StageActor->NaniteTriangleThreshold;
 		TranslationContext->RenderContext = StageActor->RenderContext;
 		TranslationContext->MaterialPurpose = StageActor->MaterialPurpose;
+		TranslationContext->RootMotionHandling = StageActor->RootMotionHandling;
 		TranslationContext->MaterialToPrimvarToUVIndex = &StageActor->MaterialToPrimvarToUVIndex;
 		TranslationContext->BlendShapesByPath = &StageActor->BlendShapesByPath;
 		TranslationContext->InfoCache = StageActor->InfoCache;
@@ -539,6 +540,7 @@ struct FUsdStageActorImpl
 			EventAttributes.Emplace( TEXT( "NaniteTriangleThreshold" ), LexToString( StageActor->NaniteTriangleThreshold ) );
 			EventAttributes.Emplace( TEXT( "RenderContext" ), StageActor->RenderContext.ToString() );
 			EventAttributes.Emplace( TEXT( "MaterialPurpose" ), StageActor->MaterialPurpose.ToString() );
+			EventAttributes.Emplace( TEXT( "RootMotionHandling" ), LexToString( ( uint8 ) StageActor->RootMotionHandling ) );
 
 			const bool bAutomated = false;
 			IUsdClassesModule::SendAnalytics( MoveTemp( EventAttributes ), TEXT( "Open" ), bAutomated, ElapsedSeconds, NumberOfFrames, Extension );
@@ -624,6 +626,7 @@ AUsdStageActor::AUsdStageActor()
 	, PurposesToLoad( (int32) EUsdPurpose::Proxy )
 	, NaniteTriangleThreshold( (uint64) 1000000 )
 	, MaterialPurpose( *UnrealIdentifiers::MaterialAllPurpose )
+	, RootMotionHandling( EUsdRootMotionHandling::NoAdditionalRootMotion )
 	, Time( 0.0f )
 	, bIsTransitioningIntoPIE( false )
 	, bIsModifyingAProperty( false )
@@ -1457,6 +1460,15 @@ void AUsdStageActor::SetMaterialPurpose( const FName& NewMaterialPurpose )
 	LoadUsdStage();
 }
 
+USDSTAGE_API void AUsdStageActor::SetRootMotionHandling( EUsdRootMotionHandling NewHandlingStrategy )
+{
+	const bool bMarkDirty = false;
+	Modify( bMarkDirty );
+
+	RootMotionHandling = NewHandlingStrategy;
+	LoadUsdStage();
+}
+
 void AUsdStageActor::SetTime(float InTime)
 {
 	const bool bMarkDirty = false;
@@ -2748,6 +2760,10 @@ void AUsdStageActor::HandlePropertyChangedEvent( FPropertyChangedEvent& Property
 	{
 		SetMaterialPurpose( MaterialPurpose );
 	}
+	else if ( PropertyName == GET_MEMBER_NAME_CHECKED( AUsdStageActor, RootMotionHandling ) )
+	{
+		SetRootMotionHandling( RootMotionHandling );
+	}
 
 	bIsModifyingAProperty = false;
 }
diff --git a/Engine/Plugins/Importers/USDImporter/Source/USDStage/Public/USDStageActor.h b/Engine/Plugins/Importers/USDImporter/Source/USDStage/Public/USDStageActor.h
--- a/Engine/Plugins/Importers/USDImporter/Source/USDStage/Public/USDStageActor.h
+++ b/Engine/Plugins/Importers/USDImporter/Source/USDStage/Public/USDStageActor.h
@@ -87,6 +87,10 @@ public:
 	UPROPERTY( EditAnywhere, BlueprintReadOnly, Category = "USD" )
 	FName MaterialPurpose;
 
+	// Describes what to add to the root bone animation within generated AnimSequences, if anything
+	UPROPERTY( EditAnywhere, BlueprintReadOnly, Category = "USD" )
+	EUsdRootMotionHandling RootMotionHandling = EUsdRootMotionHandling::NoAdditionalRootMotion;
+
 public:
 	UFUNCTION(BlueprintCallable, Category = "USD", meta = (CallInEditor = "true"))
 	USDSTAGE_API void SetRootLayer(const FString& RootFilePath );
@@ -118,6 +122,9 @@ public:
 	UFUNCTION( BlueprintCallable, Category = "USD", meta = ( CallInEditor = "true" ) )
 	USDSTAGE_API void SetMaterialPurpose( const FName& NewMaterialPurpose );
 
+	UFUNCTION( BlueprintCallable, Category = "USD", meta = ( CallInEditor = "true" ) )
+	USDSTAGE_API void SetRootMotionHandling( EUsdRootMotionHandling NewHandlingStrategy );
+
 	UFUNCTION(BlueprintCallable, Category = "USD", meta = (CallInEditor = "true"))
 	USDSTAGE_API float GetTime() const { return Time; }
 
diff --git a/Engine/Plugins/Importers/USDImporter/Source/USDStageEditor/Private/SUSDStage.cpp b/Engine/Plugins/Importers/USDImporter/Source/USDStageEditor/Private/SUSDStage.cpp
--- a/Engine/Plugins/Importers/USDImporter/Source/USDStageEditor/Private/SUSDStage.cpp
+++ b/Engine/Plugins/Importers/USDImporter/Source/USDStageEditor/Private/SUSDStage.cpp
@@ -724,6 +724,11 @@ void SUsdStage::FillOptionsMenu(FMenuBuilder& MenuBuilder)
 			LOCTEXT( "MaterialPurpose_ToolTip", "Material purpose to use when parsing material bindings in addition to the \"allPurpose\" fallback" ),
 			FNewMenuDelegate::CreateSP( this, &SUsdStage::FillMaterialPurposeSubMenu ) );
 
+		MenuBuilder.AddSubMenu(
+			LOCTEXT( "RootMotionHandling", "Root motion handling" ),
+			LOCTEXT( "RootMotionHandling_ToolTip", "Choose how to handle root motion for generated AnimSequences and LevelSequences" ),
+			FNewMenuDelegate::CreateSP( this, &SUsdStage::FillRootMotionSubMenu ) );
+
 		MenuBuilder.AddSubMenu(
 			LOCTEXT( "Collapsing", "Collapsing" ),
 			LOCTEXT( "Collapsing_ToolTip", "Whether to try to combine individual assets and components of the same type on a Kind-per-Kind basis, like multiple Mesh prims into a single Static Mesh" ),
@@ -1111,6 +1116,57 @@ void SUsdStage::FillMaterialPurposeSubMenu( FMenuBuilder& MenuBuilder )
 	MenuBuilder.AddWidget( Box, FText::GetEmpty(), bNoIndent );
 }
 
+void SUsdStage::FillRootMotionSubMenu( FMenuBuilder& MenuBuilder )
+{
+	auto AddRootMotionEntry = [&](const EUsdRootMotionHandling& HandlingStrategy, const FText& Text)
+	{
+		UEnum* Enum = StaticEnum< EUsdRootMotionHandling >();
+		FText ToolTip = Enum->GetToolTipTextByIndex( Enum->GetIndexByValue( static_cast< uint8 >( HandlingStrategy ) ) );
+
+		MenuBuilder.AddMenuEntry(
+			Text,
+			ToolTip,
+			FSlateIcon(),
+			FUIAction(
+				FExecuteAction::CreateLambda([this, HandlingStrategy]()
+				{
+					if(AUsdStageActor* StageActor = ViewModel.UsdStageActor.Get())
+					{
+						FScopedTransaction Transaction(FText::Format(
+							LOCTEXT("RootMotionHandlingTransaction", "Change root motion handling strategy for USD stage actor '{0}'"),
+							FText::FromString(StageActor->GetActorLabel())
+						));
+
+						// c.f. comment in SUsdStage::FillCollapsingSubMenu
+						TGuardValue<bool> MaintainSelectionGuard( bUpdatingViewportSelection, true );
+
+						StageActor->Modify();
+						StageActor->SetRootMotionHandling( HandlingStrategy );
+					}
+				}),
+				FCanExecuteAction::CreateLambda([this]()
+				{
+					return ViewModel.UsdStageActor.Get() != nullptr;
+				}),
+				FIsActionChecked::CreateLambda([this, HandlingStrategy]()
+				{
+					if(AUsdStageActor* StageActor = ViewModel.UsdStageActor.Get())
+					{
+						return StageActor->RootMotionHandling == HandlingStrategy;
+					}
+					return false;
+				})
+			),
+			NAME_None,
+			EUserInterfaceActionType::RadioButton
+		);
+	};
+
+	AddRootMotionEntry( EUsdRootMotionHandling::NoAdditionalRootMotion, LOCTEXT("NoAdditionalRootMotionText", "No additional root motion"));
+	AddRootMotionEntry( EUsdRootMotionHandling::UseMotionFromSkelRoot, LOCTEXT("UseMotionFromSkelRootText", "Use motion from SkelRoot"));
+	AddRootMotionEntry( EUsdRootMotionHandling::UseMotionFromSkeleton, LOCTEXT("UseMotionFromSkeletonText", "Use motion from Skeleton"));
+}
+
 void SUsdStage::FillCollapsingSubMenu( FMenuBuilder& MenuBuilder )
 {
 	MenuBuilder.AddMenuEntry(
diff --git a/Engine/Plugins/Importers/USDImporter/Source/USDStageEditor/Private/SUSDStage.h b/Engine/Plugins/Importers/USDImporter/Source/USDStageEditor/Private/SUSDStage.h
--- a/Engine/Plugins/Importers/USDImporter/Source/USDStageEditor/Private/SUSDStage.h
+++ b/Engine/Plugins/Importers/USDImporter/Source/USDStageEditor/Private/SUSDStage.h
@@ -47,6 +47,7 @@ protected:
 	void FillPurposesToLoadSubMenu( FMenuBuilder& MenuBuilder );
 	void FillRenderContextSubMenu( FMenuBuilder& MenuBuilder );
 	void FillMaterialPurposeSubMenu( FMenuBuilder& MenuBuilder );
+	void FillRootMotionSubMenu( FMenuBuilder& MenuBuilder );
 	void FillCollapsingSubMenu( FMenuBuilder& MenuBuilder );
 	void FillInterpolationTypeSubMenu( FMenuBuilder& MenuBuilder );
 	void FillSelectionSubMenu( FMenuBuilder& MenuBuilder );
diff --git a/Engine/Plugins/Importers/USDImporter/Source/USDStageImporter/Private/USDStageImporter.cpp b/Engine/Plugins/Importers/USDImporter/Source/USDStageImporter/Private/USDStageImporter.cpp
--- a/Engine/Plugins/Importers/USDImporter/Source/USDStageImporter/Private/USDStageImporter.cpp
+++ b/Engine/Plugins/Importers/USDImporter/Source/USDStageImporter/Private/USDStageImporter.cpp
@@ -1277,6 +1277,7 @@ namespace UsdStageImporterImpl
 				EventAttributes.Emplace( TEXT( "NaniteTriangleThreshold" ), LexToString( ImportContext.ImportOptions->NaniteTriangleThreshold ) );
 				EventAttributes.Emplace( TEXT( "RenderContextToImport" ), ImportContext.ImportOptions->RenderContextToImport.ToString() );
 				EventAttributes.Emplace( TEXT( "MaterialPurpose" ), ImportContext.ImportOptions->MaterialPurpose.ToString() );
+				EventAttributes.Emplace( TEXT( "RootMotionHandling" ), LexToString( ( uint8 ) ImportContext.ImportOptions->RootMotionHandling ) );
 				EventAttributes.Emplace( TEXT( "OverrideStageOptions" ), ImportContext.ImportOptions->bOverrideStageOptions );
 				if( ImportContext.ImportOptions->bOverrideStageOptions )
 				{
@@ -1494,6 +1495,7 @@ void UUsdStageImporter::ImportFromFile(FUsdStageImportContext& ImportContext)
 	TranslationContext->NaniteTriangleThreshold = ImportContext.ImportOptions->NaniteTriangleThreshold;
 	TranslationContext->RenderContext = ImportContext.ImportOptions->RenderContextToImport;
 	TranslationContext->MaterialPurpose = ImportContext.ImportOptions->MaterialPurpose;
+	TranslationContext->RootMotionHandling = ImportContext.ImportOptions->RootMotionHandling;
 	TranslationContext->ParentComponent = ImportContext.SceneActor ? ImportContext.SceneActor->GetRootComponent() : nullptr;
 	TranslationContext->KindsToCollapse = ( EUsdDefaultKind ) ImportContext.ImportOptions->KindsToCollapse;
 	TranslationContext->bMergeIdenticalMaterialSlots = ImportContext.ImportOptions->bMergeIdenticalMaterialSlots;
@@ -1589,6 +1591,7 @@ bool UUsdStageImporter::ReimportSingleAsset(FUsdStageImportContext& ImportContex
 	TranslationContext->NaniteTriangleThreshold = ImportContext.ImportOptions->NaniteTriangleThreshold;
 	TranslationContext->RenderContext = ImportContext.ImportOptions->RenderContextToImport;
 	TranslationContext->MaterialPurpose = ImportContext.ImportOptions->MaterialPurpose;
+	TranslationContext->RootMotionHandling = ImportContext.ImportOptions->RootMotionHandling;
 	TranslationContext->KindsToCollapse = ( EUsdDefaultKind ) ImportContext.ImportOptions->KindsToCollapse;
 	TranslationContext->bMergeIdenticalMaterialSlots = ImportContext.ImportOptions->bMergeIdenticalMaterialSlots;
 	TranslationContext->bCollapseTopLevelPointInstancers = ImportContext.ImportOptions->bCollapseTopLevelPointInstancers;
diff --git a/Engine/Plugins/Importers/USDImporter/Source/USDStageImporter/Public/USDStageImportOptions.h b/Engine/Plugins/Importers/USDImporter/Source/USDStageImporter/Public/USDStageImportOptions.h
--- a/Engine/Plugins/Importers/USDImporter/Source/USDStageImporter/Public/USDStageImportOptions.h
+++ b/Engine/Plugins/Importers/USDImporter/Source/USDStageImporter/Public/USDStageImportOptions.h
@@ -82,6 +82,10 @@ public:
 	UPROPERTY( BlueprintReadWrite, config, EditAnywhere, Category = "USD options" )
 	FName MaterialPurpose;
 
+	// Describes what to add to the root bone animation within generated AnimSequences, if anything
+	UPROPERTY( BlueprintReadWrite, config, EditAnywhere, Category = "USD options", meta = ( EditCondition = bImportSkeletalAnimations ) )
+	EUsdRootMotionHandling RootMotionHandling = EUsdRootMotionHandling::NoAdditionalRootMotion;
+
 	/** Whether to use the specified StageOptions instead of the stage's own settings */
 	UPROPERTY( BlueprintReadWrite, config, EditAnywhere, Category = "USD options" )
 	bool bOverrideStageOptions;
diff --git a/Engine/Plugins/Importers/USDImporter/Source/USDUtilities/Private/USDPrimConversion.cpp b/Engine/Plugins/Importers/USDImporter/Source/USDUtilities/Private/USDPrimConversion.cpp
--- a/Engine/Plugins/Importers/USDImporter/Source/USDUtilities/Private/USDPrimConversion.cpp
+++ b/Engine/Plugins/Importers/USDImporter/Source/USDUtilities/Private/USDPrimConversion.cpp
@@ -233,7 +233,7 @@ namespace UE
 
 			// Turns OutTransform into the UE-space relative (local to parent) transform for Xformable, paying attention to if it
 			// or any of its ancestors has the '!resetXformStack!' xformOp.
-			void GetPrimConvertedRelativeTransform( pxr::UsdGeomXformable Xformable, double UsdTimeCode, FTransform& OutTransform )
+			void GetPrimConvertedRelativeTransform( pxr::UsdGeomXformable Xformable, double UsdTimeCode, FTransform& OutTransform, bool bIgnoreLocalTransform = false )
 			{
 				if ( !Xformable )
 				{
@@ -246,7 +246,17 @@ namespace UE
 				pxr::UsdStageRefPtr UsdStage = UsdPrim.GetStage();
 
 				bool bResetTransformStack = false;
-				UsdToUnreal::ConvertXformable( UsdStage, Xformable, OutTransform, UsdTimeCode, &bResetTransformStack );
+				if( bIgnoreLocalTransform )
+				{
+					FTransform Dummy;
+					UsdToUnreal::ConvertXformable( UsdStage, Xformable, Dummy, UsdTimeCode, &bResetTransformStack );
+
+					OutTransform = FTransform::Identity;
+				}
+				else
+				{
+					UsdToUnreal::ConvertXformable( UsdStage, Xformable, OutTransform, UsdTimeCode, &bResetTransformStack );
+				}
 
 				// If we have the resetXformStack op on this prim's xformOpOrder we have to essentially use its transform
 				// as the world transform (i.e. we have to discard the parent transforms). We won't do this here, and will instead
@@ -373,7 +383,7 @@ bool UsdToUnreal::ConvertXformable( const pxr::UsdStageRefPtr& Stage, const pxr:
 	return true;
 }
 
-bool UsdToUnreal::ConvertXformable( const pxr::UsdStageRefPtr& Stage, const pxr::UsdTyped& Schema, USceneComponent& SceneComponent, double EvalTime )
+bool UsdToUnreal::ConvertXformable( const pxr::UsdStageRefPtr& Stage, const pxr::UsdTyped& Schema, USceneComponent& SceneComponent, double EvalTime, bool bUsePrimTransform )
 {
 	pxr::UsdGeomXformable Xformable( Schema );
 	if ( !Xformable )
@@ -387,7 +397,7 @@ bool UsdToUnreal::ConvertXformable( const pxr::UsdStageRefPtr& Stage, const pxr:
 
 	// Transform
 	FTransform Transform;
-	UE::USDPrimConversionImpl::Private::GetPrimConvertedRelativeTransform( Xformable, EvalTime, Transform );
+	UE::USDPrimConversionImpl::Private::GetPrimConvertedRelativeTransform( Xformable, EvalTime, Transform, !bUsePrimTransform );
 	SceneComponent.SetRelativeTransform( Transform );
 
 	SceneComponent.Modify();
@@ -920,7 +930,7 @@ bool UsdToUnreal::ConvertTransformTimeSamples( const UE::FUsdStage& Stage, const
 	return true;
 }
 
-UsdToUnreal::FPropertyTrackReader UsdToUnreal::CreatePropertyTrackReader( const UE::FUsdPrim& Prim, const FName& PropertyPath )
+UsdToUnreal::FPropertyTrackReader UsdToUnreal::CreatePropertyTrackReader( const UE::FUsdPrim& Prim, const FName& PropertyPath, bool bIgnorePrimLocalTransform )
 {
 	UsdToUnreal::FPropertyTrackReader Reader;
 
@@ -935,12 +945,12 @@ UsdToUnreal::FPropertyTrackReader UsdToUnreal::CreatePropertyTrackReader( const
 		if ( PropertyPath == UnrealIdentifiers::TransformPropertyName )
 		{
 			FTransform Default = FTransform::Identity;
-			UE::USDPrimConversionImpl::Private::GetPrimConvertedRelativeTransform( Xformable, UsdUtils::GetDefaultTimeCode(), Default );
+			UE::USDPrimConversionImpl::Private::GetPrimConvertedRelativeTransform( Xformable, UsdUtils::GetDefaultTimeCode(), Default, bIgnorePrimLocalTransform );
 
-			Reader.TransformReader = [UsdStage, Xformable, Default]( double UsdTimeCode )
+			Reader.TransformReader = [UsdStage, Xformable, Default, bIgnorePrimLocalTransform]( double UsdTimeCode )
 			{
 				FTransform Result = Default;
-				UE::USDPrimConversionImpl::Private::GetPrimConvertedRelativeTransform( Xformable, UsdTimeCode, Result );
+				UE::USDPrimConversionImpl::Private::GetPrimConvertedRelativeTransform( Xformable, UsdTimeCode, Result, bIgnorePrimLocalTransform );
 				return Result;
 			};
 			return Reader;
diff --git a/Engine/Plugins/Importers/USDImporter/Source/USDUtilities/Private/USDSkeletalDataConversion.cpp b/Engine/Plugins/Importers/USDImporter/Source/USDUtilities/Private/USDSkeletalDataConversion.cpp
--- a/Engine/Plugins/Importers/USDImporter/Source/USDUtilities/Private/USDSkeletalDataConversion.cpp
+++ b/Engine/Plugins/Importers/USDImporter/Source/USDUtilities/Private/USDSkeletalDataConversion.cpp
@@ -11,6 +11,7 @@
 #include "USDLayerUtils.h"
 #include "USDLog.h"
 #include "USDMemory.h"
+#include "USDPrimConversion.h"
 #include "USDTypesConversion.h"
 
 #include "UsdWrappers/SdfLayer.h"
@@ -58,6 +59,7 @@
 	#include "pxr/usd/usdGeom/primvarsAPI.h"
 	#include "pxr/usd/usdGeom/subset.h"
 	#include "pxr/usd/usdGeom/tokens.h"
+	#include "pxr/usd/usdGeom/xformable.h"
 	#include "pxr/usd/usdGeom/xformCache.h"
 	#include "pxr/usd/usdShade/tokens.h"
 	#include "pxr/usd/usdSkel/animation.h"
@@ -1838,7 +1840,15 @@ bool UsdToUnreal::ConvertSkinnedMesh(
 
 // Using UsdSkelSkeletonQuery instead of UsdSkelAnimQuery as it automatically does the joint remapping when we ask it to compute joint transforms.
 // It also initializes the joint transforms with the rest pose, if available, in case the animation doesn't provide data for all joints.
-bool UsdToUnreal::ConvertSkelAnim( const pxr::UsdSkelSkeletonQuery& InUsdSkeletonQuery, const pxr::VtArray<pxr::UsdSkelSkinningQuery>* InSkinningTargets, const UsdUtils::FBlendShapeMap* InBlendShapes, bool bInInterpretLODs, UAnimSequence* OutSkeletalAnimationAsset, float* OutStartOffsetSeconds )
+bool UsdToUnreal::ConvertSkelAnim(
+	const pxr::UsdSkelSkeletonQuery& InUsdSkeletonQuery,
+	const pxr::VtArray<pxr::UsdSkelSkinningQuery>* InSkinningTargets,
+	const UsdUtils::FBlendShapeMap* InBlendShapes,
+	bool bInInterpretLODs,
+	const pxr::UsdPrim& RootMotionPrim,
+	UAnimSequence* OutSkeletalAnimationAsset,
+	float* OutStartOffsetSeconds
+)
 {
 	TRACE_CPUPROFILER_EVENT_SCOPE( UsdToUnreal::ConvertSkelAnim );
 
@@ -1917,8 +1927,32 @@ bool UsdToUnreal::ConvertSkelAnim( const pxr::UsdSkelSkeletonQuery& InUsdSkeleto
 		LastBlendShapeSampleTimeCode = BlendShapeTimeSamples[ BlendShapeTimeSamples.size() - 1 ];
 	}
 
+	TUsdStore<std::vector<double>> UsdRootMotionPrimTimeSamples;
+	TOptional<double> FirstRootMotionTimeCode;
+	TOptional<double> LastRootMotionTimeCode;
+	TUsdStore<pxr::UsdGeomXformable> RootMotionXformable;
+	{
+		FScopedUsdAllocs UsdAllocs;
+
+		// Note how we don't care whether the root motion is animated or not and will use RootMotionXformable
+		// regardless, to have a similar effect in case its just a single non-animated transform
+		RootMotionXformable = pxr::UsdGeomXformable{ RootMotionPrim };
+		if ( RootMotionXformable.Get() )
+		{
+			std::vector< double > UsdTimeSamples;
+			if ( RootMotionXformable.Get().GetTimeSamples(&UsdTimeSamples) )
+			{
+				if ( UsdTimeSamples.size() > 0 )
+				{
+					FirstRootMotionTimeCode = UsdTimeSamples[ 0 ];
+					LastRootMotionTimeCode = UsdTimeSamples[ UsdTimeSamples.size() - 1 ];
+				}
+			}
+		}
+	}
+
 	// Nothing to do: we don't actually have joints or blend shape time samples
-	if ( !FirstJointSampleTimeCode.IsSet() && !FirstBlendShapeSampleTimeCode.IsSet() )
+	if ( !FirstJointSampleTimeCode.IsSet() && !FirstBlendShapeSampleTimeCode.IsSet() && !FirstRootMotionTimeCode.IsSet() )
 	{
 		return true;
 	}
@@ -1933,8 +1967,20 @@ bool UsdToUnreal::ConvertSkelAnim( const pxr::UsdSkelSkeletonQuery& InUsdSkeleto
 	// which is important because later our composition of tracks and subsections within a LevelSequence will reapply analogous
 	// offsets and scalings anyway
 
-	const double StageStartTimeCode = FMath::Min( FirstJointSampleTimeCode.Get( TNumericLimits<double>::Max() ), FirstBlendShapeSampleTimeCode.Get( TNumericLimits<double>::Max() ) );
-	const double StageEndTimeCode = FMath::Max( LastJointSampleTimeCode.Get( TNumericLimits<double>::Lowest() ), LastBlendShapeSampleTimeCode.Get( TNumericLimits<double>::Lowest() ) );
+	const double StageStartTimeCode = FMath::Min(
+		FirstJointSampleTimeCode.Get( TNumericLimits<double>::Max() ),
+		FMath::Min(
+			FirstBlendShapeSampleTimeCode.Get( TNumericLimits<double>::Max() ),
+			FirstRootMotionTimeCode.Get( TNumericLimits<double>::Max() )
+		)
+	);
+	const double StageEndTimeCode = FMath::Max(
+		LastJointSampleTimeCode.Get( TNumericLimits<double>::Lowest() ),
+		FMath::Max(
+			LastBlendShapeSampleTimeCode.Get( TNumericLimits<double>::Lowest() ),
+			LastRootMotionTimeCode.Get( TNumericLimits<double>::Lowest() )
+		)
+	);
 	const double StageStartSeconds = StageStartTimeCode / StageTimeCodesPerSecond;
 	const double StageSequenceLengthTimeCodes = StageEndTimeCode - StageStartTimeCode;
 	const double LayerSequenceLengthTimeCodes = StageSequenceLengthTimeCodes / Offset.Scale;
@@ -1975,6 +2021,8 @@ bool UsdToUnreal::ConvertSkelAnim( const pxr::UsdSkelSkeletonQuery& InUsdSkeleto
 			JointTrack.ScaleKeys.Reserve(NumBakedFrames);
 		}
 
+		FTransform RootMotionTransform;
+
 		pxr::VtArray<pxr::GfMatrix4d> UsdJointTransforms;
 		for ( int32 FrameIndex = 0; FrameIndex < NumBakedFrames; ++FrameIndex )
 		{
@@ -1986,6 +2034,27 @@ bool UsdToUnreal::ConvertSkelAnim( const pxr::UsdSkelSkeletonQuery& InUsdSkeleto
 				pxr::GfMatrix4d& UsdJointTransform = UsdJointTransforms[ BoneIndex ];
 				FTransform UEJointTransform = UsdToUnreal::ConvertMatrix( StageInfo, UsdJointTransform );
 
+				// Concatenate the root bone transform with the transform track actually present on the skel root as a
+				// whole
+				if ( BoneIndex == 0 )
+				{
+					// We don't care about resetXformStack here: We'll always use the root motion prim's transform as
+					// a local transformation anyway
+					bool* OutResetTransformStack = nullptr;
+					const bool bSuccess = UsdToUnreal::ConvertXformable(
+						Stage.Get(),
+						RootMotionXformable.Get(),
+						RootMotionTransform,
+						StageFrameTimeCodes,
+						OutResetTransformStack
+					);
+
+					if ( bSuccess )
+					{
+						UEJointTransform = UEJointTransform * RootMotionTransform;
+					}
+				}
+
 				FRawAnimSequenceTrack& JointTrack = JointTracks[ BoneIndex ];
 				JointTrack.PosKeys.Add( FVector3f(UEJointTransform.GetTranslation()) );
 				JointTrack.RotKeys.Add( FQuat4f(UEJointTransform.GetRotation()) );
diff --git a/Engine/Plugins/Importers/USDImporter/Source/USDUtilities/Public/USDPrimConversion.h b/Engine/Plugins/Importers/USDImporter/Source/USDUtilities/Public/USDPrimConversion.h
--- a/Engine/Plugins/Importers/USDImporter/Source/USDUtilities/Public/USDPrimConversion.h
+++ b/Engine/Plugins/Importers/USDImporter/Source/USDUtilities/Public/USDPrimConversion.h
@@ -60,7 +60,26 @@ namespace UsdToUnreal
 	UE_DEPRECATED( 5.0, "Use CreatePropertyTrackReader and ConvertTransformTimeSamples (c.f. FUsdLevelSequenceHelperImpl::AddCommonTracks)" )
 	USDUTILITIES_API bool ConvertXformable( const pxr::UsdTyped& Schema, UMovieScene3DTransformTrack& MovieSceneTrack, const FMovieSceneSequenceTransform& SequenceTransform );
 
-	USDUTILITIES_API bool ConvertXformable( const pxr::UsdStageRefPtr& Stage, const pxr::UsdTyped& Schema, USceneComponent& SceneComponent, double EvalTime );
+	/**
+	 * Converts a pxr::UsdGeomXformable's attribute values into an USceneComponent's property values
+	 * @param Stage - Stage that contains the prim to convert
+	 * @param Schema - Prim to convert
+	 * @param SceneComponent - Output component to receieve the converted data
+	 * @param EvalTime - Time at which to sample the prim's attributes
+	 * @param bUsePrimTransform - Whether to convert transform data or not
+	 * @return Whether the conversion was successful or not.
+	 */
+	USDUTILITIES_API bool ConvertXformable( const pxr::UsdStageRefPtr& Stage, const pxr::UsdTyped& Schema, USceneComponent& SceneComponent, double EvalTime, bool bUsePrimTransform = true );
+
+	/**
+	 * Converts a pxr::UsdGeomXformable's transform attribute values into a UE-space FTransform
+	 * @param Stage - Stage that contains the prim to convert
+	 * @param Schema - Prim to convert
+	 * @param OutTransform - Output transform to receieve the converted data
+	 * @param EvalTime - Time at which to sample the prim's attributes
+	 * @param bOutResetTransformStack - Output parameter that, when available, will be set to true if this prim uses the "!resetXformStack!" op
+	 * @return Whether the conversion was successful or not.
+	 */
 	USDUTILITIES_API bool ConvertXformable( const pxr::UsdStageRefPtr& Stage, const pxr::UsdTyped& Schema, FTransform& OutTransform, double EvalTime, bool* bOutResetTransformStack = nullptr );
 
 	UE_DEPRECATED( 5.0, "Prefer the overload that receives an UE::FUsdPrim" )
@@ -95,8 +114,9 @@ namespace UsdToUnreal
 	 * Creates an FPropertyTrackReader that can be used to repeatedly read and convert values of attributes that correspond to PropertyPath in UE.
 	 * e.g. calling this with a RectLight prim and UnrealIdentifiers::IntensityPropertyName will return a reader with a FloatReader member that checks
 	 * the intensity, exposure, width, height of Prim at each UsdTimeCode and returns the corresponding float value for Intensity.
+	 * When we generate a TransformReader, if bIgnorePrimLocalTransform is true it will cause it to ignore that prim's local transform for the reader.
 	 */
-	USDUTILITIES_API FPropertyTrackReader CreatePropertyTrackReader( const UE::FUsdPrim& Prim, const FName& PropertyPath );
+	USDUTILITIES_API FPropertyTrackReader CreatePropertyTrackReader( const UE::FUsdPrim& Prim, const FName& PropertyPath, bool bIgnorePrimLocalTransform = false );
 }
 
 namespace UnrealToUsd
diff --git a/Engine/Plugins/Importers/USDImporter/Source/USDUtilities/Public/USDSkeletalDataConversion.h b/Engine/Plugins/Importers/USDImporter/Source/USDUtilities/Public/USDSkeletalDataConversion.h
--- a/Engine/Plugins/Importers/USDImporter/Source/USDUtilities/Public/USDSkeletalDataConversion.h
+++ b/Engine/Plugins/Importers/USDImporter/Source/USDUtilities/Public/USDSkeletalDataConversion.h
@@ -195,6 +195,7 @@ namespace UsdToUnreal
 	 * @param InSkinningTargets - Skinned meshes that use the skeleton of InUsdSkeletonQuery. Required to fetch the blend shape ordering of each mesh. Optional (can be nullptr to ignore)
 	 * @param InBlendShapes - Converted blend shape data that will be used to interpret blend shape weights as morph target weight float curves. Optional (can be nullptr to ignore)
 	 * @param InInterpretLODs - Whether we try parsing animation data from all LODs of skinning meshes that are inside LOD variant sets
+	 * @param RootMotionPrim - Optional prim whose transform animation will be concatenaded into the root bone's regular joint animation
 	 * @param OutSkeletalAnimationAsset - Output parameter that will be filled with the converted data
 	 * @param OutStartOffsetSeconds - Optional output parameter that will be filled with the offset in seconds of when this UAnimSequence asset should be played since the start of its layer
 	 *								  to match the intended composed animation. The baked UAnimSequence will only contain the range between the first and last joint and/or blend shape
@@ -202,7 +203,15 @@ namespace UsdToUnreal
 	 *                                use this offset when animating USkeletalMeshComponents, and those drive their UAnimSequences with seconds.
 	 * @return Whether the conversion was successful or not.
 	 */
-	USDUTILITIES_API bool ConvertSkelAnim( const pxr::UsdSkelSkeletonQuery& InUsdSkeletonQuery, const pxr::VtArray<pxr::UsdSkelSkinningQuery>* InSkinningTargets, const UsdUtils::FBlendShapeMap* InBlendShapes, bool bInInterpretLODs, UAnimSequence* OutSkeletalAnimationAsset, float* OutStartOffsetSeconds=nullptr);
+	USDUTILITIES_API bool ConvertSkelAnim(
+		const pxr::UsdSkelSkeletonQuery& InUsdSkeletonQuery,
+		const pxr::VtArray<pxr::UsdSkelSkinningQuery>* InSkinningTargets,
+		const UsdUtils::FBlendShapeMap* InBlendShapes,
+		bool bInInterpretLODs,
+		const pxr::UsdPrim& RootMotionPrim,
+		UAnimSequence* OutSkeletalAnimationAsset,
+		float* OutStartOffsetSeconds=nullptr
+	);
 
 	/**
 	 * Builds a USkeletalMesh and USkeleton from the imported data in SkelMeshImportData
diff --git a/Engine/Plugins/Importers/USDImporter/Source/UnrealUSDWrapper/Public/UnrealUSDWrapper.h b/Engine/Plugins/Importers/USDImporter/Source/UnrealUSDWrapper/Public/UnrealUSDWrapper.h
--- a/Engine/Plugins/Importers/USDImporter/Source/UnrealUSDWrapper/Public/UnrealUSDWrapper.h
+++ b/Engine/Plugins/Importers/USDImporter/Source/UnrealUSDWrapper/Public/UnrealUSDWrapper.h
@@ -180,6 +180,22 @@ enum class EUsdInterpolationType : uint8
 	Linear
 };
 
+UENUM()
+enum class EUsdRootMotionHandling
+{
+	// Use for the root bone just its regular joint animation as described on the SkelAnimation prim.
+	NoAdditionalRootMotion,
+
+	// Use the transform animation from the SkelRoot prim in addition to the root bone joint animation as
+	// described on the SkelAnimation prim. Note that the SkelRoot prim's Sequencer transform track will no longer
+	// contain the transform animation data used in this manner, so as to not apply the animation twice.
+	UseMotionFromSkelRoot,
+
+	// Use the transform animation from the Skeleton prim in addition to the root bone joint animation as
+	// described on the SkelAnimation prim.
+	UseMotionFromSkeleton
+};
+
 class IUnrealUSDWrapperModule : public IModuleInterface
 {
 };
