USD: Allow dragging and dropping layers to reorder them.

diff --git a/Engine/Plugins/Importers/USDImporter/Source/USDStageEditor/Private/SUSDLayersTreeView.cpp b/Engine/Plugins/Importers/USDImporter/Source/USDStageEditor/Private/SUSDLayersTreeView.cpp
--- a/Engine/Plugins/Importers/USDImporter/Source/USDStageEditor/Private/SUSDLayersTreeView.cpp
+++ b/Engine/Plugins/Importers/USDImporter/Source/USDStageEditor/Private/SUSDLayersTreeView.cpp
@@ -4,13 +4,16 @@
 
 #include "SUSDStageEditorStyle.h"
 #include "USDClassesModule.h"
+#include "USDDragDropOp.h"
 #include "USDLayersViewModel.h"
 #include "USDLayerUtils.h"
+#include "USDLog.h"
 #include "USDMemory.h"
 #include "USDProjectSettings.h"
 #include "USDStageActor.h"
 #include "USDStageModule.h"
 #include "USDTypesConversion.h"
+#include "UsdWrappers/SdfChangeBlock.h"
 #include "UsdWrappers/SdfLayer.h"
 #include "UsdWrappers/UsdStage.h"
 
@@ -23,6 +26,7 @@
 #include "Modules/ModuleManager.h"
 #include "ScopedTransaction.h"
 #include "Styling/AppStyle.h"
+#include "Styling/SlateIconFinder.h"
 #include "Styling/SlateTypes.h"
 #include "Templates/SharedPointer.h"
 #include "Widgets/Images/SImage.h"
@@ -268,9 +272,560 @@ void SUsdLayersTreeView::Refresh( AUsdStageActor* UsdStageActor, bool bResync )
 	RequestTreeRefresh();
 }
 
+FReply SUsdLayersTreeView::OnRowDragDetected( const FGeometry& Geometry, const FPointerEvent& PointerEvent )
+{
+	TArray<FUsdLayerViewModelRef> Items = GetSelectedItems();
+
+	TSet<TSharedRef<IUsdTreeViewItem>> DraggedItems;
+	DraggedItems.Reserve( Items.Num() );
+
+	for ( const FUsdLayerViewModelRef& Item : Items )
+	{
+		// Don't let a drag operation begin if we're dragging the root or session layer
+		if ( Item->ParentItem == nullptr )
+		{
+			return FReply::Unhandled();
+		}
+
+		DraggedItems.Add( Item );
+	}
+
+	const FSlateBrush* Icon = FAppStyle::GetBrush( TEXT( "Layer.Icon16x" ) );
+	FText DefaultHoverText = DraggedItems.Num() == 1
+		? FText::Format(
+			LOCTEXT( "DefaultTextSingle", "USD layer '{0}'" ),
+			FText::FromString( StaticCastSharedRef<FUsdLayerViewModel>( *DraggedItems.CreateConstIterator() )->LayerIdentifier )
+		)
+		: FText::Format( LOCTEXT( "DefaultTextMultiple", "{0} USD layers" ), DraggedItems.Num() );
+
+	TSharedRef<FUsdDragDropOp> Op = MakeShared<FUsdDragDropOp>();
+	Op->OpType = EUsdDragDropOpType::Layers;
+	Op->DraggedItems = DraggedItems;
+	Op->SetToolTip( DefaultHoverText, Icon );
+	Op->SetupDefaults();
+	Op->Construct();  // Required to initialize the little window that shows tooltips
+
+	return FReply::Handled().BeginDragDrop( Op );
+}
+
+void SUsdLayersTreeView::OnRowDragLeave( const FDragDropEvent& Event )
+{
+	if ( TSharedPtr<FUsdDragDropOp> Op = Event.GetOperationAs<FUsdDragDropOp>() )
+	{
+		Op->ResetToDefaultToolTip();
+	}
+}
+
+namespace UE::USDLayersTreeViewImpl::Private
+{
+	// Checks if we can add all of our dragged layers as sublayers to Parent
+	UsdUtils::ECanInsertSublayerResult CanAddAsSubLayers(
+		const UE::FSdfLayer& Parent,
+		const TSet<TSharedRef<IUsdTreeViewItem>>& SubLayerItems
+	)
+	{
+		for ( const TSharedRef<IUsdTreeViewItem>& Item : SubLayerItems )
+		{
+			const FUsdLayerViewModelRef& SubLayerItem = StaticCastSharedRef<FUsdLayerViewModel>( Item );
+
+			UsdUtils::ECanInsertSublayerResult Result = UsdUtils::CanInsertSubLayer(
+				Parent,
+				*SubLayerItem->LayerIdentifier
+			);
+
+			if ( Result != UsdUtils::ECanInsertSublayerResult::Success )
+			{
+				return Result;
+			}
+		}
+
+		return UsdUtils::ECanInsertSublayerResult::Success;
+	}
+
+	// Checks if this operation would effectively do nothing
+	bool IsNoOp(
+		const TArray< FUsdLayerViewModelRef >& ExistingItems,
+		const TSet< TSharedRef<IUsdTreeViewItem> >& DraggedItems,
+		int32 TargetItemIndex,
+		EItemDropZone DropZone
+	)
+	{
+		switch ( DropZone )
+		{
+			case EItemDropZone::AboveItem:
+			{
+				ensure( TargetItemIndex != INDEX_NONE );
+
+				// If dropped, we'll add the layers before the target index within ExistingItems, so here we just have
+				// to check if the previous items of ExistingItems (before the target) are the ones we're dragging
+
+				// We have dragged more items than there are items after the target, so there is
+				// no way this can be a no-op
+				if ( DraggedItems.Num() > TargetItemIndex )
+				{
+					return false;
+				}
+
+				// See if all the items immediately before the target are the same as the ones we're
+				// dragging
+				TSet<FUsdLayerViewModelRef> UniqueItemsBeforeTarget;
+				for ( int32 Index = TargetItemIndex - 1; Index >= 0; --Index )
+				{
+					const FUsdLayerViewModelRef& OtherItem = ExistingItems[ Index ];
+
+					if ( DraggedItems.Contains( OtherItem ) )
+					{
+						UniqueItemsBeforeTarget.Add( OtherItem );
+					}
+					else
+					{
+						// We found something else other than the dragged items
+						break;
+					}
+				}
+
+				// This means there is another item before the target that is *not* something we dragged:
+				// That means we're dragging something that is not immediately after the target, so
+				// this operation would actually do something
+				if ( UniqueItemsBeforeTarget.Num() < DraggedItems.Num() )
+				{
+					return false;
+				}
+				break;
+			}
+			case EItemDropZone::BelowItem:
+			{
+				ensure( TargetItemIndex != INDEX_NONE );
+
+				// If dropped, we'll add the layers after the target index within ExistingItems, so here we just have
+				// to check if the next items of ExistingItems (after the target) are the ones we're dragging
+
+				// We have dragged more items than there are items after the target, so there is
+				// no way this can be a no-op
+				if ( DraggedItems.Num() > ExistingItems.Num() - TargetItemIndex - 1 )
+				{
+					return false;
+				}
+
+				// See if all the items immediately after the target are the same as the ones we're
+				// dragging
+				TSet<FUsdLayerViewModelRef> UniqueItemsAfterTarget;
+				for ( int32 Index = TargetItemIndex + 1; Index < ExistingItems.Num(); ++Index )
+				{
+					const FUsdLayerViewModelRef& OtherItem = ExistingItems[ Index ];
+
+					if ( DraggedItems.Contains( OtherItem ) )
+					{
+						UniqueItemsAfterTarget.Add( OtherItem );
+					}
+					else
+					{
+						// We found something else other than the dragged items
+						break;
+					}
+				}
+
+				// This means there is another item after the target that is *not* something we dragged:
+				// That means we're dragging something that is not immediately after the target, so
+				// this operation would actually do something
+				if ( UniqueItemsAfterTarget.Num() < DraggedItems.Num() )
+				{
+					return false;
+				}
+				break;
+			}
+			case EItemDropZone::OntoItem:
+			{
+				// If dropped, we add the layers at then end of the sublayer list, so here we just
+				// have to check if the last items of the sublayer list are the ones we're dragging
+
+				if ( DraggedItems.Num() > ExistingItems.Num() )
+				{
+					return false;
+				}
+
+				// See if the last ExistingItems are the same as the ones we're dragging
+				TSet<FUsdLayerViewModelRef> UniqueExistingItemsTail;
+				for ( int32 Index = ExistingItems.Num() - 1; Index >= 0; --Index )
+				{
+					const FUsdLayerViewModelRef& OtherItem = ExistingItems[ Index ];
+
+					if ( DraggedItems.Contains( OtherItem ) )
+					{
+						UniqueExistingItemsTail.Add( OtherItem );
+					}
+					else
+					{
+						// We found something else other than the dragged items
+						break;
+					}
+				}
+
+				// This means there is another item at the end of ExistingItems that is *not* something we dragged:
+				// That means we're dragging something that is not already at the end of ExistingITems, and so that
+				// this operation would actually do something
+				if ( UniqueExistingItemsTail.Num() < DraggedItems.Num() )
+				{
+					return false;
+				}
+
+				break;
+			}
+			default:
+			{
+				return false;
+				break;
+			}
+		}
+
+		return true;
+	}
+}
+
+TOptional<EItemDropZone> SUsdLayersTreeView::OnRowCanAcceptDrop( const FDragDropEvent& Event, EItemDropZone Zone, FUsdLayerViewModelRef Item )
+{
+	using namespace UE::USDLayersTreeViewImpl::Private;
+
+	TOptional<EItemDropZone> Result;
+	if ( TSharedPtr<FUsdDragDropOp> Op = Event.GetOperationAs<FUsdDragDropOp>() )
+	{
+		FText NewHoverText;
+		const FSlateBrush* NewHoverIcon = nullptr;
+		const int32 NumItems = Op->DraggedItems.Num();
+		bool bResetToDefaultAndBlockOp = false;
+
+		if ( Op->OpType == EUsdDragDropOpType::Layers && NumItems > 0 )
+		{
+			// Don't accept dropping onto one of the currently dragged layers
+			for ( const TSharedRef<IUsdTreeViewItem>& DraggedItem : Op->DraggedItems )
+			{
+				if ( StaticCastSharedRef<FUsdLayerViewModel>( DraggedItem ) == Item )
+				{
+					bResetToDefaultAndBlockOp = true;
+					break;
+				}
+			}
+
+			if ( !bResetToDefaultAndBlockOp )
+			{
+				// Special case when hovering the root/session layer. We can only add as sublayers here, as we don't
+				// support swapping root/session layers or having multiple of them
+				// (USD doesn't either... that would be more like reopening a stage?)
+				if ( Item->ParentItem == nullptr )
+				{
+					Zone = EItemDropZone::OntoItem;
+				}
+
+				FUsdLayerViewModel* ParentItem = nullptr;
+				FText MessageEnd;
+				switch ( Zone )
+				{
+					case EItemDropZone::AboveItem:
+					{
+						ParentItem = Item->ParentItem;
+						MessageEnd = FText::Format(
+							LOCTEXT(
+								"DropAboveMessageEnd",
+								", before '{0}'"
+							),
+							FText::FromString( FPaths::GetCleanFilename( Item->LayerIdentifier ) )
+						);
+						break;
+					}
+					case EItemDropZone::BelowItem:
+					{
+						ParentItem = Item->ParentItem;
+						MessageEnd = FText::Format(
+							LOCTEXT(
+								"DropBelowMessageEnd",
+								", after '{0}'"
+							),
+							FText::FromString( FPaths::GetCleanFilename( Item->LayerIdentifier ) )
+						);
+						break;
+					}
+					default:
+					case EItemDropZone::OntoItem:
+					{
+						ParentItem = &Item.Get();
+						MessageEnd = FText::GetEmpty();
+						break;
+					}
+				}
+
+				const UsdUtils::ECanInsertSublayerResult CanAddResult = CanAddAsSubLayers(
+					ParentItem->GetLayer(),
+					Op->DraggedItems
+				);
+
+				// Find the index of the hovered item within its parent, so we know where to insert the dragged layers
+				// into. We don't need this for the EItemDropZone::OntoItem case as we always append them at the end
+				int32 ItemIndexInParent = INDEX_NONE;
+				if ( CanAddResult == UsdUtils::ECanInsertSublayerResult::Success && Zone != EItemDropZone::OntoItem )
+				{
+					for ( int32 Index = 0; Index < ParentItem->Children.Num(); ++Index )
+					{
+						const FUsdLayerViewModelRef& SiblingItem = ParentItem->Children[ Index ];
+						if ( SiblingItem == Item )
+						{
+							ItemIndexInParent = Index;
+							break;
+						}
+					}
+				}
+
+				// We can't allow duplicate sublayers on the same list, so here we need to check if
+				// ParentItem->Children has any item with an identifier that matches one in dragged items list *without
+				// being one of the dragged items itself*.
+				// We have to check for this particular error here and not within CanAddAsSubLayers because
+				// we need the actual tree view items.
+				bool bHasDuplicateLayer = false;
+				if ( CanAddResult == UsdUtils::ECanInsertSublayerResult::Success )
+				{
+					TSet<FString> DraggedItemIdentifiers;
+					DraggedItemIdentifiers.Reserve( Op->DraggedItems.Num() );
+					for ( const TSharedRef<IUsdTreeViewItem>& DraggedItem : Op->DraggedItems )
+					{
+						const FUsdLayerViewModelRef& DraggedLayerItem =
+							StaticCastSharedRef<FUsdLayerViewModel>( DraggedItem );
+
+						DraggedItemIdentifiers.Add( DraggedLayerItem->LayerIdentifier );
+					}
+
+					for ( const FUsdLayerViewModelRef& SiblingItem : ParentItem->Children )
+					{
+						if ( !Op->DraggedItems.Contains( StaticCastSharedRef<IUsdTreeViewItem>( SiblingItem ) )
+							&& DraggedItemIdentifiers.Contains( SiblingItem->LayerIdentifier ) )
+						{
+							bHasDuplicateLayer = true;
+							break;
+						}
+					}
+				}
+
+				const FText ErrorMessage = bHasDuplicateLayer
+					? LOCTEXT( "CanDropDuplicate_Text", "Cannot add duplicate SubLayer!" )
+					: UsdUtils::ToText( CanAddResult );
+
+				const bool bIsNoOp = !bHasDuplicateLayer
+					&& CanAddResult == UsdUtils::ECanInsertSublayerResult::Success
+					&& IsNoOp(
+						ParentItem->Children,
+						Op->DraggedItems,
+						ItemIndexInParent,
+						Zone
+					);
+
+				// Drag and drop would do nothing
+				if ( bIsNoOp )
+				{
+					bResetToDefaultAndBlockOp = true;
+				}
+				// Drag and drop cannot be performed for some reason
+				else if ( bHasDuplicateLayer || CanAddResult != UsdUtils::ECanInsertSublayerResult::Success )
+				{
+					NewHoverText = FText::Format(
+						LOCTEXT(
+							"CanDropError_Text",
+							"Cannot add dragged {0}|plural(one=layer,other=layers) as {0}|plural(one=sublayer,other=sublayers) of '{1}': {2}"
+						),
+						NumItems,
+						FText::FromString( FPaths::GetCleanFilename( ParentItem->LayerIdentifier ) ),
+						ErrorMessage
+					);
+				}
+				// Can drop the one dragged layer
+				else if ( NumItems == 1 )
+				{
+					FString FirstDraggedLayer = FPaths::GetCleanFilename(
+						StaticCastSharedRef<FUsdLayerViewModel>(
+							*Op->DraggedItems.CreateConstIterator()
+						)->LayerIdentifier
+					);
+
+					NewHoverText = FText::Format(
+						LOCTEXT(
+							"CanDropSingle_Text",
+							"Add layer '{0}' as a sublayer of '{1}'{2}"
+						),
+						FText::FromString( FirstDraggedLayer ),
+						FText::FromString( FPaths::GetCleanFilename( ParentItem->LayerIdentifier ) ),
+						MessageEnd
+					);
+
+					Result = { Zone };
+				}
+				// Can drop multiple dragged layers
+				else
+				{
+					NewHoverText = FText::Format(
+						LOCTEXT(
+							"CanDropMultiple_Text",
+							"Add {0} layers as sublayers of '{1}'{2}"
+						),
+						NumItems,
+						FText::FromString( FPaths::GetCleanFilename( ParentItem->LayerIdentifier ) ),
+						MessageEnd
+					);
+
+					Result = { Zone };
+				}
+			}
+		}
+		else
+		{
+			NewHoverText = LOCTEXT(
+				"CanDropUnsupported_Text",
+				"Can only drag and drop layers for now"
+			);
+		}
+
+		if ( bResetToDefaultAndBlockOp )
+		{
+			Op->ResetToDefaultToolTip();
+			Result.Reset();
+		}
+		else
+		{
+			NewHoverIcon = FAppStyle::GetBrush( Result.IsSet()
+				? TEXT( "Graph.ConnectorFeedback.Ok" )
+				: TEXT( "Graph.ConnectorFeedback.Error" )
+			);
+
+			Op->SetToolTip( NewHoverText, NewHoverIcon );
+		}
+	}
+
+	return Result;
+}
+
+FReply SUsdLayersTreeView::OnRowAcceptDrop( const FDragDropEvent& Event, EItemDropZone Zone, FUsdLayerViewModelRef Item )
+{
+	// This doesn't work very well USD-wise yet, and I don't know how we can possibly undo/redo layer reassignments
+	// like this. We do need a transaction though, as this may trigger the creation or deletion of UObjects
+	FScopedTransaction Transaction( LOCTEXT( "OnAcceptDropTransaction", "Drop dragged USD layers" ) );
+
+	if ( TSharedPtr<FUsdDragDropOp> Op = Event.GetOperationAs<FUsdDragDropOp>() )
+	{
+		// Only accept layers for now
+		const int32 NumItems = Op->DraggedItems.Num();
+		if ( Op->OpType == EUsdDragDropOpType::Layers && NumItems > 0 )
+		{
+			UE::FSdfChangeBlock ChangeBlock;
+
+			// Make sure our layers aren't closed by USD between us removing them and readding
+			TArray<UE::FSdfLayer> LayerPins;
+			LayerPins.Reserve( Op->DraggedItems.Num() );
+
+			// Remove the dragged items from their parents before we even fetch where to insert them into as this will
+			// affect the target indices
+			for ( const TSharedRef<IUsdTreeViewItem>& DraggedItem : Op->DraggedItems )
+			{
+				const FUsdLayerViewModelRef& LayerItem = StaticCastSharedRef<FUsdLayerViewModel>( DraggedItem );
+				UE::FSdfLayer LayerPin = LayerItem->GetLayer();
+				LayerPins.Add( LayerPin );
+
+				// We are only allowed to drag off things that have parents: There would never be anywhere to add the
+				// root/session layer into as that would create loops/strange situations
+				FUsdLayerViewModel* OldParent = LayerItem->ParentItem;
+				if ( ensure( OldParent ) )
+				{
+					if ( UE::FSdfLayer OldParentLayer = OldParent->GetLayer() )
+					{
+						int32 OldIndex = OldParent->Children.IndexOfByKey( LayerItem->AsShared() );
+						ensure( OldIndex != INDEX_NONE );
+
+						OldParentLayer.RemoveSubLayerPath( OldIndex );
+
+						// If we have dragged multiple items, the upcoming insertion of this item may affect all the
+						// other removals indices too, so we need to make sure our widgets are up-to-date. They won't
+						// be updated on the next iteration of this loop even if we remove the outer FSdfChangeBlock,
+						// as the resync would end up creating a new batch of widgets anyway
+						OldParent->Children.RemoveAt( OldIndex );
+					}
+				}
+			}
+
+			UE::FSdfLayer ParentLayer;
+			int32 TargetSubLayerIndex = INDEX_NONE;
+			switch ( Zone )
+			{
+				case EItemDropZone::AboveItem:
+				{
+					// Here we'll add the dragged items as siblings of Item, before it
+
+					FUsdLayerViewModel* ParentItem = Item->ParentItem;
+					if ( ensure( ParentItem ) )
+					{
+						ParentLayer = ParentItem->GetLayer();
+						TargetSubLayerIndex = ParentItem->Children.IndexOfByKey( Item );
+						ensure( TargetSubLayerIndex != INDEX_NONE );
+					}
+					break;
+				}
+				case EItemDropZone::BelowItem:
+				{
+					// Here we'll add the dragged items as siblings of Item, after it
+
+					FUsdLayerViewModel* ParentItem = Item->ParentItem;
+					if ( ensure( ParentItem ) )
+					{
+						ParentLayer = ParentItem->GetLayer();
+						TargetSubLayerIndex = ParentItem->Children.IndexOfByKey( Item );
+						ensure( TargetSubLayerIndex != INDEX_NONE );
+					}
+					break;
+				}
+				default:
+				case EItemDropZone::OntoItem:
+				{
+					// Just add the dragged items as children of Item, at the end of its list of children
+
+					ParentLayer = Item->GetLayer();
+					break;
+				}
+			}
+
+			for ( const TSharedRef<IUsdTreeViewItem>& DraggedItem : Op->DraggedItems )
+			{
+				const FUsdLayerViewModelRef& LayerItem = StaticCastSharedRef<FUsdLayerViewModel>( DraggedItem );
+				UE::FSdfLayer LayerPin = LayerItem->GetLayer();
+
+				const bool bInserted = UsdUtils::InsertSubLayer(
+					ParentLayer,
+					*LayerPin.GetIdentifier(),
+					Zone == EItemDropZone::OntoItem
+						? -1 // Always at the end
+						: Zone == EItemDropZone::AboveItem
+							? TargetSubLayerIndex++
+							: ++TargetSubLayerIndex
+				);
+
+				if ( !bInserted )
+				{
+					UE_LOG( LogUsd, Warning, TEXT( "Failed to insert layer '%s' as a sublayer of '%s' at index '%d'" ),
+						*LayerPin.GetIdentifier(),
+						*ParentLayer.GetIdentifier(),
+						TargetSubLayerIndex
+					);
+				}
+
+			}
+
+			return FReply::Handled();
+		}
+	}
+
+	return FReply::Unhandled();
+}
+
 TSharedRef< ITableRow > SUsdLayersTreeView::OnGenerateRow( FUsdLayerViewModelRef InDisplayNode, const TSharedRef< STableViewBase >& OwnerTable )
 {
-	return SNew( SUsdTreeRow< FUsdLayerViewModelRef >, InDisplayNode, OwnerTable, SharedData );
+	return SNew( SUsdTreeRow< FUsdLayerViewModelRef >, InDisplayNode, OwnerTable, SharedData )
+		.OnDragDetected( this, &SUsdLayersTreeView::OnRowDragDetected )
+		.OnDragLeave( this, &SUsdLayersTreeView::OnRowDragLeave )
+		.OnCanAcceptDrop( this, &SUsdLayersTreeView::OnRowCanAcceptDrop )
+		.OnAcceptDrop( this, &SUsdLayersTreeView::OnRowAcceptDrop );
 }
 
 void SUsdLayersTreeView::OnGetChildren( FUsdLayerViewModelRef InParent, TArray< FUsdLayerViewModelRef >& OutChildren ) const
@@ -382,7 +937,7 @@ TSharedPtr< SWidget > SUsdLayersTreeView::ConstructLayerContextMenu()
 			FSlateIcon(),
 			FUIAction(
 				FExecuteAction::CreateSP( this, &SUsdLayersTreeView::OnAddSubLayer ),
-				FCanExecuteAction::CreateSP( this, &SUsdLayersTreeView::CanAddSubLayer )
+				FCanExecuteAction::CreateSP( this, &SUsdLayersTreeView::CanInsertSubLayer )
 			),
 			NAME_None,
 			EUserInterfaceActionType::Button
@@ -394,7 +949,7 @@ TSharedPtr< SWidget > SUsdLayersTreeView::ConstructLayerContextMenu()
 			FSlateIcon(),
 			FUIAction(
 				FExecuteAction::CreateSP( this, &SUsdLayersTreeView::OnNewSubLayer ),
-				FCanExecuteAction::CreateSP( this, &SUsdLayersTreeView::CanAddSubLayer )
+				FCanExecuteAction::CreateSP( this, &SUsdLayersTreeView::CanInsertSubLayer )
 			),
 			NAME_None,
 			EUserInterfaceActionType::Button
@@ -706,7 +1261,7 @@ void SUsdLayersTreeView::OnExportSelectedLayers() const
 	}
 }
 
-bool SUsdLayersTreeView::CanAddSubLayer() const
+bool SUsdLayersTreeView::CanInsertSubLayer() const
 {
 	return GetSelectedItems().Num() > 0;
 }
diff --git a/Engine/Plugins/Importers/USDImporter/Source/USDStageEditor/Private/SUSDLayersTreeView.h b/Engine/Plugins/Importers/USDImporter/Source/USDStageEditor/Private/SUSDLayersTreeView.h
--- a/Engine/Plugins/Importers/USDImporter/Source/USDStageEditor/Private/SUSDLayersTreeView.h
+++ b/Engine/Plugins/Importers/USDImporter/Source/USDStageEditor/Private/SUSDLayersTreeView.h
@@ -20,6 +20,13 @@ public:
 	void Construct( const FArguments& InArgs, AUsdStageActor* UsdStageActor );
 	void Refresh( AUsdStageActor* UsdStageActor, bool bResync );
 
+	// Drag and drop interface for our rows
+	FReply OnRowDragDetected( const FGeometry& Geometry, const FPointerEvent& PointerEvent );
+	void OnRowDragLeave( const FDragDropEvent& Event );
+	TOptional<EItemDropZone> OnRowCanAcceptDrop( const FDragDropEvent& Event, EItemDropZone Zone, FUsdLayerViewModelRef Item );
+	FReply OnRowAcceptDrop( const FDragDropEvent& Event, EItemDropZone Zone, FUsdLayerViewModelRef Item );
+	// End drag and drop interface
+
 private:
 	virtual TSharedRef< ITableRow > OnGenerateRow( FUsdLayerViewModelRef InDisplayNode, const TSharedRef< STableViewBase >& OwnerTable ) override;
 	virtual void OnGetChildren( FUsdLayerViewModelRef InParent, TArray< FUsdLayerViewModelRef >& OutChildren ) const override;
@@ -41,7 +48,7 @@ private:
 
 	void OnExportSelectedLayers() const;
 
-	bool CanAddSubLayer() const;
+	bool CanInsertSubLayer() const;
 	void OnAddSubLayer();
 	void OnNewSubLayer();
 
diff --git a/Engine/Plugins/Importers/USDImporter/Source/USDStageEditor/Private/SUSDStageTreeView.cpp b/Engine/Plugins/Importers/USDImporter/Source/USDStageEditor/Private/SUSDStageTreeView.cpp
--- a/Engine/Plugins/Importers/USDImporter/Source/USDStageEditor/Private/SUSDStageTreeView.cpp
+++ b/Engine/Plugins/Importers/USDImporter/Source/USDStageEditor/Private/SUSDStageTreeView.cpp
@@ -1735,7 +1735,7 @@ void SUsdStageTreeView::ScrollItemIntoView( FUsdPrimViewModelRef TreeItem )
 	FUsdPrimViewModel* Parent = TreeItem->ParentItem;
 	while( Parent )
 	{
-		SetItemExpansion( Parent->AsShared(), true );
+		SetItemExpansion( StaticCastSharedRef< FUsdPrimViewModel>( Parent->AsShared() ), true );
 		Parent = Parent->ParentItem;
 	}
 
diff --git a/Engine/Plugins/Importers/USDImporter/Source/USDStageEditor/Private/USDDragDropOp.h b/Engine/Plugins/Importers/USDImporter/Source/USDStageEditor/Private/USDDragDropOp.h
new file mode 100644
--- /dev/null
+++ b/Engine/Plugins/Importers/USDImporter/Source/USDStageEditor/Private/USDDragDropOp.h
@@ -0,0 +1,22 @@
+// Copyright Epic Games, Inc. All Rights Reserved.
+
+#pragma once
+
+#include "DragAndDrop/DecoratedDragDropOp.h"
+
+enum class EUsdDragDropOpType
+{
+    None,
+    Prims,
+    Layers,
+    Attributes
+};
+
+class FUsdDragDropOp : public FDecoratedDragDropOp
+{
+public:
+	DRAG_DROP_OPERATOR_TYPE( FUsdDragDropOp, FDecoratedDragDropOp )
+
+    EUsdDragDropOpType OpType;
+	TSet<TSharedRef<IUsdTreeViewItem>> DraggedItems;
+};
\ No newline at end of file
diff --git a/Engine/Plugins/Importers/USDImporter/Source/USDStageEditor/Private/Widgets/SUSDTreeView.h b/Engine/Plugins/Importers/USDImporter/Source/USDStageEditor/Private/Widgets/SUSDTreeView.h
--- a/Engine/Plugins/Importers/USDImporter/Source/USDStageEditor/Private/Widgets/SUSDTreeView.h
+++ b/Engine/Plugins/Importers/USDImporter/Source/USDStageEditor/Private/Widgets/SUSDTreeView.h
@@ -88,7 +88,16 @@ template< typename ItemType >
 class SUsdTreeRow : public SMultiColumnTableRow< ItemType >
 {
 public:
+	typedef typename STableRow<ItemType>::FOnCanAcceptDrop FUsdOnCanAcceptDrop;
+	typedef typename STableRow<ItemType>::FOnAcceptDrop FUsdOnAcceptDrop;
+
 	SLATE_BEGIN_ARGS( SUsdTreeRow ) {}
+		SLATE_EVENT( FUsdOnCanAcceptDrop, OnCanAcceptDrop )
+		SLATE_EVENT( FUsdOnAcceptDrop, OnAcceptDrop )
+		SLATE_EVENT( FOnDragDetected, OnDragDetected )
+		SLATE_EVENT( FOnTableRowDragEnter, OnDragEnter )
+		SLATE_EVENT( FOnTableRowDragLeave, OnDragLeave )
+		SLATE_EVENT( FOnTableRowDrop, OnDrop )
 	SLATE_END_ARGS()
 
 	void Construct( const FArguments& InArgs, ItemType InTreeItem, const TSharedRef< STableViewBase >& OwnerTable, TSharedPtr< FSharedUsdTreeData > InSharedData )
@@ -96,7 +105,16 @@ public:
 		TreeItem = InTreeItem;
 		SharedData = InSharedData;
 
-		SMultiColumnTableRow< ItemType >::Construct( typename SMultiColumnTableRow< ItemType >::FArguments(), OwnerTable );
+		SMultiColumnTableRow< ItemType >::Construct(
+			typename SMultiColumnTableRow< ItemType >::FArguments()
+				.OnCanAcceptDrop( InArgs._OnCanAcceptDrop )
+				.OnAcceptDrop( InArgs._OnAcceptDrop )
+				.OnDragDetected( InArgs._OnDragDetected )
+				.OnDragEnter( InArgs._OnDragEnter )
+				.OnDragLeave( InArgs._OnDragLeave )
+				.OnDrop( InArgs._OnDrop ),
+			OwnerTable
+		);
 	}
 
 	/** Overridden from SMultiColumnTableRow.  Generates a widget for this column of the tree row. */
diff --git a/Engine/Plugins/Importers/USDImporter/Source/USDStageEditorViewModels/Public/USDPrimViewModel.h b/Engine/Plugins/Importers/USDImporter/Source/USDStageEditorViewModels/Public/USDPrimViewModel.h
--- a/Engine/Plugins/Importers/USDImporter/Source/USDStageEditorViewModels/Public/USDPrimViewModel.h
+++ b/Engine/Plugins/Importers/USDImporter/Source/USDStageEditorViewModels/Public/USDPrimViewModel.h
@@ -30,7 +30,7 @@ public:
 	bool bIsVisible = true;
 };
 
-class USDSTAGEEDITORVIEWMODELS_API FUsdPrimViewModel : public IUsdTreeViewItem, public TSharedFromThis< FUsdPrimViewModel >
+class USDSTAGEEDITORVIEWMODELS_API FUsdPrimViewModel : public IUsdTreeViewItem
 {
 public:
 	FUsdPrimViewModel( FUsdPrimViewModel* InParentItem, const UE::FUsdStageWeak& InUsdStage, const UE::FUsdPrim& InUsdPrim );
diff --git a/Engine/Plugins/Importers/USDImporter/Source/USDStageEditorViewModels/Public/USDTreeItemViewModel.h b/Engine/Plugins/Importers/USDImporter/Source/USDStageEditorViewModels/Public/USDTreeItemViewModel.h
--- a/Engine/Plugins/Importers/USDImporter/Source/USDStageEditorViewModels/Public/USDTreeItemViewModel.h
+++ b/Engine/Plugins/Importers/USDImporter/Source/USDStageEditorViewModels/Public/USDTreeItemViewModel.h
@@ -3,8 +3,9 @@
 #pragma once
 
 #include "HAL/Platform.h"
+#include "Templates/SharedPointer.h"
 
-class USDSTAGEEDITORVIEWMODELS_API IUsdTreeViewItem
+class USDSTAGEEDITORVIEWMODELS_API IUsdTreeViewItem : public TSharedFromThis< IUsdTreeViewItem >
 {
 public:
 	virtual ~IUsdTreeViewItem() = default;
diff --git a/Engine/Plugins/Importers/USDImporter/Source/USDUtilities/Private/USDLayerUtils.cpp b/Engine/Plugins/Importers/USDImporter/Source/USDUtilities/Private/USDLayerUtils.cpp
--- a/Engine/Plugins/Importers/USDImporter/Source/USDUtilities/Private/USDLayerUtils.cpp
+++ b/Engine/Plugins/Importers/USDImporter/Source/USDUtilities/Private/USDLayerUtils.cpp
@@ -72,12 +72,87 @@ namespace UE::UsdLayerUtilsImpl::Private
 	}
 }
 
-namespace UsdUtils
+FText UsdUtils::ToText( ECanInsertSublayerResult Result )
 {
-	std::string GetUESessionStateLayerDisplayName( const pxr::UsdStageRefPtr& Stage )
+	switch ( Result )
+	{
+	default:
+	case ECanInsertSublayerResult::Success:
+		return FText::GetEmpty();
+		break;
+	case ECanInsertSublayerResult::ErrorSubLayerNotFound:
+		return LOCTEXT( "CanAddSubLayerNotFound", "SubLayer not found!" );
+		break;
+	case ECanInsertSublayerResult::ErrorSubLayerInvalid:
+		return LOCTEXT( "CanAddSubLayerInvalid", "SubLayer is invalid!" );
+		break;
+	case ECanInsertSublayerResult::ErrorSubLayerIsParentLayer:
+		return LOCTEXT( "CanAddSubLayerIsParent", "SubLayer is the same as the parent layer!" );
+		break;
+	case ECanInsertSublayerResult::ErrorCycleDetected:
+		return LOCTEXT( "CanAddSubLayerCycle", "Cycles detected!" );
+		break;
+	}
+}
+
+UsdUtils::ECanInsertSublayerResult UsdUtils::CanInsertSubLayer(
+	const pxr::SdfLayerRefPtr& ParentLayer,
+	const TCHAR* SubLayerIdentifier
+)
+{
+	if ( !SubLayerIdentifier )
+	{
+		return ECanInsertSublayerResult::ErrorSubLayerNotFound;
+	}
+
+	FScopedUsdAllocs Allocs;
+
+	pxr::SdfLayerRefPtr SubLayer = pxr::SdfLayer::FindOrOpen( UnrealToUsd::ConvertString( SubLayerIdentifier ).Get() );
+	if ( !SubLayer )
+	{
+		return ECanInsertSublayerResult::ErrorSubLayerNotFound;
+	}
+
+	if ( SubLayer == ParentLayer )
 	{
-		return pxr::SdfLayer::GetDisplayNameFromIdentifier( Stage->GetRootLayer()->GetIdentifier() ) + "-UE-session-state.usda";
+		return ECanInsertSublayerResult::ErrorSubLayerIsParentLayer;
 	}
+
+	// We can't climb through ancestors of ParentLayer, so we have to open sublayer and see if parentlayer is a
+	// descendant of *it* in order to detect cycles
+	TFunction< ECanInsertSublayerResult( const pxr::SdfLayerRefPtr& ) > CanAddSubLayerRecursive;
+	CanAddSubLayerRecursive = [ParentLayer, &CanAddSubLayerRecursive](
+		const pxr::SdfLayerRefPtr& CurrentParent
+	) -> ECanInsertSublayerResult
+	{
+		for ( const std::string& SubLayerPath : CurrentParent->GetSubLayerPaths() )
+		{
+			// This may seem expensive, but keep in mind the main use case for this (at least for now) is for checking
+			// during layer drag and drop, where all of these layers are actually already open anyway
+			pxr::SdfLayerRefPtr ChildSubLayer = pxr::SdfLayer::FindOrOpenRelativeToLayer(
+				CurrentParent,
+				SubLayerPath
+			);
+
+			if ( !ChildSubLayer )
+			{
+				return ECanInsertSublayerResult::ErrorSubLayerInvalid;
+			}
+
+			ECanInsertSublayerResult RecursiveResult = ChildSubLayer == ParentLayer
+				 ? ECanInsertSublayerResult::ErrorCycleDetected
+				 : CanAddSubLayerRecursive( ChildSubLayer );
+
+			if ( RecursiveResult != ECanInsertSublayerResult::Success )
+			{
+				return RecursiveResult;
+			}
+		}
+
+		return ECanInsertSublayerResult::Success;
+	};
+
+	return CanAddSubLayerRecursive( SubLayer );
 }
 
 bool UsdUtils::InsertSubLayer( const pxr::SdfLayerRefPtr& ParentLayer, const TCHAR* SubLayerFile, int32 Index, double OffsetTimeCodes, double TimeCodesScale )
diff --git a/Engine/Plugins/Importers/USDImporter/Source/USDUtilities/Public/USDLayerUtils.h b/Engine/Plugins/Importers/USDImporter/Source/USDUtilities/Public/USDLayerUtils.h
--- a/Engine/Plugins/Importers/USDImporter/Source/USDUtilities/Public/USDLayerUtils.h
+++ b/Engine/Plugins/Importers/USDImporter/Source/USDUtilities/Public/USDLayerUtils.h
@@ -46,16 +46,42 @@ namespace UsdUtils
 		Save
 	};
 
+	enum class ECanInsertSublayerResult
+	{
+		Success,
+		ErrorSubLayerNotFound,
+		ErrorSubLayerInvalid,
+		ErrorSubLayerIsParentLayer,
+		ErrorCycleDetected,
+	};
+	USDUTILITIES_API FText ToText( ECanInsertSublayerResult Result );
+
+	/**
+	 * Checks whether we can add SubLayerIdentifier as a sublayer to ParentLayer.
+	 * The index is actually not relevant.
+	 */
+	USDUTILITIES_API ECanInsertSublayerResult CanInsertSubLayer(
+		const pxr::SdfLayerRefPtr& ParentLayer,
+		const TCHAR* SubLayerIdentifier
+	);
+
 	/**
 	 * Inserts the SubLayerFile path into ParentLayer as a sublayer
 	 * @param ParentLayer - Layer to receive the new sublayer entry
 	 * @param SubLayerFile - Absolute path to a usd file to use as sublayer
-	 * @param Index - Zero-based index on the list of ParentLayer's sublayers list to insert the new sublayer. -1 means "at the end" (the default)
+	 * @param Index - Zero-based index on the list of ParentLayer's sublayers list to insert the new sublayer.
+					  -1 means "at the end" (the default)
 	 * @param OffsetTimeCodes - Offset in USD time codes to use for the sublayer reference
 	 * @param TimeCodesScale - Scale to use for the sublayer reference
 	 * @return Whether the sublayer reference was added successfully or not
 	 */
-	USDUTILITIES_API bool InsertSubLayer( const pxr::SdfLayerRefPtr& ParentLayer, const TCHAR* SubLayerFile, int32 Index = -1, double OffsetTimeCodes = 0.0, double TimeCodesScale = 1.0 );
+	USDUTILITIES_API bool InsertSubLayer(
+		const pxr::SdfLayerRefPtr& ParentLayer,
+		const TCHAR* SubLayerFile,
+		int32 Index = -1,
+		double OffsetTimeCodes = 0.0,
+		double TimeCodesScale = 1.0
+	);
 
 #if WITH_EDITOR
 	/** Opens a file dialog to open or save a USD file. The returned file path will always be absolute */
diff --git a/Engine/Plugins/Importers/USDImporter/Source/UnrealUSDWrapper/Private/UsdWrappers/SdfLayer.cpp b/Engine/Plugins/Importers/USDImporter/Source/UnrealUSDWrapper/Private/UsdWrappers/SdfLayer.cpp
--- a/Engine/Plugins/Importers/USDImporter/Source/UnrealUSDWrapper/Private/UsdWrappers/SdfLayer.cpp
+++ b/Engine/Plugins/Importers/USDImporter/Source/UnrealUSDWrapper/Private/UsdWrappers/SdfLayer.cpp
@@ -520,6 +520,17 @@ namespace UE
 		return 0.0;
 	}
 
+	template<typename PtrType>
+	void FSdfLayerBase<PtrType>::RemoveSubLayerPath( int32 Index )
+	{
+#if USE_USD_SDK
+		if ( const PtrType& Ptr = Impl->GetInner() )
+		{
+			Ptr->RemoveSubLayerPath( Index );
+		}
+#endif // #if USE_USD_SDK
+	}
+
 	template<typename PtrType>
 	void FSdfLayerBase<PtrType>::SetFramesPerSecond( double FramesPerSecond )
 	{
diff --git a/Engine/Plugins/Importers/USDImporter/Source/UnrealUSDWrapper/Public/UsdWrappers/SdfLayer.h b/Engine/Plugins/Importers/USDImporter/Source/UnrealUSDWrapper/Public/UsdWrappers/SdfLayer.h
--- a/Engine/Plugins/Importers/USDImporter/Source/UnrealUSDWrapper/Public/UsdWrappers/SdfLayer.h
+++ b/Engine/Plugins/Importers/USDImporter/Source/UnrealUSDWrapper/Public/UsdWrappers/SdfLayer.h
@@ -121,6 +121,7 @@ namespace UE
 
 		bool HasFramesPerSecond() const;
 		double GetFramesPerSecond() const;
+		void RemoveSubLayerPath( int32 Index );
 		void SetFramesPerSecond( double FramesPerSecond );
 
 		int64 GetNumSubLayerPaths() const;
