USD: Added support for GroomCache. When applying the GroomAPI schema to a groom prim, it will detect and import GroomCache if animated attributes (position and width) are found on the groom. This works with the USD stage, import assets and import into level.

diff --git a/Engine/Plugins/Importers/USDImporter/Source/USDSchemas/Private/USDGroomTranslator.cpp b/Engine/Plugins/Importers/USDImporter/Source/USDSchemas/Private/USDGroomTranslator.cpp
--- a/Engine/Plugins/Importers/USDImporter/Source/USDSchemas/Private/USDGroomTranslator.cpp
+++ b/Engine/Plugins/Importers/USDImporter/Source/USDSchemas/Private/USDGroomTranslator.cpp
@@ -6,16 +6,22 @@
 
 #include "GroomAsset.h"
 #include "GroomBuilder.h"
+#include "GroomCache.h"
+#include "GroomCacheImporter.h"
 #include "GroomComponent.h"
 #include "GroomImportOptions.h"
 #include "HairDescription.h"
 #include "HairStrandsImporter.h"
 #include "Misc/ArchiveMD5.h"
+#include "Templates/UniquePtr.h"
+#include "UObject/StrongObjectPtr.h"
 
 #include "UnrealUSDWrapper.h"
 #include "USDAssetImportData.h"
 #include "USDGroomConversion.h"
+#include "USDGroomTranslatorUtils.h"
 #include "USDIntegrationUtils.h"
+#include "USDLog.h"
 #include "USDTypesConversion.h"
 
 #include "USDIncludesStart.h"
@@ -82,6 +88,19 @@ namespace UE::UsdGroomTranslator::Private
 
 		return SHAHash;
 	}
+
+	void ComputeFrameHairDescriptionHash(const pxr::UsdPrim& Prim, TArray<FHairGroupsInterpolation>& BuildSettings, int32 FrameNumber, FSHA1& Hash)
+	{
+		FHairDescription FrameHairDescription;
+		if (UsdToUnreal::ConvertGroomHierarchy(Prim, pxr::UsdTimeCode(FrameNumber), FTransform::Identity, FrameHairDescription))
+		{
+			FHairDescriptionGroups GroupsDescription;
+			FGroomBuilder::BuildHairDescriptionGroups(FrameHairDescription, GroupsDescription);
+
+			FSHAHash SHAHash = ComputeHairDescriptionHash(FrameHairDescription, BuildSettings);
+			Hash.Update((const uint8*)SHAHash.Hash, sizeof(SHAHash.Hash));
+		}
+	}
 }
 
 class FUsdGroomCreateAssetsTaskChain : public FUsdSchemaTranslatorTaskChain
@@ -91,6 +110,10 @@ public:
 		: PrimPath(InPrimPath)
 		, Context(InContext)
 	{
+		AnimInfo.Attributes = EGroomCacheAttributes::None;
+		AnimInfo.StartFrame = TNumericLimits<int32>::Max();
+		AnimInfo.EndFrame = TNumericLimits<int32>::Min();
+
 		SetupTasks();
 	}
 
@@ -101,6 +124,11 @@ protected:
 
 	FHairDescription HairDescription;
 
+	TStrongObjectPtr<UGroomImportOptions> ImportOptions;
+	TUniquePtr<FGroomCacheProcessor> GroomCacheProcessor;
+	FGroomAnimationInfo AnimInfo;
+	FSHAHash GroomCacheHash;
+
 protected:
 	UE::FUsdPrim GetPrim() const { return Context->Stage.GetPrimAtPath(PrimPath); }
 
@@ -113,7 +141,17 @@ protected:
 		Do(ESchemaTranslationLaunchPolicy::Async,
 			[this]() -> bool
 			{
-				const bool bSuccess = UsdToUnreal::ConvertGroomHierarchy(GetPrim(), pxr::UsdTimeCode::EarliestTime(), FTransform::Identity, HairDescription);
+				const bool bSuccess = UsdToUnreal::ConvertGroomHierarchy(GetPrim(), pxr::UsdTimeCode::EarliestTime(), FTransform::Identity, HairDescription, &AnimInfo);
+
+				const double StageTimeCodesPerSecond = Context->Stage.GetTimeCodesPerSecond();
+				AnimInfo.SecondsPerFrame = float(1.0f / StageTimeCodesPerSecond);
+
+				if (bSuccess && AnimInfo.IsValid())
+				{
+					AnimInfo.Duration = AnimInfo.NumFrames * AnimInfo.SecondsPerFrame;
+					AnimInfo.StartTime = AnimInfo.StartFrame * AnimInfo.SecondsPerFrame;
+					AnimInfo.EndTime = AnimInfo.EndFrame * AnimInfo.SecondsPerFrame;
+				}
 
 				return bSuccess && HairDescription.IsValid();
 			});
@@ -127,7 +165,7 @@ protected:
 				FHairDescriptionGroups GroupsDescription;
 				FGroomBuilder::BuildHairDescriptionGroups(HairDescription, GroupsDescription);
 
-				UGroomImportOptions* ImportOptions = UE::UsdGroomTranslator::Private::CreateGroomImportOptions(GroupsDescription, Context->GroomInterpolationSettings);
+				ImportOptions.Reset(UE::UsdGroomTranslator::Private::CreateGroomImportOptions(GroupsDescription, Context->GroomInterpolationSettings));
 
 				FSHAHash SHAHash = UE::UsdGroomTranslator::Private::ComputeHairDescriptionHash(HairDescription, ImportOptions->InterpolationSettings);
 
@@ -137,7 +175,7 @@ protected:
 				{
 					FName AssetName = MakeUniqueObjectName(GetTransientPackage(), UGroomAsset::StaticClass(), *FPaths::GetBaseFilename(PrimPathString));
 
-					FHairImportContext HairImportContext(ImportOptions, GetTransientPackage(), UGroomAsset::StaticClass(), AssetName, Context->ObjectFlags | EObjectFlags::RF_Public);
+					FHairImportContext HairImportContext(ImportOptions.Get(), GetTransientPackage(), UGroomAsset::StaticClass(), AssetName, Context->ObjectFlags | EObjectFlags::RF_Public);
 					UGroomAsset* ExistingAsset = nullptr;
 					GroomAsset = FHairStrandsImporter::ImportHair(HairImportContext, HairDescription, ExistingAsset);
 					if (GroomAsset)
@@ -146,7 +184,7 @@ protected:
 #if WITH_EDITOR
 						UUsdAssetImportData* ImportData = NewObject<UUsdAssetImportData>(GroomAsset, TEXT("UUSDAssetImportData"));
 						ImportData->PrimPath = PrimPathString;
-						ImportData->ImportOptions = ImportOptions;
+						ImportData->ImportOptions = ImportOptions.Get();
 
 						GroomAsset->AssetImportData = ImportData;
 #endif // WITH_EDITOR
@@ -158,7 +196,138 @@ protected:
 					Context->AssetCache->LinkAssetToPrim(PrimPathString, GroomAsset);
 				}
 
-				return true;
+				// Next step is to parse the GroomCache data if it was determined that the groom has animated attributes
+				return GroomAsset != nullptr && AnimInfo.IsValid();
+			});
+		// Parse GroomCache data into processor (Async)
+		Then(ESchemaTranslationLaunchPolicy::Async,
+			[this]() -> bool
+			{
+				UGroomAsset* GroomAsset = Cast<UGroomAsset>(Context->AssetCache->GetAssetForPrim(PrimPath.GetString()));
+				if (!GroomAsset)
+				{
+					return false;
+				}
+
+				// Compute GroomCache hash from the first and last frame HairDescription...
+				FSHA1 SHA1;
+				UE::UsdGroomTranslator::Private::ComputeFrameHairDescriptionHash(GetPrim(), ImportOptions->InterpolationSettings, AnimInfo.StartFrame, SHA1);
+				UE::UsdGroomTranslator::Private::ComputeFrameHairDescriptionHash(GetPrim(), ImportOptions->InterpolationSettings, AnimInfo.EndFrame, SHA1);
+
+				// Along with relevant AnimInfo data
+				SHA1.Update((const uint8*)&AnimInfo.NumFrames, sizeof(AnimInfo.NumFrames));
+				SHA1.Update((const uint8*)&AnimInfo.Attributes, sizeof(AnimInfo.Attributes));
+				SHA1.Final();
+
+				SHA1.GetHash(GroomCacheHash.Hash);
+
+				bool bSuccess = true;
+				UGroomCache* GroomCache = Cast<UGroomCache>(Context->AssetCache->GetCachedAsset(GroomCacheHash.ToString()));
+				if (!GroomCache)
+				{
+					GroomCacheProcessor = MakeUnique<FGroomCacheProcessor>(EGroomCacheType::Strands, AnimInfo.Attributes);
+
+					// ref. FGroomCacheImporter::ImportGroomCache
+					const TArray<FHairGroupData>& GroomHairGroupsData = GroomAsset->HairGroupsData;
+
+					// Each frame is translated into a HairDescription and processed into HairGroupData
+					// Sample one extra frame so that we can interpolate between EndFrame - 1 and EndFrame
+					for (int32 FrameIndex = AnimInfo.StartFrame; FrameIndex < AnimInfo.EndFrame + 1; ++FrameIndex)
+					{
+						FHairDescription FrameHairDescription;
+						bSuccess = UsdToUnreal::ConvertGroomHierarchy(GetPrim(), pxr::UsdTimeCode(FrameIndex), FTransform::Identity, FrameHairDescription);
+
+						if (!bSuccess)
+						{
+							break;
+						}
+
+						FHairDescriptionGroups HairDescriptionGroups;
+						if (!FGroomBuilder::BuildHairDescriptionGroups(FrameHairDescription, HairDescriptionGroups))
+						{
+							bSuccess = false;
+							break;
+						}
+
+						const uint32 GroupCount = HairDescriptionGroups.HairGroups.Num();
+
+						TArray<FHairGroupInfoWithVisibility> HairGroupsInfo = GroomAsset->HairGroupsInfo;
+						TArray<FHairDescriptionGroup> HairGroupsData;
+						HairGroupsData.SetNum(GroupCount);
+						for (uint32 GroupIndex = 0; GroupIndex < GroupCount; ++GroupIndex)
+						{
+							const FHairDescriptionGroup& HairGroup = HairDescriptionGroups.HairGroups[GroupIndex];
+							FHairDescriptionGroup& HairGroupData = HairGroupsData[GroupIndex];
+							FGroomBuilder::BuildData(HairGroup, GroomAsset->HairGroupsInterpolation[GroupIndex], HairGroupsInfo[GroupIndex], HairGroupData.Strands, HairGroupData.Guides);
+						}
+
+						// Validate that the GroomCache has the same topology as the static groom
+						if (HairGroupsData.Num() == GroomHairGroupsData.Num())
+						{
+							for (uint32 GroupIndex = 0; GroupIndex < GroupCount; ++GroupIndex)
+							{
+								if (HairGroupsData[GroupIndex].Strands.GetNumPoints() != GroomHairGroupsData[GroupIndex].Strands.BulkData.GetNumPoints())
+								{
+									bSuccess = false;
+									UE_LOG(LogUsd, Warning, TEXT("GroomCache frame %d does not have the same number of vertices as the static groom (%u instead of %u). Aborting GroomCache import."),
+										FrameIndex, HairGroupsData[GroupIndex].Strands.GetNumPoints(), GroomHairGroupsData[GroupIndex].Strands.BulkData.GetNumPoints());
+									break;
+								}
+							}
+						}
+						else
+						{
+							bSuccess = false;
+							UE_LOG(LogUsd, Warning, TEXT("GroomCache does not have the same number of groups as the static groom (%d instead of %d). Aborting GroomCache import."),
+								HairGroupsData.Num(), GroomHairGroupsData.Num());
+						}
+
+						if (!bSuccess)
+						{
+							break;
+						}
+
+						// The HairGroupData is converted into animated groom data by the GroomCacheProcessor
+						GroomCacheProcessor->AddGroomSample(MoveTemp(HairGroupsData));
+					}
+				}
+				else
+				{
+					Context->AssetCache->LinkAssetToPrim(UsdGroomTranslatorUtils::GetStrandsGroomCachePrimPath(PrimPath), GroomCache);
+				}
+
+				if (!bSuccess)
+				{
+					GroomCacheProcessor.Reset();
+				}
+
+				// Go to the next step only if it needs to create a GroomCache asset
+				return GroomCache == nullptr && GroomCacheProcessor;
+			});
+		// Create GroomCache asset from processor (Sync)
+		Then(ESchemaTranslationLaunchPolicy::Sync,
+			[this]() -> bool
+			{
+				FHairImportContext HairImportContext(nullptr, GetTransientPackage(), nullptr, FName(), Context->ObjectFlags | EObjectFlags::RF_Public);
+				FName UniqueName = MakeUniqueObjectName(GetTransientPackage(), UGroomCache::StaticClass(), *FPaths::GetBaseFilename(UsdGroomTranslatorUtils::GetStrandsGroomCachePrimPath(PrimPath)));
+
+				// Once the processing has completed successfully, the data is transferred to the GroomCache
+				UGroomCache* GroomCache = FGroomCacheImporter::ProcessToGroomCache(*GroomCacheProcessor, AnimInfo, HairImportContext, UniqueName.ToString());
+				if (GroomCache)
+				{
+#if WITH_EDITOR
+					UUsdAssetImportData* ImportData = NewObject<UUsdAssetImportData>(GroomCache, TEXT("UUSDAssetImportData"));
+					ImportData->PrimPath = UsdGroomTranslatorUtils::GetStrandsGroomCachePrimPath(PrimPath);
+					ImportData->ImportOptions = ImportOptions.Get();
+
+					GroomCache->AssetImportData = ImportData;
+#endif // WITH_EDITOR
+
+					Context->AssetCache->CacheAsset(GroomCacheHash.ToString(), GroomCache);
+					Context->AssetCache->LinkAssetToPrim(UsdGroomTranslatorUtils::GetStrandsGroomCachePrimPath(PrimPath), GroomCache);
+				}
+
+				return GroomCache != nullptr;
 			});
 	}
 };
@@ -223,6 +392,15 @@ void FUsdGroomTranslator::UpdateComponents(USceneComponent* SceneComponent)
 			}
 
 			GroomComponent->SetGroomAsset(Groom);
+
+			if (Groom)
+			{
+				UGroomCache* GroomCache = Cast<UGroomCache>(Context->AssetCache->GetAssetForPrim(UsdGroomTranslatorUtils::GetStrandsGroomCachePrimPath(PrimPath)));
+				if (GroomCache != GroomComponent->GroomCache.Get())
+				{
+					GroomComponent->SetGroomCache(GroomCache);
+				}
+			}
 		}
 
 		// Use the prim purpose in conjunction with the prim's computed visibility to toggle the visibility of the groom component
diff --git a/Engine/Plugins/Importers/USDImporter/Source/USDSchemas/Private/USDGroomTranslatorUtils.cpp b/Engine/Plugins/Importers/USDImporter/Source/USDSchemas/Private/USDGroomTranslatorUtils.cpp
--- a/Engine/Plugins/Importers/USDImporter/Source/USDSchemas/Private/USDGroomTranslatorUtils.cpp
+++ b/Engine/Plugins/Importers/USDImporter/Source/USDSchemas/Private/USDGroomTranslatorUtils.cpp
@@ -18,6 +18,7 @@
 #include "USDIntegrationUtils.h"
 #include "USDSchemaTranslator.h"
 #include "USDTypesConversion.h"
+#include "UsdWrappers/SdfPath.h"
 
 #include "USDIncludesStart.h"
 #include "pxr/usd/usd/prim.h"
@@ -284,6 +285,11 @@ namespace UsdGroomTranslatorUtils
 			}
 		}
 	}
+
+	FString GetStrandsGroomCachePrimPath(const UE::FSdfPath& PrimPath)
+	{
+		return FString::Printf(TEXT("%s_strands_cache"), *PrimPath.GetString());
+	}
 }
 
 #endif // #if USE_USD_SDK && WITH_EDITOR
diff --git a/Engine/Plugins/Importers/USDImporter/Source/USDSchemas/Private/USDGroomTranslatorUtils.h b/Engine/Plugins/Importers/USDImporter/Source/USDSchemas/Private/USDGroomTranslatorUtils.h
--- a/Engine/Plugins/Importers/USDImporter/Source/USDSchemas/Private/USDGroomTranslatorUtils.h
+++ b/Engine/Plugins/Importers/USDImporter/Source/USDSchemas/Private/USDGroomTranslatorUtils.h
@@ -18,6 +18,11 @@ PXR_NAMESPACE_CLOSE_SCOPE
 class USceneComponent;
 class UUsdAssetCache;
 
+namespace UE
+{
+	class FSdfPath;
+}
+
 namespace UsdGroomTranslatorUtils
 {
 	// Create the groom binding asset for the given Prim with GroomBindingAPI and cache it in the AssetCache
@@ -25,6 +30,9 @@ namespace UsdGroomTranslatorUtils
 
 	// Set the groom asset targeted by the given prim with GroomBindingAPI on the SceneComponent (with a GroomComponent as a direct child)
 	void SetGroomFromPrim(const pxr::UsdPrim& Prim, const UUsdAssetCache& AssetCache, USceneComponent* SceneComponent);
+
+	// Get the (pseudo) prim path for the strands GroomCache associated with a groom prim
+	FString GetStrandsGroomCachePrimPath(const UE::FSdfPath& PrimPath);
 }
 
 #endif // #if USE_USD_SDK && WITH_EDITOR
\ No newline at end of file
diff --git a/Engine/Plugins/Importers/USDImporter/Source/USDStage/Private/USDLevelSequenceHelper.cpp b/Engine/Plugins/Importers/USDImporter/Source/USDStage/Private/USDLevelSequenceHelper.cpp
--- a/Engine/Plugins/Importers/USDImporter/Source/USDStage/Private/USDLevelSequenceHelper.cpp
+++ b/Engine/Plugins/Importers/USDImporter/Source/USDStage/Private/USDLevelSequenceHelper.cpp
@@ -40,12 +40,16 @@
 #include "Components/SpotLightComponent.h"
 #include "ControlRigObjectBinding.h"
 #include "CoreMinimal.h"
+#include "GroomCache.h"
+#include "GroomComponent.h"
 #include "LevelSequence.h"
 #include "LevelSequenceActor.h"
 #include "LevelSequencePlayer.h"
 #include "Misc/ITransaction.h"
 #include "Misc/TransactionObjectEvent.h"
 #include "MovieScene.h"
+#include "MovieSceneGroomCacheSection.h"
+#include "MovieSceneGroomCacheTrack.h"
 #include "MovieSceneTimeHelpers.h"
 #include "MovieSceneTrack.h"
 #include "Rigs/FKControlRig.h"
@@ -535,6 +539,7 @@ private:
 	void AddCameraTracks( const UUsdPrimTwin& PrimTwin, const UE::FUsdPrim& Prim );
 	void AddLightTracks( const UUsdPrimTwin& PrimTwin, const UE::FUsdPrim& Prim, const TSet<FName>& PropertyPathsToRead = {} );
 	void AddSkeletalTracks( const UUsdPrimTwin& PrimTwin, const UE::FUsdPrim& Prim );
+	void AddGroomTracks( const UUsdPrimTwin& PrimTwin, const UE::FUsdPrim& Prim );
 
 	template<typename TrackType>
 	TrackType* AddTrack( const FName& PropertyPath, const UUsdPrimTwin& PrimTwin, USceneComponent& ComponentToBind, ULevelSequence& Sequence, bool bIsMuted = false  );
@@ -1863,6 +1868,67 @@ void FUsdLevelSequenceHelperImpl::AddSkeletalTracks( const UUsdPrimTwin& PrimTwi
 	PrimPathByLevelSequenceName.AddUnique( SkelAnimationSequence->GetFName(), PrimPath );
 }
 
+void FUsdLevelSequenceHelperImpl::AddGroomTracks( const UUsdPrimTwin& PrimTwin, const UE::FUsdPrim& Prim )
+{
+	UGroomComponent* ComponentToBind = Cast< UGroomComponent >( PrimTwin.GetSceneComponent() );
+	if ( !ComponentToBind )
+	{
+		return;
+	}
+
+	if ( !AssetCache )
+	{
+		return;
+	}
+
+	// Fetch the groom cache asset from the asset cache. If there's none, don't actually need to create track
+	const FString PrimPath = Prim.GetPrimPath().GetString();
+	const FString GroomCachePath = FString::Printf( TEXT( "%s_strands_cache" ), *PrimPath );
+	UGroomCache* GroomCache = Cast< UGroomCache >( AssetCache->GetAssetForPrim( GroomCachePath ) );
+	if ( !GroomCache )
+	{
+		return;
+	}
+
+	if ( ComponentToBind->GroomCache.Get() != GroomCache )
+	{
+		return;
+	}
+
+	UE::FSdfLayer GroomLayer = UsdUtils::FindLayerForPrim( Prim );
+	if ( !GroomLayer )
+	{
+		return;
+	}
+
+	ULevelSequence* GroomAnimationSequence = FindOrAddSequenceForLayer( GroomLayer, GroomLayer.GetIdentifier(), GroomLayer.GetDisplayName() );
+	if ( !GroomAnimationSequence )
+	{
+		return;
+	}
+
+	UMovieScene* MovieScene = GroomAnimationSequence->GetMovieScene();
+	if ( !MovieScene )
+	{
+		return;
+	}
+
+	const bool bReadonly = false;
+	UsdLevelSequenceHelperImpl::FMovieSceneReadonlyGuard MovieSceneReadonlyGuard{ *MovieScene, bReadonly };
+
+	const bool bIsMuted = false;
+	if ( UMovieSceneGroomCacheTrack* GroomCacheTrack = AddTrack< UMovieSceneGroomCacheTrack >( Prim.GetName(), PrimTwin, *ComponentToBind, *GroomAnimationSequence, bIsMuted ) )
+	{
+		GroomCacheTrack->RemoveAllAnimationData();
+
+		const FFrameNumber StartOffset;
+		UMovieSceneGroomCacheSection* NewSection = Cast< UMovieSceneGroomCacheSection >( GroomCacheTrack->AddNewAnimation( StartOffset, ComponentToBind ) );
+		NewSection->EvalOptions.CompletionMode = EMovieSceneCompletionMode::KeepState;
+	}
+
+	PrimPathByLevelSequenceName.AddUnique( GroomAnimationSequence->GetFName(), PrimPath );
+}
+
 void FUsdLevelSequenceHelperImpl::AddPrim( UUsdPrimTwin& PrimTwin, bool bForceVisibilityTracks )
 {
 	if ( !UsdStage )
@@ -1910,6 +1976,10 @@ void FUsdLevelSequenceHelperImpl::AddPrim( UUsdPrimTwin& PrimTwin, bool bForceVi
 	{
 		AddSkeletalTracks( PrimTwin, UsdPrim );
 	}
+	else if ( UsdUtils::PrimHasSchema( UsdPrim, UnrealIdentifiers::GroomAPI ) )
+	{
+		AddGroomTracks( PrimTwin, UsdPrim );
+	}
 
 	AddCommonTracks( PrimTwin, UsdPrim, bForceVisibilityTracks );
 
diff --git a/Engine/Plugins/Importers/USDImporter/Source/USDStage/USDStage.Build.cs b/Engine/Plugins/Importers/USDImporter/Source/USDStage/USDStage.Build.cs
--- a/Engine/Plugins/Importers/USDImporter/Source/USDStage/USDStage.Build.cs
+++ b/Engine/Plugins/Importers/USDImporter/Source/USDStage/USDStage.Build.cs
@@ -25,12 +25,14 @@ namespace UnrealBuildTool.Rules
 					"Core",
 					"CoreUObject",
 					"Engine",
+					"HairStrandsCore",
 					"LevelSequence",
 					"LiveLinkComponents", // For tracking edits to LiveLinkComponentController properties and writing to USD
 					"LiveLinkInterface",
 					"MeshDescription",
 					"MovieScene",
 					"MovieSceneTracks",
+					"Niagara",	// Needed by GroomComponent.h
 					"Projects", // So that we can use the plugin manager to find out our content dir and cook the base materials
 					"Slate",
 					"SlateCore",
diff --git a/Engine/Plugins/Importers/USDImporter/Source/USDUtilities/Private/USDConversionUtils.cpp b/Engine/Plugins/Importers/USDImporter/Source/USDUtilities/Private/USDConversionUtils.cpp
--- a/Engine/Plugins/Importers/USDImporter/Source/USDUtilities/Private/USDConversionUtils.cpp
+++ b/Engine/Plugins/Importers/USDImporter/Source/USDUtilities/Private/USDConversionUtils.cpp
@@ -40,6 +40,7 @@
 #include "Framework/Notifications/NotificationManager.h"
 #include "GeometryCache.h"
 #include "GroomAsset.h"
+#include "GroomCache.h"
 #include "InstancedFoliageActor.h"
 #include "LandscapeProxy.h"
 #include "Widgets/Notifications/SNotificationList.h"
@@ -1188,6 +1189,10 @@ UUsdAssetImportData* UsdUtils::GetAssetImportData( UObject* Asset )
 	{
 		ImportData = Cast<UUsdAssetImportData>( Groom->AssetImportData );
 	}
+	else if ( UGroomCache* GroomCache = Cast<UGroomCache>( Asset ) )
+	{
+		ImportData = Cast<UUsdAssetImportData>( GroomCache->AssetImportData );
+	}
 
 #endif
 	return ImportData;
diff --git a/Engine/Plugins/Importers/USDImporter/Source/USDUtilities/Private/USDGroomConversion.cpp b/Engine/Plugins/Importers/USDImporter/Source/USDUtilities/Private/USDGroomConversion.cpp
--- a/Engine/Plugins/Importers/USDImporter/Source/USDUtilities/Private/USDGroomConversion.cpp
+++ b/Engine/Plugins/Importers/USDImporter/Source/USDUtilities/Private/USDGroomConversion.cpp
@@ -4,6 +4,7 @@
 
 #include "USDGroomConversion.h"
 
+#include "GroomCacheData.h"
 #include "HairDescription.h"
 #include "USDPrimConversion.h"
 #include "USDTypesConversion.h"
@@ -188,11 +189,30 @@ namespace UE::UsdGroomConversion::Private
 			}
 		}
 	}
+
+	/** Fill the AnimInfo with the time data for the given attribute, if it's animated */
+	void UpdateAnimInfo(const pxr::UsdAttribute& Attr, EGroomCacheAttributes Flag, FGroomAnimationInfo* AnimInfo)
+	{
+		if (AnimInfo)
+		{
+			std::vector<double> TimeSamples;
+			if (Attr.GetTimeSamples(&TimeSamples))
+			{
+				if (TimeSamples.size() > 1)
+				{
+					AnimInfo->Attributes = AnimInfo->Attributes | Flag;
+					AnimInfo->StartFrame = FMath::Min(AnimInfo->StartFrame, TimeSamples[0]);
+					AnimInfo->EndFrame = FMath::Max(AnimInfo->EndFrame, TimeSamples[TimeSamples.size() - 1]);
+					AnimInfo->NumFrames = FMath::Max(AnimInfo->NumFrames, uint32(AnimInfo->EndFrame - AnimInfo->StartFrame + 1));
+				}
+			}
+		}
+	}
 }
 
 namespace UsdToUnreal
 {
-	bool ConvertGroomHierarchy(const pxr::UsdPrim& Prim, const pxr::UsdTimeCode& TimeCode, const FTransform& ParentTransform, FHairDescription& HairDescription)
+	bool ConvertGroomHierarchy(const pxr::UsdPrim& Prim, const pxr::UsdTimeCode& TimeCode, const FTransform& ParentTransform, FHairDescription& HairDescription, FGroomAnimationInfo* AnimInfo)
 	{
 		TRACE_CPUPROFILER_EVENT_SCOPE(UsdToUnreal::ConvertGroomHierarchy);
 
@@ -228,6 +248,8 @@ namespace UsdToUnreal
 				return false;
 			}
 
+			UE::UsdGroomConversion::Private::UpdateAnimInfo(PointsAttr, EGroomCacheAttributes::Position, AnimInfo);
+
 			// Get the starting strand and vertex IDs for this group of Curves
 			int32 StartStrandID = HairDescription.GetNumStrands();
 			int32 StartVertexID = HairDescription.GetNumVertices();
@@ -268,6 +290,8 @@ namespace UsdToUnreal
 			pxr::VtArray<float> WidthsArray;
 			if (Widths && Widths.Get(&WidthsArray, TimeCode) && !WidthsArray.empty())
 			{
+				UE::UsdGroomConversion::Private::UpdateAnimInfo(Widths, EGroomCacheAttributes::Width, AnimInfo);
+
 				// Determine the scope of the Widths attribute
 				pxr::TfToken Interpolation = Curves.GetWidthsInterpolation();
 
@@ -333,7 +357,7 @@ namespace UsdToUnreal
 
 			for (const pxr::UsdPrim& Child : Prim.GetChildren())
 			{
-				bSuccess &= ConvertGroomHierarchy(Child, TimeCode, TransformToPropagate, HairDescription);
+				bSuccess &= ConvertGroomHierarchy(Child, TimeCode, TransformToPropagate, HairDescription, AnimInfo);
 				if (!bSuccess)
 				{
 					break;
diff --git a/Engine/Plugins/Importers/USDImporter/Source/USDUtilities/Public/USDGroomConversion.h b/Engine/Plugins/Importers/USDImporter/Source/USDUtilities/Public/USDGroomConversion.h
--- a/Engine/Plugins/Importers/USDImporter/Source/USDUtilities/Public/USDGroomConversion.h
+++ b/Engine/Plugins/Importers/USDImporter/Source/USDUtilities/Public/USDGroomConversion.h
@@ -16,6 +16,7 @@ class UsdTimeCode;
 PXR_NAMESPACE_CLOSE_SCOPE
 
 class FHairDescription;
+struct FGroomAnimationInfo;
 
 namespace UsdToUnreal
 {
@@ -26,9 +27,10 @@ namespace UsdToUnreal
 	 * @param TimeCode - The timecode at which to query the prim attributes
 	 * @param ParentTransform - The prim transform to propagate to the children
 	 * @param HairDescription - The hair description where the groom prim data is outputted
+	 * @param AnimInfo -  Optional GroomAnimationInfo to fill when parsing the hierarchy
 	 * @return true if the conversion was successful; false otherwise
 	 */
-	USDUTILITIES_API bool ConvertGroomHierarchy(const pxr::UsdPrim& Prim, const pxr::UsdTimeCode& TimeCode, const FTransform& ParentTransform, FHairDescription& HairDescription);
+	USDUTILITIES_API bool ConvertGroomHierarchy(const pxr::UsdPrim& Prim, const pxr::UsdTimeCode& TimeCode, const FTransform& ParentTransform, FHairDescription& HairDescription, FGroomAnimationInfo* AnimInfo = nullptr);
 }
 
 #endif // #if USE_USD_SDK
\ No newline at end of file
diff --git a/Engine/Plugins/Runtime/HairStrands/Source/HairStrandsEditor/Private/GroomCacheImporter.cpp b/Engine/Plugins/Runtime/HairStrands/Source/HairStrandsEditor/Private/GroomCacheImporter.cpp
--- a/Engine/Plugins/Runtime/HairStrands/Source/HairStrandsEditor/Private/GroomCacheImporter.cpp
+++ b/Engine/Plugins/Runtime/HairStrands/Source/HairStrandsEditor/Private/GroomCacheImporter.cpp
@@ -19,61 +19,71 @@ DEFINE_LOG_CATEGORY_STATIC(LogGroomCacheImporter, Log, All);
 
 static UGroomCache* CreateGroomCache(EGroomCacheType Type, UObject*& InParent, const FString& ObjectName, const EObjectFlags Flags)
 {
-	// Setup package name and create one accordingly
-	FString NewPackageName = InParent->GetOutermost()->GetName();
-	if (!NewPackageName.EndsWith(ObjectName))
+	// Don't do any package handling if the parent is the transient package; just use it
+	UPackage* Package = nullptr;
+	FString SanitizedObjectName = ObjectName;
+	if (InParent != GetTransientPackage())
 	{
-		// Remove the cache suffix if the parent is from a reimport
-		if (NewPackageName.EndsWith("_strands_cache"))
+		// Setup package name and create one accordingly
+		FString NewPackageName = InParent->GetOutermost()->GetName();
+		if (!NewPackageName.EndsWith(ObjectName))
 		{
-			NewPackageName.RemoveFromEnd("_strands_cache");
-		}
-		else if (NewPackageName.EndsWith("_guides_cache"))
-		{
-			NewPackageName.RemoveFromEnd("_guides_cache");
-		}
-
-		// Append the correct suffix
-		NewPackageName += TEXT("_") + ObjectName;
-	}
+			// Remove the cache suffix if the parent is from a reimport
+			if (NewPackageName.EndsWith("_strands_cache"))
+			{
+				NewPackageName.RemoveFromEnd("_strands_cache");
+			}
+			else if (NewPackageName.EndsWith("_guides_cache"))
+			{
+				NewPackageName.RemoveFromEnd("_guides_cache");
+			}
 
-	NewPackageName = UPackageTools::SanitizePackageName(NewPackageName);
+			// Append the correct suffix
+			NewPackageName += TEXT("_") + ObjectName;
+		}
 
-	// Parent package to place new GroomCache
-	UPackage* Package = CreatePackage(*NewPackageName);
+		NewPackageName = UPackageTools::SanitizePackageName(NewPackageName);
 
-	FString CompoundObjectName = FPackageName::GetShortName(NewPackageName);
-	const FString SanitizedObjectName = ObjectTools::SanitizeObjectName(CompoundObjectName);
+		// Parent package to place new GroomCache
+		Package = CreatePackage(*NewPackageName);
 
-	UGroomCache* ExistingTypedObject = FindObject<UGroomCache>(Package, *SanitizedObjectName);
-	UObject* ExistingObject = FindObject<UObject>(Package, *SanitizedObjectName);
+		FString CompoundObjectName = FPackageName::GetShortName(NewPackageName);
+		SanitizedObjectName = ObjectTools::SanitizeObjectName(CompoundObjectName);
 
-	if (ExistingTypedObject != nullptr)
-	{
-		ExistingTypedObject->PreEditChange(nullptr);
-		return ExistingTypedObject;
-	}
-	else if (ExistingObject != nullptr)
-	{
-		// Replacing an object.  Here we go!
-		// Delete the existing object
-		const bool bDeleteSucceeded = ObjectTools::DeleteSingleObject(ExistingObject);
+		UGroomCache* ExistingTypedObject = FindObject<UGroomCache>(Package, *SanitizedObjectName);
+		UObject* ExistingObject = FindObject<UObject>(Package, *SanitizedObjectName);
 
-		if (bDeleteSucceeded)
+		if (ExistingTypedObject != nullptr)
 		{
-			// Force GC so we can cleanly create a new asset (and not do an 'in place' replacement)
-			CollectGarbage(GARBAGE_COLLECTION_KEEPFLAGS);
-
-			// Create a package for each mesh
-			Package = CreatePackage(*NewPackageName);
-			InParent = Package;
+			ExistingTypedObject->PreEditChange(nullptr);
+			return ExistingTypedObject;
 		}
-		else
+		else if (ExistingObject != nullptr)
 		{
-			// failed to delete
-			return nullptr;
+			// Replacing an object.  Here we go!
+			// Delete the existing object
+			const bool bDeleteSucceeded = ObjectTools::DeleteSingleObject(ExistingObject);
+
+			if (bDeleteSucceeded)
+			{
+				// Force GC so we can cleanly create a new asset (and not do an 'in place' replacement)
+				CollectGarbage(GARBAGE_COLLECTION_KEEPFLAGS);
+
+				// Create a package for each mesh
+				Package = CreatePackage(*NewPackageName);
+				InParent = Package;
+			}
+			else
+			{
+				// failed to delete
+				return nullptr;
+			}
 		}
 	}
+	else
+	{
+		Package = GetTransientPackage();
+	}
 
 	UGroomCache* GroomCache = NewObject<UGroomCache>(Package, FName(*SanitizedObjectName), Flags | RF_Public);
 	GroomCache->Initialize(Type);
@@ -81,7 +91,7 @@ static UGroomCache* CreateGroomCache(EGroomCacheType Type, UObject*& InParent, c
 	return GroomCache;
 }
 
-UGroomCache* ProcessToGroomCache(FGroomCacheProcessor& Processor, const FGroomAnimationInfo& AnimInfo, FHairImportContext& ImportContext, const FString& ObjectNameSuffix)
+UGroomCache* FGroomCacheImporter::ProcessToGroomCache(FGroomCacheProcessor& Processor, const FGroomAnimationInfo& AnimInfo, FHairImportContext& ImportContext, const FString& ObjectNameSuffix)
 {
 	if (UGroomCache* GroomCache = CreateGroomCache(Processor.GetType(), ImportContext.Parent, ObjectNameSuffix, ImportContext.Flags))
 	{
diff --git a/Engine/Plugins/Runtime/HairStrands/Source/HairStrandsEditor/Public/GroomCacheImporter.h b/Engine/Plugins/Runtime/HairStrands/Source/HairStrandsEditor/Public/GroomCacheImporter.h
--- a/Engine/Plugins/Runtime/HairStrands/Source/HairStrandsEditor/Public/GroomCacheImporter.h
+++ b/Engine/Plugins/Runtime/HairStrands/Source/HairStrandsEditor/Public/GroomCacheImporter.h
@@ -6,6 +6,7 @@
 
 struct FGroomAnimationInfo;
 struct FHairImportContext;
+class FGroomCacheProcessor;
 class IGroomTranslator;
 class UGroomAsset;
 class UGroomCache;
@@ -15,4 +16,5 @@ struct HAIRSTRANDSEDITOR_API FGroomCacheImporter
 	static TArray<UGroomCache*> ImportGroomCache(const FString& SourceFilename, TSharedPtr<IGroomTranslator> Translator, const FGroomAnimationInfo& InAnimInfo, FHairImportContext& HairImportContext, UGroomAsset* GroomAssetForCache);
 	static void SetupImportSettings(struct FGroomCacheImportSettings& ImportSettings, const FGroomAnimationInfo& AnimInfo);
 	static void ApplyImportSettings(struct FGroomCacheImportSettings& ImportSettings, FGroomAnimationInfo& AnimInfo);
+	static UGroomCache* ProcessToGroomCache(FGroomCacheProcessor& Processor, const FGroomAnimationInfo& AnimInfo, FHairImportContext& ImportContext, const FString& ObjectNameSuffix);
 };
