USD: Stage editor implementation cleanup extracted from old isolated layer mode:

Preserve the yellow "warning button" but disable it completely for now;
Preserve the wrapper classes for SdfAttributeSpec and SdfPrimSpec (unused though);
Preserve being able to click on the edit layer buttons;
Cleanup/combine the CanAdd/Add/CanRemove/Remove/HasSchema functions;
Fix layer tree view items not retaining expansion states;
Add a viewmodel to integrations panel;
Cleanup/standardize the prim info widgets and the refreshing process.

diff --git a/Engine/Plugins/Importers/USDImporter/Source/USDSchemas/Private/USDGeomXformableTranslator.cpp b/Engine/Plugins/Importers/USDImporter/Source/USDSchemas/Private/USDGeomXformableTranslator.cpp
--- a/Engine/Plugins/Importers/USDImporter/Source/USDSchemas/Private/USDGeomXformableTranslator.cpp
+++ b/Engine/Plugins/Importers/USDImporter/Source/USDSchemas/Private/USDGeomXformableTranslator.cpp
@@ -583,7 +583,7 @@ void FUsdGeomXformableTranslator::UpdateComponents( USceneComponent* SceneCompon
 		// skeletal version of the LiveLink configuration, we only handle setting up LiveLink for simple transforms
 		if ( !SceneComponent->IsA<USkeletalMeshComponent>() )
 		{
-			if ( UsdUtils::PrimHasLiveLinkSchema( Prim ) )
+			if ( UsdUtils::PrimHasSchema( Prim, UnrealIdentifiers::LiveLinkAPI ) )
 			{
 				UE::UsdXformableTranslatorImpl::Private::SetUpSceneComponentForLiveLink( Context.Get(), SceneComponent, Prim );
 			}
@@ -767,7 +767,7 @@ bool FUsdGeomXformableTranslator::CanBeCollapsed( ECollapsingType CollapsingType
 		return false;
 	}
 
-	if ( UsdUtils::PrimHasLiveLinkSchema( UsdPrim ) )
+	if ( UsdUtils::PrimHasSchema( UsdPrim, UnrealIdentifiers::LiveLinkAPI ) )
 	{
 		return false;
 	}
diff --git a/Engine/Plugins/Importers/USDImporter/Source/USDSchemas/Private/USDSkelRootTranslator.cpp b/Engine/Plugins/Importers/USDImporter/Source/USDSchemas/Private/USDSkelRootTranslator.cpp
--- a/Engine/Plugins/Importers/USDImporter/Source/USDSchemas/Private/USDSkelRootTranslator.cpp
+++ b/Engine/Plugins/Importers/USDImporter/Source/USDSchemas/Private/USDSkelRootTranslator.cpp
@@ -1341,8 +1341,8 @@ void FUsdSkelRootTranslator::UpdateComponents( USceneComponent* SceneComponent )
 
 	UE::FUsdPrim Prim = GetPrim();
 
-	const bool bPrimHasLiveLinkSchema = UsdUtils::PrimHasLiveLinkSchema( Prim );
-	const bool bPrimHasControlRigSchema = UsdUtils::PrimHasControlRigSchema( Prim );
+	const bool bPrimHasLiveLinkSchema = UsdUtils::PrimHasSchema( Prim, UnrealIdentifiers::LiveLinkAPI );
+	const bool bPrimHasControlRigSchema = UsdUtils::PrimHasSchema( Prim, UnrealIdentifiers::ControlRigAPI );
 
 	bool bPrimHasLiveLinkEnabled = bPrimHasLiveLinkSchema;
 	if ( bPrimHasLiveLinkSchema )
diff --git a/Engine/Plugins/Importers/USDImporter/Source/USDStage/Private/USDStageActor.cpp b/Engine/Plugins/Importers/USDImporter/Source/USDStage/Private/USDStageActor.cpp
--- a/Engine/Plugins/Importers/USDImporter/Source/USDStage/Private/USDStageActor.cpp
+++ b/Engine/Plugins/Importers/USDImporter/Source/USDStage/Private/USDStageActor.cpp
@@ -1253,7 +1253,7 @@ UUsdPrimTwin* AUsdStageActor::ExpandPrim( const UE::FUsdPrim& Prim, FUsdSchemaTr
 	// if we have any as if will properly deactivate them like the usual "Bake to Control Rig" workflow.
 	if ( Prim.IsA( TEXT( "SkelRoot" ) ) )
 	{
-		if ( UsdUtils::PrimHasControlRigSchema( Prim ) )
+		if ( UsdUtils::PrimHasSchema( Prim, UnrealIdentifiers::ControlRigAPI ) )
 		{
 			LevelSequenceHelper.UpdateControlRigTracks( *UsdPrimTwin );
 
diff --git a/Engine/Plugins/Importers/USDImporter/Source/USDStageEditor/Private/SUSDLayersTreeView.cpp b/Engine/Plugins/Importers/USDImporter/Source/USDStageEditor/Private/SUSDLayersTreeView.cpp
--- a/Engine/Plugins/Importers/USDImporter/Source/USDStageEditor/Private/SUSDLayersTreeView.cpp
+++ b/Engine/Plugins/Importers/USDImporter/Source/USDStageEditor/Private/SUSDLayersTreeView.cpp
@@ -218,48 +218,113 @@ protected:
 class FUsdLayerEditColumn : public FUsdTreeViewColumn, public TSharedFromThis< FUsdLayerEditColumn >
 {
 public:
-	const FSlateBrush* GetCheckedImage( const FUsdLayerViewModelRef InTreeItem ) const
+	FReply OnClicked( const FUsdLayerViewModelRef TreeItem )
 	{
-		return InTreeItem->LayerModel->bIsEditTarget
-			? FUsdStageEditorStyle::Get()->GetBrush( "UsdStageEditor.CheckBoxImage" )
-			: nullptr;
+		if ( TreeItem->UsdStage )
+		{
+			TreeItem->UsdStage.SetEditTarget( TreeItem->GetLayer() );
+		}
+
+		return FReply::Handled();
+	}
+
+	FSlateColor GetForegroundColor( const FUsdLayerViewModelRef TreeItem, const TSharedPtr< ITableRow > TableRow, const TSharedPtr< SButton > Button ) const
+	{
+		if ( !TableRow.IsValid() || !Button.IsValid() )
+		{
+			return FSlateColor::UseForeground();
+		}
+
+		const bool bIsButtonHovered = Button->IsHovered();
+		const bool bIsLayerEditTarget = TreeItem->UsdStage ? TreeItem->UsdStage.GetEditTarget() == TreeItem->GetLayer() : false;
+
+		if ( bIsLayerEditTarget )
+		{
+			return FSlateColor::UseForeground();
+		}
+		else
+		{
+			if ( !bIsButtonHovered )
+			{
+				return FLinearColor::Transparent;
+			}
+			else
+			{
+				return FAppStyle::GetSlateColor( TEXT( "Colors.ForegroundHover" ) );
+			}
+		}
 	}
 
 	virtual TSharedRef< SWidget > GenerateWidget( const TSharedPtr< IUsdTreeViewItem > InTreeItem, const TSharedPtr< ITableRow > TableRow ) override
 	{
 		const FUsdLayerViewModelRef TreeItem = StaticCastSharedRef< FUsdLayerViewModel >( InTreeItem.ToSharedRef() );
 
-		TSharedRef< SWidget > Item =
-			SNew(SImage)
-				.Image( this, &FUsdLayerEditColumn::GetCheckedImage, TreeItem );
-
 		float ItemSize = FUsdStageEditorStyle::Get()->GetFloat( "UsdStageEditor.ListItemHeight" );
 
+		TSharedPtr<SButton> Button = SNew( SButton )
+			.ContentPadding( 0 )
+			.ButtonStyle( FUsdStageEditorStyle::Get(), TEXT( "NoBorder" ) )
+			.OnClicked( this, &FUsdLayerEditColumn::OnClicked, TreeItem )
+			.ToolTip( SNew( SToolTip ).Text( LOCTEXT( "EditLayerButtonToolTip", "Edit layer" ) ) )
+			.HAlign( HAlign_Center )
+			.VAlign( VAlign_Center );
+
+		TSharedPtr<SImage> Image = SNew( SImage )
+			.Image( FUsdStageEditorStyle::Get()->GetBrush( "UsdStageEditor.CheckBoxImage" ) )
+			.ColorAndOpacity( this, &FUsdLayerEditColumn::GetForegroundColor, TreeItem, TableRow, Button );
+
+		Button->SetContent( Image.ToSharedRef() );
+
 		return SNew( SBox )
 			.HeightOverride( ItemSize )
 			.WidthOverride( ItemSize )
 			.HAlign( HAlign_Center )
 			.VAlign( VAlign_Center )
 			[
-				Item
+				Button.ToSharedRef()
 			];
 	}
 };
 
-void SUsdLayersTreeView::Construct( const FArguments& InArgs, AUsdStageActor* UsdStageActor )
+void SUsdLayersTreeView::Construct( const FArguments& InArgs )
 {
 	SUsdTreeView::Construct( SUsdTreeView::FArguments() );
 
 	OnContextMenuOpening = FOnContextMenuOpening::CreateSP( this, &SUsdLayersTreeView::ConstructLayerContextMenu );
 
-	BuildUsdLayersEntries( UsdStageActor );
+	OnExpansionChanged = FOnExpansionChanged::CreateLambda( [this]( const FUsdLayerViewModelRef& LayerViewModel, bool bIsExpanded )
+	{
+		if ( !LayerViewModel->IsValid() )
+		{
+			return;
+		}
+
+		TreeItemExpansionStates.Add( LayerViewModel->LayerIdentifier, bIsExpanded );
+	} );
 }
 
-void SUsdLayersTreeView::Refresh( AUsdStageActor* UsdStageActor, bool bResync )
+void SUsdLayersTreeView::Refresh( const UE::FUsdStageWeak& NewStage, bool bResync )
 {
 	if ( bResync )
 	{
-		BuildUsdLayersEntries( UsdStageActor );
+		bool bShouldResetExpansionStates = false;
+		if ( RootItems.Num() > 0 )
+		{
+			const FUsdLayerViewModelRef& FirstRootItem = RootItems[ 0 ];
+			const UE::FUsdStageWeak& OldStage = FirstRootItem->UsdStage;
+			bShouldResetExpansionStates = NewStage != OldStage;
+		}
+
+		if ( bShouldResetExpansionStates )
+		{
+			TreeItemExpansionStates.Reset();
+		}
+
+		UsdStage = NewStage;
+
+		BuildUsdLayersEntries();
+
+		RestoreExpansionStates();
 	}
 	else
 	{
@@ -836,17 +901,11 @@ void SUsdLayersTreeView::OnGetChildren( FUsdLayerViewModelRef InParent, TArray<
 	}
 }
 
-void SUsdLayersTreeView::BuildUsdLayersEntries( AUsdStageActor* UsdStageActor )
+void SUsdLayersTreeView::BuildUsdLayersEntries()
 {
 	RootItems.Empty();
 
-	if ( !UsdStageActor )
-	{
-		return;
-	}
-
-	// The cast here forces us to use the const version of GetUsdStage, that won't force-load the stage in case it isn't opened yet
-	if ( const UE::FUsdStage& UsdStage = const_cast< const AUsdStageActor* >( UsdStageActor )->GetUsdStage() )
+	if ( UsdStage )
 	{
 		RootItems.Add( MakeSharedUnreal< FUsdLayerViewModel >( nullptr, UsdStage, UsdStage.GetRootLayer().GetIdentifier() ) );
 		RootItems.Add( MakeSharedUnreal< FUsdLayerViewModel >( nullptr, UsdStage, UsdStage.GetSessionLayer().GetIdentifier() ) );
@@ -859,7 +918,8 @@ void SUsdLayersTreeView::SetupColumns()
 
 	SHeaderRow::FColumn::FArguments LayerMutedColumnArguments;
 	LayerMutedColumnArguments.FixedWidth( 24.f );
-
+	LayerMutedColumnArguments.HAlignHeader( HAlign_Center );
+	LayerMutedColumnArguments.HAlignCell( HAlign_Center );
 	TSharedRef< FUsdLayerMutedColumn > LayerMutedColumn = MakeShared< FUsdLayerMutedColumn >();
 	AddColumn( TEXT("Mute"), FText(), LayerMutedColumn, LayerMutedColumnArguments );
 
@@ -1373,6 +1433,36 @@ void SUsdLayersTreeView::OnRemoveSelectedLayers()
 	}
 }
 
+void SUsdLayersTreeView::RestoreExpansionStates()
+{
+	TFunction< void( const FUsdLayerViewModelRef& ) > SetExpansionRecursive = [&]( const FUsdLayerViewModelRef& Item )
+	{
+		if ( Item->IsValid() )
+		{
+			if ( bool* bFoundExpansionState = TreeItemExpansionStates.Find( Item->LayerIdentifier ) )
+			{
+				SetItemExpansion( Item, *bFoundExpansionState );
+			}
+			// Default to showing everything expanded
+			else
+			{
+				const bool bShouldExpand = true;
+				SetItemExpansion( Item, bShouldExpand );
+			}
+		}
+
+		for ( const FUsdLayerViewModelRef& Child : Item->Children )
+		{
+			SetExpansionRecursive( Child );
+		}
+	};
+
+	for ( const FUsdLayerViewModelRef& RootItem : RootItems )
+	{
+		SetExpansionRecursive( RootItem );
+	}
+}
+
 #undef LOCTEXT_NAMESPACE
 
 #endif // #if USE_USD_SDK
diff --git a/Engine/Plugins/Importers/USDImporter/Source/USDStageEditor/Private/SUSDLayersTreeView.h b/Engine/Plugins/Importers/USDImporter/Source/USDStageEditor/Private/SUSDLayersTreeView.h
--- a/Engine/Plugins/Importers/USDImporter/Source/USDStageEditor/Private/SUSDLayersTreeView.h
+++ b/Engine/Plugins/Importers/USDImporter/Source/USDStageEditor/Private/SUSDLayersTreeView.h
@@ -2,6 +2,7 @@
 
 #pragma once
 
+#include "UsdWrappers/UsdStage.h"
 #include "Widgets/SUSDTreeView.h"
 
 #include "Widgets/Views/SHeaderRow.h"
@@ -17,8 +18,8 @@ public:
 	SLATE_BEGIN_ARGS( SUsdLayersTreeView ) {}
 	SLATE_END_ARGS()
 
-	void Construct( const FArguments& InArgs, AUsdStageActor* UsdStageActor );
-	void Refresh( AUsdStageActor* UsdStageActor, bool bResync );
+	void Construct( const FArguments& InArgs );
+	void Refresh( const UE::FUsdStageWeak& NewStage, bool bResync = false );
 
 	// Drag and drop interface for our rows
 	FReply OnRowDragDetected( const FGeometry& Geometry, const FPointerEvent& PointerEvent );
@@ -27,13 +28,15 @@ public:
 	FReply OnRowAcceptDrop( const FDragDropEvent& Event, EItemDropZone Zone, FUsdLayerViewModelRef Item );
 	// End drag and drop interface
 
+	const UE::FUsdStageWeak& GetStage() const { return UsdStage; }
+
 private:
 	virtual TSharedRef< ITableRow > OnGenerateRow( FUsdLayerViewModelRef InDisplayNode, const TSharedRef< STableViewBase >& OwnerTable ) override;
 	virtual void OnGetChildren( FUsdLayerViewModelRef InParent, TArray< FUsdLayerViewModelRef >& OutChildren ) const override;
 
 	virtual void SetupColumns() override;
 
-	void BuildUsdLayersEntries( AUsdStageActor* UsdStageActor );
+	void BuildUsdLayersEntries();
 
 	TSharedPtr< SWidget > ConstructLayerContextMenu();
 
@@ -55,4 +58,13 @@ private:
 	bool CanRemoveLayer( FUsdLayerViewModelRef LayerItem ) const;
 	bool CanRemoveSelectedLayers() const;
 	void OnRemoveSelectedLayers();
+
+	void RestoreExpansionStates();
+
+private:
+	// Should always be valid, we keep the one we're given on Refresh()
+	UE::FUsdStageWeak UsdStage;
+
+	// So that we can store these across refreshes
+	TMap< FString, bool > TreeItemExpansionStates;
 };
diff --git a/Engine/Plugins/Importers/USDImporter/Source/USDStageEditor/Private/SUSDPrimInfo.cpp b/Engine/Plugins/Importers/USDImporter/Source/USDStageEditor/Private/SUSDPrimInfo.cpp
--- a/Engine/Plugins/Importers/USDImporter/Source/USDStageEditor/Private/SUSDPrimInfo.cpp
+++ b/Engine/Plugins/Importers/USDImporter/Source/USDStageEditor/Private/SUSDPrimInfo.cpp
@@ -31,7 +31,7 @@
 
 #define LOCTEXT_NAMESPACE "SUSDPrimInfo"
 
-void SUsdPrimInfo::Construct( const FArguments& InArgs, const UE::FUsdStageWeak& UsdStage, const TCHAR* PrimPath )
+void SUsdPrimInfo::Construct( const FArguments& InArgs )
 {
 	ChildSlot
 	[
@@ -43,7 +43,7 @@ void SUsdPrimInfo::Construct( const FArguments& InArgs, const UE::FUsdStageWeak&
 			SNew( SBox )
 			.Content()
 			[
-				SAssignNew( PropertiesList, SUsdPrimPropertiesList, UsdStage, PrimPath )
+				SAssignNew( PropertiesList, SUsdPrimPropertiesList )
 			]
 		]
 
@@ -53,7 +53,7 @@ void SUsdPrimInfo::Construct( const FArguments& InArgs, const UE::FUsdStageWeak&
 			SNew( SBox )
 			.Content()
 			[
-				SAssignNew( IntegrationsPanel, SUsdIntegrationsPanel, UsdStage, PrimPath )
+				SAssignNew( IntegrationsPanel, SUsdIntegrationsPanel )
 			]
 		]
 
@@ -63,7 +63,7 @@ void SUsdPrimInfo::Construct( const FArguments& InArgs, const UE::FUsdStageWeak&
 			SNew( SBox )
 			.Content()
 			[
-				SAssignNew( VariantsList, SVariantsList, UsdStage, PrimPath )
+				SAssignNew( VariantsList, SVariantsList )
 			]
 		]
 
@@ -73,7 +73,7 @@ void SUsdPrimInfo::Construct( const FArguments& InArgs, const UE::FUsdStageWeak&
 			SNew( SBox )
 			.Content()
 			[
-				SAssignNew( ReferencesList, SUsdReferencesList, UsdStage, PrimPath )
+				SAssignNew( ReferencesList, SUsdReferencesList )
 			]
 		]
 	];
diff --git a/Engine/Plugins/Importers/USDImporter/Source/USDStageEditor/Private/SUSDPrimInfo.h b/Engine/Plugins/Importers/USDImporter/Source/USDStageEditor/Private/SUSDPrimInfo.h
--- a/Engine/Plugins/Importers/USDImporter/Source/USDStageEditor/Private/SUSDPrimInfo.h
+++ b/Engine/Plugins/Importers/USDImporter/Source/USDStageEditor/Private/SUSDPrimInfo.h
@@ -15,7 +15,7 @@ public:
 	SLATE_BEGIN_ARGS( SUsdPrimInfo ) {}
 	SLATE_END_ARGS()
 
-	void Construct( const FArguments& InArgs, const UE::FUsdStageWeak& UsdStage, const TCHAR* PrimPath );
+	void Construct( const FArguments& InArgs );
 	void SetPrimPath( const UE::FUsdStageWeak& UsdStage, const TCHAR* PrimPath );
 
 private:
diff --git a/Engine/Plugins/Importers/USDImporter/Source/USDStageEditor/Private/SUSDStage.cpp b/Engine/Plugins/Importers/USDImporter/Source/USDStageEditor/Private/SUSDStage.cpp
--- a/Engine/Plugins/Importers/USDImporter/Source/USDStageEditor/Private/SUSDStage.cpp
+++ b/Engine/Plugins/Importers/USDImporter/Source/USDStageEditor/Private/SUSDStage.cpp
@@ -40,6 +40,7 @@
 #include "Styling/AppStyle.h"
 #include "UObject/StrongObjectPtr.h"
 #include "Widgets/Input/SEditableTextBox.h"
+#include "Widgets/Input/SButton.h"
 #include "Widgets/Input/SSpinBox.h"
 #include "Widgets/Layout/SSplitter.h"
 
@@ -49,7 +50,7 @@
 
 namespace SUSDStageConstants
 {
-	static const FMargin SectionPadding( 1.f, 4.f, 1.f, 1.f );
+	static const FMargin SectionPadding( 1.f, 1.f, 1.f, 1.f );
 }
 
 namespace SUSDStageImpl
@@ -216,6 +217,18 @@ void SUsdStage::Construct( const FArguments& InArgs )
 					MakeMainMenu()
 				]
 
+				+ SHorizontalBox::Slot()
+				.AutoWidth()
+				[
+					SNew( SBox )
+					.HAlign( HAlign_Right )
+					.VAlign( VAlign_Fill )
+					.Padding( FMargin( 0.0f, 0.0f, 10.0f, 0.0f ) )
+					[
+						MakeFocusWarningButton()
+					]
+				]
+
 				+ SHorizontalBox::Slot()
 				.AutoWidth()
 				[
@@ -247,7 +260,7 @@ void SUsdStage::Construct( const FArguments& InArgs )
 						SNew( SBorder )
 						.BorderImage( FAppStyle::GetBrush(TEXT("ToolPanel.GroupBorder")) )
 						[
-							SAssignNew( UsdStageTreeView, SUsdStageTreeView, ViewModel.UsdStageActor.Get() )
+							SAssignNew( UsdStageTreeView, SUsdStageTreeView )
 							.OnPrimSelectionChanged( this, &SUsdStage::OnPrimSelectionChanged )
 						]
 					]
@@ -258,7 +271,7 @@ void SUsdStage::Construct( const FArguments& InArgs )
 						SNew( SBorder )
 						.BorderImage( FAppStyle::GetBrush("ToolPanel.GroupBorder") )
 						[
-							SAssignNew( UsdPrimInfoWidget, SUsdPrimInfo, UsdStage, TEXT("/") )
+							SAssignNew( UsdPrimInfoWidget, SUsdPrimInfo )
 						]
 					]
 				]
@@ -270,7 +283,7 @@ void SUsdStage::Construct( const FArguments& InArgs )
 					SNew(SBorder)
 					.BorderImage( FAppStyle::GetBrush(TEXT("ToolPanel.GroupBorder")) )
 					[
-						SAssignNew( UsdLayersTreeView, SUsdLayersTreeView, ViewModel.UsdStageActor.Get() )
+						SAssignNew( UsdLayersTreeView, SUsdLayersTreeView )
 					]
 				]
 			]
@@ -379,7 +392,10 @@ void SUsdStage::SetupStageActorDelegates()
 					if ( this->UsdLayersTreeView && ViewModel.UsdStageActor.IsValid() )
 					{
 						constexpr bool bResync = false;
-						this->UsdLayersTreeView->Refresh( ViewModel.UsdStageActor.Get(), bResync );
+						this->UsdLayersTreeView->Refresh(
+							ViewModel.UsdStageActor.Get()->GetOrLoadUsdStage(),
+							bResync
+						);
 					}
 				});
 			}
@@ -393,7 +409,10 @@ void SUsdStage::SetupStageActorDelegates()
 					if ( this->UsdLayersTreeView && ViewModel.UsdStageActor.IsValid() )
 					{
 						constexpr bool bResync = false;
-						this->UsdLayersTreeView->Refresh( ViewModel.UsdStageActor.Get(), bResync );
+						this->UsdLayersTreeView->Refresh(
+							ViewModel.UsdStageActor.Get()->GetOrLoadUsdStage(),
+							bResync
+						);
 					}
 				});
 			}
@@ -530,6 +549,34 @@ TSharedRef< SWidget > SUsdStage::MakeActorPickerMenuContent()
 	return SNullWidget::NullWidget;
 }
 
+TSharedRef< SWidget > SUsdStage::MakeFocusWarningButton()
+{
+	// TODO: Will be used for focus mode in the future
+	return SNew( SButton )
+		.ButtonStyle( &FAppStyle::Get().GetWidgetStyle<FButtonStyle>( "EditorViewportToolBar.WarningButton" ) )
+		.OnClicked_Lambda( [this]() -> FReply
+		{
+			return FReply::Handled();
+		})
+		.Visibility_Lambda( [this]() -> EVisibility
+		{
+			return EVisibility::Hidden;
+		})
+		.ToolTipText_Lambda( [this]() -> FText
+		{
+			return FText::Format(
+				LOCTEXT( "FocusWarningButtonToolTip", "The USD Stage editor and the Unreal level are displaying a focused stage with '{0}' as its root layer.\nPress this button to revert to displaying the entire composed stage." ),
+				FText::GetEmpty()
+			);
+		})
+		.Content()
+		[
+			SNew( STextBlock )
+			.TextStyle( &FAppStyle::Get().GetWidgetStyle<FTextBlockStyle>( "SmallText" ) )
+			.Text( LOCTEXT( "FocusWarningButtonText", "Focus Mode" ) )
+		];
+}
+
 void SUsdStage::FillFileMenu( FMenuBuilder& MenuBuilder )
 {
 	MenuBuilder.BeginSection( "File", LOCTEXT("File", "File") );
@@ -1748,7 +1795,8 @@ void SUsdStage::OnReloadStage()
 
 	if ( UsdLayersTreeView )
 	{
-		UsdLayersTreeView->Refresh( ViewModel.UsdStageActor.Get(), true );
+		const bool bResync = true;
+		UsdLayersTreeView->Refresh( ViewModel.UsdStageActor.Get()->GetOrLoadUsdStage(), bResync );
 	}
 }
 
@@ -1758,7 +1806,8 @@ void SUsdStage::OnResetStage()
 
 	if ( UsdLayersTreeView )
 	{
-		UsdLayersTreeView->Refresh( ViewModel.UsdStageActor.Get(), true );
+		const bool bResync = true;
+		UsdLayersTreeView->Refresh( ViewModel.UsdStageActor.Get()->GetOrLoadUsdStage(), bResync );
 	}
 }
 
@@ -1843,10 +1892,15 @@ void SUsdStage::SetActor( AUsdStageActor* InUsdStageActor )
 
 	SetupStageActorDelegates();
 
-	if ( this->UsdPrimInfoWidget && InUsdStageActor )
+	if( InUsdStageActor )
 	{
-		// Just reset to the pseudoroot for now
-		this->UsdPrimInfoWidget->SetPrimPath( ViewModel.UsdStageActor->GetOrLoadUsdStage(), TEXT( "/" ) );
+		UE::FUsdStage UsdStage = InUsdStageActor->GetOrLoadUsdStage();
+
+		if ( this->UsdPrimInfoWidget && InUsdStageActor )
+		{
+			// Just reset to the pseudoroot for now
+			this->UsdPrimInfoWidget->SetPrimPath( UsdStage, TEXT( "/" ) );
+		}
 	}
 
 	Refresh();
@@ -1859,12 +1913,13 @@ void SUsdStage::Refresh()
 
 	if (UsdLayersTreeView)
 	{
-		UsdLayersTreeView->Refresh( StageActor, true );
+		const bool bResync = true;
+		UsdLayersTreeView->Refresh( ViewModel.UsdStageActor.Get()->GetOrLoadUsdStage(), bResync );
 	}
 
 	if (UsdStageTreeView)
 	{
-		UsdStageTreeView->Refresh( StageActor );
+		UsdStageTreeView->Refresh( StageActor->GetOrLoadUsdStage() );
 
 		// Refresh will generate brand new RootItems, so let's immediately select the new item
 		// that corresponds to the prim we're supposed to be selecting
diff --git a/Engine/Plugins/Importers/USDImporter/Source/USDStageEditor/Private/SUSDStage.h b/Engine/Plugins/Importers/USDImporter/Source/USDStageEditor/Private/SUSDStage.h
--- a/Engine/Plugins/Importers/USDImporter/Source/USDStageEditor/Private/SUSDStage.h
+++ b/Engine/Plugins/Importers/USDImporter/Source/USDStageEditor/Private/SUSDStage.h
@@ -2,6 +2,8 @@
 
 #pragma once
 
+#include "USDStageViewModel.h"
+
 #include "CoreMinimal.h"
 #include "Layout/Visibility.h"
 #include "Input/Reply.h"
@@ -11,8 +13,6 @@
 #include "Widgets/DeclarativeSyntaxSupport.h"
 #include "Widgets/SCompoundWidget.h"
 
-#include "USDStageViewModel.h"
-
 class AUsdStageActor;
 class FLevelCollectionModel;
 class FMenuBuilder;
@@ -39,6 +39,7 @@ protected:
 	TSharedRef< SWidget > MakeMainMenu();
 	TSharedRef< SWidget > MakeActorPickerMenu();
 	TSharedRef< SWidget > MakeActorPickerMenuContent();
+	TSharedRef< SWidget > MakeFocusWarningButton();
 	void FillFileMenu( FMenuBuilder& MenuBuilder );
 	void FillActionsMenu( FMenuBuilder& MenuBuilder );
 	void FillOptionsMenu( FMenuBuilder& MenuBuilder );
@@ -74,8 +75,6 @@ protected:
 
 	void OnStageActorLoaded( AUsdStageActor* InUsdStageActor );
 
-	void OnMapChanged( UWorld* World, EMapChangeType ChangeType );
-
 	void OnViewportSelectionChanged( UObject* NewSelection );
 
 	int32 GetNaniteTriangleThresholdValue() const;
diff --git a/Engine/Plugins/Importers/USDImporter/Source/USDStageEditor/Private/SUSDStageTreeView.cpp b/Engine/Plugins/Importers/USDImporter/Source/USDStageEditor/Private/SUSDStageTreeView.cpp
--- a/Engine/Plugins/Importers/USDImporter/Source/USDStageEditor/Private/SUSDStageTreeView.cpp
+++ b/Engine/Plugins/Importers/USDImporter/Source/USDStageEditor/Private/SUSDStageTreeView.cpp
@@ -8,6 +8,7 @@
 #include "USDDuplicateType.h"
 #include "USDLayerUtils.h"
 #include "USDOptionsWindow.h"
+#include "USDPrimViewModel.h"
 #include "USDReferenceOptions.h"
 #include "USDStageActor.h"
 #include "USDStageModule.h"
@@ -184,6 +185,7 @@ public:
 
 				return EVisibility::Collapsed;
 			})
+			.ToolTipText( LOCTEXT( "TogglePayloadToolTip", "Toggle payload" ) )
 			.IsChecked( this, &FUsdStagePayloadColumn::IsChecked, StaticCastSharedPtr< FUsdPrimViewModel >( InTreeItem ) )
 			.OnCheckStateChanged( this, &FUsdStagePayloadColumn::OnCheckedPayload, StaticCastSharedPtr< FUsdPrimViewModel >( InTreeItem ) );
 
@@ -307,7 +309,7 @@ public:
 	}
 };
 
-void SUsdStageTreeView::Construct( const FArguments& InArgs, AUsdStageActor* InUsdStageActor )
+void SUsdStageTreeView::Construct( const FArguments& InArgs )
 {
 	SUsdTreeView::Construct( SUsdTreeView::FArguments() );
 
@@ -375,8 +377,6 @@ void SUsdStageTreeView::Construct( const FArguments& InArgs, AUsdStageActor* InU
 		FExecuteAction::CreateSP( this, &SUsdStageTreeView::OnRenamePrim ),
 		FCanExecuteAction::CreateLambda( [this]() { return SUsdStageTreeView::DoesPrimExistOnStage() && GetSelectedItems().Num() == 1; } )
 	);
-
-	Refresh( InUsdStageActor );
 }
 
 TSharedRef< ITableRow > SUsdStageTreeView::OnGenerateRow( FUsdPrimViewModelRef InDisplayNode, const TSharedRef<STableViewBase>& OwnerTable )
@@ -392,30 +392,24 @@ void SUsdStageTreeView::OnGetChildren( FUsdPrimViewModelRef InParent, TArray< FU
 	}
 }
 
-void SUsdStageTreeView::Refresh( AUsdStageActor* InUsdStageActor )
+void SUsdStageTreeView::Refresh( const UE::FUsdStageWeak& NewStage )
 {
 	UE::FUsdStageWeak OldStage = RootItems.Num() > 0 ? RootItems[0]->UsdStage : UE::FUsdStageWeak();
-	UE::FUsdStageWeak NewStage = InUsdStageActor
-		? static_cast< UE::FUsdStageWeak > ( static_cast< const AUsdStageActor* >( InUsdStageActor )->GetUsdStage() )
-		: UE::FUsdStageWeak();
 
 	RootItems.Empty();
-	if ( UsdStageActor.Get() != InUsdStageActor || NewStage != OldStage )
-	{
-		TreeItemExpansionStates.Reset();
-	}
 
-	UsdStageActor = InUsdStageActor;
-	if ( !UsdStageActor.IsValid() )
+	UsdStage = NewStage;
+
+	if ( OldStage != NewStage )
 	{
-		return;
+		TreeItemExpansionStates.Reset();
 	}
 
 	if ( NewStage )
 	{
 		if ( UE::FUsdPrim RootPrim = NewStage.GetPseudoRoot() )
 		{
-			RootItems.Add( MakeShared< FUsdPrimViewModel >( nullptr, NewStage, RootPrim ) );
+			RootItems.Add( MakeShared< FUsdPrimViewModel >( nullptr, UsdStage, RootPrim ) );
 		}
 
 		RestoreExpansionStates();
@@ -448,7 +442,7 @@ void SUsdStageTreeView::RefreshPrim( const FString& PrimPath, bool bResync )
 	// We couldn't find the target prim, do a full refresh instead
 	else
 	{
-		Refresh( UsdStageActor.Get() );
+		Refresh( UsdStage );
 	}
 
 	if ( bResync )
@@ -766,8 +760,16 @@ TSharedPtr< SWidget > SUsdStageTreeView::ConstructPrimContextMenu()
 			LOCTEXT( "SetUpLiveLink_ToolTip", "Sets up the generated component for Live Link and store the connection details to the USD Stage" ),
 			FSlateIcon(),
 			FUIAction(
-				FExecuteAction::CreateSP( this, &SUsdStageTreeView::OnSetUpLiveLink ),
-				FCanExecuteAction::CreateSP( this, &SUsdStageTreeView::CanSetUpLiveLink )
+				FExecuteAction::CreateSP(
+					this,
+					&SUsdStageTreeView::OnApplySchema,
+					FName{ *UsdToUnreal::ConvertToken( UnrealIdentifiers::LiveLinkAPI ) }
+				),
+				FCanExecuteAction::CreateSP(
+					this,
+					&SUsdStageTreeView::CanApplySchema,
+					FName{ *UsdToUnreal::ConvertToken( UnrealIdentifiers::LiveLinkAPI ) }
+				)
 			),
 			NAME_None,
 			EUserInterfaceActionType::Button
@@ -778,8 +780,16 @@ TSharedPtr< SWidget > SUsdStageTreeView::ConstructPrimContextMenu()
 			LOCTEXT( "RemoveLiveLink_ToolTip", "Reverses the Live Link configuration on the component and removes the connection details from the USD Stage" ),
 			FSlateIcon(),
 			FUIAction(
-				FExecuteAction::CreateSP( this, &SUsdStageTreeView::OnRemoveLiveLink ),
-				FCanExecuteAction::CreateSP( this, &SUsdStageTreeView::CanRemoveLiveLink )
+				FExecuteAction::CreateSP(
+					this,
+					&SUsdStageTreeView::OnRemoveSchema,
+					FName{ *UsdToUnreal::ConvertToken( UnrealIdentifiers::LiveLinkAPI ) }
+				),
+				FCanExecuteAction::CreateSP(
+					this,
+					&SUsdStageTreeView::CanRemoveSchema,
+					FName{ *UsdToUnreal::ConvertToken( UnrealIdentifiers::LiveLinkAPI ) }
+				)
 			),
 			NAME_None,
 			EUserInterfaceActionType::Button
@@ -790,8 +800,16 @@ TSharedPtr< SWidget > SUsdStageTreeView::ConstructPrimContextMenu()
 			LOCTEXT( "SetUpControlRig_ToolTip", "Sets up the generated component for Control Rig integration and store the connection details to the USD Stage" ),
 			FSlateIcon(),
 			FUIAction(
-				FExecuteAction::CreateSP( this, &SUsdStageTreeView::OnSetUpControlRig ),
-				FCanExecuteAction::CreateSP( this, &SUsdStageTreeView::CanSetUpControlRig )
+				FExecuteAction::CreateSP(
+					this,
+					&SUsdStageTreeView::OnApplySchema,
+					FName{ *UsdToUnreal::ConvertToken( UnrealIdentifiers::ControlRigAPI ) }
+				),
+				FCanExecuteAction::CreateSP(
+					this,
+					&SUsdStageTreeView::CanApplySchema,
+					FName{ *UsdToUnreal::ConvertToken( UnrealIdentifiers::ControlRigAPI ) }
+				)
 			),
 			NAME_None,
 			EUserInterfaceActionType::Button
@@ -802,8 +820,16 @@ TSharedPtr< SWidget > SUsdStageTreeView::ConstructPrimContextMenu()
 			LOCTEXT( "RemoveControlRig_ToolTip", "Reverses the Control Rig configuration on the component and removes the connection details from the USD Stage" ),
 			FSlateIcon(),
 			FUIAction(
-				FExecuteAction::CreateSP( this, &SUsdStageTreeView::OnRemoveControlRig ),
-				FCanExecuteAction::CreateSP( this, &SUsdStageTreeView::CanRemoveControlRig )
+				FExecuteAction::CreateSP(
+					this,
+					&SUsdStageTreeView::OnRemoveSchema,
+					FName{ *UsdToUnreal::ConvertToken( UnrealIdentifiers::ControlRigAPI ) }
+				),
+				FCanExecuteAction::CreateSP(
+					this,
+					&SUsdStageTreeView::CanRemoveSchema,
+					FName{ *UsdToUnreal::ConvertToken( UnrealIdentifiers::ControlRigAPI ) }
+				)
 			),
 			NAME_None,
 			EUserInterfaceActionType::Button
@@ -814,10 +840,16 @@ TSharedPtr< SWidget > SUsdStageTreeView::ConstructPrimContextMenu()
 			LOCTEXT( "ApplyGroomSchema_ToolTip", "Applies the Groom schema to interpret the prim and its children as a groom" ),
 			FSlateIcon(),
 			FUIAction(
-				FExecuteAction::CreateSP( this, &SUsdStageTreeView::OnApplyGroomSchema ),
-				FCanExecuteAction::CreateSP( this, &SUsdStageTreeView::CanApplyGroomSchema ),
-				FIsActionChecked(),
-				FIsActionButtonVisible::CreateSP( this, &SUsdStageTreeView::CanApplyGroomSchema )
+				FExecuteAction::CreateSP(
+					this,
+					&SUsdStageTreeView::OnApplySchema,
+					FName{ *UsdToUnreal::ConvertToken( UnrealIdentifiers::GroomAPI ) }
+				),
+				FCanExecuteAction::CreateSP(
+					this,
+					&SUsdStageTreeView::CanApplySchema,
+					FName{ *UsdToUnreal::ConvertToken( UnrealIdentifiers::GroomAPI ) }
+				)
 			),
 			NAME_None,
 			EUserInterfaceActionType::Button
@@ -828,10 +860,16 @@ TSharedPtr< SWidget > SUsdStageTreeView::ConstructPrimContextMenu()
 			LOCTEXT( "RemoveGroomSchema_ToolTip", "Removes the Groom schema from the prim to stop interpreting it as a groom" ),
 			FSlateIcon(),
 			FUIAction(
-				FExecuteAction::CreateSP( this, &SUsdStageTreeView::OnRemoveGroomSchema ),
-				FCanExecuteAction::CreateSP( this, &SUsdStageTreeView::CanRemoveGroomSchema ),
-				FIsActionChecked(),
-				FIsActionButtonVisible::CreateSP( this, &SUsdStageTreeView::CanRemoveGroomSchema )
+				FExecuteAction::CreateSP(
+					this,
+					&SUsdStageTreeView::OnRemoveSchema,
+					FName{ *UsdToUnreal::ConvertToken( UnrealIdentifiers::GroomAPI ) }
+				),
+				FCanExecuteAction::CreateSP(
+					this,
+					&SUsdStageTreeView::CanRemoveSchema,
+					FName{ *UsdToUnreal::ConvertToken( UnrealIdentifiers::GroomAPI ) }
+				)
 			),
 			NAME_None,
 			EUserInterfaceActionType::Button
@@ -847,11 +885,6 @@ TSharedPtr< SWidget > SUsdStageTreeView::ConstructPrimContextMenu()
 
 void SUsdStageTreeView::OnAddChildPrim()
 {
-	if ( !UsdStageActor.IsValid() )
-	{
-		return;
-	}
-
 	TArray< FUsdPrimViewModelRef > MySelectedItems = GetSelectedItems();
 
 	// Add a new child prim
@@ -869,7 +902,7 @@ void SUsdStageTreeView::OnAddChildPrim()
 	// Add a new top-level prim (direct child of the pseudo-root prim)
 	else
 	{
-		FUsdPrimViewModelRef TreeItem = MakeShared< FUsdPrimViewModel >( nullptr, UsdStageActor->GetOrLoadUsdStage() );
+		FUsdPrimViewModelRef TreeItem = MakeShared< FUsdPrimViewModel >( nullptr, UsdStage );
 		RootItems.Add( TreeItem );
 
 		PendingRenameItem = TreeItem;
@@ -881,11 +914,6 @@ void SUsdStageTreeView::OnAddChildPrim()
 
 void SUsdStageTreeView::OnCutPrim()
 {
-	if ( !UsdStageActor.IsValid() )
-	{
-		return;
-	}
-
 	FScopedTransaction Transaction( LOCTEXT( "CutPrimTransaction", "Cut prims" ) );
 
 	UE::FSdfChangeBlock Block;
@@ -908,11 +936,6 @@ void SUsdStageTreeView::OnCutPrim()
 
 void SUsdStageTreeView::OnCopyPrim()
 {
-	if ( !UsdStageActor.IsValid() )
-	{
-		return;
-	}
-
 	TArray< FUsdPrimViewModelRef > MySelectedItems = GetSelectedItems();
 
 	TArray<UE::FUsdPrim> Prims;
@@ -931,7 +954,7 @@ void SUsdStageTreeView::OnCopyPrim()
 
 void SUsdStageTreeView::OnPastePrim()
 {
-	if ( !UsdStageActor.IsValid() )
+	if ( !UsdStage )
 	{
 		return;
 	}
@@ -947,8 +970,7 @@ void SUsdStageTreeView::OnPastePrim()
 	// This happens when right-clicking the background area without selecting any prim
 	if ( MySelectedItems.IsEmpty() )
 	{
-		UE::FUsdStage Stage = UsdStageActor->GetOrLoadUsdStage();
-		ParentPrims.Add( Stage.GetPseudoRoot() );
+		ParentPrims.Add( UsdStage.GetPseudoRoot() );
 	}
 	else
 	{
@@ -973,13 +995,7 @@ void SUsdStageTreeView::OnPastePrim()
 
 void SUsdStageTreeView::OnDuplicatePrim( EUsdDuplicateType DuplicateType )
 {
-	if ( !UsdStageActor.IsValid() )
-	{
-		return;
-	}
-
-	UE::FUsdStage Stage = UsdStageActor->GetOrLoadUsdStage();
-	if ( !Stage )
+	if ( !UsdStage )
 	{
 		return;
 	}
@@ -1001,16 +1017,11 @@ void SUsdStageTreeView::OnDuplicatePrim( EUsdDuplicateType DuplicateType )
 		}
 	}
 
-	UsdUtils::DuplicatePrims( Prims, DuplicateType, Stage.GetEditTarget() );
+	UsdUtils::DuplicatePrims( Prims, DuplicateType, UsdStage.GetEditTarget() );
 }
 
 void SUsdStageTreeView::OnDeletePrim()
 {
-	if ( !UsdStageActor.IsValid() )
-	{
-		return;
-	}
-
 	FScopedTransaction Transaction( LOCTEXT( "DeletePrimTransaction", "Delete prims" ) );
 
 	UE::FSdfChangeBlock Block;
@@ -1025,11 +1036,6 @@ void SUsdStageTreeView::OnDeletePrim()
 
 void SUsdStageTreeView::OnRenamePrim()
 {
-	if ( !UsdStageActor.IsValid() )
-	{
-		return;
-	}
-
 	TArray< FUsdPrimViewModelRef > MySelectedItems = GetSelectedItems();
 
 	if ( MySelectedItems.Num() > 0 )
@@ -1044,8 +1050,7 @@ void SUsdStageTreeView::OnRenamePrim()
 
 void SUsdStageTreeView::OnAddReference()
 {
-	UE::FUsdStage& Stage = UsdStageActor->GetOrLoadUsdStage();
-	if ( !UsdStageActor.IsValid() || !Stage || !Stage.IsEditTargetValid() )
+	if ( !UsdStage || !UsdStage.IsEditTargetValid() )
 	{
 		return;
 	}
@@ -1092,11 +1097,6 @@ void SUsdStageTreeView::OnAddReference()
 
 void SUsdStageTreeView::OnClearReferences()
 {
-	if ( !UsdStageActor.IsValid() )
-	{
-		return;
-	}
-
 	FScopedTransaction Transaction( LOCTEXT( "ClearReferenceTransaction", "Clear references to USD layers" ) );
 
 	TArray< FUsdPrimViewModelRef > MySelectedItems = GetSelectedItems();
@@ -1109,8 +1109,7 @@ void SUsdStageTreeView::OnClearReferences()
 
 void SUsdStageTreeView::OnAddPayload()
 {
-	UE::FUsdStage& Stage = UsdStageActor->GetOrLoadUsdStage();
-	if ( !UsdStageActor.IsValid() || !Stage || !Stage.IsEditTargetValid() )
+	if ( !UsdStage || !UsdStage.IsEditTargetValid() )
 	{
 		return;
 	}
@@ -1157,11 +1156,6 @@ void SUsdStageTreeView::OnAddPayload()
 
 void SUsdStageTreeView::OnClearPayloads()
 {
-	if ( !UsdStageActor.IsValid() )
-	{
-		return;
-	}
-
 	FScopedTransaction Transaction( LOCTEXT( "ClearPayloadTransaction", "Clear payloads to USD layers" ) );
 
 	TArray< FUsdPrimViewModelRef > MySelectedItems = GetSelectedItems();
@@ -1172,14 +1166,12 @@ void SUsdStageTreeView::OnClearPayloads()
 	}
 }
 
-void SUsdStageTreeView::OnSetUpLiveLink()
+void SUsdStageTreeView::OnApplySchema( FName SchemaName )
 {
-	if ( !UsdStageActor.IsValid() )
-	{
-		return;
-	}
-
-	FScopedTransaction Transaction( LOCTEXT( "SetUpLiveLinkTransaction", "Set up Live Link for selected prims" ) );
+	FScopedTransaction Transaction( FText::Format(
+		LOCTEXT( "ApplySchemaTransaction", "Apply the '{0}' schema onto selected prims" ),
+		FText::FromName( SchemaName )
+	));
 
 	TArray< FUsdPrimViewModelRef > MySelectedItems = GetSelectedItems();
 
@@ -1187,18 +1179,16 @@ void SUsdStageTreeView::OnSetUpLiveLink()
 
 	for ( FUsdPrimViewModelRef SelectedItem : MySelectedItems )
 	{
-		SelectedItem->SetUpLiveLink();
+		SelectedItem->ApplySchema( SchemaName );
 	}
 }
 
-void SUsdStageTreeView::OnRemoveLiveLink()
+void SUsdStageTreeView::OnRemoveSchema( FName SchemaName )
 {
-	if ( !UsdStageActor.IsValid() )
-	{
-		return;
-	}
-
-	FScopedTransaction Transaction( LOCTEXT( "RemoveLiveLinkTransaction", "Reverses the Live Link configuration for selected prims" ) );
+	FScopedTransaction Transaction( FText::Format(
+		LOCTEXT( "RemoveSchemaTransaction", "Remove the '{0}' schema from selected prims" ),
+		FText::FromName( SchemaName )
+	) );
 
 	TArray< FUsdPrimViewModelRef > MySelectedItems = GetSelectedItems();
 
@@ -1206,94 +1196,52 @@ void SUsdStageTreeView::OnRemoveLiveLink()
 
 	for ( FUsdPrimViewModelRef SelectedItem : MySelectedItems )
 	{
-		SelectedItem->RemoveLiveLink();
+		SelectedItem->RemoveSchema( SchemaName );
 	}
 }
 
-void SUsdStageTreeView::OnSetUpControlRig()
+bool SUsdStageTreeView::CanApplySchema( FName SchemaName )
 {
-	if ( !UsdStageActor.IsValid() )
+	if ( !UsdStage || !UsdStage.IsEditTargetValid() )
 	{
-		return;
+		return false;
 	}
 
-	FScopedTransaction Transaction( LOCTEXT( "SetUpControlRigTransaction", "Set up Control Rig for selected prims" ) );
-
 	TArray< FUsdPrimViewModelRef > MySelectedItems = GetSelectedItems();
 
-	UE::FSdfChangeBlock Block;
-
 	for ( FUsdPrimViewModelRef SelectedItem : MySelectedItems )
 	{
-		SelectedItem->SetUpControlRig();
-	}
-}
-
-void SUsdStageTreeView::OnRemoveControlRig()
-{
-	if ( !UsdStageActor.IsValid() )
-	{
-		return;
+		if ( SelectedItem->CanApplySchema( SchemaName ) )
+		{
+			return true;
+		}
 	}
 
-	FScopedTransaction Transaction( LOCTEXT( "RemoveControlRigTransaction", "Reverses the Live Link configuration for selected prims" ) );
-
-	TArray< FUsdPrimViewModelRef > MySelectedItems = GetSelectedItems();
-
-	UE::FSdfChangeBlock Block;
-
-	for ( FUsdPrimViewModelRef SelectedItem : MySelectedItems )
-	{
-		SelectedItem->RemoveControlRig();
-	}
+	return false;
 }
 
-void SUsdStageTreeView::OnApplyGroomSchema()
+bool SUsdStageTreeView::CanRemoveSchema( FName SchemaName )
 {
-	if ( !UsdStageActor.IsValid() )
+	if ( !UsdStage || !UsdStage.IsEditTargetValid() )
 	{
-		return;
+		return false;
 	}
 
-	FScopedTransaction Transaction( LOCTEXT( "ApplyGroomSchemaTransaction", "Apply Groom schema to selected prims" ) );
-
 	TArray< FUsdPrimViewModelRef > MySelectedItems = GetSelectedItems();
 
-	UE::FSdfChangeBlock Block;
-
 	for ( FUsdPrimViewModelRef SelectedItem : MySelectedItems )
 	{
-		SelectedItem->ApplyGroomSchema();
-	}
-}
-
-void SUsdStageTreeView::OnRemoveGroomSchema()
-{
-	if ( !UsdStageActor.IsValid() )
-	{
-		return;
+		if ( SelectedItem->CanRemoveSchema( SchemaName ) )
+		{
+			return true;
+		}
 	}
 
-	FScopedTransaction Transaction( LOCTEXT( "RemoveGroomSchemaTransaction", "Remove Groom schema from selected prims" ) );
-
-	TArray< FUsdPrimViewModelRef > MySelectedItems = GetSelectedItems();
-
-	UE::FSdfChangeBlock Block;
-
-	for ( FUsdPrimViewModelRef SelectedItem : MySelectedItems )
-	{
-		SelectedItem->RemoveGroomSchema();
-	}
+	return false;
 }
 
 bool SUsdStageTreeView::CanAddChildPrim() const
 {
-	if ( !UsdStageActor.IsValid() )
-	{
-		return false;
-	}
-
-	UE::FUsdStage UsdStage =  UsdStageActor->GetOrLoadUsdStage();
 	if ( !UsdStage )
 	{
 		return false;
@@ -1327,12 +1275,6 @@ bool SUsdStageTreeView::CanAddChildPrim() const
 
 bool SUsdStageTreeView::CanPastePrim() const
 {
-	if ( !UsdStageActor.IsValid() )
-	{
-		return false;
-	}
-
-	UE::FUsdStage UsdStage = UsdStageActor->GetOrLoadUsdStage();
 	if ( !UsdStage )
 	{
 		return false;
@@ -1343,12 +1285,6 @@ bool SUsdStageTreeView::CanPastePrim() const
 
 bool SUsdStageTreeView::DoesPrimExistOnStage() const
 {
-	if ( !UsdStageActor.IsValid() )
-	{
-		return false;
-	}
-
-	UE::FUsdStage UsdStage = UsdStageActor->GetOrLoadUsdStage();
 	if ( !UsdStage || !UsdStage.IsEditTargetValid() )
 	{
 		return false;
@@ -1369,12 +1305,6 @@ bool SUsdStageTreeView::DoesPrimExistOnStage() const
 
 bool SUsdStageTreeView::DoesPrimExistOnEditTarget() const
 {
-	if ( !UsdStageActor.IsValid() )
-	{
-		return false;
-	}
-
-	UE::FUsdStage UsdStage =  UsdStageActor->GetOrLoadUsdStage();
 	if ( !UsdStage || !UsdStage.IsEditTargetValid() )
 	{
 		return false;
@@ -1409,162 +1339,6 @@ bool SUsdStageTreeView::CanDuplicateAllLocalLayerSpecs() const
 	return false;
 }
 
-bool SUsdStageTreeView::CanSetUpLiveLink() const
-{
-	if ( !UsdStageActor.IsValid() )
-	{
-		return false;
-	}
-
-	UE::FUsdStage UsdStage = UsdStageActor->GetOrLoadUsdStage();
-	if ( !UsdStage || !UsdStage.IsEditTargetValid() )
-	{
-		return false;
-	}
-
-	TArray< FUsdPrimViewModelRef > MySelectedItems = GetSelectedItems();
-
-	for ( FUsdPrimViewModelRef SelectedItem : MySelectedItems )
-	{
-		if ( SelectedItem->CanSetUpLiveLink() )
-		{
-			return true;
-		}
-	}
-
-	return false;
-}
-
-bool SUsdStageTreeView::CanRemoveLiveLink() const
-{
-	if ( !UsdStageActor.IsValid() )
-	{
-		return false;
-	}
-
-	UE::FUsdStage UsdStage = UsdStageActor->GetOrLoadUsdStage();
-	if ( !UsdStage || !UsdStage.IsEditTargetValid() )
-	{
-		return false;
-	}
-
-	TArray< FUsdPrimViewModelRef > MySelectedItems = GetSelectedItems();
-
-	for ( FUsdPrimViewModelRef SelectedItem : MySelectedItems )
-	{
-		if ( SelectedItem->CanRemoveLiveLink() )
-		{
-			return true;
-		}
-	}
-
-	return false;
-}
-
-bool SUsdStageTreeView::CanSetUpControlRig() const
-{
-	if ( !UsdStageActor.IsValid() )
-	{
-		return false;
-	}
-
-	UE::FUsdStage UsdStage = UsdStageActor->GetOrLoadUsdStage();
-	if ( !UsdStage || !UsdStage.IsEditTargetValid() )
-	{
-		return false;
-	}
-
-	TArray< FUsdPrimViewModelRef > MySelectedItems = GetSelectedItems();
-
-	for ( FUsdPrimViewModelRef SelectedItem : MySelectedItems )
-	{
-		if ( SelectedItem->CanSetUpControlRig() )
-		{
-			return true;
-		}
-	}
-
-	return false;
-}
-
-bool SUsdStageTreeView::CanRemoveControlRig() const
-{
-	if ( !UsdStageActor.IsValid() )
-	{
-		return false;
-	}
-
-	UE::FUsdStage UsdStage = UsdStageActor->GetOrLoadUsdStage();
-	if ( !UsdStage || !UsdStage.IsEditTargetValid() )
-	{
-		return false;
-	}
-
-	TArray< FUsdPrimViewModelRef > MySelectedItems = GetSelectedItems();
-
-	for ( FUsdPrimViewModelRef SelectedItem : MySelectedItems )
-	{
-		if ( SelectedItem->CanRemoveControlRig() )
-		{
-			return true;
-		}
-	}
-
-	return false;
-}
-
-bool SUsdStageTreeView::CanApplyGroomSchema() const
-{
-	if ( !UsdStageActor.IsValid() )
-	{
-		return false;
-	}
-
-	UE::FUsdStage UsdStage = UsdStageActor->GetOrLoadUsdStage();
-	if ( !UsdStage || !UsdStage.IsEditTargetValid() )
-	{
-		return false;
-	}
-
-	TArray< FUsdPrimViewModelRef > MySelectedItems = GetSelectedItems();
-
-	for ( FUsdPrimViewModelRef SelectedItem : MySelectedItems )
-	{
-		if ( SelectedItem->CanApplyGroomSchema() )
-		{
-			return true;
-		}
-	}
-
-	return false;
-}
-
-bool SUsdStageTreeView::CanRemoveGroomSchema() const
-{
-	if ( !UsdStageActor.IsValid() )
-	{
-		return false;
-	}
-
-	UE::FUsdStage UsdStage = UsdStageActor->GetOrLoadUsdStage();
-	if ( !UsdStage || !UsdStage.IsEditTargetValid() )
-	{
-		return false;
-	}
-
-	TArray< FUsdPrimViewModelRef > MySelectedItems = GetSelectedItems();
-
-	for ( FUsdPrimViewModelRef SelectedItem : MySelectedItems )
-	{
-		if ( SelectedItem->CanRemoveGroomSchema() )
-		{
-			return true;
-		}
-	}
-
-	return false;
-}
-
 void SUsdStageTreeView::RequestListRefresh()
 {
 	SUsdTreeView< FUsdPrimViewModelRef >::RequestListRefresh();
@@ -1603,7 +1377,7 @@ void SUsdStageTreeView::RestoreExpansionStates()
 
 void SUsdStageTreeView::OnToggleAllPayloads( EPayloadsTrigger PayloadsTrigger )
 {
-	if ( !UsdStageActor.IsValid() )
+	if ( !UsdStage )
 	{
 		return;
 	}
@@ -1665,8 +1439,6 @@ void SUsdStageTreeView::OnToggleAllPayloads( EPayloadsTrigger PayloadsTrigger )
 	if ( PrimsToLoad.Num() + PrimsToUnload.Num() > 0 )
 	{
 		UE::FSdfChangeBlock GroupNotices;
-
-		UE::FUsdStage UsdStage = UsdStageActor->GetOrLoadUsdStage();
 		UsdStage.LoadAndUnload( PrimsToLoad, PrimsToUnload );
 	}
 }
diff --git a/Engine/Plugins/Importers/USDImporter/Source/USDStageEditor/Private/SUSDStageTreeView.h b/Engine/Plugins/Importers/USDImporter/Source/USDStageEditor/Private/SUSDStageTreeView.h
--- a/Engine/Plugins/Importers/USDImporter/Source/USDStageEditor/Private/SUSDStageTreeView.h
+++ b/Engine/Plugins/Importers/USDImporter/Source/USDStageEditor/Private/SUSDStageTreeView.h
@@ -2,14 +2,13 @@
 
 #pragma once
 
+#include "USDPrimViewModel.h"
+#include "UsdWrappers/UsdStage.h"
 #include "Widgets/SUSDTreeView.h"
 
 #include "Widgets/Views/SHeaderRow.h"
 #include "Widgets/Views/STreeView.h"
 
-#include "USDMemory.h"
-#include "USDPrimViewModel.h"
-
 class AUsdStageActor;
 class FUICommandList;
 enum class EPayloadsTrigger;
@@ -27,8 +26,9 @@ public:
 		SLATE_EVENT( FOnPrimSelectionChanged, OnPrimSelectionChanged )
 	SLATE_END_ARGS()
 
-	void Construct( const FArguments& InArgs, AUsdStageActor* InUsdStageActor );
-	void Refresh( AUsdStageActor* InUsdStageActor );
+	void Construct( const FArguments& InArgs );
+
+	void Refresh( const UE::FUsdStageWeak& NewStage );
 	void RefreshPrim( const FString& PrimPath, bool bResync );
 
 	FUsdPrimViewModelPtr GetItemFromPrimPath( const FString& PrimPath );
@@ -71,12 +71,10 @@ private:
 	void OnAddPayload();
 	void OnClearPayloads();
 
-	void OnSetUpLiveLink();
-	void OnRemoveLiveLink();
-	void OnSetUpControlRig();
-	void OnRemoveControlRig();
-	void OnApplyGroomSchema();
-	void OnRemoveGroomSchema();
+	void OnApplySchema( FName SchemaName );
+	void OnRemoveSchema( FName SchemaName );
+	bool CanApplySchema( FName SchemaName );
+	bool CanRemoveSchema( FName SchemaName );
 
 	bool CanAddChildPrim() const;
 	bool CanPastePrim() const;
@@ -84,19 +82,14 @@ private:
 	bool DoesPrimExistOnEditTarget() const;
 	bool CanDuplicateAllLocalLayerSpecs() const;
 
-	bool CanSetUpLiveLink() const;
-	bool CanRemoveLiveLink() const;
-	bool CanSetUpControlRig() const;
-	bool CanRemoveControlRig() const;
-	bool CanApplyGroomSchema() const;
-	bool CanRemoveGroomSchema() const;
-
 	/** Uses TreeItemExpansionStates to travel the tree and call SetItemExpansion */
 	void RestoreExpansionStates();
 	virtual void RequestListRefresh() override;
 
 private:
-	TWeakObjectPtr< AUsdStageActor > UsdStageActor;
+	// Should always be valid, we keep the one we're given on Refresh()
+	UE::FUsdStageWeak UsdStage;
+
 	TWeakPtr< FUsdPrimViewModel > PendingRenameItem;
 
 	// So that we can store these across refreshes
diff --git a/Engine/Plugins/Importers/USDImporter/Source/USDStageEditor/Private/USDStageEditorSettings.h b/Engine/Plugins/Importers/USDImporter/Source/USDStageEditor/Private/USDStageEditorSettings.h
--- a/Engine/Plugins/Importers/USDImporter/Source/USDStageEditor/Private/USDStageEditorSettings.h
+++ b/Engine/Plugins/Importers/USDImporter/Source/USDStageEditor/Private/USDStageEditorSettings.h
@@ -16,6 +16,10 @@ class UUsdStageEditorSettings : public UObject
 
 public:
 	/** Whether our prim selection in SUSDStageTreeView is kept synchronized with the viewport selection */
-	UPROPERTY(config, EditAnywhere, Category=Options)
+	UPROPERTY( config, EditAnywhere, Category = USD )
 	bool bSelectionSynced = false;
+
+	/** Whether to automatically set a layer as edit target when isolating it */
+	UPROPERTY( config, EditAnywhere, Category = USD )
+	bool bIsolateLayerSyncedWithEditTarget = true;
 };
\ No newline at end of file
diff --git a/Engine/Plugins/Importers/USDImporter/Source/USDStageEditor/Private/Widgets/SUSDIntegrationsPanel.cpp b/Engine/Plugins/Importers/USDImporter/Source/USDStageEditor/Private/Widgets/SUSDIntegrationsPanel.cpp
--- a/Engine/Plugins/Importers/USDImporter/Source/USDStageEditor/Private/Widgets/SUSDIntegrationsPanel.cpp
+++ b/Engine/Plugins/Importers/USDImporter/Source/USDStageEditor/Private/Widgets/SUSDIntegrationsPanel.cpp
@@ -438,7 +438,7 @@ TSharedRef< SWidget > SUsdIntegrationsPanelRow::GenerateWidgetForColumn( const F
 		];
 }
 
-void SUsdIntegrationsPanel::Construct( const FArguments& InArgs, const UE::FUsdStageWeak& InUsdStage, const TCHAR* InPrimPath )
+void SUsdIntegrationsPanel::Construct( const FArguments& InArgs )
 {
 	SAssignNew( HeaderRowWidget, SHeaderRow )
 
@@ -453,74 +453,20 @@ void SUsdIntegrationsPanel::Construct( const FArguments& InArgs, const UE::FUsdS
 	SListView::Construct
 	(
 		SListView::FArguments()
-		.ListItemsSource( &Attributes )
+		.ListItemsSource( &ViewModel.Attributes )
 		.OnGenerateRow( this, &SUsdIntegrationsPanel::OnGenerateRow )
 		.HeaderRow( HeaderRowWidget )
 	);
 
-	SetPrimPath( InUsdStage, InPrimPath );
+	SetVisibility( EVisibility::Collapsed ); // Start hidden until SetPrimPath displays us
 }
 
 void SUsdIntegrationsPanel::SetPrimPath( const UE::FUsdStageWeak& InUsdStage, const TCHAR* InPrimPath )
 {
-	PrimPath = InPrimPath;
-	UsdStage = InUsdStage;
+	ViewModel.UpdateAttributes( InUsdStage, InPrimPath );
 
-	Attributes.Reset();
+	SetVisibility( ViewModel.Attributes.Num() > 0 ? EVisibility::Visible : EVisibility::Collapsed );
 
-	EVisibility PanelVisibility = EVisibility::Collapsed;
-
-	// Important as this function may be reentrant in case our ApplyXSchema calls create attributes
-	UE::FSdfChangeBlock Block;
-
-	if ( InUsdStage )
-	{
-		if ( UE::FUsdPrim Prim = InUsdStage.GetPrimAtPath( UE::FSdfPath{ InPrimPath } ) )
-		{
-			if ( !Prim.IsPseudoRoot() )
-			{
-				const bool bHasLiveLink = UsdUtils::PrimHasLiveLinkSchema( Prim );
-				const bool bHasControlRig = UsdUtils::PrimHasControlRigSchema( Prim );
-				const bool bIsSkelRoot = Prim.IsA( TEXT( "SkelRoot" ) );
-
-				TArray<FString> AttributeNames;
-
-				if ( bHasLiveLink )
-				{
-					if ( bIsSkelRoot )
-					{
-						AttributeNames.Add( *UsdToUnreal::ConvertToken( UnrealIdentifiers::UnrealAnimBlueprintPath ) );
-					}
-
-					AttributeNames.Add( *UsdToUnreal::ConvertToken( UnrealIdentifiers::UnrealLiveLinkSubjectName ) );
-					AttributeNames.Add( *UsdToUnreal::ConvertToken( UnrealIdentifiers::UnrealLiveLinkEnabled ) );
-				}
-
-				if( bHasControlRig && bIsSkelRoot )
-				{
-					AttributeNames.Add( *UsdToUnreal::ConvertToken( UnrealIdentifiers::UnrealControlRigPath ) );
-					AttributeNames.Add( *UsdToUnreal::ConvertToken( UnrealIdentifiers::UnrealUseFKControlRig ) );
-					AttributeNames.Add( *UsdToUnreal::ConvertToken( UnrealIdentifiers::UnrealControlRigReduceKeys ) );
-					AttributeNames.Add( *UsdToUnreal::ConvertToken( UnrealIdentifiers::UnrealControlRigReductionTolerance ) );
-				}
-
-				if ( bHasLiveLink || bHasControlRig )
-				{
-					PanelVisibility = EVisibility::Visible;
-				}
-
-				for ( const FString& AttributeName : AttributeNames )
-				{
-					if ( UE::FUsdAttribute Attr = Prim.GetAttribute( *AttributeName ) )
-					{
-						Attributes.Add( MakeShared<UE::FUsdAttribute>( Attr ) );
-					}
-				}
-			}
-		}
-	}
-
-	SetVisibility( PanelVisibility );
 	RequestListRefresh();
 }
 
diff --git a/Engine/Plugins/Importers/USDImporter/Source/USDStageEditor/Private/Widgets/SUSDIntegrationsPanel.h b/Engine/Plugins/Importers/USDImporter/Source/USDStageEditor/Private/Widgets/SUSDIntegrationsPanel.h
--- a/Engine/Plugins/Importers/USDImporter/Source/USDStageEditor/Private/Widgets/SUSDIntegrationsPanel.h
+++ b/Engine/Plugins/Importers/USDImporter/Source/USDStageEditor/Private/Widgets/SUSDIntegrationsPanel.h
@@ -2,8 +2,10 @@
 
 #pragma once
 
+#include "USDIntegrationsViewModel.h"
 #include "UsdWrappers/UsdStage.h"
 
+#include "Templates/SharedPointer.h"
 #include "Widgets/Views/SListView.h"
 
 #include "SUSDIntegrationsPanel.generated.h"
@@ -57,18 +59,15 @@ class SUsdIntegrationsPanel : public SListView< TSharedPtr<UE::FUsdAttribute> >
 	SLATE_END_ARGS()
 
 public:
-	void Construct( const FArguments& InArgs, const UE::FUsdStageWeak& InUsdStage, const TCHAR* InPrimPath );
-	void SetPrimPath( const UE::FUsdStageWeak& UsdStage, const TCHAR* InPrimPath );
+	void Construct( const FArguments& InArgs );
+	void SetPrimPath( const UE::FUsdStageWeak& UsdStage, const TCHAR* PrimPath );
 
 protected:
 	TSharedRef< ITableRow > OnGenerateRow( TSharedPtr<UE::FUsdAttribute> InAttr, const TSharedRef< STableViewBase >& OwnerTable );
 
 private:
 	TSharedPtr< SHeaderRow > HeaderRowWidget;
-
-	TArray<TSharedPtr<UE::FUsdAttribute>> Attributes;
-	UE::FUsdStageWeak UsdStage;
-	FString PrimPath;
+	FUsdIntegrationsViewModel ViewModel;
 };
 
 #endif // USE_USD_SDK
\ No newline at end of file
diff --git a/Engine/Plugins/Importers/USDImporter/Source/USDStageEditor/Private/Widgets/SUSDPrimPropertiesList.cpp b/Engine/Plugins/Importers/USDImporter/Source/USDStageEditor/Private/Widgets/SUSDPrimPropertiesList.cpp
--- a/Engine/Plugins/Importers/USDImporter/Source/USDStageEditor/Private/Widgets/SUSDPrimPropertiesList.cpp
+++ b/Engine/Plugins/Importers/USDImporter/Source/USDStageEditor/Private/Widgets/SUSDPrimPropertiesList.cpp
@@ -650,10 +650,8 @@ void SUsdPrimPropertyRow::OnCheckBoxCheckStateChanged( ECheckBoxState NewState )
 	UsdPrimAttribute->SetAttributeValue( NewValue );
 }
 
-void SUsdPrimPropertiesList::Construct( const FArguments& InArgs, const UE::FUsdStageWeak& UsdStage, const TCHAR* InPrimPath )
+void SUsdPrimPropertiesList::Construct( const FArguments& InArgs )
 {
-	GeneratePropertiesList( UsdStage, InPrimPath );
-
 	// Clear map as usd file may have additional Kinds now
 	UsdPrimPropertiesListImpl::ResetOptions(TEXT("Kind"));
 
@@ -683,9 +681,8 @@ TSharedRef< ITableRow > SUsdPrimPropertiesList::OnGenerateRow( TSharedPtr< FUsdP
 
 void SUsdPrimPropertiesList::GeneratePropertiesList( const UE::FUsdStageWeak& UsdStage, const TCHAR* InPrimPath )
 {
-	float TimeCode = 0.f;
-	ViewModel.UsdStage = UsdStage;
-	ViewModel.Refresh( InPrimPath, TimeCode );
+	const float TimeCode = 0.f;
+	ViewModel.Refresh( UsdStage, InPrimPath, TimeCode );
 }
 
 void SUsdPrimPropertiesList::SetPrimPath( const UE::FUsdStageWeak& UsdStage, const TCHAR* InPrimPath )
diff --git a/Engine/Plugins/Importers/USDImporter/Source/USDStageEditor/Private/Widgets/SUSDPrimPropertiesList.h b/Engine/Plugins/Importers/USDImporter/Source/USDStageEditor/Private/Widgets/SUSDPrimPropertiesList.h
--- a/Engine/Plugins/Importers/USDImporter/Source/USDStageEditor/Private/Widgets/SUSDPrimPropertiesList.h
+++ b/Engine/Plugins/Importers/USDImporter/Source/USDStageEditor/Private/Widgets/SUSDPrimPropertiesList.h
@@ -13,7 +13,7 @@ class SUsdPrimPropertiesList : public SListView< TSharedPtr< FUsdPrimAttributeVi
 	SLATE_END_ARGS()
 
 public:
-	void Construct( const FArguments& InArgs, const UE::FUsdStageWeak& UsdStage, const TCHAR* InPrimPath );
+	void Construct( const FArguments& InArgs );
 	void SetPrimPath( const UE::FUsdStageWeak& UsdStage, const TCHAR* InPrimPath );
 
 protected:
diff --git a/Engine/Plugins/Importers/USDImporter/Source/USDStageEditor/Private/Widgets/SUSDReferencesList.cpp b/Engine/Plugins/Importers/USDImporter/Source/USDStageEditor/Private/Widgets/SUSDReferencesList.cpp
--- a/Engine/Plugins/Importers/USDImporter/Source/USDStageEditor/Private/Widgets/SUSDReferencesList.cpp
+++ b/Engine/Plugins/Importers/USDImporter/Source/USDStageEditor/Private/Widgets/SUSDReferencesList.cpp
@@ -55,7 +55,7 @@ TSharedRef< SWidget > SUsdReferenceRow::GenerateWidgetForColumn( const FName& Co
 		];
 }
 
-void SUsdReferencesList::Construct( const FArguments& InArgs, const UE::FUsdStageWeak& UsdStage, const TCHAR* PrimPath )
+void SUsdReferencesList::Construct( const FArguments& InArgs )
 {
 	SAssignNew( HeaderRowWidget, SHeaderRow )
 
@@ -71,7 +71,7 @@ void SUsdReferencesList::Construct( const FArguments& InArgs, const UE::FUsdStag
 		.HeaderRow( HeaderRowWidget )
 	);
 
-	SetPrimPath( UsdStage, PrimPath );
+	SetVisibility( EVisibility::Collapsed ); // Start hidden until SetPrimPath displays us
 }
 
 TSharedRef< ITableRow > SUsdReferencesList::OnGenerateRow( TSharedPtr< FUsdReference > InDisplayNode, const TSharedRef< STableViewBase >& OwnerTable )
@@ -83,18 +83,7 @@ void SUsdReferencesList::SetPrimPath( const UE::FUsdStageWeak& UsdStage, const T
 {
 	ViewModel.UpdateReferences( UsdStage, PrimPath );
 
-	EVisibility NewVisibility = EVisibility::Collapsed;
-	if ( UsdStage )
-	{
-		if ( UE::FUsdPrim UsdPrim = UsdStage.GetPrimAtPath( UE::FSdfPath{ PrimPath } ) )
-		{
-			if ( UsdPrim.HasAuthoredReferences() )
-			{
-				NewVisibility = EVisibility::Visible;
-			}
-		}
-	}
-	SetVisibility( NewVisibility );
+	SetVisibility( ViewModel.References.Num() > 0 ? EVisibility::Visible : EVisibility::Collapsed );
 
 	RequestListRefresh();
 }
diff --git a/Engine/Plugins/Importers/USDImporter/Source/USDStageEditor/Private/Widgets/SUSDReferencesList.h b/Engine/Plugins/Importers/USDImporter/Source/USDStageEditor/Private/Widgets/SUSDReferencesList.h
--- a/Engine/Plugins/Importers/USDImporter/Source/USDStageEditor/Private/Widgets/SUSDReferencesList.h
+++ b/Engine/Plugins/Importers/USDImporter/Source/USDStageEditor/Private/Widgets/SUSDReferencesList.h
@@ -31,7 +31,7 @@ class SUsdReferencesList : public SListView< TSharedPtr< FUsdReference > >
 	SLATE_END_ARGS()
 
 public:
-	void Construct( const FArguments& InArgs, const UE::FUsdStageWeak& UsdStage, const TCHAR* PrimPath );
+	void Construct( const FArguments& InArgs );
 	void SetPrimPath( const UE::FUsdStageWeak& UsdStage, const TCHAR* PrimPath );
 
 protected:
diff --git a/Engine/Plugins/Importers/USDImporter/Source/USDStageEditor/Private/Widgets/SUSDVariantSetsList.cpp b/Engine/Plugins/Importers/USDImporter/Source/USDStageEditor/Private/Widgets/SUSDVariantSetsList.cpp
--- a/Engine/Plugins/Importers/USDImporter/Source/USDStageEditor/Private/Widgets/SUSDVariantSetsList.cpp
+++ b/Engine/Plugins/Importers/USDImporter/Source/USDStageEditor/Private/Widgets/SUSDVariantSetsList.cpp
@@ -84,7 +84,7 @@ void SUsdVariantRow::OnSelectionChanged( TSharedPtr< FString > NewValue, ESelect
 	OnVariantSelectionChanged.ExecuteIfBound( VariantSet.ToSharedRef(), NewValue );
 }
 
-void SVariantsList::Construct( const FArguments& InArgs, const UE::FUsdStageWeak& UsdStage, const TCHAR* InPrimPath )
+void SVariantsList::Construct( const FArguments& InArgs )
 {
 	SAssignNew( HeaderRowWidget, SHeaderRow )
 
@@ -104,7 +104,7 @@ void SVariantsList::Construct( const FArguments& InArgs, const UE::FUsdStageWeak
 		.HeaderRow( HeaderRowWidget )
 	);
 
-	SetPrimPath( UsdStage, InPrimPath );
+	SetVisibility( EVisibility::Collapsed ); // Start hidden until SetPrimPath displays us
 }
 
 TSharedRef< ITableRow > SVariantsList::OnGenerateRow( TSharedPtr< FUsdVariantSetViewModel > InDisplayNode, const TSharedRef< STableViewBase >& OwnerTable )
@@ -115,21 +115,9 @@ TSharedRef< ITableRow > SVariantsList::OnGenerateRow( TSharedPtr< FUsdVariantSet
 
 void SVariantsList::SetPrimPath( const UE::FUsdStageWeak& UsdStage, const TCHAR* InPrimPath )
 {
-	PrimPath = InPrimPath;
-	ViewModel.UpdateVariantSets( UsdStage, *PrimPath );
+	ViewModel.UpdateVariantSets( UsdStage, InPrimPath );
 
-	EVisibility NewVisibility = EVisibility::Collapsed;
-	if ( UsdStage )
-	{
-		if ( UE::FUsdPrim UsdPrim = UsdStage.GetPrimAtPath( UE::FSdfPath{ InPrimPath } ) )
-		{
-			if ( UsdPrim.HasVariantSets() )
-			{
-				NewVisibility = EVisibility::Visible;
-			}
-		}
-	}
-	SetVisibility( NewVisibility );
+	SetVisibility( ViewModel.VariantSets.Num() > 0 ? EVisibility::Visible : EVisibility::Collapsed );
 
 	RequestListRefresh();
 }
diff --git a/Engine/Plugins/Importers/USDImporter/Source/USDStageEditor/Private/Widgets/SUSDVariantSetsList.h b/Engine/Plugins/Importers/USDImporter/Source/USDStageEditor/Private/Widgets/SUSDVariantSetsList.h
--- a/Engine/Plugins/Importers/USDImporter/Source/USDStageEditor/Private/Widgets/SUSDVariantSetsList.h
+++ b/Engine/Plugins/Importers/USDImporter/Source/USDStageEditor/Private/Widgets/SUSDVariantSetsList.h
@@ -45,18 +45,15 @@ class SVariantsList : public SListView< TSharedPtr< FUsdVariantSetViewModel > >
 	SLATE_END_ARGS()
 
 public:
-	void Construct( const FArguments& InArgs, const UE::FUsdStageWeak& UsdStage, const TCHAR* InPrimPath );
+	void Construct( const FArguments& InArgs );
 	void SetPrimPath( const UE::FUsdStageWeak& UsdStage, const TCHAR* InPrimPath );
 
 protected:
-	void UpdateVariantSets( const TCHAR* InPrimPath );
 	TSharedRef< ITableRow > OnGenerateRow( TSharedPtr< FUsdVariantSetViewModel > InDisplayNode, const TSharedRef< STableViewBase >& OwnerTable );
 
 	void OnVariantSelectionChanged( const TSharedRef< FUsdVariantSetViewModel >& VariantSet, const TSharedPtr< FString >& NewValue );
 
 private:
-	FString PrimPath;
-
 	FUsdVariantSetsViewModel ViewModel;
 	TSharedPtr< SHeaderRow > HeaderRowWidget;
 };
diff --git a/Engine/Plugins/Importers/USDImporter/Source/USDStageEditorViewModels/Private/USDIntegrationsViewModel.cpp b/Engine/Plugins/Importers/USDImporter/Source/USDStageEditorViewModels/Private/USDIntegrationsViewModel.cpp
new file mode 100644
--- /dev/null
+++ b/Engine/Plugins/Importers/USDImporter/Source/USDStageEditorViewModels/Private/USDIntegrationsViewModel.cpp
@@ -0,0 +1,75 @@
+// Copyright Epic Games, Inc. All Rights Reserved.
+
+#include "USDIntegrationsViewModel.h"
+
+#include "USDIntegrationUtils.h"
+#include "USDTypesConversion.h"
+
+#include "UsdWrappers/SdfAttributeSpec.h"
+#include "UsdWrappers/SdfPath.h"
+#include "UsdWrappers/UsdPrim.h"
+#include "UsdWrappers/UsdStage.h"
+
+void FUsdIntegrationsViewModel::UpdateAttributes( const UE::FUsdStageWeak& InUsdStage, const TCHAR* InPrimPath )
+{
+#if USE_USD_SDK
+	UsdStage = InUsdStage;
+	PrimPath = InPrimPath;
+
+	Attributes.Reset();
+
+	if ( !UsdStage || PrimPath.IsEmpty() || UE::FSdfPath{ InPrimPath }.IsAbsoluteRootPath() )
+	{
+		return;
+	}
+
+	TArray<FString> AttributeNames;
+
+	TFunction<void(bool, bool)> AddLiveLinkAttributes = [&AttributeNames]( bool bHasLiveLink, bool bIsSkelRoot )
+	{
+		if ( bHasLiveLink )
+		{
+			if ( bIsSkelRoot )
+			{
+				AttributeNames.Add( *UsdToUnreal::ConvertToken( UnrealIdentifiers::UnrealAnimBlueprintPath ) );
+			}
+
+			AttributeNames.Add( *UsdToUnreal::ConvertToken( UnrealIdentifiers::UnrealLiveLinkSubjectName ) );
+			AttributeNames.Add( *UsdToUnreal::ConvertToken( UnrealIdentifiers::UnrealLiveLinkEnabled ) );
+		}
+	};
+
+	TFunction<void( bool, bool )> AddControlRigAttributes = [&AttributeNames]( bool bHasControlRig, bool bIsSkelRoot )
+	{
+		if ( bHasControlRig && bIsSkelRoot )
+		{
+			AttributeNames.Add( *UsdToUnreal::ConvertToken( UnrealIdentifiers::UnrealControlRigPath ) );
+			AttributeNames.Add( *UsdToUnreal::ConvertToken( UnrealIdentifiers::UnrealUseFKControlRig ) );
+			AttributeNames.Add( *UsdToUnreal::ConvertToken( UnrealIdentifiers::UnrealControlRigReduceKeys ) );
+			AttributeNames.Add( *UsdToUnreal::ConvertToken( UnrealIdentifiers::UnrealControlRigReductionTolerance ) );
+		}
+	};
+
+    if ( UE::FUsdPrim Prim = InUsdStage.GetPrimAtPath( UE::FSdfPath{ InPrimPath } ) )
+    {
+        if ( !Prim.IsPseudoRoot() )
+        {
+            const bool bHasLiveLink = UsdUtils::PrimHasSchema( Prim, UnrealIdentifiers::LiveLinkAPI );
+            const bool bHasControlRig = UsdUtils::PrimHasSchema( Prim, UnrealIdentifiers::ControlRigAPI );
+            const bool bIsSkelRoot = Prim.IsA( TEXT( "SkelRoot" ) );
+
+            AddLiveLinkAttributes( bHasLiveLink, bIsSkelRoot );
+            AddControlRigAttributes( bHasControlRig, bIsSkelRoot );
+
+            for ( const FString& AttributeName : AttributeNames )
+            {
+                if ( UE::FUsdAttribute Attr = Prim.GetAttribute( *AttributeName ) )
+                {
+					Attributes.Add( MakeShared<UE::FUsdAttribute>( Attr ) );
+                }
+            }
+        }
+    }
+#endif // USE_USD_SDK
+}
+
diff --git a/Engine/Plugins/Importers/USDImporter/Source/USDStageEditorViewModels/Private/USDPrimAttributesViewModel.cpp b/Engine/Plugins/Importers/USDImporter/Source/USDStageEditorViewModels/Private/USDPrimAttributesViewModel.cpp
--- a/Engine/Plugins/Importers/USDImporter/Source/USDStageEditorViewModels/Private/USDPrimAttributesViewModel.cpp
+++ b/Engine/Plugins/Importers/USDImporter/Source/USDStageEditorViewModels/Private/USDPrimAttributesViewModel.cpp
@@ -199,11 +199,13 @@ void FUsdPrimAttributesViewModel::SetPrimAttribute( const FString& AttributeName
 #endif // #if USE_USD_SDK
 }
 
-void FUsdPrimAttributesViewModel::Refresh( const TCHAR* InPrimPath, float TimeCode )
+void FUsdPrimAttributesViewModel::Refresh( const UE::FUsdStageWeak& InUsdStage, const TCHAR* InPrimPath, float TimeCode )
 {
 	FScopedUnrealAllocs UnrealAllocs;
 
+	UsdStage = InUsdStage;
 	PrimPath = InPrimPath;
+
 	PrimAttributes.Reset();
 
 #if USE_USD_SDK
diff --git a/Engine/Plugins/Importers/USDImporter/Source/USDStageEditorViewModels/Private/USDPrimViewModel.cpp b/Engine/Plugins/Importers/USDImporter/Source/USDStageEditorViewModels/Private/USDPrimViewModel.cpp
--- a/Engine/Plugins/Importers/USDImporter/Source/USDStageEditorViewModels/Private/USDPrimViewModel.cpp
+++ b/Engine/Plugins/Importers/USDImporter/Source/USDStageEditorViewModels/Private/USDPrimViewModel.cpp
@@ -10,22 +10,17 @@
 
 #include "UsdWrappers/SdfLayer.h"
 #include "UsdWrappers/SdfPath.h"
-
-#include "Misc/Paths.h"
-#include "ScopedTransaction.h"
-
-#include <iterator>
+#include "UsdWrappers/UsdPrim.h"
+#include "UsdWrappers/UsdStage.h"
 
 #if USE_USD_SDK
 
 #include "USDIncludesStart.h"
 	#include "pxr/pxr.h"
-	#include "pxr/usd/sdf/namespaceEdit.h"
 	#include "pxr/usd/sdf/path.h"
 	#include "pxr/usd/usd/modelAPI.h"
 	#include "pxr/usd/usd/payloads.h"
 	#include "pxr/usd/usd/prim.h"
-	#include "pxr/usd/usd/references.h"
 	#include "pxr/usd/usd/tokens.h"
 	#include "pxr/usd/usdGeom/xform.h"
 	#include "pxr/usd/usdSkel/root.h"
@@ -33,105 +28,20 @@
 
 #endif // #if USE_USD_SDK
 
-namespace UE::USDPrimViewModel::Private
-{
-#if USE_USD_SDK
-	bool CanAddOrRemoveCustomSchema( pxr::UsdPrim Prim, pxr::TfToken SchemaName, bool bAddSchema )
-	{
-		if ( !Prim )
-		{
-			return false;
-		}
-
-		FScopedUsdAllocs UsdAllocs;
-
-		pxr::TfType Schema = pxr::UsdSchemaRegistry::GetTypeFromSchemaTypeName( SchemaName );
-		if ( !ensure( static_cast< bool >( Schema ) ) )
-		{
-			return false;
-		}
-
-		// Check if the schema is compatible with this prim
-		if ( bAddSchema && !Prim.CanApplyAPI( Schema ) )
-		{
-			return false;
-		}
-
-		pxr::UsdStageRefPtr Stage = Prim.GetStage();
-		pxr::SdfLayerRefPtr EditTarget = Stage->GetEditTarget().GetLayer();
-		if ( !Stage || !EditTarget )
-		{
-			return false;
-		}
-
-		bool bAlreadyHasSchema = false;
-		if ( EditTarget == Stage->GetRootLayer() )
-		{
-			bAlreadyHasSchema = Prim.HasAPI( Schema );
-		}
-		else
-		{
-			// In the future we'll have better layer editing facilities. For now, lets make sure that
-			// we only show that we can setup/remove a schema for a prim if it can be done *at that
-			// particular edit target*, which is what will happen when clicking the button anyway.
-			// For example, a user could:
-			//  - Add the schema on a sublayer;
-			//  - Remove the schema on the root layer;
-			//  - Select the sublayer as the edit target;
-			// If we always just queried the existing stage, the user would see that they can add the schema
-			// (because the composed prim does not have the schema) but clicking the button would do nothing,
-			// as the sublayer prim spec already has the schema...
-
-			pxr::SdfPrimSpecHandle PrimSpecOnLayer = EditTarget->GetPrimAtPath( Prim.GetPath() );
-			if( !PrimSpecOnLayer )
-			{
-				// We can always just make a new 'over' on this layer
-				return true;
-			}
-
-			// We can only add the schema if its not present on this layer, but we should let USD sort the
-			// Ops out because we could have Ops to add/delete/prepend/append the same schema on the same prim...
-			pxr::SdfListOp<pxr::TfToken> Ops = PrimSpecOnLayer->GetInfo( pxr::UsdTokens->apiSchemas ).Get<pxr::SdfTokenListOp>();
-			std::vector<pxr::TfToken> AppliedOps;
-			Ops.ApplyOperations(&AppliedOps);
-			bAlreadyHasSchema = std::find( AppliedOps.begin(), AppliedOps.end(), SchemaName ) != AppliedOps.end();
-		}
-
-		return bAddSchema != bAlreadyHasSchema;
-	}
-
-	void RemoveCustomSchema( pxr::UsdPrim Prim, pxr::TfToken SchemaName )
-	{
-		if ( !Prim )
-		{
-			return;
-		}
-
-		FScopedUsdAllocs UsdAllocs;
-
-		pxr::TfType Schema = pxr::UsdSchemaRegistry::GetTypeFromSchemaTypeName( SchemaName );
-		ensure( static_cast< bool >( Schema ) );
-
-		ensure( Prim.RemoveAPI( Schema ) );
-	}
-
-#endif // #if USE_USD_SDK
-}
-
-FUsdPrimViewModel::FUsdPrimViewModel( FUsdPrimViewModel* InParentItem, const UE::FUsdStageWeak& InUsdStage, const UE::FUsdPrim& InUsdPrim )
-	: FUsdPrimViewModel( InParentItem, InUsdStage )
-{
-	UsdPrim = InUsdPrim;
-
-	RefreshData( false );
-	FillChildren();
-}
+#define LOCTEXT_NAMESPACE "USDPrimViewModel"
 
-FUsdPrimViewModel::FUsdPrimViewModel( FUsdPrimViewModel* InParentItem, const UE::FUsdStageWeak& InUsdStage )
+FUsdPrimViewModel::FUsdPrimViewModel(
+	FUsdPrimViewModel* InParentItem,
+	const UE::FUsdStageWeak& InUsdStage,
+	const UE::FUsdPrim& InPrim
+)
 	: UsdStage( InUsdStage )
+	, UsdPrim( InPrim )
 	, ParentItem( InParentItem )
 	, RowData( MakeShared< FUsdPrimModel >() )
 {
+	RefreshData( false );
+	FillChildren();
 }
 
 TArray< FUsdPrimViewModelRef >& FUsdPrimViewModel::UpdateChildren()
@@ -307,121 +217,53 @@ void FUsdPrimViewModel::TogglePayload()
 	}
 }
 
-void FUsdPrimViewModel::SetUpLiveLink()
+void FUsdPrimViewModel::ApplySchema( FName SchemaName )
 {
 #if USE_USD_SDK
-	UsdUtils::ApplyLiveLinkSchema( UsdPrim );
+	UsdUtils::ApplySchema( UsdPrim, UnrealToUsd::ConvertToken( *SchemaName.ToString() ).Get() );
 #endif // #if USE_USD_SDK
 }
 
-bool FUsdPrimViewModel::CanSetUpLiveLink() const
+bool FUsdPrimViewModel::CanApplySchema( FName SchemaName ) const
 {
 #if USE_USD_SDK
-	const bool bAddSchema = true;
-	return UsdPrim
-		&& !UsdPrim.IsPseudoRoot()
-		&& pxr::UsdPrim{ UsdPrim }.IsA<pxr::UsdGeomXformable>() // SkelRoot is also an Xformable
-		&& UE::USDPrimViewModel::Private::CanAddOrRemoveCustomSchema( UsdPrim, UnrealIdentifiers::LiveLinkAPI, bAddSchema );
-#else
-	return false;
-#endif // #if USE_USD_SDK
-}
-
-void FUsdPrimViewModel::RemoveLiveLink()
-{
-#if USE_USD_SDK
-	UE::USDPrimViewModel::Private::RemoveCustomSchema( UsdPrim, UnrealIdentifiers::LiveLinkAPI );
-#endif // #if USE_USD_SDK
-}
-
-bool FUsdPrimViewModel::CanRemoveLiveLink() const
-{
-#if USE_USD_SDK
-	const bool bAddSchema = false;
-	return UsdPrim
-		&& !UsdPrim.IsPseudoRoot()
-		&& pxr::UsdPrim{ UsdPrim }.IsA<pxr::UsdGeomXformable>() // SkelRoot is also an Xformable
-		&& UE::USDPrimViewModel::Private::CanAddOrRemoveCustomSchema( UsdPrim, UnrealIdentifiers::LiveLinkAPI, bAddSchema );
-#else
-	return false;
-#endif // #if USE_USD_SDK
-}
-
-void FUsdPrimViewModel::SetUpControlRig()
-{
-#if USE_USD_SDK
-	UsdUtils::ApplyControlRigSchema( UsdPrim );
-#endif // #if USE_USD_SDK
-}
+	if ( !UsdPrim || UsdPrim.IsPseudoRoot() )
+	{
+		return false;
+	}
 
-bool FUsdPrimViewModel::CanSetUpControlRig() const
-{
-#if USE_USD_SDK
-	const bool bAddSchema = true;
-	return UsdPrim
-		&& !UsdPrim.IsPseudoRoot()
-		&& pxr::UsdPrim{ UsdPrim }.IsA<pxr::UsdSkelRoot>()
-		&& UE::USDPrimViewModel::Private::CanAddOrRemoveCustomSchema( UsdPrim, UnrealIdentifiers::ControlRigAPI, bAddSchema );
-#else
-	return false;
-#endif // #if USE_USD_SDK
-}
+	FScopedUsdAllocs UsdAllocs;
 
-void FUsdPrimViewModel::RemoveControlRig()
-{
-#if USE_USD_SDK
-	UE::USDPrimViewModel::Private::RemoveCustomSchema( UsdPrim, UnrealIdentifiers::ControlRigAPI );
-#endif // #if USE_USD_SDK
-}
+	pxr::UsdPrim PxrUsdPrim{ UsdPrim };
+	pxr::TfToken SchemaToken = UnrealToUsd::ConvertToken( *SchemaName.ToString() ).Get();
 
-bool FUsdPrimViewModel::CanRemoveControlRig() const
-{
-#if USE_USD_SDK
-	const bool bAddSchema = false;
-	return UsdPrim
-		&& !UsdPrim.IsPseudoRoot()
-		&& pxr::UsdPrim{ UsdPrim }.IsA<pxr::UsdSkelRoot>()
-		&& UE::USDPrimViewModel::Private::CanAddOrRemoveCustomSchema( UsdPrim, UnrealIdentifiers::ControlRigAPI, bAddSchema );
-#else
-	return false;
-#endif // #if USE_USD_SDK
-}
+	if ( SchemaToken == UnrealIdentifiers::ControlRigAPI && !PxrUsdPrim.IsA<pxr::UsdSkelRoot>() )
+	{
+		return false;
+	}
 
-void FUsdPrimViewModel::ApplyGroomSchema()
-{
-#if USE_USD_SDK
-	UsdUtils::ApplySchema( UsdPrim, UnrealIdentifiers::GroomAPI );
-#endif // #if USE_USD_SDK
-}
+	if ( !PxrUsdPrim.IsA<pxr::UsdGeomXformable>() )
+	{
+		return false;
+	}
 
-bool FUsdPrimViewModel::CanApplyGroomSchema() const
-{
-#if USE_USD_SDK
-	const bool bAddSchema = true;
-	return UsdPrim
-		&& !UsdPrim.IsPseudoRoot()
-		&& pxr::UsdPrim{ UsdPrim }.IsA<pxr::UsdGeomXformable>() // UsdGeomCurves is also an Xformable
-		&& UE::USDPrimViewModel::Private::CanAddOrRemoveCustomSchema( UsdPrim, UnrealIdentifiers::GroomAPI, bAddSchema );
+	return UsdUtils::CanApplySchema( UsdPrim, SchemaToken );
 #else
 	return false;
 #endif // #if USE_USD_SDK
 }
 
-void FUsdPrimViewModel::RemoveGroomSchema()
+void FUsdPrimViewModel::RemoveSchema( FName SchemaName )
 {
 #if USE_USD_SDK
-	UE::USDPrimViewModel::Private::RemoveCustomSchema( UsdPrim, UnrealIdentifiers::GroomAPI );
+	UsdUtils::RemoveSchema( UsdPrim, UnrealToUsd::ConvertToken( *SchemaName.ToString() ).Get() );
 #endif // #if USE_USD_SDK
 }
 
-bool FUsdPrimViewModel::CanRemoveGroomSchema() const
+bool FUsdPrimViewModel::CanRemoveSchema( FName SchemaName ) const
 {
 #if USE_USD_SDK
-	const bool bAddSchema = false;
-	return UsdPrim
-		&& !UsdPrim.IsPseudoRoot()
-		&& pxr::UsdPrim{ UsdPrim }.IsA<pxr::UsdGeomXformable>() // UsdGeomCurves is also an Xformable
-		&& UE::USDPrimViewModel::Private::CanAddOrRemoveCustomSchema( UsdPrim, UnrealIdentifiers::GroomAPI, bAddSchema );
+	return UsdUtils::CanRemoveSchema( UsdPrim, UnrealToUsd::ConvertToken( *SchemaName.ToString() ).Get() );
 #else
 	return false;
 #endif // #if USE_USD_SDK
@@ -492,3 +334,5 @@ void FUsdPrimViewModel::ClearPayloads()
 	Payloads.ClearPayloads();
 #endif // #if USE_USD_SDK
 }
+
+#undef LOCTEXT_NAMESPACE
\ No newline at end of file
diff --git a/Engine/Plugins/Importers/USDImporter/Source/USDStageEditorViewModels/Private/USDReferencesViewModel.cpp b/Engine/Plugins/Importers/USDImporter/Source/USDStageEditorViewModels/Private/USDReferencesViewModel.cpp
--- a/Engine/Plugins/Importers/USDImporter/Source/USDStageEditorViewModels/Private/USDReferencesViewModel.cpp
+++ b/Engine/Plugins/Importers/USDImporter/Source/USDStageEditorViewModels/Private/USDReferencesViewModel.cpp
@@ -23,7 +23,7 @@ void FUsdReferencesViewModel::UpdateReferences( const UE::FUsdStageWeak& UsdStag
 {
 	References.Reset();
 
-	if ( !UsdStage )
+	if ( !UsdStage || !PrimPath || FString{ PrimPath }.IsEmpty() || UE::FSdfPath{ PrimPath }.IsAbsoluteRootPath() )
 	{
 		return;
 	}
diff --git a/Engine/Plugins/Importers/USDImporter/Source/USDStageEditorViewModels/Private/USDVariantSetsViewModel.cpp b/Engine/Plugins/Importers/USDImporter/Source/USDStageEditorViewModels/Private/USDVariantSetsViewModel.cpp
--- a/Engine/Plugins/Importers/USDImporter/Source/USDStageEditorViewModels/Private/USDVariantSetsViewModel.cpp
+++ b/Engine/Plugins/Importers/USDImporter/Source/USDStageEditorViewModels/Private/USDVariantSetsViewModel.cpp
@@ -26,7 +26,7 @@ FUsdVariantSetViewModel::FUsdVariantSetViewModel( FUsdVariantSetsViewModel* InOw
 void FUsdVariantSetViewModel::SetVariantSelection( const TSharedPtr< FString >& InVariantSelection )
 {
 #if USE_USD_SDK
-	if ( !Owner->UsdPrim )
+	if ( !Owner || !Owner->UsdStage )
 	{
 		return;
 	}
@@ -41,31 +41,33 @@ void FUsdVariantSetViewModel::SetVariantSelection( const TSharedPtr< FString >&
 
 	FScopedUsdAllocs UsdAllocs;
 
-	std::string UsdVariantSelection;
+	std::string UsdVariantSetName = UnrealToUsd::ConvertString( *SetName ).Get();
+	std::string UsdVariantSelection = VariantSelection
+		? UnrealToUsd::ConvertString( *( *VariantSelection ) ).Get()
+		: std::string{};
 
-	if ( VariantSelection )
+	if ( const UE::FUsdPrim& Prim = Owner->UsdStage.GetPrimAtPath( UE::FSdfPath{ *Owner->PrimPath } ) )
 	{
-		UsdVariantSelection = UnrealToUsd::ConvertString( *(*VariantSelection ) ).Get();
+		pxr::UsdVariantSets UsdVariantSets = pxr::UsdPrim{ Prim }.GetVariantSets();
+		UsdVariantSets.SetSelection( UsdVariantSetName, UsdVariantSelection );
 	}
 
-	pxr::UsdVariantSets UsdVariantSets = pxr::UsdPrim( Owner->UsdPrim ).GetVariantSets();
-	UsdVariantSets.SetSelection( UnrealToUsd::ConvertString( *SetName ).Get(), UsdVariantSelection );
 #endif // #if USE_USD_SDK
 }
 
-void FUsdVariantSetsViewModel::UpdateVariantSets( const UE::FUsdStageWeak& InUsdStage, const TCHAR* PrimPath )
+void FUsdVariantSetsViewModel::UpdateVariantSets( const UE::FUsdStageWeak& InUsdStage, const TCHAR* InPrimPath )
 {
 #if USE_USD_SDK
 	VariantSets.Reset();
 	UsdStage = InUsdStage;
+	PrimPath = InPrimPath;
 
-	if ( !UsdStage )
+	if ( !UsdStage || PrimPath.IsEmpty() || UE::FSdfPath{ InPrimPath }.IsAbsoluteRootPath() )
 	{
 		return;
 	}
 
-	UsdPrim = UsdStage.GetPrimAtPath( UE::FSdfPath( PrimPath ) );
-
+	UE::FUsdPrim UsdPrim = UsdStage.GetPrimAtPath( UE::FSdfPath( *PrimPath ) );
 	if ( !UsdPrim )
 	{
 		return;
@@ -83,7 +85,7 @@ void FUsdVariantSetsViewModel::UpdateVariantSets( const UE::FUsdStageWeak& InUsd
 		FUsdVariantSetViewModel VariantSet( this );
 		VariantSet.SetName = UsdToUnreal::ConvertString( UsdVariantSetName.c_str() );
 
-		pxr::UsdVariantSet UsdVariantSet =  pxr::UsdPrim( UsdPrim ).GetVariantSet( UsdVariantSetName.c_str() );
+		pxr::UsdVariantSet UsdVariantSet = pxr::UsdPrim( UsdPrim ).GetVariantSet( UsdVariantSetName.c_str() );
 		VariantSet.VariantSelection = MakeSharedUnreal< FString >( UsdToUnreal::ConvertString( UsdVariantSet.GetVariantSelection().c_str() ) );
 
 		std::vector< std::string > VariantNames = UsdVariantSet.GetVariantNames();
@@ -96,4 +98,4 @@ void FUsdVariantSetsViewModel::UpdateVariantSets( const UE::FUsdStageWeak& InUsd
 		VariantSets.Add( MakeSharedUnreal< FUsdVariantSetViewModel >( MoveTemp( VariantSet ) ) );
 	}
 #endif // #if USE_USD_SDK
-}
\ No newline at end of file
+}
diff --git a/Engine/Plugins/Importers/USDImporter/Source/USDStageEditorViewModels/Public/USDIntegrationsViewModel.h b/Engine/Plugins/Importers/USDImporter/Source/USDStageEditorViewModels/Public/USDIntegrationsViewModel.h
new file mode 100644
--- /dev/null
+++ b/Engine/Plugins/Importers/USDImporter/Source/USDStageEditorViewModels/Public/USDIntegrationsViewModel.h
@@ -0,0 +1,21 @@
+// Copyright Epic Games, Inc. All Rights Reserved.
+
+#pragma once
+
+#include "UsdWrappers/UsdAttribute.h"
+#include "UsdWrappers/UsdStage.h"
+
+#include "CoreMinimal.h"
+#include "Templates/SharedPointer.h"
+
+class USDSTAGEEDITORVIEWMODELS_API FUsdIntegrationsViewModel
+{
+public:
+	void UpdateAttributes( const UE::FUsdStageWeak& UsdStage, const TCHAR* PrimPath );
+
+public:
+	TArray< TSharedPtr<UE::FUsdAttribute> > Attributes;
+
+	FString PrimPath;
+    UE::FUsdStageWeak UsdStage;
+};
\ No newline at end of file
diff --git a/Engine/Plugins/Importers/USDImporter/Source/USDStageEditorViewModels/Public/USDPrimAttributesViewModel.h b/Engine/Plugins/Importers/USDImporter/Source/USDStageEditorViewModels/Public/USDPrimAttributesViewModel.h
--- a/Engine/Plugins/Importers/USDImporter/Source/USDStageEditorViewModels/Public/USDPrimAttributesViewModel.h
+++ b/Engine/Plugins/Importers/USDImporter/Source/USDStageEditorViewModels/Public/USDPrimAttributesViewModel.h
@@ -2,13 +2,12 @@
 
 #pragma once
 
-#include "Misc/TVariant.h"
-#include "Templates/SharedPointer.h"
-
 #include "USDValueConversion.h"
 
 #include "UsdWrappers/UsdStage.h"
 
+#include "Templates/SharedPointer.h"
+
 class FUsdPrimAttributesViewModel;
 
 class USDSTAGEEDITORVIEWMODELS_API FUsdPrimAttributeViewModel : public TSharedFromThis< FUsdPrimAttributeViewModel >
@@ -40,12 +39,12 @@ public:
 
 	void SetPrimAttribute( const FString& AttributeName, const UsdUtils::FConvertedVtValue& Value );
 
-	void Refresh( const TCHAR* InPrimPath, float TimeCode );
+	void Refresh( const UE::FUsdStageWeak& UsdStage, const TCHAR* PrimPath, float TimeCode );
 
 public:
-	UE::FUsdStageWeak UsdStage;
 	TArray< TSharedPtr< FUsdPrimAttributeViewModel > > PrimAttributes;
 
 private:
+	UE::FUsdStageWeak UsdStage;
 	FString PrimPath;
 };
diff --git a/Engine/Plugins/Importers/USDImporter/Source/USDStageEditorViewModels/Public/USDPrimViewModel.h b/Engine/Plugins/Importers/USDImporter/Source/USDStageEditorViewModels/Public/USDPrimViewModel.h
--- a/Engine/Plugins/Importers/USDImporter/Source/USDStageEditorViewModels/Public/USDPrimViewModel.h
+++ b/Engine/Plugins/Importers/USDImporter/Source/USDStageEditorViewModels/Public/USDPrimViewModel.h
@@ -3,12 +3,14 @@
 #pragma once
 
 #include "USDTreeItemViewModel.h"
-#include "Templates/SharedPointer.h"
-#include "Internationalization/Text.h"
 
+#include "UsdWrappers/ForwardDeclarations.h"
 #include "UsdWrappers/UsdStage.h"
 #include "UsdWrappers/UsdPrim.h"
 
+#include "Templates/SharedPointer.h"
+#include "Internationalization/Text.h"
+
 using FUsdPrimViewModelRef = TSharedRef< class FUsdPrimViewModel >;
 using FUsdPrimViewModelPtr = TSharedPtr< class FUsdPrimViewModel >;
 
@@ -19,7 +21,7 @@ public:
 	FText GetType() const { return Type; }
 	bool HasPayload() const { return bHasPayload; }
 	bool IsLoaded() const { return bIsLoaded; }
-	bool HasCompositionArcs() const {return bHasCompositionArcs; }
+	bool HasCompositionArcs() const { return bHasCompositionArcs; }
 	bool IsVisible() const { return bIsVisible; }
 
 	FText Name;
@@ -33,9 +35,7 @@ public:
 class USDSTAGEEDITORVIEWMODELS_API FUsdPrimViewModel : public IUsdTreeViewItem
 {
 public:
-	FUsdPrimViewModel( FUsdPrimViewModel* InParentItem, const UE::FUsdStageWeak& InUsdStage, const UE::FUsdPrim& InUsdPrim );
-
-	FUsdPrimViewModel( FUsdPrimViewModel* InParentItem, const UE::FUsdStageWeak& InUsdStage );
+	FUsdPrimViewModel( FUsdPrimViewModel* InParentItem, const UE::FUsdStageWeak& InUsdStage, const UE::FUsdPrim& InPrim = {} );
 
 	TArray< FUsdPrimViewModelRef >& UpdateChildren();
 
@@ -47,23 +47,10 @@ public:
 	void ToggleVisibility();
 	void TogglePayload();
 
-	void SetUpLiveLink();
-	bool CanSetUpLiveLink() const;
-
-	void RemoveLiveLink();
-	bool CanRemoveLiveLink() const;
-
-	void SetUpControlRig();
-	bool CanSetUpControlRig() const;
-
-	void RemoveControlRig();
-	bool CanRemoveControlRig() const;
-
-	void ApplyGroomSchema();
-	bool CanApplyGroomSchema() const;
-
-	void RemoveGroomSchema();
-	bool CanRemoveGroomSchema() const;
+	void ApplySchema( FName SchemaName );
+	bool CanApplySchema( FName SchemaName ) const;
+	void RemoveSchema( FName SchemaName );
+	bool CanRemoveSchema( FName SchemaName ) const;
 
 	// Returns true if this prim has at least one spec on the stage's local layer stack
 	bool HasSpecsOnLocalLayer() const;
@@ -83,8 +70,8 @@ public:
 public:
 	UE::FUsdStageWeak UsdStage;
 	UE::FUsdPrim UsdPrim;
-	FUsdPrimViewModel* ParentItem;
 
+	FUsdPrimViewModel* ParentItem;
 	TArray< FUsdPrimViewModelRef > Children;
 
 	TSharedRef< FUsdPrimModel > RowData; // Data model
diff --git a/Engine/Plugins/Importers/USDImporter/Source/USDStageEditorViewModels/Public/USDVariantSetsViewModel.h b/Engine/Plugins/Importers/USDImporter/Source/USDStageEditorViewModels/Public/USDVariantSetsViewModel.h
--- a/Engine/Plugins/Importers/USDImporter/Source/USDStageEditorViewModels/Public/USDVariantSetsViewModel.h
+++ b/Engine/Plugins/Importers/USDImporter/Source/USDStageEditorViewModels/Public/USDVariantSetsViewModel.h
@@ -32,7 +32,7 @@ public:
 
 public:
 	UE::FUsdStageWeak UsdStage;
-	UE::FUsdPrim UsdPrim;
+	FString PrimPath;
 
 	TArray< TSharedPtr< FUsdVariantSetViewModel > > VariantSets;
 };
diff --git a/Engine/Plugins/Importers/USDImporter/Source/USDUtilities/Private/USDConversionUtils.cpp b/Engine/Plugins/Importers/USDImporter/Source/USDUtilities/Private/USDConversionUtils.cpp
--- a/Engine/Plugins/Importers/USDImporter/Source/USDUtilities/Private/USDConversionUtils.cpp
+++ b/Engine/Plugins/Importers/USDImporter/Source/USDUtilities/Private/USDConversionUtils.cpp
@@ -378,6 +378,16 @@ bool UsdUtils::HasCompositionArcs( const pxr::UsdPrim& Prim )
 	return Prim.HasAuthoredReferences() || Prim.HasPayload() || Prim.HasAuthoredInherits() || Prim.HasAuthoredSpecializes() || Prim.HasVariantSets();
 }
 
+bool UsdUtils::HasCompositionArcs( const pxr::SdfPrimSpecHandle& PrimSpec )
+{
+	if ( !PrimSpec || !PrimSpec->GetActive() )
+	{
+		return false;
+	}
+
+	return PrimSpec->HasReferences() || PrimSpec->HasPayloads() || PrimSpec->HasInheritPaths() || PrimSpec->HasSpecializes() || PrimSpec->HasVariantSetNames();
+}
+
 UClass* UsdUtils::GetActorTypeForPrim( const pxr::UsdPrim& Prim )
 {
 	// If we have this attribute and a valid child camera prim then we'll assume
diff --git a/Engine/Plugins/Importers/USDImporter/Source/USDUtilities/Private/USDIntegrationUtils.cpp b/Engine/Plugins/Importers/USDImporter/Source/USDUtilities/Private/USDIntegrationUtils.cpp
--- a/Engine/Plugins/Importers/USDImporter/Source/USDUtilities/Private/USDIntegrationUtils.cpp
+++ b/Engine/Plugins/Importers/USDImporter/Source/USDUtilities/Private/USDIntegrationUtils.cpp
@@ -4,6 +4,7 @@
 
 #include "UnrealUSDWrapper.h"
 #include "USDAttributeUtils.h"
+#include "USDLog.h"
 #include "USDTypesConversion.h"
 #include "UsdWrappers/SdfChangeBlock.h"
 #include "UsdWrappers/UsdAttribute.h"
@@ -14,28 +15,84 @@
 
 #if USE_USD_SDK
 #include "USDIncludesStart.h"
+	#include "pxr/usd/sdf/primSpec.h"
 	#include "pxr/usd/usd/attribute.h"
 	#include "pxr/usd/usd/prim.h"
+	#include "pxr/usd/usd/tokens.h"
+	#include "pxr/usd/usdGeom/xformable.h"
+	#include "pxr/usd/usdSkel/root.h"
 #include "USDIncludesEnd.h"
 
 #define LOCTEXT_NAMESPACE "USDIntegrationUtils"
 
-bool UsdUtils::PrimHasSchema( const pxr::UsdPrim& Prim, const pxr::TfToken& SchemaToken )
+namespace UE::UsdIntegrationUtils::Private
 {
-	if ( !Prim )
+	bool CanApplyOrRemoveSchema( pxr::UsdPrim Prim, pxr::TfToken SchemaName, bool bApplySchema )
 	{
-		return false;
-	}
+		if ( !Prim || Prim.IsPseudoRoot() )
+		{
+			return false;
+		}
 
-	FScopedUsdAllocs Allocs;
+		FScopedUsdAllocs UsdAllocs;
 
-	pxr::TfType Schema = pxr::UsdSchemaRegistry::GetTypeFromSchemaTypeName( SchemaToken );
-	ensure( static_cast<bool>( Schema ) );
+		pxr::TfType Schema = pxr::UsdSchemaRegistry::GetTypeFromSchemaTypeName( SchemaName );
+		if ( !ensure( static_cast< bool >( Schema ) ) )
+		{
+			return false;
+		}
 
-	return Prim.HasAPI( Schema );
+		// Check if the schema is compatible with this prim
+		if ( bApplySchema && !Prim.CanApplyAPI( Schema ) )
+		{
+			return false;
+		}
+
+		pxr::UsdStageRefPtr Stage = Prim.GetStage();
+		pxr::SdfLayerRefPtr EditTarget = Stage->GetEditTarget().GetLayer();
+		if ( !Stage || !EditTarget )
+		{
+			return false;
+		}
+
+		bool bAlreadyHasSchema = false;
+		if ( EditTarget == Stage->GetRootLayer() )
+		{
+			bAlreadyHasSchema = Prim.HasAPI( Schema );
+		}
+		else
+		{
+			// In the future we'll have better layer editing facilities. For now, lets make sure that
+			// we only show that we can setup/remove a schema for a prim if it can be done *at that
+			// particular edit target*, which is what will happen when adding/removing anyway.
+			// For example, a user could:
+			//  - Add the schema on a sublayer;
+			//  - Remove the schema on the root layer;
+			//  - Select the sublayer as the edit target;
+			// If we always just queried the existing stage, the user would see that they can add the schema
+			// (because the composed prim does not have the schema) but adding the schema would do nothing,
+			// as the sublayer prim spec already has the schema...
+
+			pxr::SdfPrimSpecHandle PrimSpecOnLayer = EditTarget->GetPrimAtPath( Prim.GetPath() );
+			if ( !PrimSpecOnLayer )
+			{
+				// We can always just make a new 'over' on this layer
+				return true;
+			}
+
+			// We can only add the schema if its not present on this layer, but we should let USD sort the
+			// Ops out because we could have Ops to add/delete/prepend/append the same schema on the same prim...
+			pxr::SdfListOp<pxr::TfToken> Ops = PrimSpecOnLayer->GetInfo( pxr::UsdTokens->apiSchemas ).Get<pxr::SdfTokenListOp>();
+			std::vector<pxr::TfToken> AppliedOps;
+			Ops.ApplyOperations( &AppliedOps );
+			bAlreadyHasSchema = std::find( AppliedOps.begin(), AppliedOps.end(), SchemaName ) != AppliedOps.end();
+		}
+
+		return bApplySchema != bAlreadyHasSchema;
+	}
 }
 
-bool UsdUtils::ApplySchema( const pxr::UsdPrim& Prim, const pxr::TfToken& SchemaToken )
+bool UsdUtils::PrimHasSchema( const pxr::UsdPrim& Prim, const pxr::TfToken& SchemaToken )
 {
 	if ( !Prim )
 	{
@@ -47,23 +104,16 @@ bool UsdUtils::ApplySchema( const pxr::UsdPrim& Prim, const pxr::TfToken& Schema
 	pxr::TfType Schema = pxr::UsdSchemaRegistry::GetTypeFromSchemaTypeName( SchemaToken );
 	ensure( static_cast<bool>( Schema ) );
 
-	return Prim.ApplyAPI( Schema );
+	return Prim.HasAPI( Schema );
 }
 
-bool UsdUtils::PrimHasLiveLinkSchema( const pxr::UsdPrim& Prim )
+bool UsdUtils::CanApplySchema( pxr::UsdPrim Prim, pxr::TfToken SchemaName )
 {
-	if ( !Prim )
-	{
-		return false;
-	}
-
-    FScopedUsdAllocs Allocs;
-
-    pxr::TfType Schema = pxr::UsdSchemaRegistry::GetTypeFromSchemaTypeName( UnrealIdentifiers::LiveLinkAPI );
-    return Prim.HasAPI( Schema );
+	const bool bApplySchema = true;
+	return UE::UsdIntegrationUtils::Private::CanApplyOrRemoveSchema( Prim, SchemaName, bApplySchema );
 }
 
-bool UsdUtils::PrimHasControlRigSchema( const pxr::UsdPrim& Prim )
+bool UsdUtils::ApplySchema( const pxr::UsdPrim& Prim, const pxr::TfToken& SchemaToken )
 {
 	if ( !Prim )
 	{
@@ -72,27 +122,19 @@ bool UsdUtils::PrimHasControlRigSchema( const pxr::UsdPrim& Prim )
 
 	FScopedUsdAllocs Allocs;
 
-	pxr::TfType Schema = pxr::UsdSchemaRegistry::GetTypeFromSchemaTypeName( UnrealIdentifiers::ControlRigAPI );
-	return Prim.HasAPI( Schema );
+	pxr::TfType Schema = pxr::UsdSchemaRegistry::GetTypeFromSchemaTypeName( SchemaToken );
+	ensure( static_cast<bool>( Schema ) );
+
+	return Prim.ApplyAPI( Schema );
 }
 
-bool UsdUtils::ApplyControlRigSchema( const pxr::UsdPrim& Prim )
+bool UsdUtils::CanRemoveSchema( pxr::UsdPrim Prim, pxr::TfToken SchemaName )
 {
-	if ( !Prim )
-	{
-		return false;
-	}
-
-	FScopedUsdAllocs Allocs;
-
-	pxr::TfType Schema = pxr::UsdSchemaRegistry::GetTypeFromSchemaTypeName( UnrealIdentifiers::ControlRigAPI );
-	ensure( static_cast< bool >( Schema ) );
-
-	ensure( Prim.ApplyAPI( Schema ) );
-	return true;
+	const bool bApplySchema = false;
+	return UE::UsdIntegrationUtils::Private::CanApplyOrRemoveSchema( Prim, SchemaName, bApplySchema );
 }
 
-bool UsdUtils::ApplyLiveLinkSchema( const pxr::UsdPrim& Prim )
+bool UsdUtils::RemoveSchema( const pxr::UsdPrim& Prim, const pxr::TfToken& SchemaToken )
 {
 	if ( !Prim )
 	{
@@ -101,17 +143,16 @@ bool UsdUtils::ApplyLiveLinkSchema( const pxr::UsdPrim& Prim )
 
 	FScopedUsdAllocs Allocs;
 
-	pxr::TfType Schema = pxr::UsdSchemaRegistry::GetTypeFromSchemaTypeName( UnrealIdentifiers::LiveLinkAPI );
+	pxr::TfType Schema = pxr::UsdSchemaRegistry::GetTypeFromSchemaTypeName( SchemaToken );
 	ensure( static_cast< bool >( Schema ) );
 
-	ensure( Prim.ApplyAPI( Schema ) );
-	return true;
+	return Prim.RemoveAPI( Schema );
 }
 
 void UnrealToUsd::ConvertLiveLinkProperties( const UActorComponent* InComponent, pxr::UsdPrim& InOutPrim )
 {
 #if WITH_EDITOR
-	if ( !InOutPrim || !UsdUtils::PrimHasLiveLinkSchema( InOutPrim ) )
+	if ( !InOutPrim || !UsdUtils::PrimHasSchema( InOutPrim, UnrealIdentifiers::LiveLinkAPI ) )
 	{
 		return;
 	}
diff --git a/Engine/Plugins/Importers/USDImporter/Source/USDUtilities/Public/USDConversionUtils.h b/Engine/Plugins/Importers/USDImporter/Source/USDUtilities/Public/USDConversionUtils.h
--- a/Engine/Plugins/Importers/USDImporter/Source/USDUtilities/Public/USDConversionUtils.h
+++ b/Engine/Plugins/Importers/USDImporter/Source/USDUtilities/Public/USDConversionUtils.h
@@ -34,8 +34,11 @@ PXR_NAMESPACE_OPEN_SCOPE
 
 	class UsdStage;
 	template< typename T > class TfRefPtr;
-
 	using UsdStageRefPtr = TfRefPtr< UsdStage >;
+
+	class SdfPrimSpec;
+	template <class T> class SdfHandle;
+	using SdfPrimSpecHandle = SdfHandle< SdfPrimSpec >;
 PXR_NAMESPACE_CLOSE_SCOPE
 
 #endif // #if USE_USD_SDK
@@ -102,6 +105,7 @@ namespace UsdUtils
 	USDUTILITIES_API int32 GetUsdStageNumFrames( const pxr::UsdStageRefPtr& Stage );
 
 	USDUTILITIES_API bool HasCompositionArcs( const pxr::UsdPrim& Prim );
+	USDUTILITIES_API bool HasCompositionArcs( const pxr::SdfPrimSpecHandle& PrimSpec );
 
 	USDUTILITIES_API UClass* GetActorTypeForPrim( const pxr::UsdPrim& Prim );
 	USDUTILITIES_API UClass* GetComponentTypeForPrim( const pxr::UsdPrim& Prim );
diff --git a/Engine/Plugins/Importers/USDImporter/Source/USDUtilities/Public/USDIntegrationUtils.h b/Engine/Plugins/Importers/USDImporter/Source/USDUtilities/Public/USDIntegrationUtils.h
--- a/Engine/Plugins/Importers/USDImporter/Source/USDUtilities/Public/USDIntegrationUtils.h
+++ b/Engine/Plugins/Importers/USDImporter/Source/USDUtilities/Public/USDIntegrationUtils.h
@@ -22,27 +22,20 @@ class ULiveLinkComponentController;
 namespace UsdUtils
 {
 #if USE_USD_SDK
-
 	// Returns true if Prim has the given Schema
 	USDUTILITIES_API bool PrimHasSchema( const pxr::UsdPrim& Prim, const pxr::TfToken& Schema );
 
+	// Returns true if we can apply a single-apply API schema named "SchemaName" to Prim
+	USDUTILITIES_API bool CanApplySchema( pxr::UsdPrim Prim, pxr::TfToken SchemaName );
+
 	// Apply the given Schema to the Prim. Returns true if the application was successful
 	USDUTILITIES_API bool ApplySchema( const pxr::UsdPrim& Prim, const pxr::TfToken& Schema );
 
-	// Returns true in case Prim has our custom "LiveLinkAPI" schema
-	USDUTILITIES_API bool PrimHasLiveLinkSchema( const pxr::UsdPrim& Prim );
-
-	// Returns true in case Prim has our custom "ControlRigAPI" schema
-	USDUTILITIES_API bool PrimHasControlRigSchema( const pxr::UsdPrim& Prim );
-
-	// Apply the control rig schema to the prim, creates all the required attributes and sets default values
-	// if they're otherwise not authored
-	USDUTILITIES_API bool ApplyControlRigSchema( const pxr::UsdPrim& Prim );
-
-	// Apply the live link schema to the prim, creates all the required attributes and sets default values
-	// if they're otherwise not authored
-	USDUTILITIES_API bool ApplyLiveLinkSchema( const pxr::UsdPrim& Prim );
+	// Returns true if we can remove a single-apply API schema named "SchemaName" from Prim
+	USDUTILITIES_API bool CanRemoveSchema( pxr::UsdPrim Prim, pxr::TfToken SchemaName );
 
+	// Remove the given Schema from the Prim. Returns true if was successful
+	USDUTILITIES_API bool RemoveSchema( const pxr::UsdPrim& Prim, const pxr::TfToken& Schema );
 #endif // USE_USD_SDK
 }
 
diff --git a/Engine/Plugins/Importers/USDImporter/Source/UnrealUSDWrapper/Private/UsdWrappers/SdfAttributeSpec.cpp b/Engine/Plugins/Importers/USDImporter/Source/UnrealUSDWrapper/Private/UsdWrappers/SdfAttributeSpec.cpp
new file mode 100644
--- /dev/null
+++ b/Engine/Plugins/Importers/USDImporter/Source/UnrealUSDWrapper/Private/UsdWrappers/SdfAttributeSpec.cpp
@@ -0,0 +1,174 @@
+// Copyright Epic Games, Inc. All Rights Reserved.
+
+#include "UsdWrappers/SdfAttributeSpec.h"
+
+#include "UsdWrappers/VtValue.h"
+
+#include "USDMemory.h"
+
+#if USE_USD_SDK
+
+#include "USDIncludesStart.h"
+	#include "pxr/usd/usd/attribute.h"
+#include "USDIncludesEnd.h"
+
+#endif // #if USE_USD_SDK
+
+namespace UE
+{
+	namespace Internal
+	{
+		class FSdfAttributeSpecImpl
+		{
+		public:
+			FSdfAttributeSpecImpl() = default;
+
+#if USE_USD_SDK
+			explicit FSdfAttributeSpecImpl( const pxr::SdfAttributeSpecHandle& InSdfAttributeSpec )
+				: PxrSdfAttributeSpec( InSdfAttributeSpec )
+			{
+			}
+
+			explicit FSdfAttributeSpecImpl( pxr::SdfAttributeSpecHandle&& InSdfAttributeSpec )
+				: PxrSdfAttributeSpec( MoveTemp( InSdfAttributeSpec ) )
+			{
+			}
+
+			TUsdStore< pxr::SdfAttributeSpecHandle > PxrSdfAttributeSpec;
+#endif // #if USE_USD_SDK
+		};
+	}
+}
+
+namespace UE
+{
+	FSdfAttributeSpec::FSdfAttributeSpec()
+	{
+		FScopedUnrealAllocs UnrealAllocs;
+		Impl =  MakeUnique< Internal::FSdfAttributeSpecImpl >();
+	}
+
+	FSdfAttributeSpec::FSdfAttributeSpec( const FSdfAttributeSpec& Other )
+	{
+#if USE_USD_SDK
+		FScopedUnrealAllocs UnrealAllocs;
+		Impl = MakeUnique< Internal::FSdfAttributeSpecImpl >( Other.Impl->PxrSdfAttributeSpec.Get() );
+#endif // #if USE_USD_SDK
+	}
+
+	FSdfAttributeSpec::FSdfAttributeSpec( FSdfAttributeSpec&& Other ) = default;
+
+	FSdfAttributeSpec::~FSdfAttributeSpec()
+	{
+		FScopedUnrealAllocs UnrealAllocs;
+		Impl.Reset();
+	}
+
+	FSdfAttributeSpec& FSdfAttributeSpec::operator=( const FSdfAttributeSpec& Other )
+	{
+#if USE_USD_SDK
+		FScopedUnrealAllocs UnrealAllocs;
+		Impl = MakeUnique< Internal::FSdfAttributeSpecImpl >(  Other.Impl->PxrSdfAttributeSpec.Get() );
+#endif // #if USE_USD_SDK
+		return *this;
+	}
+
+	FSdfAttributeSpec& FSdfAttributeSpec::operator=( FSdfAttributeSpec&& Other )
+	{
+		FScopedUnrealAllocs UnrealAllocs;
+		Impl = MoveTemp( Other.Impl );
+
+		return *this;
+	}
+
+	bool FSdfAttributeSpec::operator==( const FSdfAttributeSpec& Other ) const
+	{
+#if USE_USD_SDK
+		return Impl->PxrSdfAttributeSpec.Get() == Other.Impl->PxrSdfAttributeSpec.Get();
+#else
+		return false;
+#endif // #if USE_USD_SDK
+	}
+
+	bool FSdfAttributeSpec::operator!=( const FSdfAttributeSpec& Other ) const
+	{
+		return !( *this == Other );
+	}
+
+	FSdfAttributeSpec::operator bool() const
+	{
+#if USE_USD_SDK
+		return (bool)Impl->PxrSdfAttributeSpec.Get();
+#else
+		return false;
+#endif // #if USE_USD_SDK
+	}
+
+#if USE_USD_SDK
+	FSdfAttributeSpec::FSdfAttributeSpec( const pxr::SdfAttributeSpecHandle& InSdfAttributeSpec )
+	{
+		FScopedUnrealAllocs UnrealAllocs;
+		Impl =  MakeUnique< Internal::FSdfAttributeSpecImpl >( InSdfAttributeSpec );
+	}
+
+	FSdfAttributeSpec::FSdfAttributeSpec( pxr::SdfAttributeSpecHandle&& InSdfAttributeSpec )
+	{
+		FScopedUnrealAllocs UnrealAllocs;
+		Impl = MakeUnique< Internal::FSdfAttributeSpecImpl >( MoveTemp( InSdfAttributeSpec ) );
+	}
+
+	FSdfAttributeSpec& FSdfAttributeSpec::operator=( const pxr::SdfAttributeSpecHandle& InSdfAttributeSpec )
+	{
+		FScopedUnrealAllocs UnrealAllocs;
+		Impl = MakeUnique< Internal::FSdfAttributeSpecImpl >( InSdfAttributeSpec );
+		return *this;
+	}
+
+	FSdfAttributeSpec& FSdfAttributeSpec::operator=( pxr::SdfAttributeSpecHandle&& InSdfAttributeSpec )
+	{
+		FScopedUnrealAllocs UnrealAllocs;
+		Impl = MakeUnique< Internal::FSdfAttributeSpecImpl >( MoveTemp( InSdfAttributeSpec ) );
+		return *this;
+	}
+
+	FSdfAttributeSpec::operator pxr::SdfAttributeSpecHandle&()
+	{
+		return Impl->PxrSdfAttributeSpec.Get();
+	}
+
+	FSdfAttributeSpec::operator const pxr::SdfAttributeSpecHandle&() const
+	{
+		return Impl->PxrSdfAttributeSpec.Get();
+	}
+#endif // #if USE_USD_SDK
+
+	UE::FVtValue FSdfAttributeSpec::GetDefaultValue() const
+	{
+#if USE_USD_SDK
+		FScopedUsdAllocs UsdAllocs;
+		return UE::FVtValue{ Impl->PxrSdfAttributeSpec.Get()->GetDefaultValue() };
+#else
+		return {};
+#endif // #if USE_USD_SDK
+	}
+
+	bool FSdfAttributeSpec::SetDefaultValue( const UE::FVtValue& Value )
+	{
+#if USE_USD_SDK
+		FScopedUsdAllocs UsdAllocs;
+		return Impl->PxrSdfAttributeSpec.Get()->SetDefaultValue( Value.GetUsdValue() );
+#else
+		return false;
+#endif // #if USE_USD_SDK
+	}
+
+	FString FSdfAttributeSpec::GetName() const
+	{
+#if USE_USD_SDK
+		FScopedUsdAllocs UsdAllocs;
+		return ANSI_TO_TCHAR( Impl->PxrSdfAttributeSpec.Get()->GetName().c_str() );
+#else
+		return {};
+#endif // #if USE_USD_SDK
+	}
+}
diff --git a/Engine/Plugins/Importers/USDImporter/Source/UnrealUSDWrapper/Private/UsdWrappers/SdfLayer.cpp b/Engine/Plugins/Importers/USDImporter/Source/UnrealUSDWrapper/Private/UsdWrappers/SdfLayer.cpp
--- a/Engine/Plugins/Importers/USDImporter/Source/UnrealUSDWrapper/Private/UsdWrappers/SdfLayer.cpp
+++ b/Engine/Plugins/Importers/USDImporter/Source/UnrealUSDWrapper/Private/UsdWrappers/SdfLayer.cpp
@@ -3,6 +3,7 @@
 #include "UsdWrappers/SdfLayer.h"
 
 #include "UsdWrappers/SdfPath.h"
+#include "UsdWrappers/SdfPrimSpec.h"
 
 #include "USDMemory.h"
 
@@ -11,6 +12,7 @@
 #include "USDIncludesStart.h"
 	#include "pxr/usd/sdf/layer.h"
 	#include "pxr/usd/sdf/layerUtils.h"
+	#include "pxr/usd/sdf/primSpec.h"
 #include "USDIncludesEnd.h"
 
 #endif // #if USE_USD_SDK
@@ -616,6 +618,32 @@ namespace UE
 		return false;
 	}
 
+	template<typename PtrType>
+	FSdfPrimSpec FSdfLayerBase<PtrType>::GetPseudoRoot() const
+	{
+#if USE_USD_SDK
+		if ( const PtrType& Ptr = Impl->GetInner() )
+		{
+			return FSdfPrimSpec{ Ptr->GetPseudoRoot() };
+		}
+#endif // #if USE_USD_SDK
+
+		return FSdfPrimSpec{};
+	}
+
+	template<typename PtrType>
+	FSdfPrimSpec FSdfLayerBase<PtrType>::GetPrimAtPath( const FSdfPath& Path ) const
+	{
+#if USE_USD_SDK
+		if ( const PtrType& Ptr = Impl->GetInner() )
+		{
+			return FSdfPrimSpec{ Ptr->GetPrimAtPath( Path ) };
+		}
+#endif // #if USE_USD_SDK
+
+		return FSdfPrimSpec{};
+	}
+
 	template<typename PtrType>
 	TSet< double > FSdfLayerBase<PtrType>::ListTimeSamplesForPath( const FSdfPath& Path ) const
 	{
diff --git a/Engine/Plugins/Importers/USDImporter/Source/UnrealUSDWrapper/Private/UsdWrappers/SdfPath.cpp b/Engine/Plugins/Importers/USDImporter/Source/UnrealUSDWrapper/Private/UsdWrappers/SdfPath.cpp
--- a/Engine/Plugins/Importers/USDImporter/Source/UnrealUSDWrapper/Private/UsdWrappers/SdfPath.cpp
+++ b/Engine/Plugins/Importers/USDImporter/Source/UnrealUSDWrapper/Private/UsdWrappers/SdfPath.cpp
@@ -256,6 +256,15 @@ namespace UE
 #endif // #if USE_USD_SDK
 	}
 
+	FSdfPath FSdfPath::AppendProperty( FName PropertyName ) const
+	{
+#if USE_USD_SDK
+		return FSdfPath( Impl->PxrSdfPath.Get().AppendProperty( pxr::TfToken( TCHAR_TO_ANSI( *PropertyName.ToString() ) ) ) );
+#else
+		return FSdfPath();
+#endif // #if USE_USD_SDK
+	}
+
 	UE::FSdfPath FSdfPath::StripAllVariantSelections() const
 	{
 #if USE_USD_SDK
diff --git a/Engine/Plugins/Importers/USDImporter/Source/UnrealUSDWrapper/Private/UsdWrappers/SdfPrimSpec.cpp b/Engine/Plugins/Importers/USDImporter/Source/UnrealUSDWrapper/Private/UsdWrappers/SdfPrimSpec.cpp
new file mode 100644
--- /dev/null
+++ b/Engine/Plugins/Importers/USDImporter/Source/UnrealUSDWrapper/Private/UsdWrappers/SdfPrimSpec.cpp
@@ -0,0 +1,231 @@
+// Copyright Epic Games, Inc. All Rights Reserved.
+
+#include "UsdWrappers/SdfPrimSpec.h"
+
+#include "USDMemory.h"
+
+#include "UsdWrappers/SdfAttributeSpec.h"
+#include "UsdWrappers/SdfLayer.h"
+#include "UsdWrappers/SdfPath.h"
+
+#if USE_USD_SDK
+
+#include "USDIncludesStart.h"
+	#include "pxr/usd/sdf/attributeSpec.h"
+	#include "pxr/usd/sdf/primSpec.h"
+#include "USDIncludesEnd.h"
+
+#endif // #if USE_USD_SDK
+
+namespace UE
+{
+	namespace Internal
+	{
+		class FSdfPrimSpecImpl
+		{
+		public:
+			FSdfPrimSpecImpl() = default;
+
+#if USE_USD_SDK
+			explicit FSdfPrimSpecImpl( const pxr::SdfPrimSpecHandle& InSdfPrimSpec )
+				: PxrSdfPrimSpec( InSdfPrimSpec )
+			{
+			}
+
+			explicit FSdfPrimSpecImpl( pxr::SdfPrimSpecHandle&& InSdfPrimSpec )
+				: PxrSdfPrimSpec( MoveTemp( InSdfPrimSpec ) )
+			{
+			}
+
+			TUsdStore< pxr::SdfPrimSpecHandle > PxrSdfPrimSpec;
+#endif // #if USE_USD_SDK
+		};
+	}
+
+	FSdfPrimSpec::FSdfPrimSpec()
+	{
+		FScopedUnrealAllocs UnrealAllocs;
+		Impl = MakeUnique< Internal::FSdfPrimSpecImpl >();
+	}
+
+	FSdfPrimSpec::FSdfPrimSpec( const FSdfPrimSpec& Other )
+	{
+#if USE_USD_SDK
+		FScopedUnrealAllocs UnrealAllocs;
+		Impl = MakeUnique< Internal::FSdfPrimSpecImpl >( Other.Impl->PxrSdfPrimSpec.Get() );
+#endif // #if USE_USD_SDK
+	}
+
+	FSdfPrimSpec::FSdfPrimSpec( FSdfPrimSpec&& Other ) = default;
+
+	FSdfPrimSpec::~FSdfPrimSpec()
+	{
+		FScopedUnrealAllocs UnrealAllocs;
+		Impl.Reset();
+	}
+
+	FSdfPrimSpec& FSdfPrimSpec::operator=( const FSdfPrimSpec& Other )
+	{
+#if USE_USD_SDK
+		FScopedUnrealAllocs UnrealAllocs;
+		Impl = MakeUnique< Internal::FSdfPrimSpecImpl >(  Other.Impl->PxrSdfPrimSpec.Get() );
+#endif // #if USE_USD_SDK
+		return *this;
+	}
+
+	FSdfPrimSpec& FSdfPrimSpec::operator=( FSdfPrimSpec&& Other )
+	{
+		FScopedUnrealAllocs UnrealAllocs;
+		Impl = MoveTemp(  Other.Impl );
+
+		return *this;
+	}
+
+	FSdfPrimSpec::operator bool() const
+	{
+#if USE_USD_SDK
+		return (bool)Impl->PxrSdfPrimSpec.Get();
+#else
+		return false;
+#endif // #if USE_USD_SDK
+	}
+
+	bool FSdfPrimSpec::operator==( const FSdfPrimSpec& Other ) const
+	{
+#if USE_USD_SDK
+		return Impl->PxrSdfPrimSpec.Get() == Other.Impl->PxrSdfPrimSpec.Get();
+#else
+		return false;
+#endif // #if USE_USD_SDK
+	}
+
+	bool FSdfPrimSpec::operator!=( const FSdfPrimSpec& Other ) const
+	{
+		return !( *this == Other );
+	}
+
+#if USE_USD_SDK
+	FSdfPrimSpec::FSdfPrimSpec( const pxr::SdfPrimSpecHandle& InSdfPrimSpec )
+	{
+		FScopedUnrealAllocs UnrealAllocs;
+		Impl =  MakeUnique< Internal::FSdfPrimSpecImpl >( InSdfPrimSpec );
+	}
+
+	FSdfPrimSpec::FSdfPrimSpec( pxr::SdfPrimSpecHandle&& InSdfPrimSpec )
+	{
+		FScopedUnrealAllocs UnrealAllocs;
+		Impl = MakeUnique< Internal::FSdfPrimSpecImpl >( MoveTemp( InSdfPrimSpec ) );
+	}
+
+	FSdfPrimSpec& FSdfPrimSpec::operator=( const pxr::SdfPrimSpecHandle& InSdfPrimSpec )
+	{
+		FScopedUnrealAllocs UnrealAllocs;
+		Impl = MakeUnique< Internal::FSdfPrimSpecImpl >( InSdfPrimSpec );
+		return *this;
+	}
+
+	FSdfPrimSpec& FSdfPrimSpec::operator=( pxr::SdfPrimSpecHandle&& InSdfPrimSpec )
+	{
+		FScopedUnrealAllocs UnrealAllocs;
+		Impl = MakeUnique< Internal::FSdfPrimSpecImpl >( MoveTemp( InSdfPrimSpec ) );
+		return *this;
+	}
+
+	FSdfPrimSpec::operator pxr::SdfPrimSpecHandle&()
+	{
+		return Impl->PxrSdfPrimSpec.Get();
+	}
+
+	FSdfPrimSpec::operator const pxr::SdfPrimSpecHandle&() const
+	{
+		return Impl->PxrSdfPrimSpec.Get();
+	}
+#endif // #if USE_USD_SDK
+
+	ESdfSpecType FSdfPrimSpec::GetSpecType() const
+	{
+#if USE_USD_SDK
+		// Just in case they add another or reorder these
+		static_assert( ( int ) ESdfSpecType::SdfSpecTypeUnknown == ( int ) pxr::SdfSpecType::SdfSpecTypeUnknown );
+		static_assert( ( int ) ESdfSpecType::SdfSpecTypeAttribute == ( int ) pxr::SdfSpecType::SdfSpecTypeAttribute );
+		static_assert( ( int ) ESdfSpecType::SdfSpecTypeConnection == ( int ) pxr::SdfSpecType::SdfSpecTypeConnection );
+		static_assert( ( int ) ESdfSpecType::SdfSpecTypeExpression == ( int ) pxr::SdfSpecType::SdfSpecTypeExpression );
+		static_assert( ( int ) ESdfSpecType::SdfSpecTypeMapper == ( int ) pxr::SdfSpecType::SdfSpecTypeMapper );
+		static_assert( ( int ) ESdfSpecType::SdfSpecTypeMapperArg == ( int ) pxr::SdfSpecType::SdfSpecTypeMapperArg );
+		static_assert( ( int ) ESdfSpecType::SdfSpecTypePrim == ( int ) pxr::SdfSpecType::SdfSpecTypePrim );
+		static_assert( ( int ) ESdfSpecType::SdfSpecTypePseudoRoot == ( int ) pxr::SdfSpecType::SdfSpecTypePseudoRoot );
+		static_assert( ( int ) ESdfSpecType::SdfSpecTypeRelationship == ( int ) pxr::SdfSpecType::SdfSpecTypeRelationship );
+		static_assert( ( int ) ESdfSpecType::SdfSpecTypeRelationshipTarget == ( int ) pxr::SdfSpecType::SdfSpecTypeRelationshipTarget );
+		static_assert( ( int ) ESdfSpecType::SdfSpecTypeVariant == ( int ) pxr::SdfSpecType::SdfSpecTypeVariant );
+		static_assert( ( int ) ESdfSpecType::SdfSpecTypeVariantSet == ( int ) pxr::SdfSpecType::SdfSpecTypeVariantSet );
+		static_assert( ( int ) ESdfSpecType::SdfNumSpecTypes == ( int ) pxr::SdfSpecType::SdfNumSpecTypes );
+
+		return static_cast< ESdfSpecType >( Impl->PxrSdfPrimSpec.Get()->GetSpecType() );
+#else
+		return ESdfSpecType::SdfSpecTypeUnknown;
+#endif // #if USE_USD_SDK
+	}
+
+	FSdfLayerWeak FSdfPrimSpec::GetLayer() const
+	{
+#if USE_USD_SDK
+		return FSdfLayerWeak( Impl->PxrSdfPrimSpec.Get()->GetLayer() );
+#else
+		return FSdfLayerWeak();
+#endif // #if USE_USD_SDK
+	}
+
+	FSdfPath FSdfPrimSpec::GetPath() const
+	{
+#if USE_USD_SDK
+		return FSdfPath( Impl->PxrSdfPrimSpec.Get()->GetPath() );
+#else
+		return FSdfPath();
+#endif // #if USE_USD_SDK
+	}
+
+	FSdfPrimSpec FSdfPrimSpec::GetRealNameParent() const
+	{
+#if USE_USD_SDK
+		return FSdfPrimSpec{ Impl->PxrSdfPrimSpec.Get()->GetRealNameParent() };
+#else
+		return FSdfPrimSpec{};
+#endif // #if USE_USD_SDK
+	}
+
+	bool FSdfPrimSpec::RemoveNameChild( const FSdfPrimSpec& Child )
+	{
+#if USE_USD_SDK
+		return Impl->PxrSdfPrimSpec.Get()->RemoveNameChild( Child );
+#else
+		return false;
+#endif // #if USE_USD_SDK
+	}
+
+	UE::FSdfAttributeSpec FSdfPrimSpec::GetAttributeAtPath( const UE::FSdfPath& Path ) const
+	{
+#if USE_USD_SDK
+		return UE::FSdfAttributeSpec{ Impl->PxrSdfPrimSpec.Get()->GetAttributeAtPath( Path ) };
+#else
+		return UE::FSdfAttributeSpec{};
+#endif // #if USE_USD_SDK
+	}
+
+	FName FSdfPrimSpec::GetTypeName() const
+	{
+#if USE_USD_SDK
+		return FName( ANSI_TO_TCHAR( Impl->PxrSdfPrimSpec.Get()->GetTypeName().GetString().c_str() ) );
+#else
+		return FName();
+#endif // #if USE_USD_SDK
+	}
+
+	FName FSdfPrimSpec::GetName() const
+	{
+#if USE_USD_SDK
+		return FName( ANSI_TO_TCHAR( Impl->PxrSdfPrimSpec.Get()->GetName().c_str() ) );
+#else
+		return FName();
+#endif // #if USE_USD_SDK
+	}
+}
diff --git a/Engine/Plugins/Importers/USDImporter/Source/UnrealUSDWrapper/Public/UsdWrappers/SdfAttributeSpec.h b/Engine/Plugins/Importers/USDImporter/Source/UnrealUSDWrapper/Public/UsdWrappers/SdfAttributeSpec.h
new file mode 100644
--- /dev/null
+++ b/Engine/Plugins/Importers/USDImporter/Source/UnrealUSDWrapper/Public/UsdWrappers/SdfAttributeSpec.h
@@ -0,0 +1,72 @@
+// Copyright Epic Games, Inc. All Rights Reserved.
+
+#pragma once
+
+#include "Templates/UniquePtr.h"
+
+#if USE_USD_SDK
+
+#include "USDIncludesStart.h"
+	#include "pxr/pxr.h"
+#include "USDIncludesEnd.h"
+
+PXR_NAMESPACE_OPEN_SCOPE
+	class SdfAttributeSpec;
+	template <class T> class SdfHandle;
+	using SdfAttributeSpecHandle = SdfHandle< SdfAttributeSpec >;
+PXR_NAMESPACE_CLOSE_SCOPE
+
+#endif // #if USE_USD_SDK
+
+namespace UE
+{
+	class FVtValue;
+
+	namespace Internal
+	{
+		class FSdfAttributeSpecImpl;
+	}
+
+	/**
+	 * Minimal pxr::SdfAttributeSpecHandle wrapper for Unreal that can be used from no-rtti modules.
+	 */
+	class UNREALUSDWRAPPER_API FSdfAttributeSpec
+	{
+	public:
+		FSdfAttributeSpec();
+
+		FSdfAttributeSpec( const FSdfAttributeSpec& Other );
+		FSdfAttributeSpec( FSdfAttributeSpec&& Other );
+		~FSdfAttributeSpec();
+
+		FSdfAttributeSpec& operator=( const FSdfAttributeSpec& Other );
+		FSdfAttributeSpec& operator=( FSdfAttributeSpec&& Other );
+
+		bool operator==( const FSdfAttributeSpec& Other ) const;
+		bool operator!=( const FSdfAttributeSpec& Other ) const;
+
+		explicit operator bool() const;
+
+	// Auto conversion from/to pxr::SdfAttributeSpecHandle
+	public:
+#if USE_USD_SDK
+		explicit FSdfAttributeSpec( const pxr::SdfAttributeSpecHandle& InSdfAttributeSpecHandle );
+		explicit FSdfAttributeSpec( pxr::SdfAttributeSpecHandle&& InSdfAttributeSpecHandle );
+		FSdfAttributeSpec& operator=( const pxr::SdfAttributeSpecHandle& InSdfAttributeSpecHandle );
+		FSdfAttributeSpec& operator=( pxr::SdfAttributeSpecHandle&& InSdfAttributeSpecHandle );
+
+		operator pxr::SdfAttributeSpecHandle&();
+		operator const pxr::SdfAttributeSpecHandle&() const;
+#endif // #if USE_USD_SDK
+
+	// Wrapped pxr::SdfAttributeSpec functions, refer to the USD SDK documentation
+	public:
+		UE::FVtValue GetDefaultValue() const;
+		bool SetDefaultValue( const UE::FVtValue& Value );
+
+		FString GetName() const;
+
+	private:
+		TUniquePtr< Internal::FSdfAttributeSpecImpl > Impl;
+	};
+}
\ No newline at end of file
diff --git a/Engine/Plugins/Importers/USDImporter/Source/UnrealUSDWrapper/Public/UsdWrappers/SdfLayer.h b/Engine/Plugins/Importers/USDImporter/Source/UnrealUSDWrapper/Public/UsdWrappers/SdfLayer.h
--- a/Engine/Plugins/Importers/USDImporter/Source/UnrealUSDWrapper/Public/UsdWrappers/SdfLayer.h
+++ b/Engine/Plugins/Importers/USDImporter/Source/UnrealUSDWrapper/Public/UsdWrappers/SdfLayer.h
@@ -9,6 +9,7 @@
 namespace UE
 {
 	class FSdfPath;
+	class FSdfPrimSpec;
 
 	namespace Internal
 	{
@@ -132,6 +133,9 @@ namespace UE
 
 		bool HasSpec( const FSdfPath& Path ) const;
 
+		FSdfPrimSpec GetPseudoRoot() const;
+		FSdfPrimSpec GetPrimAtPath( const FSdfPath& Path ) const;
+
 		TSet< double > ListTimeSamplesForPath( const FSdfPath& Path ) const;
 		void EraseTimeSample( const FSdfPath& Path, double Time );
 
diff --git a/Engine/Plugins/Importers/USDImporter/Source/UnrealUSDWrapper/Public/UsdWrappers/SdfPath.h b/Engine/Plugins/Importers/USDImporter/Source/UnrealUSDWrapper/Public/UsdWrappers/SdfPath.h
--- a/Engine/Plugins/Importers/USDImporter/Source/UnrealUSDWrapper/Public/UsdWrappers/SdfPath.h
+++ b/Engine/Plugins/Importers/USDImporter/Source/UnrealUSDWrapper/Public/UsdWrappers/SdfPath.h
@@ -74,6 +74,7 @@ namespace UE
 
 		FSdfPath GetParentPath() const;
 		FSdfPath AppendChild( const TCHAR* ChildName ) const;
+		FSdfPath AppendProperty( FName PropertyName ) const;
 
 		FSdfPath StripAllVariantSelections() const;
 
diff --git a/Engine/Plugins/Importers/USDImporter/Source/UnrealUSDWrapper/Public/UsdWrappers/SdfPrimSpec.h b/Engine/Plugins/Importers/USDImporter/Source/UnrealUSDWrapper/Public/UsdWrappers/SdfPrimSpec.h
new file mode 100644
--- /dev/null
+++ b/Engine/Plugins/Importers/USDImporter/Source/UnrealUSDWrapper/Public/UsdWrappers/SdfPrimSpec.h
@@ -0,0 +1,102 @@
+// Copyright Epic Games, Inc. All Rights Reserved.
+
+#pragma once
+
+#include "Templates/UniquePtr.h"
+#include "UObject/NameTypes.h"
+
+#include "UsdWrappers/ForwardDeclarations.h"
+
+#if USE_USD_SDK
+
+#include "USDIncludesStart.h"
+	#include "pxr/pxr.h"
+#include "USDIncludesEnd.h"
+
+PXR_NAMESPACE_OPEN_SCOPE
+	class SdfPrimSpec;
+	template <class T> class SdfHandle;
+	using SdfPrimSpecHandle = SdfHandle< SdfPrimSpec >;
+PXR_NAMESPACE_CLOSE_SCOPE
+
+#endif // #if USE_USD_SDK
+
+namespace UE
+{
+	class FSdfPath;
+	class FSdfAttributeSpec;
+
+	namespace Internal
+	{
+		class FSdfPrimSpecImpl;
+	}
+
+	/** Corresponds to pxr::SdfSpecType, refer to the USD SDK documentation */
+	enum class ESdfSpecType
+	{
+		SdfSpecTypeUnknown = 0,
+		SdfSpecTypeAttribute,
+		SdfSpecTypeConnection,
+		SdfSpecTypeExpression,
+		SdfSpecTypeMapper,
+		SdfSpecTypeMapperArg,
+		SdfSpecTypePrim,
+		SdfSpecTypePseudoRoot,
+		SdfSpecTypeRelationship,
+		SdfSpecTypeRelationshipTarget,
+		SdfSpecTypeVariant,
+		SdfSpecTypeVariantSet,
+
+		SdfNumSpecTypes
+	};
+
+	/**
+	 * Minimal pxr::SdfPrimSpecHandle wrapper for Unreal that can be used from no-rtti modules.
+	 */
+	class UNREALUSDWRAPPER_API FSdfPrimSpec
+	{
+	public:
+		FSdfPrimSpec();
+
+		FSdfPrimSpec( const FSdfPrimSpec& Other );
+		FSdfPrimSpec( FSdfPrimSpec&& Other );
+		~FSdfPrimSpec();
+
+		FSdfPrimSpec& operator=( const FSdfPrimSpec& Other );
+		FSdfPrimSpec& operator=( FSdfPrimSpec&& Other );
+
+		bool operator==( const FSdfPrimSpec& Other ) const;
+		bool operator!=( const FSdfPrimSpec& Other ) const;
+
+		explicit operator bool() const;
+
+	// Auto conversion from/to pxr::SdfPrimSpecHandle
+	public:
+#if USE_USD_SDK
+		explicit FSdfPrimSpec( const pxr::SdfPrimSpecHandle& InSdfPrimSpec );
+		explicit FSdfPrimSpec( pxr::SdfPrimSpecHandle&& InSdfPrimSpec );
+		FSdfPrimSpec& operator=( const pxr::SdfPrimSpecHandle& InSdfPrimSpec );
+		FSdfPrimSpec& operator=( pxr::SdfPrimSpecHandle&& InSdfPrimSpec );
+
+		operator pxr::SdfPrimSpecHandle&();
+		operator const pxr::SdfPrimSpecHandle&() const;
+#endif // #if USE_USD_SDK
+
+	// Wrapped pxr::SdfPrimSpecHandle functions, refer to the USD SDK documentation
+	public:
+		ESdfSpecType GetSpecType() const;
+		FSdfLayerWeak GetLayer() const;
+		FSdfPath GetPath() const;
+
+		FSdfPrimSpec GetRealNameParent() const;
+		bool RemoveNameChild(const FSdfPrimSpec& Child);
+
+		UE::FSdfAttributeSpec GetAttributeAtPath( const UE::FSdfPath& Path ) const;
+
+		FName GetTypeName() const;
+		FName GetName() const;
+
+	private:
+		TUniquePtr< Internal::FSdfPrimSpecImpl > Impl;
+	};
+}
\ No newline at end of file
