USD: Prevent re-exporting of identical assets;

Refactor AnimSequence exporter to always emit SkeletalMesh on a separate file;
Respect the AssetExportTask's "bReplaceIdentical" member on all exporters.

diff --git a/Engine/Plugins/Importers/USDImporter/Resources/PythonExamples/export_static_mesh.py b/Engine/Plugins/Importers/USDImporter/Resources/PythonExamples/export_static_mesh.py
new file mode 100644
--- /dev/null
+++ b/Engine/Plugins/Importers/USDImporter/Resources/PythonExamples/export_static_mesh.py
@@ -0,0 +1,21 @@
+import unreal
+
+OUTPUT_FILENAME = r"C:\MyFolder\output.usda"
+INPUT_ASSET_CONTENT_PATH = r"/Game/MyStaticMesh"
+
+asset = unreal.load_asset(INPUT_ASSET_CONTENT_PATH)
+
+options = unreal.StaticMeshExporterUSDOptions()
+options.stage_options.meters_per_unit = 0.02
+options.mesh_asset_options.use_payload = True
+options.mesh_asset_options.payload_format = "usda"
+options.mesh_asset_options.bake_materials = True
+
+task = unreal.AssetExportTask()
+task.set_editor_property('object', asset)
+task.set_editor_property('filename', OUTPUT_FILENAME)
+task.set_editor_property('automated', True)
+task.set_editor_property('options', options)
+task.set_editor_property('exporter', unreal.StaticMeshExporterUsd())
+
+unreal.Exporter.run_asset_export_task(task)
\ No newline at end of file
diff --git a/Engine/Plugins/Importers/USDImporter/Source/USDClasses/Public/USDUnrealAssetInfo.h b/Engine/Plugins/Importers/USDImporter/Source/USDClasses/Public/USDUnrealAssetInfo.h
new file mode 100644
--- /dev/null
+++ b/Engine/Plugins/Importers/USDImporter/Source/USDClasses/Public/USDUnrealAssetInfo.h
@@ -0,0 +1,43 @@
+// Copyright Epic Games, Inc. All Rights Reserved.
+
+#pragma once
+
+#include "USDUnrealAssetInfo.generated.h"
+
+/**
+ * Metadata added to a prim to indicate it was exported from a particular Unreal asset
+ */
+USTRUCT( BlueprintType )
+struct USDCLASSES_API FUsdUnrealAssetInfo
+{
+	GENERATED_BODY()
+
+public:
+	// Name of the Unreal asset
+    UPROPERTY( EditAnywhere, BlueprintReadWrite, Category = "Info" )
+	FString Name;
+
+	// Filepath to the layer where the asset was exported to
+    UPROPERTY( EditAnywhere, BlueprintReadWrite, Category = "Info" )
+	FString Identifier;
+
+	// Identifier string for the current asset version. Whenever the asset is updated inside Unreal, this will change
+    UPROPERTY( EditAnywhere, BlueprintReadWrite, Category = "Info" )
+	FString Version;
+
+	// Path to the exported asset (e.g. "/Game/MyMaterials/Red.Red")
+    UPROPERTY( EditAnywhere, BlueprintReadWrite, Category = "Info" )
+	FString UnrealContentPath;
+
+	// Class name of the exported asset (e.g. "StaticMesh")
+    UPROPERTY( EditAnywhere, BlueprintReadWrite, Category = "Info" )
+	FString UnrealAssetType;
+
+	// DateTime string of the moment of export (e.g. "2022.06.17-14.19.54")
+    UPROPERTY( EditAnywhere, BlueprintReadWrite, Category = "Info" )
+	FString UnrealExportTime;
+
+	// Engine version string describing the engine that exported this asset (e.g. "5.1.0-21000000+UE5")
+	UPROPERTY( EditAnywhere, BlueprintReadWrite, Category = "Info" )
+	FString UnrealEngineVersion;
+};
diff --git a/Engine/Plugins/Importers/USDImporter/Source/USDExporter/Private/AnimSequenceExporterUSD.cpp b/Engine/Plugins/Importers/USDImporter/Source/USDExporter/Private/AnimSequenceExporterUSD.cpp
--- a/Engine/Plugins/Importers/USDImporter/Source/USDExporter/Private/AnimSequenceExporterUSD.cpp
+++ b/Engine/Plugins/Importers/USDImporter/Source/USDExporter/Private/AnimSequenceExporterUSD.cpp
@@ -3,15 +3,18 @@
 #include "AnimSequenceExporterUSD.h"
 
 #include "AnimSequenceExporterUSDOptions.h"
-#include "EngineAnalytics.h"
 #include "Engine/SkeletalMesh.h"
+#include "EngineAnalytics.h"
 #include "MaterialExporterUSD.h"
+#include "SkeletalMeshExporterUSDOptions.h"
 #include "UnrealUSDWrapper.h"
 #include "USDClassesModule.h"
 #include "USDErrorUtils.h"
 #include "USDLayerUtils.h"
+#include "USDLog.h"
 #include "USDOptionsWindow.h"
 #include "USDSkeletalDataConversion.h"
+#include "USDUnrealAssetInfo.h"
 
 #include "UsdWrappers/SdfLayer.h"
 #include "UsdWrappers/SdfPath.h"
@@ -20,6 +23,7 @@
 
 #include "Animation/AnimSequence.h"
 #include "AssetExportTask.h"
+#include "UObject/GCObjectScopeGuard.h"
 
 namespace UE
 {
@@ -117,17 +121,23 @@ bool UAnimSequenceExporterUSD::ExportBinary( UObject* Object, const TCHAR* Type,
 
 	FScopedUsdMessageLog UsdMessageLog;
 
+	// We may dispatch another export task in between, so lets cache this for ourselves as it may
+	// be overwritten
+	const FString AnimSequenceFile = UExporter::CurrentFilename;
+
 	UAnimSequenceExporterUSDOptions* Options = nullptr;
 	if ( ExportTask )
 	{
 		Options = Cast<UAnimSequenceExporterUSDOptions>( ExportTask->Options );
 	}
-	if ( !Options && ( !ExportTask || !ExportTask->bAutomated ) )
+	if ( !Options )
 	{
 		Options = GetMutableDefault<UAnimSequenceExporterUSDOptions>();
-		if ( Options )
+
+		// Prompt with an options dialog if we can
+		if ( Options && ( !ExportTask || !ExportTask->bAutomated ) )
 		{
-			Options->PreviewMeshOptions.MaterialBakingOptions.TexturesDir.Path = FPaths::Combine( FPaths::GetPath( UExporter::CurrentFilename ), TEXT( "Textures" ) );
+			Options->PreviewMeshOptions.MaterialBakingOptions.TexturesDir.Path = FPaths::Combine( FPaths::GetPath( AnimSequenceFile ), TEXT( "Textures" ) );
 
 			const bool bIsImport = false;
 			const bool bContinue = SUsdOptionsWindow::ShowImportExportOptions( *Options, bIsImport );
@@ -137,45 +147,17 @@ bool UAnimSequenceExporterUSD::ExportBinary( UObject* Object, const TCHAR* Type,
 			}
 		}
 	}
-
-	double StartTime = FPlatformTime::Cycles64();
-
-	// If bUsePayload is true, we'll intercept the filename so that we write the mesh data to
-	// "C:/MyFolder/file_payload.usda" and create an "asset" file "C:/MyFolder/file.usda" that uses it
-	// as a payload, pointing at the default prim
-	FString PayloadFilename = UExporter::CurrentFilename;
-	const FString& AssetFilename = UExporter::CurrentFilename;
-	if ( Options && Options->PreviewMeshOptions.bUsePayload )
-	{
-		FString PathPart;
-		FString FilenamePart;
-		FString ExtensionPart;
-		FPaths::Split( PayloadFilename, PathPart, FilenamePart, ExtensionPart );
-
-		if ( FormatExtension.Contains( Options->PreviewMeshOptions.PayloadFormat ) )
-		{
-			ExtensionPart = Options->PreviewMeshOptions.PayloadFormat;
-		}
-
-		PayloadFilename = FPaths::Combine( PathPart, FilenamePart + TEXT( "_payload." ) + ExtensionPart );
-	}
-
-	UE::FUsdStage AssetStage = UnrealUSDWrapper::NewStage( *AssetFilename );
-	if ( !AssetStage )
+	if ( !Options )
 	{
 		return false;
 	}
 
-	if ( Options )
-	{
-		UsdUtils::SetUsdStageMetersPerUnit( AssetStage, Options->StageOptions.MetersPerUnit );
-		UsdUtils::SetUsdStageUpAxis( AssetStage, Options->StageOptions.UpAxis );
-	}
+	// See comment on the analogous line within StaticMeshExporterUSD.cpp
+	ExportTask->bPrompt = false;
 
-	FString PrimPath;
-	UE::FUsdPrim SkelRootPrim;
+	// Export preview mesh if needed
 	USkeletalMesh* SkeletalMesh = nullptr;
-
+	FString MeshAssetFile;
 	if ( Options && Options->bExportPreviewMesh )
 	{
 		SkeletalMesh = AnimSequence->GetPreviewMesh();
@@ -194,111 +176,212 @@ bool UAnimSequenceExporterUSD::ExportBinary( UObject* Object, const TCHAR* Type,
 
 		if ( SkeletalMesh )
 		{
-			PrimPath = TEXT( "/" ) + UsdUtils::SanitizeUsdIdentifier( *SkeletalMesh->GetName() );
-			SkelRootPrim = AssetStage.DefinePrim( UE::FSdfPath( *PrimPath ), TEXT( "SkelRoot" ) );
-
-			if ( SkelRootPrim )
+			FString PathPart;
+			FString FilenamePart;
+			FString ExtensionPart;
+			FPaths::Split( AnimSequenceFile, PathPart, FilenamePart, ExtensionPart );
+			MeshAssetFile = FPaths::Combine( PathPart, FilenamePart + TEXT( "_SkeletalMesh." ) + ExtensionPart );
+
+			USkeletalMeshExporterUSDOptions* SkeletalMeshOptions = GetMutableDefault<USkeletalMeshExporterUSDOptions>();
+			SkeletalMeshOptions->StageOptions = Options->StageOptions;
+			SkeletalMeshOptions->MeshAssetOptions = Options->PreviewMeshOptions;
+			SkeletalMeshOptions->bReExportIdenticalAssets = Options->bReExportIdenticalAssets;
+
+			UAssetExportTask* LevelExportTask = NewObject<UAssetExportTask>();
+			FGCObjectScopeGuard ExportTaskGuard( LevelExportTask );
+			LevelExportTask->Object = SkeletalMesh;
+			LevelExportTask->Options = SkeletalMeshOptions;
+			LevelExportTask->Exporter = nullptr;
+			LevelExportTask->Filename = MeshAssetFile;
+			LevelExportTask->bSelected = false;
+			LevelExportTask->bReplaceIdentical = ExportTask->bReplaceIdentical;
+			LevelExportTask->bPrompt = false;
+			LevelExportTask->bUseFileArchive = false;
+			LevelExportTask->bWriteEmptyFiles = false;
+			LevelExportTask->bAutomated = true; // Pretend this is an automated task so it doesn't pop the options dialog
+
+			const bool bSucceeded = UExporter::RunAssetExportTask( LevelExportTask );
+			if ( !bSucceeded )
 			{
-				AssetStage.SetDefaultPrim( SkelRootPrim );
+				MeshAssetFile = {};
 			}
+		}
+		else
+		{
+			FUsdLogManager::LogMessage(
+				EMessageSeverity::Warning,
+				FText::Format( NSLOCTEXT( "AnimSequenceExporterUSD", "InvalidSkelMesh", "Couldn't find the skeletal mesh to export for anim sequence {0}." ), FText::FromName( AnimSequence->GetFName() ) ) );
+		}
+	}
+
+	// Collect the target paths for our SkelAnimation prim and its SkelRoot, if any
+	UE::FSdfPath SkelRootPath;
+	UE::FSdfPath SkelAnimPath;
+	if ( MeshAssetFile.IsEmpty() )
+	{
+		SkelAnimPath = UE::FSdfPath::AbsoluteRootPath().AppendChild(
+			*UsdUtils::SanitizeUsdIdentifier( *AnimSequence->GetName() )
+		);
+	}
+	else
+	{
+		SkelRootPath = UE::FSdfPath::AbsoluteRootPath().AppendChild(
+			*UsdUtils::SanitizeUsdIdentifier( *SkeletalMesh->GetName() )
+		);
+
+		SkelAnimPath = SkelRootPath.AppendChild(
+			*UsdUtils::SanitizeUsdIdentifier( *AnimSequence->GetName() )
+		);
+	}
+
+	FString AnimSequenceVersion;
+	if ( UAnimDataModel* DataModel = AnimSequence->GetDataModel() )
+	{
+		AnimSequenceVersion = DataModel->GenerateGuid().ToString();
+
+		// We could just use the GUID directly but all other asset types end up with SHA hash size so lets be
+		// consistent
+		FSHA1 SHA1;
+		SHA1.UpdateWithString( *AnimSequenceVersion, AnimSequenceVersion.Len() );
+		SHA1.Final();
+		FSHAHash Hash;
+		SHA1.GetHash( &Hash.Hash[ 0 ] );
+		AnimSequenceVersion = Hash.ToString();
+	}
 
-			// Using payload: Convert mesh data through the asset stage (that references the payload) so that we can
-			// author mesh data on the payload layer and material data on the asset layer
-			if ( Options->PreviewMeshOptions.bUsePayload )
+	// Check if we already have exported what we plan on exporting anyway
+	if ( FPaths::FileExists( AnimSequenceFile ) && !AnimSequenceVersion.IsEmpty() )
+	{
+		if ( !ExportTask->bReplaceIdentical )
+		{
+			UE_LOG( LogUsd, Log,
+				TEXT( "Skipping export of asset '%s' as the target file '%s' already exists." ),
+				*Object->GetPathName(),
+				*AnimSequenceFile
+			);
+			return false;
+		}
+		// If we don't want to re-export this asset we need to check if its the same version
+		else if ( !Options->bReExportIdenticalAssets )
+		{
+			// Don't use the stage cache here as we want this stage to close within this scope in case
+			// we have to overwrite its files due to e.g. missing payload or anything like that
+			const bool bUseStageCache = false;
+			const EUsdInitialLoadSet InitialLoadSet = EUsdInitialLoadSet::LoadNone;
+			if ( UE::FUsdStage TempStage = UnrealUSDWrapper::OpenStage( *AnimSequenceFile, InitialLoadSet, bUseStageCache ) )
 			{
-				if ( UE::FUsdStage PayloadStage = UnrealUSDWrapper::NewStage( *PayloadFilename ) )
+				if ( UE::FUsdPrim SkelAnimPrim = TempStage.GetPrimAtPath( SkelAnimPath ) )
 				{
-					UsdUtils::SetUsdStageMetersPerUnit( PayloadStage, Options->StageOptions.MetersPerUnit );
-					UsdUtils::SetUsdStageUpAxis( PayloadStage, Options->StageOptions.UpAxis );
+					FUsdUnrealAssetInfo Info = UsdUtils::GetPrimAssetInfo( SkelAnimPrim );
 
-					if ( UE::FUsdPrim PayloadSkelRootPrim = PayloadStage.DefinePrim( UE::FSdfPath( *PrimPath ), TEXT( "SkelRoot" ) ) )
-					{
-						UnrealToUsd::ConvertSkeletalMesh( SkeletalMesh, PayloadSkelRootPrim, UsdUtils::GetDefaultTimeCode(), &AssetStage, Options->PreviewMeshOptions.LowestMeshLOD, Options->PreviewMeshOptions.HighestMeshLOD );
-						PayloadStage.SetDefaultPrim( PayloadSkelRootPrim );
-					}
+					const bool bVersionMatches = !Info.Version.IsEmpty() && Info.Version == AnimSequenceVersion;
 
-					PayloadStage.GetRootLayer().Save();
+					const bool bAssetTypeMatches = !Info.UnrealAssetType.IsEmpty()
+						&& Info.UnrealAssetType == Object->GetClass()->GetName();
 
-					if ( SkelRootPrim )
+					if ( bVersionMatches && bAssetTypeMatches )
 					{
-						UsdUtils::AddPayload( SkelRootPrim, *PayloadFilename );
+						UE_LOG( LogUsd, Log,
+							TEXT( "Skipping export of asset '%s' as the target file '%s' already contains up-to-date exported data." ),
+							*Object->GetPathName(),
+							*AnimSequenceFile
+						);
+						return true;
 					}
 				}
 			}
-			// Not using payload: Just author everything on the current edit target of the payload (== asset) layer
-			else
-			{
-				UnrealToUsd::ConvertSkeletalMesh( SkeletalMesh, SkelRootPrim, UsdUtils::GetDefaultTimeCode(), nullptr, Options->PreviewMeshOptions.LowestMeshLOD, Options->PreviewMeshOptions.HighestMeshLOD );
-			}
-		}
-		else
-		{
-			FUsdLogManager::LogMessage(
-				EMessageSeverity::Warning,
-				FText::Format( NSLOCTEXT( "AnimSequenceExporterUSD", "InvalidSkelMesh", "Couldn't find the skeletal mesh to export for anim sequence {0}." ), FText::FromName( AnimSequence->GetFName() ) ) );
 		}
 	}
 
-	PrimPath += TEXT( "/" ) + UsdUtils::SanitizeUsdIdentifier( *AnimSequence->GetName() );
-	UE::FUsdPrim SkelAnimPrim = AssetStage.DefinePrim( UE::FSdfPath( *PrimPath ), TEXT("SkelAnimation") );
-	if ( !SkelAnimPrim )
+	double StartTime = FPlatformTime::Cycles64();
+
+	UE::FUsdStage AnimationStage = UnrealUSDWrapper::NewStage( *AnimSequenceFile );
+	if ( !AnimationStage )
 	{
 		return false;
 	}
 
-	if ( !AssetStage.GetDefaultPrim() )
-	{
-		AssetStage.SetDefaultPrim( SkelAnimPrim );
-	}
+	UE::FUsdPrim SkelRootPrim;
+	UE::FUsdPrim SkelAnimPrim;
 
-	const double StartTimeCode = 0.0;
-	const double EndTimeCode = AnimSequence->GetNumberOfSampledKeys() - 1;
-	UsdUtils::AddTimeCodeRangeToLayer( AssetStage.GetRootLayer(), StartTimeCode, EndTimeCode );
+	// Haven't exported the SkeletalMesh, just make a stage with a SkelAnimation prim
+	if( MeshAssetFile.IsEmpty() )
+	{
+		SkelAnimPrim = AnimationStage.DefinePrim( SkelAnimPath, TEXT( "SkelAnimation" ) );
+		if ( !SkelAnimPrim )
+		{
+			return false;
+		}
 
-	AssetStage.SetTimeCodesPerSecond( AnimSequence->GetSamplingFrameRate().AsDecimal() );
+		AnimationStage.SetDefaultPrim( SkelAnimPrim );
+	}
+	// Exported a SkeletalMesh prim elsewhere, create a SkelRoot containing this SkelAnimation prim
+	else
+	{
+		SkelRootPrim = AnimationStage.DefinePrim( SkelRootPath, TEXT( "SkelRoot" ) );
+		if ( !SkelRootPrim )
+		{
+			return false;
+		}
 
-	UnrealToUsd::ConvertAnimSequence( AnimSequence, SkelAnimPrim );
+		SkelAnimPrim = AnimationStage.DefinePrim( SkelAnimPath, TEXT( "SkelAnimation" ) );
+		if ( !SkelAnimPrim )
+		{
+			return false;
+		}
 
-	if ( SkelRootPrim )
-	{
+		AnimationStage.SetDefaultPrim( SkelRootPrim );
 		UsdUtils::BindAnimationSource( SkelRootPrim, SkelAnimPrim );
+
+		// Add a reference to the SkelRoot of the static mesh, which will compose in the Mesh and Skeleton prims
+		UsdUtils::AddReference( SkelRootPrim, *MeshAssetFile );
 	}
 
-	// Bake materials and replace unrealMaterials with references to the baked files.
-	// Do this last because we will need to check the default prim of the stage during the traversal for replacing baked materials
-	if ( Options && Options->bExportPreviewMesh && SkeletalMesh && Options->PreviewMeshOptions.bBakeMaterials )
+	// Configure stage metadata
 	{
-		TSet<UMaterialInterface*> MaterialsToBake;
-		for ( const FSkeletalMaterial& SkeletalMaterial : SkeletalMesh->GetMaterials() )
+		if ( Options )
 		{
-			MaterialsToBake.Add( SkeletalMaterial.MaterialInterface );
+			UsdUtils::SetUsdStageMetersPerUnit( AnimationStage, Options->StageOptions.MetersPerUnit );
+			UsdUtils::SetUsdStageUpAxis( AnimationStage, Options->StageOptions.UpAxis );
 		}
 
-		const bool bIsAssetLayer = true;
-		UMaterialExporterUsd::ExportMaterialsForStage(
-			MaterialsToBake.Array(),
-			Options->PreviewMeshOptions.MaterialBakingOptions,
-			AssetStage,
-			bIsAssetLayer,
-			Options->PreviewMeshOptions.bUsePayload,
-			Options->PreviewMeshOptions.bRemoveUnrealMaterials
-		);
+		const double StartTimeCode = 0.0;
+		const double EndTimeCode = AnimSequence->GetNumberOfSampledKeys() - 1;
+		UsdUtils::AddTimeCodeRangeToLayer( AnimationStage.GetRootLayer(), StartTimeCode, EndTimeCode );
+
+		AnimationStage.SetTimeCodesPerSecond( AnimSequence->GetSamplingFrameRate().AsDecimal() );
+	}
+
+	UnrealToUsd::ConvertAnimSequence( AnimSequence, SkelAnimPrim );
+
+	// Write asset info now that we finished exporting
+	{
+		FUsdUnrealAssetInfo Info;
+		Info.Name = AnimSequence->GetName();
+		Info.Identifier = AnimSequenceFile;
+		Info.Version = AnimSequenceVersion;
+		Info.UnrealContentPath = AnimSequence->GetPathName();
+		Info.UnrealAssetType = AnimSequence->GetClass()->GetName();
+		Info.UnrealExportTime = FDateTime::Now().ToString();
+		Info.UnrealEngineVersion = FEngineVersion::Current().ToString();
+
+		UsdUtils::SetPrimAssetInfo( SkelAnimPrim, Info );
 	}
 
-	AssetStage.GetRootLayer().Save();
+	AnimationStage.GetRootLayer().Save();
 
 	// Analytics
 	{
 		bool bAutomated = ExportTask ? ExportTask->bAutomated : false;
 		double ElapsedSeconds = FPlatformTime::ToSeconds64( FPlatformTime::Cycles64() - StartTime );
-		FString Extension = FPaths::GetExtension( UExporter::CurrentFilename );
+		FString Extension = FPaths::GetExtension( AnimSequenceFile );
 
 		UE::AnimSequenceExporterUSD::Private::SendAnalytics(
 			Object,
 			Options,
 			bAutomated,
 			ElapsedSeconds,
-			UsdUtils::GetUsdStageNumFrames( AssetStage ),
+			UsdUtils::GetUsdStageNumFrames( AnimationStage ),
 			Extension
 		);
 	}
diff --git a/Engine/Plugins/Importers/USDImporter/Source/USDExporter/Private/LevelSequenceExporterUSD.cpp b/Engine/Plugins/Importers/USDImporter/Source/USDExporter/Private/LevelSequenceExporterUSD.cpp
--- a/Engine/Plugins/Importers/USDImporter/Source/USDExporter/Private/LevelSequenceExporterUSD.cpp
+++ b/Engine/Plugins/Importers/USDImporter/Source/USDExporter/Private/LevelSequenceExporterUSD.cpp
@@ -10,6 +10,7 @@
 #include "USDLog.h"
 #include "USDOptionsWindow.h"
 #include "USDPrimConversion.h"
+#include "USDUnrealAssetInfo.h"
 
 #include "UsdWrappers/SdfLayer.h"
 #include "UsdWrappers/SdfPath.h"
@@ -466,6 +467,9 @@
 
 				ULevelSequenceExporterUsdOptions* ExportOptions;
 
+				// Where we store our ExportTask's bReplaceIdentical, which indicates if we should overwrite files or not
+				bool bReplaceIdentical;
+
 				// Our own read-only sequencer that we use to play the level sequences while we bake them out one frame at a time
 				TSharedRef<ISequencer> Sequencer;
 
@@ -941,55 +945,7 @@
 				// Overwriting other files is OK, as we want to allow a "repeatedly export over the same files" workflow
 				FString UniqueFilePath = UsdUtils::GetUniqueName( FilePath, Context.UsedFilePaths );
 
-				UE::FUsdStage UsdStage = UnrealUSDWrapper::NewStage( *UniqueFilePath );
-				if ( !UsdStage )
-				{
-					return;
-				}
-
+				// Try exporting subsequences if needed
 				if ( Context.ExportOptions )
 				{
-					UsdUtils::SetUsdStageMetersPerUnit( UsdStage, Context.ExportOptions->StageOptions.MetersPerUnit );
-					UsdUtils::SetUsdStageUpAxis( UsdStage, Context.ExportOptions->StageOptions.UpAxis );
-					UsdStage.SetTimeCodesPerSecond( Context.ExportOptions->TimeCodesPerSecond );
-				}
-
-				// Set this so that if we open this exported sequence back up in UE the Sequencer will start up showing
-				// the same DisplayRate as the original exported sequence
-				UsdStage.SetFramesPerSecond( MovieScene->GetDisplayRate().AsDecimal() );
-
-				UE::FUsdPrim RootPrim = UsdStage.OverridePrim( UE::FSdfPath( TEXT( "/Root" ) ) );
-				if ( !RootPrim )
-				{
-					return;
-				}
-
-				UsdStage.SetDefaultPrim( RootPrim );
-
-				// Currently we bake the full composed result of each level sequence into a single layer,
-				// because USD can't compose individual layers in the same way (with blending and so on). So here
-				// we make sure that our sequencer exports each MovieSceneSequence as if it was a root, emitting the
-				// same result as if we had exported that subsequence's LevelSequence by itself
-				Context.Sequencer->ResetToNewRootSequence( MovieSceneSequence );
-
-				TMap< UMovieSceneSequence*, TArray<FMovieSceneSequenceID> > SequenceInstances = GetSequenceHierarchyInstances( MovieSceneSequence, Context.Sequencer.Get() );
-
-				TMap<USceneComponent*, FCombinedComponentBakers> Bakers;
-				GenerateBakersForMovieScene( Context, MovieSceneSequence, SequenceInstances, UsdStage, Bakers );
-
-				// Bake this MovieScene
-				// We bake each MovieScene individually instead of doing one large simultaneous bake because this way
-				// not only we avoid having to handle FMovieSceneSequenceTransforms when writing out the UsdTimeCodes,
-				// we guarantee we'll get the same result as if we exported each subsequence individually.
-				// They could have differed, for example, if we had a limited the playback range of a subsequence
-				BakeMovieSceneSequence( Context, MovieSceneSequence, UsdStage, Bakers );
-
-				const TRange< FFrameNumber > PlaybackRange = MovieScene->GetPlaybackRange();
-				const FFrameRate Resolution = MovieScene->GetTickResolution();
-				const FFrameRate DisplayRate = MovieScene->GetDisplayRate();
-				const double StageTimeCodesPerSecond = UsdStage.GetTimeCodesPerSecond();
-				const FFrameRate StageFrameRate( StageTimeCodesPerSecond, 1 );
-
-				if ( Context.ExportOptions )
-				{
 					if ( Context.ExportOptions->bExportSubsequencesAsLayers )
@@ -1042,15 +998,129 @@
 							}
 						}
 					}
+				}
 
-					// We can add the level as a sublayer to every exported subsequence, so that each can be opened individually and
-					// automatically load the level layer. It doesn't matter much if the parent stage has composed the level
-					// sublayer multiple times (in case we add subsequence layers as sublayers in the future), as the prims will
-					// just all override each other with the same data
-					if ( Context.ExportOptions->bUseExportedLevelAsSublayer && FPaths::FileExists( Context.LevelFilePath ) )
+				FString LevelSequenceVersion = MovieSceneSequence.GetSignature().ToString();
+				{
+					// We could just use the GUID directly but all other asset types end up with SHA hash size so lets be
+					// consistent
+					FSHA1 SHA1;
+					SHA1.UpdateWithString( *LevelSequenceVersion, LevelSequenceVersion.Len() );
+					SHA1.Final();
+					FSHAHash Hash;
+					SHA1.GetHash( &Hash.Hash[ 0 ] );
+					LevelSequenceVersion = Hash.ToString();
+				}
+
+				// Check if we already have exported what we plan on exporting anyway
+				if ( FPaths::FileExists( UniqueFilePath ) )
+				{
+					if ( !Context.bReplaceIdentical )
 					{
-						UsdUtils::InsertSubLayer( UsdStage.GetRootLayer(), *Context.LevelFilePath );
+						UE_LOG( LogUsd, Log,
+							TEXT( "Skipping export of asset '%s' as the target file '%s' already exists." ),
+							*MovieSceneSequence.GetPathName(),
+							*UExporter::CurrentFilename
+						);
+						return;
+					}
+					// If we don't want to re-export this level sequence we need to check if its the same version
+					else if ( !Context.ExportOptions->bReExportIdenticalLevelsAndSequences )
+					{
+						// Don't use the stage cache here as we want this stage to close within this scope in case
+						// we have to overwrite its files due to e.g. missing payload or anything like that
+						const bool bUseStageCache = false;
+						const EUsdInitialLoadSet InitialLoadSet = EUsdInitialLoadSet::LoadNone;
+						if ( UE::FUsdStage TempStage = UnrealUSDWrapper::OpenStage( *UniqueFilePath, InitialLoadSet, bUseStageCache ) )
+						{
+							if ( UE::FUsdPrim RootPrim = TempStage.GetDefaultPrim() )
+							{
+								FUsdUnrealAssetInfo Info = UsdUtils::GetPrimAssetInfo( RootPrim );
+
+								const bool bVersionMatches = !Info.Version.IsEmpty() && Info.Version == LevelSequenceVersion;
+
+								const bool bAssetTypeMatches = !Info.UnrealAssetType.IsEmpty()
+									&& Info.UnrealAssetType == MovieSceneSequence.GetClass()->GetName();
+
+								if ( bVersionMatches && bAssetTypeMatches )
+								{
+									UE_LOG( LogUsd, Log,
+										TEXT( "Skipping export of asset '%s' as the target file '%s' already contains up-to-date exported data." ),
+										*MovieSceneSequence.GetPathName(),
+										*UniqueFilePath
+									);
+									return;
+								}
+							}
+						}
 					}
 				}
 
+				UE::FUsdStage UsdStage = UnrealUSDWrapper::NewStage( *UniqueFilePath );
+				if ( !UsdStage )
+				{
+					return;
+				}
+
+				if ( Context.ExportOptions )
+				{
+					UsdUtils::SetUsdStageMetersPerUnit( UsdStage, Context.ExportOptions->StageOptions.MetersPerUnit );
+					UsdUtils::SetUsdStageUpAxis( UsdStage, Context.ExportOptions->StageOptions.UpAxis );
+					UsdStage.SetTimeCodesPerSecond( Context.ExportOptions->TimeCodesPerSecond );
+				}
+
+				// Set this so that if we open this exported sequence back up in UE the Sequencer will start up showing
+				// the same DisplayRate as the original exported sequence
+				UsdStage.SetFramesPerSecond( MovieScene->GetDisplayRate().AsDecimal() );
+
+				UE::FUsdPrim RootPrim = UsdStage.OverridePrim( UE::FSdfPath( TEXT( "/Root" ) ) );
+				if ( !RootPrim )
+				{
+					return;
+				}
+
+				UsdStage.SetDefaultPrim( RootPrim );
+
+				// Currently we bake the full composed result of each level sequence into a single layer,
+				// because USD can't compose individual layers in the same way (with blending and so on). So here
+				// we make sure that our sequencer exports each MovieSceneSequence as if it was a root, emitting the
+				// same result as if we had exported that subsequence's LevelSequence by itself
+				Context.Sequencer->ResetToNewRootSequence( MovieSceneSequence );
+
+				TMap< UMovieSceneSequence*, TArray<FMovieSceneSequenceID> > SequenceInstances = GetSequenceHierarchyInstances( MovieSceneSequence, Context.Sequencer.Get() );
+
+				TMap<USceneComponent*, FCombinedComponentBakers> Bakers;
+				GenerateBakersForMovieScene( Context, MovieSceneSequence, SequenceInstances, UsdStage, Bakers );
+
+				// Bake this MovieScene
+				// We bake each MovieScene individually instead of doing one large simultaneous bake because this way
+				// not only we avoid having to handle FMovieSceneSequenceTransforms when writing out the UsdTimeCodes,
+				// we guarantee we'll get the same result as if we exported each subsequence individually.
+				// They could have differed, for example, if we had a limited the playback range of a subsequence
+				BakeMovieSceneSequence( Context, MovieSceneSequence, UsdStage, Bakers );
+
+				// We can add the level as a sublayer to every exported subsequence, so that each can be opened individually and
+				// automatically load the level layer. It doesn't matter much if the parent stage has composed the level
+				// sublayer multiple times (in case we add subsequence layers as sublayers in the future), as the prims will
+				// just all override each other with the same data
+				if ( Context.ExportOptions && Context.ExportOptions->bUseExportedLevelAsSublayer && FPaths::FileExists( Context.LevelFilePath ) )
+				{
+					UsdUtils::InsertSubLayer( UsdStage.GetRootLayer(), *Context.LevelFilePath );
+				}
+
+				// Write asset info now that we finished exporting
+				if ( UE::FUsdPrim AssetDefaultPrim = UsdStage.GetDefaultPrim() )
+				{
+					FUsdUnrealAssetInfo Info;
+					Info.Name = MovieSceneSequence.GetName();
+					Info.Identifier = UniqueFilePath;
+					Info.Version = LevelSequenceVersion;
+					Info.UnrealContentPath = MovieSceneSequence.GetPathName();
+					Info.UnrealAssetType = MovieSceneSequence.GetClass()->GetName();
+					Info.UnrealExportTime = FDateTime::Now().ToString();
+					Info.UnrealEngineVersion = FEngineVersion::Current().ToString();
+
+					UsdUtils::SetPrimAssetInfo( AssetDefaultPrim, Info );
+				}
+
 				Context.ExportedMovieScenes.Add( &MovieSceneSequence, UniqueFilePath );
@@ -1106,28 +1176,31 @@
 	{
 		Options = Cast<ULevelSequenceExporterUsdOptions>( ExportTask->Options );
 	}
-
 	if ( !Options )
 	{
 		Options = GetMutableDefault<ULevelSequenceExporterUsdOptions>();
-	}
-	if ( !Options )
-	{
-		return false;
-	}
 
-	if ( !ExportTask || !ExportTask->bAutomated )
-	{
-		Options->LevelExportOptions.AssetFolder.Path = FPaths::Combine( FPaths::GetPath( UExporter::CurrentFilename ), TEXT( "Assets" ) );
-		Options->TimeCodesPerSecond = MovieScene->GetDisplayRate().AsDecimal();
-
-		const bool bIsImport = false;
-		const bool bContinue = SUsdOptionsWindow::ShowImportExportOptions( *Options, bIsImport );
-		if ( !bContinue )
+		// Prompt with an options dialog if we can
+		if ( Options && ( !ExportTask || !ExportTask->bAutomated ) )
 		{
-			return false;
+			Options->LevelExportOptions.AssetFolder.Path = FPaths::Combine( FPaths::GetPath( UExporter::CurrentFilename ), TEXT( "Assets" ) );
+			Options->TimeCodesPerSecond = MovieScene->GetDisplayRate().AsDecimal();
+
+			const bool bIsImport = false;
+			const bool bContinue = SUsdOptionsWindow::ShowImportExportOptions( *Options, bIsImport );
+			if ( !bContinue )
+			{
+				return false;
+			}
 		}
 	}
+	if ( !Options )
+	{
+		return false;
+	}
 
+	// See comment on the analogous line within StaticMeshExporterUSD.cpp
+	ExportTask->bPrompt = false;
+
 	double StartTime = FPlatformTime::Cycles64();
 
@@ -1173,6 +1246,7 @@
 
 	LevelSequenceExporterImpl::FLevelSequenceExportContext Context{ *LevelSequence, TempSequencer.ToSharedRef(), SpawnRegister.ToSharedRef() };
 	Context.ExportOptions = Options;
+	Context.bReplaceIdentical = ExportTask->bReplaceIdentical;
 
 	// Spawn (but hide) all spawnables so that they will also show up on the level export if we need them to.
 	// We have to traverse the template IDs when spawning spawnables, because we'll want to force each individual spawnable of each
@@ -1205,6 +1279,8 @@
 			ULevelExporterUSDOptions* LevelOptions = GetMutableDefault<ULevelExporterUSDOptions>();
 			LevelOptions->StageOptions = Options->StageOptions;
 			LevelOptions->Inner = Options->LevelExportOptions;
+			LevelOptions->bReExportIdenticalAssets = Options->bReExportIdenticalAssets;
+			LevelOptions->bReExportIdenticalLevelsAndSequences = Options->bReExportIdenticalLevelsAndSequences;
 
 			UAssetExportTask* LevelExportTask = NewObject<UAssetExportTask>();
 			FGCObjectScopeGuard ExportTaskGuard( LevelExportTask );
@@ -1213,11 +1289,11 @@
 			LevelExportTask->Exporter = nullptr;
 			LevelExportTask->Filename = Context.LevelFilePath;
 			LevelExportTask->bSelected = false;
-			LevelExportTask->bReplaceIdentical = true;
+			LevelExportTask->bReplaceIdentical = ExportTask->bReplaceIdentical;
 			LevelExportTask->bPrompt = false;
 			LevelExportTask->bUseFileArchive = false;
 			LevelExportTask->bWriteEmptyFiles = false;
-			LevelExportTask->bAutomated = true; // Pretend this is an automated task so it doesn't pop the dialog
+			LevelExportTask->bAutomated = true; // Pretend this is an automated task so it doesn't pop the options dialog
 
 			UExporter::RunAssetExportTask( LevelExportTask );
 		}
diff --git a/Engine/Plugins/Importers/USDImporter/Source/USDExporter/Private/MaterialExporterUSD.cpp b/Engine/Plugins/Importers/USDImporter/Source/USDExporter/Private/MaterialExporterUSD.cpp
--- a/Engine/Plugins/Importers/USDImporter/Source/USDExporter/Private/MaterialExporterUSD.cpp
+++ b/Engine/Plugins/Importers/USDImporter/Source/USDExporter/Private/MaterialExporterUSD.cpp
@@ -7,9 +7,11 @@
 #include "USDClassesModule.h"
 #include "USDConversionUtils.h"
 #include "USDGeomMeshConversion.h"
+#include "USDLog.h"
 #include "USDMemory.h"
 #include "USDOptionsWindow.h"
 #include "USDShadeConversion.h"
+#include "USDUnrealAssetInfo.h"
 
 #include "UsdWrappers/SdfLayer.h"
 #include "UsdWrappers/SdfPath.h"
@@ -17,51 +19,159 @@
 #include "UsdWrappers/UsdStage.h"
 
 #include "AssetExportTask.h"
+#include "Engine/Font.h"
+#include "Engine/Texture.h"
 #include "EngineAnalytics.h"
 #include "HAL/FileManager.h"
 #include "IMaterialBakingModule.h"
 #include "MaterialOptions.h"
+#include "Materials/MaterialInstance.h"
+#include "Materials/MaterialInstanceConstant.h"
 #include "Materials/MaterialInterface.h"
 #include "Misc/Paths.h"
 #include "Modules/ModuleManager.h"
 #include "UObject/UObjectGlobals.h"
+#include "VT/RuntimeVirtualTexture.h"
 
-namespace UE
+namespace UE::MaterialExporterUSD::Private
 {
-	namespace MaterialExporterUSD
+	void SendAnalytics( const UMaterialInterface& Material, const FUsdMaterialBakingOptions& Options, bool bAutomated, double ElapsedSeconds, double NumberOfFrames, const FString& Extension )
 	{
-		namespace Private
+		if ( FEngineAnalytics::IsAvailable() )
 		{
-			void SendAnalytics( const UMaterialInterface& Material, const FUsdMaterialBakingOptions& Options, bool bAutomated, double ElapsedSeconds, double NumberOfFrames, const FString& Extension )
+			FString ClassName = Material.GetClass()->GetName();
+
+			TArray<FAnalyticsEventAttribute> EventAttributes;
+
+			EventAttributes.Emplace( TEXT( "AssetType" ), ClassName );
+
+			FString BakedPropertiesString;
 			{
-				if ( FEngineAnalytics::IsAvailable() )
+				const UEnum* PropertyEnum = StaticEnum<EMaterialProperty>();
+				for ( const FPropertyEntry& PropertyEntry : Options.Properties )
 				{
-					FString ClassName = Material.GetClass()->GetName();
+					FString PropertyString = PropertyEnum->GetNameByValue( PropertyEntry.Property ).ToString();
+					PropertyString.RemoveFromStart( TEXT( "MP_" ) );
+					BakedPropertiesString += PropertyString + TEXT( ", " );
+				}
 
-					TArray<FAnalyticsEventAttribute> EventAttributes;
+				BakedPropertiesString.RemoveFromEnd( TEXT( ", " ) );
+			}
 
-					EventAttributes.Emplace( TEXT( "AssetType" ), ClassName );
+			EventAttributes.Emplace( TEXT( "BakedProperties" ), BakedPropertiesString );
+			EventAttributes.Emplace( TEXT( "DefaultTextureSize" ), Options.DefaultTextureSize.ToString() );
 
-					FString BakedPropertiesString;
-					{
-						const UEnum* PropertyEnum = StaticEnum<EMaterialProperty>();
-						for ( const FPropertyEntry& PropertyEntry : Options.Properties )
-						{
-							FString PropertyString = PropertyEnum->GetNameByValue( PropertyEntry.Property ).ToString();
-							PropertyString.RemoveFromStart( TEXT( "MP_" ) );
-							BakedPropertiesString += PropertyString + TEXT( ", " );
-						}
-
-						BakedPropertiesString.RemoveFromEnd( TEXT( ", " ) );
-					}
+			IUsdClassesModule::SendAnalytics( MoveTemp( EventAttributes ), FString::Printf( TEXT( "Export.%s" ), *ClassName ), bAutomated, ElapsedSeconds, NumberOfFrames, Extension );
+		}
+	}
 
-					EventAttributes.Emplace( TEXT( "BakedProperties" ), BakedPropertiesString );
-					EventAttributes.Emplace( TEXT( "DefaultTextureSize" ), Options.DefaultTextureSize.ToString() );
+	FString GetMaterialHashString( const UMaterialInterface& Material )
+	{
+		FSHA1 VersionHash;
+		if ( const UMaterial* Parent = Material.GetMaterial() )
+		{
+			// Add StateId to the combined hash, which is updated on every compile
+			VersionHash.Update( reinterpret_cast< const uint8* >( &Parent->StateId[ 0 ] ), sizeof( FGuid ) );
 
-					IUsdClassesModule::SendAnalytics( MoveTemp( EventAttributes ), FString::Printf( TEXT( "Export.%s" ), *ClassName ), bAutomated, ElapsedSeconds, NumberOfFrames, Extension );
-				}
+			// Add material hash to the combined hash, which is updated when switching some material settings
+			if ( const FMaterialResource* MaterialResource = Parent->GetMaterialResource( GMaxRHIFeatureLevel ) )
+			{
+				MaterialResource = Parent->GetMaterialResource( GMaxRHIFeatureLevel );
+
+				FMaterialShaderMapId ShaderMapID;
+				MaterialResource->GetShaderMapId( GMaxRHIShaderPlatform, nullptr, ShaderMapID );
+
+				FSHAHash MaterialHash;
+				ShaderMapID.GetMaterialHash( MaterialHash );
+
+				VersionHash.Update( &MaterialHash.Hash[ 0 ], sizeof( MaterialHash.Hash ) );
+			}
+			else
+			{
+				// We should have resources by now, but if we don't just add some random value to the hash so that
+				// we never consider this material similar to the one we previously exported, since we can't guarantee
+				// that at all without looking at its shader map ID hash
+				FGuid NewGuid = FGuid::NewGuid();
+				VersionHash.Update( reinterpret_cast< const uint8* >( &NewGuid[ 0 ] ), sizeof( FGuid ) );
 			}
 		}
+
+		// Add material override parameters to the hash
+		if ( const UMaterialInstanceConstant* InstanceConstant = Cast<const UMaterialInstanceConstant>( &Material ) )
+		{
+			// MIC already have a nice GUID we can use
+			VersionHash.Update( reinterpret_cast< const uint8* >( &InstanceConstant->ParameterStateId[ 0 ] ), sizeof( FGuid ) );
+		}
+		else if ( const UMaterialInstance* Instance = Cast<const UMaterialInstance>( &Material ) )
+		{
+			// Manually hash everything...
+			for ( const FScalarParameterValue& Parameter : Instance->ScalarParameterValues )
+			{
+				uint32 NameHash = GetTypeHash( Parameter.ParameterInfo.Name );
+				VersionHash.Update( reinterpret_cast< const uint8* >( &NameHash ), sizeof( uint32 ) );
+				VersionHash.Update( reinterpret_cast< const uint8* >( &Parameter.ParameterInfo.Association ), sizeof( Parameter.ParameterInfo.Association ) );
+				VersionHash.Update( reinterpret_cast< const uint8* >( &Parameter.ParameterInfo.Index ), sizeof( Parameter.ParameterInfo.Index ) );
+				VersionHash.Update( reinterpret_cast< const uint8* >( &Parameter.ParameterValue ), sizeof( Parameter.ParameterValue ) );
+				VersionHash.Update( reinterpret_cast< const uint8* >( &Parameter.ExpressionGUID ), sizeof( Parameter.ExpressionGUID ) );
+			}
+			for ( const FVectorParameterValue& Parameter : Instance->VectorParameterValues )
+			{
+				uint32 NameHash = GetTypeHash( Parameter.ParameterInfo.Name );
+				VersionHash.Update( reinterpret_cast< const uint8* >( &NameHash ), sizeof( uint32 ) );
+				VersionHash.Update( reinterpret_cast< const uint8* >( &Parameter.ParameterInfo.Association ), sizeof( Parameter.ParameterInfo.Association ) );
+				VersionHash.Update( reinterpret_cast< const uint8* >( &Parameter.ParameterInfo.Index ), sizeof( Parameter.ParameterInfo.Index ) );
+				VersionHash.Update( reinterpret_cast< const uint8* >( &Parameter.ParameterValue ), sizeof( Parameter.ParameterValue ) );
+				VersionHash.Update( reinterpret_cast< const uint8* >( &Parameter.ExpressionGUID ), sizeof( Parameter.ExpressionGUID ) );
+			}
+			for ( const FDoubleVectorParameterValue& Parameter : Instance->DoubleVectorParameterValues )
+			{
+				uint32 NameHash = GetTypeHash( Parameter.ParameterInfo.Name );
+				VersionHash.Update( reinterpret_cast< const uint8* >( &NameHash ), sizeof( uint32 ) );
+				VersionHash.Update( reinterpret_cast< const uint8* >( &Parameter.ParameterInfo.Association ), sizeof( Parameter.ParameterInfo.Association ) );
+				VersionHash.Update( reinterpret_cast< const uint8* >( &Parameter.ParameterInfo.Index ), sizeof( Parameter.ParameterInfo.Index ) );
+				VersionHash.Update( reinterpret_cast< const uint8* >( &Parameter.ParameterValue ), sizeof( Parameter.ParameterValue ) );
+				VersionHash.Update( reinterpret_cast< const uint8* >( &Parameter.ExpressionGUID ), sizeof( Parameter.ExpressionGUID ) );
+			}
+			for ( const FTextureParameterValue& Parameter : Instance->TextureParameterValues )
+			{
+				uint32 NameHash = GetTypeHash( Parameter.ParameterInfo.Name );
+				VersionHash.Update( reinterpret_cast< const uint8* >( &NameHash ), sizeof( uint32 ) );
+				VersionHash.Update( reinterpret_cast< const uint8* >( &Parameter.ParameterInfo.Association ), sizeof( Parameter.ParameterInfo.Association ) );
+				VersionHash.Update( reinterpret_cast< const uint8* >( &Parameter.ParameterInfo.Index ), sizeof( Parameter.ParameterInfo.Index ) );
+
+				FString TexturePath = Parameter.ParameterValue ? Parameter.ParameterValue->GetPathName() : FString{};
+				VersionHash.UpdateWithString( *TexturePath, TexturePath.Len() );
+				VersionHash.Update( reinterpret_cast< const uint8* >( &Parameter.ExpressionGUID ), sizeof( Parameter.ExpressionGUID ) );
+			}
+			for ( const FRuntimeVirtualTextureParameterValue& Parameter : Instance->RuntimeVirtualTextureParameterValues )
+			{
+				uint32 NameHash = GetTypeHash( Parameter.ParameterInfo.Name );
+				VersionHash.Update( reinterpret_cast< const uint8* >( &NameHash ), sizeof( uint32 ) );
+				VersionHash.Update( reinterpret_cast< const uint8* >( &Parameter.ParameterInfo.Association ), sizeof( Parameter.ParameterInfo.Association ) );
+				VersionHash.Update( reinterpret_cast< const uint8* >( &Parameter.ParameterInfo.Index ), sizeof( Parameter.ParameterInfo.Index ) );
+
+				FString TexturePath = Parameter.ParameterValue ? Parameter.ParameterValue->GetPathName() : FString{};
+				VersionHash.UpdateWithString( *TexturePath, TexturePath.Len() );
+				VersionHash.Update( reinterpret_cast< const uint8* >( &Parameter.ExpressionGUID ), sizeof( Parameter.ExpressionGUID ) );
+			}
+			for ( const FFontParameterValue& Parameter : Instance->FontParameterValues )
+			{
+				uint32 NameHash = GetTypeHash( Parameter.ParameterInfo.Name );
+				VersionHash.Update( reinterpret_cast< const uint8* >( &NameHash ), sizeof( uint32 ) );
+				VersionHash.Update( reinterpret_cast< const uint8* >( &Parameter.ParameterInfo.Association ), sizeof( Parameter.ParameterInfo.Association ) );
+				VersionHash.Update( reinterpret_cast< const uint8* >( &Parameter.ParameterInfo.Index ), sizeof( Parameter.ParameterInfo.Index ) );
+
+				FString FontPath = Parameter.FontValue ? Parameter.FontValue->GetPathName() : FString{};
+				VersionHash.UpdateWithString( *FontPath, FontPath.Len() );
+				VersionHash.Update( reinterpret_cast< const uint8* >( &Parameter.FontPage ), sizeof( Parameter.FontPage ) );
+				VersionHash.Update( reinterpret_cast< const uint8* >( &Parameter.ExpressionGUID ), sizeof( Parameter.ExpressionGUID ) );
+			}
+		}
+
+		VersionHash.Final();
+		FSHAHash Hash;
+		VersionHash.GetHash( &Hash.Hash[ 0 ] );
+		return Hash.ToString();
 	}
 }
 
@@ -98,10 +208,12 @@ bool UMaterialExporterUsd::ExportBinary( UObject* Object, const TCHAR* Type, FAr
 	{
 		Options = Cast<UMaterialExporterUSDOptions>( ExportTask->Options );
 	}
-	if ( !Options && ( !ExportTask || !ExportTask->bAutomated ) )
+	if ( !Options )
 	{
 		Options = GetMutableDefault<UMaterialExporterUSDOptions>();
-		if ( Options )
+
+		// Prompt with an options dialog if we can
+		if ( Options && ( !ExportTask || !ExportTask->bAutomated ) )
 		{
 			Options->MaterialBakingOptions.TexturesDir.Path = FPaths::Combine( FPaths::GetPath( UExporter::CurrentFilename ), TEXT( "Textures" ) );
 
@@ -113,16 +225,93 @@ bool UMaterialExporterUsd::ExportBinary( UObject* Object, const TCHAR* Type, FAr
 			}
 		}
 	}
+	if ( !Options )
+	{
+		return false;
+	}
+
+	// See comment on the analogous line within StaticMeshExporterUSD.cpp
+	ExportTask->bPrompt = false;
+
+	// Make sure the material has a resource as we'll use it to fetch some version GUIDs and ExportMaterial
+	// can't do it directly due to receiving the material by const ref
+	if ( UMaterial* Parent = Material->GetMaterial() )
+	{
+		const FMaterialResource* MaterialResource = Parent->GetMaterialResource( GMaxRHIFeatureLevel );
+		if ( MaterialResource == nullptr )
+		{
+			Parent->ForceRecompileForRendering();
+		}
+	}
 
-	return UMaterialExporterUsd::ExportMaterial( *Material, Options->MaterialBakingOptions, FFilePath{ UExporter::CurrentFilename } );
+	return UMaterialExporterUsd::ExportMaterial(
+		*Material,
+		Options->MaterialBakingOptions,
+		FFilePath{ UExporter::CurrentFilename },
+		ExportTask->bReplaceIdentical,
+		Options->bReExportIdenticalAssets
+	);
 #else
 	return false;
 #endif // #if USE_USD_SDK
 }
 
-bool UMaterialExporterUsd::ExportMaterial( const UMaterialInterface& Material, const FUsdMaterialBakingOptions& Options, const FFilePath& FilePath )
+bool UMaterialExporterUsd::ExportMaterial(
+	const UMaterialInterface& Material,
+	const FUsdMaterialBakingOptions& Options,
+	const FFilePath& FilePath,
+	bool bReplaceIdentical,
+	bool bReExportIdenticalAssets
+)
 {
 #if USE_USD_SDK
+
+	FString MaterialHashString = UE::MaterialExporterUSD::Private::GetMaterialHashString( Material );
+
+	// Check if we already have exported what we plan on exporting anyway
+	if ( FPaths::FileExists( FilePath.FilePath ) )
+	{
+		if ( !bReplaceIdentical )
+		{
+			UE_LOG( LogUsd, Log,
+				TEXT( "Skipping export of asset '%s' as the target file '%s' already exists." ),
+				*Material.GetPathName(),
+				*UExporter::CurrentFilename
+			);
+			return false;
+		}
+		// If we don't want to re-export this asset we need to check if its the same version
+		else if ( !bReExportIdenticalAssets )
+		{
+			// Don't use the asset cache here as we want this stage to close within this scope in case
+			// we have to overwrite its files due to e.g. missing payload or anything like that
+			const bool bUseAssetCache = false;
+			const EUsdInitialLoadSet InitialLoadSet = EUsdInitialLoadSet::LoadNone;
+			if ( UE::FUsdStage TempStage = UnrealUSDWrapper::OpenStage( *FilePath.FilePath, InitialLoadSet, bUseAssetCache ) )
+			{
+				if ( UE::FUsdPrim DefaultPrim = TempStage.GetDefaultPrim() )
+				{
+					FUsdUnrealAssetInfo Info = UsdUtils::GetPrimAssetInfo( DefaultPrim );
+
+					const bool bVersionMatches = !Info.Version.IsEmpty() && Info.Version == MaterialHashString;
+
+					const bool bAssetTypeMatches = !Info.UnrealAssetType.IsEmpty()
+						&& Info.UnrealAssetType == Material.GetClass()->GetName();
+
+					if ( bVersionMatches && bAssetTypeMatches )
+					{
+						UE_LOG( LogUsd, Log,
+							TEXT( "Skipping export of asset '%s' as the target file '%s' already contains up-to-date exported data." ),
+							*Material.GetPathName(),
+							*FilePath.FilePath
+						);
+						return true;
+					}
+				}
+			}
+		}
+	}
+
 	double StartTime = FPlatformTime::Cycles64();
 
 	UE::FUsdStage UsdStage = UnrealUSDWrapper::NewStage( *FilePath.FilePath );
@@ -151,6 +340,20 @@ bool UMaterialExporterUsd::ExportMaterial( const UMaterialInterface& Material, c
 		RootPrim
 	);
 
+	// Write asset info now that we finished exporting
+	{
+		FUsdUnrealAssetInfo Info;
+		Info.Name = Material.GetName();
+		Info.Identifier = UExporter::CurrentFilename;
+		Info.Version = MaterialHashString;
+		Info.UnrealContentPath = Material.GetPathName();
+		Info.UnrealAssetType = Material.GetClass()->GetName();
+		Info.UnrealExportTime = FDateTime::Now().ToString();
+		Info.UnrealEngineVersion = FEngineVersion::Current().ToString();
+
+		UsdUtils::SetPrimAssetInfo( RootPrim, Info );
+	}
+
 	UsdStage.GetRootLayer().Save();
 
 	// Analytics
@@ -175,7 +378,10 @@ bool UMaterialExporterUsd::ExportMaterialsForStage(
 	const UE::FUsdStage& Stage,
 	bool bIsAssetLayer,
 	bool bUsePayload,
-	bool bRemoveUnrealMaterials )
+	bool bRemoveUnrealMaterials,
+	bool bReplaceIdentical,
+	bool bReExportIdenticalAssets
+)
 {
 #if USE_USD_SDK
 	if ( !Stage )
@@ -198,13 +404,24 @@ bool UMaterialExporterUsd::ExportMaterialsForStage(
 	TSet<FString> UsedFilePathsWithoutExt;
 
 	TMap<FString, FString> MaterialPathNameToFilePath;
-	for ( const UMaterialInterface* Material : Materials )
+	for ( UMaterialInterface* Material : Materials )
 	{
 		if ( !Material )
 		{
 			continue;
 		}
 
+		// Make sure the material has a resource as we'll use it to fetch some version GUIDs and ExportMaterial
+		// can't do it directly due to receiving the material by const ref
+		if ( UMaterial* Parent = Material->GetMaterial() )
+		{
+			const FMaterialResource* MaterialResource = Parent->GetMaterialResource( GMaxRHIFeatureLevel );
+			if ( MaterialResource == nullptr )
+			{
+				Parent->ForceRecompileForRendering();
+			}
+		}
+
 		// "/Game/ContentFolder/Blue.Blue"
 		FString MaterialPathName = Material->GetPathName();
 
@@ -217,7 +434,13 @@ bool UMaterialExporterUsd::ExportMaterialsForStage(
 		// "C:/MyFolder/Export/Blue_4.usda"
 		FString FinalPath = FString::Printf( TEXT( "%s.%s" ), *FinalPathNoExt, *ExtensionNoDot );
 
-		if ( UMaterialExporterUsd::ExportMaterial( *Material, Options, FFilePath{ FinalPath } ) )
+		if ( UMaterialExporterUsd::ExportMaterial(
+			*Material,
+			Options,
+			FFilePath{ FinalPath },
+			bReplaceIdentical,
+			bReExportIdenticalAssets
+		))
 		{
 			UsedFilePathsWithoutExt.Add( FinalPathNoExt );
 			MaterialPathNameToFilePath.Add( MaterialPathName, FinalPath );
diff --git a/Engine/Plugins/Importers/USDImporter/Source/USDExporter/Private/SkeletalMeshExporterUSD.cpp b/Engine/Plugins/Importers/USDImporter/Source/USDExporter/Private/SkeletalMeshExporterUSD.cpp
--- a/Engine/Plugins/Importers/USDImporter/Source/USDExporter/Private/SkeletalMeshExporterUSD.cpp
+++ b/Engine/Plugins/Importers/USDImporter/Source/USDExporter/Private/SkeletalMeshExporterUSD.cpp
@@ -7,10 +7,12 @@
 #include "SkeletalMeshExporterUSDOptions.h"
 #include "USDClassesModule.h"
 #include "USDConversionUtils.h"
+#include "USDLog.h"
 #include "USDMemory.h"
 #include "USDOptionsWindow.h"
 #include "USDSkeletalDataConversion.h"
 #include "USDTypesConversion.h"
+#include "USDUnrealAssetInfo.h"
 
 #include "UsdWrappers/SdfLayer.h"
 #include "UsdWrappers/SdfPath.h"
@@ -114,10 +116,12 @@ bool USkeletalMeshExporterUsd::ExportBinary( UObject* Object, const TCHAR* Type,
 	{
 		Options = Cast<USkeletalMeshExporterUSDOptions>( ExportTask->Options );
 	}
-	if ( !Options && ( !ExportTask || !ExportTask->bAutomated ) )
+	if ( !Options )
 	{
 		Options = GetMutableDefault<USkeletalMeshExporterUSDOptions>();
-		if ( Options )
+
+		// Prompt with an options dialog if we can
+		if ( Options && ( !ExportTask || !ExportTask->bAutomated ) )
 		{
 			Options->MeshAssetOptions.MaterialBakingOptions.TexturesDir.Path = FPaths::Combine( FPaths::GetPath( UExporter::CurrentFilename ), TEXT( "Textures" ) );
 
@@ -129,8 +133,13 @@ bool USkeletalMeshExporterUsd::ExportBinary( UObject* Object, const TCHAR* Type,
 			}
 		}
 	}
+	if ( !Options )
+	{
+		return false;
+	}
 
-	double StartTime = FPlatformTime::Cycles64();
+	// See comment on the analogous line within StaticMeshExporterUSD.cpp
+	ExportTask->bPrompt = false;
 
 	// If bUsePayload is true, we'll intercept the filename so that we write the mesh data to
 	// "C:/MyFolder/file_payload.usda" and create an "asset" file "C:/MyFolder/file.usda" that uses it
@@ -151,6 +160,89 @@ bool USkeletalMeshExporterUsd::ExportBinary( UObject* Object, const TCHAR* Type,
 		PayloadFilename = FPaths::Combine( PathPart, FilenamePart + TEXT( "_payload." ) + ExtensionPart );
 	}
 
+	// Get a simple GUID hash/identifier of our mesh
+	FString DDCKeyHash;
+	if( FSkeletalMeshRenderData* RenderData = SkeletalMesh->GetResourceForRendering() )
+	{
+		const FString& DDCKey = RenderData->DerivedDataKey;
+		FSHA1 SHA1;
+		SHA1.UpdateWithString( *DDCKey, DDCKey.Len() );
+		SHA1.Final();
+		FSHAHash Hash;
+		SHA1.GetHash( &Hash.Hash[ 0 ] );
+		DDCKeyHash = Hash.ToString();
+	}
+
+	// Check if we already have exported what we plan on exporting anyway
+	if ( FPaths::FileExists( UExporter::CurrentFilename ) && FPaths::FileExists( PayloadFilename ) )
+	{
+		if ( !ExportTask->bReplaceIdentical )
+		{
+			UE_LOG( LogUsd, Log,
+				TEXT( "Skipping export of asset '%s' as the target file '%s' already exists." ),
+				*Object->GetPathName(),
+				*UExporter::CurrentFilename
+			);
+			return false;
+		}
+		// If we don't want to re-export this asset we need to check if its the same version
+		else if ( !Options->bReExportIdenticalAssets )
+		{
+			// Don't use the stage cache here as we want this stage to close within this scope in case
+			// we have to overwrite its files due to e.g. missing payload or anything like that
+			const bool bUseStageCache = false;
+			const EUsdInitialLoadSet InitialLoadSet = EUsdInitialLoadSet::LoadNone;
+			if ( UE::FUsdStage TempStage = UnrealUSDWrapper::OpenStage( *UExporter::CurrentFilename, InitialLoadSet, bUseStageCache ) )
+			{
+				if ( UE::FUsdPrim DefaultPrim = TempStage.GetDefaultPrim() )
+				{
+					FUsdUnrealAssetInfo Info = UsdUtils::GetPrimAssetInfo( DefaultPrim );
+
+					const bool bVersionMatches = !Info.Version.IsEmpty() && Info.Version == DDCKeyHash;
+
+					const bool bAssetTypeMatches = !Info.UnrealAssetType.IsEmpty()
+						&& Info.UnrealAssetType == SkeletalMesh->GetClass()->GetName();
+
+					if ( bVersionMatches && bAssetTypeMatches )
+					{
+						UE_LOG( LogUsd, Log,
+							TEXT( "Skipping export of asset '%s' as the target file '%s' already contains up-to-date exported data." ),
+							*SkeletalMesh->GetPathName(),
+							*UExporter::CurrentFilename
+						);
+
+						// Even if we're not going to export the mesh, we may still need to re-bake materials
+						if ( Options->MeshAssetOptions.bBakeMaterials )
+						{
+							TSet<UMaterialInterface*> MaterialsToBake;
+							for ( const FSkeletalMaterial& SkeletalMaterial : SkeletalMesh->GetMaterials() )
+							{
+								MaterialsToBake.Add( SkeletalMaterial.MaterialInterface );
+							}
+
+							const bool bIsAssetLayer = true;
+							UMaterialExporterUsd::ExportMaterialsForStage(
+								MaterialsToBake.Array(),
+								Options->MeshAssetOptions.MaterialBakingOptions,
+								TempStage,
+								bIsAssetLayer,
+								Options->MeshAssetOptions.bUsePayload,
+								Options->MeshAssetOptions.bRemoveUnrealMaterials,
+								ExportTask->bReplaceIdentical,
+								Options->bReExportIdenticalAssets
+							);
+							TempStage.GetRootLayer().Save();
+						}
+
+						return true;
+					}
+				}
+			}
+		}
+	}
+
+	double StartTime = FPlatformTime::Cycles64();
+
 	// UsdStage is the payload stage when exporting with payloads, or just the single stage otherwise
 	UE::FUsdStage UsdStage = UnrealUSDWrapper::NewStage( *PayloadFilename );
 	if ( !UsdStage )
@@ -204,6 +296,21 @@ bool USkeletalMeshExporterUsd::ExportBinary( UObject* Object, const TCHAR* Type,
 
 	UnrealToUsd::ConvertSkeletalMesh( SkeletalMesh, RootPrim, UsdUtils::GetDefaultTimeCode(), &AssetStage, Options->MeshAssetOptions.LowestMeshLOD, Options->MeshAssetOptions.HighestMeshLOD );
 
+	// Write asset info now that we finished exporting
+	if ( UE::FUsdPrim AssetDefaultPrim = AssetStage.GetDefaultPrim() )
+	{
+		FUsdUnrealAssetInfo Info;
+		Info.Name = SkeletalMesh->GetName();
+		Info.Identifier = UExporter::CurrentFilename;
+		Info.Version = DDCKeyHash;
+		Info.UnrealContentPath = SkeletalMesh->GetPathName();
+		Info.UnrealAssetType = SkeletalMesh->GetClass()->GetName();
+		Info.UnrealExportTime = FDateTime::Now().ToString();
+		Info.UnrealEngineVersion = FEngineVersion::Current().ToString();
+
+		UsdUtils::SetPrimAssetInfo( AssetDefaultPrim, Info );
+	}
+
 	// Bake materials and replace unrealMaterials with references to the baked files.
 	if ( Options->MeshAssetOptions.bBakeMaterials )
 	{
@@ -220,7 +327,9 @@ bool USkeletalMeshExporterUsd::ExportBinary( UObject* Object, const TCHAR* Type,
 			AssetStage,
 			bIsAssetLayer,
 			Options->MeshAssetOptions.bUsePayload,
-			Options->MeshAssetOptions.bRemoveUnrealMaterials
+			Options->MeshAssetOptions.bRemoveUnrealMaterials,
+			ExportTask->bReplaceIdentical,
+			Options->bReExportIdenticalAssets
 		);
 	}
 
diff --git a/Engine/Plugins/Importers/USDImporter/Source/USDExporter/Private/StaticMeshExporterUSD.cpp b/Engine/Plugins/Importers/USDImporter/Source/USDExporter/Private/StaticMeshExporterUSD.cpp
--- a/Engine/Plugins/Importers/USDImporter/Source/USDExporter/Private/StaticMeshExporterUSD.cpp
+++ b/Engine/Plugins/Importers/USDImporter/Source/USDExporter/Private/StaticMeshExporterUSD.cpp
@@ -8,9 +8,11 @@
 #include "USDClassesModule.h"
 #include "USDConversionUtils.h"
 #include "USDGeomMeshConversion.h"
+#include "USDLog.h"
 #include "USDMemory.h"
 #include "USDOptionsWindow.h"
 #include "USDTypesConversion.h"
+#include "USDUnrealAssetInfo.h"
 
 #include "UsdWrappers/SdfLayer.h"
 #include "UsdWrappers/SdfPath.h"
@@ -124,10 +126,12 @@ bool UStaticMeshExporterUsd::ExportBinary( UObject* Object, const TCHAR* Type, F
 	{
 		Options = Cast<UStaticMeshExporterUSDOptions>( ExportTask->Options );
 	}
-	if ( !Options && ( !ExportTask || !ExportTask->bAutomated ) )
+	if ( !Options )
 	{
 		Options = GetMutableDefault<UStaticMeshExporterUSDOptions>();
-		if ( Options )
+
+		// Prompt with an options dialog if we can
+		if ( Options && ( !ExportTask || !ExportTask->bAutomated ) )
 		{
 			Options->MeshAssetOptions.MaterialBakingOptions.TexturesDir.Path = FPaths::Combine( FPaths::GetPath( UExporter::CurrentFilename ), TEXT( "Textures" ) );
 
@@ -139,14 +143,24 @@ bool UStaticMeshExporterUsd::ExportBinary( UObject* Object, const TCHAR* Type, F
 			}
 		}
 	}
+	if ( !Options )
+	{
+		return false;
+	}
 
-	double StartTime = FPlatformTime::Cycles64();
+	// The intended export workflow has us exporting into the Ar argument, and would only later actually write to the
+	// file, potentially prompting the user beforehand if they want to overwrite an existing file or not (check
+	// UExporter::RunAssetExportTask). On the USD export workflows we always write to the file directly ourselves as
+	// its easier to just let the USD SDK do it, which means that the "Do you want to overwrite?" file prompt happens
+	// way after we already saved everything, and has no effect. Lets force bPrompt to false then, so this dialog
+	// doesn't appear under any circumstances
+	ExportTask->bPrompt = false;
 
 	// If bUsePayload is true, we'll intercept the filename so that we write the mesh data to
 	// "C:/MyFolder/file_payload.usda" and create an "asset" file "C:/MyFolder/file.usda" that uses it
 	// as a payload, pointing at the default prim
 	FString PayloadFilename = UExporter::CurrentFilename;
-	if ( Options && Options->MeshAssetOptions.bUsePayload )
+	if ( Options->MeshAssetOptions.bUsePayload )
 	{
 		FString PathPart;
 		FString FilenamePart;
@@ -161,6 +175,89 @@ bool UStaticMeshExporterUsd::ExportBinary( UObject* Object, const TCHAR* Type, F
 		PayloadFilename = FPaths::Combine( PathPart, FilenamePart + TEXT( "_payload." ) + ExtensionPart );
 	}
 
+	// Get a simple GUID hash/identifier of our mesh
+	FString DDCKeyHash;
+	if ( FStaticMeshRenderData* RenderData = StaticMesh->GetRenderData() )
+	{
+		const FString& DDCKey = RenderData->DerivedDataKey;
+		FSHA1 SHA1;
+		SHA1.UpdateWithString( *DDCKey, DDCKey.Len() );
+		SHA1.Final();
+		FSHAHash Hash;
+		SHA1.GetHash( &Hash.Hash[ 0 ] );
+		DDCKeyHash = Hash.ToString();
+	}
+
+	// Check if we already have exported what we plan on exporting anyway
+	if ( FPaths::FileExists( UExporter::CurrentFilename ) && FPaths::FileExists( PayloadFilename ) )
+	{
+		if( !ExportTask->bReplaceIdentical )
+		{
+			UE_LOG( LogUsd, Log,
+				TEXT( "Skipping export of asset '%s' as the target file '%s' already exists." ),
+				*Object->GetPathName(),
+				*UExporter::CurrentFilename
+			);
+			return false;
+		}
+		// If we don't want to re-export this asset we need to check if its the same version
+		else if ( !Options->bReExportIdenticalAssets )
+		{
+			// Don't use the stage cache here as we want this stage to close within this scope in case
+			// we have to overwrite its files due to e.g. missing payload or anything like that
+			const bool bUseStageCache = false;
+			const EUsdInitialLoadSet InitialLoadSet = EUsdInitialLoadSet::LoadNone;
+			if ( UE::FUsdStage TempStage = UnrealUSDWrapper::OpenStage( *UExporter::CurrentFilename, InitialLoadSet, bUseStageCache ) )
+			{
+				if ( UE::FUsdPrim DefaultPrim = TempStage.GetDefaultPrim() )
+				{
+					FUsdUnrealAssetInfo Info =	UsdUtils::GetPrimAssetInfo(	DefaultPrim );
+
+					const bool bVersionMatches = !Info.Version.IsEmpty() && Info.Version == DDCKeyHash;
+
+					const bool bAssetTypeMatches = !Info.UnrealAssetType.IsEmpty()
+						&& Info.UnrealAssetType == StaticMesh->GetClass()->GetName();
+
+					if ( bVersionMatches && bAssetTypeMatches )
+					{
+						UE_LOG( LogUsd, Log,
+							TEXT( "Skipping export of asset '%s' as the target file '%s' already contains up-to-date exported data." ),
+							*StaticMesh->GetPathName(),
+							*UExporter::CurrentFilename
+						);
+
+						// Even if we're not going to export the mesh, we may still need to re-bake materials
+						if ( Options->MeshAssetOptions.bBakeMaterials )
+						{
+							TSet<UMaterialInterface*> MaterialsToBake;
+							for ( const FStaticMaterial& StaticMaterial : StaticMesh->GetStaticMaterials() )
+							{
+								MaterialsToBake.Add( StaticMaterial.MaterialInterface );
+							}
+
+							const bool bIsAssetLayer = true;
+							UMaterialExporterUsd::ExportMaterialsForStage(
+								MaterialsToBake.Array(),
+								Options->MeshAssetOptions.MaterialBakingOptions,
+								TempStage,
+								bIsAssetLayer,
+								Options->MeshAssetOptions.bUsePayload,
+								Options->MeshAssetOptions.bRemoveUnrealMaterials,
+								ExportTask->bReplaceIdentical,
+								Options->bReExportIdenticalAssets
+							);
+							TempStage.GetRootLayer().Save();
+						}
+
+						return true;
+					}
+				}
+			}
+		}
+	}
+
+	double StartTime = FPlatformTime::Cycles64();
+
 	// UsdStage is the payload stage when exporting with payloads, or just the single stage otherwise
 	UE::FUsdStage UsdStage = UnrealUSDWrapper::NewStage( *PayloadFilename );
 	if ( !UsdStage )
@@ -168,11 +265,8 @@ bool UStaticMeshExporterUsd::ExportBinary( UObject* Object, const TCHAR* Type, F
 		return false;
 	}
 
-	if ( Options )
-	{
-		UsdUtils::SetUsdStageMetersPerUnit( UsdStage, Options->StageOptions.MetersPerUnit );
-		UsdUtils::SetUsdStageUpAxis( UsdStage, Options->StageOptions.UpAxis );
-	}
+	UsdUtils::SetUsdStageMetersPerUnit( UsdStage, Options->StageOptions.MetersPerUnit );
+	UsdUtils::SetUsdStageUpAxis( UsdStage, Options->StageOptions.UpAxis );
 
 	FString RootPrimPath = ( TEXT( "/" ) + UsdUtils::SanitizeUsdIdentifier( *StaticMesh->GetName() ) );
 
@@ -189,7 +283,7 @@ bool UStaticMeshExporterUsd::ExportBinary( UObject* Object, const TCHAR* Type, F
 
 	// Using payload: Convert mesh data through the asset stage (that references the payload) so that we can
 	// author mesh data on the payload layer and material data on the asset layer
-	if ( Options && Options->MeshAssetOptions.bUsePayload )
+	if ( Options->MeshAssetOptions.bUsePayload )
 	{
 		AssetStage = UnrealUSDWrapper::NewStage( *UExporter::CurrentFilename );
 		if ( AssetStage )
@@ -213,6 +307,21 @@ bool UStaticMeshExporterUsd::ExportBinary( UObject* Object, const TCHAR* Type, F
 
 	UnrealToUsd::ConvertStaticMesh( StaticMesh, RootPrim, UsdUtils::GetDefaultTimeCode(), &AssetStage, Options->MeshAssetOptions.LowestMeshLOD, Options->MeshAssetOptions.HighestMeshLOD );
 
+	// Write asset info now that we finished exporting
+	if( UE::FUsdPrim AssetDefaultPrim = AssetStage.GetDefaultPrim() )
+	{
+		FUsdUnrealAssetInfo Info;
+		Info.Name = StaticMesh->GetName();
+		Info.Identifier = UExporter::CurrentFilename;
+		Info.Version = DDCKeyHash;
+		Info.UnrealContentPath = StaticMesh->GetPathName();
+		Info.UnrealAssetType = StaticMesh->GetClass()->GetName();
+		Info.UnrealExportTime = FDateTime::Now().ToString();
+		Info.UnrealEngineVersion = FEngineVersion::Current().ToString();
+
+		UsdUtils::SetPrimAssetInfo( AssetDefaultPrim, Info );
+	}
+
 	// Bake materials and replace unrealMaterials with references to the baked files.
 	if ( Options->MeshAssetOptions.bBakeMaterials )
 	{
@@ -229,7 +338,9 @@ bool UStaticMeshExporterUsd::ExportBinary( UObject* Object, const TCHAR* Type, F
 			AssetStage,
 			bIsAssetLayer,
 			Options->MeshAssetOptions.bUsePayload,
-			Options->MeshAssetOptions.bRemoveUnrealMaterials
+			Options->MeshAssetOptions.bRemoveUnrealMaterials,
+			ExportTask->bReplaceIdentical,
+			Options->bReExportIdenticalAssets
 		);
 	}
 
diff --git a/Engine/Plugins/Importers/USDImporter/Source/USDExporter/Private/USDConversionBlueprintContext.cpp b/Engine/Plugins/Importers/USDImporter/Source/USDExporter/Private/USDConversionBlueprintContext.cpp
--- a/Engine/Plugins/Importers/USDImporter/Source/USDExporter/Private/USDConversionBlueprintContext.cpp
+++ b/Engine/Plugins/Importers/USDImporter/Source/USDExporter/Private/USDConversionBlueprintContext.cpp
@@ -8,6 +8,7 @@
 #include "USDLog.h"
 #include "USDPrimConversion.h"
 #include "USDShadeConversion.h"
+#include "USDUnrealAssetInfo.h"
 
 #include "UsdWrappers/SdfPath.h"
 #include "UsdWrappers/UsdPrim.h"
@@ -479,3 +480,36 @@ int32 UUsdConversionBlueprintContext::GetUsdStageNumFrames()
 #endif // USE_USD_SDK
 }
 
+void UUsdConversionBlueprintContext::SetPrimAssetInfo( const FString& PrimPath, const FUsdUnrealAssetInfo& Info )
+{
+#if USE_USD_SDK
+	if ( !Stage )
+	{
+		return;
+	}
+
+	UE::FUsdPrim Prim = UnrealToUsdImpl::GetPrim( Stage, PrimPath );
+	if ( !Prim )
+	{
+		return;
+	}
+
+	UsdUtils::SetPrimAssetInfo( Prim, Info );
+#endif // USE_USD_SDK
+}
+
+FUsdUnrealAssetInfo UUsdConversionBlueprintContext::GetPrimAssetInfo( const FString& PrimPath )
+{
+#if USE_USD_SDK
+	if ( Stage )
+	{
+		if ( UE::FUsdPrim Prim = UnrealToUsdImpl::GetPrim( Stage, PrimPath ) )
+		{
+			return UsdUtils::GetPrimAssetInfo( Prim );
+		}
+	}
+#endif // USE_USD_SDK
+
+	return {};
+}
+
diff --git a/Engine/Plugins/Importers/USDImporter/Source/USDExporter/Private/USDConversionBlueprintLibrary.cpp b/Engine/Plugins/Importers/USDImporter/Source/USDExporter/Private/USDConversionBlueprintLibrary.cpp
--- a/Engine/Plugins/Importers/USDImporter/Source/USDExporter/Private/USDConversionBlueprintLibrary.cpp
+++ b/Engine/Plugins/Importers/USDImporter/Source/USDExporter/Private/USDConversionBlueprintLibrary.cpp
@@ -6,6 +6,7 @@
 #include "UnrealUSDWrapper.h"
 #include "USDClassesModule.h"
 #include "USDConversionUtils.h"
+#include "USDExporterModule.h"
 #include "USDLayerUtils.h"
 #include "USDLog.h"
 
@@ -289,6 +290,11 @@ TSet<AActor*> UUsdConversionBlueprintLibrary::GetActorsToConvert( UWorld* World
 	return Result;
 }
 
+FString UUsdConversionBlueprintLibrary::GeneratePackageVersionGuidString( const UPackage* Package )
+{
+	return IUsdExporterModule::GeneratePackageVersionGuidString( Package );
+}
+
 FString UUsdConversionBlueprintLibrary::MakePathRelativeToLayer( const FString& AnchorLayerPath, const FString& PathToMakeRelative )
 {
 #if USE_USD_SDK
diff --git a/Engine/Plugins/Importers/USDImporter/Source/USDExporter/Private/USDExporterModule.cpp b/Engine/Plugins/Importers/USDImporter/Source/USDExporter/Private/USDExporterModule.cpp
--- a/Engine/Plugins/Importers/USDImporter/Source/USDExporter/Private/USDExporterModule.cpp
+++ b/Engine/Plugins/Importers/USDImporter/Source/USDExporter/Private/USDExporterModule.cpp
@@ -6,8 +6,15 @@
 #include "USDAssetOptions.h"
 #include "USDMemory.h"
 
+#include "HAL/FileManager.h"
 #include "Modules/ModuleManager.h"
 #include "PropertyEditorModule.h"
+#include "UObject/ObjectSaveContext.h"
+
+namespace UE::UsdExporterModule::Private
+{
+	TMap<FString, int32> PackagePathNameToDirtyCounter;
+}
 
 class FUsdExporterModule : public IUsdExporterModule
 {
@@ -34,16 +41,80 @@ public:
 				HighestMeshLODProperty->SetMetaData( TEXT( "ClampMax" ), LexToString( MAX_MESH_LOD_COUNT - 1 ) );
 			}
 		}
+
+		PackageMarkedDirtyEventHandle = UPackage::PackageMarkedDirtyEvent.AddLambda([](const UPackage* Package, bool bWasDirty)
+		{
+			if ( Package )
+			{
+				UE::UsdExporterModule::Private::PackagePathNameToDirtyCounter.FindOrAdd( Package->GetPathName() )++;
+			}
+		});
+
+		PackageSavedWithContextEventHandle = UPackage::PackageSavedWithContextEvent.AddLambda(
+			[]( const FString& PackageFilename, UPackage* Package, FObjectPostSaveContext ObjectSaveContext )
+			{
+				if ( Package )
+				{
+					UE::UsdExporterModule::Private::PackagePathNameToDirtyCounter.Remove( Package->GetPathName() );
+				}
+			}
+		);
 	}
 
 	virtual void ShutdownModule() override
 	{
+		UPackage::PackageMarkedDirtyEvent.Remove( PackageSavedWithContextEventHandle );
+		UPackage::PackageMarkedDirtyEvent.Remove( PackageMarkedDirtyEventHandle );
+
 		if ( FPropertyEditorModule* PropertyModule = FModuleManager::GetModulePtr< FPropertyEditorModule >( TEXT( "PropertyEditor" ) ) )
 		{
 			PropertyModule->UnregisterCustomClassLayout( TEXT( "LevelExporterUSDOptions" ) );
 			PropertyModule->UnregisterCustomClassLayout( TEXT( "LevelSequenceExporterUSDOptions" ) );
 		}
 	}
+
+private:
+	FDelegateHandle PackageMarkedDirtyEventHandle;
+	FDelegateHandle PackageSavedWithContextEventHandle;
 };
 
+FString IUsdExporterModule::GeneratePackageVersionGuidString( const UPackage* Package )
+{
+	if ( !Package )
+	{
+		return {};
+	}
+
+	// Hash package's persistent Guid
+	const FGuid& Guid = Package->GetPersistentGuid();
+	FSHA1 SHA1;
+	SHA1.Update( reinterpret_cast< const uint8* >( &Guid ), sizeof( Guid ) );
+
+	// Hash last modified date
+	FString PackageFullName = Package->GetPathName();
+	FString FileName;
+	if ( FPackageName::TryConvertLongPackageNameToFilename( PackageFullName, FileName ) )
+	{
+		FFileStatData StatData = IFileManager::Get().GetStatData( *FileName );
+		if ( StatData.bIsValid )
+		{
+			FString ModifiedTimeString = StatData.ModificationTime.ToString();
+			SHA1.UpdateWithString( *ModifiedTimeString, ModifiedTimeString.Len() );
+		}
+	}
+
+	// If this asset is currently dirty, also hash how many times it was dirtied in this session
+	// If its ever saved, we'll reset this counter but update the last saved date
+	if ( int32* DirtyCounter = UE::UsdExporterModule::Private::PackagePathNameToDirtyCounter.Find( PackageFullName ) )
+	{
+		SHA1.Update( reinterpret_cast< const uint8* >( DirtyCounter ), sizeof( *DirtyCounter ) );
+	}
+
+	SHA1.Final();
+	FSHAHash Hash;
+	SHA1.GetHash( &Hash.Hash[ 0 ] );
+	return Hash.ToString();
+}
+
 IMPLEMENT_MODULE_USD( FUsdExporterModule, USDExporter );
+
diff --git a/Engine/Plugins/Importers/USDImporter/Source/USDExporter/Public/AnimSequenceExporterUSDOptions.h b/Engine/Plugins/Importers/USDImporter/Source/USDExporter/Public/AnimSequenceExporterUSDOptions.h
--- a/Engine/Plugins/Importers/USDImporter/Source/USDExporter/Public/AnimSequenceExporterUSDOptions.h
+++ b/Engine/Plugins/Importers/USDImporter/Source/USDExporter/Public/AnimSequenceExporterUSDOptions.h
@@ -29,4 +29,11 @@ public:
 	/** Export options to use for the preview mesh, if enabled */
 	UPROPERTY( EditAnywhere, config, BlueprintReadWrite, Category = "Mesh options", meta = ( ShowOnlyInnerProperties, EditCondition = bExportPreviewMesh ) )
 	FUsdMeshAssetOptions PreviewMeshOptions;
+
+	/**
+	 * Whether to export any asset (StaticMesh, Material, etc.) even if the existing file already describes the same version of a compatible asset.
+	 * This is only checked when bReplaceIdentical is set on the asset export task. Otherwise we'll never overwrite files.
+	 */
+	UPROPERTY( EditAnywhere, config, BlueprintReadWrite, Category = "Collision", meta = ( DisplayName = "Re-export Identical Assets" ) )
+	bool bReExportIdenticalAssets = false;
 };
diff --git a/Engine/Plugins/Importers/USDImporter/Source/USDExporter/Public/LevelExporterUSDOptions.h b/Engine/Plugins/Importers/USDImporter/Source/USDExporter/Public/LevelExporterUSDOptions.h
--- a/Engine/Plugins/Importers/USDImporter/Source/USDExporter/Public/LevelExporterUSDOptions.h
+++ b/Engine/Plugins/Importers/USDImporter/Source/USDExporter/Public/LevelExporterUSDOptions.h
@@ -94,6 +94,20 @@ public:
     UPROPERTY( EditAnywhere, config, BlueprintReadWrite, Category = "Export settings", meta = ( ShowOnlyInnerProperties ) )
     FLevelExporterUSDOptionsInner Inner;
 
+	/**
+	 * Whether to export levels and LevelSequences even if the existing files already describe the same versions of compatible assets.
+	 * This is only checked when bReplaceIdentical is set on the asset export task. Otherwise we'll never overwrite files.
+	 */
+	UPROPERTY( EditAnywhere, config, BlueprintReadWrite, Category = "Collision", meta = ( DisplayName = "Re-export Identical Levels and Sequences" ) )
+	bool bReExportIdenticalLevelsAndSequences = false;
+
+	/**
+	 * Whether to export any asset (StaticMesh, Material, etc.) even if the existing file already describes the same version of a compatible asset.
+	 * This is only checked when bReplaceIdentical is set on the asset export task. Otherwise we'll never overwrite files.
+	 */
+	UPROPERTY( EditAnywhere, config, BlueprintReadWrite, Category = "Collision", meta = ( DisplayName = "Re-export Identical Assets" ) )
+	bool bReExportIdenticalAssets = false;
+
 public:
 	// We temporarily stash our export task here as a way of passing our options down to
 	// the Python exporter, that does the actual level exporting.
diff --git a/Engine/Plugins/Importers/USDImporter/Source/USDExporter/Public/LevelSequenceExporterUSDOptions.h b/Engine/Plugins/Importers/USDImporter/Source/USDExporter/Public/LevelSequenceExporterUSDOptions.h
--- a/Engine/Plugins/Importers/USDImporter/Source/USDExporter/Public/LevelSequenceExporterUSDOptions.h
+++ b/Engine/Plugins/Importers/USDImporter/Source/USDExporter/Public/LevelSequenceExporterUSDOptions.h
@@ -63,10 +63,23 @@ public:
 	UPROPERTY( EditAnywhere, config, BlueprintReadWrite, Category = "Level Sequence", meta = ( EditCondition = "bExportLevel" ) )
 	bool bUseExportedLevelAsSublayer = false;
 
+	/**
+	 * Whether to export levels and LevelSequences even if the existing files already describe the same versions of compatible assets.
+	 * This is only checked when bReplaceIdentical is set on the asset export task. Otherwise we'll never overwrite files.
+	 */
+	UPROPERTY( EditAnywhere, config, BlueprintReadWrite, Category = "Collision", meta = ( DisplayName = "Re-export Identical Levels and Sequences" ) )
+	bool bReExportIdenticalLevelsAndSequences = false;
+
+	/**
+	 * Whether to export any asset (StaticMesh, Material, etc.) even if the existing file already describes the same version of a compatible asset.
+	 * This is only checked when bReplaceIdentical is set on the asset export task. Otherwise we'll never overwrite files.
+	 */
+	UPROPERTY( EditAnywhere, config, BlueprintReadWrite, Category = "Collision", meta = ( DisplayName = "Re-export Identical Assets" ) )
+	bool bReExportIdenticalAssets = false;
+
 	UPROPERTY( EditAnywhere, BlueprintReadWrite, Category = "Level Export", meta = ( EditCondition = "bExportLevel" ) )
 	TWeakObjectPtr<UWorld> Level = nullptr;
 
 	UPROPERTY( EditAnywhere, config, BlueprintReadWrite, Category = "Level Export", meta = ( EditCondition = "bExportLevel" ) )
 	FLevelExporterUSDOptionsInner LevelExportOptions;
-
 };
diff --git a/Engine/Plugins/Importers/USDImporter/Source/USDExporter/Public/MaterialExporterUSD.h b/Engine/Plugins/Importers/USDImporter/Source/USDExporter/Public/MaterialExporterUSD.h
--- a/Engine/Plugins/Importers/USDImporter/Source/USDExporter/Public/MaterialExporterUSD.h
+++ b/Engine/Plugins/Importers/USDImporter/Source/USDExporter/Public/MaterialExporterUSD.h
@@ -28,9 +28,17 @@ public:
 	 * @param Material - Material to bake
 	 * @param Options - Options to use when baking
 	 * @param FilePath - Filepath of the destination file (e.g. "C:/MyFolder/Red.usda")
+	 * @param ReplaceIdentical - Whether to overwrite the destination file in case it already exists
+	 * @param bReExportIdenticalAssets - Whether to overwrite the destination file even if it already describes an up-to-date version of Material
 	 * @return Whether the export was successful or not.
 	 */
-	static bool ExportMaterial( const UMaterialInterface& Material, const FUsdMaterialBakingOptions& Options, const FFilePath& FilePath );
+	static bool ExportMaterial(
+		const UMaterialInterface& Material,
+		const FUsdMaterialBakingOptions& Options,
+		const FFilePath& FilePath,
+		bool bReplaceIdentical = true,
+		bool bReExportIdenticalAssets = false
+	);
 
 	/**
 	 * Exports the provided materials next to the stage's root layer on disk, replacing usages of unrealMaterials within `Stage` with references
@@ -45,6 +53,8 @@ public:
 	 *                      behaviors when authoring the material binding relationships.
 	 * @param bRemoveUnrealMaterials - Whether to remove the `unrealMaterial` attributes after replacing them with material bindings.
 	 *                                 Important because the `unrealMaterial` attributes will be used as a higher priority when determining material assignments
+	 * @param ReplaceIdentical - Whether to overwrite the destination files in case they already exist
+	 * @param bReExportIdenticalAssets - Whether to overwrite the destination files even if they already describe up-to-date versions of Materials
 	 * @return Whether the export was successful or not.
 	 */
 	static bool ExportMaterialsForStage(
@@ -53,6 +63,8 @@ public:
 		const UE::FUsdStage& Stage,
 		bool bIsAssetLayer,
 		bool bUsePayload,
-		bool bRemoveUnrealMaterials
+		bool bRemoveUnrealMaterials,
+		bool bReplaceIdentical = true,
+		bool bReExportIdenticalAssets = false
 	);
 };
diff --git a/Engine/Plugins/Importers/USDImporter/Source/USDExporter/Public/MaterialExporterUSDOptions.h b/Engine/Plugins/Importers/USDImporter/Source/USDExporter/Public/MaterialExporterUSDOptions.h
--- a/Engine/Plugins/Importers/USDImporter/Source/USDExporter/Public/MaterialExporterUSDOptions.h
+++ b/Engine/Plugins/Importers/USDImporter/Source/USDExporter/Public/MaterialExporterUSDOptions.h
@@ -21,4 +21,11 @@ class USDEXPORTER_API UMaterialExporterUSDOptions : public UObject
 public:
 	UPROPERTY( EditAnywhere, config, BlueprintReadWrite, Category = "Material baking options", meta = ( ShowOnlyInnerProperties ))
 	FUsdMaterialBakingOptions MaterialBakingOptions;
+
+	/**
+	 * Whether to export any asset (StaticMesh, Material, etc.) even if the existing file already describes the same version of a compatible asset.
+	 * This is only checked when bReplaceIdentical is set on the asset export task. Otherwise we'll never overwrite files.
+	 */
+	UPROPERTY( EditAnywhere, config, BlueprintReadWrite, Category = "Collision", meta = ( DisplayName = "Re-export Identical Assets" ) )
+	bool bReExportIdenticalAssets = false;
 };
diff --git a/Engine/Plugins/Importers/USDImporter/Source/USDExporter/Public/SkeletalMeshExporterUSDOptions.h b/Engine/Plugins/Importers/USDImporter/Source/USDExporter/Public/SkeletalMeshExporterUSDOptions.h
--- a/Engine/Plugins/Importers/USDImporter/Source/USDExporter/Public/SkeletalMeshExporterUSDOptions.h
+++ b/Engine/Plugins/Importers/USDImporter/Source/USDExporter/Public/SkeletalMeshExporterUSDOptions.h
@@ -23,4 +23,11 @@ public:
 
 	UPROPERTY( EditAnywhere, config, BlueprintReadWrite, Category = "Mesh options", meta = ( ShowOnlyInnerProperties ) )
 	FUsdMeshAssetOptions MeshAssetOptions;
+
+	/**
+	 * Whether to export any asset (StaticMesh, Material, etc.) even if the existing file already describes the same version of a compatible asset.
+	 * This is only checked when bReplaceIdentical is set on the asset export task. Otherwise we'll never overwrite files.
+	 */
+	UPROPERTY( EditAnywhere, config, BlueprintReadWrite, Category = "Collision", meta = ( DisplayName = "Re-export Identical Assets" ) )
+	bool bReExportIdenticalAssets = false;
 };
diff --git a/Engine/Plugins/Importers/USDImporter/Source/USDExporter/Public/StaticMeshExporterUSDOptions.h b/Engine/Plugins/Importers/USDImporter/Source/USDExporter/Public/StaticMeshExporterUSDOptions.h
--- a/Engine/Plugins/Importers/USDImporter/Source/USDExporter/Public/StaticMeshExporterUSDOptions.h
+++ b/Engine/Plugins/Importers/USDImporter/Source/USDExporter/Public/StaticMeshExporterUSDOptions.h
@@ -23,4 +23,11 @@ public:
 
 	UPROPERTY( EditAnywhere, config, BlueprintReadWrite, Category = "Mesh options", meta = ( ShowOnlyInnerProperties ) )
 	FUsdMeshAssetOptions MeshAssetOptions;
+
+	/**
+	 * Whether to export any asset (StaticMesh, Material, etc.) even if the existing file already describes the same version of a compatible asset.
+	 * This is only checked when bReplaceIdentical is set on the asset export task. Otherwise we'll never overwrite files.
+	 */
+	UPROPERTY( EditAnywhere, config, BlueprintReadWrite, Category = "Collision", meta = ( DisplayName = "Re-export Identical Assets" ) )
+	bool bReExportIdenticalAssets = false;
 };
diff --git a/Engine/Plugins/Importers/USDImporter/Source/USDExporter/Public/USDConversionBlueprintContext.h b/Engine/Plugins/Importers/USDImporter/Source/USDExporter/Public/USDConversionBlueprintContext.h
--- a/Engine/Plugins/Importers/USDImporter/Source/USDExporter/Public/USDConversionBlueprintContext.h
+++ b/Engine/Plugins/Importers/USDImporter/Source/USDExporter/Public/USDConversionBlueprintContext.h
@@ -3,9 +3,8 @@
 #pragma once
 
 #include "UnrealUSDWrapper.h"
-
 #include "USDStageOptions.h"
-
+#include "USDUnrealAssetInfo.h"
 #include "UsdWrappers/SdfLayer.h"
 #include "UsdWrappers/UsdStage.h"
 
@@ -167,4 +166,12 @@ public:
 	/** Returns how many frames of animation the stage would generate if imported */
 	UFUNCTION( BlueprintCallable, Category = "Conversion utils" )
 	int32 GetUsdStageNumFrames();
+
+	/** Adds to Prim the assetInfo metadata the values described in Info */
+	UFUNCTION( BlueprintCallable, Category = "Conversion utils" )
+	void SetPrimAssetInfo( const FString& PrimPath, const FUsdUnrealAssetInfo& Info );
+
+	/** Retrieves from Prim the assetInfo metadata values that we use as export metadata, when exporting Unreal assets */
+	UFUNCTION( BlueprintCallable, Category = "Conversion utils" )
+	FUsdUnrealAssetInfo GetPrimAssetInfo( const FString& PrimPath );
 };
diff --git a/Engine/Plugins/Importers/USDImporter/Source/USDExporter/Public/USDConversionBlueprintLibrary.h b/Engine/Plugins/Importers/USDImporter/Source/USDExporter/Public/USDConversionBlueprintLibrary.h
--- a/Engine/Plugins/Importers/USDImporter/Source/USDExporter/Public/USDConversionBlueprintLibrary.h
+++ b/Engine/Plugins/Importers/USDImporter/Source/USDExporter/Public/USDConversionBlueprintLibrary.h
@@ -58,6 +58,15 @@ public:
 	UFUNCTION( BlueprintCallable, Category = "World utils" )
 	static TSet<AActor*> GetActorsToConvert( UWorld* World );
 
+	/**
+	 * Generates a unique identifier string that involves the package's persistent guid, the corresponding file save
+	 * date and time, and the number of times the package has been dirtied since last being saved.
+	 * This can be used to track the version of exported assets and levels, to prevent reexporting of actors and
+	 * components.
+	 */
+	UFUNCTION( BlueprintCallable, Category = "World utils" )
+	static FString GeneratePackageVersionGuidString( const UPackage* Package );
+
 	UFUNCTION( BlueprintCallable, Category = "Layer utils" )
 	static FString MakePathRelativeToLayer( const FString& AnchorLayerPath, const FString& PathToMakeRelative );
 
diff --git a/Engine/Plugins/Importers/USDImporter/Source/USDExporter/Public/USDExporterModule.h b/Engine/Plugins/Importers/USDImporter/Source/USDExporter/Public/USDExporterModule.h
--- a/Engine/Plugins/Importers/USDImporter/Source/USDExporter/Public/USDExporterModule.h
+++ b/Engine/Plugins/Importers/USDImporter/Source/USDExporter/Public/USDExporterModule.h
@@ -2,10 +2,17 @@
 
 #pragma once
 
+#include "CoreMinimal.h"
 #include "Modules/ModuleInterface.h"
 
 class IUsdExporterModule : public IModuleInterface
 {
 public:
-
+	/**
+	 * Generates a unique identifier string that involves the package's persistent guid, the corresponding file save
+	 * date and time, and the number of times the package has been dirtied since last being saved.
+	 * This can be used to track the version of exported assets and levels, to prevent reexporting of actors and
+	 * components.
+	 */
+	static FString GeneratePackageVersionGuidString( const UPackage* Package );
 };
diff --git a/Engine/Plugins/Importers/USDImporter/Source/USDUtilities/Private/USDConversionUtils.cpp b/Engine/Plugins/Importers/USDImporter/Source/USDUtilities/Private/USDConversionUtils.cpp
--- a/Engine/Plugins/Importers/USDImporter/Source/USDUtilities/Private/USDConversionUtils.cpp
+++ b/Engine/Plugins/Importers/USDImporter/Source/USDUtilities/Private/USDConversionUtils.cpp
@@ -10,6 +10,7 @@
 #include "USDLog.h"
 #include "USDProjectSettings.h"
 #include "USDTypesConversion.h"
+#include "USDUnrealAssetInfo.h"
 
 #include "UsdWrappers/SdfLayer.h"
 #include "UsdWrappers/SdfPath.h"
@@ -2320,6 +2321,156 @@ TArray<UE::FSdfPath> UsdUtils::DuplicatePrims( const TArray<UE::FUsdPrim>& Prims
 	return Result;
 }
 
+void UsdUtils::SetPrimAssetInfo( UE::FUsdPrim& Prim, const FUsdUnrealAssetInfo& Info )
+{
+#if USE_USD_SDK
+	FScopedUsdAllocs Allocs;
+
+	pxr::UsdPrim UsdPrim{ Prim };
+	if ( !UsdPrim )
+	{
+		return;
+	}
+
+	// Just fetch the dictionary already since we'll add custom keys anyway
+	pxr::VtDictionary AssetInfoDict = UsdPrim.GetAssetInfo();
+
+	if ( !Info.Name.IsEmpty() )
+	{
+		AssetInfoDict.SetValueAtPath(
+			pxr::UsdModelAPIAssetInfoKeys->name,
+			pxr::VtValue{ UnrealToUsd::ConvertString( *Info.Name ).Get() }
+		);
+	}
+
+	if ( !Info.Identifier.IsEmpty() )
+	{
+		AssetInfoDict.SetValueAtPath(
+			pxr::UsdModelAPIAssetInfoKeys->identifier,
+			pxr::VtValue{ pxr::SdfAssetPath{ UnrealToUsd::ConvertString( *Info.Identifier ).Get() } }
+		);
+	}
+
+	if ( !Info.Version.IsEmpty() )
+	{
+		AssetInfoDict.SetValueAtPath(
+			pxr::UsdModelAPIAssetInfoKeys->version,
+			pxr::VtValue{ UnrealToUsd::ConvertString( *Info.Version ).Get() }
+		);
+	}
+
+	if ( !Info.UnrealContentPath.IsEmpty() )
+	{
+		AssetInfoDict.SetValueAtPath(
+			UnrealIdentifiers::UnrealContentPath,
+			pxr::VtValue{ UnrealToUsd::ConvertString( *Info.UnrealContentPath ).Get() }
+		);
+	}
+
+	if ( !Info.UnrealAssetType.IsEmpty() )
+	{
+		AssetInfoDict.SetValueAtPath(
+			UnrealIdentifiers::UnrealAssetType,
+			pxr::VtValue{ UnrealToUsd::ConvertString( *Info.UnrealAssetType ).Get() }
+		);
+	}
+
+	if ( !Info.UnrealExportTime.IsEmpty() )
+	{
+		AssetInfoDict.SetValueAtPath(
+			UnrealIdentifiers::UnrealExportTime,
+			pxr::VtValue{ UnrealToUsd::ConvertString( *Info.UnrealExportTime ).Get() }
+		);
+	}
+
+	if ( !Info.UnrealEngineVersion.IsEmpty() )
+	{
+		AssetInfoDict.SetValueAtPath(
+			UnrealIdentifiers::UnrealEngineVersion,
+			pxr::VtValue{ UnrealToUsd::ConvertString( *Info.UnrealEngineVersion ).Get() }
+		);
+	}
+
+	UsdPrim.SetAssetInfo( AssetInfoDict );
+#endif // USE_USD_SDK
+}
+
+FUsdUnrealAssetInfo UsdUtils::GetPrimAssetInfo( const UE::FUsdPrim& Prim )
+{
+	FUsdUnrealAssetInfo Result;
+
+#if USE_USD_SDK
+	FScopedUsdAllocs Allocs;
+
+	pxr::UsdPrim UsdPrim{ Prim };
+	if ( !UsdPrim )
+	{
+		return Result;
+	}
+
+	// Just fetch the dictionary already since we'll fetch custom keys anyway
+	pxr::VtDictionary AssetInfoDict = UsdPrim.GetAssetInfo();
+
+	if ( const pxr::VtValue* Value = AssetInfoDict.GetValueAtPath( pxr::UsdModelAPIAssetInfoKeys->name ) )
+	{
+		if ( Value->IsHolding<std::string>() )
+		{
+			Result.Name = UsdToUnreal::ConvertString( Value->Get<std::string>() );
+		}
+	}
+
+	if ( const pxr::VtValue* Value = AssetInfoDict.GetValueAtPath( pxr::UsdModelAPIAssetInfoKeys->identifier ) )
+	{
+		if ( Value->IsHolding<pxr::SdfAssetPath>() )
+		{
+			Result.Identifier = UsdToUnreal::ConvertString( Value->Get<pxr::SdfAssetPath>().GetAssetPath() );
+		}
+	}
+
+	if ( const pxr::VtValue* Value = AssetInfoDict.GetValueAtPath( pxr::UsdModelAPIAssetInfoKeys->version ) )
+	{
+		if ( Value->IsHolding<std::string>() )
+		{
+			Result.Version = UsdToUnreal::ConvertString( Value->Get<std::string>() );
+		}
+	}
+
+	if ( const pxr::VtValue* Value = AssetInfoDict.GetValueAtPath( UnrealIdentifiers::UnrealContentPath ) )
+	{
+		if ( Value->IsHolding<std::string>() )
+		{
+			Result.UnrealContentPath = UsdToUnreal::ConvertString( Value->Get<std::string>() );
+		}
+	}
+
+	if ( const pxr::VtValue* Value = AssetInfoDict.GetValueAtPath( UnrealIdentifiers::UnrealAssetType ) )
+	{
+		if ( Value->IsHolding<std::string>() )
+		{
+			Result.UnrealAssetType = UsdToUnreal::ConvertString( Value->Get<std::string>() );
+		}
+	}
+
+	if ( const pxr::VtValue* Value = AssetInfoDict.GetValueAtPath( UnrealIdentifiers::UnrealExportTime ) )
+	{
+		if ( Value->IsHolding<std::string>() )
+		{
+			Result.UnrealExportTime = UsdToUnreal::ConvertString( Value->Get<std::string>() );
+		}
+	}
+
+	if ( const pxr::VtValue* Value = AssetInfoDict.GetValueAtPath( UnrealIdentifiers::UnrealEngineVersion ) )
+	{
+		if ( Value->IsHolding<std::string>() )
+		{
+			Result.UnrealEngineVersion = UsdToUnreal::ConvertString( Value->Get<std::string>() );
+		}
+	}
+#endif // USE_USD_SDK
+
+	return Result;
+}
+
 #if USE_USD_SDK
 #undef LOCTEXT_NAMESPACE
 #endif
diff --git a/Engine/Plugins/Importers/USDImporter/Source/USDUtilities/Public/USDConversionUtils.h b/Engine/Plugins/Importers/USDImporter/Source/USDUtilities/Public/USDConversionUtils.h
--- a/Engine/Plugins/Importers/USDImporter/Source/USDUtilities/Public/USDConversionUtils.h
+++ b/Engine/Plugins/Importers/USDImporter/Source/USDUtilities/Public/USDConversionUtils.h
@@ -47,6 +47,7 @@ class USceneComponent;
 class UUsdAssetImportData;
 enum class EUsdDuplicateType : uint8;
 enum class EUsdUpAxis : uint8;
+struct FUsdUnrealAssetInfo;
 namespace UE
 {
 	class FUsdPrim;
@@ -345,5 +346,11 @@ namespace UsdUtils
 		EUsdDuplicateType DuplicateType,
 		const UE::FSdfLayer& TargetLayer = UE::FSdfLayer{}
 	);
+
+	/** Adds to Prim the assetInfo metadata the values described in Info */
+	USDUTILITIES_API void SetPrimAssetInfo( UE::FUsdPrim& Prim, const FUsdUnrealAssetInfo& Info );
+
+	/** Retrieves from Prim the assetInfo metadata values that we use as export metadata, when exporting Unreal assets */
+	USDUTILITIES_API FUsdUnrealAssetInfo GetPrimAssetInfo( const UE::FUsdPrim& Prim );
 }
 
diff --git a/Engine/Plugins/Importers/USDImporter/Source/UnrealUSDWrapper/Private/UnrealUSDWrapper.cpp b/Engine/Plugins/Importers/USDImporter/Source/UnrealUSDWrapper/Private/UnrealUSDWrapper.cpp
--- a/Engine/Plugins/Importers/USDImporter/Source/UnrealUSDWrapper/Private/UnrealUSDWrapper.cpp
+++ b/Engine/Plugins/Importers/USDImporter/Source/UnrealUSDWrapper/Private/UnrealUSDWrapper.cpp
@@ -148,6 +148,11 @@ namespace UnrealIdentifiers
 	const TfToken GroomBindingAPI = TfToken( "GroomBindingAPI" );
 	const TfToken UnrealGroomToBind = TfToken( "unreal:groomBinding:groom" );
 	const TfToken UnrealGroomReferenceMesh = TfToken( "unreal:groomBinding:referenceMesh" );
+
+	const TfToken UnrealContentPath = TfToken( "unreal:contentPath" );
+	const TfToken UnrealAssetType = TfToken( "unreal:assetType" );
+	const TfToken UnrealExportTime = TfToken( "unreal:exportTime" );
+	const TfToken UnrealEngineVersion = TfToken( "unreal:engineVersion" );
 }
 
 std::string FUsdAttribute::GetUnrealPropertyPath( const pxr::UsdAttribute& Attribute )
diff --git a/Engine/Plugins/Importers/USDImporter/Source/UnrealUSDWrapper/Public/UnrealUSDWrapper.h b/Engine/Plugins/Importers/USDImporter/Source/UnrealUSDWrapper/Public/UnrealUSDWrapper.h
--- a/Engine/Plugins/Importers/USDImporter/Source/UnrealUSDWrapper/Public/UnrealUSDWrapper.h
+++ b/Engine/Plugins/Importers/USDImporter/Source/UnrealUSDWrapper/Public/UnrealUSDWrapper.h
@@ -416,6 +416,11 @@ namespace UnrealIdentifiers
 	extern UNREALUSDWRAPPER_API const pxr::TfToken GroomBindingAPI;
 	extern UNREALUSDWRAPPER_API const pxr::TfToken UnrealGroomToBind;
 	extern UNREALUSDWRAPPER_API const pxr::TfToken UnrealGroomReferenceMesh;
+
+	extern UNREALUSDWRAPPER_API const pxr::TfToken UnrealContentPath;
+	extern UNREALUSDWRAPPER_API const pxr::TfToken UnrealAssetType;
+	extern UNREALUSDWRAPPER_API const pxr::TfToken UnrealExportTime;
+	extern UNREALUSDWRAPPER_API const pxr::TfToken UnrealEngineVersion;
 #endif // #if USE_USD_SDK
 
 	extern UNREALUSDWRAPPER_API const TCHAR* Invisible;
diff --git a/Engine/Source/Runtime/Engine/Public/MaterialShared.h b/Engine/Source/Runtime/Engine/Public/MaterialShared.h
--- a/Engine/Source/Runtime/Engine/Public/MaterialShared.h
+++ b/Engine/Source/Runtime/Engine/Public/MaterialShared.h
@@ -984,7 +984,7 @@ public:
 
 #if WITH_EDITOR
 	/** Hashes the material-specific part of this shader map Id. */
-	void GetMaterialHash(FSHAHash& OutHash) const;
+	ENGINE_API void GetMaterialHash(FSHAHash& OutHash) const;
 #endif
 
 	/** 
