USD: Isolated layer/stage mode (second implementation).

diff --git a/Engine/Plugins/Importers/USDImporter/Source/USDStage/Private/USDStageActor.cpp b/Engine/Plugins/Importers/USDImporter/Source/USDStage/Private/USDStageActor.cpp
--- a/Engine/Plugins/Importers/USDImporter/Source/USDStage/Private/USDStageActor.cpp
+++ b/Engine/Plugins/Importers/USDImporter/Source/USDStage/Private/USDStageActor.cpp
@@ -404,7 +404,7 @@ struct FUsdStageActorImpl
 	 * correct value of bHiddenInGame.
 	 * Note that bHiddenInGame corresponds to computed visibility, and the component tags correspond to individual prim-level visibilities
 	 */
-	static void MakeVisible( UUsdPrimTwin& UsdPrimTwin, UE::FUsdStage& Stage )
+	static void MakeVisible( UUsdPrimTwin& UsdPrimTwin, const UE::FUsdStage& Stage )
 	{
 		// Find the highest invisible prim parent: Nothing above this can possibly change with what we're doing
 		UUsdPrimTwin* Iter = &UsdPrimTwin;
@@ -767,6 +767,45 @@ AUsdStageActor::AUsdStageActor()
 	}
 }
 
+void AUsdStageActor::IsolateLayer( const UE::FSdfLayer& Layer )
+{
+	if ( IsolatedStage && IsolatedStage.GetRootLayer() == Layer )
+	{
+		return;
+	}
+
+	// Stop isolating
+	if ( !Layer || Layer == UsdStage.GetRootLayer() )
+	{
+		IsolatedStage = UE::FUsdStage{};
+
+		UsdListener.Register( UsdStage );
+
+		LoadUsdStage();
+		return;
+	}
+
+	if ( UsdStage )
+	{
+		// Ideally we'd check to see if Layer belonged to the stage first, but there isn't any foolproof way to check
+		// for that because if the layer is muted it won't show up on the layer stack at all. Its also valid to mute
+		// any layer anyway (even one that doesn't belong to the stage yet), so checking the list of muted layers also
+		// wouldn't improve anything: It would always be possible to trick that code into isolating an external layer
+		// by just muting the layer first...
+
+		// We really want our own stage for this and not something from the stage cache.
+		// Plus, this means its easier to cleanup: Just drop our IsolatedStage
+		const bool bUseStageCache = false;
+		IsolatedStage = UnrealUSDWrapper::OpenStage( Layer, {}, EUsdInitialLoadSet::LoadAll, bUseStageCache );
+		IsolatedStage.SetEditTarget( IsolatedStage.GetRootLayer() );
+		IsolatedStage.SetInterpolationType( InterpolationType );
+
+		UsdListener.Register( IsolatedStage );
+
+		LoadUsdStage();
+	}
+}
+
 void AUsdStageActor::OnUsdObjectsChanged( const UsdUtils::FObjectChangesByPath& InfoChanges, const UsdUtils::FObjectChangesByPath& ResyncChanges )
 {
 #if USE_USD_SDK
@@ -794,7 +833,8 @@ void AUsdStageActor::OnUsdObjectsChanged( const UsdUtils::FObjectChangesByPath&
 	// and the UE::FUsdStage will still be opened and valid (even though we intend on closing/changing it when we transact). It could be problematic/wasteful if we
 	// responded to those notices, so just early out here. We can do this check because our RootLayer property will already have the new value
 	{
-		const UE::FSdfLayer& StageRoot = Stage.GetRootLayer();
+		const UE::FUsdStage& BaseStage = GetBaseUsdStage();
+		const UE::FSdfLayer& StageRoot = BaseStage.GetRootLayer();
 		if ( !StageRoot )
 		{
 			return;
@@ -932,7 +972,7 @@ void AUsdStageActor::OnUsdObjectsChanged( const UsdUtils::FObjectChangesByPath&
 		// However, if we know that only Prim resynced, we can traverse the tree root down and if we reach Prim and its collapsing state
 		// hasn't updated from before, we don't have to update its subtree at all, or sibling subtrees.
 		TSharedRef< FUsdSchemaTranslationContext > TranslationContext = FUsdStageActorImpl::CreateUsdSchemaTranslationContext( this, TEXT( "/" ) );
-		InfoCache->RebuildCacheForSubtree( UsdStage.GetPseudoRoot(), TranslationContext.Get() );
+		InfoCache->RebuildCacheForSubtree( Stage.GetPseudoRoot(), TranslationContext.Get() );
 	}
 
 	for ( const TPair< FString, bool >& PrimChangedInfo : SortedPrimsChangedList )
@@ -1320,19 +1360,29 @@ UE::FUsdStage& AUsdStageActor::GetUsdStage()
 {
 	OpenUsdStage();
 
-	return UsdStage;
+	return IsolatedStage ? IsolatedStage : UsdStage;
 }
 
 const UE::FUsdStage& AUsdStageActor::GetUsdStage() const
+{
+	return IsolatedStage ? IsolatedStage : UsdStage;
+}
+
+const UE::FUsdStage& AUsdStageActor::GetBaseUsdStage() const
 {
 	return UsdStage;
 }
 
+const UE::FUsdStage& AUsdStageActor::GetIsolatedUsdStage() const
+{
+	return IsolatedStage;
+}
+
 UE::FUsdStage& AUsdStageActor::GetOrLoadUsdStage()
 {
 	OpenUsdStage();
 
-	return UsdStage;
+	return IsolatedStage ? IsolatedStage : UsdStage;
 }
 
 void AUsdStageActor::SetRootLayer( const FString& RootFilePath )
@@ -1480,14 +1530,15 @@ void AUsdStageActor::SetTime(float InTime)
 
 USceneComponent* AUsdStageActor::GetGeneratedComponent( const FString& PrimPath )
 {
-	if ( !UsdStage )
+	const UE::FUsdStage& CurrentStage = static_cast< const AUsdStageActor* >( this )->GetUsdStage();
+	if ( !CurrentStage )
 	{
 		return nullptr;
 	}
 
 	// We can't query our InfoCache with invalid paths, as we're using ensures to track when we miss the cache (which shouldn't ever happen)
 	UE::FSdfPath UsdPath{ *PrimPath };
-	if ( !UsdStage.GetPrimAtPath( UsdPath ) )
+	if ( !CurrentStage.GetPrimAtPath( UsdPath ) )
 	{
 		return nullptr;
 	}
@@ -1508,14 +1559,15 @@ USceneComponent* AUsdStageActor::GetGeneratedComponent( const FString& PrimPath
 
 TArray<UObject*> AUsdStageActor::GetGeneratedAssets( const FString& PrimPath )
 {
-	if ( !UsdStage )
+	const UE::FUsdStage& CurrentStage = static_cast<const AUsdStageActor*>(this)->GetUsdStage();
+	if ( !CurrentStage )
 	{
 		return {};
 	}
 
 	// We can't query our InfoCache with invalid paths, as we're using ensures to track when we miss the cache (which shouldn't ever happen)
 	UE::FSdfPath UsdPath{ *PrimPath };
-	if ( !UsdStage.GetPrimAtPath( UsdPath ) )
+	if ( !CurrentStage.GetPrimAtPath( UsdPath ) )
 	{
 		return {};
 	}
@@ -1622,6 +1674,8 @@ void AUsdStageActor::OpenUsdStage()
 	OnPreStageChanged.Broadcast();
 
 	UsdStage = UnrealUSDWrapper::OpenStage( *AbsPath, InitialLoadSet );
+	IsolatedStage = UE::FUsdStage{};
+
 	if ( UsdStage )
 	{
 		UsdStage.SetEditTarget( UsdStage.GetRootLayer() );
@@ -1648,6 +1702,7 @@ void AUsdStageActor::CloseUsdStage()
 
 	FUsdStageActorImpl::DiscardStage( UsdStage, this );
 	UsdStage = UE::FUsdStage();
+	IsolatedStage = UE::FUsdStage();  // We don't keep our isolated stages on the stage cache
 	LevelSequenceHelper.Init( UE::FUsdStage() ); // Drop the helper's reference to the stage
 	OnStageChanged.Broadcast();
 }
@@ -1743,6 +1798,7 @@ void AUsdStageActor::OnObjectsReplaced( const TMap<UObject*, UObject*>& ObjectRe
 			// Can't just move out of this one as TUsdStore expects its TOptional to always contain a value, and we may
 			// still need to use the bool operator on it to test for validity
 			NewActor->UsdStage = UsdStage;
+			NewActor->IsolatedStage = IsolatedStage;
 
 			NewActor->LevelSequenceHelper = MoveTemp( LevelSequenceHelper );
 			NewActor->LevelSequence = LevelSequence;
@@ -1826,18 +1882,29 @@ void AUsdStageActor::LoadUsdStage()
 
 	FScopedUsdMessageLog ScopedMessageLog;
 
-	// If we're in here we don't expect our current stage to be the same as the new stage we're trying to load, so
-	// get rid of it so that OpenUsdStage can open it
-	UsdStage = UE::FUsdStage();
-
-	OpenUsdStage();
-	if ( !UsdStage )
+	// If we have an isolated stage we may have been asked to refresh to display it: We should keep our UsdStage opened
+	UE::FUsdStage StageToLoad;
+	if ( IsolatedStage )
 	{
-		OnStageChanged.Broadcast();
-		return;
+		StageToLoad = IsolatedStage;
 	}
+	else
+	{
+		// If we're in here we don't expect our current stage to be the same as the new stage we're trying to load, so
+		// get rid of it so that OpenUsdStage can open it
+		UsdStage = UE::FUsdStage();
 
-	UsdStage.SetInterpolationType( InterpolationType );
+		OpenUsdStage();
+		if ( !UsdStage )
+		{
+			OnStageChanged.Broadcast();
+			return;
+		}
+
+		StageToLoad = UsdStage;
+	}
+
+	StageToLoad.SetInterpolationType( InterpolationType );
 
 	ReloadAnimations();
 
@@ -1854,20 +1921,20 @@ void AUsdStageActor::LoadUsdStage()
 	{
 		InfoCache = MakeShared<FUsdInfoCache>();
 	}
-	InfoCache->RebuildCacheForSubtree( UsdStage.GetPseudoRoot(), TranslationContext.Get() );
+	InfoCache->RebuildCacheForSubtree( StageToLoad.GetPseudoRoot(), TranslationContext.Get() );
 
 	SlowTask.EnterProgressFrame( 0.7f );
-	LoadAssets( *TranslationContext, UsdStage.GetPseudoRoot() );
+	LoadAssets( *TranslationContext, StageToLoad.GetPseudoRoot() );
 
 	SlowTask.EnterProgressFrame( 0.2f );
-	UpdatePrim( UsdStage.GetPseudoRoot().GetPrimPath(), true, *TranslationContext );
+	UpdatePrim( StageToLoad.GetPseudoRoot().GetPrimPath(), true, *TranslationContext );
 
 	TranslationContext->CompleteTasks();
 
 	// Keep our old Time value if we're loading the stage during initialization, so that we can save/load Time values
-	if ( UsdStage.GetRootLayer() && IsActorInitialized() )
+	if ( StageToLoad.GetRootLayer() && IsActorInitialized() )
 	{
-		SetTime( UsdStage.GetRootLayer().GetStartTimeCode() );
+		SetTime( StageToLoad.GetRootLayer().GetStartTimeCode() );
 
 		// If we're an instance of a blueprint that derives the stage actor and we're in the editor preview world, it means we're the
 		// blueprint preview actor. We (the instance) will load the stage and update our Time to StartTimeCode, but our CDO will not.
@@ -1908,7 +1975,7 @@ void AUsdStageActor::LoadUsdStage()
 	FUsdStageActorImpl::SendAnalytics(
 		this,
 		ElapsedSeconds,
-		UsdUtils::GetUsdStageNumFrames( UsdStage ),
+		UsdUtils::GetUsdStageNumFrames( StageToLoad ),
 		FPaths::GetExtension( RootLayer.FilePath )
 	);
 #endif // #if USE_USD_SDK
@@ -2011,7 +2078,8 @@ void AUsdStageActor::ReloadAnimations()
 	}
 	TGuardValue<bool> ReentrantGuard( bIsReentrant, true );
 
-	if ( !UsdStage )
+	const UE::FUsdStage& CurrentStage = GetOrLoadUsdStage();
+	if ( !CurrentStage )
 	{
 		return;
 	}
@@ -2036,7 +2104,7 @@ void AUsdStageActor::ReloadAnimations()
 		const bool bMarkDirty = false;
 		Modify(bMarkDirty);
 
-		LevelSequence = LevelSequenceHelper.Init( UsdStage );
+		LevelSequence = LevelSequenceHelper.Init( CurrentStage );
 		LevelSequenceHelper.BindToUsdStageActor( this );
 
 #if WITH_EDITOR
@@ -2341,7 +2409,8 @@ void AUsdStageActor::PostRegisterAllComponents()
 
 	// This may fail if our stage happened to not spawn any components, actors or assets, but by that
 	// point "being loaded" doesn't really mean anything anyway
-	const bool bStageIsLoaded = UsdStage && ( ( RootUsdTwin && RootUsdTwin->GetSceneComponent() != nullptr ) || ( AssetCache && AssetCache->GetNumAssets() > 0 ) );
+	const bool bStageIsLoaded = GetBaseUsdStage()
+		&& ( ( RootUsdTwin && RootUsdTwin->GetSceneComponent() != nullptr ) || ( AssetCache && AssetCache->GetNumAssets() > 0 ) );
 
 	// Blocks loading stage when going into PIE, if we already have something loaded (we'll want to duplicate stuff instead).
 	// We need to allow loading when going into PIE when we have nothing loaded yet because the MovieRenderQueue (or other callers)
@@ -2433,7 +2502,8 @@ void AUsdStageActor::UnregisterAllComponents( bool bForReregister )
 	}
 #endif // WITH_EDITOR
 
-	const bool bStageIsLoaded = UsdStage && ( ( RootUsdTwin && RootUsdTwin->GetSceneComponent() != nullptr ) || ( AssetCache && AssetCache->GetNumAssets() > 0 ) );
+	const bool bStageIsLoaded = GetBaseUsdStage()
+		&& ( ( RootUsdTwin && RootUsdTwin->GetSceneComponent() != nullptr ) || ( AssetCache && AssetCache->GetNumAssets() > 0 ) );
 
 	UWorld* World = GetWorld();
 	if ( bIsTransitioningIntoPIE && bStageIsLoaded && ( !World || World->WorldType == EWorldType::PIE ) )
@@ -2464,14 +2534,15 @@ void AUsdStageActor::PostUnregisterAllComponents()
 
 void AUsdStageActor::OnPreUsdImport( FString FilePath )
 {
-	if ( !UsdStage || !HasAuthorityOverStage() )
+	const UE::FUsdStage& CurrentStage = static_cast< const AUsdStageActor* >( this )->GetUsdStage();
+	if ( !CurrentStage || !HasAuthorityOverStage() )
 	{
 		return;
 	}
 
 	// Stop listening to events because a USD import may temporarily modify the stage (e.g. when importing with
 	// a different MetersPerUnit value), and we don't want to respond to the notices in the meantime
-	FString RootPath = UsdStage.GetRootLayer().GetRealPath();
+	FString RootPath = CurrentStage.GetRootLayer().GetRealPath();
 	FPaths::NormalizeFilename( RootPath );
 	if ( RootPath == FilePath )
 	{
@@ -2481,13 +2552,14 @@ void AUsdStageActor::OnPreUsdImport( FString FilePath )
 
 void AUsdStageActor::OnPostUsdImport( FString FilePath )
 {
-	if ( !UsdStage || !HasAuthorityOverStage() )
+	const UE::FUsdStage& CurrentStage = static_cast< const AUsdStageActor* >( this )->GetUsdStage();
+	if ( !CurrentStage || !HasAuthorityOverStage() )
 	{
 		return;
 	}
 
 	// Resume listening to events
-	FString RootPath = UsdStage.GetRootLayer().GetRealPath();
+	FString RootPath = CurrentStage.GetRootLayer().GetRealPath();
 	FPaths::NormalizeFilename( RootPath );
 	if ( RootPath == FilePath )
 	{
@@ -2594,6 +2666,8 @@ void AUsdStageActor::OnObjectPropertyChanged( UObject* ObjectBeingModified, FPro
 		}
 	}
 
+	const UE::FUsdStage& CurrentStage = static_cast< const AUsdStageActor* >( this )->GetUsdStage();
+
 	FString PrimPath = ObjectsToWatch[ PrimObject ];
 
 	if ( UUsdPrimTwin* UsdPrimTwin = GetRootPrimTwin()->Find( PrimPath ) )
@@ -2608,14 +2682,14 @@ void AUsdStageActor::OnObjectPropertyChanged( UObject* ObjectBeingModified, FPro
 
 		if ( PrimSceneComponent )
 		{
-			if ( UsdStage )
+			if ( CurrentStage )
 			{
 				// This block is important, as it not only prevents us from getting into infinite loops with the USD notices,
 				// but it also guarantees that if we have an object property change, the corresponding stage notice is not also
 				// independently saved to the transaction via the UUsdTransactor, which would be duplication
 				FScopedBlockNoticeListening BlockNotices( this );
 
-				UE::FUsdPrim UsdPrim = UsdStage.GetPrimAtPath( UE::FSdfPath( *PrimPath ) );
+				UE::FUsdPrim UsdPrim = CurrentStage.GetPrimAtPath( UE::FSdfPath( *PrimPath ) );
 
 				// We want to keep component visibilities in sync with USD, which uses inherited visibilities
 				// To accomplish that while blocking notices we must always propagate component visibility changes manually.
@@ -2630,7 +2704,7 @@ void AUsdStageActor::OnObjectPropertyChanged( UObject* ObjectBeingModified, FPro
 					}
 					else
 					{
-						FUsdStageActorImpl::MakeVisible( *UsdPrimTwin, UsdStage );
+						FUsdStageActorImpl::MakeVisible( *UsdPrimTwin, CurrentStage );
 					}
 				}
 
@@ -2638,11 +2712,11 @@ void AUsdStageActor::OnObjectPropertyChanged( UObject* ObjectBeingModified, FPro
 
 				UnrealToUsd::ConvertLiveLinkProperties( Controller ? Cast<UActorComponent>( Controller ) : PrimSceneComponent, UsdPrim );
 
-				UnrealToUsd::ConvertSceneComponent( UsdStage, PrimSceneComponent, UsdPrim );
+				UnrealToUsd::ConvertSceneComponent( CurrentStage, PrimSceneComponent, UsdPrim );
 
 				if ( UMeshComponent* MeshComponent = Cast< UMeshComponent >( PrimSceneComponent ) )
 				{
-					UnrealToUsd::ConvertMeshComponent( UsdStage, MeshComponent, UsdPrim );
+					UnrealToUsd::ConvertMeshComponent( CurrentStage, MeshComponent, UsdPrim );
 				}
 				else if ( UsdPrim && UsdPrim.IsA( TEXT( "Camera" ) ) )
 				{
@@ -2784,12 +2858,13 @@ bool AUsdStageActor::HasAuthorityOverStage() const
 void AUsdStageActor::OnSkelAnimationBaked( const FString& SkelRootPrimPath )
 {
 #if USE_USD_SDK
-	if ( !UsdStage || !GRegenerateSkeletalAssetsOnControlRigBake )
+	const UE::FUsdStage& CurrentStage = static_cast< const AUsdStageActor* >( this )->GetUsdStage();
+	if ( !CurrentStage || !GRegenerateSkeletalAssetsOnControlRigBake )
 	{
 		return;
 	}
 
-	UE::FUsdPrim SkelRootPrim = UsdStage.GetPrimAtPath( UE::FSdfPath{ *SkelRootPrimPath } );
+	UE::FUsdPrim SkelRootPrim = CurrentStage.GetPrimAtPath( UE::FSdfPath{ *SkelRootPrimPath } );
 	if ( !SkelRootPrim || !SkelRootPrim.IsA( TEXT( "SkelRoot" ) ) )
 	{
 		return;
@@ -2936,7 +3011,8 @@ void AUsdStageActor::LoadAssets( FUsdSchemaTranslationContext& TranslationContex
 void AUsdStageActor::AnimatePrims()
 {
 	// Don't try to animate if we don't have a stage opened
-	if ( !UsdStage )
+	const UE::FUsdStage& CurrentStage = static_cast< const AUsdStageActor* >( this )->GetUsdStage();
+	if ( !CurrentStage )
 	{
 		return;
 	}
@@ -2961,7 +3037,7 @@ void AUsdStageActor::AnimatePrims()
 		UE::FSdfPath PrimPath( *PrimToAnimate );
 
 		IUsdSchemasModule& SchemasModule = FModuleManager::Get().LoadModuleChecked< IUsdSchemasModule >( "USDSchemas" );
-		if ( TSharedPtr< FUsdSchemaTranslator > SchemaTranslator = SchemasModule.GetTranslatorRegistry().CreateTranslatorForSchema( TranslationContext, UE::FUsdTyped( UsdStage.GetPrimAtPath( PrimPath ) ) ) )
+		if ( TSharedPtr< FUsdSchemaTranslator > SchemaTranslator = SchemasModule.GetTranslatorRegistry().CreateTranslatorForSchema( TranslationContext, UE::FUsdTyped( CurrentStage.GetPrimAtPath( PrimPath ) ) ) )
 		{
 			if ( UUsdPrimTwin* UsdPrimTwin = GetRootPrimTwin()->Find( PrimToAnimate ) )
 			{
diff --git a/Engine/Plugins/Importers/USDImporter/Source/USDStage/Public/USDStageActor.h b/Engine/Plugins/Importers/USDImporter/Source/USDStage/Public/USDStageActor.h
--- a/Engine/Plugins/Importers/USDImporter/Source/USDStage/Public/USDStageActor.h
+++ b/Engine/Plugins/Importers/USDImporter/Source/USDStage/Public/USDStageActor.h
@@ -197,6 +197,8 @@ public:
 public:
 	AUsdStageActor();
 
+	USDSTAGE_API void IsolateLayer( const UE::FSdfLayer& Layer );
+
 	USDSTAGE_API void Reset() override;
 	void Refresh() const;
 	void ReloadAnimations();
@@ -275,9 +277,19 @@ public:
 	UE_DEPRECATED( 4.27, "Use the const version if you don't wish to load the stage on-demand, or use GetOrLoadUsdStage if you do" )
 	USDSTAGE_API UE::FUsdStage& GetUsdStage();
 
+	// Loads the stage with RootLayer if its not loaded already, and returns the either the isolated stage (if any) or
+	// the base stage
 	USDSTAGE_API UE::FUsdStage& GetOrLoadUsdStage();
+
+	// Returns either the isolated stage (if any) or the base stage
 	USDSTAGE_API const UE::FUsdStage& GetUsdStage() const;
 
+	// Always returns the base stage, regardless of whether we have an isolated stage or not
+	USDSTAGE_API const UE::FUsdStage& GetBaseUsdStage() const;
+
+	// Always returns the isolated stage, being an invalid stage in case we're not isolating anything
+	USDSTAGE_API const UE::FUsdStage& GetIsolatedUsdStage() const;
+
 	FUsdListener& GetUsdListener() { return UsdListener; }
 	const FUsdListener& GetUsdListener() const { return UsdListener; }
 
@@ -318,8 +330,13 @@ private:
 	 */
 	TMap< FString, TMap< FString, int32 > > MaterialToPrimvarToUVIndex;
 
+	// The main UsdStage that is currently opened
 	UE::FUsdStage UsdStage;
 
+	// Another stage that has as root layer one of the non-root local layers of UsdStage. This is the stage we'll
+	// be displaying if it is valid, otherwise we'll be displaying UsdStage directly.
+	UE::FUsdStage IsolatedStage;
+
 	/**
 	 * We use PostRegisterAllComponents and PostUnregisterAllComponents as main entry points to decide when to load/unload
 	 * the USD stage. These are the three exceptions we must avoid though:
diff --git a/Engine/Plugins/Importers/USDImporter/Source/USDStageEditor/Private/SUSDLayersTreeView.cpp b/Engine/Plugins/Importers/USDImporter/Source/USDStageEditor/Private/SUSDLayersTreeView.cpp
--- a/Engine/Plugins/Importers/USDImporter/Source/USDStageEditor/Private/SUSDLayersTreeView.cpp
+++ b/Engine/Plugins/Importers/USDImporter/Source/USDStageEditor/Private/SUSDLayersTreeView.cpp
@@ -69,9 +69,16 @@ namespace UE::USDLayersTreeViewImpl::Private
 	}
 }
 
-class FUsdLayerNameColumn : public FUsdTreeViewColumn
+class FUsdLayerNameColumn : public FUsdTreeViewColumn, public TSharedFromThis< FUsdLayerNameColumn >
 {
 public:
+	FSlateColor GetForegroundColor( const FUsdLayerViewModelRef TreeItem ) const
+	{
+		return TreeItem->IsInIsolatedStage()
+			? FSlateColor::UseForeground()
+			: FSlateColor::UseSubduedForeground();
+	}
+
 	virtual TSharedRef< SWidget > GenerateWidget( const TSharedPtr< IUsdTreeViewItem > InTreeItem, const TSharedPtr< ITableRow > TableRow ) override
 	{
 		FUsdLayerViewModelRef TreeItem = StaticCastSharedRef< FUsdLayerViewModel >( InTreeItem.ToSharedRef() );
@@ -80,8 +87,9 @@ public:
 		return SNew( SBox )
 			.VAlign( VAlign_Center )
 			[
-				SNew(STextBlock)
+				SNew( STextBlock )
 				.Text( TreeItem, &FUsdLayerViewModel::GetDisplayName )
+				.ColorAndOpacity( this, &FUsdLayerNameColumn::GetForegroundColor, TreeItem )
 				.ToolTipText_Lambda( [TreeItemWeak]
 				{
 					if ( TSharedPtr<FUsdLayerViewModel> PinnedTreeItem = TreeItemWeak.Pin() )
@@ -113,7 +121,7 @@ public:
 		{
 			return nullptr;
 		}
-		else if ( TreeItem->LayerModel->bIsMuted )
+		else if ( TreeItem->IsLayerMuted() )
 		{
 			return bIsButtonHovered
 				? FAppStyle::GetBrush( "Level.NotVisibleHighlightIcon16x" )
@@ -159,6 +167,7 @@ public:
 		}
 
 		FUsdLayerViewModelRef TreeItem = StaticCastSharedRef< FUsdLayerViewModel >( InTreeItem.ToSharedRef() );
+		FUsdLayerViewModelWeak TreeItemWeak = TreeItem;
 		const float ItemSize = FUsdStageEditorStyle::Get()->GetFloat( "UsdStageEditor.ListItemHeight" );
 
 		if ( !TreeItem->CanMuteLayer() )
@@ -172,6 +181,14 @@ public:
 
 		TSharedPtr<SButton> Button = SNew( SButton )
 			.ContentPadding( 0 )
+			.IsEnabled_Lambda([TreeItemWeak]()->bool
+			{
+				if ( TSharedPtr<FUsdLayerViewModel> PinnedTreeItem = TreeItemWeak.Pin() )
+				{
+					return PinnedTreeItem->IsInIsolatedStage();
+				}
+				return false;
+			})
 			.ButtonStyle( FUsdStageEditorStyle::Get(), TEXT("NoBorder") )
 			.OnClicked( this, &FUsdLayerMutedColumn::OnClicked, TreeItem )
 			.ToolTip( SNew( SToolTip ).Text( LOCTEXT( "MuteLayerTooltip", "Mute or unmute this layer" ) ) )
@@ -196,22 +213,15 @@ public:
 protected:
 	bool CanMuteLayer( FUsdLayerViewModelRef LayerItem ) const
 	{
-		if ( !LayerItem->IsValid() )
-		{
-			return false;
-		}
-
-		return LayerItem->CanMuteLayer();
+		return LayerItem->IsValid() && LayerItem->CanMuteLayer();
 	}
 
 	void ToggleMuteLayer( FUsdLayerViewModelRef LayerItem )
 	{
-		if ( !LayerItem->IsValid() || !CanMuteLayer( LayerItem ) )
+		if ( CanMuteLayer( LayerItem ) )
 		{
-			return;
+			LayerItem->ToggleMuteLayer();
 		}
-
-		LayerItem->ToggleMuteLayer();
 	}
 };
 
@@ -220,12 +230,7 @@ class FUsdLayerEditColumn : public FUsdTreeViewColumn, public TSharedFromThis< F
 public:
 	FReply OnClicked( const FUsdLayerViewModelRef TreeItem )
 	{
-		if ( TreeItem->UsdStage )
-		{
-			TreeItem->UsdStage.SetEditTarget( TreeItem->GetLayer() );
-		}
-
-		return FReply::Handled();
+		return TreeItem->EditLayer() ? FReply::Handled() : FReply::Unhandled();
 	}
 
 	FSlateColor GetForegroundColor( const FUsdLayerViewModelRef TreeItem, const TSharedPtr< ITableRow > TableRow, const TSharedPtr< SButton > Button ) const
@@ -236,7 +241,7 @@ public:
 		}
 
 		const bool bIsButtonHovered = Button->IsHovered();
-		const bool bIsLayerEditTarget = TreeItem->UsdStage ? TreeItem->UsdStage.GetEditTarget() == TreeItem->GetLayer() : false;
+		const bool bIsLayerEditTarget = TreeItem->IsEditTarget();
 
 		if ( bIsLayerEditTarget )
 		{
@@ -258,11 +263,20 @@ public:
 	virtual TSharedRef< SWidget > GenerateWidget( const TSharedPtr< IUsdTreeViewItem > InTreeItem, const TSharedPtr< ITableRow > TableRow ) override
 	{
 		const FUsdLayerViewModelRef TreeItem = StaticCastSharedRef< FUsdLayerViewModel >( InTreeItem.ToSharedRef() );
+		FUsdLayerViewModelWeak TreeItemWeak = TreeItem;
 
 		float ItemSize = FUsdStageEditorStyle::Get()->GetFloat( "UsdStageEditor.ListItemHeight" );
 
 		TSharedPtr<SButton> Button = SNew( SButton )
 			.ContentPadding( 0 )
+			.IsEnabled_Lambda([TreeItemWeak]()->bool
+			{
+				if ( TSharedPtr<FUsdLayerViewModel> PinnedTreeItem = TreeItemWeak.Pin() )
+				{
+					return PinnedTreeItem->IsInIsolatedStage();
+				}
+				return false;
+			})
 			.ButtonStyle( FUsdStageEditorStyle::Get(), TEXT( "NoBorder" ) )
 			.OnClicked( this, &FUsdLayerEditColumn::OnClicked, TreeItem )
 			.ToolTip( SNew( SToolTip ).Text( LOCTEXT( "EditLayerButtonToolTip", "Edit layer" ) ) )
@@ -290,6 +304,8 @@ void SUsdLayersTreeView::Construct( const FArguments& InArgs )
 {
 	SUsdTreeView::Construct( SUsdTreeView::FArguments() );
 
+	LayerIsolatedDelegate = InArgs._OnLayerIsolated;
+
 	OnContextMenuOpening = FOnContextMenuOpening::CreateSP( this, &SUsdLayersTreeView::ConstructLayerContextMenu );
 
 	OnExpansionChanged = FOnExpansionChanged::CreateLambda( [this]( const FUsdLayerViewModelRef& LayerViewModel, bool bIsExpanded )
@@ -303,7 +319,7 @@ void SUsdLayersTreeView::Construct( const FArguments& InArgs )
 	} );
 }
 
-void SUsdLayersTreeView::Refresh( const UE::FUsdStageWeak& NewStage, bool bResync )
+void SUsdLayersTreeView::Refresh( const UE::FUsdStageWeak& NewStage, const UE::FUsdStageWeak& InIsolatedStage, bool bResync )
 {
 	if ( bResync )
 	{
@@ -321,6 +337,7 @@ void SUsdLayersTreeView::Refresh( const UE::FUsdStageWeak& NewStage, bool bResyn
 		}
 
 		UsdStage = NewStage;
+		IsolatedStage = InIsolatedStage;
 
 		BuildUsdLayersEntries();
 
@@ -907,8 +924,8 @@ void SUsdLayersTreeView::BuildUsdLayersEntries()
 
 	if ( UsdStage )
 	{
-		RootItems.Add( MakeSharedUnreal< FUsdLayerViewModel >( nullptr, UsdStage, UsdStage.GetRootLayer().GetIdentifier() ) );
-		RootItems.Add( MakeSharedUnreal< FUsdLayerViewModel >( nullptr, UsdStage, UsdStage.GetSessionLayer().GetIdentifier() ) );
+		RootItems.Add( MakeSharedUnreal< FUsdLayerViewModel >( nullptr, UsdStage, IsolatedStage, UsdStage.GetRootLayer().GetIdentifier() ) );
+		RootItems.Add( MakeSharedUnreal< FUsdLayerViewModel >( nullptr, UsdStage, IsolatedStage, UsdStage.GetSessionLayer().GetIdentifier() ) );
 	}
 }
 
@@ -939,6 +956,54 @@ TSharedPtr< SWidget > SUsdLayersTreeView::ConstructLayerContextMenu()
 	FMenuBuilder LayerOptions( true, nullptr );
 	LayerOptions.BeginSection( "Layer", LOCTEXT("Layer", "Layer") );
 	{
+		LayerOptions.AddMenuEntry(
+			TAttribute<FText>::Create( TAttribute<FText>::FGetter::CreateLambda( [this]()
+			{
+				if ( IsolatedStage )
+				{
+					TArray< FUsdLayerViewModelRef > MySelectedItems = GetSelectedItems();
+					if ( MySelectedItems.Num() == 1 && IsolatedStage )
+					{
+						const UE::FSdfLayerWeak& Layer = MySelectedItems[ 0 ]->GetLayer();
+						if ( Layer == IsolatedStage.GetRootLayer() || Layer == UsdStage.GetRootLayer() )
+						{
+							return LOCTEXT( "StopIsolatingStage_Text", "Stop isolating" );
+						}
+					}
+				}
+
+				return LOCTEXT( "IsolateStage_Text", "Isolate" );
+			})),
+			TAttribute<FText>::Create( TAttribute<FText>::FGetter::CreateLambda( [this]()
+			{
+				if ( IsolatedStage )
+				{
+					TArray< FUsdLayerViewModelRef > MySelectedItems = GetSelectedItems();
+					if ( MySelectedItems.Num() == 1 )
+					{
+						const UE::FSdfLayerWeak& Layer = MySelectedItems[ 0 ]->GetLayer();
+						if ( Layer == IsolatedStage.GetRootLayer() )
+						{
+							return LOCTEXT( "StopIsolatingStage_ToolTip", "Stops isolating this layer and go back to showing the full composed stage" );
+						}
+						else if ( IsolatedStage && Layer == UsdStage.GetRootLayer() )
+						{
+							return LOCTEXT( "LeaveIsolatedMode_ToolTip", "Stops isolating on any layer and go back to showing the full composed stage" );
+						}
+					}
+				}
+
+				return LOCTEXT( "IsolateStage_ToolTip", "Isolate a stage with this layer as its root layer" );
+			})),
+			FSlateIcon(),
+			FUIAction(
+				FExecuteAction::CreateSP( this, &SUsdLayersTreeView::OnIsolateSelectedLayer ),
+				FCanExecuteAction::CreateSP( this, &SUsdLayersTreeView::CanIsolateSelectedLayer )
+			),
+			NAME_None,
+			EUserInterfaceActionType::Button
+		);
+
 		LayerOptions.AddMenuEntry(
 			LOCTEXT("EditLayer", "Edit"),
 			LOCTEXT("EditLayer_ToolTip", "Sets the layer as the edit target"),
@@ -1034,6 +1099,62 @@ TSharedPtr< SWidget > SUsdLayersTreeView::ConstructLayerContextMenu()
 	return MenuWidget;
 }
 
+bool SUsdLayersTreeView::CanIsolateSelectedLayer() const
+{
+	TArray< FUsdLayerViewModelRef > MySelectedItems = GetSelectedItems();
+
+	if ( MySelectedItems.Num() != 1 )
+	{
+		return false;
+	}
+
+	for ( FUsdLayerViewModelRef SelectedItem : MySelectedItems )
+	{
+		const UE::FSdfLayerWeak& Layer = SelectedItem->GetLayer();
+		const bool bLayerIsStageRoot = Layer == UsdStage.GetRootLayer();
+
+		return SelectedItem->IsValid() && (
+			// If we're right clicking a sublayer its always OK
+			!bLayerIsStageRoot
+			// If we're isolating and selecting either a sublayer or the root layer (i.e. not the session layer)
+			|| ( IsolatedStage && ( SelectedItem->ParentItem || bLayerIsStageRoot ) )
+		);
+	}
+
+	return false;
+}
+
+void SUsdLayersTreeView::OnIsolateSelectedLayer()
+{
+	if ( !LayerIsolatedDelegate.IsBound() )
+	{
+		return;
+	}
+
+	TArray< FUsdLayerViewModelRef > MySelectedItems = GetSelectedItems();
+
+	if ( MySelectedItems.Num() != 1 )
+	{
+		return;
+	}
+
+	for ( FUsdLayerViewModelRef SelectedItem : MySelectedItems )
+	{
+		UE::FSdfLayer Layer = SelectedItem->GetLayer();
+
+		if ( IsolatedStage && IsolatedStage.GetRootLayer() == Layer )
+		{
+			// We're already isolating this one -> Toggle isolate mode off
+			LayerIsolatedDelegate.Execute( UE::FSdfLayer{} );
+		}
+		else
+		{
+			// Isolate the provided layer
+			LayerIsolatedDelegate.Execute( Layer );
+		}
+	}
+}
+
 bool SUsdLayersTreeView::CanEditSelectedLayer() const
 {
 	bool bHasEditableLayer = false;
diff --git a/Engine/Plugins/Importers/USDImporter/Source/USDStageEditor/Private/SUSDLayersTreeView.h b/Engine/Plugins/Importers/USDImporter/Source/USDStageEditor/Private/SUSDLayersTreeView.h
--- a/Engine/Plugins/Importers/USDImporter/Source/USDStageEditor/Private/SUSDLayersTreeView.h
+++ b/Engine/Plugins/Importers/USDImporter/Source/USDStageEditor/Private/SUSDLayersTreeView.h
@@ -2,6 +2,7 @@
 
 #pragma once
 
+#include "UsdWrappers/ForwardDeclarations.h"
 #include "UsdWrappers/UsdStage.h"
 #include "Widgets/SUSDTreeView.h"
 
@@ -11,15 +12,19 @@
 class AUsdStageActor;
 
 using FUsdLayerViewModelRef = TSharedRef< class FUsdLayerViewModel >;
+using FUsdLayerViewModelWeak = TWeakPtr< class FUsdLayerViewModel >;
+
+DECLARE_DELEGATE_OneParam( FOnLayerIsolated, const UE::FSdfLayer& );
 
 class SUsdLayersTreeView : public SUsdTreeView< FUsdLayerViewModelRef >
 {
 public:
 	SLATE_BEGIN_ARGS( SUsdLayersTreeView ) {}
+		SLATE_EVENT( FOnLayerIsolated, OnLayerIsolated )
 	SLATE_END_ARGS()
 
 	void Construct( const FArguments& InArgs );
-	void Refresh( const UE::FUsdStageWeak& NewStage, bool bResync = false );
+	void Refresh( const UE::FUsdStageWeak& NewStage, const UE::FUsdStageWeak& IsolatedStage = {}, bool bResync = false);
 
 	// Drag and drop interface for our rows
 	FReply OnRowDragDetected( const FGeometry& Geometry, const FPointerEvent& PointerEvent );
@@ -40,6 +45,9 @@ private:
 
 	TSharedPtr< SWidget > ConstructLayerContextMenu();
 
+	bool CanIsolateSelectedLayer() const;
+	void OnIsolateSelectedLayer();
+
 	bool CanEditSelectedLayer() const;
 	void OnEditSelectedLayer();
 
@@ -65,6 +73,12 @@ private:
 	// Should always be valid, we keep the one we're given on Refresh()
 	UE::FUsdStageWeak UsdStage;
 
+	// A stage we create based on one of the sublayers of UsdStage
+	UE::FUsdStageWeak IsolatedStage;
+
 	// So that we can store these across refreshes
 	TMap< FString, bool > TreeItemExpansionStates;
+
+	// Used so that we can isolate the new layer without coupling to the SUSDStage widget too much
+	FOnLayerIsolated LayerIsolatedDelegate;
 };
diff --git a/Engine/Plugins/Importers/USDImporter/Source/USDStageEditor/Private/SUSDStage.cpp b/Engine/Plugins/Importers/USDImporter/Source/USDStageEditor/Private/SUSDStage.cpp
--- a/Engine/Plugins/Importers/USDImporter/Source/USDStageEditor/Private/SUSDStage.cpp
+++ b/Engine/Plugins/Importers/USDImporter/Source/USDStageEditor/Private/SUSDStage.cpp
@@ -225,7 +225,7 @@ void SUsdStage::Construct( const FArguments& InArgs )
 					.VAlign( VAlign_Fill )
 					.Padding( FMargin( 0.0f, 0.0f, 10.0f, 0.0f ) )
 					[
-						MakeFocusWarningButton()
+						MakeIsolateWarningButton()
 					]
 				]
 
@@ -284,6 +284,7 @@ void SUsdStage::Construct( const FArguments& InArgs )
 					.BorderImage( FAppStyle::GetBrush(TEXT("ToolPanel.GroupBorder")) )
 					[
 						SAssignNew( UsdLayersTreeView, SUsdLayersTreeView )
+						.OnLayerIsolated( this, &SUsdStage::OnLayerIsolated )
 					]
 				]
 			]
@@ -291,6 +292,12 @@ void SUsdStage::Construct( const FArguments& InArgs )
 	];
 
 	SetupStageActorDelegates();
+
+	// We're opening the USD Stage editor for the first time and we already have a stage: Display it immediately
+	if ( UsdStage )
+	{
+		Refresh();
+	}
 }
 
 void SUsdStage::SetupStageActorDelegates()
@@ -321,7 +328,7 @@ void SUsdStage::SetupStageActorDelegates()
 						 ViewModel.UsdStageActor.IsValid() &&
 						 ( bViewingTheUpdatedPrim || ( bViewingStageProperties && bStageUpdated ) ) )
 					{
-						this->UsdPrimInfoWidget->SetPrimPath( ViewModel.UsdStageActor->GetOrLoadUsdStage(), *PrimPath );
+						this->UsdPrimInfoWidget->SetPrimPath( GetCurrentStage(), *PrimPath);
 					}
 
 					// If we resynced our selected prim or our ancestor and have selection sync enabled, try to refresh it so that we're
@@ -349,9 +356,7 @@ void SUsdStage::SetupStageActorDelegates()
 
 						if ( this->UsdPrimInfoWidget )
 						{
-							// The cast here forces us to use the const version of GetUsdStage, that won't force-load the stage in case it isn't opened yet
-							const UE::FUsdStage& UsdStage = static_cast< const AUsdStageActor* >( ViewModel.UsdStageActor.Get( bEvenIfPendingKill ) )->GetUsdStage();
-							this->UsdPrimInfoWidget->SetPrimPath( UsdStage, TEXT("/") );
+							this->UsdPrimInfoWidget->SetPrimPath( GetCurrentStage(), TEXT("/"));
 						}
 
 						if ( this->UsdStageTreeView )
@@ -389,11 +394,13 @@ void SUsdStage::SetupStageActorDelegates()
 			{
 				AsyncTask( ENamedThreads::GameThread, [this]()
 				{
-					if ( this->UsdLayersTreeView && ViewModel.UsdStageActor.IsValid() )
+					AUsdStageActor* StageActor = ViewModel.UsdStageActor.Get();
+					if ( this->UsdLayersTreeView && StageActor )
 					{
 						constexpr bool bResync = false;
-						this->UsdLayersTreeView->Refresh(
-							ViewModel.UsdStageActor.Get()->GetOrLoadUsdStage(),
+						UsdLayersTreeView->Refresh(
+							StageActor->GetBaseUsdStage(),
+							StageActor->GetIsolatedUsdStage(),
 							bResync
 						);
 					}
@@ -406,11 +413,13 @@ void SUsdStage::SetupStageActorDelegates()
 			{
 				AsyncTask( ENamedThreads::GameThread, [this]()
 				{
-					if ( this->UsdLayersTreeView && ViewModel.UsdStageActor.IsValid() )
+					AUsdStageActor* StageActor = ViewModel.UsdStageActor.Get();
+					if ( this->UsdLayersTreeView && StageActor )
 					{
 						constexpr bool bResync = false;
-						this->UsdLayersTreeView->Refresh(
-							ViewModel.UsdStageActor.Get()->GetOrLoadUsdStage(),
+						UsdLayersTreeView->Refresh(
+							StageActor->GetBaseUsdStage(),
+							StageActor->GetIsolatedUsdStage(),
 							bResync
 						);
 					}
@@ -549,31 +558,47 @@ TSharedRef< SWidget > SUsdStage::MakeActorPickerMenuContent()
 	return SNullWidget::NullWidget;
 }
 
-TSharedRef< SWidget > SUsdStage::MakeFocusWarningButton()
+TSharedRef< SWidget > SUsdStage::MakeIsolateWarningButton()
 {
-	// TODO: Will be used for focus mode in the future
 	return SNew( SButton )
 		.ButtonStyle( &FAppStyle::Get().GetWidgetStyle<FButtonStyle>( "EditorViewportToolBar.WarningButton" ) )
+		.ButtonColorAndOpacity( FLinearColor{ 1.0f, 0.4f, 0.0f, 1.0f } )
 		.OnClicked_Lambda( [this]() -> FReply
 		{
+			// Isolating nothing is how we un-isolate
+			OnLayerIsolated( UE::FSdfLayerWeak{} );
 			return FReply::Handled();
 		})
 		.Visibility_Lambda( [this]() -> EVisibility
 		{
+			if ( const AUsdStageActor* StageActor = ViewModel.UsdStageActor.Get() )
+			{
+				if ( StageActor->GetIsolatedUsdStage() )
+				{
+					return EVisibility::Visible;
+				}
+			}
 			return EVisibility::Hidden;
 		})
 		.ToolTipText_Lambda( [this]() -> FText
 		{
-			return FText::Format(
-				LOCTEXT( "FocusWarningButtonToolTip", "The USD Stage editor and the Unreal level are displaying a focused stage with '{0}' as its root layer.\nPress this button to revert to displaying the entire composed stage." ),
-				FText::GetEmpty()
-			);
+			if ( const AUsdStageActor* StageActor = ViewModel.UsdStageActor.Get() )
+			{
+				if ( const UE::FUsdStage& CurrentStage = StageActor->GetUsdStage() )
+				{
+					return FText::Format(
+						LOCTEXT( "IsolateWarningButtonToolTip", "The USD Stage editor and the Unreal level are displaying an isolated stage with '{0}' as its root layer.\nPress this button to revert to displaying the entire composed stage." ),
+						FText::FromString( CurrentStage.GetRootLayer().GetDisplayName() )
+					);
+				}
+			}
+
+			return FText::GetEmpty();
 		})
 		.Content()
 		[
 			SNew( STextBlock )
-			.TextStyle( &FAppStyle::Get().GetWidgetStyle<FTextBlockStyle>( "SmallText" ) )
-			.Text( LOCTEXT( "FocusWarningButtonText", "Focus Mode" ) )
+			.Text( LOCTEXT( "IsolateWarningButtonText", "Isolated Mode" ) )
 		];
 }
 
@@ -1793,10 +1818,12 @@ void SUsdStage::OnReloadStage()
 
 	ViewModel.ReloadStage();
 
-	if ( UsdLayersTreeView )
+	const AUsdStageActor* StageActor = ViewModel.UsdStageActor.Get();
+
+	if ( UsdLayersTreeView && StageActor )
 	{
 		const bool bResync = true;
-		UsdLayersTreeView->Refresh( ViewModel.UsdStageActor.Get()->GetOrLoadUsdStage(), bResync );
+		UsdLayersTreeView->Refresh( StageActor->GetBaseUsdStage(), StageActor->GetIsolatedUsdStage(), bResync );
 	}
 }
 
@@ -1804,10 +1831,12 @@ void SUsdStage::OnResetStage()
 {
 	ViewModel.ResetStage();
 
-	if ( UsdLayersTreeView )
+	const AUsdStageActor* StageActor = ViewModel.UsdStageActor.Get();
+
+	if ( UsdLayersTreeView && StageActor )
 	{
 		const bool bResync = true;
-		UsdLayersTreeView->Refresh( ViewModel.UsdStageActor.Get()->GetOrLoadUsdStage(), bResync );
+		UsdLayersTreeView->Refresh( StageActor->GetBaseUsdStage(), StageActor->GetIsolatedUsdStage(), bResync );
 	}
 }
 
@@ -1819,6 +1848,16 @@ void SUsdStage::OnClose()
 	Refresh();
 }
 
+void SUsdStage::OnLayerIsolated( const UE::FSdfLayer& IsolatedLayer )
+{
+	if ( AUsdStageActor* StageActor = ViewModel.UsdStageActor.Get() )
+	{
+		StageActor->IsolateLayer( IsolatedLayer );
+
+		Refresh();
+	}
+}
+
 void SUsdStage::OnImport()
 {
 	ViewModel.ImportStage();
@@ -1911,15 +1950,21 @@ void SUsdStage::Refresh()
 	// May be nullptr, but that is ok. Its how the widgets are reset
 	AUsdStageActor* StageActor = ViewModel.UsdStageActor.Get();
 
-	if (UsdLayersTreeView)
+	if ( UsdLayersTreeView )
 	{
+		UE::FUsdStageWeak IsolatedStage = StageActor ? StageActor->GetIsolatedUsdStage() : UE::FUsdStage{};
+		UE::FUsdStageWeak BaseStage = StageActor ? StageActor->GetBaseUsdStage() : UE::FUsdStage{};
+
+		// The layers tree view always needs to receive both the full stage as well as the isolated
 		const bool bResync = true;
-		UsdLayersTreeView->Refresh( ViewModel.UsdStageActor.Get()->GetOrLoadUsdStage(), bResync );
+		UsdLayersTreeView->Refresh( BaseStage, IsolatedStage, bResync );
 	}
 
-	if (UsdStageTreeView)
+	if ( UsdStageTreeView )
 	{
-		UsdStageTreeView->Refresh( StageActor->GetOrLoadUsdStage() );
+		UE::FUsdStageWeak CurrentStage = StageActor ? StageActor->GetOrLoadUsdStage() : UE::FUsdStage{};
+
+		UsdStageTreeView->Refresh( CurrentStage );
 
 		// Refresh will generate brand new RootItems, so let's immediately select the new item
 		// that corresponds to the prim we're supposed to be selecting
@@ -2030,6 +2075,17 @@ void SUsdStage::OnNaniteTriangleThresholdValueCommitted( int32 InValue, ETextCom
 	CurrentNaniteThreshold = InValue;
 }
 
+UE::FUsdStageWeak SUsdStage::GetCurrentStage() const
+{
+	const AUsdStageActor* StageActor = ViewModel.UsdStageActor.Get();
+	if ( StageActor )
+	{
+		return StageActor->GetUsdStage();
+	}
+
+	return UE::FUsdStageWeak{};
+}
+
 #endif // #if USE_USD_SDK
 
 #undef LOCTEXT_NAMESPACE
diff --git a/Engine/Plugins/Importers/USDImporter/Source/USDStageEditor/Private/SUSDStage.h b/Engine/Plugins/Importers/USDImporter/Source/USDStageEditor/Private/SUSDStage.h
--- a/Engine/Plugins/Importers/USDImporter/Source/USDStageEditor/Private/SUSDStage.h
+++ b/Engine/Plugins/Importers/USDImporter/Source/USDStageEditor/Private/SUSDStage.h
@@ -3,6 +3,7 @@
 #pragma once
 
 #include "USDStageViewModel.h"
+#include "UsdWrappers/UsdStage.h"
 
 #include "CoreMinimal.h"
 #include "Layout/Visibility.h"
@@ -39,7 +40,7 @@ protected:
 	TSharedRef< SWidget > MakeMainMenu();
 	TSharedRef< SWidget > MakeActorPickerMenu();
 	TSharedRef< SWidget > MakeActorPickerMenuContent();
-	TSharedRef< SWidget > MakeFocusWarningButton();
+	TSharedRef< SWidget > MakeIsolateWarningButton();
 	void FillFileMenu( FMenuBuilder& MenuBuilder );
 	void FillActionsMenu( FMenuBuilder& MenuBuilder );
 	void FillOptionsMenu( FMenuBuilder& MenuBuilder );
@@ -63,6 +64,8 @@ protected:
 	void OnResetStage();
 	void OnClose();
 
+	void OnLayerIsolated( const UE::FSdfLayer& IsolatedLayer );
+
 	void OnImport();
 
 	void OnPrimSelectionChanged( const TArray<FString>& PrimPath );
@@ -81,6 +84,8 @@ protected:
 	void OnNaniteTriangleThresholdValueChanged( int32 InValue );
 	void OnNaniteTriangleThresholdValueCommitted( int32 InValue, ETextCommit::Type InCommitType );
 
+	UE::FUsdStageWeak GetCurrentStage() const;
+
 protected:
 	TSharedPtr< class SUsdStageTreeView > UsdStageTreeView;
 	TSharedPtr< class SUsdPrimInfo > UsdPrimInfoWidget;
diff --git a/Engine/Plugins/Importers/USDImporter/Source/USDStageEditorViewModels/Private/USDLayersViewModel.cpp b/Engine/Plugins/Importers/USDImporter/Source/USDStageEditorViewModels/Private/USDLayersViewModel.cpp
--- a/Engine/Plugins/Importers/USDImporter/Source/USDStageEditorViewModels/Private/USDLayersViewModel.cpp
+++ b/Engine/Plugins/Importers/USDImporter/Source/USDStageEditorViewModels/Private/USDLayersViewModel.cpp
@@ -23,10 +23,16 @@
 
 #endif // #if USE_USD_SDK
 
-FUsdLayerViewModel::FUsdLayerViewModel( FUsdLayerViewModel* InParentItem, const UE::FUsdStageWeak& InUsdStage, const FString& InLayerIdentifier )
+FUsdLayerViewModel::FUsdLayerViewModel(
+	FUsdLayerViewModel* InParentItem,
+	const UE::FUsdStageWeak& InUsdStage,
+	const UE::FUsdStageWeak& InIsolatedStage,
+	const FString& InLayerIdentifier
+)
 	: LayerModel( MakeShared< FUsdLayerModel >() )
 	, ParentItem( InParentItem )
 	, UsdStage( InUsdStage )
+	, IsolatedStage( InIsolatedStage )
 	, LayerIdentifier( InLayerIdentifier )
 {
 	RefreshData();
@@ -119,7 +125,7 @@ void FUsdLayerViewModel::FillChildren()
 				// Prevent infinite recursions if a sublayer refers to a parent of the same hierarchy
 				if ( !AllLayerIdentifiers.Contains( AssetPathRelativeToLayer ) )
 				{
-					Children.Add( MakeShared< FUsdLayerViewModel >( this, UsdStage, AssetPathRelativeToLayer ) );
+					Children.Add( MakeShared< FUsdLayerViewModel >( this, UsdStage, IsolatedStage, AssetPathRelativeToLayer ) );
 				}
 			}
 		}
@@ -140,11 +146,13 @@ void FUsdLayerViewModel::RefreshData()
 	FScopedUsdAllocs UsdAllocs;
 
 	pxr::UsdStageRefPtr UsdStageRef( UsdStage );
+	pxr::UsdStageRefPtr IsolatedStageRef( IsolatedStage );
 
-	const TUsdStore< std::string > UsdLayerIdentifier = UnrealToUsd::ConvertString( *LayerIdentifier );
+	const std::string UsdLayerIdentifier = UnrealToUsd::ConvertString( *LayerIdentifier ).Get();
 
-	LayerModel->DisplayName = UsdToUnreal::ConvertString( pxr::SdfLayer::GetDisplayNameFromIdentifier( UsdLayerIdentifier.Get() ) );
-	LayerModel->bIsMuted = UsdStageRef->IsLayerMuted( UsdLayerIdentifier.Get() );
+	LayerModel->DisplayName = UsdToUnreal::ConvertString(
+		pxr::SdfLayer::GetDisplayNameFromIdentifier( UsdLayerIdentifier )
+	);
 
 	pxr::SdfLayerRefPtr UsdLayer = UE::FSdfLayer::FindOrOpen( *LayerIdentifier );
 	if ( UsdLayer )
@@ -152,10 +160,29 @@ void FUsdLayerViewModel::RefreshData()
 		LayerModel->bIsDirty = UsdLayer->IsDirty() && !UsdLayer->IsAnonymous();
 	}
 
+	const bool bIsIsolatedRootLayer = IsolatedStageRef && IsolatedStageRef->GetRootLayer() == UsdLayer;
+
+	// Note: Originally this was a proper check to see if UsdLayer was inside IsolatedStage's layer stack. However, this
+	// breaks down if the layer is also muted: When that happens the stages basically pretend the layer doesn't exist,
+	// and so any muted layer on the base stage would behave as if it was "non-isolated"
+	LayerModel->bIsInIsolatedStage = !IsolatedStageRef
+		|| bIsIsolatedRootLayer
+		|| ( ParentItem && ParentItem->IsInIsolatedStage() );
+
 	const pxr::SdfLayerHandle& EditTargetLayer = UsdStageRef->GetEditTarget().GetLayer();
-	LayerModel->bIsEditTarget = ( EditTargetLayer
-		? FCStringAnsi::Stricmp( EditTargetLayer->GetIdentifier().c_str(), UsdLayerIdentifier.Get().c_str() ) == 0
-		: false );
+	const pxr::SdfLayerHandle& FocusedEditTarget = IsolatedStageRef
+		? IsolatedStageRef->GetEditTarget().GetLayer()
+		: pxr::SdfLayerHandle{};
+
+	LayerModel->bIsEditTarget = ( IsolatedStageRef )
+		? LayerModel->bIsInIsolatedStage && FocusedEditTarget == UsdLayer
+		: EditTargetLayer == UsdLayer;
+
+	LayerModel->bIsMuted = ( IsolatedStageRef && LayerModel->bIsInIsolatedStage )
+		// If we isolating, we're only muted if we're muted on the isolated stage
+		? IsolatedStageRef->IsLayerMuted( UsdLayerIdentifier )
+		// Otherwise we're muted any time we're muted in the base stage
+		: UsdStageRef->IsLayerMuted( UsdLayerIdentifier );
 
 	for ( const TSharedRef< FUsdLayerViewModel >& Child : Children )
 	{
@@ -193,7 +220,10 @@ bool FUsdLayerViewModel::CanMuteLayer() const
 		return false;
 	}
 
-	return !UsdStage.GetRootLayer().GetIdentifier().Equals( LayerIdentifier, ESearchCase::Type::IgnoreCase ) && !LayerModel->bIsEditTarget;
+	UE::FUsdStageWeak AffectedStage = IsolatedStage ? IsolatedStage : UsdStage;
+
+	return !AffectedStage.GetRootLayer().GetIdentifier().Equals( LayerIdentifier, ESearchCase::Type::IgnoreCase )
+		&& !LayerModel->bIsEditTarget;
 }
 
 void FUsdLayerViewModel::ToggleMuteLayer()
@@ -208,21 +238,42 @@ void FUsdLayerViewModel::ToggleMuteLayer()
 
 	const TUsdStore< std::string > UsdLayerIdentifier = UnrealToUsd::ConvertString( *LayerIdentifier );
 
-	pxr::UsdStageRefPtr UsdStageRef( UsdStage );
+	pxr::UsdStageRefPtr CurrentStage( IsolatedStage ? IsolatedStage : UsdStage );
 
-	if ( UsdStageRef->IsLayerMuted( UsdLayerIdentifier.Get() ) )
+	if ( CurrentStage->IsLayerMuted( UsdLayerIdentifier.Get() ) )
 	{
-		UsdStageRef->UnmuteLayer( UsdLayerIdentifier.Get() );
+		CurrentStage->UnmuteLayer( UsdLayerIdentifier.Get() );
 	}
 	else
 	{
-		UsdStageRef->MuteLayer( UsdLayerIdentifier.Get() );
+		CurrentStage->MuteLayer( UsdLayerIdentifier.Get() );
 	}
 
-	RefreshData();
+	// We have to refresh all of our view models here, because this layer can appear more than once on the layer
+	// stack and we want all of them to show the updated muted/unmuted icon
+	FUsdLayerViewModel* Ancestor = this;
+	while ( true )
+	{
+		FUsdLayerViewModel* AncestorParent = Ancestor->ParentItem;
+		if ( !AncestorParent )
+		{
+			break;
+		}
+
+		Ancestor = AncestorParent;
+	}
+	if ( Ancestor )
+	{
+		Ancestor->RefreshData();
+	}
 #endif // #if USE_USD_SDK
 }
 
+bool FUsdLayerViewModel::IsEditTarget() const
+{
+	return LayerModel->bIsEditTarget;
+}
+
 bool FUsdLayerViewModel::CanEditLayer() const
 {
 	return !LayerModel->bIsMuted;
@@ -231,17 +282,25 @@ bool FUsdLayerViewModel::CanEditLayer() const
 bool FUsdLayerViewModel::EditLayer()
 {
 	UE::FSdfLayer Layer( GetLayer() );
-	if ( !UsdStage || !Layer || !CanEditLayer() )
+
+	UE::FUsdStageWeak AffectedStage = IsolatedStage ? IsolatedStage : UsdStage;
+
+	if ( !AffectedStage || !Layer || !CanEditLayer() )
 	{
 		return false;
 	}
 
-	UsdStage.SetEditTarget( Layer );
+	AffectedStage.SetEditTarget( Layer );
 	RefreshData();
 
 	return true;
 }
 
+bool FUsdLayerViewModel::IsInIsolatedStage() const
+{
+	return LayerModel->bIsInIsolatedStage;
+}
+
 void FUsdLayerViewModel::AddSubLayer( const TCHAR* SubLayerIdentifier )
 {
 #if USE_USD_SDK
diff --git a/Engine/Plugins/Importers/USDImporter/Source/USDStageEditorViewModels/Public/USDLayersViewModel.h b/Engine/Plugins/Importers/USDImporter/Source/USDStageEditorViewModels/Public/USDLayersViewModel.h
--- a/Engine/Plugins/Importers/USDImporter/Source/USDStageEditorViewModels/Public/USDLayersViewModel.h
+++ b/Engine/Plugins/Importers/USDImporter/Source/USDStageEditorViewModels/Public/USDLayersViewModel.h
@@ -15,12 +15,18 @@ public:
 	bool bIsEditTarget = false;
 	bool bIsMuted = false;
 	bool bIsDirty = false;
+	bool bIsInIsolatedStage = true;
 };
 
 class USDSTAGEEDITORVIEWMODELS_API FUsdLayerViewModel : public IUsdTreeViewItem
 {
 public:
-	explicit FUsdLayerViewModel( FUsdLayerViewModel* InParentItem, const UE::FUsdStageWeak& InUsdStage, const FString& InLayerIdentifier );
+	explicit FUsdLayerViewModel(
+		FUsdLayerViewModel* InParentItem,
+		const UE::FUsdStageWeak& InUsdStage,
+		const UE::FUsdStageWeak& InIsolatedStage,
+		const FString& InLayerIdentifier
+	);
 
 	bool IsValid() const;
 
@@ -37,9 +43,12 @@ public:
 	bool CanMuteLayer() const;
 	void ToggleMuteLayer();
 
+	bool IsEditTarget() const;
 	bool CanEditLayer() const;
 	bool EditLayer();
 
+	bool IsInIsolatedStage() const;
+
 	void AddSubLayer( const TCHAR* SubLayerIdentifier );
 	void NewSubLayer( const TCHAR* SubLayerIdentifier );
 	bool RemoveSubLayer( int32 SubLayerIndex );
@@ -52,5 +61,6 @@ public:
 	TArray< TSharedRef< FUsdLayerViewModel > > Children;
 
 	UE::FUsdStageWeak UsdStage;
+	UE::FUsdStageWeak IsolatedStage;
 	FString LayerIdentifier;
 };
diff --git a/Engine/Plugins/Importers/USDImporter/Source/UnrealUSDWrapper/Private/UnrealUSDWrapper.cpp b/Engine/Plugins/Importers/USDImporter/Source/UnrealUSDWrapper/Private/UnrealUSDWrapper.cpp
--- a/Engine/Plugins/Importers/USDImporter/Source/UnrealUSDWrapper/Private/UnrealUSDWrapper.cpp
+++ b/Engine/Plugins/Importers/USDImporter/Source/UnrealUSDWrapper/Private/UnrealUSDWrapper.cpp
@@ -908,60 +908,210 @@ TArray<FString> UnrealUSDWrapper::GetNativeFileFormats()
 	return Result;
 }
 
-UE::FUsdStage UnrealUSDWrapper::OpenStage( const TCHAR* Identifier, EUsdInitialLoadSet InitialLoadSet, bool bUseStageCache )
+namespace UE::UnrealUSDWrapper::Private
 {
-	if ( !Identifier || FCString::Strlen( Identifier ) == 0 )
+	UE::FUsdStage OpenStageImpl(
+		const TCHAR* RootIdentifier,
+		const TCHAR* SessionIdentifier,
+		EUsdInitialLoadSet InitialLoadSet,
+		bool bUseStageCache,
+		const TArray<FString>* PopulationMask,
+		bool bForceReloadLayersFromDisk
+	)
 	{
-		return UE::FUsdStage();
-	}
+		if ( !RootIdentifier || FCString::Strlen( RootIdentifier ) == 0 )
+		{
+			return UE::FUsdStage();
+		}
 
 #if USE_USD_SDK
-	FScopedUsdAllocs UsdAllocs;
+		FScopedUsdAllocs UsdAllocs;
 
-	pxr::SdfLayerHandleSet LoadedLayers = pxr::SdfLayer::GetLoadedLayers();
-	pxr::UsdStageRefPtr Stage;
+		pxr::SdfLayerHandleSet LoadedLayers = pxr::SdfLayer::GetLoadedLayers();
+		pxr::UsdStageRefPtr Stage;
 
-	TOptional<pxr::UsdStageCacheContext> StageCacheContext;
-	if ( bUseStageCache )
-	{
-		StageCacheContext.Emplace( pxr::UsdUtilsStageCache::Get() );
-	}
+		TOptional<pxr::UsdStageCacheContext> StageCacheContext;
+		if ( bUseStageCache )
+		{
+			StageCacheContext.Emplace( pxr::UsdUtilsStageCache::Get() );
+		}
 
-	FString IdentifierStr = FString(Identifier);
-	if ( FPaths::FileExists( IdentifierStr ) )
-	{
-		Stage = pxr::UsdStage::Open( TCHAR_TO_ANSI( Identifier ), pxr::UsdStage::InitialLoadSet( InitialLoadSet ) );
-	}
-	else if ( IdentifierStr.RemoveFromStart( UnrealIdentifiers::IdentifierPrefix ) )
-	{
-		pxr::SdfLayerRefPtr RootLayer = pxr::SdfLayer::Find( TCHAR_TO_ANSI( *IdentifierStr ) );
-		if ( RootLayer )
+		pxr::UsdStagePopulationMask Mask;
+		if( PopulationMask )
 		{
-			Stage = pxr::UsdStage::Open( RootLayer, pxr::UsdStage::InitialLoadSet( InitialLoadSet ) );
+			// The USD OpenMaskedStage functions don't actually consult or populate the stage cache
+			ensure( bUseStageCache == false );
+
+			for ( const FString& AllowedPrimPath : *PopulationMask )
+			{
+				Mask.Add( pxr::SdfPath{ TCHAR_TO_ANSI( *AllowedPrimPath ) } );
+			}
 		}
-	}
 
-	if ( !bUseStageCache && Stage )
-	{
-		// Layers are cached in the layer registry independently of the stage cache. If the layer is already in the registry by the time
-		// we try to open a stage, even if we're not using a stage cache at all the layer will be reused and the file will *not* be re-read.
-		// Here we keep track of these loaded layers and manually reload the ones that were reused, because if we're passing false for
-		// bUseStageCache we really expect the files to be fully re-read
-		pxr::SdfLayerHandleVector StageLayers = Stage->GetLayerStack();
-		for ( pxr::SdfLayerHandle StageLayer : StageLayers )
+		static_assert( ( int ) pxr::UsdStage::InitialLoadSet::LoadAll == ( int ) EUsdInitialLoadSet::LoadAll );
+		static_assert( ( int ) pxr::UsdStage::InitialLoadSet::LoadNone == ( int ) EUsdInitialLoadSet::LoadNone );
+		pxr::UsdStage::InitialLoadSet LoadSet = static_cast< pxr::UsdStage::InitialLoadSet >( InitialLoadSet );
+
+		FString IdentifierStr = FString( RootIdentifier );
+		if ( FPaths::FileExists( IdentifierStr ) )
 		{
-			if ( LoadedLayers.count( StageLayer ) > 0 )
+			if( PopulationMask )
+			{
+				Stage = pxr::UsdStage::OpenMasked( TCHAR_TO_ANSI( *IdentifierStr ), Mask, LoadSet );
+			}
+			else
 			{
-				StageLayer->Reload();
+				Stage = pxr::UsdStage::Open( TCHAR_TO_ANSI( *IdentifierStr ), LoadSet );
+			}
+		}
+		else
+		{
+			FString SessionIdentifierStr = FString{ SessionIdentifier };
+
+			IdentifierStr.RemoveFromStart( UnrealIdentifiers::IdentifierPrefix );
+			SessionIdentifierStr.RemoveFromStart( UnrealIdentifiers::IdentifierPrefix );
+
+			pxr::SdfLayerRefPtr RootLayer = pxr::SdfLayer::Find( TCHAR_TO_ANSI( *IdentifierStr ) );
+			pxr::SdfLayerRefPtr SessionLayer = pxr::SdfLayer::Find( TCHAR_TO_ANSI( *SessionIdentifierStr ) );
+			if ( RootLayer )
+			{
+				if ( PopulationMask )
+				{
+					// We use this additional check so we don't have to replicate USD's "_CreateAnonymousSessionLayer"
+					// Basically we can't pass an invalid session layer pointer here the stage will actually end up
+					// with no session layer at all
+					if ( SessionLayer )
+					{
+						Stage = pxr::UsdStage::OpenMasked( RootLayer, SessionLayer, Mask, LoadSet );
+					}
+					else
+					{
+						Stage = pxr::UsdStage::OpenMasked( RootLayer, Mask, LoadSet );
+					}
+				}
+				else
+				{
+					if( SessionLayer )
+					{
+						Stage = pxr::UsdStage::Open( RootLayer, SessionLayer, LoadSet );
+					}
+					else
+					{
+						Stage = pxr::UsdStage::Open( RootLayer, LoadSet );
+					}
+				}
+			}
+		}
+
+		if ( bForceReloadLayersFromDisk && Stage )
+		{
+			// Layers are cached in the layer registry independently of the stage cache. If the layer is already in
+			// the registry by the time we try to open a stage, even if we're not using a stage cache at all the
+			// layer will be reused and the file will *not* be re-read, so here we manually reload them.
+			pxr::SdfLayerHandleVector StageLayers = Stage->GetLayerStack();
+			for ( pxr::SdfLayerHandle StageLayer : StageLayers )
+			{
+				if ( LoadedLayers.count( StageLayer ) > 0 )
+				{
+					StageLayer->Reload();
+				}
 			}
 		}
-	}
 
-	return UE::FUsdStage( Stage );
+		return UE::FUsdStage( Stage );
 #else
-	UsdWrapperUtils::CheckIfForceDisabled();
-	return UE::FUsdStage();
+		UsdWrapperUtils::CheckIfForceDisabled();
+		return UE::FUsdStage();
 #endif // #if USE_USD_SDK
+	}
+}
+
+UE::FUsdStage UnrealUSDWrapper::OpenStage(
+	const TCHAR* Identifier,
+	EUsdInitialLoadSet InitialLoadSet,
+	bool bUseStageCache,
+	bool bForceReloadLayersFromDisk
+)
+{
+	const TArray<FString>* PopulationMask = nullptr;
+	const TCHAR* SessionIdentifier = nullptr;
+	return UE::UnrealUSDWrapper::Private::OpenStageImpl(
+		Identifier,
+		SessionIdentifier,
+		InitialLoadSet,
+		bUseStageCache,
+		nullptr,
+		bForceReloadLayersFromDisk
+	);
+}
+
+ UE::FUsdStage UnrealUSDWrapper::OpenStage(
+	 UE::FSdfLayer RootLayer,
+	 UE::FSdfLayer SessionLayer,
+	 EUsdInitialLoadSet InitialLoadSet,
+	 bool bUseStageCache,
+	 bool bForceReloadLayersFromDisk
+ )
+{
+	 if( !RootLayer )
+	 {
+		 return {};
+	 }
+
+	 const TArray<FString>* PopulationMask = nullptr;
+	 const TCHAR* SessionIdentifier = nullptr;
+	 return UE::UnrealUSDWrapper::Private::OpenStageImpl(
+		 *RootLayer.GetIdentifier(),
+		 SessionLayer ? *SessionLayer.GetIdentifier() : nullptr,
+		 InitialLoadSet,
+		 bUseStageCache,
+		 nullptr,
+		 bForceReloadLayersFromDisk
+	 );
+}
+
+UE::FUsdStage UnrealUSDWrapper::OpenMaskedStage(
+	const TCHAR* Identifier,
+	EUsdInitialLoadSet InitialLoadSet,
+	const TArray<FString>& PopulationMask,
+	bool bForceReloadLayersFromDisk
+)
+{
+	const bool bUseStageCache = false;
+	const TCHAR* SessionIdentifier = nullptr;
+	return UE::UnrealUSDWrapper::Private::OpenStageImpl(
+		Identifier,
+		SessionIdentifier,
+		InitialLoadSet,
+		bUseStageCache,
+		&PopulationMask,
+		bForceReloadLayersFromDisk
+	);
+}
+
+UE::FUsdStage UnrealUSDWrapper::OpenMaskedStage(
+	UE::FSdfLayer RootLayer,
+	UE::FSdfLayer SessionLayer,
+	EUsdInitialLoadSet InitialLoadSet,
+	const TArray<FString>& PopulationMask,
+	bool bForceReloadLayersFromDisk
+)
+{
+	if ( !RootLayer )
+	{
+		return {};
+	}
+
+	const bool bUseStageCache = false;
+	const TCHAR* SessionIdentifier = nullptr;
+	return UE::UnrealUSDWrapper::Private::OpenStageImpl(
+		*RootLayer.GetIdentifier(),
+		SessionLayer ? *SessionLayer.GetIdentifier() : nullptr,
+		InitialLoadSet,
+		bUseStageCache,
+		&PopulationMask,
+		bForceReloadLayersFromDisk
+	);
 }
 
 UE::FUsdStage UnrealUSDWrapper::NewStage( const TCHAR* FilePath )
diff --git a/Engine/Plugins/Importers/USDImporter/Source/UnrealUSDWrapper/Public/UnrealUSDWrapper.h b/Engine/Plugins/Importers/USDImporter/Source/UnrealUSDWrapper/Public/UnrealUSDWrapper.h
--- a/Engine/Plugins/Importers/USDImporter/Source/UnrealUSDWrapper/Public/UnrealUSDWrapper.h
+++ b/Engine/Plugins/Importers/USDImporter/Source/UnrealUSDWrapper/Public/UnrealUSDWrapper.h
@@ -225,15 +225,47 @@ public:
 	UNREALUSDWRAPPER_API static TArray<FString> GetNativeFileFormats();
 
 	/**
-	 * Opens a file as a root layer of an USD stage, and returns that stage.
+	 * Opens a USD stage from a file on disk or existing layers, with a population mask or not.
 	 * @param Identifier - Path to a file that the USD SDK can open (or the identifier of a root layer), which will become the root layer of the new stage
+	 * @param RootLayer - Existing root layer to use for the new stage, instead of reading it from disk
+	 * @param SessionLayer - Existing session layer to use for the new stage, instead of creating a new one
 	 * @param InitialLoadSet - How to handle USD payloads when opening this stage
+	 * @param PopulationMask - List of prim paths to import, following the USD population mask rules
 	 * @param bUseStageCache - If true, and the stage is already opened in the stage cache (or the layers are already loaded in the registry) then
 	 *						   the file reading may be skipped, and the existing stage returned. When false, the stage and all its referenced layers
 	 *						   will be re-read anew, and the stage will not be added to the stage cache.
+	 * @param bForceReloadLayersFromDisk - USD layers are always cached in the layer registry, so trying to reopen an
+	 *                                     already opened layer will just fetch it from memory (potentially with
+	 *                                     edits). If this is true, all local layers used by the stage will be force-
+	 *                                     reloaded from disk
 	 * @return The opened stage, which may be invalid.
 	 */
-	UNREALUSDWRAPPER_API static UE::FUsdStage OpenStage( const TCHAR* Identifier, EUsdInitialLoadSet InitialLoadSet, bool bUseStageCache = true );
+	UNREALUSDWRAPPER_API static UE::FUsdStage OpenStage(
+		const TCHAR* Identifier,
+		EUsdInitialLoadSet InitialLoadSet,
+		bool bUseStageCache = true,
+		bool bForceReloadLayersFromDisk = false
+	);
+	UNREALUSDWRAPPER_API static UE::FUsdStage OpenStage(
+		UE::FSdfLayer RootLayer,
+		UE::FSdfLayer SessionLayer,
+		EUsdInitialLoadSet InitialLoadSet,
+		bool bUseStageCache = true,
+		bool bForceReloadLayersFromDisk = false
+	);
+	UNREALUSDWRAPPER_API static UE::FUsdStage OpenMaskedStage(
+		const TCHAR* Identifier,
+		EUsdInitialLoadSet InitialLoadSet,
+		const TArray<FString>& PopulationMask,
+		bool bForceReloadLayersFromDisk = false
+	);
+	UNREALUSDWRAPPER_API static UE::FUsdStage OpenMaskedStage(
+		UE::FSdfLayer RootLayer,
+		UE::FSdfLayer SessionLayer,
+		EUsdInitialLoadSet InitialLoadSet,
+		const TArray<FString>& PopulationMask,
+		bool bForceReloadLayersFromDisk = false
+	);
 
 	/** Creates a new USD root layer file, opens it as a new stage and returns that stage */
 	UNREALUSDWRAPPER_API static UE::FUsdStage NewStage( const TCHAR* FilePath );
