USD: Allow importing specific prim paths

diff --git a/Engine/Plugins/Importers/USDImporter/Resources/PythonExamples/import_specific_prims.py b/Engine/Plugins/Importers/USDImporter/Resources/PythonExamples/import_specific_prims.py
new file mode 100644
--- /dev/null
+++ b/Engine/Plugins/Importers/USDImporter/Resources/PythonExamples/import_specific_prims.py
@@ -0,0 +1,68 @@
+"""
+This script describes how to import a USD Stage into actors and assets, also optionally specifying specific prim
+paths to import.
+
+The `prims_to_import` property can be left at its default value of ["/"] (a list with just the "/" prim path in it)
+to import the entire stage.
+
+The provided paths in `prims_to_import` are used as a USD population mask when opening the stage, and as such are
+subject to the rules described here: https://graphics.pixar.com/usd/release/api/class_usd_stage_population_mask.html
+
+As an example, consider the following USD stage:
+
+        #usda 1.0
+        def Xform "ParentA"
+        {
+            def Xform "ChildA"
+            {
+            }
+
+            def Xform "ChildB"
+            {
+            }
+        }
+        def Xform "ParentB"
+        {
+            def Xform "ChildC"
+            {
+            }
+
+            def Xform "ChildD"
+            {
+            }
+        }
+
+In general, the main thing to keep in mind is that if "/ParentA" is within `prims_to_import`, ParentA and *all* of its
+children will be imported. As a consequence, having both "/ParentA" and "/ParentA/ChildA" on the list is reduntant, as
+"/ParentA" will already lead to "/ParentA/ChildA" being imported, as previously mentioned.
+"""
+
+import unreal
+
+ROOT_LAYER_FILENAME = r"C:\MyFolder\my_scene.usda"
+DESTINATION_CONTENT_PATH = r"/Game/Imported/"
+
+options = unreal.UsdStageImportOptions()
+options.import_actors = True
+options.import_geometry = True
+options.import_skeletal_animations = True
+options.import_level_sequences = True
+options.import_materials = True
+options.prim_path_folder_structure = False
+options.prims_to_import = [
+    "/ParentA",  # This will import ParentA, ChildA and ChildB
+    "/ParentB/ChildC"  # This will import ParentB and ChildC only (and *not* import ChildD)
+]
+
+task = unreal.AssetImportTask()
+task.set_editor_property('filename', ROOT_LAYER_FILENAME)
+task.set_editor_property('destination_path', DESTINATION_CONTENT_PATH)
+task.set_editor_property('automated', True)
+task.set_editor_property('options', options)
+task.set_editor_property('replace_existing', True)
+
+asset_tools = unreal.AssetToolsHelpers.get_asset_tools()
+asset_tools.import_asset_tasks([task])
+
+asset_paths = task.get_editor_property("imported_object_paths")
+success = asset_paths and len(asset_paths) > 0
\ No newline at end of file
diff --git a/Engine/Plugins/Importers/USDImporter/Source/USDExporter/Private/AnimSequenceExporterUSD.cpp b/Engine/Plugins/Importers/USDImporter/Source/USDExporter/Private/AnimSequenceExporterUSD.cpp
--- a/Engine/Plugins/Importers/USDImporter/Source/USDExporter/Private/AnimSequenceExporterUSD.cpp
+++ b/Engine/Plugins/Importers/USDImporter/Source/USDExporter/Private/AnimSequenceExporterUSD.cpp
@@ -139,8 +139,7 @@ bool UAnimSequenceExporterUSD::ExportBinary( UObject* Object, const TCHAR* Type,
 		{
 			Options->PreviewMeshOptions.MaterialBakingOptions.TexturesDir.Path = FPaths::Combine( FPaths::GetPath( AnimSequenceFile ), TEXT( "Textures" ) );
 
-			const bool bIsImport = false;
-			const bool bContinue = SUsdOptionsWindow::ShowImportExportOptions( *Options, bIsImport );
+			const bool bContinue = SUsdOptionsWindow::ShowExportOptions( *Options );
 			if ( !bContinue )
 			{
 				return false;
diff --git a/Engine/Plugins/Importers/USDImporter/Source/USDExporter/Private/LevelExporterUSD.cpp b/Engine/Plugins/Importers/USDImporter/Source/USDExporter/Private/LevelExporterUSD.cpp
--- a/Engine/Plugins/Importers/USDImporter/Source/USDExporter/Private/LevelExporterUSD.cpp
+++ b/Engine/Plugins/Importers/USDImporter/Source/USDExporter/Private/LevelExporterUSD.cpp
@@ -80,8 +80,7 @@ bool ULevelExporterUSD::ExportBinary( UObject* Object, const TCHAR* Type, FArchi
 	{
 		Options->Inner.AssetFolder.Path = FPaths::Combine( FPaths::GetPath( UExporter::CurrentFilename ), TEXT( "Assets" ) );
 
-		const bool bIsImport = false;
-		const bool bContinue = SUsdOptionsWindow::ShowImportExportOptions( *Options, bIsImport );
+		const bool bContinue = SUsdOptionsWindow::ShowExportOptions( *Options );
 		if ( !bContinue )
 		{
 			return false;
diff --git a/Engine/Plugins/Importers/USDImporter/Source/USDExporter/Private/LevelSequenceExporterUSD.cpp b/Engine/Plugins/Importers/USDImporter/Source/USDExporter/Private/LevelSequenceExporterUSD.cpp
--- a/Engine/Plugins/Importers/USDImporter/Source/USDExporter/Private/LevelSequenceExporterUSD.cpp
+++ b/Engine/Plugins/Importers/USDImporter/Source/USDExporter/Private/LevelSequenceExporterUSD.cpp
@@ -1186,8 +1186,7 @@ bool ULevelSequenceExporterUsd::ExportBinary( UObject* Object, const TCHAR* Type
 			Options->LevelExportOptions.AssetFolder.Path = FPaths::Combine( FPaths::GetPath( UExporter::CurrentFilename ), TEXT( "Assets" ) );
 			Options->TimeCodesPerSecond = MovieScene->GetDisplayRate().AsDecimal();
 
-			const bool bIsImport = false;
-			const bool bContinue = SUsdOptionsWindow::ShowImportExportOptions( *Options, bIsImport );
+			const bool bContinue = SUsdOptionsWindow::ShowExportOptions( *Options );
 			if ( !bContinue )
 			{
 				return false;
diff --git a/Engine/Plugins/Importers/USDImporter/Source/USDExporter/Private/MaterialExporterUSD.cpp b/Engine/Plugins/Importers/USDImporter/Source/USDExporter/Private/MaterialExporterUSD.cpp
--- a/Engine/Plugins/Importers/USDImporter/Source/USDExporter/Private/MaterialExporterUSD.cpp
+++ b/Engine/Plugins/Importers/USDImporter/Source/USDExporter/Private/MaterialExporterUSD.cpp
@@ -217,8 +217,7 @@ bool UMaterialExporterUsd::ExportBinary( UObject* Object, const TCHAR* Type, FAr
 		{
 			Options->MaterialBakingOptions.TexturesDir.Path = FPaths::Combine( FPaths::GetPath( UExporter::CurrentFilename ), TEXT( "Textures" ) );
 
-			const bool bIsImport = false;
-			const bool bContinue = SUsdOptionsWindow::ShowImportExportOptions( *Options, bIsImport );
+			const bool bContinue = SUsdOptionsWindow::ShowExportOptions( *Options );
 			if ( !bContinue )
 			{
 				return false;
diff --git a/Engine/Plugins/Importers/USDImporter/Source/USDExporter/Private/SkeletalMeshExporterUSD.cpp b/Engine/Plugins/Importers/USDImporter/Source/USDExporter/Private/SkeletalMeshExporterUSD.cpp
--- a/Engine/Plugins/Importers/USDImporter/Source/USDExporter/Private/SkeletalMeshExporterUSD.cpp
+++ b/Engine/Plugins/Importers/USDImporter/Source/USDExporter/Private/SkeletalMeshExporterUSD.cpp
@@ -125,8 +125,7 @@ bool USkeletalMeshExporterUsd::ExportBinary( UObject* Object, const TCHAR* Type,
 		{
 			Options->MeshAssetOptions.MaterialBakingOptions.TexturesDir.Path = FPaths::Combine( FPaths::GetPath( UExporter::CurrentFilename ), TEXT( "Textures" ) );
 
-			const bool bIsImport = false;
-			const bool bContinue = SUsdOptionsWindow::ShowImportExportOptions( *Options, bIsImport );
+			const bool bContinue = SUsdOptionsWindow::ShowExportOptions( *Options );
 			if ( !bContinue )
 			{
 				return false;
diff --git a/Engine/Plugins/Importers/USDImporter/Source/USDExporter/Private/StaticMeshExporterUSD.cpp b/Engine/Plugins/Importers/USDImporter/Source/USDExporter/Private/StaticMeshExporterUSD.cpp
--- a/Engine/Plugins/Importers/USDImporter/Source/USDExporter/Private/StaticMeshExporterUSD.cpp
+++ b/Engine/Plugins/Importers/USDImporter/Source/USDExporter/Private/StaticMeshExporterUSD.cpp
@@ -135,8 +135,7 @@ bool UStaticMeshExporterUsd::ExportBinary( UObject* Object, const TCHAR* Type, F
 		{
 			Options->MeshAssetOptions.MaterialBakingOptions.TexturesDir.Path = FPaths::Combine( FPaths::GetPath( UExporter::CurrentFilename ), TEXT( "Textures" ) );
 
-			const bool bIsImport = false;
-			const bool bContinue = SUsdOptionsWindow::ShowImportExportOptions( *Options, bIsImport );
+			const bool bContinue = SUsdOptionsWindow::ShowExportOptions( *Options );
 			if ( !bContinue )
 			{
 				return false;
diff --git a/Engine/Plugins/Importers/USDImporter/Source/USDStageEditor/Private/Widgets/SUSDTreeView.h b/Engine/Plugins/Importers/USDImporter/Source/USDStageEditor/Private/Widgets/SUSDTreeView.h
deleted file mode 100644
diff --git a/Engine/Plugins/Importers/USDImporter/Source/USDStageEditorViewModels/Public/USDLayersViewModel.h b/Engine/Plugins/Importers/USDImporter/Source/USDStageEditorViewModels/Public/USDLayersViewModel.h
--- a/Engine/Plugins/Importers/USDImporter/Source/USDStageEditorViewModels/Public/USDLayersViewModel.h
+++ b/Engine/Plugins/Importers/USDImporter/Source/USDStageEditorViewModels/Public/USDLayersViewModel.h
@@ -5,8 +5,8 @@
 #include "Internationalization/Text.h"
 #include "Templates/SharedPointer.h"
 
-#include "USDTreeItemViewModel.h"
 #include "UsdWrappers/UsdStage.h"
+#include "Widgets/IUSDTreeViewItem.h"
 
 class USDSTAGEEDITORVIEWMODELS_API FUsdLayerModel : public TSharedFromThis< FUsdLayerModel >
 {
diff --git a/Engine/Plugins/Importers/USDImporter/Source/USDStageEditorViewModels/Public/USDPrimViewModel.h b/Engine/Plugins/Importers/USDImporter/Source/USDStageEditorViewModels/Public/USDPrimViewModel.h
--- a/Engine/Plugins/Importers/USDImporter/Source/USDStageEditorViewModels/Public/USDPrimViewModel.h
+++ b/Engine/Plugins/Importers/USDImporter/Source/USDStageEditorViewModels/Public/USDPrimViewModel.h
@@ -2,7 +2,9 @@
 
 #pragma once
 
-#include "USDTreeItemViewModel.h"
+#include "UsdWrappers/UsdStage.h"
+#include "UsdWrappers/UsdPrim.h"
+#include "Widgets/IUSDTreeViewItem.h"
 
 #include "UsdWrappers/ForwardDeclarations.h"
 #include "UsdWrappers/UsdStage.h"
diff --git a/Engine/Plugins/Importers/USDImporter/Source/USDStageEditorViewModels/Public/USDTreeItemViewModel.h b/Engine/Plugins/Importers/USDImporter/Source/USDStageEditorViewModels/Public/USDTreeItemViewModel.h
deleted file mode 100644
diff --git a/Engine/Plugins/Importers/USDImporter/Source/USDStageImporter/Private/USDOptionsWindow.cpp b/Engine/Plugins/Importers/USDImporter/Source/USDStageImporter/Private/USDOptionsWindow.cpp
--- a/Engine/Plugins/Importers/USDImporter/Source/USDStageImporter/Private/USDOptionsWindow.cpp
+++ b/Engine/Plugins/Importers/USDImporter/Source/USDStageImporter/Private/USDOptionsWindow.cpp
@@ -2,41 +2,36 @@
 
 #include "USDOptionsWindow.h"
 
+#include "USDStageImportOptions.h"
+#include "UsdWrappers/SdfLayer.h"
+#include "UsdWrappers/SdfPath.h"
+#include "UsdWrappers/UsdPrim.h"
+#include "UsdWrappers/UsdStage.h"
+#include "Widgets/SUSDStagePreviewTree.h"
+
 #include "CoreMinimal.h"
 #include "Framework/Application/SlateApplication.h"
+#include "HAL/PlatformApplicationMisc.h"
 #include "Interfaces/IMainFrameModule.h"
 #include "Misc/ScopedSlowTask.h"
 #include "Modules/ModuleManager.h"
 #include "PropertyEditorModule.h"
 #include "Widgets/Input/SButton.h"
+#include "Widgets/Input/SSearchBox.h"
 #include "Widgets/Layout/SBox.h"
+#include "Widgets/Layout/SSplitter.h"
 #include "Widgets/Layout/SUniformGridPanel.h"
 #include "Widgets/SBoxPanel.h"
 #include "Widgets/SWindow.h"
 
 #define LOCTEXT_NAMESPACE "USDOptionsWindow"
 
-bool SUsdOptionsWindow::ShowImportExportOptions( UObject& OptionsObject, bool bIsImport )
-{
-	if ( bIsImport )
-	{
-		return SUsdOptionsWindow::ShowOptions(
-			OptionsObject,
-			LOCTEXT( "USDImportOptionsTitle", "USD Import Options" ),
-			LOCTEXT( "USDOptionWindow_Import", "Import" )
-		);
-	}
-	else
-	{
-		return SUsdOptionsWindow::ShowOptions(
-			OptionsObject,
-			LOCTEXT( "USDExportOptionsTitle", "USD Export Options" ),
-			LOCTEXT( "USDOptionWindow_Export", "Export" )
-		);
-	}
-}
-
-bool SUsdOptionsWindow::ShowOptions( UObject& OptionsObject, const FText& WindowTitle, const FText& AcceptText )
+bool SUsdOptionsWindow::ShowOptions(
+	UObject& OptionsObject,
+	const FText& WindowTitle,
+	const FText& AcceptText,
+	const UE::FUsdStage* Stage
+)
 {
 	TSharedPtr<SWindow> ParentWindow;
 
@@ -46,9 +41,22 @@ bool SUsdOptionsWindow::ShowOptions( UObject& OptionsObject, const FText& Window
 		ParentWindow = MainFrame.GetParentWindow();
 	}
 
+	FDisplayMetrics DisplayMetrics;
+	FSlateApplication::Get().GetDisplayMetrics( DisplayMetrics );
+	const float DPIScaleFactor = FPlatformApplicationMisc::GetDPIScaleFactorAtPoint(
+		DisplayMetrics.PrimaryDisplayWorkAreaRect.Left,
+		DisplayMetrics.PrimaryDisplayWorkAreaRect.Top
+	);
+
+	const FVector2D ClientSize{
+		320.0f * DPIScaleFactor * ( Stage == nullptr ? 1 : 2 ), // Make it wider if we're going to show two panes
+		320.0f * DPIScaleFactor
+	};
+
 	TSharedRef<SWindow> Window = SNew( SWindow )
 		.Title( WindowTitle )
-		.SizingRule( ESizingRule::Autosized );
+		.SizingRule( ESizingRule::UserSized )
+		.ClientSize( ClientSize );
 
 	TSharedPtr<SUsdOptionsWindow> OptionsWindow;
 	Window->SetContent
@@ -57,6 +65,7 @@ bool SUsdOptionsWindow::ShowOptions( UObject& OptionsObject, const FText& Window
 		.OptionsObject( &OptionsObject )
 		.AcceptText( AcceptText )
 		.WidgetWindow( Window )
+		.Stage( Stage )
 	);
 
 	// Preemptively make sure we have a progress dialog created before showing our modal. This because the progress
@@ -70,7 +79,35 @@ bool SUsdOptionsWindow::ShowOptions( UObject& OptionsObject, const FText& Window
 	const bool bSlowTaskWindow = false;
 	FSlateApplication::Get().AddModalWindow( Window, ParentWindow, bSlowTaskWindow );
 
-	return OptionsWindow->UserAccepted();
+	const bool bAccepted = OptionsWindow->UserAccepted();
+	if ( bAccepted )
+	{
+		if ( UUsdStageImportOptions* StageImportOptions = Cast< UUsdStageImportOptions>( &OptionsObject ) )
+		{
+			StageImportOptions->PrimsToImport = OptionsWindow->GetSelectedFullPrimPaths();
+		}
+	}
+
+	return bAccepted;
+}
+
+bool SUsdOptionsWindow::ShowImportOptions( UObject& OptionsObject, const UE::FUsdStage* StageToImport )
+{
+	return SUsdOptionsWindow::ShowOptions(
+		OptionsObject,
+		LOCTEXT( "USDImportOptionsTitle", "USD Import Options" ),
+		LOCTEXT( "USDOptionWindow_Import", "Import" ),
+		StageToImport
+	);
+}
+
+bool SUsdOptionsWindow::ShowExportOptions( UObject& OptionsObject )
+{
+	return SUsdOptionsWindow::ShowOptions(
+		OptionsObject,
+		LOCTEXT( "USDExportOptionsTitle", "USD Export Options" ),
+		LOCTEXT( "USDOptionWindow_Export", "Export" )
+	);
 }
 
 void SUsdOptionsWindow::Construct( const FArguments& InArgs )
@@ -78,56 +115,124 @@ void SUsdOptionsWindow::Construct( const FArguments& InArgs )
 	OptionsObject = InArgs._OptionsObject;
 	Window = InArgs._WidgetWindow;
 	AcceptText = InArgs._AcceptText;
+	const UE::FUsdStage* Stage = InArgs._Stage;
 	bAccepted = false;
 
-	TSharedPtr<SBox> DetailsViewBox;
+	// Prepare box that shows the OptionsObject's properties
+	TSharedPtr<SWidget> DetailsViewBox;
+	{
+		FPropertyEditorModule& PropertyEditorModule = FModuleManager::GetModuleChecked<FPropertyEditorModule>( "PropertyEditor" );
 
-	ChildSlot
-	[
-		SNew( SVerticalBox )
-		+ SVerticalBox::Slot()
-		.AutoHeight()
-		.Padding( 2 )
-		[
-			SAssignNew( DetailsViewBox, SBox )
-			.MaxDesiredHeight( 450.0f )
-			.MinDesiredWidth( 550.0f )
-		]
+		FDetailsViewArgs DetailsViewArgs;
+		DetailsViewArgs.bAllowSearch = false;
+		DetailsViewArgs.NameAreaSettings = FDetailsViewArgs::HideNameArea;
+		TSharedPtr<IDetailsView> DetailsView = PropertyEditorModule.CreateDetailView( DetailsViewArgs );
+		DetailsView->SetObject( OptionsObject );
+
+		DetailsViewBox = DetailsView.ToSharedRef();
+	}
 
-		+ SVerticalBox::Slot()
-		.AutoHeight()
-		.HAlign( HAlign_Right )
-		.Padding( 2 )
+	TSharedRef<SSplitter> Splitter = SNew( SSplitter ).Orientation( Orient_Horizontal );
+	{
+		// If we have a stage, show the preview tree on the left
+		if ( Stage )
+		{
+			SAssignNew( StagePreviewTree, SUsdStagePreviewTree, *Stage );
+
+			Splitter->AddSlot()
+			.Value(0.5f)
+			[
+				SNew( SVerticalBox )
+
+				+SVerticalBox::Slot()
+				.AutoHeight()
+				.Padding( 0.0f )
+				[
+					SNew( SHorizontalBox )
+					+SHorizontalBox::Slot()
+					.FillWidth( 1.0f )
+					.Padding( 0.0f )
+					.VAlign( VAlign_Center )
+					[
+						SNew( STextBlock )
+						.Font( FAppStyle::GetFontStyle( "CurveEd.LabelFont" ) )
+						.Text( LOCTEXT( "TreeDescriptionText", "Select prims to import:" ) )
+					]
+
+					+SHorizontalBox::Slot()
+					.AutoWidth()
+					.MaxWidth( 300.0f )
+					.Padding( FMargin( 5.0f, 2.0f ) )
+					.VAlign( VAlign_Center )
+					[
+						SNew( SSearchBox )
+						.HintText( LOCTEXT( "SearchInitialText", "Filter prims..." ) )
+						.OnTextChanged_Lambda( [this]( const FText& SearchText )
+						{
+							if ( StagePreviewTree )
+							{
+								StagePreviewTree->SetFilterText( SearchText );
+							}
+						})
+					]
+				]
+
+				+SVerticalBox::Slot()
+				.FillHeight( 1.0f )
+				.Padding( 0.0f )
+				[
+					StagePreviewTree.ToSharedRef()
+				]
+			];
+		}
+
+		Splitter->AddSlot()
+		.Value(0.5f)
 		[
-			SNew( SUniformGridPanel )
-			.SlotPadding( 2 )
+			DetailsViewBox.ToSharedRef()
+		];
+	}
 
-			+ SUniformGridPanel::Slot( 0, 0 )
+	ChildSlot
+	[
+		SNew( SBorder )
+		.Padding( FMargin( 10.0f, 0.0f, 10.0f, 0.0f ) )
+		.BorderImage( FAppStyle::GetBrush( "ToolPanel.DarkGroupBorder" ) )
+		[
+			SNew( SVerticalBox )
+			+SVerticalBox::Slot()
+			.FillHeight( 1.0f )
+			.Padding( 0.0f )
 			[
-				SNew( SButton )
-				.HAlign( HAlign_Center )
-				.Text( AcceptText )
-				.OnClicked( this, &SUsdOptionsWindow::OnAccept )
+				Splitter
 			]
 
-			+ SUniformGridPanel::Slot( 1, 0 )
+			+SVerticalBox::Slot()
+			.AutoHeight()
+			.HAlign( HAlign_Right )
+			.Padding( 8 )
 			[
-				SNew( SButton )
-				.HAlign( HAlign_Center )
-				.Text( LOCTEXT( "USDOptionWindow_Cancel", "Cancel" ) )
-				.OnClicked( this, &SUsdOptionsWindow::OnCancel )
+				SNew( SUniformGridPanel )
+				.SlotPadding( 2 )
+
+				+ SUniformGridPanel::Slot( 0, 0 )
+				[
+					SNew( SButton )
+					.HAlign( HAlign_Center )
+					.Text( AcceptText )
+					.OnClicked( this, &SUsdOptionsWindow::OnAccept )
+				]
+
+				+ SUniformGridPanel::Slot( 1, 0 )
+				[
+					SNew( SButton )
+					.HAlign( HAlign_Center )
+					.Text( LOCTEXT( "USDOptionWindow_Cancel", "Cancel" ) )
+					.OnClicked( this, &SUsdOptionsWindow::OnCancel )
+				]
 			]
 		]
 	];
-
-	FPropertyEditorModule& PropertyEditorModule = FModuleManager::GetModuleChecked<FPropertyEditorModule>( "PropertyEditor" );
-	FDetailsViewArgs DetailsViewArgs;
-	DetailsViewArgs.bAllowSearch = false;
-	DetailsViewArgs.NameAreaSettings = FDetailsViewArgs::HideNameArea;
-	TSharedPtr<IDetailsView> DetailsView = PropertyEditorModule.CreateDetailView( DetailsViewArgs );
-
-	DetailsViewBox->SetContent( DetailsView.ToSharedRef() );
-	DetailsView->SetObject( OptionsObject );
 }
 
 bool SUsdOptionsWindow::SupportsKeyboardFocus() const
@@ -174,4 +279,14 @@ bool SUsdOptionsWindow::UserAccepted() const
 	return bAccepted;
 }
 
+TArray<FString> SUsdOptionsWindow::GetSelectedFullPrimPaths() const
+{
+	if ( StagePreviewTree )
+	{
+		return StagePreviewTree->GetSelectedFullPrimPaths();
+	}
+
+	return {};
+}
+
 #undef LOCTEXT_NAMESPACE
\ No newline at end of file
diff --git a/Engine/Plugins/Importers/USDImporter/Source/USDStageImporter/Private/USDStageAssetImportFactory.cpp b/Engine/Plugins/Importers/USDImporter/Source/USDStageImporter/Private/USDStageAssetImportFactory.cpp
--- a/Engine/Plugins/Importers/USDImporter/Source/USDStageImporter/Private/USDStageAssetImportFactory.cpp
+++ b/Engine/Plugins/Importers/USDImporter/Source/USDStageImporter/Private/USDStageAssetImportFactory.cpp
@@ -60,6 +60,9 @@ UObject* UUsdStageAssetImportFactory::FactoryCreateFile(UClass* InClass, UObject
 		ImportContext.ImportOptions = Cast<UUsdStageImportOptions>( AssetImportTask->Options );
 	}
 
+	// When importing from file we don't want to use any opened stage
+	ImportContext.bReadFromStageCache = false;
+
 	const FString InitialPackagePath = InParent ? InParent->GetName() : TEXT( "/Game/" );
 	const bool bIsReimport = false;
 	const bool bAllowActorImport = false;
@@ -153,6 +156,8 @@ EReimportResult::Type UUsdStageAssetImportFactory::Reimport(UObject* Obj)
 		ImportContext.ImportOptions = DuplicateObject( ReimportOptions, ImportData );
 	}
 
+	ImportContext.bReadFromStageCache = false;
+
 	const bool bIsReimport = true;
 	const bool bAllowActorImport = false;
 	if (!ImportContext.Init(Obj->GetName(), ImportData->GetFirstFilename(), Obj->GetName(), Obj->GetFlags(), IsAutomatedImport(), bIsReimport, bAllowActorImport))
diff --git a/Engine/Plugins/Importers/USDImporter/Source/USDStageImporter/Private/USDStageImportContext.cpp b/Engine/Plugins/Importers/USDImporter/Source/USDStageImporter/Private/USDStageImportContext.cpp
--- a/Engine/Plugins/Importers/USDImporter/Source/USDStageImporter/Private/USDStageImportContext.cpp
+++ b/Engine/Plugins/Importers/USDImporter/Source/USDStageImporter/Private/USDStageImportContext.cpp
@@ -36,6 +36,13 @@ bool FUsdStageImportContext::Init(const FString& InName, const FString& InFilePa
 
 	FPaths::NormalizeFilename(FilePath);
 
+	// Open the stage if we haven't yet, as we'll need it open to show the preview tree
+	if ( !Stage )
+	{
+		const EUsdInitialLoadSet InitialLoadSet = EUsdInitialLoadSet::LoadAll;
+		Stage = UnrealUSDWrapper::OpenStage( *FilePath, InitialLoadSet, bReadFromStageCache );
+	}
+
 	if(!bIsAutomated)
 	{
 		// Show dialog for content folder
@@ -59,8 +66,7 @@ bool FUsdStageImportContext::Init(const FString& InName, const FString& InFilePa
 		ImportOptions->EnableActorImport( bAllowActorImport );
 
 		// Show dialog for import options
-		const bool bIsImport = true;
-		bool bProceedWithImport = SUsdOptionsWindow::ShowImportExportOptions( *ImportOptions, bIsImport );
+		bool bProceedWithImport = SUsdOptionsWindow::ShowImportOptions( *ImportOptions, &Stage );
 		if (!bProceedWithImport)
 		{
 			return false;
diff --git a/Engine/Plugins/Importers/USDImporter/Source/USDStageImporter/Private/USDStageImportFactory.cpp b/Engine/Plugins/Importers/USDImporter/Source/USDStageImporter/Private/USDStageImportFactory.cpp
--- a/Engine/Plugins/Importers/USDImporter/Source/USDStageImporter/Private/USDStageImportFactory.cpp
+++ b/Engine/Plugins/Importers/USDImporter/Source/USDStageImporter/Private/USDStageImportFactory.cpp
@@ -47,6 +47,9 @@ UObject* UUsdStageImportFactory::FactoryCreateFile(UClass* InClass, UObject* InP
 		ImportContext.ImportOptions = Cast<UUsdStageImportOptions>( AssetImportTask->Options );
 	}
 
+	// When importing from file we don't want to use any opened stage
+	ImportContext.bReadFromStageCache = false;
+
 #if USE_USD_SDK
 	const FString InitialPackagePath =InParent ? InParent->GetName() : TEXT( "/Game/" );
 	const bool bIsReimport = false;
diff --git a/Engine/Plugins/Importers/USDImporter/Source/USDStageImporter/Private/USDStageImportOptionsCustomization.cpp b/Engine/Plugins/Importers/USDImporter/Source/USDStageImporter/Private/USDStageImportOptionsCustomization.cpp
--- a/Engine/Plugins/Importers/USDImporter/Source/USDStageImporter/Private/USDStageImportOptionsCustomization.cpp
+++ b/Engine/Plugins/Importers/USDImporter/Source/USDStageImporter/Private/USDStageImportOptionsCustomization.cpp
@@ -47,6 +47,13 @@ void FUsdStageImportOptionsCustomization::CustomizeDetails(IDetailLayoutBuilder&
 		return;
 	}
 
+	// Hide this property since we'll show the preview tree for it
+	IDetailCategoryBuilder& PrimsToImportCatBuilder = DetailLayoutBuilder.EditCategory( TEXT( "Prims to Import" ) );
+	if ( TSharedPtr<IPropertyHandle> PrimsToImportProperty = DetailLayoutBuilder.GetProperty( GET_MEMBER_NAME_CHECKED( UUsdStageImportOptions, PrimsToImport ) ) )
+	{
+		DetailLayoutBuilder.HideProperty( PrimsToImportProperty );
+	}
+
 	IUsdSchemasModule& UsdSchemasModule = FModuleManager::Get().LoadModuleChecked< IUsdSchemasModule >( TEXT( "USDSchemas" ) );
 
 	RenderContextComboBoxItems.Reset();
diff --git a/Engine/Plugins/Importers/USDImporter/Source/USDStageImporter/Private/USDStageImporter.cpp b/Engine/Plugins/Importers/USDImporter/Source/USDStageImporter/Private/USDStageImporter.cpp
--- a/Engine/Plugins/Importers/USDImporter/Source/USDStageImporter/Private/USDStageImporter.cpp
+++ b/Engine/Plugins/Importers/USDImporter/Source/USDStageImporter/Private/USDStageImporter.cpp
@@ -54,24 +54,75 @@
 
 namespace UsdStageImporterImpl
 {
-	void LoadStageFromFilePath(FUsdStageImportContext& ImportContext)
+	void OpenStage( FUsdStageImportContext& ImportContext, bool bNeedsMasking )
 	{
-		const FString FilePath = IFileManager::Get().ConvertToAbsolutePathForExternalAppForRead(*ImportContext.FilePath);
+		const FString FilePath = !ImportContext.FilePath.IsEmpty()
+			? IFileManager::Get().ConvertToAbsolutePathForExternalAppForRead( *ImportContext.FilePath )
+			: FString{};
+
+		if ( FilePath.IsEmpty() && !ImportContext.Stage )
+		{
+			// There's nothing we can do here to create a stage
+			return;
+		}
 
 		UsdUtils::StartMonitoringErrors();
 
-		for ( const UE::FUsdStage& OpenedStage : UnrealUSDWrapper::GetAllStagesFromCache() )
+		UE::FUsdStage Stage;
+		if ( bNeedsMasking )
 		{
-			FString RootPath = OpenedStage.GetRootLayer().GetRealPath();
-			FPaths::NormalizeFilename( RootPath );
-			if ( ImportContext.FilePath == RootPath )
+			// If we're masking we'll make a unique stage for us.
+			// Also, within USD the OpenMasked methods do not consult the stage cache anyway
+			ImportContext.bReadFromStageCache = false;
+			ImportContext.bStageWasOriginallyOpenInCache = false;
+
+			// Here even though we won't use the stage cache we don't want to force reload the stage's layers either:
+			// The whole point is that we're just opening another masked "view" into the same composed (potentially
+			// modified) layers
+			const bool bForceReloadLayersFromDisk = false;
+
+			// We don't have a file path to reopen, so just reopen the existing layers
+			if ( FilePath.IsEmpty() )
 			{
-				ImportContext.bStageWasOriginallyOpenInCache = true;
-				break;
+				Stage = UnrealUSDWrapper::OpenMaskedStage(
+					ImportContext.Stage.GetRootLayer(),
+					ImportContext.Stage.GetSessionLayer(),
+					EUsdInitialLoadSet::LoadAll,
+					ImportContext.ImportOptions->PrimsToImport,
+					bForceReloadLayersFromDisk
+				);
+			}
+			else
+			{
+				Stage = UnrealUSDWrapper::OpenMaskedStage(
+					*FilePath,
+					EUsdInitialLoadSet::LoadAll,
+					ImportContext.ImportOptions->PrimsToImport,
+					bForceReloadLayersFromDisk
+				);
 			}
 		}
+		else
+		{
+			for ( const UE::FUsdStage& OpenedStage : UnrealUSDWrapper::GetAllStagesFromCache() )
+			{
+				FString RootPath = OpenedStage.GetRootLayer().GetRealPath();
+				FPaths::NormalizeFilename( RootPath );
+				if ( ImportContext.FilePath == RootPath )
+				{
+					ImportContext.bStageWasOriginallyOpenInCache = true;
+					break;
+				}
+			}
 
-		UE::FUsdStage Stage = UnrealUSDWrapper::OpenStage( *FilePath, EUsdInitialLoadSet::LoadAll, ImportContext.bReadFromStageCache );
+			const bool bForceReloadLayersFromDisk = !ImportContext.bReadFromStageCache;
+			Stage = UnrealUSDWrapper::OpenStage(
+				*FilePath,
+				EUsdInitialLoadSet::LoadAll,
+				ImportContext.bReadFromStageCache,
+				bForceReloadLayersFromDisk
+			);
+		}
 
 		TArray<FString> ErrorStrings = UsdUtils::GetErrorsAndStopMonitoring();
 		FString Error = FString::Join(ErrorStrings, TEXT("\n"));
@@ -1447,9 +1498,17 @@ void UUsdStageImporter::ImportFromFile(FUsdStageImportContext& ImportContext)
 
 	double StartTime = FPlatformTime::Cycles64();
 
-	if ( !ImportContext.Stage && !ImportContext.FilePath.IsEmpty() )
+	if ( ImportContext.ImportOptions->PrimsToImport.Num() == 0 )
 	{
-		UsdStageImporterImpl::LoadStageFromFilePath( ImportContext );
+		return;
+	}
+
+	const bool bNeedsMasking =
+		ImportContext.ImportOptions->PrimsToImport != TArray<FString>{ UE::FSdfPath::AbsoluteRootPath().GetString() };
+
+	if ( !ImportContext.Stage || bNeedsMasking )
+	{
+		UsdStageImporterImpl::OpenStage( ImportContext, bNeedsMasking );
 	}
 
 	if ( !ImportContext.Stage )
@@ -1550,9 +1609,13 @@ bool UUsdStageImporter::ReimportSingleAsset(FUsdStageImportContext& ImportContex
 #if USE_USD_SDK
 	double StartTime = FPlatformTime::Cycles64();
 
-	if ( !ImportContext.Stage && !ImportContext.FilePath.IsEmpty() )
+	// TODO: Maybe change this whole reimporting approach to just taking advantage of a population mask instead?
+	const bool bNeedsMasking =
+		ImportContext.ImportOptions->PrimsToImport != TArray<FString>{ UE::FSdfPath::AbsoluteRootPath().GetString() };
+
+	if ( !ImportContext.Stage || bNeedsMasking )
 	{
-		UsdStageImporterImpl::LoadStageFromFilePath( ImportContext );
+		UsdStageImporterImpl::OpenStage( ImportContext, bNeedsMasking );
 	}
 
 	if ( !ImportContext.Stage )
diff --git a/Engine/Plugins/Importers/USDImporter/Source/USDStageImporter/Private/Widgets/SUSDStagePreviewTree.cpp b/Engine/Plugins/Importers/USDImporter/Source/USDStageImporter/Private/Widgets/SUSDStagePreviewTree.cpp
new file mode 100644
--- /dev/null
+++ b/Engine/Plugins/Importers/USDImporter/Source/USDStageImporter/Private/Widgets/SUSDStagePreviewTree.cpp
@@ -0,0 +1,743 @@
+// Copyright Epic Games, Inc. All Rights Reserved.
+
+#include "Widgets/SUSDStagePreviewTree.h"
+
+#include "UsdWrappers/UsdPrim.h"
+#include "UsdWrappers/UsdStage.h"
+
+#include "Framework/Application/SlateApplication.h"
+#include "Framework/MultiBox/MultiBoxBuilder.h"
+#include "Framework/Notifications/NotificationManager.h"
+#include "Widgets/Input/SCheckBox.h"
+#include "Widgets/Notifications/SNotificationList.h"
+#include "Widgets/Views/STableRow.h"
+
+#define LOCTEXT_NAMESPACE "UsdStagePreviewTree"
+
+namespace UE::UsdStagePreviewTree::Private
+{
+	const float CheckBoxColumnWidth = 24.0f;
+
+	bool AreAllViewsCheckedRecursive( const SUsdStagePreviewTree& Tree, const TArray<FUsdPrimPreviewModelViewRef>& Views )
+	{
+		for ( const FUsdPrimPreviewModelViewRef& View : Views )
+		{
+			if ( !Tree.GetModel( View->ModelIndex ).bShouldImport )
+			{
+				return false;
+			}
+
+			AreAllViewsCheckedRecursive( Tree, View->Children );
+		}
+
+		return true;
+	};
+}
+
+class FUsdPreviewTreeImportColumn : public FUsdTreeViewColumn
+{
+public:
+	virtual TSharedRef< SWidget > GenerateWidget(
+		const TSharedPtr< IUsdTreeViewItem > InTreeItem,
+		const TSharedPtr< ITableRow > TableRow
+	) override
+	{
+		const FUsdPrimPreviewModelViewWeak WeakTreeItem = StaticCastSharedPtr< FUsdPrimPreviewModelView >( InTreeItem );
+		TSharedPtr< SUsdStagePreviewTreeRow > Row = StaticCastSharedPtr< SUsdStagePreviewTreeRow >( TableRow );
+		TWeakPtr< SUsdStagePreviewTree > WeakTree = Row ? Row->GetOwnerTree() : nullptr;
+
+		return SNew( SBox )
+			.VAlign( VAlign_Center )
+			.HAlign( HAlign_Center )
+			.WidthOverride( UE::UsdStagePreviewTree::Private::CheckBoxColumnWidth )
+			[
+				SNew( SCheckBox )
+				.IsChecked_Lambda( [WeakTreeItem, WeakTree]() -> ECheckBoxState
+				{
+					FUsdPrimPreviewModelViewPtr PinnedTreeItem = WeakTreeItem.Pin();
+					TSharedPtr< SUsdStagePreviewTree > PinnedTree = WeakTree.Pin();
+					if( PinnedTreeItem && PinnedTree )
+					{
+						FUsdPrimPreviewModel& Model = PinnedTree->GetModel( PinnedTreeItem->ModelIndex );
+						return Model.bShouldImport ? ECheckBoxState::Checked : ECheckBoxState::Unchecked;
+					}
+
+					return ECheckBoxState::Undetermined;
+				})
+				.OnCheckStateChanged_Lambda([WeakTreeItem, WeakTree]( ECheckBoxState NewCheckedState )
+				{
+					FUsdPrimPreviewModelViewPtr PinnedTreeItem = WeakTreeItem.Pin();
+					TSharedPtr< SUsdStagePreviewTree > PinnedTree = WeakTree.Pin();
+					if ( PinnedTreeItem && PinnedTree )
+					{
+						if ( NewCheckedState == ECheckBoxState::Checked )
+						{
+							// Check recursively when shift clicking (tree views support recursive expansion when
+							// shift-clicking the expander arrows, so lets do that too on the checkboxes!)
+							if ( FSlateApplication::Get().GetModifierKeys().IsShiftDown() )
+							{
+								const bool bChecked = true;
+								PinnedTree->CheckItemsRecursively( { PinnedTreeItem.ToSharedRef() }, bChecked );
+							}
+							else
+							{
+								FUsdPrimPreviewModel& Model = PinnedTree->GetModel( PinnedTreeItem->ModelIndex );
+								Model.bShouldImport = true;
+
+								// If we want to import some prim, we must import all ancestors too
+								int32 ParentIndex = Model.ParentIndex;
+								while ( ParentIndex != INDEX_NONE )
+								{
+									FUsdPrimPreviewModel& ParentModel = PinnedTree->GetModel( ParentIndex );
+									ParentModel.bShouldImport = true;
+									ParentIndex = ParentModel.ParentIndex;
+								}
+							}
+						}
+						else
+						{
+							const bool bChecked = false;
+							PinnedTree->CheckItemsRecursively( { PinnedTreeItem.ToSharedRef() }, bChecked );
+						}
+					}
+				})
+			];
+	}
+};
+
+class FUsdPreviewTreeNameColumn : public FUsdTreeViewColumn
+{
+public:
+	virtual TSharedRef< SWidget > GenerateWidget(
+		const TSharedPtr< IUsdTreeViewItem > InTreeItem,
+		const TSharedPtr< ITableRow > TableRow
+	) override
+	{
+		FUsdPrimPreviewModelViewPtr TreeItem = StaticCastSharedPtr< FUsdPrimPreviewModelView >( InTreeItem );
+		TSharedPtr< SUsdStagePreviewTreeRow > Row = StaticCastSharedPtr< SUsdStagePreviewTreeRow >( TableRow );
+		TSharedPtr< SUsdStagePreviewTree > Tree = Row ? Row->GetOwnerTree() : nullptr;
+
+		if ( Tree && TreeItem )
+		{
+			FUsdPrimPreviewModel& Model = Tree->GetModel( TreeItem->ModelIndex );
+
+			return SNew( SBox )
+				.VAlign( VAlign_Center )
+				[
+					SNew( STextBlock )
+					.HighlightText( Tree->GetFilterText() )
+					.Text( Model.Name )
+				];
+		}
+
+		return SNullWidget::NullWidget;
+	}
+};
+
+class FUsdPreviewTreeTypeColumn : public FUsdTreeViewColumn
+{
+public:
+	virtual TSharedRef< SWidget > GenerateWidget(
+		const TSharedPtr< IUsdTreeViewItem > InTreeItem,
+		const TSharedPtr< ITableRow > TableRow
+	) override
+	{
+		FUsdPrimPreviewModelViewPtr TreeItem = StaticCastSharedPtr< FUsdPrimPreviewModelView >( InTreeItem );
+		TSharedPtr< SUsdStagePreviewTreeRow > Row = StaticCastSharedPtr< SUsdStagePreviewTreeRow >( TableRow );
+		TSharedPtr< SUsdStagePreviewTree > Tree = Row ? Row->GetOwnerTree() : nullptr;
+
+		if ( Tree && TreeItem )
+		{
+			FUsdPrimPreviewModel& Model = Tree->GetModel( TreeItem->ModelIndex );
+
+			return SNew( SBox )
+				.VAlign( VAlign_Center )
+				[
+					SNew( STextBlock )
+					.HighlightText( Tree->GetFilterText() )
+					.Text( Model.TypeName )
+				];
+		}
+
+		return SNullWidget::NullWidget;
+	}
+};
+
+void SUsdStagePreviewTree::Construct( const FArguments& InArgs, const UE::FUsdStage& InStage )
+{
+	SUsdTreeView::Construct( SUsdTreeView::FArguments() );
+
+	OnContextMenuOpening = FOnContextMenuOpening::CreateSP( this, &SUsdStagePreviewTree::ConstructPrimContextMenu );
+	OnExpansionChanged = FOnExpansionChanged::CreateLambda(
+		[this]( FUsdPrimPreviewModelViewRef Item, bool bIsExpanded )
+		{
+			ItemModels[ Item->ModelIndex ].bIsExpanded = bIsExpanded;
+		}
+	);
+	OnSetExpansionRecursive = FOnSetExpansionRecursive::CreateLambda(
+		[this]( FUsdPrimPreviewModelViewRef Item, bool bShouldExpand )
+		{
+			ExpandItemsRecursively( { Item }, bShouldExpand );
+		}
+	);
+
+	if ( InStage )
+	{
+		if ( UE::FUsdPrim RootPrim = InStage.GetPseudoRoot() )
+		{
+			TFunction<int32( const UE::FUsdPrim& )> RecursivelyAddModels;
+			RecursivelyAddModels = [this, &RecursivelyAddModels]( const UE::FUsdPrim& Prim ) -> int32
+			{
+				int32 NewIndex = ItemModels.Num();
+				FUsdPrimPreviewModel& NewPrimModel = ItemModels.Emplace_GetRef();
+				NewPrimModel.Name = FText::FromName( Prim.GetName() );
+				NewPrimModel.TypeName = FText::FromName( Prim.GetTypeName() );
+
+				TArray<UE::FUsdPrim> ChildPrims = Prim.GetChildren();
+				NewPrimModel.ChildIndices.Reset( ChildPrims.Num() );
+				ItemModels.Reserve( ItemModels.Num() + ChildPrims.Num() );
+				for ( const UE::FUsdPrim& ChildPrim : ChildPrims )
+				{
+					// WARNING: This will likely invalidate the NewPrimModel reference! Don't use it past this line!
+					int32 NewChildIndex = RecursivelyAddModels( ChildPrim );
+
+					ItemModels[ NewChildIndex ].ParentIndex = NewIndex;
+					ItemModels[ NewIndex ].ChildIndices.Add( NewChildIndex );
+				}
+
+				return NewIndex;
+			};
+
+			TArray<UE::FUsdPrim> TopLevelPrims = RootPrim.GetChildren();
+			ItemModels.Reset( TopLevelPrims.Num() ); // Not the best guess but better than nothing
+			for ( const UE::FUsdPrim& TopLevelPrim : TopLevelPrims )
+			{
+				RecursivelyAddModels( TopLevelPrim );
+			}
+
+			RebuildModelViews();
+		}
+	}
+}
+
+TArray<FString> SUsdStagePreviewTree::GetSelectedFullPrimPaths() const
+{
+	// This will eventually end up in a USD population mask, so we'll follow its rules, which essentially are:
+	// - Any prims in this list will have *all* of their children imported;
+	// - This list should not contain both a child and an ancestor.
+	// This basically amounts to the rule that if a prim P has all of its descendants selected, we put P on the list
+	// and no other descendant. Otherwise, we list all of the individual descendants that are enabled
+
+	const static FString Divider = TEXT( "/" );
+
+	// Since this is the default case, it may be worth it to try a quick scan to see if we're just trying to import
+	// everything. If that's the case all we need to do is return the root prim path
+	bool bShouldImportAllPrims = true;
+	for ( const FUsdPrimPreviewModel& Model : ItemModels )
+	{
+		if ( !Model.bShouldImport )
+		{
+			bShouldImportAllPrims = false;
+			break;
+		}
+	}
+	if ( bShouldImportAllPrims )
+	{
+		return { Divider };
+	}
+
+	TFunction<void( const TArray<int32>&, FString, TArray<FString>&, bool& )> RecursivelyCollectPaths;
+	RecursivelyCollectPaths = [this, &RecursivelyCollectPaths](
+		const TArray<int32>& Indices,
+		FString PathSoFar,
+		TArray<FString>& OutChildPaths,
+		bool& bOutAllDescendantsImported
+	)
+	{
+		for ( int32 Index : Indices )
+		{
+			const FUsdPrimPreviewModel& Model = ItemModels[ Index ];
+			if ( !Model.bShouldImport )
+			{
+				bOutAllDescendantsImported = false;
+				continue;
+			}
+
+			FString PrimFullPath = PathSoFar + Model.Name.ToString();
+
+			TArray<FString> ChildPaths;
+			bool bAllDescendantsImported = true;
+			RecursivelyCollectPaths(
+				Model.ChildIndices,
+				PrimFullPath + Divider,
+				ChildPaths,
+				bAllDescendantsImported
+			);
+
+			// If all of the Prim's descendents were imported, fully discard ChildPaths and just add Prim directly
+			// to the list.
+			// Alternatively, if we imported *nothing else* but we made it this far (i.e. we want to import this prim)
+			// then we actually want to import this intermediate prim directly, so also add it directly to the list.
+			if ( bAllDescendantsImported || ChildPaths.Num() == 0 )
+			{
+				OutChildPaths.Add( MoveTemp( PrimFullPath ) );
+			}
+			else
+			{
+				OutChildPaths.Append( MoveTemp( ChildPaths ) );
+			}
+
+			bOutAllDescendantsImported &= bAllDescendantsImported;
+		}
+	};
+
+	TArray<int32> RootItemModels;
+	RootItemModels.Reserve( RootItems.Num() );
+	for ( const FUsdPrimPreviewModelViewRef& RootItem : RootItems )
+	{
+		RootItemModels.Add( RootItem->ModelIndex );
+	}
+
+	TArray<FString> Result;
+	bool bUnusedDummy = false; // At this point we know we won't import *everything* anyway, don't bother checking
+	RecursivelyCollectPaths( RootItemModels, Divider, Result, bUnusedDummy );
+	return Result;
+}
+
+void SUsdStagePreviewTree::SetFilterText( const FText& NewText )
+{
+	CurrentFilterText = NewText;
+	RebuildModelViews();
+}
+
+TSharedRef< ITableRow > SUsdStagePreviewTree::OnGenerateRow(
+	FUsdPrimPreviewModelViewRef InDisplayNode,
+	const TSharedRef<STableViewBase>& OwnerTable
+)
+{
+	return SNew( SUsdStagePreviewTreeRow, InDisplayNode, OwnerTable, SharedData );
+}
+
+void SUsdStagePreviewTree::OnGetChildren(
+	FUsdPrimPreviewModelViewRef InParent,
+	TArray< FUsdPrimPreviewModelViewRef >& OutChildren
+) const
+{
+	OutChildren = InParent->Children;
+}
+
+void SUsdStagePreviewTree::SetupColumns()
+{
+	HeaderRowWidget->ClearColumns();
+
+	// Import checkbox column with a checkbox on the header itself (so we can't just use AddColumn)
+	{
+		const static FName ImportColumnKey = TEXT( "ImportColumn" );
+
+		SharedData->Columns.Add( ImportColumnKey, MakeShared< FUsdPreviewTreeImportColumn >() );
+
+		HeaderRowWidget->AddColumn(
+			SHeaderRow::Column( ImportColumnKey )
+			.FixedWidth( UE::UsdStagePreviewTree::Private::CheckBoxColumnWidth )
+			.HeaderContentPadding( FMargin( 0 ) )
+			[
+				SNew( SBox )
+				.VAlign( VAlign_Center )
+				.HAlign( HAlign_Center )
+				.WidthOverride( UE::UsdStagePreviewTree::Private::CheckBoxColumnWidth )
+				[
+					SNew( SCheckBox )
+					.ToolTipText( LOCTEXT( "HeaderCheckboxToolTip", "Toggle all currently displayed items" ) )
+					.IsChecked_Lambda( [this]() -> ECheckBoxState
+					{
+						return UE::UsdStagePreviewTree::Private::AreAllViewsCheckedRecursive( *this, RootItems )
+							? ECheckBoxState::Checked
+							: ECheckBoxState::Unchecked;
+					})
+					.OnCheckStateChanged_Lambda( [this]( ECheckBoxState NewState )
+					{
+						const bool bChecked = NewState == ECheckBoxState::Checked;
+						CheckItemsRecursively( RootItems, bChecked );
+					})
+				]
+			]
+		);
+	}
+
+	// Prim name column
+	{
+		TSharedRef< FUsdPreviewTreeNameColumn > PrimNameColumn = MakeShared< FUsdPreviewTreeNameColumn >();
+		PrimNameColumn->bIsMainColumn = true;
+
+		SHeaderRow::FColumn::FArguments Arguments;
+		Arguments.FillWidth( 70.f );
+
+		AddColumn( TEXT("Prim"), LOCTEXT( "Prim", "Prim" ), PrimNameColumn, Arguments );
+	}
+
+	// Prim type column
+	{
+		SHeaderRow::FColumn::FArguments Arguments;
+		Arguments.FillWidth( 20.f );
+		AddColumn( TEXT( "Type" ), LOCTEXT( "Type", "Type" ), MakeShared< FUsdPreviewTreeTypeColumn >(), Arguments );
+	}
+}
+
+TArray<FUsdPrimPreviewModelViewRef> SUsdStagePreviewTree::GetAncestorSelectedViews()
+{
+	TArray< FUsdPrimPreviewModelViewRef > SelectedViews = GetSelectedItems();
+
+	TSet<int32> SelectedIndices;
+	SelectedIndices.Reserve( SelectedViews.Num() );
+	for ( const FUsdPrimPreviewModelViewRef& SelectedItem : SelectedViews )
+	{
+		SelectedIndices.Add( SelectedItem->ModelIndex );
+	}
+
+	TArray< FUsdPrimPreviewModelViewRef > AncestorSelectedViews;
+	AncestorSelectedViews.Reserve( SelectedViews.Num() );
+
+	for ( const FUsdPrimPreviewModelViewRef& SelectedView : SelectedViews )
+	{
+		bool bIsAncestor = true;
+
+		int32 ParentIndex = ItemModels[ SelectedView->ModelIndex ].ParentIndex;
+		while ( ParentIndex != INDEX_NONE )
+		{
+			if ( SelectedIndices.Contains( ParentIndex ) )
+			{
+				bIsAncestor = false;
+				break;
+			}
+
+			ParentIndex = ItemModels[ ParentIndex ].ParentIndex;
+		}
+
+		if ( bIsAncestor )
+		{
+			AncestorSelectedViews.Add( SelectedView );
+		}
+	}
+
+	return AncestorSelectedViews;
+}
+
+TSharedPtr< SWidget > SUsdStagePreviewTree::ConstructPrimContextMenu()
+{
+	const bool bInShouldCloseWindowAfterMenuSelection = true;
+	const TSharedPtr< const FUICommandList > CommandList = nullptr;
+	FMenuBuilder MenuBuilder{ bInShouldCloseWindowAfterMenuSelection, CommandList };
+
+	MenuBuilder.BeginSection( "Import", LOCTEXT( "ImportText", "Import" ) );
+	{
+		MenuBuilder.AddMenuEntry(
+			LOCTEXT( "CheckDescendants_Text", "Check Recursively" ),
+			LOCTEXT( "CheckDescendants_ToolTip", "Check this prim and all displayed descendants recursively" ),
+			FSlateIcon(),
+			FUIAction(
+				FExecuteAction::CreateLambda( [this]()
+				{
+					const bool bShouldSelect = true;
+					CheckItemsRecursively( GetAncestorSelectedViews(), bShouldSelect );
+				}),
+				FCanExecuteAction()
+			),
+			NAME_None,
+			EUserInterfaceActionType::Button
+		);
+
+		MenuBuilder.AddMenuEntry(
+			LOCTEXT( "UncheckDescendants_Text", "Uncheck Recursively" ),
+			LOCTEXT( "UncheckDescendants_ToolTip", "Uncheck this prim and all displayed descendants recursively" ),
+			FSlateIcon(),
+			FUIAction(
+				FExecuteAction::CreateLambda( [this]()
+				{
+					const bool bShouldSelect = false;
+					CheckItemsRecursively( GetAncestorSelectedViews(), bShouldSelect );
+				}),
+				FCanExecuteAction()
+			),
+			NAME_None,
+			EUserInterfaceActionType::Button
+		);
+	}
+	MenuBuilder.EndSection();
+
+	MenuBuilder.BeginSection( "Expansion", LOCTEXT( "ExpansionText", "Expansion" ) );
+	{
+		MenuBuilder.AddMenuEntry(
+			LOCTEXT( "ExpandDescendants_Text", "Expand Recursively" ),
+			LOCTEXT( "ExpandDescendants_ToolTip", "Expand this prim and all displayed descendants recursively" ),
+			FSlateIcon(),
+			FUIAction(
+				FExecuteAction::CreateLambda( [this]()
+				{
+					const bool bShouldExpand = true;
+					ExpandItemsRecursively( GetAncestorSelectedViews(), bShouldExpand );
+				}),
+				FCanExecuteAction()
+			),
+			NAME_None,
+			EUserInterfaceActionType::Button
+		);
+
+		MenuBuilder.AddMenuEntry(
+			LOCTEXT( "CollapseDescendants_Text", "Collapse Recursively" ),
+			LOCTEXT( "CollapseDescendants_ToolTip", "Collapse this prim and all displayed descendants recursively" ),
+			FSlateIcon(),
+			FUIAction(
+				FExecuteAction::CreateLambda( [this]()
+				{
+					const bool bShouldExpand = false;
+					ExpandItemsRecursively( GetAncestorSelectedViews(), bShouldExpand );
+				}),
+				FCanExecuteAction()
+			),
+			NAME_None,
+			EUserInterfaceActionType::Button
+		);
+	}
+	MenuBuilder.EndSection();
+
+	return MenuBuilder.MakeWidget();
+}
+
+void SUsdStagePreviewTree::RebuildModelViews()
+{
+	const FString CurrentFilterString = CurrentFilterText.ToString().ToLower();
+
+	// Update our models caching whether they pass the filter or not
+	for ( FUsdPrimPreviewModel& Model : ItemModels )
+	{
+		Model.bPassesFilter = CurrentFilterString.IsEmpty()
+			|| Model.Name.ToString().ToLower().Contains( CurrentFilterString )
+			|| Model.TypeName.ToString().ToLower().Contains( CurrentFilterString );
+
+		// If this model passes the filter, we need to ensure all of its ancestors are marked as passing too
+		if ( Model.bPassesFilter )
+		{
+			int32 ParentIndex = Model.ParentIndex;
+			while ( ParentIndex != INDEX_NONE )
+			{
+				FUsdPrimPreviewModel& ParentModel = ItemModels[ ParentIndex ];
+				if ( ParentModel.bPassesFilter )
+				{
+					break;
+				}
+				ParentModel.bPassesFilter = true;
+				ParentIndex = ParentModel.ParentIndex;
+			}
+		}
+	}
+
+	// Rebuild and keep track of this mapping as its very cheap to do right now
+	ItemModelsToViews.Reset();
+	ItemModelsToViews.Reserve( ItemModels.Num() );
+
+	// Create a view for each model *that passes the filter*, having nullptr instead if they don't
+	TArray< FUsdPrimPreviewModelViewPtr > LinearViews;
+	LinearViews.SetNum( ItemModels.Num() );
+	for ( int32 Index = 0; Index < ItemModels.Num(); ++Index )
+	{
+		const FUsdPrimPreviewModel& Model = ItemModels[ Index ];
+
+		if ( Model.bPassesFilter )
+		{
+			FUsdPrimPreviewModelViewRef View = MakeShared<FUsdPrimPreviewModelView>();
+			View->ModelIndex = Index;
+
+			ItemModelsToViews.Add( Index, View );
+			LinearViews[ Index ] = MoveTemp( View );
+		}
+	}
+
+	// Now that we have views created for each model we want, connect parent/child based on model indices
+	RootItems.Reset();
+	SparseItemInfos.Reset();
+	for ( int32 Index = 0; Index < ItemModels.Num(); ++Index )
+	{
+		const FUsdPrimPreviewModel& Model = ItemModels[ Index ];
+		FUsdPrimPreviewModelViewPtr& View = LinearViews[ Index ];
+		if ( !View )
+		{
+			// Model didn't pass the filter
+			continue;
+		}
+
+		// Set the view's parent
+		if ( Model.ParentIndex == INDEX_NONE )
+		{
+			RootItems.Add( View.ToSharedRef() );
+		}
+		else
+		{
+			View->Parent = LinearViews[ Model.ParentIndex ];
+		}
+
+		// Add child models it they pass the filter, also already tracking if they're expanded or not
+		bool bHasExpandedChildren = false;
+		View->Children.Reserve( Model.ChildIndices.Num() );
+		for ( int32 ChildIndex : Model.ChildIndices )
+		{
+			const FUsdPrimPreviewModel& ChildModel = ItemModels[ ChildIndex ];
+			if ( ChildModel.bPassesFilter )
+			{
+				View->Children.Add( LinearViews[ ChildIndex ].ToSharedRef() );
+				bHasExpandedChildren |= ChildModel.bIsExpanded;
+			}
+		}
+
+		const bool bShouldExpand = true;
+		SparseItemInfos.Add( View.ToSharedRef(), FSparseItemInfo( Model.bIsExpanded, bHasExpandedChildren ) );
+	}
+
+	RequestTreeRefresh();
+}
+
+void SUsdStagePreviewTree::CheckItemsRecursively( const TArray<FUsdPrimPreviewModelViewRef>& Items, bool bCheck )
+{
+	TFunction< void(const TArray<FUsdPrimPreviewModelViewRef>&)> CheckViewsRecursively;
+	CheckViewsRecursively = [this, &CheckViewsRecursively]( const TArray<FUsdPrimPreviewModelViewRef>& Views )
+	{
+		for ( const FUsdPrimPreviewModelViewRef& View : Views )
+		{
+			ItemModels[ View->ModelIndex ].bShouldImport = true;
+			CheckViewsRecursively( View->Children );
+		}
+	};
+
+	TArray<FString> UncheckedInvisiblePrims;
+	TFunction< void( const TArray<int32>& )> UncheckModelsRecursively;
+	UncheckModelsRecursively = [this, &UncheckModelsRecursively, &UncheckedInvisiblePrims](
+		const TArray<int32>& ModelIndices
+	)
+	{
+		for ( int32 ModelIndex : ModelIndices )
+		{
+			FUsdPrimPreviewModel& Model = ItemModels[ ModelIndex ];
+
+			if ( Model.bShouldImport && !ItemModelsToViews.Contains( ModelIndex ) )
+			{
+				UncheckedInvisiblePrims.Add( Model.Name.ToString() );
+			}
+
+			Model.bShouldImport = false;
+			UncheckModelsRecursively( Model.ChildIndices );
+		}
+	};
+
+	// Attention here: For consistency, if we are unchecking we must guarantee that no descendant models are ever left
+	// checked, whether we have views for them or not. On the other hand if we are just checking something, it is fine
+	// if not all of the model descendants are checked in the process, so let's only check what is shown on the screen
+	if ( bCheck )
+	{
+		CheckViewsRecursively( Items );
+
+		// If we want to check a view, all of its ancestor *models* must also be checked, regardless of what is shown
+		for ( const FUsdPrimPreviewModelViewRef& View : Items )
+		{
+			int32 ParentIndex = ItemModels[ View->ModelIndex ].ParentIndex;
+			while ( ParentIndex != INDEX_NONE )
+			{
+				FUsdPrimPreviewModel& ParentModel = ItemModels[ ParentIndex ];
+				ParentModel.bShouldImport = true;
+				ParentIndex = ParentModel.ParentIndex;
+			}
+		}
+	}
+	else
+	{
+		TArray<int32> ModelIndices;
+		ModelIndices.Reserve( Items.Num() );
+		for ( const FUsdPrimPreviewModelViewRef& Item : Items )
+		{
+			ModelIndices.Add( Item->ModelIndex );
+		}
+		UncheckModelsRecursively( ModelIndices );
+
+		// Notify user in case we unchecked some prims that currently are not being shown
+		if( UncheckedInvisiblePrims.Num() > 0 )
+		{
+			const int32 MaxNumPrimsToList = 15;
+			const static FString Delimiter = TEXT( ", " );
+			FString PrimListString;
+			for ( int32 Index = 0; Index < FMath::Min( UncheckedInvisiblePrims.Num(), MaxNumPrimsToList ); ++Index )
+			{
+				PrimListString += UncheckedInvisiblePrims[ Index ] + Delimiter;
+			}
+			if ( UncheckedInvisiblePrims.Num() > MaxNumPrimsToList )
+			{
+				PrimListString += TEXT( "..." );
+			}
+			else
+			{
+				PrimListString.RemoveFromEnd( Delimiter );
+			}
+
+			FNotificationInfo Toast( LOCTEXT( "UncheckedInvisibleTitle", "USD Prims to import" ) );
+			Toast.SubText = FText::Format(
+				UncheckedInvisiblePrims.Num() == 1
+					? LOCTEXT( "UncheckedInvisibleSingular", "This prim is not currently shown, but was also unchecked:\n\n{0}" )
+					: LOCTEXT( "UncheckedInvisiblePlural", "These prims are not currently shown, but were also unchecked:\n\n{0}" ),
+				FText::FromString( PrimListString )
+			);
+			Toast.Image = FCoreStyle::Get().GetBrush( TEXT( "MessageLog.Warning" ) );
+			Toast.bUseLargeFont = false;
+			Toast.bFireAndForget = true;
+			Toast.FadeOutDuration = 0.5f;
+			Toast.ExpireDuration = 4.5f;
+			Toast.bUseThrobber = true;
+			Toast.bUseSuccessFailIcons = false;
+			FSlateNotificationManager::Get().AddNotification( Toast );
+		}
+	}
+}
+
+void SUsdStagePreviewTree::ExpandItemsRecursively(
+	const TArray<FUsdPrimPreviewModelViewRef>& Items,
+	bool bExpand
+)
+{
+	TFunction< void( const TArray<FUsdPrimPreviewModelViewRef>& )> ExpandRecursively;
+	ExpandRecursively = [this, &ExpandRecursively, bExpand]( const TArray<FUsdPrimPreviewModelViewRef>& Views )
+	{
+		for ( const FUsdPrimPreviewModelViewRef& View : Views )
+		{
+			FUsdPrimPreviewModel& Model = ItemModels[ View->ModelIndex ];
+			Model.bIsExpanded = bExpand;
+
+			// We're going to expand recursively anyway, so we know it will have expanded children
+			const bool bHasExpandedChildren = bExpand && View->Children.Num() > 0;
+			const bool bIsExpanded = bExpand;
+			SparseItemInfos.Add( View, FSparseItemInfo( bIsExpanded, bHasExpandedChildren ) );
+
+			ExpandRecursively( View->Children );
+		}
+	};
+	ExpandRecursively( Items );
+
+	RequestTreeRefresh();
+}
+
+void SUsdStagePreviewTreeRow::Construct(
+	const FArguments& InArgs,
+	FUsdPrimPreviewModelViewRef InTreeItem,
+	const TSharedRef<STableViewBase>& OwnerTable,
+	TSharedPtr<FSharedUsdTreeData> InSharedData
+)
+{
+	SUsdTreeRow< FUsdPrimPreviewModelViewRef >::Construct(
+		typename SUsdTreeRow< FUsdPrimPreviewModelViewRef >::FArguments(),
+		InTreeItem,
+		OwnerTable,
+		InSharedData
+	);
+}
+
+TSharedPtr<SUsdStagePreviewTree> SUsdStagePreviewTreeRow::GetOwnerTree() const
+{
+	return StaticCastSharedPtr< SUsdStagePreviewTree>( OwnerTablePtr.Pin() );
+}
+
+#undef LOCTEXT_NAMESPACE
\ No newline at end of file
diff --git a/Engine/Plugins/Importers/USDImporter/Source/USDStageImporter/Public/USDOptionsWindow.h b/Engine/Plugins/Importers/USDImporter/Source/USDStageImporter/Public/USDOptionsWindow.h
--- a/Engine/Plugins/Importers/USDImporter/Source/USDStageImporter/Public/USDOptionsWindow.h
+++ b/Engine/Plugins/Importers/USDImporter/Source/USDStageImporter/Public/USDOptionsWindow.h
@@ -5,6 +5,8 @@
 #include "Widgets/SCompoundWidget.h"
 #include "Widgets/SWindow.h"
 
+#include "UsdWrappers/ForwardDeclarations.h"
+
 /**
  * Slate window used to show import/export options for the USDImporter plugin
  */
@@ -15,11 +17,21 @@ public:
 		SLATE_ARGUMENT( UObject*, OptionsObject )
 		SLATE_ARGUMENT( TSharedPtr<SWindow>, WidgetWindow )
 		SLATE_ARGUMENT( FText, AcceptText )
+		SLATE_ARGUMENT( const UE::FUsdStage*, Stage )
 	SLATE_END_ARGS()
 
 public:
-	static bool ShowImportExportOptions( UObject& OptionsObject, bool bIsImport );
-	static bool ShowOptions( UObject& OptionsObject, const FText& WindowTitle, const FText& AcceptText );
+	// Show the options window with the given text
+	static bool ShowOptions(
+		UObject& OptionsObject,
+		const FText& WindowTitle,
+		const FText& AcceptText,
+		const UE::FUsdStage* Stage = nullptr
+	);
+
+	// Shortcut functions that show the standard import/export text
+	static bool ShowImportOptions( UObject& OptionsObject, const UE::FUsdStage* StageToImport );
+	static bool ShowExportOptions( UObject& OptionsObject );
 
 	void Construct( const FArguments& InArgs );
 	virtual bool SupportsKeyboardFocus() const override;
@@ -29,10 +41,14 @@ public:
 
 	virtual FReply OnKeyDown( const FGeometry& MyGeometry, const FKeyEvent& InKeyEvent ) override;
 	bool UserAccepted() const;
+	TArray<FString> GetSelectedFullPrimPaths() const;
 
 private:
 	UObject* OptionsObject;
+
 	TWeakPtr< SWindow > Window;
+	TSharedPtr< class SUsdStagePreviewTree > StagePreviewTree;
+
 	FText AcceptText;
 	bool bAccepted;
 };
diff --git a/Engine/Plugins/Importers/USDImporter/Source/USDStageImporter/Public/USDStageImportOptions.h b/Engine/Plugins/Importers/USDImporter/Source/USDStageImporter/Public/USDStageImportOptions.h
--- a/Engine/Plugins/Importers/USDImporter/Source/USDStageImporter/Public/USDStageImportOptions.h
+++ b/Engine/Plugins/Importers/USDImporter/Source/USDStageImporter/Public/USDStageImportOptions.h
@@ -64,7 +64,13 @@ public:
 	UPROPERTY(BlueprintReadWrite, config, EditAnywhere, Category = "DataToImport", meta = (DisplayName = "Materials & Textures"))
 	bool bImportMaterials;
 
-
+	/**
+	 * List of paths of prims to import (e.g. ["/Root/MyBox", "/Root/OtherPrim"]).
+	 * Importing a prim will import its entire subtree.
+	 * If this list contains the root prim path the entire stage will be imported (default value).
+	 */
+	UPROPERTY( BlueprintReadWrite, EditAnywhere, Category = "Prims to Import" )
+	TArray<FString> PrimsToImport = TArray<FString>{ TEXT( "/" ) };
 
 	/** Only import prims with these specific purposes from the USD file */
 	UPROPERTY(BlueprintReadWrite, config, EditAnywhere, Category= "USD options", meta = (Bitmask, BitmaskEnum="/Script/UnrealUSDWrapper.EUsdPurpose"))
diff --git a/Engine/Plugins/Importers/USDImporter/Source/USDStageImporter/Public/USDStageImportOptionsCustomization.h b/Engine/Plugins/Importers/USDImporter/Source/USDStageImporter/Public/USDStageImportOptionsCustomization.h
deleted file mode 100644
diff --git a/Engine/Plugins/Importers/USDImporter/Source/USDStageImporter/Public/Widgets/IUSDTreeViewItem.h b/Engine/Plugins/Importers/USDImporter/Source/USDStageImporter/Public/Widgets/IUSDTreeViewItem.h
new file mode 100644
--- /dev/null
+++ b/Engine/Plugins/Importers/USDImporter/Source/USDStageImporter/Public/Widgets/IUSDTreeViewItem.h
@@ -0,0 +1,12 @@
+// Copyright Epic Games, Inc. All Rights Reserved.
+
+#pragma once
+
+#include "HAL/Platform.h"
+#include "Templates/SharedPointer.h"
+
+class USDSTAGEIMPORTER_API IUsdTreeViewItem : public TSharedFromThis< IUsdTreeViewItem >
+{
+public:
+	virtual ~IUsdTreeViewItem() = default;
+};
diff --git a/Engine/Plugins/Importers/USDImporter/Source/USDStageImporter/Public/Widgets/SUSDStagePreviewTree.h b/Engine/Plugins/Importers/USDImporter/Source/USDStageImporter/Public/Widgets/SUSDStagePreviewTree.h
new file mode 100644
--- /dev/null
+++ b/Engine/Plugins/Importers/USDImporter/Source/USDStageImporter/Public/Widgets/SUSDStagePreviewTree.h
@@ -0,0 +1,112 @@
+// Copyright Epic Games, Inc. All Rights Reserved.
+
+#pragma once
+
+#include "UsdWrappers/ForwardDeclarations.h"
+#include "Widgets/IUSDTreeViewItem.h"
+#include "Widgets/SUSDTreeView.h"
+
+#include "Templates/SharedPointer.h"
+
+using FUsdPrimPreviewModelViewRef = TSharedRef< struct FUsdPrimPreviewModelView >;
+using FUsdPrimPreviewModelViewPtr = TSharedPtr< struct FUsdPrimPreviewModelView >;
+using FUsdPrimPreviewModelViewWeak = TWeakPtr< struct FUsdPrimPreviewModelView >;
+
+struct FUsdPrimPreviewModel
+{
+	// Prim's name
+	FText Name;
+
+	// Prim's type name (like "Mesh" or "Camera")
+	FText TypeName;
+
+	// Whether this prim should be imported or not
+	bool bShouldImport = true;
+
+	// Only used to make filtering items with SUsdStagePreviewTree::CurrentFilterText a bit faster: This is not
+	// actually used to display anything
+	bool bPassesFilter = true;
+
+	// Our source of truth for expansion state. Whenever expansion changes on the views we write to this,
+	// and whenever we rebuild our views we write this to the tree's sparse infos about its items
+	bool bIsExpanded = true;
+
+	int32 ParentIndex = INDEX_NONE;
+	TArray<int32> ChildIndices;
+};
+
+struct FUsdPrimPreviewModelView : public IUsdTreeViewItem
+{
+	int32 ModelIndex = INDEX_NONE;
+
+	FUsdPrimPreviewModelViewWeak Parent;
+    TArray< FUsdPrimPreviewModelViewRef > Children;
+};
+
+class USDSTAGEIMPORTER_API SUsdStagePreviewTree : public SUsdTreeView< FUsdPrimPreviewModelViewRef >
+{
+public:
+	SLATE_BEGIN_ARGS( SUsdStagePreviewTree ) {}
+	SLATE_END_ARGS()
+
+	void Construct( const FArguments& InArgs, const UE::FUsdStage& InStage );
+
+    TArray<FString> GetSelectedFullPrimPaths() const;
+
+	void SetFilterText( const FText& NewText );
+	const FText& GetFilterText() const { return CurrentFilterText; }
+
+	FUsdPrimPreviewModel& GetModel( int32 ModelIndex ) { return ItemModels[ ModelIndex ]; }
+	const FUsdPrimPreviewModel& GetModel( int32 ModelIndex ) const { return ItemModels[ ModelIndex ]; }
+
+	void CheckItemsRecursively( const TArray<FUsdPrimPreviewModelViewRef>& Items, bool bCheck );
+	void ExpandItemsRecursively( const TArray<FUsdPrimPreviewModelViewRef>& Items, bool bExpand );
+
+private:
+    // Begin SUsdTreeView interface
+	virtual TSharedRef< ITableRow > OnGenerateRow(
+		FUsdPrimPreviewModelViewRef InDisplayNode,
+		const TSharedRef< STableViewBase >& OwnerTable
+	) override;
+
+	virtual void OnGetChildren(
+		FUsdPrimPreviewModelViewRef InParent,
+		TArray< FUsdPrimPreviewModelViewRef >& OutChildren
+	) const override;
+
+	virtual void SetupColumns() override;
+    // End SUsdTreeView interface
+
+	// Returns all selected model views, excluding child views of any other selected view.
+	// i.e. if both a child and an ancestor view are selected, the returned array will contain only the ancestor.
+	TArray<FUsdPrimPreviewModelViewRef> GetAncestorSelectedViews();
+
+	TSharedPtr< SWidget > ConstructPrimContextMenu();
+
+	// Rebuilds our RootItems array of FUsdPrimPreviewModelViewRefs based on CurrentFilterText and the ItemModels that
+	// pass the filter
+	void RebuildModelViews();
+
+private:
+	FText CurrentFilterText;
+
+	TArray<FUsdPrimPreviewModel> ItemModels;
+	TMap<int32, FUsdPrimPreviewModelViewRef> ItemModelsToViews;
+};
+
+// Custom row so that we can fetch the owning SUsdStagePreviewTree from the columns
+class USDSTAGEIMPORTER_API SUsdStagePreviewTreeRow : public SUsdTreeRow< FUsdPrimPreviewModelViewRef >
+{
+public:
+	SLATE_BEGIN_ARGS( SUsdStagePreviewTreeRow ) {}
+	SLATE_END_ARGS()
+
+	void Construct(
+		const FArguments& InArgs,
+		FUsdPrimPreviewModelViewRef InTreeItem,
+		const TSharedRef< STableViewBase >& OwnerTable,
+		TSharedPtr< FSharedUsdTreeData > InSharedData
+	);
+
+	TSharedPtr< SUsdStagePreviewTree > GetOwnerTree() const;
+};
\ No newline at end of file
diff --git a/Engine/Plugins/Importers/USDImporter/Source/USDStageImporter/Public/Widgets/SUSDTreeView.h b/Engine/Plugins/Importers/USDImporter/Source/USDStageImporter/Public/Widgets/SUSDTreeView.h
new file mode 100644
--- /dev/null
+++ b/Engine/Plugins/Importers/USDImporter/Source/USDStageImporter/Public/Widgets/SUSDTreeView.h
@@ -0,0 +1,165 @@
+// Copyright Epic Games, Inc. All Rights Reserved.
+
+#pragma once
+
+#include "Templates/SharedPointer.h"
+#include "Widgets/SWidget.h"
+#include "Widgets/Views/SHeaderRow.h"
+#include "Widgets/Views/STableRow.h"
+#include "Widgets/Views/STreeView.h"
+
+class IUsdTreeViewItem;
+
+class FUsdTreeViewColumn
+{
+public:
+	virtual ~FUsdTreeViewColumn() = default;
+
+	virtual TSharedRef< SWidget > GenerateWidget( const TSharedPtr< IUsdTreeViewItem > TreeItem, const TSharedPtr< ITableRow > TableRow ) = 0;
+
+public:
+	bool bIsMainColumn = false;
+};
+
+struct FSharedUsdTreeData
+{
+	TMap< FName, TSharedRef< FUsdTreeViewColumn > > Columns;
+};
+
+template< typename ItemType >
+class SUsdTreeView : public STreeView< ItemType >
+{
+public:
+	SLATE_BEGIN_ARGS( SUsdTreeView ) {}
+	SLATE_END_ARGS()
+
+	void Construct( const FArguments& InArgs );
+
+	void AddColumn( FName ColumnKey, FText ColumnLabel, TSharedRef< FUsdTreeViewColumn > Column, SHeaderRow::FColumn::FArguments ColumnArguments = SHeaderRow::FColumn::FArguments() );
+	virtual void SetupColumns() = 0;
+
+	virtual TSharedRef< ITableRow > OnGenerateRow( ItemType InDisplayNode, const TSharedRef< STableViewBase >& OwnerTable ) = 0;
+	virtual void OnGetChildren( ItemType InParent, TArray< ItemType >& OutChildren ) const = 0;
+
+	virtual void OnTreeItemScrolledIntoView( ItemType TreeItem, const TSharedPtr< ITableRow >& Widget ) {}
+
+protected:
+	TSharedPtr< SHeaderRow > HeaderRowWidget;
+	TArray< ItemType > RootItems;
+
+	TSharedPtr< FSharedUsdTreeData > SharedData;
+};
+
+template< typename ItemType >
+inline void SUsdTreeView< ItemType >::AddColumn( FName ColumnKey, FText ColumnLabel, TSharedRef< FUsdTreeViewColumn > Column, SHeaderRow::FColumn::FArguments ColumnArguments )
+{
+	ColumnArguments.ColumnId( ColumnKey );
+	ColumnArguments.DefaultLabel( ColumnLabel );
+
+	HeaderRowWidget->AddColumn( ColumnArguments );
+
+	SharedData->Columns.Add( ColumnKey, MoveTemp( Column ) );
+}
+
+template< typename ItemType >
+inline void SUsdTreeView< ItemType >::Construct( const FArguments& InArgs )
+{
+	SharedData = MakeShared< FSharedUsdTreeData >();
+
+	HeaderRowWidget =
+			SNew( SHeaderRow )
+				.Visibility( EVisibility::Visible );
+
+	SetupColumns();
+
+	STreeView< ItemType >::Construct
+	(
+		typename STreeView< ItemType >::FArguments()
+		.TreeItemsSource( &RootItems )
+		.OnGenerateRow( this, &SUsdTreeView< ItemType >::OnGenerateRow )
+		.OnGetChildren( this, &SUsdTreeView< ItemType >::OnGetChildren )
+		.OnItemScrolledIntoView( this, &SUsdTreeView< ItemType >::OnTreeItemScrolledIntoView )
+		.HeaderRow( HeaderRowWidget )
+	);
+}
+
+template< typename ItemType >
+class SUsdTreeRow : public SMultiColumnTableRow< ItemType >
+{
+public:
+	typedef typename STableRow<ItemType>::FOnCanAcceptDrop FUsdOnCanAcceptDrop;
+	typedef typename STableRow<ItemType>::FOnAcceptDrop FUsdOnAcceptDrop;
+
+	SLATE_BEGIN_ARGS( SUsdTreeRow ) {}
+		SLATE_EVENT( FUsdOnCanAcceptDrop, OnCanAcceptDrop )
+		SLATE_EVENT( FUsdOnAcceptDrop, OnAcceptDrop )
+		SLATE_EVENT( FOnDragDetected, OnDragDetected )
+		SLATE_EVENT( FOnTableRowDragEnter, OnDragEnter )
+		SLATE_EVENT( FOnTableRowDragLeave, OnDragLeave )
+		SLATE_EVENT( FOnTableRowDrop, OnDrop )
+	SLATE_END_ARGS()
+
+	void Construct( const FArguments& InArgs, ItemType InTreeItem, const TSharedRef< STableViewBase >& OwnerTable, TSharedPtr< FSharedUsdTreeData > InSharedData )
+	{
+		TreeItem = InTreeItem;
+		SharedData = InSharedData;
+
+		SMultiColumnTableRow< ItemType >::Construct(
+			typename SMultiColumnTableRow< ItemType >::FArguments()
+				.OnCanAcceptDrop( InArgs._OnCanAcceptDrop )
+				.OnAcceptDrop( InArgs._OnAcceptDrop )
+				.OnDragDetected( InArgs._OnDragDetected )
+				.OnDragEnter( InArgs._OnDragEnter )
+				.OnDragLeave( InArgs._OnDragLeave )
+				.OnDrop( InArgs._OnDrop ),
+			OwnerTable
+		);
+	}
+
+	/** Overridden from SMultiColumnTableRow.  Generates a widget for this column of the tree row. */
+	virtual TSharedRef< SWidget > GenerateWidgetForColumn( const FName& ColumnName ) override;
+
+private:
+	TSharedPtr< typename ItemType::ElementType > TreeItem;
+	TSharedPtr< FSharedUsdTreeData > SharedData;
+};
+
+template< typename ItemType >
+inline TSharedRef< SWidget > SUsdTreeRow< ItemType >::GenerateWidgetForColumn( const FName& ColumnName )
+{
+	TSharedRef< SWidget > ColumnWidget = SNullWidget::NullWidget;
+
+	if ( SharedData && SharedData->Columns.Contains( ColumnName ) )
+	{
+		ColumnWidget = SharedData->Columns[ ColumnName ]->GenerateWidget(
+			StaticCastSharedPtr< IUsdTreeViewItem >( TSharedPtr< typename ItemType::ElementType >( TreeItem ) ),
+			this->SharedThis(this)
+		);
+	}
+
+	if ( SharedData->Columns[ ColumnName ]->bIsMainColumn )
+	{
+		return SNew( SHorizontalBox )
+			+SHorizontalBox::Slot()
+			.AutoWidth()
+			.Padding(6, 0, 0, 0)
+			[
+				SNew( SExpanderArrow, this->SharedThis(this) ).IndentAmount(12)
+			]
+
+			+SHorizontalBox::Slot()
+			.FillWidth(1.0f)
+			[
+				ColumnWidget
+			];
+	}
+	else
+	{
+		return SNew( SHorizontalBox )
+			+SHorizontalBox::Slot()
+			.AutoWidth()
+			[
+				ColumnWidget
+			];
+	}
+}
diff --git a/Engine/Plugins/Importers/USDImporter/Source/USDStageImporter/USDStageImporter.Build.cs b/Engine/Plugins/Importers/USDImporter/Source/USDStageImporter/USDStageImporter.Build.cs
--- a/Engine/Plugins/Importers/USDImporter/Source/USDStageImporter/USDStageImporter.Build.cs
+++ b/Engine/Plugins/Importers/USDImporter/Source/USDStageImporter/USDStageImporter.Build.cs
@@ -12,6 +12,7 @@ namespace UnrealBuildTool.Rules
 			PrivateDependencyModuleNames.AddRange(
 				new string[]
 				{
+					"ApplicationCore",
 					"Core",
 					"CoreUObject",
 					"EditorFramework",
diff --git a/Engine/Plugins/Importers/USDImporter/Source/UnrealUSDWrapper/Private/UnrealUSDWrapper.cpp b/Engine/Plugins/Importers/USDImporter/Source/UnrealUSDWrapper/Private/UnrealUSDWrapper.cpp
--- a/Engine/Plugins/Importers/USDImporter/Source/UnrealUSDWrapper/Private/UnrealUSDWrapper.cpp
+++ b/Engine/Plugins/Importers/USDImporter/Source/UnrealUSDWrapper/Private/UnrealUSDWrapper.cpp
@@ -939,7 +939,7 @@ namespace UE::UnrealUSDWrapper::Private
 		pxr::UsdStagePopulationMask Mask;
 		if( PopulationMask )
 		{
-			// The USD OpenMaskedStage functions don't actually consult or populate the stage cache
+			// The USD OpenMasked functions don't actually consult or populate the stage cache
 			ensure( bUseStageCache == false );
 
 			for ( const FString& AllowedPrimPath : *PopulationMask )
