USD: Need to close Sequencer when done exporting

diff --git a/Engine/Plugins/Importers/USDImporter/Source/USDExporter/Private/LevelSequenceExporterUSD.cpp b/Engine/Plugins/Importers/USDImporter/Source/USDExporter/Private/LevelSequenceExporterUSD.cpp
--- a/Engine/Plugins/Importers/USDImporter/Source/USDExporter/Private/LevelSequenceExporterUSD.cpp
+++ b/Engine/Plugins/Importers/USDImporter/Source/USDExporter/Private/LevelSequenceExporterUSD.cpp
@@ -1338,6 +1338,7 @@ bool ULevelSequenceExporterUsd::ExportBinary( UObject* Object, const TCHAR* Type
 
 	// Set this back to Stopped or else it will keep the editor viewport controls permanently hidden
 	TempSequencer->SetPlaybackStatus( EMovieScenePlayerStatus::Stopped );
+	TempSequencer->Close();
 
 	if ( GEditor && GIsEditor && AssetEditorSubsystem )
 	{
