USD: Actually save configs after picking to not prompt again on any of the new dialogs.

diff --git a/Engine/Plugins/Importers/USDImporter/Source/USDStage/Private/USDStageActorCustomization.cpp b/Engine/Plugins/Importers/USDImporter/Source/USDStage/Private/USDStageActorCustomization.cpp
--- a/Engine/Plugins/Importers/USDImporter/Source/USDStage/Private/USDStageActorCustomization.cpp
+++ b/Engine/Plugins/Importers/USDImporter/Source/USDStage/Private/USDStageActorCustomization.cpp
@@ -223,6 +223,7 @@ void FUsdStageActorCustomization::CustomizeDetails( IDetailLayoutBuilder& Detail
 							if ( UUsdProjectSettings* ProjectSettings = GetMutableDefault<UUsdProjectSettings>() )
 							{
 								ProjectSettings->AdditionalMaterialPurposes.AddUnique( NewPurpose );
+								ProjectSettings->SaveConfig();
 							}
 						}
 
diff --git a/Engine/Plugins/Importers/USDImporter/Source/USDStageEditor/Private/SUSDLayersTreeView.cpp b/Engine/Plugins/Importers/USDImporter/Source/USDStageEditor/Private/SUSDLayersTreeView.cpp
--- a/Engine/Plugins/Importers/USDImporter/Source/USDStageEditor/Private/SUSDLayersTreeView.cpp
+++ b/Engine/Plugins/Importers/USDImporter/Source/USDStageEditor/Private/SUSDLayersTreeView.cpp
@@ -1253,6 +1253,7 @@ void SUsdLayersTreeView::OnClearSelectedLayers()
 
 				UUsdProjectSettings* Settings = GetMutableDefault<UUsdProjectSettings>();
 				Settings->bShowConfirmationWhenClearingLayers = false;
+				Settings->SaveConfig();
 
 				ClearSelectedLayersInner();
 			}
diff --git a/Engine/Plugins/Importers/USDImporter/Source/USDStageEditor/Private/SUSDStage.cpp b/Engine/Plugins/Importers/USDImporter/Source/USDStageEditor/Private/SUSDStage.cpp
--- a/Engine/Plugins/Importers/USDImporter/Source/USDStageEditor/Private/SUSDStage.cpp
+++ b/Engine/Plugins/Importers/USDImporter/Source/USDStageEditor/Private/SUSDStage.cpp
@@ -1173,6 +1173,7 @@ void SUsdStage::FillMaterialPurposeSubMenu( FMenuBuilder& MenuBuilder )
 					if ( UUsdProjectSettings* ProjectSettings = GetMutableDefault<UUsdProjectSettings>() )
 					{
 						ProjectSettings->AdditionalMaterialPurposes.AddUnique( NewPurpose );
+						ProjectSettings->SaveConfig();
 					}
 				}
 
diff --git a/Engine/Plugins/Importers/USDImporter/Source/USDStageEditor/Private/USDStageEditorModule.cpp b/Engine/Plugins/Importers/USDImporter/Source/USDStageEditor/Private/USDStageEditorModule.cpp
--- a/Engine/Plugins/Importers/USDImporter/Source/USDStageEditor/Private/USDStageEditorModule.cpp
+++ b/Engine/Plugins/Importers/USDImporter/Source/USDStageEditor/Private/USDStageEditorModule.cpp
@@ -160,6 +160,8 @@ namespace UE::UsdStageEditorModule::Private
 				: bShouldSave
 					? EUsdSaveDialogBehavior::AlwaysSave
 					: EUsdSaveDialogBehavior::NeverSave;
+
+			Settings->SaveConfig();
 		}
 
 		for ( const FUsdSaveDialogRowData& ReturnedRow : Rows )
diff --git a/Engine/Plugins/Importers/USDImporter/Source/USDStageImporter/Private/USDStageImportOptionsCustomization.cpp b/Engine/Plugins/Importers/USDImporter/Source/USDStageImporter/Private/USDStageImportOptionsCustomization.cpp
--- a/Engine/Plugins/Importers/USDImporter/Source/USDStageImporter/Private/USDStageImportOptionsCustomization.cpp
+++ b/Engine/Plugins/Importers/USDImporter/Source/USDStageImporter/Private/USDStageImportOptionsCustomization.cpp
@@ -226,6 +226,7 @@ void FUsdStageImportOptionsCustomization::CustomizeDetails(IDetailLayoutBuilder&
 							if ( UUsdProjectSettings* ProjectSettings = GetMutableDefault<UUsdProjectSettings>() )
 							{
 								ProjectSettings->AdditionalMaterialPurposes.AddUnique( NewPurpose );
+								ProjectSettings->SaveConfig();
 							}
 						}
 
diff --git a/Engine/Plugins/Importers/USDImporter/Source/USDUtilities/Private/USDAttributeUtils.cpp b/Engine/Plugins/Importers/USDImporter/Source/USDUtilities/Private/USDAttributeUtils.cpp
--- a/Engine/Plugins/Importers/USDImporter/Source/USDUtilities/Private/USDAttributeUtils.cpp
+++ b/Engine/Plugins/Importers/USDImporter/Source/USDUtilities/Private/USDAttributeUtils.cpp
@@ -263,6 +263,7 @@ void UsdUtils::NotifyIfOverriddenOpinion( const UE::FUsdAttribute& Attribute )
 					{
 						// This is flipped because the default checkbox message is "Don't prompt again"
 						Settings->bShowOverriddenOpinionsWarning = NewState == ECheckBoxState::Unchecked;
+						Settings->SaveConfig();
 					}
 				});
 
diff --git a/Engine/Plugins/Importers/USDImporter/Source/USDUtilities/Private/USDConversionUtils.cpp b/Engine/Plugins/Importers/USDImporter/Source/USDUtilities/Private/USDConversionUtils.cpp
--- a/Engine/Plugins/Importers/USDImporter/Source/USDUtilities/Private/USDConversionUtils.cpp
+++ b/Engine/Plugins/Importers/USDImporter/Source/USDUtilities/Private/USDConversionUtils.cpp
@@ -286,6 +286,7 @@ namespace USDConversionUtilsImpl
 				{
 					// This is flipped because the default checkbox message is "Don't prompt again"
 					Settings->bShowWarningOnIncompleteDuplication = NewState == ECheckBoxState::Unchecked;
+					Settings->SaveConfig();
 				}
 			} );
 
