USD: Fix missing invalid prim checks on FUsdPrimViewModel.

diff --git a/Engine/Plugins/Importers/USDImporter/Source/USDStageEditorViewModels/Private/USDPrimViewModel.cpp b/Engine/Plugins/Importers/USDImporter/Source/USDStageEditorViewModels/Private/USDPrimViewModel.cpp
--- a/Engine/Plugins/Importers/USDImporter/Source/USDStageEditorViewModels/Private/USDPrimViewModel.cpp
+++ b/Engine/Plugins/Importers/USDImporter/Source/USDStageEditorViewModels/Private/USDPrimViewModel.cpp
@@ -111,7 +111,15 @@ TArray< FUsdPrimViewModelRef >& FUsdPrimViewModel::UpdateChildren()
 void FUsdPrimViewModel::FillChildren()
 {
 #if USE_USD_SDK
+	if ( !UsdPrim )
+	{
+		return;
+	}
+
+	FScopedUsdAllocs UsdAllocs;
 	pxr::UsdPrimSiblingRange PrimChildren = pxr::UsdPrim( UsdPrim ).GetFilteredChildren( pxr::UsdTraverseInstanceProxies( pxr::UsdPrimAllPrimsPredicate ) );
+
+	FScopedUnrealAllocs UnrealAllocs;
 	for ( pxr::UsdPrim Child : PrimChildren )
 	{
 		Children.Add( MakeShared< FUsdPrimViewModel >( this, UsdStage, UE::FUsdPrim( Child ) ) );
@@ -202,7 +210,7 @@ void FUsdPrimViewModel::ToggleVisibility()
 
 void FUsdPrimViewModel::TogglePayload()
 {
-	if ( UsdPrim.HasPayload() )
+	if ( UsdPrim && UsdPrim.HasPayload() )
 	{
 		if ( UsdPrim.IsLoaded() )
 		{
@@ -318,6 +326,11 @@ void FUsdPrimViewModel::DefinePrim( const TCHAR* PrimName )
 void FUsdPrimViewModel::ClearReferences()
 {
 #if USE_USD_SDK
+	if ( !UsdPrim )
+	{
+		return;
+	}
+
 	FScopedUsdAllocs UsdAllocs;
 
 	pxr::UsdReferences References = pxr::UsdPrim( UsdPrim ).GetReferences();
@@ -328,6 +341,11 @@ void FUsdPrimViewModel::ClearReferences()
 void FUsdPrimViewModel::ClearPayloads()
 {
 #if USE_USD_SDK
+	if ( !UsdPrim )
+	{
+		return;
+	}
+
 	FScopedUsdAllocs UsdAllocs;
 
 	pxr::UsdPayloads Payloads = pxr::UsdPrim( UsdPrim ).GetPayloads();
