USD: Improve feedback on delete/rename/cut prim buttons when the selected prim is not reachable.

Disable some right-click operations if they would have led to no action;
Add tooltip explaining why the disabled actions are disabled.

diff --git a/Engine/Plugins/Importers/USDImporter/Source/USDStageEditor/Private/SUSDStageTreeView.cpp b/Engine/Plugins/Importers/USDImporter/Source/USDStageEditor/Private/SUSDStageTreeView.cpp
--- a/Engine/Plugins/Importers/USDImporter/Source/USDStageEditor/Private/SUSDStageTreeView.cpp
+++ b/Engine/Plugins/Importers/USDImporter/Source/USDStageEditor/Private/SUSDStageTreeView.cpp
@@ -41,6 +41,12 @@
 
 #if USE_USD_SDK
 
+namespace UE::USDStageTreeView::Private
+{
+	static const FText NoSpecOnLocalLayerStack =
+		LOCTEXT( "NoLocalSpecToolTip", "This prim needs at least one spec on the stage's local layer stack for this option to be usable" );
+}
+
 enum class EPayloadsTrigger
 {
 	Load,
@@ -350,7 +356,7 @@ void SUsdStageTreeView::Construct( const FArguments& InArgs )
 	UICommandList->MapAction(
 		FGenericCommands::Get().Cut,
 		FExecuteAction::CreateSP( this, &SUsdStageTreeView::OnCutPrim ),
-		FCanExecuteAction::CreateSP( this, &SUsdStageTreeView::DoesPrimExistOnStage )
+		FCanExecuteAction::CreateSP( this, &SUsdStageTreeView::DoesPrimHaveSpecOnLocalLayerStack )
 	);
 	UICommandList->MapAction(
 		FGenericCommands::Get().Copy,
@@ -370,12 +376,15 @@ void SUsdStageTreeView::Construct( const FArguments& InArgs )
 	UICommandList->MapAction(
 		FGenericCommands::Get().Delete,
 		FExecuteAction::CreateSP( this, &SUsdStageTreeView::OnDeletePrim ),
-		FCanExecuteAction::CreateSP( this, &SUsdStageTreeView::DoesPrimExistOnStage )
+		FCanExecuteAction::CreateSP( this, &SUsdStageTreeView::DoesPrimHaveSpecOnLocalLayerStack )
 	);
 	UICommandList->MapAction(
 		FGenericCommands::Get().Rename,
 		FExecuteAction::CreateSP( this, &SUsdStageTreeView::OnRenamePrim ),
-		FCanExecuteAction::CreateLambda( [this]() { return SUsdStageTreeView::DoesPrimExistOnStage() && GetSelectedItems().Num() == 1; } )
+		FCanExecuteAction::CreateLambda( [this]()
+		{
+			return SUsdStageTreeView::DoesPrimHaveSpecOnLocalLayerStack() && GetSelectedItems().Num() == 1;
+		})
 	);
 }
 
@@ -619,14 +628,19 @@ TSharedPtr< SWidget > SUsdStageTreeView::ConstructPrimContextMenu()
 			FGenericCommands::Get().Cut,
 			NAME_None,
 			LOCTEXT( "Cut_Text", "Cut" ),
-			LOCTEXT( "Cut_ToolTip", "Cuts this prim selection from the stage" )
+			TAttribute<FText>::Create( TAttribute<FText>::FGetter::CreateLambda( [this]()
+			{
+				return DoesPrimHaveSpecOnLocalLayerStack()
+					? LOCTEXT( "Cut_ToolTip", "Cuts the selected prim's specs from the stage's local layer stack" )
+					: UE::USDStageTreeView::Private::NoSpecOnLocalLayerStack;
+			}))
 		);
 
 		PrimOptions.AddMenuEntry(
 			FGenericCommands::Get().Copy,
 			NAME_None,
 			LOCTEXT( "Copy_Text", "Copy" ),
-			LOCTEXT( "Copy_ToolTip", "Copies the selected prims from the stage" )
+			LOCTEXT( "Copy_ToolTip", "Copies the selected prims" )
 		);
 
 		PrimOptions.AddMenuEntry(
@@ -649,14 +663,24 @@ TSharedPtr< SWidget > SUsdStageTreeView::ConstructPrimContextMenu()
 			FGenericCommands::Get().Delete,
 			NAME_None,
 			LOCTEXT( "Delete_Text", "Delete" ),
-			LOCTEXT( "Delete_ToolTip", "Deletes the selected prims" )
+			TAttribute<FText>::Create( TAttribute<FText>::FGetter::CreateLambda( [this]()
+			{
+				return DoesPrimHaveSpecOnLocalLayerStack()
+					? LOCTEXT( "Delete_ToolTip", "Deletes the selected prim's specs from the local layer stack" )
+					: UE::USDStageTreeView::Private::NoSpecOnLocalLayerStack;
+			}))
 		);
 
 		PrimOptions.AddMenuEntry(
 			FGenericCommands::Get().Rename,
 			NAME_None,
 			LOCTEXT( "Rename_Text", "Rename" ),
-			LOCTEXT( "Rename_ToolTip", "Renames the selected prim" )
+			TAttribute<FText>::Create( TAttribute<FText>::FGetter::CreateLambda( [this]()
+			{
+				return DoesPrimHaveSpecOnLocalLayerStack()
+					? LOCTEXT( "Rename_ToolTip", "Renames the selected prim's specs on the local layer stack" )
+					: UE::USDStageTreeView::Private::NoSpecOnLocalLayerStack;
+			}))
 		);
 	}
 	PrimOptions.EndSection();
@@ -1324,7 +1348,7 @@ bool SUsdStageTreeView::DoesPrimExistOnEditTarget() const
 	return false;
 }
 
-bool SUsdStageTreeView::CanDuplicateAllLocalLayerSpecs() const
+bool SUsdStageTreeView::DoesPrimHaveSpecOnLocalLayerStack() const
 {
 	TArray< FUsdPrimViewModelRef > MySelectedItems = GetSelectedItems();
 
@@ -1463,7 +1487,7 @@ void SUsdStageTreeView::FillDuplicateSubmenu( FMenuBuilder& MenuBuilder )
 		{
 			return DoesPrimExistOnEditTarget()
 				? LOCTEXT( "DuplicateSingleValid_ToolTip", "Duplicate the prim's specs on the current edit target only" )
-				: LOCTEXT( "DuplicateSingleInvalid_ToolTip", "This prim needs a spec in the local edit target for this option to be usable!" );
+				: UE::USDStageTreeView::Private::NoSpecOnLocalLayerStack;
 		})),
 		FSlateIcon(),
 		FUIAction(
@@ -1478,14 +1502,14 @@ void SUsdStageTreeView::FillDuplicateSubmenu( FMenuBuilder& MenuBuilder )
 		LOCTEXT( "DuplicateAllLocal_Text", "All local layer specs" ),
 		TAttribute<FText>::Create( TAttribute<FText>::FGetter::CreateLambda( [this]()
 		{
-			return CanDuplicateAllLocalLayerSpecs()
+			return DoesPrimHaveSpecOnLocalLayerStack()
 				? LOCTEXT( "DuplicateAllLocalValid_ToolTip", "Duplicate each of the prim's specs across the entire stage" )
-				: LOCTEXT( "DuplicateAllLocalInvalid_ToolTip", "This prim needs at least one spec on the stage's local layer stack for this option to be usable!" );
+				: UE::USDStageTreeView::Private::NoSpecOnLocalLayerStack;
 		})),
 		FSlateIcon(),
 		FUIAction(
 			FExecuteAction::CreateSP( this, &SUsdStageTreeView::OnDuplicatePrim, EUsdDuplicateType::AllLocalLayerSpecs ),
-			FCanExecuteAction::CreateSP( this, &SUsdStageTreeView::CanDuplicateAllLocalLayerSpecs )
+			FCanExecuteAction::CreateSP( this, &SUsdStageTreeView::DoesPrimHaveSpecOnLocalLayerStack )
 		),
 		NAME_None,
 		EUserInterfaceActionType::Button
diff --git a/Engine/Plugins/Importers/USDImporter/Source/USDStageEditor/Private/SUSDStageTreeView.h b/Engine/Plugins/Importers/USDImporter/Source/USDStageEditor/Private/SUSDStageTreeView.h
--- a/Engine/Plugins/Importers/USDImporter/Source/USDStageEditor/Private/SUSDStageTreeView.h
+++ b/Engine/Plugins/Importers/USDImporter/Source/USDStageEditor/Private/SUSDStageTreeView.h
@@ -80,7 +80,7 @@ private:
 	bool CanPastePrim() const;
 	bool DoesPrimExistOnStage() const;
 	bool DoesPrimExistOnEditTarget() const;
-	bool CanDuplicateAllLocalLayerSpecs() const;
+	bool DoesPrimHaveSpecOnLocalLayerStack() const;
 
 	/** Uses TreeItemExpansionStates to travel the tree and call SetItemExpansion */
 	void RestoreExpansionStates();
