USD: Add some invalid attribute checks in case the integration panel is left displaying an invalid prim.

diff --git a/Engine/Plugins/Importers/USDImporter/Source/USDStageEditor/Private/Widgets/SUSDIntegrationsPanel.cpp b/Engine/Plugins/Importers/USDImporter/Source/USDStageEditor/Private/Widgets/SUSDIntegrationsPanel.cpp
--- a/Engine/Plugins/Importers/USDImporter/Source/USDStageEditor/Private/Widgets/SUSDIntegrationsPanel.cpp
+++ b/Engine/Plugins/Importers/USDImporter/Source/USDStageEditor/Private/Widgets/SUSDIntegrationsPanel.cpp
@@ -159,7 +159,7 @@ TSharedRef< SWidget > SUsdIntegrationsPanelRow::GenerateWidgetForColumn( const F
 					);
 
 					UE::FVtValue Value;
-					if ( UsdUtils::SetUnderlyingValue( Value, UsdPath ) && !Value.IsEmpty() )
+					if ( AttributeCopy && *AttributeCopy && UsdUtils::SetUnderlyingValue( Value, UsdPath ) && !Value.IsEmpty() )
 					{
 						UE::FSdfChangeBlock Block;
 						AttributeCopy->Set( Value );
@@ -217,7 +217,7 @@ TSharedRef< SWidget > SUsdIntegrationsPanelRow::GenerateWidgetForColumn( const F
 					);
 
 					UE::FVtValue Value;
-					if ( UsdUtils::SetUnderlyingValue( Value, UsdPath ) && !Value.IsEmpty() )
+					if ( AttributeCopy && *AttributeCopy && UsdUtils::SetUnderlyingValue( Value, UsdPath ) && !Value.IsEmpty() )
 					{
 						UE::FSdfChangeBlock Block;
 						AttributeCopy->Set( Value );
@@ -257,7 +257,7 @@ TSharedRef< SWidget > SUsdIntegrationsPanelRow::GenerateWidgetForColumn( const F
 				.IsChecked_Lambda( [AttributeCopy]() ->ECheckBoxState
 				{
 					UE::FVtValue Value;
-					if ( AttributeCopy->Get( Value ) )
+					if ( AttributeCopy && *AttributeCopy && AttributeCopy->Get( Value ) )
 					{
 						if ( TOptional<bool> UnderlyingValue = UsdUtils::GetUnderlyingValue<bool>( Value ) )
 						{
@@ -276,7 +276,7 @@ TSharedRef< SWidget > SUsdIntegrationsPanelRow::GenerateWidgetForColumn( const F
 					);
 
 					UE::FVtValue Value;
-					if ( UsdUtils::SetUnderlyingValue( Value, NewValue == ECheckBoxState::Checked ) && !Value.IsEmpty() )
+					if ( AttributeCopy && *AttributeCopy && UsdUtils::SetUnderlyingValue( Value, NewValue == ECheckBoxState::Checked ) && !Value.IsEmpty() )
 					{
 						UE::FSdfChangeBlock Block;
 						AttributeCopy->Set( Value );
@@ -302,7 +302,7 @@ TSharedRef< SWidget > SUsdIntegrationsPanelRow::GenerateWidgetForColumn( const F
 					);
 
 					UE::FVtValue Value;
-					if ( UsdUtils::SetUnderlyingValue( Value, NewValue ) && !Value.IsEmpty() )
+					if ( AttributeCopy && *AttributeCopy && UsdUtils::SetUnderlyingValue( Value, NewValue ) && !Value.IsEmpty() )
 					{
 						UE::FSdfChangeBlock Block;
 						AttributeCopy->Set( Value );
@@ -312,7 +312,7 @@ TSharedRef< SWidget > SUsdIntegrationsPanelRow::GenerateWidgetForColumn( const F
 				.Value_Lambda( [AttributeCopy]() -> float
 				{
 					UE::FVtValue Value;
-					if ( AttributeCopy->Get( Value ) )
+					if ( AttributeCopy && *AttributeCopy && AttributeCopy->Get( Value ) )
 					{
 						if ( TOptional<float> UnderlyingValue = UsdUtils::GetUnderlyingValue<float>( Value ) )
 						{
@@ -336,24 +336,26 @@ TSharedRef< SWidget > SUsdIntegrationsPanelRow::GenerateWidgetForColumn( const F
 			.Value_Lambda( [AttributeCopy]()
 			{
 				SLiveLinkSubjectRepresentationPicker::FLiveLinkSourceSubjectRole Result;
-
-				Result.Role = AttributeCopy->GetPrim().IsA(TEXT("SkelRoot"))
-					? ULiveLinkAnimationRole::StaticClass()
-					: ULiveLinkTransformRole::StaticClass();
-
-				FString SubjectName;
-
-				UE::FVtValue Value;
-				if ( AttributeCopy->Get( Value ) )
+				if ( AttributeCopy && *AttributeCopy )
 				{
-					if ( TOptional<std::string> UnderlyingValue = UsdUtils::GetUnderlyingValue<std::string>( Value ) )
+					Result.Role = AttributeCopy->GetPrim().IsA(TEXT("SkelRoot"))
+						? ULiveLinkAnimationRole::StaticClass()
+						: ULiveLinkTransformRole::StaticClass();
+
+					UE::FVtValue Value;
+					if ( AttributeCopy->Get( Value ) )
 					{
-						SubjectName = UsdToUnreal::ConvertString( UnderlyingValue.GetValue() );
+						if ( TOptional<std::string> UnderlyingValue = UsdUtils::GetUnderlyingValue<std::string>( Value ) )
+						{
+							Result.Subject = FLiveLinkSubjectName{
+								FName{
+									*UsdToUnreal::ConvertString( UnderlyingValue.GetValue() )
+								}
+							};
+						}
 					}
 				}
 
-				Result.Subject = FLiveLinkSubjectName{ FName{*SubjectName} };
-
 				return Result;
 			})
 			.OnValueChanged_Lambda( [AttributeCopy]( SLiveLinkSubjectRepresentationPicker::FLiveLinkSourceSubjectRole NewValue )
@@ -369,7 +371,7 @@ TSharedRef< SWidget > SUsdIntegrationsPanelRow::GenerateWidgetForColumn( const F
 				);
 
 				UE::FVtValue Value;
-				if ( UsdUtils::SetUnderlyingValue( Value, UsdValue ) && !Value.IsEmpty() )
+				if ( AttributeCopy && *AttributeCopy && UsdUtils::SetUnderlyingValue( Value, UsdValue ) && !Value.IsEmpty() )
 				{
 					UE::FSdfChangeBlock Block;
 					AttributeCopy->Set( Value );
@@ -388,7 +390,7 @@ TSharedRef< SWidget > SUsdIntegrationsPanelRow::GenerateWidgetForColumn( const F
 				.IsChecked_Lambda( [AttributeCopy]()
 				{
 					UE::FVtValue Value;
-					if ( AttributeCopy->Get( Value ) )
+					if ( AttributeCopy && *AttributeCopy && AttributeCopy->Get( Value ) )
 					{
 						if ( TOptional<bool> UnderlyingValue = UsdUtils::GetUnderlyingValue<bool>( Value ) )
 						{
@@ -408,7 +410,7 @@ TSharedRef< SWidget > SUsdIntegrationsPanelRow::GenerateWidgetForColumn( const F
 					);
 
 					UE::FVtValue Value;
-					if ( UsdUtils::SetUnderlyingValue( Value, NewState == ECheckBoxState::Checked ) && !Value.IsEmpty() )
+					if ( AttributeCopy && *AttributeCopy && UsdUtils::SetUnderlyingValue( Value, NewState == ECheckBoxState::Checked ) && !Value.IsEmpty() )
 					{
 						UE::FSdfChangeBlock Block;
 						AttributeCopy->Set( Value );
