USD: Don't register static mesh components before we're done with our updates.

diff --git a/Engine/Plugins/Importers/USDImporter/Source/USDSchemas/Private/USDGeomXformableTranslator.cpp b/Engine/Plugins/Importers/USDImporter/Source/USDSchemas/Private/USDGeomXformableTranslator.cpp
--- a/Engine/Plugins/Importers/USDImporter/Source/USDSchemas/Private/USDGeomXformableTranslator.cpp
+++ b/Engine/Plugins/Importers/USDImporter/Source/USDSchemas/Private/USDGeomXformableTranslator.cpp
@@ -573,7 +573,8 @@ void FUsdGeomXformableTranslator::UpdateComponents( USceneComponent* SceneCompon
 
 				StaticMeshComponent->SetStaticMesh( PrimStaticMesh );
 
-				StaticMeshComponent->RegisterComponent();
+				// We can't register yet, as UsdToUnreal::ConvertXformable below us may want to move the component.
+				// We'll always re-register when needed below, though.
 			}
 		}
 
