USD: Fix potential crash when trying to remove prims;

diff --git a/Engine/Plugins/Importers/USDImporter/Source/USDUtilities/Private/USDConversionUtils.cpp b/Engine/Plugins/Importers/USDImporter/Source/USDUtilities/Private/USDConversionUtils.cpp
--- a/Engine/Plugins/Importers/USDImporter/Source/USDUtilities/Private/USDConversionUtils.cpp
+++ b/Engine/Plugins/Importers/USDImporter/Source/USDUtilities/Private/USDConversionUtils.cpp
@@ -1412,6 +1412,11 @@ bool UsdUtils::RenamePrim( UE::FUsdPrim& Prim, const TCHAR* NewPrimName )
 	bool bCanApply = true;
 	for ( const pxr::SdfPrimSpecHandle& Spec : SpecStack )
 	{
+		if ( !Spec )
+		{
+			continue;
+		}
+
 		pxr::SdfPath SpecPath = Spec->GetPath();
 		if ( !SpecPath.IsPrimPath() )
 		{
@@ -1736,6 +1741,11 @@ UE::FSdfPath UsdUtils::GetPrimSpecPathForLayer( const UE::FUsdPrim& Prim, const
 	std::size_t LargestPathLength = 0;
 	for ( const pxr::SdfPrimSpecHandle& Spec : UsdPrim.GetPrimStack() )
 	{
+		if ( !Spec )
+		{
+			continue;
+		}
+
 		pxr::SdfPath SpecPath = Spec->GetPath();
 		if ( !SpecPath.IsPrimPath() )
 		{
