USD: Fix "save dirty layers" dialog popping up when creating a new stage;

Move New Stage creation to stage actor;
Add new utils function to set prim kind with the EUsdDefaultKind enum;
Prevent stray UI refreshes from triggering full stage loads.

diff --git a/Engine/Plugins/Importers/USDImporter/Source/USDStage/Private/USDStageActor.cpp b/Engine/Plugins/Importers/USDImporter/Source/USDStage/Private/USDStageActor.cpp
--- a/Engine/Plugins/Importers/USDImporter/Source/USDStage/Private/USDStageActor.cpp
+++ b/Engine/Plugins/Importers/USDImporter/Source/USDStage/Private/USDStageActor.cpp
@@ -767,6 +767,33 @@ AUsdStageActor::AUsdStageActor()
 	}
 }
 
+void AUsdStageActor::NewStage()
+{
+#if USE_USD_SDK
+	UE::FUsdStage NewStage = UnrealUSDWrapper::NewStage();
+	if ( !NewStage )
+	{
+		return;
+	}
+
+	// We'll create an in-memory stage, and so the "RootLayer" path we'll use will be a
+	// magic path that is guaranteed to never exist in a filesystem due to invalid characters.
+	UE::FSdfLayer Layer = NewStage.GetRootLayer();
+	if ( !Layer )
+	{
+		return;
+	}
+	FString	StagePath = FString( UnrealIdentifiers::IdentifierPrefix ) + Layer.GetIdentifier();
+
+	UE::FUsdPrim RootPrim = NewStage.DefinePrim( UE::FSdfPath{ TEXT( "/Root" ) }, TEXT( "Xform" ) );
+	ensure( UsdUtils::SetDefaultKind( RootPrim, EUsdDefaultKind::Assembly ) );
+
+	NewStage.SetDefaultPrim( RootPrim );
+
+	SetRootLayer( StagePath );
+#endif // USE_USD_SDK
+}
+
 void AUsdStageActor::IsolateLayer( const UE::FSdfLayer& Layer )
 {
 	if ( IsolatedStage && IsolatedStage.GetRootLayer() == Layer )
diff --git a/Engine/Plugins/Importers/USDImporter/Source/USDStage/Public/USDStageActor.h b/Engine/Plugins/Importers/USDImporter/Source/USDStage/Public/USDStageActor.h
--- a/Engine/Plugins/Importers/USDImporter/Source/USDStage/Public/USDStageActor.h
+++ b/Engine/Plugins/Importers/USDImporter/Source/USDStage/Public/USDStageActor.h
@@ -197,6 +197,9 @@ public:
 public:
 	AUsdStageActor();
 
+	// Creates a brand new, memory-only USD stage and opens it
+	USDSTAGE_API void NewStage();
+
 	USDSTAGE_API void IsolateLayer( const UE::FSdfLayer& Layer );
 
 	USDSTAGE_API void Reset() override;
diff --git a/Engine/Plugins/Importers/USDImporter/Source/USDStageEditor/Private/SUSDStage.cpp b/Engine/Plugins/Importers/USDImporter/Source/USDStageEditor/Private/SUSDStage.cpp
--- a/Engine/Plugins/Importers/USDImporter/Source/USDStageEditor/Private/SUSDStage.cpp
+++ b/Engine/Plugins/Importers/USDImporter/Source/USDStageEditor/Private/SUSDStage.cpp
@@ -194,9 +194,10 @@ void SUsdStage::Construct( const FArguments& InArgs )
 
 	UE::FUsdStage UsdStage;
 
-	if ( ViewModel.UsdStageActor.IsValid() )
+	const AUsdStageActor* StageActor = ViewModel.UsdStageActor.Get();
+	if ( StageActor )
 	{
-		UsdStage = ViewModel.UsdStageActor->GetOrLoadUsdStage();
+		UsdStage = static_cast< const AUsdStageActor* >( StageActor )->GetUsdStage();
 	}
 
 	ChildSlot
@@ -1585,7 +1586,7 @@ void SUsdStage::FillNaniteThresholdSubMenu( FMenuBuilder& MenuBuilder )
 
 void SUsdStage::OnNew()
 {
-	ViewModel.NewStage( nullptr );
+	ViewModel.NewStage();
 }
 
 void SUsdStage::OnOpen()
@@ -1600,28 +1601,31 @@ void SUsdStage::OnOpen()
 
 void SUsdStage::OnSave()
 {
-	UE::FUsdStage UsdStage;
-	if ( ViewModel.UsdStageActor.IsValid() )
+	const AUsdStageActor* StageActor = ViewModel.UsdStageActor.Get();
+	if ( !StageActor )
 	{
-		UsdStage = ViewModel.UsdStageActor->GetOrLoadUsdStage();
+		return;
 	}
 
-	if ( UsdStage )
+	UE::FUsdStage UsdStage = static_cast< const AUsdStageActor* >( StageActor )->GetUsdStage();
+	if ( !UsdStage )
 	{
-		if ( UE::FSdfLayer RootLayer = UsdStage.GetRootLayer() )
+		return;
+	}
+
+	if ( UE::FSdfLayer RootLayer = UsdStage.GetRootLayer() )
+	{
+		FString RealPath = RootLayer.GetRealPath();
+		if ( FPaths::FileExists( RealPath ) )
 		{
-			FString RealPath = RootLayer.GetRealPath();
-			if ( FPaths::FileExists( RealPath ) )
-			{
-				ViewModel.SaveStage();
-			}
-			else
+			ViewModel.SaveStage();
+		}
+		else
+		{
+			TOptional< FString > UsdFilePath = UsdUtils::BrowseUsdFile( UsdUtils::EBrowseFileMode::Save );
+			if ( UsdFilePath )
 			{
-				TOptional< FString > UsdFilePath = UsdUtils::BrowseUsdFile( UsdUtils::EBrowseFileMode::Save );
-				if ( UsdFilePath )
-				{
-					ViewModel.SaveStageAs( *UsdFilePath.GetValue() );
-				}
+				ViewModel.SaveStageAs( *UsdFilePath.GetValue() );
 			}
 		}
 	}
@@ -1879,7 +1883,7 @@ void SUsdStage::OnPrimSelectionChanged( const TArray<FString>& PrimPaths )
 
 	if ( UsdPrimInfoWidget )
 	{
-		UE::FUsdStage UsdStage = StageActor->GetOrLoadUsdStage();
+		UE::FUsdStage UsdStage = static_cast< const AUsdStageActor* >( StageActor )->GetUsdStage();
 
 		SelectedPrimPath = PrimPaths.Num() == 1 ? PrimPaths[ 0 ] : TEXT( "" );
 		UsdPrimInfoWidget->SetPrimPath( UsdStage, *SelectedPrimPath );
@@ -1934,7 +1938,7 @@ void SUsdStage::SetActor( AUsdStageActor* InUsdStageActor )
 
 	if( InUsdStageActor )
 	{
-		UE::FUsdStage UsdStage = InUsdStageActor->GetOrLoadUsdStage();
+		UE::FUsdStage UsdStage = static_cast< const AUsdStageActor* >( InUsdStageActor )->GetUsdStage();
 
 		if ( this->UsdPrimInfoWidget && InUsdStageActor )
 		{
@@ -1963,7 +1967,9 @@ void SUsdStage::Refresh()
 
 	if ( UsdStageTreeView )
 	{
-		UE::FUsdStageWeak CurrentStage = StageActor ? StageActor->GetOrLoadUsdStage() : UE::FUsdStage{};
+		UE::FUsdStageWeak CurrentStage = StageActor
+			? static_cast< const AUsdStageActor* >( StageActor )->GetUsdStage()
+			: UE::FUsdStage{};
 
 		UsdStageTreeView->Refresh( CurrentStage );
 
diff --git a/Engine/Plugins/Importers/USDImporter/Source/USDStageEditorViewModels/Private/USDStageViewModel.cpp b/Engine/Plugins/Importers/USDImporter/Source/USDStageEditorViewModels/Private/USDStageViewModel.cpp
--- a/Engine/Plugins/Importers/USDImporter/Source/USDStageEditorViewModels/Private/USDStageViewModel.cpp
+++ b/Engine/Plugins/Importers/USDImporter/Source/USDStageEditorViewModels/Private/USDStageViewModel.cpp
@@ -55,51 +55,29 @@ namespace UsdViewModelImpl
 }
 #endif // #if USE_USD_SDK
 
-void FUsdStageViewModel::NewStage( const TCHAR* FilePath )
+void FUsdStageViewModel::NewStage()
 {
-	FScopedTransaction Transaction( FText::Format(
-		LOCTEXT("NewStageTransaction", "Created new USD stage '{0}'"),
-		FText::FromString( FilePath )
-	));
+	FScopedTransaction Transaction( LOCTEXT( "NewStageTransaction", "Created new USD stage" ) );
 
-	UE::FUsdStage UsdStage = FilePath ? UnrealUSDWrapper::NewStage( FilePath ) : UnrealUSDWrapper::NewStage();
-	if ( !UsdStage )
-	{
-		return;
-	}
+	UsdUtils::StartMonitoringErrors();
 
-	// If we pass in nullptr, we'll create an in-memory stage, and so the "RootLayer" path we'll send down to the
-	// stage actor will be a magic path that is guaranteed to never exist in a filesystem due to invalid characters.
-	// The stage actor will interpret that, and try to load the stage from the stage cache
-	FString StagePath = FilePath;
-	if ( FilePath == nullptr )
+	if ( !UsdStageActor.IsValid() )
 	{
-		UE::FSdfLayer Layer = UsdStage.GetRootLayer();
-		if ( Layer )
-		{
-			StagePath = FString( UnrealIdentifiers::IdentifierPrefix ) + Layer.GetIdentifier();
-		}
+		IUsdStageModule& UsdStageModule = FModuleManager::GetModuleChecked< IUsdStageModule >( TEXT( "USDStage" ) );
+		UsdStageActor = &UsdStageModule.GetUsdStageActor( GWorld );
 	}
 
-#if USE_USD_SDK
+	if ( AUsdStageActor* StageActor = UsdStageActor.Get() )
 	{
-		FScopedUsdAllocs UsdAllocs;
-
-		// Create default prim
-		pxr::UsdGeomXform RootPrim = pxr::UsdGeomXform::Define( UsdStage, UnrealToUsd::ConvertPath( TEXT("/Root") ).Get() );
-		pxr::UsdModelAPI( RootPrim ).SetKind( pxr::KindTokens->assembly );
-
-		// Set default prim
-		( (pxr::UsdStageRefPtr&)UsdStage )->SetDefaultPrim( RootPrim.GetPrim() );
+		StageActor->Modify();
+		StageActor->NewStage();
 	}
-#endif // #if USE_USD_SDK
 
-	OpenStage( *StagePath );
+	UsdUtils::ShowErrorsAndStopMonitoring();
 }
 
 void FUsdStageViewModel::OpenStage( const TCHAR* FilePath )
 {
-
 	UsdUtils::StartMonitoringErrors();
 
 	if ( !UsdStageActor.IsValid() )
@@ -117,6 +95,8 @@ void FUsdStageViewModel::OpenStage( const TCHAR* FilePath )
 	{
 		UE_LOG(LogUsd, Error, TEXT("Failed to find a AUsdStageActor that could open stage '%s'!"), FilePath);
 	}
+
+	UsdUtils::ShowErrorsAndStopMonitoring();
 }
 
 void FUsdStageViewModel::ReloadStage()
diff --git a/Engine/Plugins/Importers/USDImporter/Source/USDStageEditorViewModels/Public/USDStageViewModel.h b/Engine/Plugins/Importers/USDImporter/Source/USDStageEditorViewModels/Public/USDStageViewModel.h
--- a/Engine/Plugins/Importers/USDImporter/Source/USDStageEditorViewModels/Public/USDStageViewModel.h
+++ b/Engine/Plugins/Importers/USDImporter/Source/USDStageEditorViewModels/Public/USDStageViewModel.h
@@ -12,8 +12,7 @@ class UPackage;
 class USDSTAGEEDITORVIEWMODELS_API FUsdStageViewModel
 {
 public:
-	/** Can pass nullptr to create an in-memory stage */
-	void NewStage( const TCHAR* FilePath );
+	void NewStage();
 	void OpenStage( const TCHAR* FilePath );
 	void ReloadStage();
 	void ResetStage();
diff --git a/Engine/Plugins/Importers/USDImporter/Source/USDUtilities/Private/USDConversionUtils.cpp b/Engine/Plugins/Importers/USDImporter/Source/USDUtilities/Private/USDConversionUtils.cpp
--- a/Engine/Plugins/Importers/USDImporter/Source/USDUtilities/Private/USDConversionUtils.cpp
+++ b/Engine/Plugins/Importers/USDImporter/Source/USDUtilities/Private/USDConversionUtils.cpp
@@ -981,6 +981,60 @@ EUsdDefaultKind UsdUtils::GetDefaultKind( const pxr::UsdPrim& Prim )
 	return Result;
 }
 
+bool UsdUtils::SetDefaultKind( pxr::UsdPrim& Prim, EUsdDefaultKind NewKind )
+{
+	if ( !Prim )
+	{
+		return false;
+	}
+
+	FScopedUsdAllocs Allocs;
+
+	const int32 NewKindInt = static_cast< int32 >( NewKind );
+	const bool bSingleFlagSet = NewKindInt != 0 && ( NewKindInt & ( NewKindInt - 1 ) ) == 0;
+	if ( !bSingleFlagSet )
+	{
+		return false;
+	}
+
+	pxr::TfToken NewKindToken;
+	switch ( NewKind )
+	{
+		default:
+		case EUsdDefaultKind::Model:
+		{
+			NewKindToken = pxr::KindTokens->model;
+			break;
+		}
+		case EUsdDefaultKind::Component:
+		{
+			NewKindToken = pxr::KindTokens->component;
+			break;
+		}
+		case EUsdDefaultKind::Group:
+		{
+			NewKindToken = pxr::KindTokens->group;
+			break;
+		}
+		case EUsdDefaultKind::Assembly:
+		{
+			NewKindToken = pxr::KindTokens->assembly;
+			break;
+		}
+		case EUsdDefaultKind::Subcomponent:
+		{
+			NewKindToken = pxr::KindTokens->subcomponent;
+			break;
+		}
+	}
+	if ( NewKindToken.IsEmpty() )
+	{
+		return false;
+	}
+
+	return IUsdPrim::SetKind( Prim, NewKindToken );
+}
+
 TArray< TUsdStore< pxr::UsdPrim > > UsdUtils::GetAllPrimsOfType( const pxr::UsdPrim& StartPrim, const pxr::TfType& SchemaType, const TArray< TUsdStore< pxr::TfType > >& ExcludeSchemaTypes )
 {
     return GetAllPrimsOfType( StartPrim, SchemaType, []( const pxr::UsdPrim& ) { return false; }, ExcludeSchemaTypes );
diff --git a/Engine/Plugins/Importers/USDImporter/Source/USDUtilities/Public/USDConversionUtils.h b/Engine/Plugins/Importers/USDImporter/Source/USDUtilities/Public/USDConversionUtils.h
--- a/Engine/Plugins/Importers/USDImporter/Source/USDUtilities/Public/USDConversionUtils.h
+++ b/Engine/Plugins/Importers/USDImporter/Source/USDUtilities/Public/USDConversionUtils.h
@@ -171,6 +171,15 @@ namespace UsdUtils
 	/** Returns whether Prim belongs to any of the default kinds, or a kind derived from them. The result can be a union of different kinds. */
 	USDUTILITIES_API EUsdDefaultKind GetDefaultKind( const pxr::UsdPrim& Prim );
 
+	/**
+	 * Sets the prim kind for Prim to the provided NewKind.
+	 * Differs from IUsdPrim::SetKind in that we don't have to provide the USD token.
+	 * @param NewKind - New kind to set. Must be a single flag, and not a combination of multiple plags
+	 * @param Prim - Prim to receive the new Kind
+	 * @return Whether we managed to set the kind or not.
+	 */
+	USDUTILITIES_API bool SetDefaultKind( pxr::UsdPrim& Prim, EUsdDefaultKind NewKind );
+
 	/**
 	 * Returns all prims of type SchemaType (or a descendant type) in the subtree of prims rooted at StartPrim.
 	 * Stops going down the subtrees when it hits a schema type to exclude.
