USD: Clean up the Analytics handling.

Remove code duplication (was required at the time as some stuff was introduced in a patch version);
Move level export analytics to C++;
Add analytics events to detect when importing grooms and groom caches;
Properly fetch whethe the material export is automated or not.

diff --git a/Engine/Plugins/Importers/USDImporter/Source/USDClasses/Private/USDStageOptions.cpp b/Engine/Plugins/Importers/USDImporter/Source/USDClasses/Private/USDStageOptions.cpp
new file mode 100644
--- /dev/null
+++ b/Engine/Plugins/Importers/USDImporter/Source/USDClasses/Private/USDStageOptions.cpp
@@ -0,0 +1,14 @@
+// Copyright Epic Games, Inc. All Rights Reserved.
+
+#include "USDStageOptions.h"
+
+#include "AnalyticsEventAttribute.h"
+
+void UsdUtils::AddAnalyticsAttributes(
+	const FUsdStageOptions& Options,
+	TArray< FAnalyticsEventAttribute >& InOutAttributes
+)
+{
+	InOutAttributes.Emplace( TEXT( "MetersPerUnit" ), LexToString( Options.MetersPerUnit ) );
+	InOutAttributes.Emplace( TEXT( "UpAxis" ), Options.UpAxis == EUsdUpAxis::YAxis ? TEXT( "Y" ) : TEXT( "Z" ) );
+}
diff --git a/Engine/Plugins/Importers/USDImporter/Source/USDClasses/Public/USDStageOptions.h b/Engine/Plugins/Importers/USDImporter/Source/USDClasses/Public/USDStageOptions.h
--- a/Engine/Plugins/Importers/USDImporter/Source/USDClasses/Public/USDStageOptions.h
+++ b/Engine/Plugins/Importers/USDImporter/Source/USDClasses/Public/USDStageOptions.h
@@ -4,6 +4,8 @@
 
 #include "USDStageOptions.generated.h"
 
+struct FAnalyticsEventAttribute;
+
 UENUM( BlueprintType )
 enum class EUsdUpAxis : uint8
 {
@@ -24,3 +26,11 @@ struct USDCLASSES_API FUsdStageOptions
 	UPROPERTY( EditAnywhere, config, BlueprintReadWrite, Category = "Stage options" )
 	EUsdUpAxis UpAxis = EUsdUpAxis::ZAxis;
 };
+
+namespace UsdUtils
+{
+	USDCLASSES_API void AddAnalyticsAttributes(
+		const FUsdStageOptions& Options,
+		TArray< FAnalyticsEventAttribute >& InOutAttributes
+	);
+}
diff --git a/Engine/Plugins/Importers/USDImporter/Source/USDExporter/Private/AnimSequenceExporterUSD.cpp b/Engine/Plugins/Importers/USDImporter/Source/USDExporter/Private/AnimSequenceExporterUSD.cpp
--- a/Engine/Plugins/Importers/USDImporter/Source/USDExporter/Private/AnimSequenceExporterUSD.cpp
+++ b/Engine/Plugins/Importers/USDImporter/Source/USDExporter/Private/AnimSequenceExporterUSD.cpp
@@ -25,69 +25,40 @@
 #include "AssetExportTask.h"
 #include "UObject/GCObjectScopeGuard.h"
 
-namespace UE
+namespace UE::AnimSequenceExporterUSD::Private
 {
-	namespace AnimSequenceExporterUSD
+	void SendAnalytics(
+		UObject* Asset,
+		UAnimSequenceExporterUSDOptions* Options,
+		bool bAutomated,
+		double ElapsedSeconds,
+		double NumberOfFrames,
+		const FString& Extension
+	)
 	{
-		namespace Private
+		if ( !Asset || !FEngineAnalytics::IsAvailable() )
 		{
-			void SendAnalytics( UObject* Asset, UAnimSequenceExporterUSDOptions* Options, bool bAutomated, double ElapsedSeconds, double NumberOfFrames, const FString& Extension )
-			{
-				if ( !Asset )
-				{
-					return;
-				}
-
-				if ( FEngineAnalytics::IsAvailable() )
-				{
-					FString ClassName = Asset->GetClass()->GetName();
-
-					TArray<FAnalyticsEventAttribute> EventAttributes;
+			return;
+		}
 
-					EventAttributes.Emplace( TEXT( "AssetType" ), ClassName );
+		FString ClassName = Asset->GetClass()->GetName();
 
-					if ( Options )
-					{
-						EventAttributes.Emplace( TEXT( "MetersPerUnit" ), LexToString( Options->StageOptions.MetersPerUnit ) );
-						EventAttributes.Emplace( TEXT( "UpAxis" ), Options->StageOptions.UpAxis == EUsdUpAxis::YAxis ? TEXT( "Y" ) : TEXT( "Z" ) );
-
-						EventAttributes.Emplace( TEXT( "ExportPreviewMesh" ), LexToString( Options->bExportPreviewMesh ) );
-						if ( Options->bExportPreviewMesh )
-						{
-							EventAttributes.Emplace( TEXT( "UsePayload" ), LexToString( Options->PreviewMeshOptions.bUsePayload ) );
-							if ( Options->PreviewMeshOptions.bUsePayload )
-							{
-								EventAttributes.Emplace( TEXT( "PayloadFormat" ), Options->PreviewMeshOptions.PayloadFormat );
-							}
-							EventAttributes.Emplace( TEXT( "BakeMaterials" ), Options->PreviewMeshOptions.bBakeMaterials );
-							if ( Options->PreviewMeshOptions.bBakeMaterials )
-							{
-								FString BakedPropertiesString;
-								{
-									const UEnum* PropertyEnum = StaticEnum<EMaterialProperty>();
-									for ( const FPropertyEntry& PropertyEntry : Options->PreviewMeshOptions.MaterialBakingOptions.Properties )
-									{
-										FString PropertyString = PropertyEnum->GetNameByValue( PropertyEntry.Property ).ToString();
-										PropertyString.RemoveFromStart( TEXT( "MP_" ) );
-										BakedPropertiesString += PropertyString + TEXT( ", " );
-									}
-
-									BakedPropertiesString.RemoveFromEnd( TEXT( ", " ) );
-								}
-
-								EventAttributes.Emplace( TEXT( "RemoveUnrealMaterials" ), Options->PreviewMeshOptions.bRemoveUnrealMaterials );
-								EventAttributes.Emplace( TEXT( "BakedProperties" ), BakedPropertiesString );
-								EventAttributes.Emplace( TEXT( "DefaultTextureSize" ), Options->PreviewMeshOptions.MaterialBakingOptions.DefaultTextureSize.ToString() );
-							}
-							EventAttributes.Emplace( TEXT( "LowestMeshLOD" ), LexToString( Options->PreviewMeshOptions.LowestMeshLOD ) );
-							EventAttributes.Emplace( TEXT( "HighestMeshLOD" ), LexToString( Options->PreviewMeshOptions.HighestMeshLOD ) );
-						}
-					}
+		TArray<FAnalyticsEventAttribute> EventAttributes;
+		EventAttributes.Emplace( TEXT( "AssetType" ), ClassName );
 
-					IUsdClassesModule::SendAnalytics( MoveTemp( EventAttributes ), FString::Printf( TEXT( "Export.%s" ), *ClassName ), bAutomated, ElapsedSeconds, NumberOfFrames, Extension );
-				}
-			}
+		if ( Options )
+		{
+			UsdUtils::AddAnalyticsAttributes( *Options, EventAttributes );
 		}
+
+		IUsdClassesModule::SendAnalytics(
+			MoveTemp( EventAttributes ),
+			FString::Printf( TEXT( "Export.%s" ), *ClassName ),
+			bAutomated,
+			ElapsedSeconds,
+			NumberOfFrames,
+			Extension
+		);
 	}
 }
 
diff --git a/Engine/Plugins/Importers/USDImporter/Source/USDExporter/Private/AnimSequenceExporterUSDOptions.cpp b/Engine/Plugins/Importers/USDImporter/Source/USDExporter/Private/AnimSequenceExporterUSDOptions.cpp
new file mode 100644
--- /dev/null
+++ b/Engine/Plugins/Importers/USDImporter/Source/USDExporter/Private/AnimSequenceExporterUSDOptions.cpp
@@ -0,0 +1,19 @@
+// Copyright Epic Games, Inc. All Rights Reserved.
+
+#include "AnimSequenceExporterUSDOptions.h"
+
+#include "AnalyticsEventAttribute.h"
+
+void UsdUtils::AddAnalyticsAttributes(
+	const UAnimSequenceExporterUSDOptions& Options,
+	TArray< FAnalyticsEventAttribute >& InOutAttributes
+)
+{
+	UsdUtils::AddAnalyticsAttributes( Options.StageOptions, InOutAttributes );
+	InOutAttributes.Emplace( TEXT( "ExportPreviewMesh" ), LexToString( Options.bExportPreviewMesh ) );
+	if ( Options.bExportPreviewMesh )
+	{
+		UsdUtils::AddAnalyticsAttributes( Options.PreviewMeshOptions, InOutAttributes );
+	}
+	InOutAttributes.Emplace( TEXT( "ReExportIdenticalAssets" ), Options.bReExportIdenticalAssets );
+}
diff --git a/Engine/Plugins/Importers/USDImporter/Source/USDExporter/Private/LevelExporterUSDOptions.cpp b/Engine/Plugins/Importers/USDImporter/Source/USDExporter/Private/LevelExporterUSDOptions.cpp
new file mode 100644
--- /dev/null
+++ b/Engine/Plugins/Importers/USDImporter/Source/USDExporter/Private/LevelExporterUSDOptions.cpp
@@ -0,0 +1,49 @@
+// Copyright Epic Games, Inc. All Rights Reserved.
+
+#include "LevelExporterUSDOptions.h"
+
+#include "USDAssetOptions.h"
+
+#include "AnalyticsEventAttribute.h"
+
+TArray<FString> ULevelExporterUSDOptions::GetUsdExtensions()
+{
+	TArray<FString> Extensions = UnrealUSDWrapper::GetNativeFileFormats();
+	Extensions.Remove( TEXT( "usdz" ) );
+	return Extensions;
+}
+
+void UsdUtils::AddAnalyticsAttributes(
+	const FLevelExporterUSDOptionsInner& Options,
+	TArray< FAnalyticsEventAttribute >& InOutAttributes
+)
+{
+	InOutAttributes.Emplace( TEXT( "SelectionOnly" ), Options.bSelectionOnly );
+	InOutAttributes.Emplace( TEXT( "ExportActorFolders" ), Options.bExportActorFolders );
+	InOutAttributes.Emplace( TEXT( "IgnoreSequencerAnimations" ), Options.bIgnoreSequencerAnimations );
+	InOutAttributes.Emplace( TEXT( "ExportFoliageOnActorsLayer" ), Options.bExportFoliageOnActorsLayer );
+	UsdUtils::AddAnalyticsAttributes( Options.AssetOptions, InOutAttributes );
+	InOutAttributes.Emplace( TEXT( "LowestLandscapeLOD" ), LexToString( Options.LowestLandscapeLOD ) );
+	InOutAttributes.Emplace( TEXT( "HighestLandscapeLOD" ), LexToString( Options.HighestLandscapeLOD ) );
+	InOutAttributes.Emplace( TEXT( "LandscapeBakeResolution" ), Options.LandscapeBakeResolution.ToString() );
+	InOutAttributes.Emplace( TEXT( "ExportSublayers" ), LexToString( Options.bExportSublayers ) );
+	InOutAttributes.Emplace( TEXT( "NumLevelsToIgnore" ), LexToString( Options.LevelsToIgnore.Num() ) );
+}
+
+void UsdUtils::AddAnalyticsAttributes(
+	const ULevelExporterUSDOptions& Options,
+	TArray< FAnalyticsEventAttribute >& InOutAttributes
+)
+{
+	UsdUtils::AddAnalyticsAttributes( Options.StageOptions, InOutAttributes );
+	InOutAttributes.Emplace( TEXT( "StartTimeCode" ), LexToString( Options.StartTimeCode ) );
+	InOutAttributes.Emplace( TEXT( "EndTimeCode" ), LexToString( Options.EndTimeCode ) );
+	UsdUtils::AddAnalyticsAttributes( Options.Inner, InOutAttributes );
+	InOutAttributes.Emplace( TEXT( "EndTimeCode" ), LexToString( Options.EndTimeCode ) );
+	InOutAttributes.Emplace(
+		TEXT( "ReExportIdenticalLevelsAndSequences" ),
+		Options.bReExportIdenticalLevelsAndSequences
+	);
+	InOutAttributes.Emplace( TEXT( "ReExportIdenticalAssets" ), Options.bReExportIdenticalAssets );
+}
+
diff --git a/Engine/Plugins/Importers/USDImporter/Source/USDExporter/Private/LevelSequenceExporterUSD.cpp b/Engine/Plugins/Importers/USDImporter/Source/USDExporter/Private/LevelSequenceExporterUSD.cpp
--- a/Engine/Plugins/Importers/USDImporter/Source/USDExporter/Private/LevelSequenceExporterUSD.cpp
+++ b/Engine/Plugins/Importers/USDImporter/Source/USDExporter/Private/LevelSequenceExporterUSD.cpp
@@ -2,6 +2,7 @@
 
 #include "LevelSequenceExporterUSD.h"
 
+#include "LevelExporterUSDOptions.h"
 #include "LevelSequenceExporterUSDOptions.h"
 #include "UnrealUSDWrapper.h"
 #include "USDClassesModule.h"
@@ -58,96 +59,32 @@ namespace UE
 				const FString& Extension
 			)
 			{
-				if ( !LevelSequence || ExportedStages.Num() == 0 )
+				if ( !LevelSequence || ExportedStages.Num() == 0 || !FEngineAnalytics::IsAvailable() )
 				{
 					return;
 				}
 
-				if ( FEngineAnalytics::IsAvailable() )
-				{
-					FString ClassName = LevelSequence->GetClass()->GetName();
-
-					TArray<FAnalyticsEventAttribute> EventAttributes;
-
-					EventAttributes.Emplace( TEXT( "AssetType" ), ClassName );
-
-					if ( Options )
-					{
-						// Stage options
-						EventAttributes.Emplace( TEXT( "MetersPerUnit" ), LexToString( Options->StageOptions.MetersPerUnit ) );
-						EventAttributes.Emplace( TEXT( "UpAxis" ), Options->StageOptions.UpAxis == EUsdUpAxis::YAxis ? TEXT( "Y" ) : TEXT( "Z" ) );
-
-						// LevelSequence options
-						EventAttributes.Emplace( TEXT( "TimeCodesPerSecond" ), LexToString( Options->TimeCodesPerSecond ) );
-						EventAttributes.Emplace( TEXT( "OverrideExportRange" ), Options->bOverrideExportRange );
-						if( Options->bOverrideExportRange )
-						{
-							EventAttributes.Emplace( TEXT( "StartFrame" ), LexToString( Options->StartFrame ) );
-							EventAttributes.Emplace( TEXT( "EndFrame" ), LexToString( Options->EndFrame ) );
-						}
-						EventAttributes.Emplace( TEXT( "ExportSubsequencesAsLayers" ), Options->bExportSubsequencesAsLayers );
-						EventAttributes.Emplace( TEXT( "ExportLevel" ), Options->bExportLevel );
-						if( Options->bExportLevel )
-						{
-							EventAttributes.Emplace( TEXT( "UseExportedLevelAsSublayer" ), Options->bUseExportedLevelAsSublayer );
-						}
-						EventAttributes.Emplace( TEXT( "NumExportedLevelSequenceLayers" ), ExportedStages.Num() );
+				FString ClassName = LevelSequence->GetClass()->GetName();
 
-						if( Options->bExportLevel )
-						{
-							// Level options
-							EventAttributes.Emplace( TEXT( "SelectionOnly" ), Options->LevelExportOptions.bSelectionOnly );
-							EventAttributes.Emplace( TEXT( "ExportActorFolders" ), Options->LevelExportOptions.bExportActorFolders );
-							EventAttributes.Emplace( TEXT( "IgnoreSequencerAnimations" ), Options->LevelExportOptions.bIgnoreSequencerAnimations );
-							EventAttributes.Emplace( TEXT( "ExportFoliageOnActorsLayer" ), Options->LevelExportOptions.bExportFoliageOnActorsLayer );
-							EventAttributes.Emplace( TEXT( "LowestLandscapeLOD" ), LexToString( Options->LevelExportOptions.LowestLandscapeLOD ) );
-							EventAttributes.Emplace( TEXT( "HighestLandscapeLOD" ), LexToString( Options->LevelExportOptions.HighestLandscapeLOD ) );
-							EventAttributes.Emplace( TEXT( "LandscapeBakeResolution" ), Options->LevelExportOptions.LandscapeBakeResolution.ToString() );
-							EventAttributes.Emplace( TEXT( "ExportSublayers" ), LexToString( Options->LevelExportOptions.bExportSublayers ) );
-							EventAttributes.Emplace( TEXT( "NumLevelsToIgnore" ), LexToString( Options->LevelExportOptions.LevelsToIgnore.Num() ) );
-						}
-
-						// Asset options
-						EventAttributes.Emplace( TEXT( "UsePayload" ), LexToString( Options->LevelExportOptions.AssetOptions.bUsePayload ) );
-						if( Options->LevelExportOptions.AssetOptions.bUsePayload )
-						{
-							EventAttributes.Emplace( TEXT( "PayloadFormat" ), Options->LevelExportOptions.AssetOptions.PayloadFormat );
-						}
-						EventAttributes.Emplace( TEXT( "BakeMaterials" ), Options->LevelExportOptions.AssetOptions.bBakeMaterials );
-						EventAttributes.Emplace( TEXT( "LowestMeshLOD" ), LexToString( Options->LevelExportOptions.AssetOptions.LowestMeshLOD ) );
-						EventAttributes.Emplace( TEXT( "HighestMeshLOD" ), LexToString( Options->LevelExportOptions.AssetOptions.HighestMeshLOD ) );
+				TArray<FAnalyticsEventAttribute> EventAttributes;
+				EventAttributes.Emplace( TEXT( "AssetType" ), ClassName );
+				EventAttributes.Emplace( TEXT( "NumExportedLevelSequenceLayers" ), ExportedStages.Num() );
 
-						// Material baking options
-						if( Options->LevelExportOptions.AssetOptions.bBakeMaterials )
-						{
-							FString BakedPropertiesString;
-							{
-								const UEnum* PropertyEnum = StaticEnum<EMaterialProperty>();
-								for ( const FPropertyEntry& PropertyEntry : Options->LevelExportOptions.AssetOptions.MaterialBakingOptions.Properties )
-								{
-									FString PropertyString = PropertyEnum->GetNameByValue( PropertyEntry.Property ).ToString();
-									PropertyString.RemoveFromStart( TEXT( "MP_" ) );
-									BakedPropertiesString += PropertyString + TEXT( ", " );
-								}
-
-								BakedPropertiesString.RemoveFromEnd( TEXT( ", " ) );
-							}
-
-							EventAttributes.Emplace( TEXT( "RemoveUnrealMaterials" ), Options->LevelExportOptions.AssetOptions.bRemoveUnrealMaterials );
-							EventAttributes.Emplace( TEXT( "BakedProperties" ), BakedPropertiesString );
-							EventAttributes.Emplace( TEXT( "DefaultTextureSize" ), Options->LevelExportOptions.AssetOptions.MaterialBakingOptions.DefaultTextureSize.ToString() );
-						}
-
-						IUsdClassesModule::SendAnalytics(
-							MoveTemp( EventAttributes ),
-							FString::Printf( TEXT( "Export.%s" ), *ClassName ),
-							bAutomated,
-							ElapsedSeconds,
-							( Options->EndFrame - Options->StartFrame ),
-							Extension
-						);
-					}
+				int32 NumFrames = 0;
+				if ( Options )
+				{
+					NumFrames = Options->EndFrame - Options->StartFrame;
+					UsdUtils::AddAnalyticsAttributes( *Options, EventAttributes );
 				}
+
+				IUsdClassesModule::SendAnalytics(
+					MoveTemp( EventAttributes ),
+					FString::Printf( TEXT( "Export.%s" ), *ClassName ),
+					bAutomated,
+					ElapsedSeconds,
+					NumFrames,
+					Extension
+				);
 			}
 
 #if USE_USD_SDK
diff --git a/Engine/Plugins/Importers/USDImporter/Source/USDExporter/Private/LevelSequenceExporterUSDOptions.cpp b/Engine/Plugins/Importers/USDImporter/Source/USDExporter/Private/LevelSequenceExporterUSDOptions.cpp
new file mode 100644
--- /dev/null
+++ b/Engine/Plugins/Importers/USDImporter/Source/USDExporter/Private/LevelSequenceExporterUSDOptions.cpp
@@ -0,0 +1,36 @@
+// Copyright Epic Games, Inc. All Rights Reserved.
+
+#include "LevelSequenceExporterUSDOptions.h"
+
+#include "LevelExporterUSDOptions.h"
+
+#include "AnalyticsEventAttribute.h"
+
+void UsdUtils::AddAnalyticsAttributes(
+	const ULevelSequenceExporterUsdOptions& Options,
+	TArray< struct FAnalyticsEventAttribute >& InOutAttributes
+)
+{
+	InOutAttributes.Emplace( TEXT( "TimeCodesPerSecond" ), LexToString( Options.TimeCodesPerSecond ) );
+	InOutAttributes.Emplace( TEXT( "OverrideExportRange" ), Options.bOverrideExportRange );
+	if ( Options.bOverrideExportRange )
+	{
+		InOutAttributes.Emplace( TEXT( "StartFrame" ), LexToString( Options.StartFrame ) );
+		InOutAttributes.Emplace( TEXT( "EndFrame" ), LexToString( Options.EndFrame ) );
+	}
+	InOutAttributes.Emplace( TEXT( "ExportSubsequencesAsLayers" ), Options.bExportSubsequencesAsLayers );
+	InOutAttributes.Emplace( TEXT( "ExportLevel" ), Options.bExportLevel );
+	if ( Options.bExportLevel )
+	{
+		InOutAttributes.Emplace( TEXT( "UseExportedLevelAsSublayer" ), Options.bUseExportedLevelAsSublayer );
+	}
+	InOutAttributes.Emplace(
+        TEXT( "ReExportIdenticalLevelsAndSequences" ),
+        Options.bReExportIdenticalLevelsAndSequences
+    );
+	InOutAttributes.Emplace( TEXT( "ReExportIdenticalAssets" ), Options.bReExportIdenticalAssets );
+	if ( Options.bExportLevel )
+	{
+		UsdUtils::AddAnalyticsAttributes( Options.LevelExportOptions, InOutAttributes );
+	}
+}
\ No newline at end of file
diff --git a/Engine/Plugins/Importers/USDImporter/Source/USDExporter/Private/MaterialExporterUSD.cpp b/Engine/Plugins/Importers/USDImporter/Source/USDExporter/Private/MaterialExporterUSD.cpp
--- a/Engine/Plugins/Importers/USDImporter/Source/USDExporter/Private/MaterialExporterUSD.cpp
+++ b/Engine/Plugins/Importers/USDImporter/Source/USDExporter/Private/MaterialExporterUSD.cpp
@@ -35,34 +35,35 @@
 
 namespace UE::MaterialExporterUSD::Private
 {
-	void SendAnalytics( const UMaterialInterface& Material, const FUsdMaterialBakingOptions& Options, bool bAutomated, double ElapsedSeconds, double NumberOfFrames, const FString& Extension )
+	void SendAnalytics(
+		const UMaterialInterface& Material,
+		const FUsdMaterialBakingOptions& Options,
+		bool bAutomated,
+		double ElapsedSeconds,
+		double NumberOfFrames,
+		const FString& Extension
+	)
 	{
-		if ( FEngineAnalytics::IsAvailable() )
+		if ( !FEngineAnalytics::IsAvailable() )
 		{
-			FString ClassName = Material.GetClass()->GetName();
+			return;
+		}
 
-			TArray<FAnalyticsEventAttribute> EventAttributes;
+		FString ClassName = Material.GetClass()->GetName();
 
-			EventAttributes.Emplace( TEXT( "AssetType" ), ClassName );
+		TArray<FAnalyticsEventAttribute> EventAttributes;
+		EventAttributes.Emplace( TEXT( "AssetType" ), ClassName );
 
-			FString BakedPropertiesString;
-			{
-				const UEnum* PropertyEnum = StaticEnum<EMaterialProperty>();
-				for ( const FPropertyEntry& PropertyEntry : Options.Properties )
-				{
-					FString PropertyString = PropertyEnum->GetNameByValue( PropertyEntry.Property ).ToString();
-					PropertyString.RemoveFromStart( TEXT( "MP_" ) );
-					BakedPropertiesString += PropertyString + TEXT( ", " );
-				}
+		UsdUtils::AddAnalyticsAttributes( Options, EventAttributes );
 
-				BakedPropertiesString.RemoveFromEnd( TEXT( ", " ) );
-			}
-
-			EventAttributes.Emplace( TEXT( "BakedProperties" ), BakedPropertiesString );
-			EventAttributes.Emplace( TEXT( "DefaultTextureSize" ), Options.DefaultTextureSize.ToString() );
-
-			IUsdClassesModule::SendAnalytics( MoveTemp( EventAttributes ), FString::Printf( TEXT( "Export.%s" ), *ClassName ), bAutomated, ElapsedSeconds, NumberOfFrames, Extension );
-		}
+		IUsdClassesModule::SendAnalytics(
+			MoveTemp( EventAttributes ),
+			FString::Printf( TEXT( "Export.%s" ), *ClassName ),
+			bAutomated,
+			ElapsedSeconds,
+			NumberOfFrames,
+			Extension
+		);
 	}
 
 	FString GetMaterialHashString( const UMaterialInterface& Material )
@@ -248,7 +249,8 @@ bool UMaterialExporterUsd::ExportBinary( UObject* Object, const TCHAR* Type, FAr
 		Options->MaterialBakingOptions,
 		FFilePath{ UExporter::CurrentFilename },
 		ExportTask->bReplaceIdentical,
-		Options->bReExportIdenticalAssets
+		Options->bReExportIdenticalAssets,
+		ExportTask->bAutomated
 	);
 #else
 	return false;
@@ -260,7 +262,8 @@ bool UMaterialExporterUsd::ExportMaterial(
 	const FUsdMaterialBakingOptions& Options,
 	const FFilePath& FilePath,
 	bool bReplaceIdentical,
-	bool bReExportIdenticalAssets
+	bool bReExportIdenticalAssets,
+	bool bIsAutomated
 )
 {
 #if USE_USD_SDK
@@ -357,12 +360,18 @@ bool UMaterialExporterUsd::ExportMaterial(
 
 	// Analytics
 	{
-		bool bAutomated = false;  // TODO: Can't properly fetch this yet as it would change the function signature
 		double ElapsedSeconds = FPlatformTime::ToSeconds64( FPlatformTime::Cycles64() - StartTime );
 		FString Extension = FPaths::GetExtension( FilePath.FilePath );
 		double NumberOfFrames = 1 + UsdStage.GetEndTimeCode() - UsdStage.GetStartTimeCode();
 
-		UE::MaterialExporterUSD::Private::SendAnalytics( Material, Options, bAutomated, ElapsedSeconds, NumberOfFrames, Extension );
+		UE::MaterialExporterUSD::Private::SendAnalytics(
+			Material,
+			Options,
+			bIsAutomated,
+			ElapsedSeconds,
+			NumberOfFrames,
+			Extension
+		);
 	}
 
 	return true;
@@ -379,7 +388,8 @@ bool UMaterialExporterUsd::ExportMaterialsForStage(
 	bool bUsePayload,
 	bool bRemoveUnrealMaterials,
 	bool bReplaceIdentical,
-	bool bReExportIdenticalAssets
+	bool bReExportIdenticalAssets,
+	bool bIsAutomated
 )
 {
 #if USE_USD_SDK
@@ -438,7 +448,8 @@ bool UMaterialExporterUsd::ExportMaterialsForStage(
 			Options,
 			FFilePath{ FinalPath },
 			bReplaceIdentical,
-			bReExportIdenticalAssets
+			bReExportIdenticalAssets,
+			bIsAutomated
 		))
 		{
 			UsedFilePathsWithoutExt.Add( FinalPathNoExt );
diff --git a/Engine/Plugins/Importers/USDImporter/Source/USDExporter/Private/MaterialExporterUSDOptions.cpp b/Engine/Plugins/Importers/USDImporter/Source/USDExporter/Private/MaterialExporterUSDOptions.cpp
new file mode 100644
--- /dev/null
+++ b/Engine/Plugins/Importers/USDImporter/Source/USDExporter/Private/MaterialExporterUSDOptions.cpp
@@ -0,0 +1,14 @@
+// Copyright Epic Games, Inc. All Rights Reserved.
+
+#include "MaterialExporterUSDOptions.h"
+
+#include "AnalyticsEventAttribute.h"
+
+void UsdUtils::AddAnalyticsAttributes(
+	const UMaterialExporterUSDOptions& Options,
+	TArray< FAnalyticsEventAttribute >& InOutAttributes
+)
+{
+	UsdUtils::AddAnalyticsAttributes( Options.MaterialBakingOptions, InOutAttributes );
+	InOutAttributes.Emplace( TEXT( "ReExportIdenticalAssets" ), Options.bReExportIdenticalAssets );
+}
\ No newline at end of file
diff --git a/Engine/Plugins/Importers/USDImporter/Source/USDExporter/Private/SkeletalMeshExporterUSD.cpp b/Engine/Plugins/Importers/USDImporter/Source/USDExporter/Private/SkeletalMeshExporterUSD.cpp
--- a/Engine/Plugins/Importers/USDImporter/Source/USDExporter/Private/SkeletalMeshExporterUSD.cpp
+++ b/Engine/Plugins/Importers/USDImporter/Source/USDExporter/Private/SkeletalMeshExporterUSD.cpp
@@ -22,64 +22,40 @@
 #include "AssetExportTask.h"
 #include "Engine/SkeletalMesh.h"
 
-namespace UE
+namespace UE::SkeletalMeshExporterUSD::Private
 {
-	namespace SkeletalMeshExporterUSD
+	void SendAnalytics(
+		UObject* Asset,
+		USkeletalMeshExporterUSDOptions* Options,
+		bool bAutomated,
+		double ElapsedSeconds,
+		double NumberOfFrames,
+		const FString& Extension
+	)
 	{
-		namespace Private
+		if ( !Asset || !FEngineAnalytics::IsAvailable() )
 		{
-			void SendAnalytics( UObject* Asset, USkeletalMeshExporterUSDOptions* Options, bool bAutomated, double ElapsedSeconds, double NumberOfFrames, const FString& Extension )
-			{
-				if ( !Asset )
-				{
-					return;
-				}
-
-				if ( FEngineAnalytics::IsAvailable() )
-				{
-					FString ClassName = Asset->GetClass()->GetName();
+			return;
+		}
 
-					TArray<FAnalyticsEventAttribute> EventAttributes;
+		FString ClassName = Asset->GetClass()->GetName();
 
-					EventAttributes.Emplace( TEXT( "AssetType" ), ClassName );
+		TArray<FAnalyticsEventAttribute> EventAttributes;
+		EventAttributes.Emplace( TEXT( "AssetType" ), ClassName );
 
-					if ( Options )
-					{
-						EventAttributes.Emplace( TEXT( "MetersPerUnit" ), LexToString( Options->StageOptions.MetersPerUnit ) );
-						EventAttributes.Emplace( TEXT( "UpAxis" ), Options->StageOptions.UpAxis == EUsdUpAxis::YAxis ? TEXT( "Y" ) : TEXT( "Z" ) );
-						EventAttributes.Emplace( TEXT( "UsePayload" ), LexToString( Options->MeshAssetOptions.bUsePayload ) );
-						if ( Options->MeshAssetOptions.bUsePayload )
-						{
-							EventAttributes.Emplace( TEXT( "PayloadFormat" ), Options->MeshAssetOptions.PayloadFormat );
-						}
-						EventAttributes.Emplace( TEXT( "BakeMaterials" ), Options->MeshAssetOptions.bBakeMaterials );
-						if ( Options->MeshAssetOptions.bBakeMaterials )
-						{
-							FString BakedPropertiesString;
-							{
-								const UEnum* PropertyEnum = StaticEnum<EMaterialProperty>();
-								for ( const FPropertyEntry& PropertyEntry : Options->MeshAssetOptions.MaterialBakingOptions.Properties )
-								{
-									FString PropertyString = PropertyEnum->GetNameByValue( PropertyEntry.Property ).ToString();
-									PropertyString.RemoveFromStart( TEXT( "MP_" ) );
-									BakedPropertiesString += PropertyString + TEXT( ", " );
-								}
-
-								BakedPropertiesString.RemoveFromEnd( TEXT( ", " ) );
-							}
-
-							EventAttributes.Emplace( TEXT( "RemoveUnrealMaterials" ), Options->MeshAssetOptions.bRemoveUnrealMaterials );
-							EventAttributes.Emplace( TEXT( "BakedProperties" ), BakedPropertiesString );
-							EventAttributes.Emplace( TEXT( "DefaultTextureSize" ), Options->MeshAssetOptions.MaterialBakingOptions.DefaultTextureSize.ToString() );
-						}
-						EventAttributes.Emplace( TEXT( "LowestMeshLOD" ), LexToString( Options->MeshAssetOptions.LowestMeshLOD ) );
-						EventAttributes.Emplace( TEXT( "HighestMeshLOD" ), LexToString( Options->MeshAssetOptions.HighestMeshLOD ) );
-					}
-
-					IUsdClassesModule::SendAnalytics( MoveTemp( EventAttributes ), FString::Printf( TEXT( "Export.%s" ), *ClassName ), bAutomated, ElapsedSeconds, NumberOfFrames, Extension );
-				}
-			}
+		if ( Options )
+		{
+			UsdUtils::AddAnalyticsAttributes( *Options, EventAttributes );
 		}
+
+		IUsdClassesModule::SendAnalytics(
+			MoveTemp( EventAttributes ),
+			FString::Printf( TEXT( "Export.%s" ), *ClassName ),
+			bAutomated,
+			ElapsedSeconds,
+			NumberOfFrames,
+			Extension
+		);
 	}
 }
 
@@ -228,7 +204,8 @@ bool USkeletalMeshExporterUsd::ExportBinary( UObject* Object, const TCHAR* Type,
 								Options->MeshAssetOptions.bUsePayload,
 								Options->MeshAssetOptions.bRemoveUnrealMaterials,
 								ExportTask->bReplaceIdentical,
-								Options->bReExportIdenticalAssets
+								Options->bReExportIdenticalAssets,
+								ExportTask->bAutomated
 							);
 							TempStage.GetRootLayer().Save();
 						}
@@ -328,7 +305,8 @@ bool USkeletalMeshExporterUsd::ExportBinary( UObject* Object, const TCHAR* Type,
 			Options->MeshAssetOptions.bUsePayload,
 			Options->MeshAssetOptions.bRemoveUnrealMaterials,
 			ExportTask->bReplaceIdentical,
-			Options->bReExportIdenticalAssets
+			Options->bReExportIdenticalAssets,
+			ExportTask->bAutomated
 		);
 	}
 
diff --git a/Engine/Plugins/Importers/USDImporter/Source/USDExporter/Private/SkeletalMeshExporterUSDOptions.cpp b/Engine/Plugins/Importers/USDImporter/Source/USDExporter/Private/SkeletalMeshExporterUSDOptions.cpp
new file mode 100644
--- /dev/null
+++ b/Engine/Plugins/Importers/USDImporter/Source/USDExporter/Private/SkeletalMeshExporterUSDOptions.cpp
@@ -0,0 +1,15 @@
+// Copyright Epic Games, Inc. All Rights Reserved.
+
+#include "SkeletalMeshExporterUSDOptions.h"
+
+#include "AnalyticsEventAttribute.h"
+
+void UsdUtils::AddAnalyticsAttributes(
+	const USkeletalMeshExporterUSDOptions& Options,
+	TArray< FAnalyticsEventAttribute >& InOutAttributes
+)
+{
+	UsdUtils::AddAnalyticsAttributes( Options.StageOptions, InOutAttributes );
+	UsdUtils::AddAnalyticsAttributes( Options.MeshAssetOptions, InOutAttributes );
+	InOutAttributes.Emplace( TEXT( "ReExportIdenticalAssets" ), Options.bReExportIdenticalAssets );
+}
diff --git a/Engine/Plugins/Importers/USDImporter/Source/USDExporter/Private/StaticMeshExporterUSD.cpp b/Engine/Plugins/Importers/USDImporter/Source/USDExporter/Private/StaticMeshExporterUSD.cpp
--- a/Engine/Plugins/Importers/USDImporter/Source/USDExporter/Private/StaticMeshExporterUSD.cpp
+++ b/Engine/Plugins/Importers/USDImporter/Source/USDExporter/Private/StaticMeshExporterUSD.cpp
@@ -23,64 +23,40 @@
 #include "CoreMinimal.h"
 #include "Engine/StaticMesh.h"
 
-namespace UE
+namespace UE::StaticMeshExporterUSD::Private
 {
-	namespace StaticMeshExporterUSD
+	void SendAnalytics(
+		UObject* Asset,
+		UStaticMeshExporterUSDOptions* Options,
+		bool bAutomated,
+		double ElapsedSeconds,
+		double NumberOfFrames,
+		const FString& Extension
+	)
 	{
-		namespace Private
+		if ( !Asset || !FEngineAnalytics::IsAvailable() )
 		{
-			void SendAnalytics( UObject* Asset, UStaticMeshExporterUSDOptions* Options, bool bAutomated, double ElapsedSeconds, double NumberOfFrames, const FString& Extension )
-			{
-				if ( !Asset )
-				{
-					return;
-				}
-
-				if ( FEngineAnalytics::IsAvailable() )
-				{
-					FString ClassName = Asset->GetClass()->GetName();
+			return;
+		}
 
-					TArray<FAnalyticsEventAttribute> EventAttributes;
+		FString ClassName = Asset->GetClass()->GetName();
 
-					EventAttributes.Emplace( TEXT( "AssetType" ), ClassName );
+		TArray<FAnalyticsEventAttribute> EventAttributes;
+		EventAttributes.Emplace( TEXT( "AssetType" ), ClassName );
 
-					if ( Options )
-					{
-						EventAttributes.Emplace( TEXT( "MetersPerUnit" ), LexToString( Options->StageOptions.MetersPerUnit ) );
-						EventAttributes.Emplace( TEXT( "UpAxis" ), Options->StageOptions.UpAxis == EUsdUpAxis::YAxis ? TEXT( "Y" ) : TEXT( "Z" ) );
-						EventAttributes.Emplace( TEXT( "UsePayload" ), LexToString( Options->MeshAssetOptions.bUsePayload ) );
-						if ( Options->MeshAssetOptions.bUsePayload )
-						{
-							EventAttributes.Emplace( TEXT( "PayloadFormat" ), Options->MeshAssetOptions.PayloadFormat );
-						}
-						EventAttributes.Emplace( TEXT( "BakeMaterials" ), Options->MeshAssetOptions.bBakeMaterials );
-						if ( Options->MeshAssetOptions.bBakeMaterials )
-						{
-							FString BakedPropertiesString;
-							{
-								const UEnum* PropertyEnum = StaticEnum<EMaterialProperty>();
-								for ( const FPropertyEntry& PropertyEntry : Options->MeshAssetOptions.MaterialBakingOptions.Properties )
-								{
-									FString PropertyString = PropertyEnum->GetNameByValue( PropertyEntry.Property ).ToString();
-									PropertyString.RemoveFromStart( TEXT( "MP_" ) );
-									BakedPropertiesString += PropertyString + TEXT( ", " );
-								}
-
-								BakedPropertiesString.RemoveFromEnd( TEXT( ", " ) );
-							}
-
-							EventAttributes.Emplace( TEXT( "RemoveUnrealMaterials" ), Options->MeshAssetOptions.bRemoveUnrealMaterials );
-							EventAttributes.Emplace( TEXT( "BakedProperties" ), BakedPropertiesString );
-							EventAttributes.Emplace( TEXT( "DefaultTextureSize" ), Options->MeshAssetOptions.MaterialBakingOptions.DefaultTextureSize.ToString() );
-						}
-						EventAttributes.Emplace( TEXT( "LowestMeshLOD" ), LexToString( Options->MeshAssetOptions.LowestMeshLOD ) );
-						EventAttributes.Emplace( TEXT( "HighestMeshLOD" ), LexToString( Options->MeshAssetOptions.HighestMeshLOD ) );
-					}
-
-					IUsdClassesModule::SendAnalytics( MoveTemp( EventAttributes ), FString::Printf( TEXT( "Export.%s" ), *ClassName ), bAutomated, ElapsedSeconds, NumberOfFrames, Extension );
-				}
-			}
+		if ( Options )
+		{
+			UsdUtils::AddAnalyticsAttributes( *Options, EventAttributes );
 		}
+
+		IUsdClassesModule::SendAnalytics(
+			MoveTemp( EventAttributes ),
+			FString::Printf( TEXT( "Export.%s" ), *ClassName ),
+			bAutomated,
+			ElapsedSeconds,
+			NumberOfFrames,
+			Extension
+		);
 	}
 }
 
@@ -243,7 +219,8 @@ bool UStaticMeshExporterUsd::ExportBinary( UObject* Object, const TCHAR* Type, F
 								Options->MeshAssetOptions.bUsePayload,
 								Options->MeshAssetOptions.bRemoveUnrealMaterials,
 								ExportTask->bReplaceIdentical,
-								Options->bReExportIdenticalAssets
+								Options->bReExportIdenticalAssets,
+								ExportTask->bAutomated
 							);
 							TempStage.GetRootLayer().Save();
 						}
@@ -339,7 +316,8 @@ bool UStaticMeshExporterUsd::ExportBinary( UObject* Object, const TCHAR* Type, F
 			Options->MeshAssetOptions.bUsePayload,
 			Options->MeshAssetOptions.bRemoveUnrealMaterials,
 			ExportTask->bReplaceIdentical,
-			Options->bReExportIdenticalAssets
+			Options->bReExportIdenticalAssets,
+			ExportTask->bAutomated
 		);
 	}
 
diff --git a/Engine/Plugins/Importers/USDImporter/Source/USDExporter/Private/StaticMeshExporterUSDOptions.cpp b/Engine/Plugins/Importers/USDImporter/Source/USDExporter/Private/StaticMeshExporterUSDOptions.cpp
new file mode 100644
--- /dev/null
+++ b/Engine/Plugins/Importers/USDImporter/Source/USDExporter/Private/StaticMeshExporterUSDOptions.cpp
@@ -0,0 +1,15 @@
+// Copyright Epic Games, Inc. All Rights Reserved.
+
+#include "StaticMeshExporterUSDOptions.h"
+
+#include "AnalyticsEventAttribute.h"
+
+void UsdUtils::AddAnalyticsAttributes(
+	const UStaticMeshExporterUSDOptions& Options,
+	TArray< FAnalyticsEventAttribute >& InOutAttributes
+)
+{
+	UsdUtils::AddAnalyticsAttributes( Options.StageOptions, InOutAttributes );
+	UsdUtils::AddAnalyticsAttributes( Options.MeshAssetOptions, InOutAttributes );
+	InOutAttributes.Emplace( TEXT( "ReExportIdenticalAssets" ), Options.bReExportIdenticalAssets );
+}
\ No newline at end of file
diff --git a/Engine/Plugins/Importers/USDImporter/Source/USDExporter/Private/USDAssetOptions.cpp b/Engine/Plugins/Importers/USDImporter/Source/USDExporter/Private/USDAssetOptions.cpp
new file mode 100644
--- /dev/null
+++ b/Engine/Plugins/Importers/USDImporter/Source/USDExporter/Private/USDAssetOptions.cpp
@@ -0,0 +1,47 @@
+// Copyright Epic Games, Inc. All Rights Reserved.
+
+#include "USDAssetOptions.h"
+
+#include "AnalyticsEventAttribute.h"
+
+void UsdUtils::AddAnalyticsAttributes(
+	const FUsdMaterialBakingOptions& Options,
+	TArray< FAnalyticsEventAttribute >& InOutAttributes
+)
+{
+	FString BakedPropertiesString;
+	{
+		const UEnum* PropertyEnum = StaticEnum<EMaterialProperty>();
+		for ( const FPropertyEntry& PropertyEntry : Options.Properties )
+		{
+			FString PropertyString = PropertyEnum->GetNameByValue( PropertyEntry.Property ).ToString();
+			PropertyString.RemoveFromStart( TEXT( "MP_" ) );
+			BakedPropertiesString += PropertyString + TEXT( ", " );
+		}
+
+		BakedPropertiesString.RemoveFromEnd( TEXT( ", " ) );
+	}
+
+	InOutAttributes.Emplace( TEXT( "BakedProperties" ), BakedPropertiesString );
+	InOutAttributes.Emplace( TEXT( "DefaultTextureSize" ), Options.DefaultTextureSize.ToString() );
+}
+
+void UsdUtils::AddAnalyticsAttributes(
+	const FUsdMeshAssetOptions& Options,
+	TArray< FAnalyticsEventAttribute >& InOutAttributes
+)
+{
+	InOutAttributes.Emplace( TEXT( "UsePayload" ), LexToString( Options.bUsePayload ) );
+	if ( Options.bUsePayload )
+	{
+		InOutAttributes.Emplace( TEXT( "PayloadFormat" ), Options.PayloadFormat );
+	}
+	InOutAttributes.Emplace( TEXT( "BakeMaterials" ), Options.bBakeMaterials );
+	InOutAttributes.Emplace( TEXT( "RemoveUnrealMaterials" ), Options.bRemoveUnrealMaterials );
+	if ( Options.bBakeMaterials )
+	{
+		UsdUtils::AddAnalyticsAttributes( Options.MaterialBakingOptions, InOutAttributes );
+	}
+	InOutAttributes.Emplace( TEXT( "LowestMeshLOD" ), LexToString( Options.LowestMeshLOD ) );
+	InOutAttributes.Emplace( TEXT( "HighestMeshLOD" ), LexToString( Options.HighestMeshLOD ) );
+}
\ No newline at end of file
diff --git a/Engine/Plugins/Importers/USDImporter/Source/USDExporter/Private/USDConversionBlueprintLibrary.cpp b/Engine/Plugins/Importers/USDImporter/Source/USDExporter/Private/USDConversionBlueprintLibrary.cpp
--- a/Engine/Plugins/Importers/USDImporter/Source/USDExporter/Private/USDConversionBlueprintLibrary.cpp
+++ b/Engine/Plugins/Importers/USDImporter/Source/USDExporter/Private/USDConversionBlueprintLibrary.cpp
@@ -3,6 +3,7 @@
 #include "USDConversionBlueprintLibrary.h"
 
 #include "LevelExporterUSD.h"
+#include "LevelExporterUSDOptions.h"
 #include "UnrealUSDWrapper.h"
 #include "USDClassesModule.h"
 #include "USDConversionUtils.h"
@@ -463,6 +464,25 @@ TArray<FTransform> UUsdConversionBlueprintLibrary::GetInstanceTransforms( AInsta
 	return Result;
 }
 
+TArray<FAnalyticsEventAttr> UUsdConversionBlueprintLibrary::GetAnalyticsAttributes( const ULevelExporterUSDOptions* Options )
+{
+	TArray<FAnalyticsEventAttr> Attrs;
+	if ( Options )
+	{
+		TArray<FAnalyticsEventAttribute> Attributes;
+		UsdUtils::AddAnalyticsAttributes( *Options, Attributes );
+
+		Attrs.Reserve( Attributes.Num() );
+		for(const FAnalyticsEventAttribute& Attribute : Attributes )
+		{
+			FAnalyticsEventAttr& NewAttr = Attrs.Emplace_GetRef();
+			NewAttr.Name = Attribute.GetName();
+			NewAttr.Value = Attribute.GetValue();
+		}
+	}
+	return Attrs;
+}
+
 void UUsdConversionBlueprintLibrary::SendAnalytics( const TArray<FAnalyticsEventAttr>& Attrs, const FString& EventName, bool bAutomated, double ElapsedSeconds, double NumberOfFrames, const FString& Extension )
 {
 	TArray<FAnalyticsEventAttribute> Converted;
diff --git a/Engine/Plugins/Importers/USDImporter/Source/USDExporter/Public/AnimSequenceExporterUSDOptions.h b/Engine/Plugins/Importers/USDImporter/Source/USDExporter/Public/AnimSequenceExporterUSDOptions.h
--- a/Engine/Plugins/Importers/USDImporter/Source/USDExporter/Public/AnimSequenceExporterUSDOptions.h
+++ b/Engine/Plugins/Importers/USDImporter/Source/USDExporter/Public/AnimSequenceExporterUSDOptions.h
@@ -9,6 +9,8 @@
 
 #include "AnimSequenceExporterUSDOptions.generated.h"
 
+struct FAnalyticsEventAttribute;
+
 /**
  * Options for exporting skeletal mesh animations to USD format.
  */
@@ -37,3 +39,11 @@ public:
 	UPROPERTY( EditAnywhere, config, BlueprintReadWrite, Category = "Collision", meta = ( DisplayName = "Re-export Identical Assets" ) )
 	bool bReExportIdenticalAssets = false;
 };
+
+namespace UsdUtils
+{
+	USDEXPORTER_API void AddAnalyticsAttributes(
+		const UAnimSequenceExporterUSDOptions& Options,
+		TArray< FAnalyticsEventAttribute >& InOutAttributes
+	);
+}
diff --git a/Engine/Plugins/Importers/USDImporter/Source/USDExporter/Public/LevelExporterUSDOptions.h b/Engine/Plugins/Importers/USDImporter/Source/USDExporter/Public/LevelExporterUSDOptions.h
--- a/Engine/Plugins/Importers/USDImporter/Source/USDExporter/Public/LevelExporterUSDOptions.h
+++ b/Engine/Plugins/Importers/USDImporter/Source/USDExporter/Public/LevelExporterUSDOptions.h
@@ -12,6 +12,8 @@
 
 #include "LevelExporterUSDOptions.generated.h"
 
+struct FAnalyticsEventAttribute;
+
 USTRUCT( BlueprintType )
 struct USDEXPORTER_API FLevelExporterUSDOptionsInner
 {
@@ -118,10 +120,17 @@ public:
 
 private:
 	UFUNCTION()
-	static TArray<FString> GetUsdExtensions()
-	{
-		TArray<FString> Extensions = UnrealUSDWrapper::GetNativeFileFormats();
-		Extensions.Remove( TEXT( "usdz" ) );
-		return Extensions;
-	}
+	static TArray<FString> GetUsdExtensions();
 };
+
+namespace UsdUtils
+{
+	USDEXPORTER_API void AddAnalyticsAttributes(
+		const FLevelExporterUSDOptionsInner& Options,
+		TArray< FAnalyticsEventAttribute >& InOutAttributes
+	);
+	USDEXPORTER_API void AddAnalyticsAttributes(
+		const ULevelExporterUSDOptions& Options,
+		TArray< FAnalyticsEventAttribute >& InOutAttributes
+	);
+}
\ No newline at end of file
diff --git a/Engine/Plugins/Importers/USDImporter/Source/USDExporter/Public/LevelSequenceExporterUSDOptions.h b/Engine/Plugins/Importers/USDImporter/Source/USDExporter/Public/LevelSequenceExporterUSDOptions.h
--- a/Engine/Plugins/Importers/USDImporter/Source/USDExporter/Public/LevelSequenceExporterUSDOptions.h
+++ b/Engine/Plugins/Importers/USDImporter/Source/USDExporter/Public/LevelSequenceExporterUSDOptions.h
@@ -10,6 +10,8 @@
 
 #include "LevelSequenceExporterUSDOptions.generated.h"
 
+struct FAnalyticsEventAttribute;
+
 /**
  * Options for level sequences to the USD format.
  */
@@ -83,3 +85,11 @@ public:
 	UPROPERTY( EditAnywhere, config, BlueprintReadWrite, Category = "Level Export", meta = ( EditCondition = "bExportLevel" ) )
 	FLevelExporterUSDOptionsInner LevelExportOptions;
 };
+
+namespace UsdUtils
+{
+	USDEXPORTER_API void AddAnalyticsAttributes(
+		const ULevelSequenceExporterUsdOptions& Options,
+		TArray< FAnalyticsEventAttribute >& InOutAttributes
+	);
+}
\ No newline at end of file
diff --git a/Engine/Plugins/Importers/USDImporter/Source/USDExporter/Public/MaterialExporterUSD.h b/Engine/Plugins/Importers/USDImporter/Source/USDExporter/Public/MaterialExporterUSD.h
--- a/Engine/Plugins/Importers/USDImporter/Source/USDExporter/Public/MaterialExporterUSD.h
+++ b/Engine/Plugins/Importers/USDImporter/Source/USDExporter/Public/MaterialExporterUSD.h
@@ -30,6 +30,7 @@ public:
 	 * @param FilePath - Filepath of the destination file (e.g. "C:/MyFolder/Red.usda")
 	 * @param ReplaceIdentical - Whether to overwrite the destination file in case it already exists
 	 * @param bReExportIdenticalAssets - Whether to overwrite the destination file even if it already describes an up-to-date version of Material
+	 * @param bIsAutomated - Whether the export is being done by a script or not. Just used for analytics
 	 * @return Whether the export was successful or not.
 	 */
 	static bool ExportMaterial(
@@ -37,7 +38,8 @@ public:
 		const FUsdMaterialBakingOptions& Options,
 		const FFilePath& FilePath,
 		bool bReplaceIdentical = true,
-		bool bReExportIdenticalAssets = false
+		bool bReExportIdenticalAssets = false,
+		bool bIsAutomated = false
 	);
 
 	/**
@@ -55,6 +57,7 @@ public:
 	 *                                 Important because the `unrealMaterial` attributes will be used as a higher priority when determining material assignments
 	 * @param ReplaceIdentical - Whether to overwrite the destination files in case they already exist
 	 * @param bReExportIdenticalAssets - Whether to overwrite the destination files even if they already describe up-to-date versions of Materials
+	 * @param bIsAutomated - Whether the export is being done by a script or not. Just used for analytics
 	 * @return Whether the export was successful or not.
 	 */
 	static bool ExportMaterialsForStage(
@@ -65,6 +68,7 @@ public:
 		bool bUsePayload,
 		bool bRemoveUnrealMaterials,
 		bool bReplaceIdentical = true,
-		bool bReExportIdenticalAssets = false
+		bool bReExportIdenticalAssets = false,
+		bool bIsAutomated = false
 	);
 };
diff --git a/Engine/Plugins/Importers/USDImporter/Source/USDExporter/Public/MaterialExporterUSDOptions.h b/Engine/Plugins/Importers/USDImporter/Source/USDExporter/Public/MaterialExporterUSDOptions.h
--- a/Engine/Plugins/Importers/USDImporter/Source/USDExporter/Public/MaterialExporterUSDOptions.h
+++ b/Engine/Plugins/Importers/USDImporter/Source/USDExporter/Public/MaterialExporterUSDOptions.h
@@ -8,6 +8,8 @@
 
 #include "MaterialExporterUSDOptions.generated.h"
 
+struct FAnalyticsEventAttribute;
+
 /**
  * Options for exporting materials to USD format.
  * We use a dedicated object instead of reusing the MaterialBaking module as automated export tasks
@@ -29,3 +31,11 @@ public:
 	UPROPERTY( EditAnywhere, config, BlueprintReadWrite, Category = "Collision", meta = ( DisplayName = "Re-export Identical Assets" ) )
 	bool bReExportIdenticalAssets = false;
 };
+
+namespace UsdUtils
+{
+	USDEXPORTER_API void AddAnalyticsAttributes(
+		const UMaterialExporterUSDOptions& Options,
+		TArray< FAnalyticsEventAttribute >& InOutAttributes
+	);
+}
\ No newline at end of file
diff --git a/Engine/Plugins/Importers/USDImporter/Source/USDExporter/Public/SkeletalMeshExporterUSDOptions.h b/Engine/Plugins/Importers/USDImporter/Source/USDExporter/Public/SkeletalMeshExporterUSDOptions.h
--- a/Engine/Plugins/Importers/USDImporter/Source/USDExporter/Public/SkeletalMeshExporterUSDOptions.h
+++ b/Engine/Plugins/Importers/USDImporter/Source/USDExporter/Public/SkeletalMeshExporterUSDOptions.h
@@ -9,6 +9,8 @@
 
 #include "SkeletalMeshExporterUSDOptions.generated.h"
 
+struct FAnalyticsEventAttribute;
+
 /**
  * Options for exporting skeletal meshes to USD format.
  */
@@ -31,3 +33,11 @@ public:
 	UPROPERTY( EditAnywhere, config, BlueprintReadWrite, Category = "Collision", meta = ( DisplayName = "Re-export Identical Assets" ) )
 	bool bReExportIdenticalAssets = false;
 };
+
+namespace UsdUtils
+{
+	USDEXPORTER_API void AddAnalyticsAttributes(
+		const USkeletalMeshExporterUSDOptions& Options,
+		TArray< FAnalyticsEventAttribute >& InOutAttributes
+	);
+}
diff --git a/Engine/Plugins/Importers/USDImporter/Source/USDExporter/Public/StaticMeshExporterUSDOptions.h b/Engine/Plugins/Importers/USDImporter/Source/USDExporter/Public/StaticMeshExporterUSDOptions.h
--- a/Engine/Plugins/Importers/USDImporter/Source/USDExporter/Public/StaticMeshExporterUSDOptions.h
+++ b/Engine/Plugins/Importers/USDImporter/Source/USDExporter/Public/StaticMeshExporterUSDOptions.h
@@ -9,6 +9,8 @@
 
 #include "StaticMeshExporterUSDOptions.generated.h"
 
+struct FAnalyticsEventAttribute;
+
 /**
  * Options for exporting static meshes to USD format.
  */
@@ -31,3 +33,11 @@ public:
 	UPROPERTY( EditAnywhere, config, BlueprintReadWrite, Category = "Collision", meta = ( DisplayName = "Re-export Identical Assets" ) )
 	bool bReExportIdenticalAssets = false;
 };
+
+namespace UsdUtils
+{
+	USDEXPORTER_API void AddAnalyticsAttributes(
+		const UStaticMeshExporterUSDOptions& Options,
+		TArray< FAnalyticsEventAttribute >& InOutAttributes
+	);
+}
\ No newline at end of file
diff --git a/Engine/Plugins/Importers/USDImporter/Source/USDExporter/Public/USDAssetOptions.h b/Engine/Plugins/Importers/USDImporter/Source/USDExporter/Public/USDAssetOptions.h
--- a/Engine/Plugins/Importers/USDImporter/Source/USDExporter/Public/USDAssetOptions.h
+++ b/Engine/Plugins/Importers/USDImporter/Source/USDExporter/Public/USDAssetOptions.h
@@ -8,6 +8,8 @@
 
 #include "USDAssetOptions.generated.h"
 
+struct FAnalyticsEventAttribute;
+
 USTRUCT( BlueprintType )
 struct USDEXPORTER_API FUsdMaterialBakingOptions
 {
@@ -57,4 +59,16 @@ struct USDEXPORTER_API FUsdMeshAssetOptions
 	/** Highest of the LOD indices to export static and skeletal meshes with */
 	UPROPERTY( EditAnywhere, config, BlueprintReadWrite, Category = "Mesh options", meta = ( ClampMin = "0" ) )
 	int32 HighestMeshLOD = MAX_MESH_LOD_COUNT - 1;
-};
\ No newline at end of file
+};
+
+namespace UsdUtils
+{
+	USDEXPORTER_API void AddAnalyticsAttributes(
+		const FUsdMaterialBakingOptions& Options,
+		TArray< FAnalyticsEventAttribute >& InOutAttributes
+	);
+	USDEXPORTER_API void AddAnalyticsAttributes(
+		const FUsdMeshAssetOptions& Options,
+		TArray< FAnalyticsEventAttribute >& InOutAttributes
+	);
+}
\ No newline at end of file
diff --git a/Engine/Plugins/Importers/USDImporter/Source/USDExporter/Public/USDConversionBlueprintLibrary.h b/Engine/Plugins/Importers/USDImporter/Source/USDExporter/Public/USDConversionBlueprintLibrary.h
--- a/Engine/Plugins/Importers/USDImporter/Source/USDExporter/Public/USDConversionBlueprintLibrary.h
+++ b/Engine/Plugins/Importers/USDImporter/Source/USDExporter/Public/USDConversionBlueprintLibrary.h
@@ -11,6 +11,7 @@
 class AInstancedFoliageActor;
 class UFoliageType;
 class ULevel;
+class ULevelExporterUSDOptions;
 
 /** Wrapped static conversion functions from the UsdUtilities module, so that they can be used via scripting */
 UCLASS(meta=(ScriptName="UsdConversionLibrary"))
@@ -113,6 +114,10 @@ public:
 	UFUNCTION( BlueprintCallable, meta = ( ScriptMethod ), Category = "USD Foliage Exporter" )
 	static TArray<FTransform> GetInstanceTransforms( AInstancedFoliageActor* Actor, UFoliageType* FoliageType, ULevel* InstancesLevel = nullptr );
 
+	/** Retrieves the analytics attributes to send for the provided options object */
+	UFUNCTION( BlueprintCallable, Category = "Analytics" )
+	static TArray<FAnalyticsEventAttr> GetAnalyticsAttributes( const ULevelExporterUSDOptions* Options );
+
 	/** Defer to the USDClasses module to actually send analytics information */
 	UFUNCTION( BlueprintCallable, Category = "Analytics" )
 	static void SendAnalytics( const TArray<FAnalyticsEventAttr>& Attrs, const FString& EventName, bool bAutomated, double ElapsedSeconds, double NumberOfFrames, const FString& Extension );
diff --git a/Engine/Plugins/Importers/USDImporter/Source/USDStageImporter/Private/USDStageImportOptions.cpp b/Engine/Plugins/Importers/USDImporter/Source/USDStageImporter/Private/USDStageImportOptions.cpp
--- a/Engine/Plugins/Importers/USDImporter/Source/USDStageImporter/Private/USDStageImportOptions.cpp
+++ b/Engine/Plugins/Importers/USDImporter/Source/USDStageImporter/Private/USDStageImportOptions.cpp
@@ -5,6 +5,7 @@
 #include "USDSchemasModule.h"
 #include "USDSchemaTranslator.h"
 
+#include "AnalyticsEventAttribute.h"
 #include "Modules/ModuleManager.h"
 #include "UObject/UnrealType.h"
 
@@ -68,3 +69,37 @@ void UUsdStageImportOptions::PostEditChangeProperty(struct FPropertyChangedEvent
 		SaveConfig();
 	}
 }
+
+void UsdUtils::AddAnalyticsAttributes(
+	const UUsdStageImportOptions& Options,
+	TArray< FAnalyticsEventAttribute >& InOutAttributes
+)
+{
+	InOutAttributes.Emplace( TEXT( "ImportActors" ), LexToString( Options.bImportActors ) );
+	InOutAttributes.Emplace( TEXT( "ImportGeometry" ), LexToString( Options.bImportGeometry ) );
+	InOutAttributes.Emplace( TEXT( "ImportSkeletalAnimations" ), LexToString( Options.bImportSkeletalAnimations ) );
+	InOutAttributes.Emplace( TEXT( "ImportLevelSequences" ), LexToString( Options.bImportLevelSequences ) );
+	InOutAttributes.Emplace( TEXT( "ImportMaterials" ), LexToString( Options.bImportMaterials ) );
+	if ( Options.PrimsToImport != TArray<FString>{ TEXT( "/" ) } )
+	{
+		InOutAttributes.Emplace( TEXT( "NumPrimsToImport" ), LexToString( Options.PrimsToImport.Num() ) );
+	}
+	InOutAttributes.Emplace( TEXT( "PurposesToImport" ), LexToString( Options.PurposesToImport ) );
+	InOutAttributes.Emplace( TEXT( "NaniteTriangleThreshold" ), LexToString( Options.NaniteTriangleThreshold ) );
+	InOutAttributes.Emplace( TEXT( "RenderContextToImport" ), Options.RenderContextToImport.ToString() );
+	InOutAttributes.Emplace( TEXT( "MaterialPurpose" ), Options.MaterialPurpose.ToString() );
+	InOutAttributes.Emplace( TEXT( "RootMotionHandling" ), LexToString( ( uint8 ) Options.RootMotionHandling ) );
+	InOutAttributes.Emplace( TEXT( "OverrideStageOptions" ), Options.bOverrideStageOptions );
+	if ( Options.bOverrideStageOptions )
+	{
+		UsdUtils::AddAnalyticsAttributes( Options.StageOptions, InOutAttributes );
+	}
+	InOutAttributes.Emplace( TEXT( "NumGroomInterpolationSettings" ), LexToString( Options.GroomInterpolationSettings.Num() ) );
+	InOutAttributes.Emplace( TEXT( "ReplaceActorPolicy" ), LexToString( ( uint8 ) Options.ExistingActorPolicy ) );
+	InOutAttributes.Emplace( TEXT( "ReplaceAssetPolicy" ), LexToString( ( uint8 ) Options.ExistingAssetPolicy ) );
+	InOutAttributes.Emplace( TEXT( "PrimPathFolderStructure" ), LexToString( Options.bPrimPathFolderStructure ) );
+	InOutAttributes.Emplace( TEXT( "KindsToCollapse" ), LexToString( Options.KindsToCollapse ) );
+	InOutAttributes.Emplace( TEXT( "MergeIdenticalMaterialSlots" ), LexToString( Options.bMergeIdenticalMaterialSlots ) );
+	InOutAttributes.Emplace( TEXT( "CollapseTopLevelPointInstancers" ), LexToString( Options.bCollapseTopLevelPointInstancers ) );
+	InOutAttributes.Emplace( TEXT( "InterpretLODs" ), LexToString( Options.bInterpretLODs ) );
+}
\ No newline at end of file
diff --git a/Engine/Plugins/Importers/USDImporter/Source/USDStageImporter/Private/USDStageImporter.cpp b/Engine/Plugins/Importers/USDImporter/Source/USDStageImporter/Private/USDStageImporter.cpp
--- a/Engine/Plugins/Importers/USDImporter/Source/USDStageImporter/Private/USDStageImporter.cpp
+++ b/Engine/Plugins/Importers/USDImporter/Source/USDStageImporter/Private/USDStageImporter.cpp
@@ -35,6 +35,8 @@
 #include "EngineAnalytics.h"
 #include "EngineUtils.h"
 #include "GeometryCache.h"
+#include "GroomAsset.h"
+#include "GroomCache.h"
 #include "HAL/FileManager.h"
 #include "IAssetTools.h"
 #include "LevelSequence.h"
@@ -1300,7 +1302,13 @@ namespace UsdStageImporterImpl
 		}
 	}
 
-	void SendAnalytics( FUsdStageImportContext& ImportContext, UObject* Asset, const FString& Operation, const TSet<UObject*>& ImportedAssets, double ElapsedSeconds )
+	void SendAnalytics(
+		FUsdStageImportContext& ImportContext,
+		UObject* Asset,
+		const FString& Operation,
+		const TSet<UObject*>& ImportedAssets,
+		double ElapsedSeconds
+	)
 	{
 #if USE_USD_SDK
 		if ( FEngineAnalytics::IsAvailable() )
@@ -1319,29 +1327,7 @@ namespace UsdStageImporterImpl
 
 			if ( ImportContext.ImportOptions )
 			{
-				EventAttributes.Emplace( TEXT( "ImportActors" ), LexToString( ImportContext.ImportOptions->bImportActors ) );
-				EventAttributes.Emplace( TEXT( "ImportGeometry" ), LexToString( ImportContext.ImportOptions->bImportGeometry ) );
-				EventAttributes.Emplace( TEXT( "ImportSkeletalAnimations" ), LexToString( ImportContext.ImportOptions->bImportSkeletalAnimations ) );
-				EventAttributes.Emplace( TEXT( "ImportLevelSequences" ), LexToString( ImportContext.ImportOptions->bImportLevelSequences ) );
-				EventAttributes.Emplace( TEXT( "ImportMaterials" ), LexToString( ImportContext.ImportOptions->bImportMaterials ) );
-				EventAttributes.Emplace( TEXT( "PurposesToImport" ), LexToString( ImportContext.ImportOptions->PurposesToImport ) );
-				EventAttributes.Emplace( TEXT( "NaniteTriangleThreshold" ), LexToString( ImportContext.ImportOptions->NaniteTriangleThreshold ) );
-				EventAttributes.Emplace( TEXT( "RenderContextToImport" ), ImportContext.ImportOptions->RenderContextToImport.ToString() );
-				EventAttributes.Emplace( TEXT( "MaterialPurpose" ), ImportContext.ImportOptions->MaterialPurpose.ToString() );
-				EventAttributes.Emplace( TEXT( "RootMotionHandling" ), LexToString( ( uint8 ) ImportContext.ImportOptions->RootMotionHandling ) );
-				EventAttributes.Emplace( TEXT( "OverrideStageOptions" ), ImportContext.ImportOptions->bOverrideStageOptions );
-				if( ImportContext.ImportOptions->bOverrideStageOptions )
-				{
-					EventAttributes.Emplace( TEXT( "MetersPerUnit" ), ImportContext.ImportOptions->StageOptions.MetersPerUnit );
-					EventAttributes.Emplace( TEXT( "UpAxis" ), ImportContext.ImportOptions->StageOptions.UpAxis == EUsdUpAxis::YAxis ? TEXT( "Y" ) : TEXT( "Z" ) );
-				}
-				EventAttributes.Emplace( TEXT( "ReplaceActorPolicy" ), LexToString( (uint8)ImportContext.ImportOptions->ExistingActorPolicy ) );
-				EventAttributes.Emplace( TEXT( "ReplaceAssetPolicy" ), LexToString( (uint8)ImportContext.ImportOptions->ExistingAssetPolicy ) );
-				EventAttributes.Emplace( TEXT( "PrimPathFolderStructure" ), LexToString( ImportContext.ImportOptions->bPrimPathFolderStructure ) );
-				EventAttributes.Emplace( TEXT( "KindsToCollapse" ), LexToString( ImportContext.ImportOptions->KindsToCollapse ) );
-				EventAttributes.Emplace( TEXT( "MergeIdenticalMaterialSlots" ), LexToString( ImportContext.ImportOptions->bMergeIdenticalMaterialSlots ) );
-				EventAttributes.Emplace( TEXT( "CollapseTopLevelPointInstancers" ), LexToString( ImportContext.ImportOptions->bCollapseTopLevelPointInstancers ) );
-				EventAttributes.Emplace( TEXT( "InterpretLODs" ), LexToString( ImportContext.ImportOptions->bInterpretLODs ) );
+				UsdUtils::AddAnalyticsAttributes( *ImportContext.ImportOptions, EventAttributes );
 			}
 
 			int32 NumStaticMeshes = 0;
@@ -1351,6 +1337,8 @@ namespace UsdStageImporterImpl
 			int32 NumLevelSequences = 0;
 			int32 NumTextures = 0;
 			int32 NumGeometryCaches = 0;
+			int32 NumGroomAssets = 0;
+			int32 NumGroomCaches = 0;
 			for ( UObject* ImportedAsset : ImportedAssets )
 			{
 				if ( !ImportedAsset )
@@ -1386,6 +1374,14 @@ namespace UsdStageImporterImpl
 				{
 					++NumGeometryCaches;
 				}
+				else if ( ImportedAsset->IsA<UGroomAsset>() )
+				{
+					++NumGroomAssets;
+				}
+				else if ( ImportedAsset->IsA<UGroomCache>() )
+				{
+					++NumGroomCaches;
+				}
 			}
 			EventAttributes.Emplace( TEXT( "NumStaticMeshes" ), LexToString( NumStaticMeshes ) );
 			EventAttributes.Emplace( TEXT( "NumSkeletalMeshes" ), LexToString( NumSkeletalMeshes ) );
@@ -1394,6 +1390,8 @@ namespace UsdStageImporterImpl
 			EventAttributes.Emplace( TEXT( "NumLevelSequences" ), LexToString( NumLevelSequences ) );
 			EventAttributes.Emplace( TEXT( "NumTextures" ), LexToString( NumTextures ) );
 			EventAttributes.Emplace( TEXT( "NumGeometryCaches" ), LexToString( NumGeometryCaches ) );
+			EventAttributes.Emplace( TEXT( "NumGroomAssets" ), LexToString( NumGroomAssets ) );
+			EventAttributes.Emplace( TEXT( "NumGroomCaches" ), LexToString( NumGroomCaches ) );
 
 			FString RootLayerIdentifier = ImportContext.FilePath;
 			if ( ImportContext.Stage )
diff --git a/Engine/Plugins/Importers/USDImporter/Source/USDStageImporter/Public/USDStageImportOptions.h b/Engine/Plugins/Importers/USDImporter/Source/USDStageImporter/Public/USDStageImportOptions.h
--- a/Engine/Plugins/Importers/USDImporter/Source/USDStageImporter/Public/USDStageImportOptions.h
+++ b/Engine/Plugins/Importers/USDImporter/Source/USDStageImporter/Public/USDStageImportOptions.h
@@ -13,6 +13,8 @@
 
 #include "USDStageImportOptions.generated.h"
 
+struct FAnalyticsEventAttribute;
+
 UENUM(BlueprintType)
 enum class EReplaceActorPolicy : uint8
 {
@@ -153,3 +155,11 @@ public:
 	void EnableActorImport(bool bEnable);
 	virtual void PostEditChangeProperty(struct FPropertyChangedEvent& PropertyChangedEvent) override;
 };
+
+namespace UsdUtils
+{
+	USDSTAGEIMPORTER_API void AddAnalyticsAttributes(
+		const UUsdStageImportOptions& Options,
+		TArray< FAnalyticsEventAttribute >& InOutAttributes
+	);
+}
