USD: Work around Python GIL deadlock caused by sneaking into C++ from Python.

diff --git a/Engine/Plugins/Importers/USDImporter/Source/USDSchemas/Private/USDSkelRootTranslator.cpp b/Engine/Plugins/Importers/USDImporter/Source/USDSchemas/Private/USDSkelRootTranslator.cpp
--- a/Engine/Plugins/Importers/USDImporter/Source/USDSchemas/Private/USDSkelRootTranslator.cpp
+++ b/Engine/Plugins/Importers/USDImporter/Source/USDSchemas/Private/USDSkelRootTranslator.cpp
@@ -33,6 +33,7 @@
 #include "GroomComponent.h"
 #include "Materials/Material.h"
 #include "Materials/MaterialInstanceConstant.h"
+#include "UObject/StrongObjectPtr.h"
 #include "Rendering/SkeletalMeshLODImporterData.h"
 #include "Serialization/BufferArchive.h"
 
@@ -1418,7 +1419,28 @@ void FUsdSkelRootTranslator::UpdateComponents( USceneComponent* SceneComponent )
 
 	if ( bPrimHasLiveLinkSchema )
 	{
-		UsdSkelRootTranslatorImpl::UpdateLiveLinkProperties( Context.Get(), SkeletalMeshComponent, Prim );
+		TSharedPtr< FUsdSchemaTranslationContext > ContextPtr = Context;
+		TStrongObjectPtr< USkeletalMeshComponent > PinnedSkelMeshComponent{ SkeletalMeshComponent };
+
+		// HACK: This is a temporary work-around for a GIL deadlock. See UE-156489 and UE-156370 for more details, but
+		// the long-story short of it is that at this point we may have a callstack that originates from Python,
+		// triggers a USD stage notice and causes C++ code to run as our stage actor is listening to them. If we
+		// cause GC to run right now (which the DuplicateObject inside UpdateLiveLinkProperties will) we may cause
+		// a deadlock, as the game thread still holds the GIL and a background reference collector thread may want
+		// to acquire it too. What this does is run this part of the UpdateComponents on tick, that has a callstack
+		// that doesn't originate from Python, and so doesn't have the GIL locked
+		FTSTicker::GetCoreTicker().AddTicker(
+			FTickerDelegate::CreateLambda( [ContextPtr, PinnedSkelMeshComponent, Prim]( float Time )
+			{
+				if ( ContextPtr )
+				{
+					UsdSkelRootTranslatorImpl::UpdateLiveLinkProperties( *ContextPtr.Get(), PinnedSkelMeshComponent.Get(), Prim );
+				}
+
+				// Returning false means this is a one-off, and won't repeat
+				return false;
+			})
+		);
 	}
 
 	// Update the animation state
