USD: Added LLM tag for USD allocations.

diff --git a/Engine/Plugins/Importers/USDImporter/Source/USDExporter/Private/USDExporterModule.cpp b/Engine/Plugins/Importers/USDImporter/Source/USDExporter/Private/USDExporterModule.cpp
--- a/Engine/Plugins/Importers/USDImporter/Source/USDExporter/Private/USDExporterModule.cpp
+++ b/Engine/Plugins/Importers/USDImporter/Source/USDExporter/Private/USDExporterModule.cpp
@@ -21,6 +21,8 @@ class FUsdExporterModule : public IUsdExporterModule
 public:
 	virtual void StartupModule() override
 	{
+		LLM_SCOPE_BYTAG(Usd);
+
 		FPropertyEditorModule& PropertyModule = FModuleManager::LoadModuleChecked< FPropertyEditorModule >( TEXT( "PropertyEditor" ) );
 
 		// We intentionally use the same customization for both of these
diff --git a/Engine/Plugins/Importers/USDImporter/Source/USDSchemas/Private/USDSchemaTranslator.cpp b/Engine/Plugins/Importers/USDImporter/Source/USDSchemas/Private/USDSchemaTranslator.cpp
--- a/Engine/Plugins/Importers/USDImporter/Source/USDSchemas/Private/USDSchemaTranslator.cpp
+++ b/Engine/Plugins/Importers/USDImporter/Source/USDSchemas/Private/USDSchemaTranslator.cpp
@@ -56,6 +56,8 @@ TSharedPtr< FUsdSchemaTranslator > FUsdSchemaTranslatorRegistry::CreateTranslato
 FUsdRenderContextRegistry::FUsdRenderContextRegistry()
 {
 #if USE_USD_SDK
+	LLM_SCOPE_BYTAG(Usd);
+
 	UniversalRenderContext = FName( UsdToUnreal::ConvertToken( pxr::UsdShadeTokens->universalRenderContext ) );
 	Register( UniversalRenderContext );
 
diff --git a/Engine/Plugins/Importers/USDImporter/Source/USDSchemas/Private/USDSchemasModule.cpp b/Engine/Plugins/Importers/USDImporter/Source/USDSchemas/Private/USDSchemasModule.cpp
--- a/Engine/Plugins/Importers/USDImporter/Source/USDSchemas/Private/USDSchemasModule.cpp
+++ b/Engine/Plugins/Importers/USDImporter/Source/USDSchemas/Private/USDSchemasModule.cpp
@@ -29,6 +29,8 @@ public:
 	virtual void StartupModule() override
 	{
 #if USE_USD_SDK
+		LLM_SCOPE_BYTAG(Usd);
+
 		// Register the default translators
 		UsdGeomCameraTranslatorHandle = GetTranslatorRegistry().Register< FUsdGeomCameraTranslator >( TEXT("UsdGeomCamera") );
 		UsdGeomMeshTranslatorHandle = GetTranslatorRegistry().Register< FUsdGeomMeshTranslator >( TEXT("UsdGeomMesh") );
diff --git a/Engine/Plugins/Importers/USDImporter/Source/USDStage/Private/USDStageModule.cpp b/Engine/Plugins/Importers/USDImporter/Source/USDStage/Private/USDStageModule.cpp
--- a/Engine/Plugins/Importers/USDImporter/Source/USDStage/Private/USDStageModule.cpp
+++ b/Engine/Plugins/Importers/USDImporter/Source/USDStage/Private/USDStageModule.cpp
@@ -2,6 +2,7 @@
 
 #include "USDStageModule.h"
 
+#include "USDMemory.h"
 #include "USDStageActor.h"
 #include "USDStageActorCustomization.h"
 
@@ -20,6 +21,8 @@ public:
 	virtual void StartupModule() override
 	{
 #if WITH_EDITOR
+		LLM_SCOPE_BYTAG(Usd);
+
 		// If don't have any active references to our materials they won't be packaged into monolithic builds, and we wouldn't
 		// be able to create dynamic material instances at runtime.
 		TSharedPtr<IPlugin> Plugin = IPluginManager::Get().FindPlugin( TEXT( "USDImporter" ) );
diff --git a/Engine/Plugins/Importers/USDImporter/Source/USDStageEditor/Private/USDStageEditorModule.cpp b/Engine/Plugins/Importers/USDImporter/Source/USDStageEditor/Private/USDStageEditorModule.cpp
--- a/Engine/Plugins/Importers/USDImporter/Source/USDStageEditor/Private/USDStageEditorModule.cpp
+++ b/Engine/Plugins/Importers/USDImporter/Source/USDStageEditor/Private/USDStageEditorModule.cpp
@@ -222,6 +222,8 @@ public:
 #if USE_USD_SDK
 	static TSharedRef< SDockTab > SpawnUsdStageTab( const FSpawnTabArgs& SpawnTabArgs )
 	{
+		LLM_SCOPE_BYTAG(Usd);
+
 		return SNew(SDockTab)
 			.TabRole(ETabRole::NomadTab)
 			.Label( LOCTEXT( "USDStage", "USD Stage" ) )
@@ -237,6 +239,8 @@ public:
 
 	virtual void StartupModule() override
 	{
+		LLM_SCOPE_BYTAG(Usd);
+
 		FUsdStageEditorStyle::Initialize();
 
 		FLevelEditorModule& LevelEditorModule = FModuleManager::GetModuleChecked< FLevelEditorModule >( "LevelEditor" );
diff --git a/Engine/Plugins/Importers/USDImporter/Source/USDStageImporter/Private/USDStageImporterModule.cpp b/Engine/Plugins/Importers/USDImporter/Source/USDStageImporter/Private/USDStageImporterModule.cpp
--- a/Engine/Plugins/Importers/USDImporter/Source/USDStageImporter/Private/USDStageImporterModule.cpp
+++ b/Engine/Plugins/Importers/USDImporter/Source/USDStageImporter/Private/USDStageImporterModule.cpp
@@ -3,6 +3,7 @@
 #include "USDStageImporterModule.h"
 
 #include "UnrealUSDWrapper.h"
+#include "USDMemory.h"
 #include "USDStageImporter.h"
 #include "USDStageImportOptionsCustomization.h"
 
@@ -19,6 +20,8 @@ public:
 	virtual void StartupModule() override
 	{
 #if USE_USD_SDK
+		LLM_SCOPE_BYTAG(Usd);
+
 		IUnrealUSDWrapperModule& UnrealUSDWrapperModule = FModuleManager::Get().LoadModuleChecked< IUnrealUSDWrapperModule >(TEXT("UnrealUSDWrapper"));
 
 		FPropertyEditorModule& PropertyModule = FModuleManager::LoadModuleChecked<FPropertyEditorModule>( TEXT( "PropertyEditor" ) );
diff --git a/Engine/Plugins/Importers/USDImporter/Source/USDTests/Private/USDTestsModule.cpp b/Engine/Plugins/Importers/USDImporter/Source/USDTests/Private/USDTestsModule.cpp
--- a/Engine/Plugins/Importers/USDImporter/Source/USDTests/Private/USDTestsModule.cpp
+++ b/Engine/Plugins/Importers/USDImporter/Source/USDTests/Private/USDTestsModule.cpp
@@ -11,6 +11,7 @@ class FUsdTestsModule : public IUsdTestsModule
 public:
 	virtual void StartupModule() override
 	{
+		LLM_SCOPE_BYTAG(Usd);
 	}
 
 	virtual void ShutdownModule() override
diff --git a/Engine/Plugins/Importers/USDImporter/Source/USDUtilities/Private/USDUtilitiesModule.cpp b/Engine/Plugins/Importers/USDImporter/Source/USDUtilities/Private/USDUtilitiesModule.cpp
--- a/Engine/Plugins/Importers/USDImporter/Source/USDUtilities/Private/USDUtilitiesModule.cpp
+++ b/Engine/Plugins/Importers/USDImporter/Source/USDUtilities/Private/USDUtilitiesModule.cpp
@@ -15,6 +15,8 @@ class FUsdUtilitiesModule : public IUsdUtilitiesModule
 	virtual void StartupModule() override
 	{
 #if WITH_EDITOR
+		LLM_SCOPE_BYTAG(Usd);
+
 		FMessageLogInitializationOptions InitOptions;
 		FMessageLogModule& MessageLogModule = FModuleManager::LoadModuleChecked< FMessageLogModule >( TEXT("MessageLog") );
 		MessageLogModule.RegisterLogListing( TEXT( "USD" ), NSLOCTEXT( "USDUtilitiesModule", "USDLogListing", "USD" ), InitOptions );
diff --git a/Engine/Plugins/Importers/USDImporter/Source/UnrealUSDWrapper/Private/USDMemory.cpp b/Engine/Plugins/Importers/USDImporter/Source/UnrealUSDWrapper/Private/USDMemory.cpp
--- a/Engine/Plugins/Importers/USDImporter/Source/UnrealUSDWrapper/Private/USDMemory.cpp
+++ b/Engine/Plugins/Importers/USDImporter/Source/UnrealUSDWrapper/Private/USDMemory.cpp
@@ -2,11 +2,13 @@
 
 #include "USDMemory.h"
 
+#include "Async/Async.h"
 #include "Containers/Array.h"
 #include "HAL/PlatformTLS.h"
 #include "HAL/TlsAutoCleanup.h"
 #include "Misc/ScopeLock.h"
-#include "Async/Async.h"
+
+LLM_DEFINE_TAG( Usd );
 
 class UNREALUSDWRAPPER_API FTlsSlot final
 {
@@ -137,6 +139,8 @@ bool FUsdMemoryManager::DeactivateAllocator( EUsdActiveAllocator Allocator )
 
 void* FUsdMemoryManager::Malloc( SIZE_T Count )
 {
+	LLM_SCOPE_BYTAG(Usd);
+
 	void* Result = nullptr;
 
 #if USD_USES_SYSTEM_MALLOC
diff --git a/Engine/Plugins/Importers/USDImporter/Source/UnrealUSDWrapper/Private/UnrealUSDWrapper.cpp b/Engine/Plugins/Importers/USDImporter/Source/UnrealUSDWrapper/Private/UnrealUSDWrapper.cpp
--- a/Engine/Plugins/Importers/USDImporter/Source/UnrealUSDWrapper/Private/UnrealUSDWrapper.cpp
+++ b/Engine/Plugins/Importers/USDImporter/Source/UnrealUSDWrapper/Private/UnrealUSDWrapper.cpp
@@ -1229,6 +1229,7 @@ public:
 	virtual void StartupModule() override
 	{
 #if USE_USD_SDK
+		LLM_SCOPE_BYTAG(Usd);
 
 		// Path to USD base plugins
 		FString UsdPluginsPath = FPaths::Combine( TEXT( ".." ), TEXT( "ThirdParty" ), TEXT( "USD" ), TEXT( "UsdResources" ) );
diff --git a/Engine/Plugins/Importers/USDImporter/Source/UnrealUSDWrapper/Public/USDMemory.h b/Engine/Plugins/Importers/USDImporter/Source/UnrealUSDWrapper/Public/USDMemory.h
--- a/Engine/Plugins/Importers/USDImporter/Source/UnrealUSDWrapper/Public/USDMemory.h
+++ b/Engine/Plugins/Importers/USDImporter/Source/UnrealUSDWrapper/Public/USDMemory.h
@@ -3,12 +3,15 @@
 #pragma once
 
 #include "CoreMinimal.h"
+#include "HAL/LowLevelMemTracker.h"
 #include "Misc/Optional.h"
 #include "Modules/Boilerplate/ModuleBoilerplate.h"
 #include "Templates/SharedPointer.h"
 #include "Misc/ScopeRWLock.h"
 #include "Templates/UnrealTemplate.h"
 
+LLM_DECLARE_TAG_API( Usd, UNREALUSDWRAPPER_API );
+
 /**
  * This file is used to handle memory allocation issues when interacting with the USD SDK.
  *
