USD: Remove deprecated stuff.

diff --git a/Engine/Plugins/Importers/USDImporter/Source/USDExporter/Private/USDLevelInfo.cpp b/Engine/Plugins/Importers/USDImporter/Source/USDExporter/Private/USDLevelInfo.cpp
deleted file mode 100644
diff --git a/Engine/Plugins/Importers/USDImporter/Source/USDExporter/Private/USDLevelInfo.h b/Engine/Plugins/Importers/USDImporter/Source/USDExporter/Private/USDLevelInfo.h
deleted file mode 100644
diff --git a/Engine/Plugins/Importers/USDImporter/Source/USDSchemas/Private/MeshTranslationImpl.cpp b/Engine/Plugins/Importers/USDImporter/Source/USDSchemas/Private/MeshTranslationImpl.cpp
--- a/Engine/Plugins/Importers/USDImporter/Source/USDSchemas/Private/MeshTranslationImpl.cpp
+++ b/Engine/Plugins/Importers/USDImporter/Source/USDSchemas/Private/MeshTranslationImpl.cpp
@@ -107,26 +107,7 @@ TMap<const UsdUtils::FUsdPrimMaterialSlot*, UMaterialInterface*> MeshTranslation
 			case UsdUtils::EPrimAssignmentType::None:
 			default:
 			{
-				// Check if there is a material already on the mesh
-				UMaterialInstanceConstant* ExistingMaterialInstance = Cast< UMaterialInstanceConstant >( ExistingAssignments[ GlobalResolvedMaterialIndex ] );
-
-				// Assuming that we own the material instance and that we can change it as we wish, reuse it
-				if ( ExistingMaterialInstance && ExistingMaterialInstance->GetOuter() == GetTransientPackage() )
-				{
-#if WITH_EDITOR
-					UUsdAssetImportData* AssetImportData = Cast<UUsdAssetImportData>( ExistingMaterialInstance->AssetImportData );
-					if ( AssetImportData && AssetImportData->PrimPath == UsdToUnreal::ConvertPath( UsdPrim.GetPrimPath() ) )
-#endif // WITH_EDITOR
-					{
-						// If we have displayColor data on our prim, repurpose this material to show it
-						if ( TOptional<UsdUtils::FDisplayColorMaterial> DisplayColorDescription = UsdUtils::ExtractDisplayColorMaterial( pxr::UsdGeomMesh( UsdPrim ) ) )
-						{
-							UsdToUnreal::ConvertDisplayColor( DisplayColorDescription.GetValue(), *ExistingMaterialInstance );
-						}
-
-						Material = ExistingMaterialInstance;
-					}
-				}
+				ensure( false );
 				break;
 			}
 			}
diff --git a/Engine/Plugins/Importers/USDImporter/Source/USDSchemas/Private/USDGeomMeshTranslator.cpp b/Engine/Plugins/Importers/USDImporter/Source/USDSchemas/Private/USDGeomMeshTranslator.cpp
--- a/Engine/Plugins/Importers/USDImporter/Source/USDSchemas/Private/USDGeomMeshTranslator.cpp
+++ b/Engine/Plugins/Importers/USDImporter/Source/USDSchemas/Private/USDGeomMeshTranslator.cpp
@@ -374,26 +374,7 @@ namespace UsdGeomMeshTranslatorImpl
 				case UsdUtils::EPrimAssignmentType::None:
 				default:
 				{
-					// Check if there is a material already on the mesh
-					UMaterialInstanceConstant* ExistingMaterialInstance = GeometryCache.Materials.IsValidIndex( StaticMeshSlotIndex ) ? Cast< UMaterialInstanceConstant >( GeometryCache.Materials[ StaticMeshSlotIndex ] ) : nullptr;
-
-					// Assuming that we own the material instance and that we can change it as we wish, reuse it
-					if ( ExistingMaterialInstance && ExistingMaterialInstance->GetOuter() == GetTransientPackage() )
-					{
-#if WITH_EDITOR
-						UUsdAssetImportData* AssetImportData = Cast< UUsdAssetImportData >( ExistingMaterialInstance->AssetImportData );
-						if ( AssetImportData && AssetImportData->PrimPath == UsdToUnreal::ConvertPath( UsdPrim.GetPrimPath() ) )
-#endif // WITH_EDITOR
-						{
-							// If we have displayColor data on our prim, repurpose this material to show it
-							if ( TOptional< UsdUtils::FDisplayColorMaterial > DisplayColorDescription = UsdUtils::ExtractDisplayColorMaterial( pxr::UsdGeomMesh( UsdPrim ) ) )
-							{
-								UsdToUnreal::ConvertDisplayColor( DisplayColorDescription.GetValue(), *ExistingMaterialInstance );
-							}
-
-							Material = ExistingMaterialInstance;
-						}
-					}
+					ensure( false );
 					break;
 				}
 				}
diff --git a/Engine/Plugins/Importers/USDImporter/Source/USDStage/Private/USDLevelSequenceHelper.cpp b/Engine/Plugins/Importers/USDImporter/Source/USDStage/Private/USDLevelSequenceHelper.cpp
--- a/Engine/Plugins/Importers/USDImporter/Source/USDStage/Private/USDLevelSequenceHelper.cpp
+++ b/Engine/Plugins/Importers/USDImporter/Source/USDStage/Private/USDLevelSequenceHelper.cpp
@@ -3452,17 +3452,6 @@ public:
 };
 #endif // USE_USD_SDK
 
-
-FUsdLevelSequenceHelper::FUsdLevelSequenceHelper(TWeakObjectPtr<AUsdStageActor> InStageActor)
-	: FUsdLevelSequenceHelper()
-{
-	if (AUsdStageActor* ValidStageActor = InStageActor.Get())
-	{
-		Init( InStageActor->GetOrLoadUsdStage() );
-		BindToUsdStageActor(ValidStageActor);
-	}
-}
-
 FUsdLevelSequenceHelper::FUsdLevelSequenceHelper()
 {
 	UsdSequencerImpl = MakeUnique<FUsdLevelSequenceHelperImpl>();
@@ -3531,16 +3520,6 @@ void FUsdLevelSequenceHelper::Clear()
 	}
 }
 
-void FUsdLevelSequenceHelper::InitLevelSequence(const UE::FUsdStage& UsdStage)
-{
-	if (UsdSequencerImpl.IsValid() && UsdStage)
-	{
-		UE_LOG(LogUsd, Verbose, TEXT("InitLevelSequence"));
-
-		UsdSequencerImpl->CreateLocalLayersSequences();
-	}
-}
-
 void FUsdLevelSequenceHelper::BindToUsdStageActor(AUsdStageActor* StageActor)
 {
 	if (UsdSequencerImpl.IsValid())
diff --git a/Engine/Plugins/Importers/USDImporter/Source/USDStage/Private/USDStageActor.cpp b/Engine/Plugins/Importers/USDImporter/Source/USDStage/Private/USDStageActor.cpp
--- a/Engine/Plugins/Importers/USDImporter/Source/USDStage/Private/USDStageActor.cpp
+++ b/Engine/Plugins/Importers/USDImporter/Source/USDStage/Private/USDStageActor.cpp
@@ -1383,13 +1383,6 @@ void AUsdStageActor::UpdatePrim( const UE::FSdfPath& InUsdPrimPath, bool bResync
 	}
 }
 
-UE::FUsdStage& AUsdStageActor::GetUsdStage()
-{
-	OpenUsdStage();
-
-	return IsolatedStage ? IsolatedStage : UsdStage;
-}
-
 const UE::FUsdStage& AUsdStageActor::GetUsdStage() const
 {
 	return IsolatedStage ? IsolatedStage : UsdStage;
diff --git a/Engine/Plugins/Importers/USDImporter/Source/USDStage/Public/USDLevelSequenceHelper.h b/Engine/Plugins/Importers/USDImporter/Source/USDStage/Public/USDLevelSequenceHelper.h
--- a/Engine/Plugins/Importers/USDImporter/Source/USDStage/Public/USDLevelSequenceHelper.h
+++ b/Engine/Plugins/Importers/USDImporter/Source/USDStage/Public/USDLevelSequenceHelper.h
@@ -20,8 +20,6 @@ class USDSTAGE_API FUsdLevelSequenceHelper
 {
 public:
 	FUsdLevelSequenceHelper();
-	UE_DEPRECATED(4.27, "This function is deprecated, use the default constructor and call Init and BindToUsdStageActor(optional) instead.")
-	explicit FUsdLevelSequenceHelper(TWeakObjectPtr<AUsdStageActor> InStageActor);
 	virtual ~FUsdLevelSequenceHelper();
 
 	// Copy semantics are there for convenience only. Copied FUsdLevelSequenceHelper are empty and require a call to Init().
@@ -47,9 +45,6 @@ public:
 	/** Resets the helper, abandoning all managed LevelSequences */
 	void Clear();
 
-	UE_DEPRECATED(4.27, "This function is deprecated, use Init instead.")
-	void InitLevelSequence(const UE::FUsdStage& UsdStage);
-
 	/** Creates the time track for the StageActor */
 	void BindToUsdStageActor(AUsdStageActor* StageActor);
 	void UnbindFromUsdStageActor();
diff --git a/Engine/Plugins/Importers/USDImporter/Source/USDStage/Public/USDStageActor.h b/Engine/Plugins/Importers/USDImporter/Source/USDStage/Public/USDStageActor.h
--- a/Engine/Plugins/Importers/USDImporter/Source/USDStage/Public/USDStageActor.h
+++ b/Engine/Plugins/Importers/USDImporter/Source/USDStage/Public/USDStageActor.h
@@ -167,15 +167,6 @@ private:
 	UPROPERTY(EditAnywhere, Category = "USD")
 	float Time;
 
-	UPROPERTY()
-	float StartTimeCode_DEPRECATED;
-
-	UPROPERTY()
-	float EndTimeCode_DEPRECATED;
-
-	UPROPERTY()
-	float TimeCodesPerSecond_DEPRECATED;
-
 	UPROPERTY(VisibleAnywhere, Category = "USD", Transient)
 	TObjectPtr<ULevelSequence> LevelSequence;
 
@@ -277,9 +268,6 @@ private:
 	TObjectPtr<UUsdTransactor> Transactor;
 
 public:
-	UE_DEPRECATED( 4.27, "Use the const version if you don't wish to load the stage on-demand, or use GetOrLoadUsdStage if you do" )
-	USDSTAGE_API UE::FUsdStage& GetUsdStage();
-
 	// Loads the stage with RootLayer if its not loaded already, and returns the either the isolated stage (if any) or
 	// the base stage
 	USDSTAGE_API UE::FUsdStage& GetOrLoadUsdStage();
diff --git a/Engine/Plugins/Importers/USDImporter/Source/USDUtilities/Private/USDGeomMeshConversion.cpp b/Engine/Plugins/Importers/USDImporter/Source/USDUtilities/Private/USDGeomMeshConversion.cpp
--- a/Engine/Plugins/Importers/USDImporter/Source/USDUtilities/Private/USDGeomMeshConversion.cpp
+++ b/Engine/Plugins/Importers/USDImporter/Source/USDUtilities/Private/USDGeomMeshConversion.cpp
@@ -291,17 +291,9 @@ namespace UE::UsdGeomMeshConversion::Private
 			// This LOD has a single material assignment, just add an unrealMaterials attribute to the mesh prim
 			if ( bHasUEMaterialAssignements && UnrealMaterialsForLOD.size() == 1 )
 			{
-				const bool bHasMaterialAttribute = MaterialPrim.HasAttribute( UnrealIdentifiers::MaterialAssignments );
-				if ( bHasUEMaterialAssignements )
+				if ( pxr::UsdAttribute UEMaterialsAttribute = MaterialPrim.CreateAttribute( UnrealIdentifiers::MaterialAssignment, pxr::SdfValueTypeNames->String ) )
 				{
-					if ( pxr::UsdAttribute UEMaterialsAttribute = MaterialPrim.CreateAttribute( UnrealIdentifiers::MaterialAssignment, pxr::SdfValueTypeNames->String ) )
-					{
-						UEMaterialsAttribute.Set( UnrealMaterialsForLOD[ 0 ] );
-					}
-				}
-				else if ( bHasMaterialAttribute )
-				{
-					MaterialPrim.GetAttribute( UnrealIdentifiers::MaterialAssignments ).Clear();
+					UEMaterialsAttribute.Set( UnrealMaterialsForLOD[ 0 ] );
 				}
 			}
 			// Multiple material assignments to the same LOD (and so the same mesh prim). Need to create a GeomSubset for each UE mesh section
@@ -1348,31 +1340,6 @@ bool UsdToUnreal::ConvertGeomMeshHierarchy(
 	);
 }
 
-bool UsdToUnreal::ConvertDisplayColor( const UsdUtils::FDisplayColorMaterial& DisplayColorDescription, UMaterialInstanceConstant& Material )
-{
-	FUsdLogManager::LogMessage( EMessageSeverity::Warning, LOCTEXT( "DeprecatedConvertDisplayColor", "Converting existing instances with UsdToUnreal::ConvertDisplayColor is deprecated in favor of just calling UsdUtils::CreateDisplayColorMaterialInstanceConstant instead, and may be removed in a future release." ) );
-
-#if WITH_EDITOR
-	FString ParentPath = DisplayColorDescription.bHasOpacity
-		? TEXT("Material'/USDImporter/Materials/DisplayColorAndOpacity.DisplayColorAndOpacity'")
-		: TEXT("Material'/USDImporter/Materials/DisplayColor.DisplayColor'");
-
-	if ( UMaterialInterface* ParentMaterial = Cast< UMaterialInterface >( FSoftObjectPath( ParentPath ).TryLoad() ) )
-	{
-		UMaterialEditingLibrary::SetMaterialInstanceParent( &Material, ParentMaterial );
-	}
-
-	if ( DisplayColorDescription.bIsDoubleSided )
-	{
-		Material.BasePropertyOverrides.bOverride_TwoSided = true;
-		Material.BasePropertyOverrides.TwoSided = DisplayColorDescription.bIsDoubleSided;
-		Material.PostEditChange();
-	}
-#endif // WITH_EDITOR
-
-	return true;
-}
-
 UMaterialInstanceDynamic* UsdUtils::CreateDisplayColorMaterialInstanceDynamic( const UsdUtils::FDisplayColorMaterial& DisplayColorDescription )
 {
 	FString ParentPath;
@@ -1491,23 +1458,6 @@ UsdUtils::FUsdPrimMaterialAssignmentInfo UsdUtils::GetPrimMaterialAssignments(
 				ValidPackagePath = UsdToUnreal::ConvertString( UEMaterial );
 			}
 		}
-		else if ( pxr::UsdAttribute MaterialsAttribute = UsdPrim.GetAttribute( UnrealIdentifiers::MaterialAssignments ) )
-		{
-			pxr::VtStringArray UEMaterials;
-			MaterialsAttribute.Get( &UEMaterials, TimeCode );
-
-			if ( UEMaterials.size() > 0 && UEMaterials[ 0 ].size() > 0 )
-			{
-				ValidPackagePath = UsdToUnreal::ConvertString( UEMaterials[ 0 ] );
-
-				UE_LOG( LogUsd, Warning, TEXT( "String array attribute 'unrealMaterials' is deprecated: Use the singular string 'unrealMaterial' attribute" ) );
-				if ( UEMaterials.size() > 1 )
-				{
-					UE_LOG( LogUsd, Warning, TEXT( "Found more than one Unreal material assigned to Mesh '%s'. The first material ('%s') will be chosen, and the rest ignored." ),
-						*UsdToUnreal::ConvertPath( UsdPrim.GetPath() ), *ValidPackagePath );
-				}
-			}
-		}
 
 		if ( !ValidPackagePath.IsEmpty() )
 		{
diff --git a/Engine/Plugins/Importers/USDImporter/Source/USDUtilities/Public/USDGeomMeshConversion.h b/Engine/Plugins/Importers/USDImporter/Source/USDUtilities/Public/USDGeomMeshConversion.h
--- a/Engine/Plugins/Importers/USDImporter/Source/USDUtilities/Public/USDGeomMeshConversion.h
+++ b/Engine/Plugins/Importers/USDImporter/Source/USDUtilities/Public/USDGeomMeshConversion.h
@@ -166,12 +166,6 @@ namespace UsdToUnreal
 		bool bSkipRootPrimTransformAndVisibility = false,
 		bool bMergeIdenticalMaterialSlots = true
 	);
-
-	/** Configure Material to become a vertex color/displayColor material, according to the given description */
-	USDUTILITIES_API bool ConvertDisplayColor(
-		const UsdUtils::FDisplayColorMaterial& DisplayColorDescription,
-		UMaterialInstanceConstant& Material
-	);
 }
 
 namespace UnrealToUsd
diff --git a/Engine/Plugins/Importers/USDImporter/Source/UnrealUSDWrapper/Private/USDListener.cpp b/Engine/Plugins/Importers/USDImporter/Source/UnrealUSDWrapper/Private/USDListener.cpp
--- a/Engine/Plugins/Importers/USDImporter/Source/UnrealUSDWrapper/Private/USDListener.cpp
+++ b/Engine/Plugins/Importers/USDImporter/Source/UnrealUSDWrapper/Private/USDListener.cpp
@@ -213,8 +213,6 @@ public:
 	virtual ~FUsdListenerImpl();
 
 	FUsdListener::FOnStageEditTargetChanged OnStageEditTargetChanged;
-	FUsdListener::FOnPrimsChanged OnPrimsChanged;
-	FUsdListener::FOnStageInfoChanged OnStageInfoChanged;
 	FUsdListener::FOnLayersChanged OnLayersChanged;
 	FUsdListener::FOnObjectsChanged OnObjectsChanged;
 
@@ -231,8 +229,6 @@ protected:
 	void HandleLayersChangedNotice( const pxr::SdfNotice::LayersDidChange& Notice );
 	void HandleLayerDirtinessChangedNotice( const pxr::SdfNotice::LayerDirtinessChanged& Notice );
 
-	void EmitDeprecatedEvents( const pxr::UsdNotice::ObjectsChanged & Notice );
-
 private:
 	pxr::TfNotice::Key RegisteredObjectsChangedKey;
 	pxr::TfNotice::Key RegisteredStageEditTargetChangedKey;
@@ -281,16 +277,6 @@ FUsdListener::FOnStageEditTargetChanged& FUsdListener::GetOnStageEditTargetChang
 	return Impl->OnStageEditTargetChanged;
 }
 
-FUsdListener::FOnPrimsChanged& FUsdListener::GetOnPrimsChanged()
-{
-	return Impl->OnPrimsChanged;
-}
-
-FUsdListener::FOnStageInfoChanged& FUsdListener::GetOnStageInfoChanged()
-{
-	return Impl->OnStageInfoChanged;
-}
-
 FUsdListener::FOnLayersChanged& FUsdListener::GetOnLayersChanged()
 {
 	return Impl->OnLayersChanged;
@@ -344,7 +330,7 @@ FUsdListenerImpl::~FUsdListenerImpl()
 #if USE_USD_SDK
 void FUsdListenerImpl::HandleObjectsChangedNotice( const pxr::UsdNotice::ObjectsChanged& Notice, const pxr::UsdStageWeakPtr& Sender )
 {
-	if ( !OnPrimsChanged.IsBound() && !OnStageInfoChanged.IsBound() && !OnObjectsChanged.IsBound() )
+	if ( !OnObjectsChanged.IsBound() )
 	{
 		return;
 	}
@@ -362,8 +348,6 @@ void FUsdListenerImpl::HandleObjectsChangedNotice( const pxr::UsdNotice::Objects
 		FScopedUnrealAllocs UnrealAllocs;
 		OnObjectsChanged.Broadcast( InfoChanges, ResyncChanges );
 	}
-
-	EmitDeprecatedEvents( Notice );
 }
 
 void FUsdListenerImpl::HandleStageEditTargetChangedNotice( const pxr::UsdNotice::StageEditTargetChanged& Notice, const pxr::UsdStageWeakPtr& Sender )
@@ -426,85 +410,6 @@ void FUsdListenerImpl::HandleLayerDirtinessChangedNotice( const pxr::SdfNotice::
 	TArray< FString > LayersNames; // We don't know which layer dirtiness has changed unfortunately so we'll provide an empty list for now.
 	OnLayersChanged.Broadcast(LayersNames);
 }
-
-void FUsdListenerImpl::EmitDeprecatedEvents( const pxr::UsdNotice::ObjectsChanged& Notice )
-{
-	using namespace pxr;
-	using PathRange = UsdNotice::ObjectsChanged::PathRange;
-
-	TMap< FString, bool > PrimsChangedList;
-	TArray< FString > StageChangedFields;
-
-	FScopedUsdAllocs UsdAllocs;
-
-	PathRange PathsToUpdate = Notice.GetResyncedPaths();
-	for ( PathRange::const_iterator It = PathsToUpdate.begin(); It != PathsToUpdate.end(); ++It )
-	{
-		constexpr bool bResync = true;
-		PrimsChangedList.Add( ANSI_TO_TCHAR( It->GetAbsoluteRootOrPrimPath().GetString().c_str() ), bResync );
-	}
-
-	PathsToUpdate = Notice.GetChangedInfoOnlyPaths();
-	for ( PathRange::const_iterator PathToUpdateIt = PathsToUpdate.begin(); PathToUpdateIt != PathsToUpdate.end(); ++PathToUpdateIt )
-	{
-		const FString PrimPath = ANSI_TO_TCHAR( PathToUpdateIt->GetAbsoluteRootOrPrimPath().GetString().c_str() );
-
-		if ( PrimsChangedList.Contains( PrimPath ) )
-		{
-			continue;
-		}
-
-		bool bResync = false;
-
-		// If the layer reloaded, anything could have happened, so we must resync from the layer down
-		const std::vector<const SdfChangeList::Entry*>& Changes = PathToUpdateIt.base()->second;
-		for ( const SdfChangeList::Entry* Change : Changes )
-		{
-			if ( Change && Change->flags.didReloadContent )
-			{
-				bResync = true;
-			}
-		}
-
-		// Change on the stage root
-		bool bIsRootLayer = PathToUpdateIt->GetAbsoluteRootOrPrimPath() == pxr::SdfPath::AbsoluteRootPath();
-		if ( bIsRootLayer )
-		{
-			pxr::TfTokenVector ChangedFields = PathToUpdateIt.GetChangedFields();
-
-			for ( const pxr::TfToken& ChangedField : ChangedFields )
-			{
-				StageChangedFields.Add( ANSI_TO_TCHAR( ChangedField.GetString().c_str() ) );
-
-				// Force a resync when changing these since it affects all coordinates (the transforms they had when collapsed, etc.)
-				if ( ChangedField == UsdGeomTokens->metersPerUnit || ChangedField == UsdGeomTokens->upAxis )
-				{
-					bResync = true;
-					break;
-				}
-			}
-		}
-
-		// If the change is just about the stage updating some info like startTimeSeconds or framesPerSecond, then we don't
-		// need to refresh all prims. This is important because when undo/redoing, we may resync the movie scene, which shouldn't trigger prim spawning
-		if ( !bIsRootLayer || bResync )
-		{
-			PrimsChangedList.Add( PrimPath, bResync );
-		}
-	}
-
-	if ( PrimsChangedList.Num() > 0 )
-	{
-		FScopedUnrealAllocs UnrealAllocs;
-		OnPrimsChanged.Broadcast( PrimsChangedList );
-	}
-
-	if ( StageChangedFields.Num() > 0 )
-	{
-		FScopedUnrealAllocs UnrealAllocs;
-		OnStageInfoChanged.Broadcast( StageChangedFields );
-	}
-}
 #endif // #if USE_USD_SDK
 
 FScopedBlockNotices::FScopedBlockNotices( FUsdListener& InListener )
diff --git a/Engine/Plugins/Importers/USDImporter/Source/UnrealUSDWrapper/Private/UnrealUSDWrapper.cpp b/Engine/Plugins/Importers/USDImporter/Source/UnrealUSDWrapper/Private/UnrealUSDWrapper.cpp
--- a/Engine/Plugins/Importers/USDImporter/Source/UnrealUSDWrapper/Private/UnrealUSDWrapper.cpp
+++ b/Engine/Plugins/Importers/USDImporter/Source/UnrealUSDWrapper/Private/UnrealUSDWrapper.cpp
@@ -73,17 +73,6 @@ using std::string;
 
 using namespace pxr;
 
-#ifdef PLATFORM_WINDOWS
-	#define USDWRAPPER_USE_XFORMACHE	1
-#else
-	#define USDWRAPPER_USE_XFORMACHE	0
-#endif
-
-
-#if USDWRAPPER_USE_XFORMACHE
-static TUsdStore< UsdGeomXformCache > XFormCache;
-#endif // USDWRAPPER_USE_XFORMACHE
-
 namespace UnrealIdentifiers
 {
 	static const TfToken AssetPath("unrealAssetPath");
@@ -98,7 +87,6 @@ namespace UnrealIdentifiers
 	 */
 	const TfToken LOD("LOD");
 
-	const TfToken MaterialAssignments = TfToken("unrealMaterials"); // DEPRECATED in favor of MaterialAssignment
 	const TfToken MaterialAssignment = TfToken("unrealMaterial");
 	const TfToken Unreal = TfToken("unreal");
 
diff --git a/Engine/Plugins/Importers/USDImporter/Source/UnrealUSDWrapper/Public/USDListener.h b/Engine/Plugins/Importers/USDImporter/Source/UnrealUSDWrapper/Public/USDListener.h
--- a/Engine/Plugins/Importers/USDImporter/Source/UnrealUSDWrapper/Public/USDListener.h
+++ b/Engine/Plugins/Importers/USDImporter/Source/UnrealUSDWrapper/Public/USDListener.h
@@ -112,16 +112,6 @@ public:
 	DECLARE_EVENT( FUsdListener, FOnStageEditTargetChanged );
 	FOnStageEditTargetChanged& GetOnStageEditTargetChanged();
 
-	using FPrimsChangedList = TMap< FString, bool >;
-	DECLARE_EVENT_OneParam( FUsdListener, FOnPrimsChanged, const FPrimsChangedList& );
-	UE_DEPRECATED( 4.27, "This event is deprecated, prefer FOnObjectsChanged instead" )
-	FOnPrimsChanged& GetOnPrimsChanged();
-
-	using FStageChangedFields = TArray< FString >;
-	DECLARE_EVENT_OneParam( FUsdListener, FOnStageInfoChanged, const FStageChangedFields& );
-	UE_DEPRECATED( 4.27, "This event is deprecated, prefer FOnObjectsChanged instead (stage changes are the same as changes to the pseudoroot a.k.a. '/')" )
-	FOnStageInfoChanged& GetOnStageInfoChanged();
-
 	DECLARE_EVENT_OneParam( FUsdListener, FOnLayersChanged, const TArray< FString >& );
 	FOnLayersChanged& GetOnLayersChanged();
 
diff --git a/Engine/Plugins/Importers/USDImporter/Source/UnrealUSDWrapper/Public/UnrealUSDWrapper.h b/Engine/Plugins/Importers/USDImporter/Source/UnrealUSDWrapper/Public/UnrealUSDWrapper.h
--- a/Engine/Plugins/Importers/USDImporter/Source/UnrealUSDWrapper/Public/UnrealUSDWrapper.h
+++ b/Engine/Plugins/Importers/USDImporter/Source/UnrealUSDWrapper/Public/UnrealUSDWrapper.h
@@ -394,7 +394,6 @@ namespace UnrealIdentifiers
 	extern UNREALUSDWRAPPER_API const pxr::TfToken LOD;
 
 	/* Attribute name when assigning Unreal materials to UsdGeomMeshes */
-	extern UNREALUSDWRAPPER_API const pxr::TfToken MaterialAssignments; // DEPRECATED in favor of MaterialAssignment
 	extern UNREALUSDWRAPPER_API const pxr::TfToken MaterialAssignment;
 	extern UNREALUSDWRAPPER_API const pxr::TfToken Unreal;
 
