USD: Fix export options not being included on the asset hash.

diff --git a/Engine/Plugins/Importers/USDImporter/Source/USDClasses/Private/USDClassesModule.cpp b/Engine/Plugins/Importers/USDImporter/Source/USDClasses/Private/USDClassesModule.cpp
--- a/Engine/Plugins/Importers/USDImporter/Source/USDClasses/Private/USDClassesModule.cpp
+++ b/Engine/Plugins/Importers/USDImporter/Source/USDClasses/Private/USDClassesModule.cpp
@@ -9,41 +9,41 @@
 #include "HAL/FileManager.h"
 #include "Misc/EngineVersion.h"
 #include "Misc/FileHelper.h"
+#include "Misc/PackageName.h"
 #include "Misc/Paths.h"
+#include "Misc/SecureHash.h"
 #include "Modules/ModuleManager.h"
 #include "Serialization/JsonSerializer.h"
+#include "UObject/ObjectSaveContext.h"
+#include "UObject/Package.h"
 
 DEFINE_LOG_CATEGORY( LogUsd );
 
-namespace UE
+namespace UE::USDClasses::Private
 {
-	namespace USDClasses
+	TSharedPtr<FJsonObject> ParseJSON( const FString& FileContents )
 	{
-		namespace Private
+		if ( FileContents.IsEmpty() )
 		{
-			TSharedPtr<FJsonObject> ParseJSON( const FString& FileContents )
-			{
-				if ( FileContents.IsEmpty() )
-				{
-					return nullptr;
-				}
+			return nullptr;
+		}
 
-				const TSharedRef< TJsonReader<> >& Reader = TJsonReaderFactory<>::Create( FileContents );
+		const TSharedRef< TJsonReader<> >& Reader = TJsonReaderFactory<>::Create( FileContents );
 
-				TSharedPtr<FJsonObject> DescriptorObject;
-				if ( FJsonSerializer::Deserialize( Reader, DescriptorObject ) && DescriptorObject.IsValid() )
-				{
-					return DescriptorObject;
-				}
-				else
-				{
-					UE_LOG( LogUsd, Warning, TEXT( "Failed to parse plugInfo.json file: '%s'" ), *Reader->GetErrorMessage() );
-				}
-
-				return nullptr;
-			}
+		TSharedPtr<FJsonObject> DescriptorObject;
+		if ( FJsonSerializer::Deserialize( Reader, DescriptorObject ) && DescriptorObject.IsValid() )
+		{
+			return DescriptorObject;
+		}
+		else
+		{
+			UE_LOG( LogUsd, Warning, TEXT( "Failed to parse plugInfo.json file: '%s'" ), *Reader->GetErrorMessage() );
 		}
+
+		return nullptr;
 	}
+
+	TMap<FString, int32> PackagePathNameToDirtyCounter;
 }
 
 void IUsdClassesModule::UpdatePlugInfoFiles( const FString& PluginDirectory, const FString& TargetDllFolder )
@@ -200,8 +200,86 @@ void IUsdClassesModule::SendAnalytics( TArray<FAnalyticsEventAttribute>&& InAttr
 	}
 }
 
+bool IUsdClassesModule::HashObjectPackage( const UObject* Object, FSHA1& HashToUpdate )
+{
+#if WITH_EDITOR
+	if ( !Object )
+	{
+		return false;
+	}
+
+	UPackage* Package = Object->GetOutermost();
+	if ( !Package )
+	{
+		return false;
+	}
+
+	// Hash package's persistent Guid
+	const FGuid& Guid = Package->GetPersistentGuid();
+	HashToUpdate.Update( reinterpret_cast< const uint8* >( &Guid ), sizeof( Guid ) );
+
+	// Hash last modified date
+	FString PackageFullName = Package->GetPathName();
+	FString FileName;
+	if ( FPackageName::TryConvertLongPackageNameToFilename( PackageFullName, FileName ) )
+	{
+		FFileStatData StatData = IFileManager::Get().GetStatData( *FileName );
+		if ( StatData.bIsValid )
+		{
+			FString ModifiedTimeString = StatData.ModificationTime.ToString();
+			HashToUpdate.UpdateWithString( *ModifiedTimeString, ModifiedTimeString.Len() );
+		}
+	}
+
+	// If this asset is currently dirty, also hash how many times it was dirtied in this session.
+	// If it's ever saved, we'll reset this counter but update the last saved date
+	if ( int32* DirtyCounter = UE::USDClasses::Private::PackagePathNameToDirtyCounter.Find( PackageFullName ) )
+	{
+		HashToUpdate.Update( reinterpret_cast< const uint8* >( DirtyCounter ), sizeof( *DirtyCounter ) );
+	}
+
+	return true;
+#else
+	return false;
+#endif // WITH_EDITOR
+}
+
 class FUsdClassesModule : public IUsdClassesModule
 {
+
+public:
+	virtual void StartupModule() override
+	{
+		PackageMarkedDirtyEventHandle = UPackage::PackageMarkedDirtyEvent.AddLambda(
+			[]( const UPackage* Package, bool bWasDirty )
+			{
+				if ( Package )
+				{
+					UE::USDClasses::Private::PackagePathNameToDirtyCounter.FindOrAdd( Package->GetPathName() )++;
+				}
+			}
+		);
+
+		PackageSavedWithContextEventHandle = UPackage::PackageSavedWithContextEvent.AddLambda(
+			[]( const FString& PackageFilename, UPackage* Package, FObjectPostSaveContext ObjectSaveContext )
+			{
+				if ( Package )
+				{
+					UE::USDClasses::Private::PackagePathNameToDirtyCounter.Remove( Package->GetPathName() );
+				}
+			}
+		);
+	}
+
+	virtual void ShutdownModule() override
+	{
+		UPackage::PackageSavedWithContextEvent.Remove( PackageSavedWithContextEventHandle );
+		UPackage::PackageMarkedDirtyEvent.Remove( PackageMarkedDirtyEventHandle );
+	}
+
+private:
+	FDelegateHandle PackageMarkedDirtyEventHandle;
+	FDelegateHandle PackageSavedWithContextEventHandle;
 };
 
 IMPLEMENT_MODULE( FUsdClassesModule, USDClasses );
diff --git a/Engine/Plugins/Importers/USDImporter/Source/USDClasses/Private/USDStageOptions.cpp b/Engine/Plugins/Importers/USDImporter/Source/USDClasses/Private/USDStageOptions.cpp
--- a/Engine/Plugins/Importers/USDImporter/Source/USDClasses/Private/USDStageOptions.cpp
+++ b/Engine/Plugins/Importers/USDImporter/Source/USDClasses/Private/USDStageOptions.cpp
@@ -12,3 +12,8 @@ void UsdUtils::AddAnalyticsAttributes(
 	InOutAttributes.Emplace( TEXT( "MetersPerUnit" ), LexToString( Options.MetersPerUnit ) );
 	InOutAttributes.Emplace( TEXT( "UpAxis" ), Options.UpAxis == EUsdUpAxis::YAxis ? TEXT( "Y" ) : TEXT( "Z" ) );
 }
+
+void UsdUtils::HashForExport( const FUsdStageOptions& Options, FSHA1& HashToUpdate )
+{
+	HashToUpdate.Update( reinterpret_cast< const uint8* >( &Options ), sizeof( Options ) );
+}
diff --git a/Engine/Plugins/Importers/USDImporter/Source/USDClasses/Public/USDClassesModule.h b/Engine/Plugins/Importers/USDImporter/Source/USDClasses/Public/USDClassesModule.h
--- a/Engine/Plugins/Importers/USDImporter/Source/USDClasses/Public/USDClassesModule.h
+++ b/Engine/Plugins/Importers/USDImporter/Source/USDClasses/Public/USDClassesModule.h
@@ -22,5 +22,19 @@ public:
 	 * @param NumberOfFrames - Number of time codes in the exported/imported/opened stage
 	 * @param Extension - Extension of the main USD file opened/emitted/imported (e.g. "usda" or "usd")
 	 */
-	USDCLASSES_API static void SendAnalytics( TArray<FAnalyticsEventAttribute>&& InAttributes, const FString& EventName, bool bAutomated, double ElapsedSeconds, double NumberOfFrames, const FString& Extension );
+	USDCLASSES_API static void SendAnalytics(
+		TArray<FAnalyticsEventAttribute>&& InAttributes,
+		const FString& EventName,
+		bool bAutomated,
+		double ElapsedSeconds,
+		double NumberOfFrames,
+		const FString& Extension
+	);
+
+	/**
+	 * Updates HashToUpdate with the Object's package's persistent guid, the corresponding file save
+	 * date and time, and the number of times the package has been dirtied since last being saved.
+	 * This can be used to track the version of exported assets and levels, to prevent unnecessary re-exports.
+	 */
+	USDCLASSES_API static bool HashObjectPackage( const UObject* Object, FSHA1& HashToUpdate );
 };
\ No newline at end of file
diff --git a/Engine/Plugins/Importers/USDImporter/Source/USDClasses/Public/USDStageOptions.h b/Engine/Plugins/Importers/USDImporter/Source/USDClasses/Public/USDStageOptions.h
--- a/Engine/Plugins/Importers/USDImporter/Source/USDClasses/Public/USDStageOptions.h
+++ b/Engine/Plugins/Importers/USDImporter/Source/USDClasses/Public/USDStageOptions.h
@@ -33,4 +33,9 @@ namespace UsdUtils
 		const FUsdStageOptions& Options,
 		TArray< FAnalyticsEventAttribute >& InOutAttributes
 	);
+
+	USDCLASSES_API void HashForExport(
+		const FUsdStageOptions& Options,
+		FSHA1& HashToUpdate
+	);
 }
diff --git a/Engine/Plugins/Importers/USDImporter/Source/USDExporter/Private/AnimSequenceExporterUSD.cpp b/Engine/Plugins/Importers/USDImporter/Source/USDExporter/Private/AnimSequenceExporterUSD.cpp
--- a/Engine/Plugins/Importers/USDImporter/Source/USDExporter/Private/AnimSequenceExporterUSD.cpp
+++ b/Engine/Plugins/Importers/USDImporter/Source/USDExporter/Private/AnimSequenceExporterUSD.cpp
@@ -207,12 +207,13 @@ bool UAnimSequenceExporterUSD::ExportBinary( UObject* Object, const TCHAR* Type,
 	FString AnimSequenceVersion;
 	if ( UAnimDataModel* DataModel = AnimSequence->GetDataModel() )
 	{
-		AnimSequenceVersion = DataModel->GenerateGuid().ToString();
-
-		// We could just use the GUID directly but all other asset types end up with SHA hash size so lets be
-		// consistent
 		FSHA1 SHA1;
-		SHA1.UpdateWithString( *AnimSequenceVersion, AnimSequenceVersion.Len() );
+
+		FGuid DataModelGuid = DataModel->GenerateGuid();
+		SHA1.Update( reinterpret_cast< uint8* >( &DataModelGuid ), sizeof( DataModelGuid ) );
+
+		UsdUtils::HashForAnimSequenceExport( *Options, SHA1 );
+
 		SHA1.Final();
 		FSHAHash Hash;
 		SHA1.GetHash( &Hash.Hash[ 0 ] );
diff --git a/Engine/Plugins/Importers/USDImporter/Source/USDExporter/Private/AnimSequenceExporterUSDOptions.cpp b/Engine/Plugins/Importers/USDImporter/Source/USDExporter/Private/AnimSequenceExporterUSDOptions.cpp
--- a/Engine/Plugins/Importers/USDImporter/Source/USDExporter/Private/AnimSequenceExporterUSDOptions.cpp
+++ b/Engine/Plugins/Importers/USDImporter/Source/USDExporter/Private/AnimSequenceExporterUSDOptions.cpp
@@ -17,3 +17,9 @@ void UsdUtils::AddAnalyticsAttributes(
 	}
 	InOutAttributes.Emplace( TEXT( "ReExportIdenticalAssets" ), Options.bReExportIdenticalAssets );
 }
+
+void UsdUtils::HashForAnimSequenceExport( const UAnimSequenceExporterUSDOptions& Options, FSHA1& HashToUpdate )
+{
+	UsdUtils::HashForExport( Options.StageOptions, HashToUpdate );
+	HashToUpdate.Update( reinterpret_cast< const uint8* >( &Options.bExportPreviewMesh ), sizeof( Options.bExportPreviewMesh ) );
+}
diff --git a/Engine/Plugins/Importers/USDImporter/Source/USDExporter/Private/LevelExporterUSDOptions.cpp b/Engine/Plugins/Importers/USDImporter/Source/USDExporter/Private/LevelExporterUSDOptions.cpp
--- a/Engine/Plugins/Importers/USDImporter/Source/USDExporter/Private/LevelExporterUSDOptions.cpp
+++ b/Engine/Plugins/Importers/USDImporter/Source/USDExporter/Private/LevelExporterUSDOptions.cpp
@@ -47,3 +47,42 @@ void UsdUtils::AddAnalyticsAttributes(
 	InOutAttributes.Emplace( TEXT( "ReExportIdenticalAssets" ), Options.bReExportIdenticalAssets );
 }
 
+void UsdUtils::HashForLevelExport( const FLevelExporterUSDOptionsInner& Options, FSHA1& HashToUpdate )
+{
+	HashToUpdate.Update( reinterpret_cast< const uint8* >( &Options.bSelectionOnly ), sizeof( Options.bSelectionOnly ) );
+	HashToUpdate.Update( reinterpret_cast< const uint8* >( &Options.bExportActorFolders ), sizeof( Options.bExportActorFolders ) );
+	HashToUpdate.Update( reinterpret_cast< const uint8* >( &Options.bIgnoreSequencerAnimations ), sizeof( Options.bIgnoreSequencerAnimations ) );
+	HashToUpdate.Update( reinterpret_cast< const uint8* >( &Options.bExportFoliageOnActorsLayer ), sizeof( Options.bExportFoliageOnActorsLayer ) );
+
+	// If we changed where we want the assets exported we need to re-export them and update the reference paths
+	HashToUpdate.UpdateWithString( *Options.AssetFolder.Path, Options.AssetFolder.Path.Len() );
+
+	// This affects how we author material overrides
+	HashToUpdate.Update( reinterpret_cast< const uint8* >( &Options.AssetOptions.bBakeMaterials ), sizeof( Options.AssetOptions.bBakeMaterials ) );
+	HashToUpdate.Update( reinterpret_cast< const uint8* >( &Options.AssetOptions.bRemoveUnrealMaterials ), sizeof( Options.AssetOptions.bRemoveUnrealMaterials ) );
+
+	// This affects how we author material overrides and landscapes
+	HashToUpdate.Update( reinterpret_cast< const uint8* >( &Options.AssetOptions.bUsePayload ), sizeof( Options.AssetOptions.bUsePayload ) );
+	if ( Options.AssetOptions.bUsePayload )
+	{
+		HashToUpdate.UpdateWithString( *Options.AssetOptions.PayloadFormat, Options.AssetOptions.PayloadFormat.Len() );
+	}
+
+	// If we changed to/from exporting LODs, how we author material overrides may also need to change
+	HashToUpdate.Update( reinterpret_cast< const uint8* >( &Options.AssetOptions.LowestMeshLOD ), sizeof( Options.AssetOptions.LowestMeshLOD ) );
+	HashToUpdate.Update( reinterpret_cast< const uint8* >( &Options.AssetOptions.HighestMeshLOD ), sizeof( Options.AssetOptions.HighestMeshLOD ) );
+
+	HashToUpdate.Update( reinterpret_cast< const uint8* >( &Options.LowestLandscapeLOD ), sizeof( Options.LowestLandscapeLOD ) );
+	HashToUpdate.Update( reinterpret_cast< const uint8* >( &Options.HighestLandscapeLOD ), sizeof( Options.HighestLandscapeLOD ) );
+	HashToUpdate.Update( reinterpret_cast< const uint8* >( &Options.LandscapeBakeResolution ), sizeof( Options.LandscapeBakeResolution ) );
+	HashToUpdate.Update( reinterpret_cast< const uint8* >( &Options.bExportSublayers ), sizeof( Options.bExportSublayers ) );
+}
+
+void UsdUtils::HashForLevelExport( const ULevelExporterUSDOptions& Options, FSHA1& HashToUpdate )
+{
+	UsdUtils::HashForExport( Options.StageOptions, HashToUpdate );
+	HashToUpdate.Update( reinterpret_cast< const uint8* >( &Options.StartTimeCode ), sizeof( Options.StartTimeCode ) );
+	HashToUpdate.Update( reinterpret_cast< const uint8* >( &Options.EndTimeCode ), sizeof( Options.EndTimeCode ) );
+	UsdUtils::HashForLevelExport( Options.Inner, HashToUpdate );
+}
+
diff --git a/Engine/Plugins/Importers/USDImporter/Source/USDExporter/Private/LevelSequenceExporterUSD.cpp b/Engine/Plugins/Importers/USDImporter/Source/USDExporter/Private/LevelSequenceExporterUSD.cpp
--- a/Engine/Plugins/Importers/USDImporter/Source/USDExporter/Private/LevelSequenceExporterUSD.cpp
+++ b/Engine/Plugins/Importers/USDImporter/Source/USDExporter/Private/LevelSequenceExporterUSD.cpp
@@ -943,6 +943,10 @@ namespace UE
 					// consistent
 					FSHA1 SHA1;
 					SHA1.UpdateWithString( *LevelSequenceVersion, LevelSequenceVersion.Len() );
+					if ( Context.ExportOptions )
+					{
+						UsdUtils::HashForLevelSequenceExport( *Context.ExportOptions, SHA1 );
+					}
 					SHA1.Final();
 					FSHAHash Hash;
 					SHA1.GetHash( &Hash.Hash[ 0 ] );
diff --git a/Engine/Plugins/Importers/USDImporter/Source/USDExporter/Private/LevelSequenceExporterUSDOptions.cpp b/Engine/Plugins/Importers/USDImporter/Source/USDExporter/Private/LevelSequenceExporterUSDOptions.cpp
--- a/Engine/Plugins/Importers/USDImporter/Source/USDExporter/Private/LevelSequenceExporterUSDOptions.cpp
+++ b/Engine/Plugins/Importers/USDImporter/Source/USDExporter/Private/LevelSequenceExporterUSDOptions.cpp
@@ -33,4 +33,20 @@ void UsdUtils::AddAnalyticsAttributes(
 	{
 		UsdUtils::AddAnalyticsAttributes( Options.LevelExportOptions, InOutAttributes );
 	}
-}
\ No newline at end of file
+}
+
+void UsdUtils::HashForLevelSequenceExport( const ULevelSequenceExporterUsdOptions& Options, FSHA1& HashToUpdate )
+{
+	UsdUtils::HashForExport( Options.StageOptions, HashToUpdate );
+	HashToUpdate.Update( reinterpret_cast< const uint8* >( &Options.TimeCodesPerSecond ), sizeof( Options.TimeCodesPerSecond ) );
+	HashToUpdate.Update( reinterpret_cast< const uint8* >( &Options.bOverrideExportRange ), sizeof( Options.bOverrideExportRange ) );
+	if ( Options.bOverrideExportRange )
+	{
+		HashToUpdate.Update( reinterpret_cast< const uint8* >( &Options.StartFrame ), sizeof( Options.StartFrame ) );
+		HashToUpdate.Update( reinterpret_cast< const uint8* >( &Options.EndFrame ), sizeof( Options.EndFrame ) );
+	}
+	HashToUpdate.Update( reinterpret_cast< const uint8* >( &Options.bExportSubsequencesAsLayers ), sizeof( Options.bExportSubsequencesAsLayers ) );
+
+	const bool bUsingLevelSublayer = Options.bExportLevel && Options.bUseExportedLevelAsSublayer;
+	HashToUpdate.Update( reinterpret_cast< const uint8* >( &bUsingLevelSublayer ), sizeof( bUsingLevelSublayer ) );
+}
diff --git a/Engine/Plugins/Importers/USDImporter/Source/USDExporter/Private/MaterialExporterUSD.cpp b/Engine/Plugins/Importers/USDImporter/Source/USDExporter/Private/MaterialExporterUSD.cpp
--- a/Engine/Plugins/Importers/USDImporter/Source/USDExporter/Private/MaterialExporterUSD.cpp
+++ b/Engine/Plugins/Importers/USDImporter/Source/USDExporter/Private/MaterialExporterUSD.cpp
@@ -66,9 +66,8 @@ namespace UE::MaterialExporterUSD::Private
 		);
 	}
 
-	FString GetMaterialHashString( const UMaterialInterface& Material )
+	void HashMaterial( const UMaterialInterface& Material, FSHA1& VersionHash )
 	{
-		FSHA1 VersionHash;
 		if ( const UMaterial* Parent = Material.GetMaterial() )
 		{
 			// Add StateId to the combined hash, which is updated on every compile
@@ -168,11 +167,21 @@ namespace UE::MaterialExporterUSD::Private
 				VersionHash.Update( reinterpret_cast< const uint8* >( &Parameter.ExpressionGUID ), sizeof( Parameter.ExpressionGUID ) );
 			}
 		}
+	}
 
-		VersionHash.Final();
-		FSHAHash Hash;
-		VersionHash.GetHash( &Hash.Hash[ 0 ] );
-		return Hash.ToString();
+	void HashOptions( const FUsdMaterialBakingOptions& Options, FSHA1& VersionHash )
+	{
+		VersionHash.Update(
+			reinterpret_cast< const uint8* >( &Options.DefaultTextureSize ),
+			sizeof( Options.DefaultTextureSize )
+		);
+
+		VersionHash.Update(
+			reinterpret_cast< const uint8* >( Options.Properties.GetData() ),
+			Options.Properties.Num() * Options.Properties.GetTypeSize()
+		);
+
+		VersionHash.UpdateWithString( *Options.TexturesDir.Path, Options.TexturesDir.Path.Len() );
 	}
 }
 
@@ -268,7 +277,13 @@ bool UMaterialExporterUsd::ExportMaterial(
 {
 #if USE_USD_SDK
 
-	FString MaterialHashString = UE::MaterialExporterUSD::Private::GetMaterialHashString( Material );
+	FSHAHash Hash;
+	FSHA1 SHA1;
+	UE::MaterialExporterUSD::Private::HashMaterial( Material, SHA1 );
+	UE::MaterialExporterUSD::Private::HashOptions( Options, SHA1 );
+	SHA1.Final();
+	SHA1.GetHash( &Hash.Hash[ 0 ] );
+	FString MaterialHashString = Hash.ToString();
 
 	// Check if we already have exported what we plan on exporting anyway
 	if ( FPaths::FileExists( FilePath.FilePath ) )
@@ -383,7 +398,7 @@ bool UMaterialExporterUsd::ExportMaterial(
 bool UMaterialExporterUsd::ExportMaterialsForStage(
 	const TArray<UMaterialInterface*>& Materials,
 	const FUsdMaterialBakingOptions& Options,
-	const UE::FUsdStage& Stage,
+	const FString& StageRootLayerPath,
 	bool bIsAssetLayer,
 	bool bUsePayload,
 	bool bRemoveUnrealMaterials,
@@ -393,7 +408,7 @@ bool UMaterialExporterUsd::ExportMaterialsForStage(
 )
 {
 #if USE_USD_SDK
-	if ( !Stage )
+	if ( StageRootLayerPath.IsEmpty() )
 	{
 		return false;
 	}
@@ -403,9 +418,7 @@ bool UMaterialExporterUsd::ExportMaterialsForStage(
 		return true;
 	}
 
-	const UE::FSdfLayer RootLayer = Stage.GetRootLayer();
-	const FString RootLayerPath = RootLayer.GetRealPath();
-	const FString ExtensionNoDot = FPaths::GetExtension( RootLayerPath );
+	const FString ExtensionNoDot = FPaths::GetExtension( StageRootLayerPath );
 
 	// If we have multiple materials *within this mesh* that want to be emitted to the same filepath we'll append
 	// a suffix, but we will otherwise overwrite any unrelated existing files that were there before we began the export.
@@ -457,6 +470,23 @@ bool UMaterialExporterUsd::ExportMaterialsForStage(
 		}
 	}
 
+	// We can only open the stage *after* we finished exporting the materials. This because if we're exporting over
+	// existing files, it could be that this stage still references the existing material layers that the individual
+	// material exports would try to replace, meaning the exports would fail as those files would be currently open.
+	const EUsdInitialLoadSet InitialLoadSet = EUsdInitialLoadSet::LoadAll;
+	const bool bUseStageCache = false;
+	UE::FUsdStage Stage = UnrealUSDWrapper::OpenStage( *StageRootLayerPath, InitialLoadSet, bUseStageCache );
+	if ( !Stage )
+	{
+		return false;
+	}
+
+	UE::FSdfLayer RootLayer = Stage.GetRootLayer();
+	if ( !RootLayer )
+	{
+		return false;
+	}
+
 	UsdUtils::ReplaceUnrealMaterialsWithBaked(
 		Stage,
 		RootLayer,
@@ -466,6 +496,8 @@ bool UMaterialExporterUsd::ExportMaterialsForStage(
 		bRemoveUnrealMaterials
 	);
 
+	RootLayer.Save();
+
 	return true;
 #else
 	return false;
diff --git a/Engine/Plugins/Importers/USDImporter/Source/USDExporter/Private/MaterialExporterUSDOptions.cpp b/Engine/Plugins/Importers/USDImporter/Source/USDExporter/Private/MaterialExporterUSDOptions.cpp
--- a/Engine/Plugins/Importers/USDImporter/Source/USDExporter/Private/MaterialExporterUSDOptions.cpp
+++ b/Engine/Plugins/Importers/USDImporter/Source/USDExporter/Private/MaterialExporterUSDOptions.cpp
@@ -11,4 +11,9 @@ void UsdUtils::AddAnalyticsAttributes(
 {
 	UsdUtils::AddAnalyticsAttributes( Options.MaterialBakingOptions, InOutAttributes );
 	InOutAttributes.Emplace( TEXT( "ReExportIdenticalAssets" ), Options.bReExportIdenticalAssets );
-}
\ No newline at end of file
+}
+
+void UsdUtils::HashForMaterialExport( const UMaterialExporterUSDOptions& Options, FSHA1& HashToUpdate )
+{
+	UsdUtils::HashForMaterialExport( Options.MaterialBakingOptions, HashToUpdate );
+}
diff --git a/Engine/Plugins/Importers/USDImporter/Source/USDExporter/Private/SkeletalMeshExporterUSD.cpp b/Engine/Plugins/Importers/USDImporter/Source/USDExporter/Private/SkeletalMeshExporterUSD.cpp
--- a/Engine/Plugins/Importers/USDImporter/Source/USDExporter/Private/SkeletalMeshExporterUSD.cpp
+++ b/Engine/Plugins/Importers/USDImporter/Source/USDExporter/Private/SkeletalMeshExporterUSD.cpp
@@ -142,6 +142,7 @@ bool USkeletalMeshExporterUsd::ExportBinary( UObject* Object, const TCHAR* Type,
 		const FString& DDCKey = RenderData->DerivedDataKey;
 		FSHA1 SHA1;
 		SHA1.UpdateWithString( *DDCKey, DDCKey.Len() );
+		UsdUtils::HashForSkeletalMeshExport( *Options, SHA1 );
 		SHA1.Final();
 		FSHAHash Hash;
 		SHA1.GetHash( &Hash.Hash[ 0 ] );
@@ -163,6 +164,8 @@ bool USkeletalMeshExporterUsd::ExportBinary( UObject* Object, const TCHAR* Type,
 		// If we don't want to re-export this asset we need to check if its the same version
 		else if ( !Options->bReExportIdenticalAssets )
 		{
+			bool bSkipMeshExport = false;
+
 			// Don't use the stage cache here as we want this stage to close within this scope in case
 			// we have to overwrite its files due to e.g. missing payload or anything like that
 			const bool bUseStageCache = false;
@@ -186,33 +189,37 @@ bool USkeletalMeshExporterUsd::ExportBinary( UObject* Object, const TCHAR* Type,
 							*UExporter::CurrentFilename
 						);
 
-						// Even if we're not going to export the mesh, we may still need to re-bake materials
-						if ( Options->MeshAssetOptions.bBakeMaterials )
-						{
-							TSet<UMaterialInterface*> MaterialsToBake;
-							for ( const FSkeletalMaterial& SkeletalMaterial : SkeletalMesh->GetMaterials() )
-							{
-								MaterialsToBake.Add( SkeletalMaterial.MaterialInterface );
-							}
-
-							const bool bIsAssetLayer = true;
-							UMaterialExporterUsd::ExportMaterialsForStage(
-								MaterialsToBake.Array(),
-								Options->MeshAssetOptions.MaterialBakingOptions,
-								TempStage,
-								bIsAssetLayer,
-								Options->MeshAssetOptions.bUsePayload,
-								Options->MeshAssetOptions.bRemoveUnrealMaterials,
-								ExportTask->bReplaceIdentical,
-								Options->bReExportIdenticalAssets,
-								ExportTask->bAutomated
-							);
-							TempStage.GetRootLayer().Save();
-						}
-
-						return true;
+						bSkipMeshExport = true;
+					}
+				}
+			}
+
+			if ( bSkipMeshExport )
+			{
+				// Even if we're not going to export the mesh, we may still need to re-bake materials
+				if ( Options->MeshAssetOptions.bBakeMaterials )
+				{
+					TSet<UMaterialInterface*> MaterialsToBake;
+					for ( const FSkeletalMaterial& SkeletalMaterial : SkeletalMesh->GetMaterials() )
+					{
+						MaterialsToBake.Add( SkeletalMaterial.MaterialInterface );
 					}
+
+					const bool bIsAssetLayer = true;
+					UMaterialExporterUsd::ExportMaterialsForStage(
+						MaterialsToBake.Array(),
+						Options->MeshAssetOptions.MaterialBakingOptions,
+						UExporter::CurrentFilename,
+						bIsAssetLayer,
+						Options->MeshAssetOptions.bUsePayload,
+						Options->MeshAssetOptions.bRemoveUnrealMaterials,
+						ExportTask->bReplaceIdentical,
+						Options->bReExportIdenticalAssets,
+						ExportTask->bAutomated
+					);
 				}
+
+				return true;
 			}
 		}
 	}
@@ -300,7 +307,7 @@ bool USkeletalMeshExporterUsd::ExportBinary( UObject* Object, const TCHAR* Type,
 		UMaterialExporterUsd::ExportMaterialsForStage(
 			MaterialsToBake.Array(),
 			Options->MeshAssetOptions.MaterialBakingOptions,
-			AssetStage,
+			AssetStage.GetRootLayer().GetRealPath(),
 			bIsAssetLayer,
 			Options->MeshAssetOptions.bUsePayload,
 			Options->MeshAssetOptions.bRemoveUnrealMaterials,
diff --git a/Engine/Plugins/Importers/USDImporter/Source/USDExporter/Private/SkeletalMeshExporterUSDOptions.cpp b/Engine/Plugins/Importers/USDImporter/Source/USDExporter/Private/SkeletalMeshExporterUSDOptions.cpp
--- a/Engine/Plugins/Importers/USDImporter/Source/USDExporter/Private/SkeletalMeshExporterUSDOptions.cpp
+++ b/Engine/Plugins/Importers/USDImporter/Source/USDExporter/Private/SkeletalMeshExporterUSDOptions.cpp
@@ -13,3 +13,9 @@ void UsdUtils::AddAnalyticsAttributes(
 	UsdUtils::AddAnalyticsAttributes( Options.MeshAssetOptions, InOutAttributes );
 	InOutAttributes.Emplace( TEXT( "ReExportIdenticalAssets" ), Options.bReExportIdenticalAssets );
 }
+
+void UsdUtils::HashForSkeletalMeshExport( const USkeletalMeshExporterUSDOptions& Options, FSHA1& HashToUpdate )
+{
+	UsdUtils::HashForExport( Options.StageOptions, HashToUpdate );
+	UsdUtils::HashForMeshExport( Options.MeshAssetOptions, HashToUpdate );
+}
diff --git a/Engine/Plugins/Importers/USDImporter/Source/USDExporter/Private/StaticMeshExporterUSD.cpp b/Engine/Plugins/Importers/USDImporter/Source/USDExporter/Private/StaticMeshExporterUSD.cpp
--- a/Engine/Plugins/Importers/USDImporter/Source/USDExporter/Private/StaticMeshExporterUSD.cpp
+++ b/Engine/Plugins/Importers/USDImporter/Source/USDExporter/Private/StaticMeshExporterUSD.cpp
@@ -157,6 +157,7 @@ bool UStaticMeshExporterUsd::ExportBinary( UObject* Object, const TCHAR* Type, F
 		const FString& DDCKey = RenderData->DerivedDataKey;
 		FSHA1 SHA1;
 		SHA1.UpdateWithString( *DDCKey, DDCKey.Len() );
+		UsdUtils::HashForStaticMeshExport( *Options, SHA1 );
 		SHA1.Final();
 		FSHAHash Hash;
 		SHA1.GetHash( &Hash.Hash[ 0 ] );
@@ -178,6 +179,8 @@ bool UStaticMeshExporterUsd::ExportBinary( UObject* Object, const TCHAR* Type, F
 		// If we don't want to re-export this asset we need to check if its the same version
 		else if ( !Options->bReExportIdenticalAssets )
 		{
+			bool bSkipMeshExport = false;
+
 			// Don't use the stage cache here as we want this stage to close within this scope in case
 			// we have to overwrite its files due to e.g. missing payload or anything like that
 			const bool bUseStageCache = false;
@@ -201,33 +204,37 @@ bool UStaticMeshExporterUsd::ExportBinary( UObject* Object, const TCHAR* Type, F
 							*UExporter::CurrentFilename
 						);
 
-						// Even if we're not going to export the mesh, we may still need to re-bake materials
-						if ( Options->MeshAssetOptions.bBakeMaterials )
-						{
-							TSet<UMaterialInterface*> MaterialsToBake;
-							for ( const FStaticMaterial& StaticMaterial : StaticMesh->GetStaticMaterials() )
-							{
-								MaterialsToBake.Add( StaticMaterial.MaterialInterface );
-							}
-
-							const bool bIsAssetLayer = true;
-							UMaterialExporterUsd::ExportMaterialsForStage(
-								MaterialsToBake.Array(),
-								Options->MeshAssetOptions.MaterialBakingOptions,
-								TempStage,
-								bIsAssetLayer,
-								Options->MeshAssetOptions.bUsePayload,
-								Options->MeshAssetOptions.bRemoveUnrealMaterials,
-								ExportTask->bReplaceIdentical,
-								Options->bReExportIdenticalAssets,
-								ExportTask->bAutomated
-							);
-							TempStage.GetRootLayer().Save();
-						}
-
-						return true;
+						bSkipMeshExport = true;
+					}
+				}
+			}
+
+			if ( bSkipMeshExport )
+			{
+				// Even if we're not going to export the mesh, we may still need to re-bake materials
+				if ( Options->MeshAssetOptions.bBakeMaterials )
+				{
+					TSet<UMaterialInterface*> MaterialsToBake;
+					for ( const FStaticMaterial& StaticMaterial : StaticMesh->GetStaticMaterials() )
+					{
+						MaterialsToBake.Add( StaticMaterial.MaterialInterface );
 					}
+
+					const bool bIsAssetLayer = true;
+					UMaterialExporterUsd::ExportMaterialsForStage(
+						MaterialsToBake.Array(),
+						Options->MeshAssetOptions.MaterialBakingOptions,
+						UExporter::CurrentFilename,
+						bIsAssetLayer,
+						Options->MeshAssetOptions.bUsePayload,
+						Options->MeshAssetOptions.bRemoveUnrealMaterials,
+						ExportTask->bReplaceIdentical,
+						Options->bReExportIdenticalAssets,
+						ExportTask->bAutomated
+					);
 				}
+
+				return true;
 			}
 		}
 	}
@@ -311,7 +318,7 @@ bool UStaticMeshExporterUsd::ExportBinary( UObject* Object, const TCHAR* Type, F
 		UMaterialExporterUsd::ExportMaterialsForStage(
 			MaterialsToBake.Array(),
 			Options->MeshAssetOptions.MaterialBakingOptions,
-			AssetStage,
+			AssetStage.GetRootLayer().GetRealPath(),
 			bIsAssetLayer,
 			Options->MeshAssetOptions.bUsePayload,
 			Options->MeshAssetOptions.bRemoveUnrealMaterials,
diff --git a/Engine/Plugins/Importers/USDImporter/Source/USDExporter/Private/StaticMeshExporterUSDOptions.cpp b/Engine/Plugins/Importers/USDImporter/Source/USDExporter/Private/StaticMeshExporterUSDOptions.cpp
--- a/Engine/Plugins/Importers/USDImporter/Source/USDExporter/Private/StaticMeshExporterUSDOptions.cpp
+++ b/Engine/Plugins/Importers/USDImporter/Source/USDExporter/Private/StaticMeshExporterUSDOptions.cpp
@@ -12,4 +12,10 @@ void UsdUtils::AddAnalyticsAttributes(
 	UsdUtils::AddAnalyticsAttributes( Options.StageOptions, InOutAttributes );
 	UsdUtils::AddAnalyticsAttributes( Options.MeshAssetOptions, InOutAttributes );
 	InOutAttributes.Emplace( TEXT( "ReExportIdenticalAssets" ), Options.bReExportIdenticalAssets );
-}
\ No newline at end of file
+}
+
+void UsdUtils::HashForStaticMeshExport( const UStaticMeshExporterUSDOptions& Options, FSHA1& HashToUpdate )
+{
+	UsdUtils::HashForExport( Options.StageOptions, HashToUpdate );
+	UsdUtils::HashForMeshExport( Options.MeshAssetOptions, HashToUpdate );
+}
diff --git a/Engine/Plugins/Importers/USDImporter/Source/USDExporter/Private/USDAssetOptions.cpp b/Engine/Plugins/Importers/USDImporter/Source/USDExporter/Private/USDAssetOptions.cpp
--- a/Engine/Plugins/Importers/USDImporter/Source/USDExporter/Private/USDAssetOptions.cpp
+++ b/Engine/Plugins/Importers/USDImporter/Source/USDExporter/Private/USDAssetOptions.cpp
@@ -44,4 +44,49 @@ void UsdUtils::AddAnalyticsAttributes(
 	}
 	InOutAttributes.Emplace( TEXT( "LowestMeshLOD" ), LexToString( Options.LowestMeshLOD ) );
 	InOutAttributes.Emplace( TEXT( "HighestMeshLOD" ), LexToString( Options.HighestMeshLOD ) );
-}
\ No newline at end of file
+}
+
+void UsdUtils::HashForMaterialExport( const FUsdMaterialBakingOptions& Options, FSHA1& HashToUpdate )
+{
+	HashToUpdate.Update(
+		reinterpret_cast< const uint8* >( Options.Properties.GetData() ),
+		Options.Properties.Num() * Options.Properties.GetTypeSize()
+	);
+
+	HashToUpdate.Update(
+		reinterpret_cast< const uint8* >( &Options.DefaultTextureSize ),
+		sizeof( Options.DefaultTextureSize )
+	);
+
+	// If we changed where we want the textures exported we need to re-export them and update the texture paths on the
+	// material
+	HashToUpdate.UpdateWithString( *Options.TexturesDir.Path, Options.TexturesDir.Path.Len() );
+}
+
+void UsdUtils::HashForMeshExport( const FUsdMeshAssetOptions& Options, FSHA1& HashToUpdate )
+{
+	HashToUpdate.Update( reinterpret_cast< const uint8* >( &Options.bUsePayload ), sizeof( Options.bUsePayload ) );
+	if ( Options.bUsePayload )
+	{
+		HashToUpdate.UpdateWithString( *Options.PayloadFormat, Options.PayloadFormat.Len() );
+	}
+
+	// These can affect the material bindings we'll write on the exported layer with the mesh
+	HashToUpdate.Update(
+		reinterpret_cast< const uint8* >( &Options.bBakeMaterials ),
+		sizeof( Options.bBakeMaterials )
+	);
+	HashToUpdate.Update(
+		reinterpret_cast< const uint8* >( &Options.bRemoveUnrealMaterials ),
+		sizeof( Options.bRemoveUnrealMaterials )
+	);
+
+	HashToUpdate.Update(
+		reinterpret_cast< const uint8* >( &Options.LowestMeshLOD ),
+		sizeof( Options.LowestMeshLOD )
+	);
+	HashToUpdate.Update(
+		reinterpret_cast< const uint8* >( &Options.HighestMeshLOD ),
+		sizeof( Options.HighestMeshLOD )
+	);
+}
diff --git a/Engine/Plugins/Importers/USDImporter/Source/USDExporter/Private/USDConversionBlueprintLibrary.cpp b/Engine/Plugins/Importers/USDImporter/Source/USDExporter/Private/USDConversionBlueprintLibrary.cpp
--- a/Engine/Plugins/Importers/USDImporter/Source/USDExporter/Private/USDConversionBlueprintLibrary.cpp
+++ b/Engine/Plugins/Importers/USDImporter/Source/USDExporter/Private/USDConversionBlueprintLibrary.cpp
@@ -291,9 +291,29 @@ TSet<AActor*> UUsdConversionBlueprintLibrary::GetActorsToConvert( UWorld* World
 	return Result;
 }
 
-FString UUsdConversionBlueprintLibrary::GeneratePackageVersionGuidString( const UPackage* Package )
+FString UUsdConversionBlueprintLibrary::GenerateObjectVersionString( const UObject* ObjectToExport, UObject* ExportOptions )
 {
-	return IUsdExporterModule::GeneratePackageVersionGuidString( Package );
+	if ( !ObjectToExport )
+	{
+		return {};
+	}
+
+	FSHA1 SHA1;
+
+	if ( !IUsdClassesModule::HashObjectPackage( ObjectToExport, SHA1 ) )
+	{
+		return {};
+	}
+
+	if ( ULevelExporterUSDOptions* LevelExportOptions = Cast<ULevelExporterUSDOptions>( ExportOptions ) )
+	{
+		UsdUtils::HashForLevelExport( *LevelExportOptions, SHA1 );
+	}
+
+	FSHAHash Hash;
+	SHA1.Final();
+	SHA1.GetHash( &Hash.Hash[ 0 ] );
+	return Hash.ToString();
 }
 
 FString UUsdConversionBlueprintLibrary::MakePathRelativeToLayer( const FString& AnchorLayerPath, const FString& PathToMakeRelative )
diff --git a/Engine/Plugins/Importers/USDImporter/Source/USDExporter/Private/USDExporterModule.cpp b/Engine/Plugins/Importers/USDImporter/Source/USDExporter/Private/USDExporterModule.cpp
--- a/Engine/Plugins/Importers/USDImporter/Source/USDExporter/Private/USDExporterModule.cpp
+++ b/Engine/Plugins/Importers/USDImporter/Source/USDExporter/Private/USDExporterModule.cpp
@@ -6,16 +6,10 @@
 #include "USDAssetOptions.h"
 #include "USDMemory.h"
 
-#include "HAL/FileManager.h"
 #include "Modules/ModuleManager.h"
 #include "PropertyEditorModule.h"
 #include "UObject/ObjectSaveContext.h"
 
-namespace UE::UsdExporterModule::Private
-{
-	TMap<FString, int32> PackagePathNameToDirtyCounter;
-}
-
 class FUsdExporterModule : public IUsdExporterModule
 {
 public:
@@ -43,80 +37,17 @@ public:
 				HighestMeshLODProperty->SetMetaData( TEXT( "ClampMax" ), LexToString( MAX_MESH_LOD_COUNT - 1 ) );
 			}
 		}
-
-		PackageMarkedDirtyEventHandle = UPackage::PackageMarkedDirtyEvent.AddLambda([](const UPackage* Package, bool bWasDirty)
-		{
-			if ( Package )
-			{
-				UE::UsdExporterModule::Private::PackagePathNameToDirtyCounter.FindOrAdd( Package->GetPathName() )++;
-			}
-		});
-
-		PackageSavedWithContextEventHandle = UPackage::PackageSavedWithContextEvent.AddLambda(
-			[]( const FString& PackageFilename, UPackage* Package, FObjectPostSaveContext ObjectSaveContext )
-			{
-				if ( Package )
-				{
-					UE::UsdExporterModule::Private::PackagePathNameToDirtyCounter.Remove( Package->GetPathName() );
-				}
-			}
-		);
 	}
 
 	virtual void ShutdownModule() override
 	{
-		UPackage::PackageMarkedDirtyEvent.Remove( PackageSavedWithContextEventHandle );
-		UPackage::PackageMarkedDirtyEvent.Remove( PackageMarkedDirtyEventHandle );
-
 		if ( FPropertyEditorModule* PropertyModule = FModuleManager::GetModulePtr< FPropertyEditorModule >( TEXT( "PropertyEditor" ) ) )
 		{
 			PropertyModule->UnregisterCustomClassLayout( TEXT( "LevelExporterUSDOptions" ) );
 			PropertyModule->UnregisterCustomClassLayout( TEXT( "LevelSequenceExporterUSDOptions" ) );
 		}
 	}
-
-private:
-	FDelegateHandle PackageMarkedDirtyEventHandle;
-	FDelegateHandle PackageSavedWithContextEventHandle;
 };
 
-FString IUsdExporterModule::GeneratePackageVersionGuidString( const UPackage* Package )
-{
-	if ( !Package )
-	{
-		return {};
-	}
-
-	// Hash package's persistent Guid
-	const FGuid& Guid = Package->GetPersistentGuid();
-	FSHA1 SHA1;
-	SHA1.Update( reinterpret_cast< const uint8* >( &Guid ), sizeof( Guid ) );
-
-	// Hash last modified date
-	FString PackageFullName = Package->GetPathName();
-	FString FileName;
-	if ( FPackageName::TryConvertLongPackageNameToFilename( PackageFullName, FileName ) )
-	{
-		FFileStatData StatData = IFileManager::Get().GetStatData( *FileName );
-		if ( StatData.bIsValid )
-		{
-			FString ModifiedTimeString = StatData.ModificationTime.ToString();
-			SHA1.UpdateWithString( *ModifiedTimeString, ModifiedTimeString.Len() );
-		}
-	}
-
-	// If this asset is currently dirty, also hash how many times it was dirtied in this session
-	// If its ever saved, we'll reset this counter but update the last saved date
-	if ( int32* DirtyCounter = UE::UsdExporterModule::Private::PackagePathNameToDirtyCounter.Find( PackageFullName ) )
-	{
-		SHA1.Update( reinterpret_cast< const uint8* >( DirtyCounter ), sizeof( *DirtyCounter ) );
-	}
-
-	SHA1.Final();
-	FSHAHash Hash;
-	SHA1.GetHash( &Hash.Hash[ 0 ] );
-	return Hash.ToString();
-}
-
 IMPLEMENT_MODULE_USD( FUsdExporterModule, USDExporter );
 
diff --git a/Engine/Plugins/Importers/USDImporter/Source/USDExporter/Public/AnimSequenceExporterUSDOptions.h b/Engine/Plugins/Importers/USDImporter/Source/USDExporter/Public/AnimSequenceExporterUSDOptions.h
--- a/Engine/Plugins/Importers/USDImporter/Source/USDExporter/Public/AnimSequenceExporterUSDOptions.h
+++ b/Engine/Plugins/Importers/USDImporter/Source/USDExporter/Public/AnimSequenceExporterUSDOptions.h
@@ -46,4 +46,9 @@ namespace UsdUtils
 		const UAnimSequenceExporterUSDOptions& Options,
 		TArray< FAnalyticsEventAttribute >& InOutAttributes
 	);
+
+	USDEXPORTER_API void HashForAnimSequenceExport(
+		const UAnimSequenceExporterUSDOptions& Options,
+		FSHA1& HashToUpdate
+	);
 }
diff --git a/Engine/Plugins/Importers/USDImporter/Source/USDExporter/Public/LevelExporterUSDOptions.h b/Engine/Plugins/Importers/USDImporter/Source/USDExporter/Public/LevelExporterUSDOptions.h
--- a/Engine/Plugins/Importers/USDImporter/Source/USDExporter/Public/LevelExporterUSDOptions.h
+++ b/Engine/Plugins/Importers/USDImporter/Source/USDExporter/Public/LevelExporterUSDOptions.h
@@ -129,8 +129,19 @@ namespace UsdUtils
 		const FLevelExporterUSDOptionsInner& Options,
 		TArray< FAnalyticsEventAttribute >& InOutAttributes
 	);
+
 	USDEXPORTER_API void AddAnalyticsAttributes(
 		const ULevelExporterUSDOptions& Options,
 		TArray< FAnalyticsEventAttribute >& InOutAttributes
 	);
+
+	USDEXPORTER_API void HashForLevelExport(
+		const FLevelExporterUSDOptionsInner& Options,
+		FSHA1& HashToUpdate
+	);
+
+	USDEXPORTER_API void HashForLevelExport(
+		const ULevelExporterUSDOptions& Options,
+		FSHA1& HashToUpdate
+	);
 }
\ No newline at end of file
diff --git a/Engine/Plugins/Importers/USDImporter/Source/USDExporter/Public/LevelSequenceExporterUSDOptions.h b/Engine/Plugins/Importers/USDImporter/Source/USDExporter/Public/LevelSequenceExporterUSDOptions.h
--- a/Engine/Plugins/Importers/USDImporter/Source/USDExporter/Public/LevelSequenceExporterUSDOptions.h
+++ b/Engine/Plugins/Importers/USDImporter/Source/USDExporter/Public/LevelSequenceExporterUSDOptions.h
@@ -92,4 +92,9 @@ namespace UsdUtils
 		const ULevelSequenceExporterUsdOptions& Options,
 		TArray< FAnalyticsEventAttribute >& InOutAttributes
 	);
+
+	USDEXPORTER_API void HashForLevelSequenceExport(
+		const ULevelSequenceExporterUsdOptions& Options,
+		FSHA1& HashToUpdate
+	);
 }
\ No newline at end of file
diff --git a/Engine/Plugins/Importers/USDImporter/Source/USDExporter/Public/MaterialExporterUSD.h b/Engine/Plugins/Importers/USDImporter/Source/USDExporter/Public/MaterialExporterUSD.h
--- a/Engine/Plugins/Importers/USDImporter/Source/USDExporter/Public/MaterialExporterUSD.h
+++ b/Engine/Plugins/Importers/USDImporter/Source/USDExporter/Public/MaterialExporterUSD.h
@@ -47,8 +47,11 @@ public:
 	 * to the exported UsdPreviewMaterial assets.
 	 * @param Materials - Materials to bake
 	 * @param Options - Options to use when baking
-	 * @param Stage - Stage to replace the unrealMaterials attributes in. All of it's layers will be traversed.
-	 *                The root layer folder on disk is used as destination folder for the baked materials
+	 * @param StageRootLayerPath - Path to the stage to replace the unrealMaterials attributes in.
+	 *                             Note that this stage must not be opened in case you want this function to be able to
+	 *                             replace exported materials with new files.
+	 *                             All of it's layers will be traversed.
+	 *                             The root layer folder on disk is used as destination folder for the baked materials
 	 * @param bIsAssetLayer - True when we're exporting a single mesh/animation asset. False when we're exporting a level. Dictates minor behaviors
 	 *                        when authoring the material binding relationships, e.g. whether we author them inside variants or not
 	 * @param bUsePayload - Should be True if the Stage was exported using payload files to store the actual Mesh prims. Also dictates minor
@@ -63,7 +66,7 @@ public:
 	static bool ExportMaterialsForStage(
 		const TArray<UMaterialInterface*>& Materials,
 		const FUsdMaterialBakingOptions& Options,
-		const UE::FUsdStage& Stage,
+		const FString& StageRootLayerPath,
 		bool bIsAssetLayer,
 		bool bUsePayload,
 		bool bRemoveUnrealMaterials,
diff --git a/Engine/Plugins/Importers/USDImporter/Source/USDExporter/Public/MaterialExporterUSDOptions.h b/Engine/Plugins/Importers/USDImporter/Source/USDExporter/Public/MaterialExporterUSDOptions.h
--- a/Engine/Plugins/Importers/USDImporter/Source/USDExporter/Public/MaterialExporterUSDOptions.h
+++ b/Engine/Plugins/Importers/USDImporter/Source/USDExporter/Public/MaterialExporterUSDOptions.h
@@ -38,4 +38,9 @@ namespace UsdUtils
 		const UMaterialExporterUSDOptions& Options,
 		TArray< FAnalyticsEventAttribute >& InOutAttributes
 	);
+
+	USDEXPORTER_API void HashForMaterialExport(
+		const UMaterialExporterUSDOptions& Options,
+		FSHA1& HashToUpdate
+	);
 }
\ No newline at end of file
diff --git a/Engine/Plugins/Importers/USDImporter/Source/USDExporter/Public/SkeletalMeshExporterUSDOptions.h b/Engine/Plugins/Importers/USDImporter/Source/USDExporter/Public/SkeletalMeshExporterUSDOptions.h
--- a/Engine/Plugins/Importers/USDImporter/Source/USDExporter/Public/SkeletalMeshExporterUSDOptions.h
+++ b/Engine/Plugins/Importers/USDImporter/Source/USDExporter/Public/SkeletalMeshExporterUSDOptions.h
@@ -40,4 +40,9 @@ namespace UsdUtils
 		const USkeletalMeshExporterUSDOptions& Options,
 		TArray< FAnalyticsEventAttribute >& InOutAttributes
 	);
+
+	USDEXPORTER_API void HashForSkeletalMeshExport(
+		const USkeletalMeshExporterUSDOptions& Options,
+		FSHA1& HashToUpdate
+	);
 }
diff --git a/Engine/Plugins/Importers/USDImporter/Source/USDExporter/Public/StaticMeshExporterUSDOptions.h b/Engine/Plugins/Importers/USDImporter/Source/USDExporter/Public/StaticMeshExporterUSDOptions.h
--- a/Engine/Plugins/Importers/USDImporter/Source/USDExporter/Public/StaticMeshExporterUSDOptions.h
+++ b/Engine/Plugins/Importers/USDImporter/Source/USDExporter/Public/StaticMeshExporterUSDOptions.h
@@ -40,4 +40,9 @@ namespace UsdUtils
 		const UStaticMeshExporterUSDOptions& Options,
 		TArray< FAnalyticsEventAttribute >& InOutAttributes
 	);
+
+	USDEXPORTER_API void HashForStaticMeshExport(
+		const UStaticMeshExporterUSDOptions& Options,
+		FSHA1& HashToUpdate
+	);
 }
\ No newline at end of file
diff --git a/Engine/Plugins/Importers/USDImporter/Source/USDExporter/Public/USDAssetOptions.h b/Engine/Plugins/Importers/USDImporter/Source/USDExporter/Public/USDAssetOptions.h
--- a/Engine/Plugins/Importers/USDImporter/Source/USDExporter/Public/USDAssetOptions.h
+++ b/Engine/Plugins/Importers/USDImporter/Source/USDExporter/Public/USDAssetOptions.h
@@ -67,8 +67,19 @@ namespace UsdUtils
 		const FUsdMaterialBakingOptions& Options,
 		TArray< FAnalyticsEventAttribute >& InOutAttributes
 	);
+
 	USDEXPORTER_API void AddAnalyticsAttributes(
 		const FUsdMeshAssetOptions& Options,
 		TArray< FAnalyticsEventAttribute >& InOutAttributes
 	);
+
+	USDEXPORTER_API void HashForMaterialExport(
+		const FUsdMaterialBakingOptions& Options,
+		FSHA1& HashToUpdate
+	);
+
+	USDEXPORTER_API void HashForMeshExport(
+		const FUsdMeshAssetOptions& Options,
+		FSHA1& HashToUpdate
+	);
 }
\ No newline at end of file
diff --git a/Engine/Plugins/Importers/USDImporter/Source/USDExporter/Public/USDConversionBlueprintLibrary.h b/Engine/Plugins/Importers/USDImporter/Source/USDExporter/Public/USDConversionBlueprintLibrary.h
--- a/Engine/Plugins/Importers/USDImporter/Source/USDExporter/Public/USDConversionBlueprintLibrary.h
+++ b/Engine/Plugins/Importers/USDImporter/Source/USDExporter/Public/USDConversionBlueprintLibrary.h
@@ -60,13 +60,16 @@ public:
 	static TSet<AActor*> GetActorsToConvert( UWorld* World );
 
 	/**
-	 * Generates a unique identifier string that involves the package's persistent guid, the corresponding file save
-	 * date and time, and the number of times the package has been dirtied since last being saved.
+	 * Generates a unique identifier string that involves ObjectToExport's package's persistent guid, the
+	 * corresponding file save date and time, and the number of times the package has been dirtied since last being
+	 * saved.
+	 * Optionally it can also combine that hash with a hash of the export options being used for the export, if
+	 * available.
 	 * This can be used to track the version of exported assets and levels, to prevent reexporting of actors and
 	 * components.
 	 */
 	UFUNCTION( BlueprintCallable, Category = "World utils" )
-	static FString GeneratePackageVersionGuidString( const UPackage* Package );
+	static FString GenerateObjectVersionString( const UObject* ObjectToExport, UObject* ExportOptions );
 
 	UFUNCTION( BlueprintCallable, Category = "Layer utils" )
 	static FString MakePathRelativeToLayer( const FString& AnchorLayerPath, const FString& PathToMakeRelative );
diff --git a/Engine/Plugins/Importers/USDImporter/Source/USDExporter/Public/USDExporterModule.h b/Engine/Plugins/Importers/USDImporter/Source/USDExporter/Public/USDExporterModule.h
--- a/Engine/Plugins/Importers/USDImporter/Source/USDExporter/Public/USDExporterModule.h
+++ b/Engine/Plugins/Importers/USDImporter/Source/USDExporter/Public/USDExporterModule.h
@@ -7,12 +7,4 @@
 
 class IUsdExporterModule : public IModuleInterface
 {
-public:
-	/**
-	 * Generates a unique identifier string that involves the package's persistent guid, the corresponding file save
-	 * date and time, and the number of times the package has been dirtied since last being saved.
-	 * This can be used to track the version of exported assets and levels, to prevent reexporting of actors and
-	 * components.
-	 */
-	static FString GeneratePackageVersionGuidString( const UPackage* Package );
 };
diff --git a/Engine/Plugins/Importers/USDImporter/Source/UnrealUSDWrapper/Private/UsdWrappers/UsdStage.cpp b/Engine/Plugins/Importers/USDImporter/Source/UnrealUSDWrapper/Private/UsdWrappers/UsdStage.cpp
--- a/Engine/Plugins/Importers/USDImporter/Source/UnrealUSDWrapper/Private/UsdWrappers/UsdStage.cpp
+++ b/Engine/Plugins/Importers/USDImporter/Source/UnrealUSDWrapper/Private/UsdWrappers/UsdStage.cpp
@@ -402,6 +402,48 @@ namespace UE
 		return Result;
 	}
 
+	template<typename PtrType>
+	void FUsdStageBase<PtrType>::MuteAndUnmuteLayers( const TArray<FString>& MuteLayers, const TArray<FString>& UnmuteLayers )
+	{
+#if USE_USD_SDK
+		if ( PtrType& Ptr = Impl->GetInner() )
+		{
+			FScopedUsdAllocs UsdAllocs;
+
+			std::vector<std::string> UsdMuteLayers;
+			UsdMuteLayers.reserve( MuteLayers.Num() );
+			for ( const FString& MuteLayer : MuteLayers )
+			{
+				UsdMuteLayers.push_back( TCHAR_TO_ANSI( *MuteLayer ) );
+			}
+
+			std::vector<std::string> UsdUnmuteLayers;
+			UsdUnmuteLayers.reserve( UnmuteLayers.Num() );
+			for ( const FString& UnmuteLayer : UnmuteLayers )
+			{
+				UsdUnmuteLayers.push_back( TCHAR_TO_ANSI( *UnmuteLayer ) );
+			}
+
+			Ptr->MuteAndUnmuteLayers( UsdMuteLayers, UsdUnmuteLayers );
+		}
+#endif // #if USE_USD_SDK
+	}
+
+	template<typename PtrType>
+	bool FUsdStageBase<PtrType>::IsLayerMuted( const FString& LayerIdentifier ) const
+	{
+#if USE_USD_SDK
+		if ( PtrType& Ptr = Impl->GetInner() )
+		{
+			FScopedUsdAllocs UsdAllocs;
+
+			return Ptr->IsLayerMuted( TCHAR_TO_ANSI( *LayerIdentifier ) );
+		}
+#endif // #if USE_USD_SDK
+
+		return false;
+	}
+
 	template<typename PtrType>
 	FSdfLayer FUsdStageBase<PtrType>::GetEditTarget() const
 	{
diff --git a/Engine/Plugins/Importers/USDImporter/Source/UnrealUSDWrapper/Public/UsdWrappers/UsdStage.h b/Engine/Plugins/Importers/USDImporter/Source/UnrealUSDWrapper/Public/UsdWrappers/UsdStage.h
--- a/Engine/Plugins/Importers/USDImporter/Source/UnrealUSDWrapper/Public/UsdWrappers/UsdStage.h
+++ b/Engine/Plugins/Importers/USDImporter/Source/UnrealUSDWrapper/Public/UsdWrappers/UsdStage.h
@@ -89,6 +89,9 @@ namespace UE
 		TArray<FSdfLayer> GetLayerStack( bool bIncludeSessionLayers = true ) const;
 		TArray<FSdfLayer> GetUsedLayers( bool bIncludeClipLayers = true ) const;
 
+		void MuteAndUnmuteLayers( const TArray<FString>& MuteLayers, const TArray<FString>& UnmuteLayers );
+		bool IsLayerMuted( const FString& LayerIdentifier ) const;
+
 		bool IsEditTargetValid() const;
 		void SetEditTarget( const FSdfLayer& Layer );
 		FSdfLayer GetEditTarget() const;
