USD: Recursively traverse for texture assets.

Replace deprecated usages of GetConnectedSource.

diff --git a/Engine/Plugins/Importers/USDImporter/Source/USDUtilities/Private/USDShadeConversion.cpp b/Engine/Plugins/Importers/USDImporter/Source/USDUtilities/Private/USDShadeConversion.cpp
--- a/Engine/Plugins/Importers/USDImporter/Source/USDUtilities/Private/USDShadeConversion.cpp
+++ b/Engine/Plugins/Importers/USDImporter/Source/USDUtilities/Private/USDShadeConversion.cpp
@@ -381,7 +381,7 @@ namespace UE
 					pxr::UsdShadeConnectableAPI Souce;
 					pxr::TfToken SourceName;
 					pxr::UsdShadeAttributeType SourceType;
-					if ( Input.GetConnectedSource( &Souce, &SourceName, &SourceType ) )
+					if ( pxr::UsdShadeConnectableAPI::GetConnectedSource( Input.GetAttr(), &Souce, &SourceName, &SourceType ) )
 					{
 						for ( const pxr::UsdShadeInput& DeeperInput : Souce.GetInputs() )
 						{
@@ -427,7 +427,7 @@ namespace UE
 				pxr::UsdShadeAttributeType AttributeType;
 
 				// Connected to a PrimvarReader
-				if ( StInput.GetConnectedSource( &PrimvarReader, &PrimvarReaderId, &AttributeType ) )
+				if ( pxr::UsdShadeConnectableAPI::GetConnectedSource( StInput.GetAttr(), &PrimvarReader, &PrimvarReaderId, &AttributeType ) )
 				{
 					if ( pxr::UsdShadeInput VarnameInput = PrimvarReader.GetInput( UnrealIdentifiers::Varname ) )
 					{
@@ -482,25 +482,51 @@ namespace UE
 					|| ShadeInput.GetBaseName() == UnrealIdentifiers::Normal
 				);
 
-				pxr::UsdShadeConnectableAPI Source;
-				pxr::TfToken SourceName;
-				pxr::UsdShadeAttributeType AttributeType;
+				pxr::UsdShadeConnectableAPI UsdUVTextureSource;
+				pxr::TfToken UsdUVTextureSourceName;
+				pxr::UsdShadeAttributeType UsdUVTextureAttributeType;
 
 				// Clear it, as we'll signal that it has a valid texture bound by setting it with an FTextureParameterValue
 				OutValue.Set<float>(0.0f);
-				if ( ShadeInput.GetConnectedSource( &Source, &SourceName, &AttributeType ) )
+
+				if ( pxr::UsdShadeConnectableAPI::GetConnectedSource( ShadeInput.GetAttr(), &UsdUVTextureSource, &UsdUVTextureSourceName, &UsdUVTextureAttributeType ) )
 				{
 					pxr::UsdShadeInput FileInput;
 
 					// UsdUVTexture: Get its file input
-					if ( AttributeType == pxr::UsdShadeAttributeType::Output )
+					if ( UsdUVTextureAttributeType == pxr::UsdShadeAttributeType::Output )
 					{
-						FileInput = Source.GetInput( UnrealIdentifiers::File );
+						FileInput = UsdUVTextureSource.GetInput( UnrealIdentifiers::File );
 					}
 					// Check if we are being directly passed an asset
 					else
 					{
-						FileInput = Source.GetInput( SourceName );
+						FileInput = UsdUVTextureSource.GetInput( UsdUVTextureSourceName );
+					}
+
+					// Recursively traverse "inputs:file" connections until we stop finding other connected prims
+					pxr::UsdShadeConnectableAPI TextureFileSource;
+					pxr::TfToken TextureFileSourceName;
+					pxr::UsdShadeAttributeType TextureFileAttributeType;
+					while ( FileInput )
+					{
+						if ( pxr::UsdShadeConnectableAPI::GetConnectedSource( FileInput.GetAttr(), &TextureFileSource, &TextureFileSourceName, &TextureFileAttributeType ) )
+						{
+							// Another connection, get its file input
+							if ( TextureFileAttributeType == pxr::UsdShadeAttributeType::Output )
+							{
+								FileInput = TextureFileSource.GetInput( UnrealIdentifiers::File );
+							}
+							// Check if we are being directly passed an asset
+							else
+							{
+								FileInput = TextureFileSource.GetInput( TextureFileSourceName );
+							}
+						}
+						else
+						{
+							break;
+						}
 					}
 
 					if ( FileInput && FileInput.GetTypeName() == pxr::SdfValueTypeNames->Asset ) // Check that FileInput is of type Asset
@@ -576,7 +602,7 @@ namespace UE
 
 							if ( Texture )
 							{
-								FString PrimvarName = GetPrimvarUsedAsST(Source);
+								FString PrimvarName = GetPrimvarUsedAsST( UsdUVTextureSource );
 								int32 UVIndex = 0;
 								if ( !PrimvarName.IsEmpty() && PrimvarToUVIndex )
 								{
@@ -598,9 +624,9 @@ namespace UE
 								OutTextureValue.Texture = Texture;
 								OutTextureValue.UVIndex = UVIndex;
 
-								if ( AttributeType == pxr::UsdShadeAttributeType::Output )
+								if ( UsdUVTextureAttributeType == pxr::UsdShadeAttributeType::Output )
 								{
-									const FName OutputName( UsdToUnreal::ConvertToken( SourceName ) );
+									const FName OutputName( UsdToUnreal::ConvertToken( UsdUVTextureSourceName ) );
 
 									if ( OutputName == TEXT("rgb") )
 									{
@@ -661,7 +687,7 @@ namespace UE
 				pxr::UsdShadeConnectableAPI Source;
 				pxr::TfToken SourceName;
 				pxr::UsdShadeAttributeType AttributeType;
-				if ( Input.GetConnectedSource( &Source, &SourceName, &AttributeType ) )
+				if ( pxr::UsdShadeConnectableAPI::GetConnectedSource( Input.GetAttr(), &Source, &SourceName, &AttributeType ) )
 				{
 					if ( !GetTextureParameterValue( Input, TEXTUREGROUP_WorldSpecular, OutValue, Material, TexturesCache, PrimvarToUVIndex, bForceVirtualTextures ) )
 					{
@@ -780,7 +806,7 @@ namespace UE
 				pxr::UsdShadeConnectableAPI Source;
 				pxr::TfToken SourceName;
 				pxr::UsdShadeAttributeType AttributeType;
-				if ( Input.GetConnectedSource( &Source, &SourceName, &AttributeType ) )
+				if ( pxr::UsdShadeConnectableAPI::GetConnectedSource( Input.GetAttr(), &Source, &SourceName, &AttributeType ) )
 				{
 					if ( !GetTextureParameterValue( Input, bIsNormalMap ? TEXTUREGROUP_WorldNormalMap : TEXTUREGROUP_World, OutValue, Material, TexturesCache, PrimvarToUVIndex, bForceVirtualTextures ) )
 					{
@@ -850,7 +876,7 @@ namespace UE
 				pxr::UsdShadeConnectableAPI Source;
 				pxr::TfToken SourceName;
 				pxr::UsdShadeAttributeType AttributeType;
-				if ( Input.GetConnectedSource( &Source, &SourceName, &AttributeType ) )
+				if ( pxr::UsdShadeConnectableAPI::GetConnectedSource( Input.GetAttr(), &Source, &SourceName, &AttributeType ) )
 				{
 					// Recurse because the attribute may just be pointing at some other attribute that has the data
 					// (e.g. when shader input is just "hoisted" and connected to the parent material input)
@@ -1709,11 +1735,17 @@ namespace UE
 				pxr::UsdShadeConnectableAPI Source;
 				pxr::TfToken SourceName;
 				pxr::UsdShadeAttributeType AttributeType;
-				if ( ShadeInput.GetConnectedSource( &Source, &SourceName, &AttributeType ) )
+				if ( pxr::UsdShadeConnectableAPI::GetConnectedSource( ShadeInput.GetAttr(), &Source, &SourceName, &AttributeType ) )
 				{
 					FString SourceOutputName = UsdToUnreal::ConvertToken( SourceName );
 					InOutHashState.UpdateWithString( *SourceOutputName, SourceOutputName.Len() );
 
+					// Skip in case our input is connected to an output on the same prim, or else we'll recurse forever
+					if ( Source.GetPrim() == ShadeInput.GetPrim() )
+					{
+						return;
+					}
+
 					for ( const pxr::UsdShadeInput& ChildInput : Source.GetInputs() )
 					{
 						HashShadeInput( ChildInput, InOutHashState );
@@ -2086,7 +2118,7 @@ bool UsdToUnreal::ConvertShadeInputsToParameters( const pxr::UsdShadeMaterial& U
 			pxr::UsdShadeConnectableAPI Source;
 			pxr::TfToken SourceName;
 			pxr::UsdShadeAttributeType SourceType;
-			ShadeInput.GetConnectedSource( &Source, &SourceName, &SourceType );
+			pxr::UsdShadeConnectableAPI::GetConnectedSource( ShadeInput.GetAttr(), &Source, &SourceName, &SourceType );
 
 			ConnectInput = Source.GetInput(SourceName);
 		}
@@ -2487,7 +2519,7 @@ bool UsdUtils::IsMaterialUsingUDIMs( const pxr::UsdShadeMaterial& UsdShadeMateri
 		pxr::TfToken SourceName;
 		pxr::UsdShadeAttributeType AttributeType;
 
-		if ( ShadeInput.GetConnectedSource(&Source, &SourceName, &AttributeType))
+		if ( pxr::UsdShadeConnectableAPI::GetConnectedSource( ShadeInput.GetAttr(), &Source, &SourceName, &AttributeType ) )
 		{
 			pxr::UsdShadeInput FileInput;
 
