USD: Fix LevelSequence to USD export missing tracks when actors and root components were separately bound.

diff --git a/Engine/Plugins/Importers/USDImporter/Source/USDExporter/Private/LevelSequenceExporterUSD.cpp b/Engine/Plugins/Importers/USDImporter/Source/USDExporter/Private/LevelSequenceExporterUSD.cpp
--- a/Engine/Plugins/Importers/USDImporter/Source/USDExporter/Private/LevelSequenceExporterUSD.cpp
+++ b/Engine/Plugins/Importers/USDImporter/Source/USDExporter/Private/LevelSequenceExporterUSD.cpp
@@ -580,11 +580,11 @@ namespace UE
 					return;
 				}
 
-				// Find all components (if actors, the root components of those actors) bound to the movie scene.
-				// We use components here because when exporting actors and components to USD we basically just ignore
-				// actors altogether and export the component attachment hierarchy instead.
-				// Index from component to FGuid because we may have multiple spawned actors/components for a given spawnable Guid
-				TMap<USceneComponent*, FGuid> BoundComponents;
+				// Find all objects bound to the movie scene. We'll track UObjects here because we may have tracks
+				// bound to root components but also separate tracks bound directly to the actors, and we want to
+				// capture both.
+				// Index from UObject to FGuid because we may have multiple spawned objects for a given spawnable Guid
+				TMap<UObject*, FGuid> BoundObjects;
 
 				const TArray<FMovieSceneSequenceID>* InstancesOfThisSequence = SequenceInstances.Find( &MovieSceneSequence );
 				if ( !InstancesOfThisSequence )
@@ -615,20 +615,10 @@ namespace UE
 							Context.SpawnRegister
 						);
 
-						USceneComponent* BoundComponent = Cast<USceneComponent>( BoundObject );
-						if ( !BoundComponent )
+						if ( BoundObject && ( BoundObject->IsA<USceneComponent>() || BoundObject->IsA<AActor>() ) )
 						{
-							if ( const AActor* Actor = Cast<AActor>( BoundObject ) )
-							{
-								BoundComponent = Actor->GetRootComponent();
-							}
-						}
-						if ( !BoundComponent )
-						{
-							continue;
+							BoundObjects.Add( BoundObject, Guid );
 						}
-
-						BoundComponents.Add( BoundComponent, Guid );
 					}
 
 					// Spawnables
@@ -647,29 +637,34 @@ namespace UE
 							continue;
 						}
 
-						USceneComponent* BoundComponent = Cast<USceneComponent>( BoundObject );
-						if ( !BoundComponent )
-						{
-							if ( const AActor* Actor = Cast<AActor>( BoundObject ) )
-							{
-								BoundComponent = Actor->GetRootComponent();
-							}
-						}
-						if ( !BoundComponent )
+						if ( BoundObject && ( BoundObject->IsA<USceneComponent>() || BoundObject->IsA<AActor>() ) )
 						{
-							continue;
+							BoundObjects.Add( BoundObject, Guid );
 						}
-
-						BoundComponents.Add( BoundComponent, Guid );
 					}
 				}
 
 				// Generate bakers
-				for ( const TPair<USceneComponent*, FGuid>& Pair : BoundComponents )
+				for ( const TPair<UObject*, FGuid>& Pair : BoundObjects )
 				{
-					USceneComponent* BoundComponent = Pair.Key;
+					UObject* BoundObject = Pair.Key;
 					const FGuid& Guid = Pair.Value;
 
+					// We always use components here because when exporting actors and components to USD we basically
+					// just ignore actors altogether and export the component attachment hierarchy instead
+					USceneComponent* BoundComponent = Cast<USceneComponent>( BoundObject );
+					if ( !BoundComponent )
+					{
+						if ( const AActor* Actor = Cast<AActor>( BoundObject ) )
+						{
+							BoundComponent = Actor->GetRootComponent();
+						}
+					}
+					if ( !BoundComponent )
+					{
+						continue;
+					}
+
 					FString PrimPath = UsdUtils::GetPrimPathForObject(
 						BoundComponent,
 						TEXT(""),
