USD: Fix crash when trying to parse nested SkelRoots in some scenarios.

diff --git a/Engine/Plugins/Importers/USDImporter/Source/USDSchemas/Private/USDSkelRootTranslator.cpp b/Engine/Plugins/Importers/USDImporter/Source/USDSchemas/Private/USDSkelRootTranslator.cpp
--- a/Engine/Plugins/Importers/USDImporter/Source/USDSchemas/Private/USDSkelRootTranslator.cpp
+++ b/Engine/Plugins/Importers/USDImporter/Source/USDSchemas/Private/USDSkelRootTranslator.cpp
@@ -679,6 +679,10 @@ namespace UsdSkelRootTranslatorImpl
 					IterateLODsLambda( pxr::UsdGeomMesh{ Stage->GetPrimAtPath( MeshPrimPath ) }, 0 );
 				}
 			}
+
+			// We only ever handle the first skeleton binding for now, so lets ignore material
+			// overrides on the others
+			break;
 		}
 
 		// Refresh reference to SkelRootPrim because variant switching potentially invalidated it
