USD: Respect disabled sections when exporting skeletal meshes to USD.

diff --git a/Engine/Plugins/Importers/USDImporter/Source/USDUtilities/Private/USDSkeletalDataConversion.cpp b/Engine/Plugins/Importers/USDImporter/Source/USDUtilities/Private/USDSkeletalDataConversion.cpp
--- a/Engine/Plugins/Importers/USDImporter/Source/USDUtilities/Private/USDSkeletalDataConversion.cpp
+++ b/Engine/Plugins/Importers/USDImporter/Source/USDUtilities/Private/USDSkeletalDataConversion.cpp
@@ -685,16 +685,45 @@ namespace UnrealToUsdImpl
 
 		const FUsdStageInfo StageInfo( Stage );
 
+		// FSkelMeshSection can be "disabled", at which point they don't show up in the engine. We'll skip those
+		// sections, and will use this array to help remap from a source vertex index to the vertex's corresponding
+		// index in a "packed" array of vertices that we'll push to USD
+		TArray<int32> SourceToPackedVertexIndex;
+		int32 PackedVertexIndex = 0;
+
 		// Vertices
 		{
-			const int32 VertexCount = LODModel.NumVertices;
-			if ( VertexCount == 0 )
+			if ( LODModel.NumVertices == 0 )
 			{
 				return;
 			}
 
+			// We manually collect vertices here instead of calling LODModel.GetVertices as we need to skip
+			// vertices from disabled sections, which that function won't do
 			TArray<FSoftSkinVertex> Vertices;
-			LODModel.GetVertices( Vertices );
+			Vertices.Reserve( LODModel.NumVertices );
+			SourceToPackedVertexIndex.SetNum( LODModel.NumVertices );
+			for ( int32 SectionIndex = 0; SectionIndex < LODModel.Sections.Num(); ++SectionIndex )
+			{
+				const FSkelMeshSection& Section = LODModel.Sections[ SectionIndex ];
+				if ( Section.bDisabled )
+				{
+					// Mark that these indices are skipped
+					for ( int32 Index = 0; Index < Section.NumVertices; ++Index )
+					{
+						SourceToPackedVertexIndex[ Section.BaseVertexIndex + Index ] = INDEX_NONE;
+					}
+					continue;
+				}
+
+				for ( int32 Index = 0; Index < Section.NumVertices; ++Index )
+				{
+					SourceToPackedVertexIndex[ Section.BaseVertexIndex + Index ] = PackedVertexIndex++;
+				}
+
+				Vertices.Append( Section.SoftVertices );
+			}
+			const int32 VertexCount = Vertices.Num();
 
 			// Points
 			{
@@ -799,9 +828,13 @@ namespace UnrealToUsdImpl
 					pxr::VtArray< float > JointWeights;
 					JointWeights.reserve( VertexCount * NumInfluencesPerVertex );
 
-					const int32 SectionCount = LODModel.Sections.Num();
 					for ( const FSkelMeshSection& Section : LODModel.Sections )
 					{
+						if ( Section.bDisabled )
+						{
+							continue;
+						}
+
 						for ( const FSoftSkinVertex& Vertex : Section.SoftVertices )
 						{
 							for ( int32 InfluenceIndex = 0; InfluenceIndex < NumInfluencesPerVertex; ++InfluenceIndex )
@@ -828,6 +861,11 @@ namespace UnrealToUsdImpl
 			{
 				for ( const FSkelMeshSection& Section : LODModel.Sections )
 				{
+					if ( Section.bDisabled )
+					{
+						continue;
+					}
+
 					TotalNumTriangles += Section.NumTriangles;
 				}
 
@@ -857,15 +895,24 @@ namespace UnrealToUsdImpl
 
 					for ( const FSkelMeshSection& Section : LODModel.Sections )
 					{
+						if ( Section.bDisabled )
+						{
+							continue;
+						}
+
 						int32 TriangleCount = Section.NumTriangles;
 						for ( uint32 TriangleIndex = 0; TriangleIndex < Section.NumTriangles; ++TriangleIndex )
 						{
 							for ( uint32 PointIndex = 0; PointIndex < 3; PointIndex++ )
 							{
-								int32 VertexPositionIndex = LODModel.IndexBuffer[ Section.BaseIndex + ( ( TriangleIndex * 3 ) + PointIndex ) ];
-								check( VertexPositionIndex >= 0 );
+								const int32 SourceVertexIndex =
+									LODModel.IndexBuffer[ Section.BaseIndex + ( ( TriangleIndex * 3 ) + PointIndex ) ];
 
-								FaceVertexIndices.push_back(VertexPositionIndex);
+								int32 PackedVertexPositionIndex = SourceToPackedVertexIndex[ SourceVertexIndex ];
+
+								check( PackedVertexPositionIndex >= 0 && PackedVertexPositionIndex != INDEX_NONE );
+
+								FaceVertexIndices.push_back( PackedVertexPositionIndex );
 							}
 						}
 					}
@@ -910,6 +957,10 @@ namespace UnrealToUsdImpl
 				for ( int32 SectionIndex = 0; SectionIndex < LODModel.Sections.Num(); ++SectionIndex )
 				{
 					const FSkelMeshSection& Section = LODModel.Sections[ SectionIndex ];
+					if ( Section.bDisabled )
+					{
+						continue;
+					}
 
 					// Note that we will continue on even if we have no material assignment, so as to satisfy the "partition" family condition
 					std::string SectionMaterial = UnrealMaterialsForLOD[ SectionIndex ];
@@ -934,18 +985,34 @@ namespace UnrealToUsdImpl
 					ElementTypeAttr.Set( pxr::UsdGeomTokens->face, TimeCode );
 
 					// Indices attribute
-					const uint32 TriangleCount = Section.NumTriangles;
-					const uint32 FirstTriangleIndex = Section.BaseIndex / 3; // BaseIndex is the first *vertex* instance index
-					pxr::VtArray<int> IndicesAttrValue;
-					for ( uint32 TriangleIndex = FirstTriangleIndex; TriangleIndex - FirstTriangleIndex < TriangleCount; ++TriangleIndex )
 					{
-						// Note that we add VertexInstances in sequence to the usda file for the faceVertexInstances attribute, which
-						// also constitutes our triangle order
-						IndicesAttrValue.push_back( static_cast< int >( TriangleIndex ));
-					}
+						const uint32 TriangleCount = Section.NumTriangles;
+						uint32 FirstTriangleIndex = Section.BaseIndex / 3; // BaseIndex is the first *vertex* instance index
+						pxr::VtArray<int> IndicesAttrValue;
+
+						// We may have some disabled sections (that wouldn't have emitted triangles). If so, we need to
+						// adjust our FirstTriangleIndex.
+						// This could be optimized in case vertex instances show up in LODModel.IndexBuffer according
+						// to the section order, but so far we haven't found such guarantee, so just check them all.
+						for ( int32 OtherSectionIndex = 0; OtherSectionIndex < LODModel.Sections.Num(); ++OtherSectionIndex )
+						{
+							const FSkelMeshSection& OtherSection = LODModel.Sections[ OtherSectionIndex ];
+							if ( OtherSection.bDisabled && OtherSection.BaseIndex < Section.BaseIndex )
+							{
+								FirstTriangleIndex -= OtherSection.NumTriangles;
+							}
+						}
 
-					pxr::UsdAttribute IndicesAttr = GeomSubsetSchema.CreateIndicesAttr();
-					IndicesAttr.Set( IndicesAttrValue, TimeCode );
+						for ( uint32 TriangleIndex = FirstTriangleIndex; TriangleIndex - FirstTriangleIndex < TriangleCount; ++TriangleIndex )
+						{
+							// Note that we add VertexInstances in sequence to the usda file for the faceVertexInstances attribute, which
+							// also constitutes our triangle order
+							IndicesAttrValue.push_back( static_cast< int >( TriangleIndex ));
+						}
+
+						pxr::UsdAttribute IndicesAttr = GeomSubsetSchema.CreateIndicesAttr();
+						IndicesAttr.Set( IndicesAttrValue, TimeCode );
+					}
 
 					// Family name attribute
 					pxr::UsdAttribute FamilyNameAttr = GeomSubsetSchema.CreateFamilyNameAttr();
