USD: Generate PhysicsAssets when generating SkeletalMeshes (disableable via cvar USD.GeneratePhysicsAssets)

diff --git a/Engine/Plugins/Importers/USDImporter/Source/USDSchemas/Private/USDSkelRootTranslator.cpp b/Engine/Plugins/Importers/USDImporter/Source/USDSchemas/Private/USDSkelRootTranslator.cpp
--- a/Engine/Plugins/Importers/USDImporter/Source/USDSchemas/Private/USDSkelRootTranslator.cpp
+++ b/Engine/Plugins/Importers/USDImporter/Source/USDSchemas/Private/USDSkelRootTranslator.cpp
@@ -33,7 +33,8 @@
 #include "GroomComponent.h"
 #include "Materials/Material.h"
 #include "Materials/MaterialInstanceConstant.h"
-#include "UObject/StrongObjectPtr.h"
+#include "ObjectTools.h"
+#include "PhysicsAssetUtils.h"
 #include "Rendering/SkeletalMeshLODImporterData.h"
 #include "Serialization/BufferArchive.h"
 
@@ -63,6 +64,12 @@
 
 #define LOCTEXT_NAMESPACE "UsdSkelRoot"
 
+static bool bGeneratePhysicsAssets = true;
+static FAutoConsoleVariableRef CVarGeneratePhysicsAssets(
+	TEXT( "USD.GeneratePhysicsAssets" ),
+	bGeneratePhysicsAssets,
+	TEXT( "Whether to automatically generate and assign PhysicsAssets to generated SkeletalMeshes." ) );
+
 namespace UsdSkelRootTranslatorImpl
 {
 #if WITH_EDITOR
@@ -763,6 +770,38 @@ namespace UsdSkelRootTranslatorImpl
 		return false;
 	}
 
+	UPhysicsAsset* GenerateAndAssignPhysicsAsset( USkeletalMesh* SkeletalMesh, EObjectFlags Flags )
+	{
+		FName AssetName = MakeUniqueObjectName(
+			GetTransientPackage(),
+			UPhysicsAsset::StaticClass(),
+			*FPaths::GetBaseFilename( TEXT( "PHYS_" ) + SkeletalMesh->GetName() )
+		);
+
+		UPhysicsAsset* Result = NewObject< UPhysicsAsset >(
+			GetTransientPackage(),
+			AssetName,
+			Flags | EObjectFlags::RF_Public
+		);
+
+		FPhysAssetCreateParams NewBodyData;
+		FText CreationErrorMessage;
+		bool bSuccess = FPhysicsAssetUtils::CreateFromSkeletalMesh( Result, SkeletalMesh, NewBodyData, CreationErrorMessage );
+		if ( !bSuccess )
+		{
+			UE_LOG( LogUsd, Error, TEXT( "Couldn't create PhysicsAsset for SkeletalMesh '%s': %s" ),
+				*SkeletalMesh->GetName(),
+				*CreationErrorMessage.ToString()
+			);
+
+			ensure( ObjectTools::DeleteSingleObject( Result ) );
+
+			Result = nullptr;
+		}
+
+		return Result;
+	}
+
 	class FSkelRootCreateAssetsTaskChain : public FUsdSchemaTranslatorTaskChain
 	{
 	public:
@@ -859,8 +898,9 @@ namespace UsdSkelRootTranslatorImpl
 			{
 				FString SkelRootPath = PrimPath.GetString();
 				FSHAHash SkeletalMeshHash = UsdSkelRootTranslatorImpl::ComputeSHAHash( LODIndexToSkeletalMeshImportData, SkeletonBones );
+				const FString SkeletalMeshHashString = SkeletalMeshHash.ToString();
 
-				USkeletalMesh* SkeletalMesh = Cast< USkeletalMesh >( Context->AssetCache->GetCachedAsset( SkeletalMeshHash.ToString() ) );
+				USkeletalMesh* SkeletalMesh = Cast< USkeletalMesh >( Context->AssetCache->GetCachedAsset( SkeletalMeshHashString ) );
 
 				bool bIsNew = false;
 				if ( !SkeletalMesh )
@@ -893,8 +933,32 @@ namespace UsdSkelRootTranslatorImpl
 						ImportData->PrimPath = SkelRootPath;
 						SkeletalMesh->SetAssetImportData(ImportData);
 
-						Context->AssetCache->CacheAsset( SkeletalMeshHash.ToString(), SkeletalMesh );
-						Context->AssetCache->CacheAsset( SkeletalMeshHash.ToString() + TEXT( "_Skeleton" ), SkeletalMesh->GetSkeleton() );
+						Context->AssetCache->CacheAsset( SkeletalMeshHashString, SkeletalMesh );
+						Context->AssetCache->CacheAsset( SkeletalMeshHashString + TEXT( "_Skeleton" ), SkeletalMesh->GetSkeleton() );
+					}
+
+					if ( bGeneratePhysicsAssets )
+					{
+						if ( !SkeletalMesh->GetPhysicsAsset() )
+						{
+							UPhysicsAsset* PhysicsAsset = UsdSkelRootTranslatorImpl::GenerateAndAssignPhysicsAsset(
+								SkeletalMesh,
+								Context->ObjectFlags
+							);
+
+							if ( PhysicsAsset )
+							{
+								Context->AssetCache->CacheAsset(
+									SkeletalMeshHashString + TEXT( "_PhysicsAsset" ),
+									PhysicsAsset
+								);
+							}
+						}
+					}
+					else
+					{
+						// Actively clear this so that if we toggle the cvar and reload we'll clear our physics assets
+						SkeletalMesh->SetPhysicsAsset( nullptr );
 					}
 
 					Context->AssetCache->LinkAssetToPrim( SkelRootPath, SkeletalMesh );
diff --git a/Engine/Plugins/Importers/USDImporter/Source/USDSchemas/USDSchemas.Build.cs b/Engine/Plugins/Importers/USDImporter/Source/USDSchemas/USDSchemas.Build.cs
--- a/Engine/Plugins/Importers/USDImporter/Source/USDSchemas/USDSchemas.Build.cs
+++ b/Engine/Plugins/Importers/USDImporter/Source/USDSchemas/USDSchemas.Build.cs
@@ -55,6 +55,7 @@ namespace UnrealBuildTool.Rules
 						"MaterialEditor",
 						"MDLImporter",
 						"MeshUtilities",
+						"PhysicsUtilities", // For generating UPhysicsAssets for SkeletalMeshes
 						"PropertyEditor",
 						"UnrealEd",
 					}
diff --git a/Engine/Plugins/Importers/USDImporter/Source/USDStageImporter/Private/USDStageImporter.cpp b/Engine/Plugins/Importers/USDImporter/Source/USDStageImporter/Private/USDStageImporter.cpp
--- a/Engine/Plugins/Importers/USDImporter/Source/USDStageImporter/Private/USDStageImporter.cpp
+++ b/Engine/Plugins/Importers/USDImporter/Source/USDStageImporter/Private/USDStageImporter.cpp
@@ -458,7 +458,6 @@ namespace UsdStageImporterImpl
 		else if (USkeletalMesh* SkMesh = Cast<USkeletalMesh>(Asset))
 		{
 			AssetPrefix = TEXT("SK_");
-
 		}
 		else if (USkeleton* Skeleton = Cast<USkeleton>(Asset))
 		{
@@ -466,6 +465,12 @@ namespace UsdStageImporterImpl
 
 			// Note that UsdUtils::GetAssetImportData fetches the AssetImportData of the PreviewMesh
 		}
+		else if ( UPhysicsAsset* PhysicsAsset = Cast<UPhysicsAsset>( Asset ) )
+		{
+			AssetPrefix = TEXT( "PHYS_" );
+
+			// Note that UsdUtils::GetAssetImportData fetches the AssetImportData of the PreviewMesh
+		}
 		else if ( UAnimSequence* AnimSequence = Cast<UAnimSequence>( Asset ) )
 		{
 			AssetPrefix = TEXT( "Anim_" );
@@ -711,7 +716,12 @@ namespace UsdStageImporterImpl
 				{
 					// For skeletal stuff, the primpaths point to the SkelRoot, so it is useful to place the assets in there,
 					// as we'll always have at least the skeletal mesh and the skeleton
-					if ( Asset->IsA( USkeletalMesh::StaticClass() ) || Asset->IsA( USkeleton::StaticClass() ) || Asset->IsA( UAnimSequence::StaticClass() ) || Asset->IsA( UAnimBlueprint::StaticClass() ) )
+					if ( Asset->IsA( USkeletalMesh::StaticClass() ) ||
+						 Asset->IsA( USkeleton::StaticClass() ) ||
+						 Asset->IsA( UPhysicsAsset::StaticClass() ) ||
+						 Asset->IsA( UAnimSequence::StaticClass() ) ||
+						 Asset->IsA( UAnimBlueprint::StaticClass() )
+					)
 					{
 						AssetTypeFolder = ImportData->PrimPath;
 					}
@@ -733,13 +743,19 @@ namespace UsdStageImporterImpl
 				}
 				else if (Asset->IsA(UGeometryCache::StaticClass()))
 				{
-					AssetTypeFolder = "GeometryCaches";
+					AssetTypeFolder = TEXT("GeometryCaches");
 				}
 				else if (Asset->IsA(UTexture::StaticClass()))
 				{
 					AssetTypeFolder = TEXT("Textures");
 				}
-				else if (Asset->IsA(USkeletalMesh::StaticClass()) || Asset->IsA(USkeleton::StaticClass()) || Asset->IsA(UAnimSequence::StaticClass()) || Asset->IsA(UAnimBlueprint::StaticClass()))
+				else if (
+					Asset->IsA( USkeletalMesh::StaticClass() ) ||
+					Asset->IsA( USkeleton::StaticClass() ) ||
+					Asset->IsA( UPhysicsAsset::StaticClass() ) ||
+					Asset->IsA( UAnimSequence::StaticClass() ) ||
+					Asset->IsA( UAnimBlueprint::StaticClass() )
+				)
 				{
 					AssetTypeFolder = TEXT("SkeletalMeshes");
 				}
@@ -1020,7 +1036,8 @@ namespace UsdStageImporterImpl
 
 			// Assign even if ExistingSkeletalMesh has nullptr skeleton because we must be able to cleanup the
 			// abandoned Skeleton in the transient package
-			NewSkeletalMesh->SetSkeleton(ExistingSkeletalMesh->GetSkeleton());
+			NewSkeletalMesh->SetSkeleton( ExistingSkeletalMesh->GetSkeleton() );
+			NewSkeletalMesh->SetPhysicsAsset( ExistingSkeletalMesh->GetPhysicsAsset() );
 		}
 
 		UAnimSequence* ExistingAnimSequence = Cast<UAnimSequence>( ExistingAsset );
@@ -1195,6 +1212,14 @@ namespace UsdStageImporterImpl
 						}
 					}
 
+					if ( UPhysicsAsset* PhysicsAsset = SkeletalMesh->GetPhysicsAsset() )
+					{
+						if ( ImportContext.AssetCache->IsAssetOwnedByCache( PhysicsAsset ) )
+						{
+							OutAssetsAndDependencies.Add( PhysicsAsset );
+						}
+					}
+
 					AssetQueue.Reserve( AssetQueue.Num() + SkeletalMesh->GetMaterials().Num() );
 					for ( const FSkeletalMaterial& SkeletalMaterial : SkeletalMesh->GetMaterials() )
 					{
@@ -1248,6 +1273,10 @@ namespace UsdStageImporterImpl
 				{
 					// Do nothing. Skeletons have no additional dependencies
 				}
+				else if ( UPhysicsAsset* PhysicsAsset = Cast<UPhysicsAsset>( Asset ) )
+				{
+					// Do nothing. PhysicsAssets have no additional dependencies
+				}
 				else if ( UAnimBlueprint* AnimBP = Cast<UAnimBlueprint>( Asset ) )
 				{
 					// Do nothing. AnimBlueprints have no additional dependencies
@@ -1434,7 +1463,13 @@ namespace UsdStageImporterImpl
 			}
 
 			if (
-				( !ImportContext.ImportOptions->bImportGeometry && ( Asset->IsA<UStaticMesh>() || Asset->IsA<USkeletalMesh>() || Asset->IsA<USkeleton>() || Asset->IsA<UGeometryCache>() ) ) ||
+				( !ImportContext.ImportOptions->bImportGeometry && (
+					Asset->IsA<UStaticMesh>() ||
+					Asset->IsA<USkeletalMesh>() ||
+					Asset->IsA<USkeleton>() ||
+					Asset->IsA<UPhysicsAsset>() ||
+					Asset->IsA<UGeometryCache>()
+				) ) ||
 				( !bImportSkeletalAnimations && ( Asset->IsA<UAnimSequence>() ) ) ||
 				( !ImportContext.ImportOptions->bImportLevelSequences && ( Asset->IsA<ULevelSequence>() ) ) ||
 				( !ImportContext.ImportOptions->bImportMaterials && ( Asset->IsA<UMaterialInterface>() || Asset->IsA<UTexture>() ) )
diff --git a/Engine/Plugins/Importers/USDImporter/Source/USDUtilities/Private/USDConversionUtils.cpp b/Engine/Plugins/Importers/USDImporter/Source/USDUtilities/Private/USDConversionUtils.cpp
--- a/Engine/Plugins/Importers/USDImporter/Source/USDUtilities/Private/USDConversionUtils.cpp
+++ b/Engine/Plugins/Importers/USDImporter/Source/USDUtilities/Private/USDConversionUtils.cpp
@@ -1219,6 +1219,13 @@ UUsdAssetImportData* UsdUtils::GetAssetImportData( UObject* Asset )
 			ImportData = Cast<UUsdAssetImportData>( SkMesh->GetAssetImportData() );
 		}
 	}
+	else if ( UPhysicsAsset* PhysicsAsset = Cast<UPhysicsAsset>( Asset ) )
+	{
+		if ( USkeletalMesh* SkMesh = PhysicsAsset->GetPreviewMesh() )
+		{
+			ImportData = Cast<UUsdAssetImportData>( SkMesh->GetAssetImportData() );
+		}
+	}
 	else if ( UAnimBlueprint* AnimBP = Cast<UAnimBlueprint>( Asset ) )
 	{
 		// We will always have a skeleton, but not necessarily we will have a preview mesh directly
