USD: Show warning when muting unsaved layers.

diff --git a/Engine/Plugins/Importers/USDImporter/Source/USDClasses/Public/USDProjectSettings.h b/Engine/Plugins/Importers/USDImporter/Source/USDClasses/Public/USDProjectSettings.h
--- a/Engine/Plugins/Importers/USDImporter/Source/USDClasses/Public/USDProjectSettings.h
+++ b/Engine/Plugins/Importers/USDImporter/Source/USDClasses/Public/USDProjectSettings.h
@@ -33,6 +33,9 @@ public:
 	UPROPERTY( config, EditAnywhere, Category = "USD|Dialogs" )
 	bool bShowConfirmationWhenClearingLayers = true;
 
+	UPROPERTY( config, EditAnywhere, Category = "USD|Dialogs" )
+	bool bShowConfirmationWhenMutingDirtyLayers = true;
+
 	// Whether to show the warning dialog when authoring opinions that could have no effect on the composed stage
 	UPROPERTY( config, EditAnywhere, Category = "USD|Dialogs" )
 	bool bShowOverriddenOpinionsWarning = true;
diff --git a/Engine/Plugins/Importers/USDImporter/Source/USDStageEditor/Private/SUSDLayersTreeView.cpp b/Engine/Plugins/Importers/USDImporter/Source/USDStageEditor/Private/SUSDLayersTreeView.cpp
--- a/Engine/Plugins/Importers/USDImporter/Source/USDStageEditor/Private/SUSDLayersTreeView.cpp
+++ b/Engine/Plugins/Importers/USDImporter/Source/USDStageEditor/Private/SUSDLayersTreeView.cpp
@@ -218,8 +218,92 @@ protected:
 
 	void ToggleMuteLayer( FUsdLayerViewModelRef LayerItem )
 	{
-		if ( CanMuteLayer( LayerItem ) )
+		if ( !CanMuteLayer( LayerItem ) )
 		{
+			return;
+		}
+
+		// If the layer is not dirty we can just mute it without worry and early out
+		if ( !LayerItem->IsLayerDirty() )
+		{
+			LayerItem->ToggleMuteLayer();
+			return;
+		}
+
+		// Show a warning if the layer is dirty, as muting it will discard changes
+		const UUsdProjectSettings* Settings = GetDefault<UUsdProjectSettings>();
+		if ( Settings && Settings->bShowConfirmationWhenMutingDirtyLayers )
+		{
+			static TWeakPtr<SNotificationItem> Notification;
+
+			FNotificationInfo Toast( LOCTEXT( "ConfirmMutingLayer", "Muting dirty layer" ) );
+			Toast.SubText = FText::Format(
+				LOCTEXT( "ConfirmMutingLayer_Subtext", "Layer '{0}' has unsaved changes that will be lost if muted.\n\nDo you wish to proceed muting the layer?" ),
+				LayerItem->GetDisplayName()
+			);
+			Toast.Image = FCoreStyle::Get().GetBrush( TEXT( "MessageLog.Warning" ) );
+			Toast.bUseLargeFont = false;
+			Toast.bFireAndForget = false;
+			Toast.FadeOutDuration = 0.0f;
+			Toast.ExpireDuration = 0.0f;
+			Toast.bUseThrobber = false;
+			Toast.bUseSuccessFailIcons = false;
+
+			Toast.ButtonDetails.Emplace(
+				LOCTEXT( "ConfirmMutingOkAll", "Always proceed" ),
+				FText::GetEmpty(),
+				FSimpleDelegate::CreateLambda( [LayerItem]() {
+				if ( TSharedPtr<SNotificationItem> PinnedNotification = Notification.Pin() )
+				{
+					PinnedNotification->SetCompletionState( SNotificationItem::CS_Success );
+					PinnedNotification->ExpireAndFadeout();
+
+					UUsdProjectSettings* Settings = GetMutableDefault<UUsdProjectSettings>();
+					Settings->bShowConfirmationWhenMutingDirtyLayers = false;
+					Settings->SaveConfig();
+
+					LayerItem->ToggleMuteLayer();
+				}
+			} ) );
+
+			Toast.ButtonDetails.Emplace(
+				LOCTEXT( "ConfirmMutingOk", "Proceed" ),
+				FText::GetEmpty(),
+				FSimpleDelegate::CreateLambda( [LayerItem]() {
+				if ( TSharedPtr<SNotificationItem> PinnedNotification = Notification.Pin() )
+				{
+					PinnedNotification->SetCompletionState( SNotificationItem::CS_Success );
+					PinnedNotification->ExpireAndFadeout();
+
+					LayerItem->ToggleMuteLayer();
+				}
+			} ) );
+
+			Toast.ButtonDetails.Emplace(
+				LOCTEXT( "ConfirmMutingCancel", "Cancel" ),
+				FText::GetEmpty(),
+				FSimpleDelegate::CreateLambda( []() {
+				if ( TSharedPtr<SNotificationItem> PinnedNotification = Notification.Pin() )
+				{
+					PinnedNotification->SetCompletionState( SNotificationItem::CS_Fail );
+					PinnedNotification->ExpireAndFadeout();
+				}
+			} ) );
+
+			// Only show one at a time
+			if ( !Notification.IsValid() )
+			{
+				Notification = FSlateNotificationManager::Get().AddNotification( Toast );
+			}
+
+			if ( TSharedPtr<SNotificationItem> PinnedNotification = Notification.Pin() )
+			{
+				PinnedNotification->SetCompletionState( SNotificationItem::CS_Pending );
+			}
+		}
+		else
+		{
+			// Don't show prompt, always just mute
 			LayerItem->ToggleMuteLayer();
 		}
 	}
@@ -1230,7 +1314,7 @@ void SUsdLayersTreeView::OnClearSelectedLayers()
 
 		FNotificationInfo Toast( LOCTEXT( "ConfirmClearingLayer", "Clearing cannot be undone" ) );
 		Toast.SubText = FText::Format(
-			LOCTEXT( "ConfirmClearingLayer_Subtext", "Clearing USD layers cannot be undone.\nDo you wish to proceed clearing {0}|plural(one=layer,other=layers) '{1}' ?" ),
+			LOCTEXT( "ConfirmClearingLayer_Subtext", "Clearing USD layers cannot be undone.\n\nDo you wish to proceed clearing {0}|plural(one=layer,other=layers) '{1}' ?" ),
 			SelectedLayers.Num(),
 			FText::FromString(LayerNames)
 		);
