USD: Fix missing material assignments when exporting skeletal meshes with unused entries on their LODMaterialMap.

diff --git a/Engine/Plugins/Importers/USDImporter/Source/USDUtilities/Private/USDSkeletalDataConversion.cpp b/Engine/Plugins/Importers/USDImporter/Source/USDUtilities/Private/USDSkeletalDataConversion.cpp
--- a/Engine/Plugins/Importers/USDImporter/Source/USDUtilities/Private/USDSkeletalDataConversion.cpp
+++ b/Engine/Plugins/Importers/USDImporter/Source/USDUtilities/Private/USDSkeletalDataConversion.cpp
@@ -929,7 +929,17 @@ namespace UnrealToUsdImpl
 			pxr::VtArray< std::string > UnrealMaterialsForLOD;
 			for ( const FSkelMeshSection& Section : LODModel.Sections )
 			{
-				int32 SkeletalMaterialIndex = LODMaterialMap.IsValidIndex( Section.MaterialIndex ) ? LODMaterialMap[ Section.MaterialIndex ] : Section.MaterialIndex;
+				int32 SkeletalMaterialIndex = INDEX_NONE;
+				if ( LODMaterialMap.IsValidIndex( Section.MaterialIndex ) )
+				{
+					SkeletalMaterialIndex = LODMaterialMap[ Section.MaterialIndex ];
+				}
+				// Note that the LODMaterialMap can contain INDEX_NONE to signify no remapping
+				if ( SkeletalMaterialIndex == INDEX_NONE )
+				{
+					SkeletalMaterialIndex = Section.MaterialIndex;
+				}
+
 				if ( SkeletalMaterialIndex >= 0 && SkeletalMaterialIndex < MaterialAssignments.size() )
 				{
 					UnrealMaterialsForLOD.push_back( MaterialAssignments[ SkeletalMaterialIndex ] );
