USD: Fix AnimBPs not being imported when importing scenes with live link integration;

Fix asset name prefixes to follow recommendations.

diff --git a/Engine/Plugins/Importers/USDImporter/Source/USDSchemas/Private/USDSkelRootTranslator.cpp b/Engine/Plugins/Importers/USDImporter/Source/USDSchemas/Private/USDSkelRootTranslator.cpp
--- a/Engine/Plugins/Importers/USDImporter/Source/USDSchemas/Private/USDSkelRootTranslator.cpp
+++ b/Engine/Plugins/Importers/USDImporter/Source/USDSchemas/Private/USDSkelRootTranslator.cpp
@@ -351,6 +351,7 @@ namespace UsdSkelRootTranslatorImpl
 		TArray<FSkeletalMeshImportData>& OutLODIndexToSkeletalMeshImportData,
 		TArray<UsdUtils::FUsdPrimMaterialAssignmentInfo>& OutLODIndexToMaterialInfo,
 		TArray<SkeletalMeshImportData::FBone>& OutSkeletonBones,
+		FName& OutSkeletonName,
 		UsdUtils::FBlendShapeMap* OutBlendShapes,
 		TSet<FString>& InOutUsedMorphTargetNames,
 		const TMap< FString, TMap< FString, int32 > >& InMaterialToPrimvarsUVSetNames,
@@ -404,6 +405,7 @@ namespace UsdSkelRootTranslatorImpl
 				return false;
 			}
 			OutSkeletonBones = MoveTemp(DummyImportData.RefBonesBinary);
+			OutSkeletonName = *UsdToUnreal::ConvertString( SkelQuery.GetSkeleton().GetPrim().GetName() );
 		}
 
 		TMap<int32, FSkeletalMeshImportData> LODIndexToSkeletalMeshImportDataMap;
@@ -802,6 +804,319 @@ namespace UsdSkelRootTranslatorImpl
 		return Result;
 	}
 
+	UAnimBlueprint* CreateAnimBlueprint(
+		const FUsdSchemaTranslationContext& Context,
+		const pxr::UsdPrim& Prim,
+		bool bDelayRecompilation = false,  // Whether we should recompile within this function or not
+		bool* bOutNeedsRecompile = nullptr // Whether this function wants the returned AnimBP to be recompiled
+	)
+	{
+		if ( !Prim )
+		{
+			return nullptr;
+		}
+
+		FString PrimPath = UsdToUnreal::ConvertPath( Prim.GetPath() );
+		FString PrimName = UsdToUnreal::ConvertString( Prim.GetName() );
+
+		USkeletalMesh* SkeletalMesh = Cast< USkeletalMesh >( Context.AssetCache->GetAssetForPrim( PrimPath ) );
+		if ( !SkeletalMesh )
+		{
+			return nullptr;
+		}
+
+		USkeleton* Skeleton = SkeletalMesh->GetSkeleton();
+		if ( !Skeleton )
+		{
+			return nullptr;
+		}
+
+		// Fetch relevant attributes from prim, since we know it has the schema
+		FString AnimBPPath;
+		{
+			FScopedUsdAllocs Allocs;
+
+			if ( pxr::UsdAttribute Attr = Prim.GetAttribute( UnrealIdentifiers::UnrealAnimBlueprintPath ) )
+			{
+				std::string PathString;
+				if ( Attr.Get( &PathString ) )
+				{
+					AnimBPPath = UsdToUnreal::ConvertString( PathString );
+				}
+			}
+		}
+		if ( AnimBPPath.IsEmpty() )
+		{
+			return nullptr;
+		}
+
+		bool bNeedRecompile = false;
+
+		UAnimBlueprint* AnimBP = Cast<UAnimBlueprint>( FSoftObjectPath( AnimBPPath ).TryLoad() );
+		if ( !AnimBP )
+		{
+			return nullptr;
+		}
+
+		// Create transient AnimBP based on our template, so that we can assign it a proper skeleton
+		const static FString DefaultAnimBPPath = TEXT( "/USDImporter/Blueprint/DefaultLiveLinkAnimBP.DefaultLiveLinkAnimBP" );
+		if ( DefaultAnimBPPath == AnimBPPath )
+		{
+			const FString CacheKey = UsdToUnreal::ConvertPath( Prim.GetPrimPath() ) + TEXT( "_DefaultAnimBlueprint" );
+
+			// Check if we can find an AnimBP for this prim in the asset cache (useful when doing Action->Import)
+			bool bReusedAnimBP = false;
+			if ( UAnimBlueprint* CachedAnimBP = Cast< UAnimBlueprint >( Context.AssetCache->GetCachedAsset( CacheKey ) ) )
+			{
+				if ( CachedAnimBP->TargetSkeleton == Skeleton )
+				{
+					AnimBP = CachedAnimBP;
+					bReusedAnimBP = true;
+				}
+			}
+
+			// We have to generate a new transient AnimBP
+			if ( !bReusedAnimBP )
+			{
+				FName UniqueName = MakeUniqueObjectName(
+					GetTransientPackage(),
+					UAnimBlueprint::StaticClass(),
+					*( PrimName + TEXT( "_DefaultAnimBlueprint" ) )
+				);
+
+				// Duplicate and never reuse these so that they can be assigned independent subject names if desired.
+				// Its not as if scenes will have thousands of these anyway.
+				AnimBP = DuplicateObject( AnimBP, GetTransientPackage(), UniqueName );
+
+				// Patch up the flags here or else the rest of the engine code (and our plugin code) will get confused
+				// as to what a non-transient asset in the transient package even means. It can lead to some crashes
+				// when saving after import if we don't do this
+				AnimBP->ClearFlags( AnimBP->GetFlags() );
+				AnimBP->SetFlags( Skeleton->GetFlags() );
+
+				AnimBP->TargetSkeleton = Skeleton;
+				AnimBP->bIsTemplate = false;
+
+				bNeedRecompile = true;
+
+				Context.AssetCache->CacheAsset( CacheKey, AnimBP );
+			}
+		}
+		// Path is pointing to an existing, persistent AnimBP
+		else
+		{
+			// Force skeletons to be compatible if they aren't (we need both ways!)
+			if ( !Skeleton->IsCompatible( AnimBP->TargetSkeleton ) )
+			{
+				AnimBP->TargetSkeleton->AddCompatibleSkeleton( Skeleton );
+				Skeleton->AddCompatibleSkeleton( AnimBP->TargetSkeleton );
+
+				if ( AnimBP->TargetSkeleton->GetReferenceSkeleton().GetRefBoneInfo() != Skeleton->GetReferenceSkeleton().GetRefBoneInfo() )
+				{
+					UE_LOG( LogUsd, Warning, TEXT( "Forcing AnimBlueprint '%s's Skeleton '%s' to be compatible with the Skeleton generated for prim '%s', but they may be different!" ),
+						*AnimBP->GetPathName(),
+						*AnimBP->TargetSkeleton->GetPathName(),
+						*PrimName
+					);
+				}
+			}
+		}
+
+		if ( bNeedRecompile && !bDelayRecompilation )
+		{
+			FCompilerResultsLog Results;
+			FBPCompileRequest Request( AnimBP, EBlueprintCompileOptions::None, &Results );
+			FBlueprintCompilationManager::CompileSynchronously( Request );
+			bNeedRecompile = false;
+		}
+
+		if ( bOutNeedsRecompile )
+		{
+			*bOutNeedsRecompile = bNeedRecompile;
+		}
+
+		return AnimBP;
+	}
+
+	void UpdateLiveLinkProperties( const FUsdSchemaTranslationContext& Context, USkeletalMeshComponent* Component, const pxr::UsdPrim& Prim )
+	{
+		if ( !Component || !Component->GetSkeletalMeshAsset() || !Prim )
+		{
+			return;
+		}
+
+		FString PrimName = UsdToUnreal::ConvertString( Prim.GetName() );
+
+		USkeleton* Skeleton = Component->GetSkeletalMeshAsset()->GetSkeleton();
+		if ( !Skeleton )
+		{
+			return;
+		}
+
+		UAnimBlueprint* ExistingAnimBP = nullptr;
+		if ( Component->AnimClass )
+		{
+			ExistingAnimBP = Cast<UAnimBlueprint>( Component->AnimClass->ClassGeneratedBy );
+		}
+
+		// Fetch relevant attributes from prim, since we know it has the schema
+		FString SubjectName;
+		FString AnimBPPath;
+		{
+			FScopedUsdAllocs Allocs;
+
+			if ( pxr::UsdAttribute Attr = Prim.GetAttribute( UnrealIdentifiers::UnrealLiveLinkSubjectName ) )
+			{
+				std::string SubjectNameString;
+				if ( Attr.Get( &SubjectNameString ) )
+				{
+					SubjectName = UsdToUnreal::ConvertString( SubjectNameString );
+				}
+			}
+
+			if ( pxr::UsdAttribute Attr = Prim.GetAttribute( UnrealIdentifiers::UnrealAnimBlueprintPath ) )
+			{
+				std::string PathString;
+				if ( Attr.Get( &PathString ) )
+				{
+					AnimBPPath = UsdToUnreal::ConvertString( PathString );
+				}
+			}
+		}
+
+		UAnimBlueprint* AnimBP = Cast<UAnimBlueprint>( FSoftObjectPath( AnimBPPath ).TryLoad() );
+
+		// Check if we need to change the AnimBP
+		bool bNeedRecompile = false;
+		if ( ExistingAnimBP != AnimBP )
+		{
+			const bool bDelayRecompilation = true;
+			AnimBP = UsdSkelRootTranslatorImpl::CreateAnimBlueprint(
+				Context,
+				Prim,
+				bDelayRecompilation,
+				&bNeedRecompile
+			);
+		}
+
+		// Apply subject name to live link pose AnimBlueprint node
+		// Reference: UAnimationBlueprintLibrary::AddNodeAssetOverride
+		if ( AnimBP )
+		{
+			TArray<UBlueprint*> BlueprintHierarchy;
+			AnimBP->GetBlueprintHierarchyFromClass( AnimBP->GetAnimBlueprintGeneratedClass(), BlueprintHierarchy );
+
+			TArray<UAnimGraphNode_LiveLinkPose*> LiveLinkNodes;
+
+			for ( int32 BlueprintIndex = 0; BlueprintIndex < BlueprintHierarchy.Num(); ++BlueprintIndex )
+			{
+				UBlueprint* CurrentBlueprint = BlueprintHierarchy[ BlueprintIndex ];
+
+				TArray<UEdGraph*> Graphs;
+				CurrentBlueprint->GetAllGraphs( Graphs );
+
+				for ( UEdGraph* Graph : Graphs )
+				{
+					for ( UEdGraphNode* Node : Graph->Nodes )
+					{
+						if ( UAnimGraphNode_LiveLinkPose* AnimNode = Cast<UAnimGraphNode_LiveLinkPose>( Node ) )
+						{
+							LiveLinkNodes.Add( AnimNode );
+						}
+					}
+				}
+			}
+
+			if ( LiveLinkNodes.Num() > 1 && !ExistingAnimBP )
+			{
+				UE_LOG( LogUsd, Warning, TEXT( "Found more than one LiveLinkPose blueprint node on AnimBlueprint '%s's graphs."
+					"Note that all of those nodes will have their LiveLink SubjectName's updated to '%s', as described for prim '%s'!" ),
+					*AnimBP->GetPathName(),
+					*SubjectName,
+					*UsdToUnreal::ConvertPath( Prim.GetPrimPath() )
+				);
+			}
+
+			for ( UAnimGraphNode_LiveLinkPose* Node : LiveLinkNodes )
+			{
+				const UEdGraphSchema* Schema = Node->GetSchema();
+				if ( !Schema )
+				{
+					continue;
+				}
+
+				UEdGraphPin* SubjectNamePin = nullptr;
+				for ( UEdGraphPin* Pin : Node->Pins )
+				{
+					if ( Pin->GetName() == GET_MEMBER_NAME_STRING_CHECKED( FAnimNode_LiveLinkPose, LiveLinkSubjectName ) )
+					{
+						SubjectNamePin = Pin;
+						break;
+					}
+				}
+
+				// The subject name pin is already connected to something...
+				if ( SubjectNamePin->LinkedTo.Num() != 0 )
+				{
+					if ( !ExistingAnimBP )
+					{
+						UE_LOG( LogUsd, Warning, TEXT( "Failed to update a LiveLinkPose node's 'Subject Name' to '%s' on AnimBlueprint '%s', "
+							"because the pin is already connected to some other node. Disconnect it if you want it to be updated automatically." ),
+							*SubjectName,
+							*AnimBP->GetPathName()
+						);
+					}
+
+					continue;
+				}
+
+				if ( SubjectNamePin )
+				{
+					// The pin type is FLiveLinkSubjectName, so we must create an instance of it and serialize it using
+					// UScriptStruct::ExportText to generate a proper default value string
+					FLiveLinkSubjectName Dummy;
+					Dummy.Name = *SubjectName;
+
+					FString ValueString;
+					const void* Defaults = nullptr;
+					const UObject* OwnerObject = nullptr;
+					const int32 PortFlags = EPropertyPortFlags::PPF_None;
+					const UObject* ExportRootScope = nullptr;
+					FLiveLinkSubjectName::StaticStruct()->ExportText( ValueString, &Dummy, nullptr, nullptr, EPropertyPortFlags::PPF_None, nullptr );
+
+					if ( !Schema->DoesDefaultValueMatch( *SubjectNamePin, ValueString ) )
+					{
+						SubjectNamePin->Modify();
+						Schema->TrySetDefaultValue( *SubjectNamePin, ValueString );
+
+						FBlueprintEditorUtils::MarkBlueprintAsModified( AnimBP );
+						bNeedRecompile = true;
+					}
+				}
+			}
+		}
+
+		if ( bNeedRecompile )
+		{
+			FCompilerResultsLog Results;
+			FBPCompileRequest Request( AnimBP, EBlueprintCompileOptions::None, &Results );
+			FBlueprintCompilationManager::CompileSynchronously( Request );
+
+			// We need to force the component to update its anim after we regenerate the blueprint class
+			Component->ClearAnimScriptInstance();
+			Component->InitAnim( true );
+		}
+
+		if ( AnimBP != ExistingAnimBP )
+		{
+			// This can internally change AnimationMode, but lets revert it to what it was so that we can control it from
+			// that single place in ::UpdateComponents
+			EAnimationMode::Type OldMode = Component->GetAnimationMode();
+			Component->SetAnimInstanceClass( AnimBP->GeneratedClass );
+			Component->SetAnimationMode( OldMode );
+		}
+	}
+
 	class FSkelRootCreateAssetsTaskChain : public FUsdSchemaTranslatorTaskChain
 	{
 	public:
@@ -816,6 +1131,7 @@ namespace UsdSkelRootTranslatorImpl
 		TArray<FSkeletalMeshImportData> LODIndexToSkeletalMeshImportData;
 		TArray<UsdUtils::FUsdPrimMaterialAssignmentInfo> LODIndexToMaterialInfo;
 		TArray<SkeletalMeshImportData::FBone> SkeletonBones;
+		FName SkeletonName;
 		UsdUtils::FBlendShapeMap NewBlendShapes;
 
 		// Note that we want this to be case insensitive so that our UMorphTarget FNames are unique not only due to case differences
@@ -879,6 +1195,7 @@ namespace UsdSkelRootTranslatorImpl
 					LODIndexToSkeletalMeshImportData,
 					LODIndexToMaterialInfo,
 					SkeletonBones,
+					SkeletonName,
 					OutBlendShapes,
 					UsedMorphTargetNames,
 					*MaterialToPrimvarToUVIndex,
@@ -906,7 +1223,14 @@ namespace UsdSkelRootTranslatorImpl
 				if ( !SkeletalMesh )
 				{
 					bIsNew = true;
-					SkeletalMesh = UsdToUnreal::GetSkeletalMeshFromImportData( LODIndexToSkeletalMeshImportData, SkeletonBones, NewBlendShapes, Context->ObjectFlags, *FPaths::GetBaseFilename( SkelRootPath ) );
+					SkeletalMesh = UsdToUnreal::GetSkeletalMeshFromImportData(
+						LODIndexToSkeletalMeshImportData,
+						SkeletonBones,
+						NewBlendShapes,
+						Context->ObjectFlags,
+						*FPaths::GetBaseFilename( SkelRootPath ),
+						SkeletonName
+					);
 				}
 
 				if ( SkeletalMesh )
@@ -993,6 +1317,49 @@ namespace UsdSkelRootTranslatorImpl
 				return true;
 			} );
 
+		// Create AnimBP asset if we need to
+		Then( ESchemaTranslationLaunchPolicy::Sync,
+			[ this ]()
+			{
+				UE::FUsdPrim Prim = GetPrim();
+
+				// We need to *also* create the AnimBP within CreateAssets so that we will still generate it even if
+				// we never call Create/UpdateComponents (e.g. importing without importing actors).
+				const bool bPrimHasLiveLinkSchema = UsdUtils::PrimHasSchema( Prim, UnrealIdentifiers::LiveLinkAPI );
+				if ( bPrimHasLiveLinkSchema )
+				{
+					TSharedPtr< FUsdSchemaTranslationContext > ContextPtr = Context;
+
+					// When importing, we can't delay creating the AnimBP to the next frame as that will be after the
+					// import. We don't need to worry about deadlocks though, because we will never *import* as a
+					// response to a USD event: It's always an intentional function call, where the Python GIL is
+					// properly handled
+					if ( ContextPtr && ContextPtr->bIsImporting )
+					{
+						UsdSkelRootTranslatorImpl::CreateAnimBlueprint( *ContextPtr.Get(), Prim );
+					}
+					else
+					{
+						// HACK. c.f. the large comment on the analogous timer inside
+						// FUsdSkelRootTranslator::UpdateComponents
+						FTSTicker::GetCoreTicker().AddTicker(
+							FTickerDelegate::CreateLambda( [ContextPtr, Prim]( float Time )
+							{
+								if ( ContextPtr )
+								{
+									UsdSkelRootTranslatorImpl::CreateAnimBlueprint( *ContextPtr.Get(), Prim );
+								}
+
+								// Returning false means this is a one-off, and won't repeat
+								return false;
+							})
+						);
+					}
+				}
+
+				return true;
+			});
+
 		// Create UAnimSequences (requires a completed USkeleton. Main thread as some steps of the animation compression require it)
 		Then( ESchemaTranslationLaunchPolicy::Sync,
 			[ this ]()
@@ -1108,258 +1475,6 @@ namespace UsdSkelRootTranslatorImpl
 			});
 	}
 
-	void UpdateLiveLinkProperties( const FUsdSchemaTranslationContext& Context, USkeletalMeshComponent* Component, const pxr::UsdPrim& Prim )
-	{
-		if ( !Component || !Component->GetSkeletalMeshAsset() || !Prim )
-		{
-			return;
-		}
-
-		FString PrimName = UsdToUnreal::ConvertString( Prim.GetName() );
-
-		USkeleton* Skeleton = Component->GetSkeletalMeshAsset()->GetSkeleton();
-		if ( !Skeleton )
-		{
-			return;
-		}
-
-		// Fetch relevant attributes from prim, since we know it has the schema
-		FString SubjectName;
-		FString AnimBPPath;
-		{
-			FScopedUsdAllocs Allocs;
-
-			if ( pxr::UsdAttribute Attr = Prim.GetAttribute( UnrealIdentifiers::UnrealLiveLinkSubjectName ) )
-			{
-				std::string SubjectNameString;
-				if ( Attr.Get( &SubjectNameString ) )
-				{
-					SubjectName = UsdToUnreal::ConvertString( SubjectNameString );
-				}
-			}
-
-			if ( pxr::UsdAttribute Attr = Prim.GetAttribute( UnrealIdentifiers::UnrealAnimBlueprintPath ) )
-			{
-				std::string PathString;
-				if ( Attr.Get( &PathString ) )
-				{
-					AnimBPPath = UsdToUnreal::ConvertString( PathString );
-				}
-			}
-		}
-
-		UAnimBlueprint* ExistingAnimBP = nullptr;
-		if ( Component->AnimClass )
-		{
-			ExistingAnimBP = Cast<UAnimBlueprint>( Component->AnimClass->ClassGeneratedBy );
-		}
-
-		bool bNeedRecompile = false;
-
-		// Fetch AnimBP to use
-		UAnimBlueprint* AnimBP = ExistingAnimBP;
-		if ( !AnimBPPath.IsEmpty() && ( !ExistingAnimBP || ExistingAnimBP->GetPathName() != AnimBPPath ) )
-		{
-			const static FString DefaultAnimBPPath = TEXT( "/USDImporter/Blueprint/DefaultLiveLinkAnimBP.DefaultLiveLinkAnimBP" );
-
-			AnimBP = Cast<UAnimBlueprint>( FSoftObjectPath( AnimBPPath ).TryLoad() );
-			if ( AnimBP )
-			{
-				// Create transient AnimBP based on our template, so that we can assign it a proper skeleton
-				if ( DefaultAnimBPPath == AnimBPPath )
-				{
-					// Check if we can use the AnimBP we currently have
-					bool bAlreadyHasTransientAnimBP = false;
-					if ( ExistingAnimBP )
-					{
-						if ( ExistingAnimBP->GetOutermost() == GetTransientPackage() && ExistingAnimBP->TargetSkeleton == Skeleton )
-						{
-							AnimBP = ExistingAnimBP;
-							bAlreadyHasTransientAnimBP = true;
-						}
-					}
-
-					const FString CacheKey = UsdToUnreal::ConvertPath( Prim.GetPrimPath() ) + TEXT( "_DefaultAnimBlueprint" );
-
-					// Check if we can find an AnimBP for this prim in the asset cache (useful when doing Action->Import)
-					if ( !bAlreadyHasTransientAnimBP )
-					{
-						if ( UAnimBlueprint* CachedAnimBP = Cast< UAnimBlueprint >( Context.AssetCache->GetCachedAsset( CacheKey ) ) )
-						{
-							if ( CachedAnimBP->TargetSkeleton == Skeleton )
-							{
-								AnimBP = CachedAnimBP;
-								bAlreadyHasTransientAnimBP = true;
-							}
-						}
-					}
-
-					// We have to generate a new transient AnimBP
-					if ( !bAlreadyHasTransientAnimBP )
-					{
-						FName UniqueName = MakeUniqueObjectName(
-							GetTransientPackage(),
-							UAnimBlueprint::StaticClass(),
-							*( PrimName + TEXT( "_DefaultAnimBlueprint" ) )
-						);
-
-						// Duplicate and never reuse these so that they can be assigned independent subject names if desired.
-						// Its not as if scenes will have thousands of these anyway.
-						AnimBP = DuplicateObject( AnimBP, GetTransientPackage(), UniqueName );
-
-						// Patch up the flags here or else the rest of the engine code (and our plugin code) will get confused
-						// as to what a non-transient asset in the transient package even means. It can lead to some crashes
-						// when saving after import if we don't do this
-						AnimBP->ClearFlags( AnimBP->GetFlags() );
-						AnimBP->SetFlags( Skeleton->GetFlags() );
-
-						AnimBP->TargetSkeleton = Skeleton;
-						AnimBP->bIsTemplate = false;
-
-						bNeedRecompile = true;
-
-						Context.AssetCache->CacheAsset( CacheKey, AnimBP );
-					}
-				}
-				// Path is pointing to an existing, persistent AnimBP
-				else
-				{
-					// Force skeletons to be compatible (we need both ways!)
-					AnimBP->TargetSkeleton->AddCompatibleSkeleton( Skeleton );
-					Skeleton->AddCompatibleSkeleton( AnimBP->TargetSkeleton );
-
-					// Check ExistingAnimBP so that we only emit this warning when we first set up the component (when we'll realistically not have an ExistingAnimBP yet),
-					// to try and prevent some warning spam
-					if ( AnimBP->TargetSkeleton->GetReferenceSkeleton().GetRefBoneInfo() != Skeleton->GetReferenceSkeleton().GetRefBoneInfo() && !ExistingAnimBP )
-					{
-						UE_LOG(LogUsd, Warning, TEXT("Forcing AnimBlueprint '%s's Skeleton '%s' to be compatible with the Skeleton generated for prim '%s', but they may be different!"),
-							*AnimBP->GetPathName(),
-							*AnimBP->TargetSkeleton->GetPathName(),
-							*PrimName
-						);
-					}
-				}
-			}
-		}
-
-		// Apply subject name to live link pose AnimBlueprint node
-		// Reference: UAnimationBlueprintLibrary::AddNodeAssetOverride
-		if ( AnimBP )
-		{
-			TArray<UBlueprint*> BlueprintHierarchy;
-			AnimBP->GetBlueprintHierarchyFromClass( AnimBP->GetAnimBlueprintGeneratedClass(), BlueprintHierarchy );
-
-			TArray<UAnimGraphNode_LiveLinkPose*> LiveLinkNodes;
-
-			for ( int32 BlueprintIndex = 0; BlueprintIndex < BlueprintHierarchy.Num(); ++BlueprintIndex )
-			{
-				UBlueprint* CurrentBlueprint = BlueprintHierarchy[ BlueprintIndex ];
-
-				TArray<UEdGraph*> Graphs;
-				CurrentBlueprint->GetAllGraphs( Graphs );
-
-				for ( UEdGraph* Graph : Graphs )
-				{
-					for ( UEdGraphNode* Node : Graph->Nodes )
-					{
-						if ( UAnimGraphNode_LiveLinkPose* AnimNode = Cast<UAnimGraphNode_LiveLinkPose>( Node ) )
-						{
-							LiveLinkNodes.Add( AnimNode );
-						}
-					}
-				}
-			}
-
-			if ( LiveLinkNodes.Num() > 1 && !ExistingAnimBP )
-			{
-				UE_LOG( LogUsd, Warning, TEXT( "Found more than one LiveLinkPose blueprint node on AnimBlueprint '%s's graphs."
-					"Note that all of those nodes will have their LiveLink SubjectName's updated to '%s', as described for prim '%s'!" ),
-					*AnimBP->GetPathName(),
-					*SubjectName,
-					*UsdToUnreal::ConvertPath( Prim.GetPrimPath() )
-				);
-			}
-
-			for ( UAnimGraphNode_LiveLinkPose* Node : LiveLinkNodes )
-			{
-				const UEdGraphSchema* Schema = Node->GetSchema();
-				if ( !Schema )
-				{
-					continue;
-				}
-
-				UEdGraphPin* SubjectNamePin = nullptr;
-				for ( UEdGraphPin* Pin : Node->Pins )
-				{
-					if ( Pin->GetName() == GET_MEMBER_NAME_STRING_CHECKED( FAnimNode_LiveLinkPose, LiveLinkSubjectName ) )
-					{
-						SubjectNamePin = Pin;
-						break;
-					}
-				}
-
-				// The subject name pin is already connected to something...
-				if ( SubjectNamePin->LinkedTo.Num() != 0 )
-				{
-					if ( !ExistingAnimBP )
-					{
-						UE_LOG( LogUsd, Warning, TEXT( "Failed to update a LiveLinkPose node's 'Subject Name' to '%s' on AnimBlueprint '%s', "
-							"because the pin is already connected to some other node. Disconnect it if you want it to be updated automatically." ),
-							*SubjectName,
-							*AnimBP->GetPathName()
-						);
-					}
-
-					continue;
-				}
-
-				if ( SubjectNamePin )
-				{
-					// The pin type is FLiveLinkSubjectName, so we must create an instance of it and serialize it using
-					// UScriptStruct::ExportText to generate a proper default value string
-					FLiveLinkSubjectName Dummy;
-					Dummy.Name = *SubjectName;
-
-					FString ValueString;
-					const void* Defaults = nullptr;
-					const UObject* OwnerObject = nullptr;
-					const int32 PortFlags = EPropertyPortFlags::PPF_None;
-					const UObject* ExportRootScope = nullptr;
-					FLiveLinkSubjectName::StaticStruct()->ExportText( ValueString, &Dummy, nullptr, nullptr, EPropertyPortFlags::PPF_None, nullptr );
-
-					if ( !Schema->DoesDefaultValueMatch( *SubjectNamePin, ValueString ) )
-					{
-						SubjectNamePin->Modify();
-						Schema->TrySetDefaultValue( *SubjectNamePin, ValueString );
-
-						FBlueprintEditorUtils::MarkBlueprintAsModified( AnimBP );
-						bNeedRecompile = true;
-					}
-				}
-			}
-
-			if ( bNeedRecompile )
-			{
-				FCompilerResultsLog Results;
-				FBPCompileRequest Request( AnimBP, EBlueprintCompileOptions::None, &Results );
-				FBlueprintCompilationManager::CompileSynchronously( Request );
-
-				// We need to force the component to update its anim after we regenerate the blueprint class
-				Component->ClearAnimScriptInstance();
-				Component->InitAnim(true);
-			}
-
-			if ( AnimBP != ExistingAnimBP )
-			{
-				// This can internally change AnimationMode, but lets revert it to what it was so that we can control it from
-				// that single place in ::UpdateComponents
-				EAnimationMode::Type OldMode = Component->GetAnimationMode();
-				Component->SetAnimInstanceClass( AnimBP->GeneratedClass );
-				Component->SetAnimationMode( OldMode );
-			}
-		}
-	}
-
 #endif // WITH_EDITOR
 }
 
@@ -1490,25 +1605,32 @@ void FUsdSkelRootTranslator::UpdateComponents( USceneComponent* SceneComponent )
 		TSharedPtr< FUsdSchemaTranslationContext > ContextPtr = Context;
 		TStrongObjectPtr< USkeletalMeshComponent > PinnedSkelMeshComponent{ SkeletalMeshComponent };
 
-		// HACK: This is a temporary work-around for a GIL deadlock. See UE-156489 and UE-156370 for more details, but
-		// the long-story short of it is that at this point we may have a callstack that originates from Python,
-		// triggers a USD stage notice and causes C++ code to run as our stage actor is listening to them. If we
-		// cause GC to run right now (which the DuplicateObject inside UpdateLiveLinkProperties will) we may cause
-		// a deadlock, as the game thread still holds the GIL and a background reference collector thread may want
-		// to acquire it too. What this does is run this part of the UpdateComponents on tick, that has a callstack
-		// that doesn't originate from Python, and so doesn't have the GIL locked
-		FTSTicker::GetCoreTicker().AddTicker(
-			FTickerDelegate::CreateLambda( [ContextPtr, PinnedSkelMeshComponent, Prim]( float Time )
-			{
-				if ( ContextPtr )
+		if ( ContextPtr && ContextPtr->bIsImporting )
+		{
+			UsdSkelRootTranslatorImpl::UpdateLiveLinkProperties( *ContextPtr.Get(), PinnedSkelMeshComponent.Get(), Prim );
+		}
+		else
+		{
+			// HACK: This is a temporary work-around for a GIL deadlock. See UE-156489 and UE-156370 for more details, but
+			// the long-story short of it is that at this point we may have a callstack that originates from Python,
+			// triggers a USD stage notice and causes C++ code to run as our stage actor is listening to them. If we
+			// cause GC to run right now (which the DuplicateObject inside UpdateLiveLinkProperties will) we may cause
+			// a deadlock, as the game thread still holds the GIL and a background reference collector thread may want
+			// to acquire it too. What this does is run this part of the UpdateComponents on tick, that has a callstack
+			// that doesn't originate from Python, and so doesn't have the GIL locked
+			FTSTicker::GetCoreTicker().AddTicker(
+				FTickerDelegate::CreateLambda( [ContextPtr, PinnedSkelMeshComponent, Prim]( float Time )
 				{
-					UsdSkelRootTranslatorImpl::UpdateLiveLinkProperties( *ContextPtr.Get(), PinnedSkelMeshComponent.Get(), Prim );
-				}
+					if ( ContextPtr )
+					{
+						UsdSkelRootTranslatorImpl::UpdateLiveLinkProperties( *ContextPtr.Get(), PinnedSkelMeshComponent.Get(), Prim );
+					}
 
-				// Returning false means this is a one-off, and won't repeat
-				return false;
-			})
-		);
+					// Returning false means this is a one-off, and won't repeat
+					return false;
+				})
+			);
+		}
 	}
 
 	// Update the animation state
diff --git a/Engine/Plugins/Importers/USDImporter/Source/USDSchemas/Public/USDSchemaTranslator.h b/Engine/Plugins/Importers/USDImporter/Source/USDSchemas/Public/USDSchemaTranslator.h
--- a/Engine/Plugins/Importers/USDImporter/Source/USDSchemas/Public/USDSchemaTranslator.h
+++ b/Engine/Plugins/Importers/USDImporter/Source/USDSchemas/Public/USDSchemaTranslator.h
@@ -128,6 +128,9 @@ struct USDSCHEMAS_API FUsdSchemaTranslationContext : public TSharedFromThis< FUs
 {
 	explicit FUsdSchemaTranslationContext( const UE::FUsdStage& InStage, UUsdAssetCache& InAssetCache );
 
+	/** True if we're a context created by the USDStageImporter to fully import to persistent assets and actors */
+	bool bIsImporting = false;
+
 	/** pxr::UsdStage we're translating from */
 	UE::FUsdStage Stage;
 
diff --git a/Engine/Plugins/Importers/USDImporter/Source/USDStageImporter/Private/USDStageImporter.cpp b/Engine/Plugins/Importers/USDImporter/Source/USDStageImporter/Private/USDStageImporter.cpp
--- a/Engine/Plugins/Importers/USDImporter/Source/USDStageImporter/Private/USDStageImporter.cpp
+++ b/Engine/Plugins/Importers/USDImporter/Source/USDStageImporter/Private/USDStageImporter.cpp
@@ -425,14 +425,15 @@ namespace UsdStageImporterImpl
 		}
 
 		FString AssetPrefix;
-		FString AssetSuffix;
-		FString AssetPath = Asset->GetName();
 
+		FString AssetPath = Asset->GetName();
 		if (UUsdAssetImportData* AssetImportData = UsdUtils::GetAssetImportData(Asset))
 		{
 			AssetPath = AssetImportData->PrimPath;
 		}
 
+		FString AssetName = FPaths::GetBaseFilename( AssetPath );
+
 		if (UStaticMesh* Mesh = Cast<UStaticMesh>(Asset))
 		{
 			AssetPrefix = TEXT("SM_");
@@ -450,6 +451,7 @@ namespace UsdStageImporterImpl
 						if ( PrimName.IsNumeric() )
 						{
 							AssetPath = FPaths::GetPath( AssetPath );
+							AssetName = FPaths::GetBaseFilename( AssetPath );
 						}
 					}
 				}
@@ -461,52 +463,86 @@ namespace UsdStageImporterImpl
 		}
 		else if (USkeleton* Skeleton = Cast<USkeleton>(Asset))
 		{
-			AssetSuffix = TEXT("_Skeleton");
-
-			// Note that UsdUtils::GetAssetImportData fetches the AssetImportData of the PreviewMesh
+			// Skeletons don't have asset import data, so we can't store the prim that originated them, but we do
+			// name the actual assets after the prims, so unlike for the other asset types we want to use the actual
+			// asset name here
+			AssetName = Asset->GetFName().GetPlainNameString();
+			AssetPrefix = TEXT("SKEL_");
 		}
 		else if ( UPhysicsAsset* PhysicsAsset = Cast<UPhysicsAsset>( Asset ) )
 		{
-			AssetPrefix = TEXT( "PHYS_" );
+			// See comments above on the case for USkeleton
+			FString TempName = Asset->GetFName().GetPlainNameString();
 
-			// Note that UsdUtils::GetAssetImportData fetches the AssetImportData of the PreviewMesh
+			// The asset is named after the SkelRoot prim. If we're importing back a scene that was originally exported,
+			// we should clean up these prefixes or else we may end up with something like "PHYS_SK_PrimName"
+			TempName.RemoveFromStart( TEXT( "PHYS_" ), ESearchCase::CaseSensitive );
+			TempName.RemoveFromStart( TEXT( "SK_" ), ESearchCase::CaseSensitive );
+			if ( !TempName.IsEmpty() )
+			{
+				AssetName = TempName;
+			}
+
+			AssetPrefix = TEXT( "PHYS_" );
 		}
 		else if ( UAnimSequence* AnimSequence = Cast<UAnimSequence>( Asset ) )
 		{
-			AssetPrefix = TEXT( "Anim_" );
+			AssetPrefix = TEXT( "AS_" );
 		}
 		else if (UMaterialInterface* Material = Cast<UMaterialInterface>(Asset))
 		{
-			AssetPrefix = TEXT("M_");
+			if ( Material->IsA<UMaterialInstance>() )
+			{
+				AssetPrefix = TEXT( "MI_" );
+			}
+			else
+			{
+				AssetPrefix = TEXT( "M_" );
+			}
 
 			if (UUsdAssetImportData* AssetImportData = Cast<UUsdAssetImportData>(Material->AssetImportData))
 			{
 				// The only materials with no prim path are our auto-generated displayColor materials
 				AssetPath = AssetImportData->PrimPath.IsEmpty()? TEXT("DisplayColor") : AssetImportData->PrimPath;
+				AssetName = FPaths::GetBaseFilename( AssetPath );
 			}
 		}
 		else if (UTexture* Texture = Cast<UTexture>(Asset))
 		{
 			AssetPrefix = TEXT("T_");
 		}
-
-		FString FinalName = FPaths::GetBaseFilename(AssetPath);
-		if ( !FinalName.StartsWith( AssetPrefix ) )
+		else if ( ULevelSequence* LevelSequence = Cast<ULevelSequence>( Asset ) )
 		{
-			FinalName = AssetPrefix + FinalName;
+			AssetPrefix = TEXT( "LS_" );
 		}
-		if ( !FinalName.EndsWith( AssetSuffix ) )
+		else if ( UAnimBlueprint* AnimBP = Cast<UAnimBlueprint>( Asset ) )
+		{
+			FString TempName = AssetName;
+
+			// The asset is named after the SkelRoot prim. If we're importing back a scene that was originally exported,
+			// we should clean up these prefixes or else we may end up with something like "ABP_SK_PrimName"
+			TempName.RemoveFromStart( TEXT( "ABP_" ), ESearchCase::CaseSensitive );
+			TempName.RemoveFromStart( TEXT( "SK_" ), ESearchCase::CaseSensitive );
+			if ( !TempName.IsEmpty() )
+			{
+				AssetName = TempName;
+			}
+
+			AssetPrefix = TEXT( "ABP_" );
+		}
+
+		if ( !AssetName.StartsWith( AssetPrefix ) )
 		{
-			FinalName = FinalName + AssetSuffix;
+			AssetName = AssetPrefix + AssetName;
 		}
 
 		// We don't care if our assets overwrite something in the final destination package (that conflict will be
 		// handled according to EReplaceAssetPolicy). But we do want these assets to have unique names amongst themselves
 		// or else they will overwrite each other when publishing
-		FinalName = UsdUtils::GetUniqueName( ObjectTools::SanitizeObjectName( FinalName ), UniqueAssetNames );
-		UniqueAssetNames.Add(FinalName);
+		AssetName = UsdUtils::GetUniqueName( ObjectTools::SanitizeObjectName( AssetName ), UniqueAssetNames );
+		UniqueAssetNames.Add(AssetName);
 
-		return FinalName;
+		return AssetName;
 	}
 
 	void UpdateAssetImportData( UObject* Asset, const FString& MainFilePath, UUsdStageImportOptions* ImportOptions )
@@ -1580,6 +1616,7 @@ void UUsdStageImporter::ImportFromFile(FUsdStageImportContext& ImportContext)
 	FGlobalComponentRecreateRenderStateContext RecreateRenderStateContext;
 
 	TSharedRef<FUsdSchemaTranslationContext> TranslationContext = MakeShared<FUsdSchemaTranslationContext>( ImportContext.Stage, *ImportContext.AssetCache );
+	TranslationContext->bIsImporting = true;
 	TranslationContext->Level = ImportContext.World->GetCurrentLevel();
 	TranslationContext->ObjectFlags = ImportContext.ImportObjectFlags;
 	TranslationContext->Time = static_cast< float >( UsdUtils::GetDefaultTimeCode() );
@@ -1680,6 +1717,7 @@ bool UUsdStageImporter::ReimportSingleAsset(FUsdStageImportContext& ImportContex
 	FGlobalComponentRecreateRenderStateContext RecreateRenderStateContext;
 
 	TSharedRef<FUsdSchemaTranslationContext> TranslationContext = MakeShared<FUsdSchemaTranslationContext>( ImportContext.Stage, *ImportContext.AssetCache );
+	TranslationContext->bIsImporting = true;
 	TranslationContext->Level = ImportContext.World->GetCurrentLevel();
 	TranslationContext->ObjectFlags = ImportContext.ImportObjectFlags;
 	TranslationContext->Time = static_cast< float >( UsdUtils::GetDefaultTimeCode() );
