USD: Fixed playback of USD GeometryCache due to threading issue when querying the frame data.

diff --git a/Engine/Plugins/Importers/USDImporter/Source/GeometryCacheUSD/Private/GeometryCacheUSDSceneProxy.cpp b/Engine/Plugins/Importers/USDImporter/Source/GeometryCacheUSD/Private/GeometryCacheUSDSceneProxy.cpp
--- a/Engine/Plugins/Importers/USDImporter/Source/GeometryCacheUSD/Private/GeometryCacheUSDSceneProxy.cpp
+++ b/Engine/Plugins/Importers/USDImporter/Source/GeometryCacheUSD/Private/GeometryCacheUSDSceneProxy.cpp
@@ -57,10 +57,11 @@ void FGeomCacheTrackUsdProxy::FindSampleIndexesFromTime(float Time, bool bLoopin
 	if (UGeometryCacheTrackUsd* UsdTrack = Cast<UGeometryCacheTrackUsd>(Track))
 	{
 		int32 ThisFrameIndex = UsdTrack->FindSampleIndexFromTime(Time, bLooping);
-		int32 LastFrameIndex = UsdTrack->GetEndFrameIndex();
 		OutFrameIndex = ThisFrameIndex;
 		OutNextFrameIndex = OutFrameIndex + 1;
-		InInterpolationFactor = Time - ThisFrameIndex;
+		// Clamp the Time (which is the FrameNumber) to the range of the USD GeometryCache track
+		// EndFrameIndex - 1 since last frame index is not included
+		InInterpolationFactor = FMath::Clamp(Time, float(UsdTrack->GetStartFrameIndex()), float(UsdTrack->GetEndFrameIndex() - 1)) - ThisFrameIndex;
 
 		// If playing backwards the logical order of previous and next is reversed
 		if (bIsPlayingBackwards)
diff --git a/Engine/Plugins/Importers/USDImporter/Source/GeometryCacheUSD/Private/GeometryCacheUSDStream.cpp b/Engine/Plugins/Importers/USDImporter/Source/GeometryCacheUSD/Private/GeometryCacheUSDStream.cpp
--- a/Engine/Plugins/Importers/USDImporter/Source/GeometryCacheUSD/Private/GeometryCacheUSDStream.cpp
+++ b/Engine/Plugins/Importers/USDImporter/Source/GeometryCacheUSD/Private/GeometryCacheUSDStream.cpp
@@ -20,7 +20,7 @@ static int32 kUsdReadConcurrency = 10;
 // differences, etc.) replace the version GUID below with a new one.
 // In case of merge conflicts with DDC versions, you MUST generate a new GUID
 // and set this new GUID as the version.
-#define USDSTREAM_DERIVED_DATA_VERSION TEXT("040845A286FE431A846D95D9B77EE756")
+#define USDSTREAM_DERIVED_DATA_VERSION TEXT("AB2B7CC003C54AEBBCC5ABDC1B0BFFD8")
 
 class FUsdStreamDDCUtils
 {
@@ -127,12 +127,14 @@ void FGeometryCacheUsdStream::GetMeshData(int32 FrameIndex, int32 ConcurrencyInd
 			}
 			else
 			{
-				ReadFunc(UsdTrack, FrameIndex, OutMeshData);
-
-				FMemoryWriter Ar(DerivedData, true);
-				Ar << OutMeshData;
-
-				GetDerivedDataCacheRef().Put(*DerivedDataKey, DerivedData, UsdPrimPath);
+				// Don't cache the data if it failed to read it
+				if (ReadFunc(UsdTrack, FrameIndex, OutMeshData))
+				{
+					FMemoryWriter Ar(DerivedData, true);
+					Ar << OutMeshData;
+
+					GetDerivedDataCacheRef().Put(*DerivedDataKey, DerivedData, UsdPrimPath);
+				}
 			}
 		}
 		// If the key is empty, it means the prim was invalid or not a mesh, so don't do anything
diff --git a/Engine/Plugins/Importers/USDImporter/Source/GeometryCacheUSD/Public/GeometryCacheTrackUSDTypes.h b/Engine/Plugins/Importers/USDImporter/Source/GeometryCacheUSD/Public/GeometryCacheTrackUSDTypes.h
--- a/Engine/Plugins/Importers/USDImporter/Source/GeometryCacheUSD/Public/GeometryCacheTrackUSDTypes.h
+++ b/Engine/Plugins/Importers/USDImporter/Source/GeometryCacheUSD/Public/GeometryCacheTrackUSDTypes.h
@@ -7,4 +7,4 @@
 class UGeometryCacheTrackUsd;
 struct FGeometryCacheMeshData;
 
-typedef TFunction< void( const TWeakObjectPtr<UGeometryCacheTrackUsd>, float Time, FGeometryCacheMeshData& ) > FReadUsdMeshFunction;
+typedef TFunction< bool( const TWeakObjectPtr<UGeometryCacheTrackUsd>, float Time, FGeometryCacheMeshData& ) > FReadUsdMeshFunction;
diff --git a/Engine/Plugins/Importers/USDImporter/Source/USDSchemas/Private/USDGeomMeshTranslator.cpp b/Engine/Plugins/Importers/USDImporter/Source/USDSchemas/Private/USDGeomMeshTranslator.cpp
--- a/Engine/Plugins/Importers/USDImporter/Source/USDSchemas/Private/USDGeomMeshTranslator.cpp
+++ b/Engine/Plugins/Importers/USDImporter/Source/USDSchemas/Private/USDGeomMeshTranslator.cpp
@@ -891,18 +891,18 @@ namespace UsdGeomMeshTranslatorImpl
 					UGeometryCacheTrackUsd* Track = TrackPtr.Get();
 					if ( !Track )
 					{
-						return;
+						return false;
 					}
 
 					if ( !Track->CurrentStagePinned )
 					{
-						return;
+						return false;
 					}
 
 					UE::FUsdPrim Prim = Track->CurrentStagePinned.GetPrimAtPath( UE::FSdfPath{ *Track->PrimPath } );
 					if ( !Prim )
 					{
-						return;
+						return false;
 					}
 
 					// Get MeshDescription associated with the prim
@@ -911,18 +911,20 @@ namespace UsdGeomMeshTranslatorImpl
 					TArray< UsdUtils::FUsdPrimMaterialAssignmentInfo > LODIndexToMaterialInfo;
 					const bool bAllowInterpretingLODs = false;  // GeometryCaches don't have LODs, so we will never do this
 
-					Options.TimeCode = pxr::UsdTimeCode{ Time };
-					Options.MaterialToPrimvarToUVIndex = &Track->MaterialToPrimvarToUVIndex;
+					// Need a local copy of Options to set the TimeCode since this function is called from multiple worker threads
+					UsdToUnreal::FUsdMeshConversionOptions LocalOptions(Options);
+					LocalOptions.TimeCode = pxr::UsdTimeCode{ Time };
+					LocalOptions.MaterialToPrimvarToUVIndex = &Track->MaterialToPrimvarToUVIndex;
 
 					UsdGeomMeshTranslatorImpl::LoadMeshDescriptions(
 						pxr::UsdTyped( Prim ),
 						LODIndexToMeshDescription,
 						LODIndexToMaterialInfo,
-						Options,
+						LocalOptions,
 						bAllowInterpretingLODs
 					);
 
-					// Convert the MeshDescription to MeshData
+					// Convert the MeshDescription to MeshData, first LOD only
 					for ( FMeshDescription& MeshDescription : LODIndexToMeshDescription )
 					{
 						if ( !MeshDescription.IsEmpty() )
@@ -942,8 +944,11 @@ namespace UsdGeomMeshTranslatorImpl
 							FStaticMeshOperations::ComputeTangentsAndNormals( MeshDescription, ComputeNTBsOptions );
 
 							UsdGeomMeshTranslatorImpl::GeometryCacheDataForMeshDescription( OutMeshData, MeshDescription );
+
+							return true;
 						}
 					}
+					return false;
 				}
 			);
 
diff --git a/Engine/Plugins/Runtime/GeometryCache/Source/GeometryCacheStreamer/Private/GeometryCacheStreamBase.cpp b/Engine/Plugins/Runtime/GeometryCache/Source/GeometryCacheStreamer/Private/GeometryCacheStreamBase.cpp
--- a/Engine/Plugins/Runtime/GeometryCache/Source/GeometryCacheStreamer/Private/GeometryCacheStreamBase.cpp
+++ b/Engine/Plugins/Runtime/GeometryCache/Source/GeometryCacheStreamer/Private/GeometryCacheStreamBase.cpp
@@ -404,7 +404,7 @@ void FGeometryCacheStreamBase::SetLimits(float InMaxMemoryAllowed, float InMaxCa
 		MaxMemAllowed = InMaxMemoryAllowed;
 
 		MaxCachedDuration = FMath::Min(InMaxCachedDuration, Details.Duration);
-		MaxCachedFrames = FMath::Min(FMath::FloorToInt(MaxCachedDuration / Details.SecondsPerFrame), Details.NumFrames);
+		MaxCachedFrames = FMath::Min(FMath::RoundToInt(MaxCachedDuration / Details.SecondsPerFrame), Details.NumFrames);
 	}
 }
 
