USD: Fix usage of ObjectTools on non-editor targets.

diff --git a/Engine/Plugins/Importers/USDImporter/Source/USDSchemas/Private/USDSkelRootTranslator.cpp b/Engine/Plugins/Importers/USDImporter/Source/USDSchemas/Private/USDSkelRootTranslator.cpp
--- a/Engine/Plugins/Importers/USDImporter/Source/USDSchemas/Private/USDSkelRootTranslator.cpp
+++ b/Engine/Plugins/Importers/USDImporter/Source/USDSchemas/Private/USDSkelRootTranslator.cpp
@@ -33,7 +33,6 @@
 #include "GroomComponent.h"
 #include "Materials/Material.h"
 #include "Materials/MaterialInstanceConstant.h"
-#include "ObjectTools.h"
 #include "PhysicsAssetUtils.h"
 #include "Rendering/SkeletalMeshLODImporterData.h"
 #include "Serialization/BufferArchive.h"
@@ -48,6 +47,7 @@
 	#include "Kismet2/BlueprintEditorUtils.h"
 	#include "Kismet2/CompilerResultsLog.h"
 	#include "MaterialEditingLibrary.h"
+	#include "ObjectTools.h"
 #endif // WITH_EDITOR
 
 #include "USDIncludesStart.h"
@@ -796,7 +796,9 @@ namespace UsdSkelRootTranslatorImpl
 				*CreationErrorMessage.ToString()
 			);
 
+#if WITH_EDITOR
 			ensure( ObjectTools::DeleteSingleObject( Result ) );
+#endif // WITH_EDITOR
 
 			Result = nullptr;
 		}
