USD: Fix usage of editor-only include on non-editor targets.

diff --git a/Engine/Plugins/Importers/USDImporter/Source/USDSchemas/Private/USDSkelRootTranslator.cpp b/Engine/Plugins/Importers/USDImporter/Source/USDSchemas/Private/USDSkelRootTranslator.cpp
--- a/Engine/Plugins/Importers/USDImporter/Source/USDSchemas/Private/USDSkelRootTranslator.cpp
+++ b/Engine/Plugins/Importers/USDImporter/Source/USDSchemas/Private/USDSkelRootTranslator.cpp
@@ -33,7 +33,6 @@
 #include "GroomComponent.h"
 #include "Materials/Material.h"
 #include "Materials/MaterialInstanceConstant.h"
-#include "PhysicsAssetUtils.h"
 #include "Rendering/SkeletalMeshLODImporterData.h"
 #include "Serialization/BufferArchive.h"
 
@@ -48,6 +47,7 @@
 	#include "Kismet2/CompilerResultsLog.h"
 	#include "MaterialEditingLibrary.h"
 	#include "ObjectTools.h"
+	#include "PhysicsAssetUtils.h"
 #endif // WITH_EDITOR
 
 #include "USDIncludesStart.h"
@@ -774,6 +774,7 @@ namespace UsdSkelRootTranslatorImpl
 
 	UPhysicsAsset* GenerateAndAssignPhysicsAsset( USkeletalMesh* SkeletalMesh, EObjectFlags Flags )
 	{
+#if WITH_EDITOR
 		FName AssetName = MakeUniqueObjectName(
 			GetTransientPackage(),
 			UPhysicsAsset::StaticClass(),
@@ -796,14 +797,15 @@ namespace UsdSkelRootTranslatorImpl
 				*CreationErrorMessage.ToString()
 			);
 
-#if WITH_EDITOR
 			ensure( ObjectTools::DeleteSingleObject( Result ) );
-#endif // WITH_EDITOR
 
 			Result = nullptr;
 		}
 
 		return Result;
+#else
+		return nullptr;
+#endif // WITH_EDITOR
 	}
 
 	UAnimBlueprint* CreateAnimBlueprint(
