USD: Always export to UsdGeomMesh specifying no subdivision scheme.

diff --git a/Engine/Plugins/Importers/USDImporter/Source/USDUtilities/Private/USDGeomMeshConversion.cpp b/Engine/Plugins/Importers/USDImporter/Source/USDUtilities/Private/USDGeomMeshConversion.cpp
--- a/Engine/Plugins/Importers/USDImporter/Source/USDUtilities/Private/USDGeomMeshConversion.cpp
+++ b/Engine/Plugins/Importers/USDImporter/Source/USDUtilities/Private/USDGeomMeshConversion.cpp
@@ -153,6 +153,12 @@ namespace UE::UsdGeomMeshConversion::Private
 
 			// Normals
 			{
+				// We need to emit this if we're writing normals (which we always are) because any DCC that can
+				// actually subdivide (like usdview) will just discard authored normals and fully recompute them
+				// on-demand in case they have a valid subdivision scheme (which is the default state).
+				// Reference: https://graphics.pixar.com/usd/release/api/class_usd_geom_mesh.html#UsdGeom_Mesh_Normals
+				ensure( UsdMesh.CreateSubdivisionSchemeAttr( pxr::VtValue{ pxr::UsdGeomTokens->none } ) );
+
 				pxr::UsdAttribute NormalsAttribute = UsdMesh.CreateNormalsAttr();
 				if ( NormalsAttribute )
 				{
diff --git a/Engine/Plugins/Importers/USDImporter/Source/USDUtilities/Private/USDSkeletalDataConversion.cpp b/Engine/Plugins/Importers/USDImporter/Source/USDUtilities/Private/USDSkeletalDataConversion.cpp
--- a/Engine/Plugins/Importers/USDImporter/Source/USDUtilities/Private/USDSkeletalDataConversion.cpp
+++ b/Engine/Plugins/Importers/USDImporter/Source/USDUtilities/Private/USDSkeletalDataConversion.cpp
@@ -744,6 +744,12 @@ namespace UnrealToUsdImpl
 
 			// Normals
 			{
+				// We need to emit this if we're writing normals (which we always are) because any DCC that can
+				// actually subdivide (like usdview) will just discard authored normals and fully recompute them
+				// on-demand in case they have a valid subdivision scheme (which is the default state).
+				// Reference: https://graphics.pixar.com/usd/release/api/class_usd_geom_mesh.html#UsdGeom_Mesh_Normals
+				ensure( UsdLODPrimGeomMesh.CreateSubdivisionSchemeAttr( pxr::VtValue{ pxr::UsdGeomTokens->none } ) );
+
 				pxr::UsdAttribute NormalsAttribute = UsdLODPrimGeomMesh.CreateNormalsAttr();
 				if ( NormalsAttribute )
 				{
