USD: Force prim info widget to refresh whenever the SUsdStage window as a whole refreshes;

Change delayed UI refreshes from AsyncTasks to core tickers as the former can be called in a non-deterministic way.

diff --git a/Engine/Plugins/Importers/USDImporter/Source/USDStageEditor/Private/SUSDStage.cpp b/Engine/Plugins/Importers/USDImporter/Source/USDStageEditor/Private/SUSDStage.cpp
--- a/Engine/Plugins/Importers/USDImporter/Source/USDStageEditor/Private/SUSDStage.cpp
+++ b/Engine/Plugins/Importers/USDImporter/Source/USDStageEditor/Private/SUSDStage.cpp
@@ -312,33 +312,38 @@ void SUsdStage::SetupStageActorDelegates()
 			{
 				// The USD notices may come from a background USD TBB thread, but we should only update slate from the main/slate threads.
 				// We can't retrieve the FSlateApplication singleton here (because that can also only be used from the main/slate threads),
-				// so we must use Async or core tickers here
-				AsyncTask( ENamedThreads::GameThread, [this, PrimPath, bResync]()
-				{
-					if ( this->UsdStageTreeView )
+				// so we must use core tickers here
+				FTSTicker::GetCoreTicker().AddTicker(
+					FTickerDelegate::CreateLambda( [this, PrimPath, bResync]( float Time )
 					{
-						this->UsdStageTreeView->RefreshPrim( PrimPath, bResync );
-						this->UsdStageTreeView->RequestTreeRefresh();
-					}
+						if ( this->UsdStageTreeView )
+						{
+							this->UsdStageTreeView->RefreshPrim( PrimPath, bResync );
+							this->UsdStageTreeView->RequestTreeRefresh();
+						}
 
-					const bool bViewingTheUpdatedPrim = SelectedPrimPath.Equals( PrimPath, ESearchCase::IgnoreCase );
-					const bool bViewingStageProperties = SelectedPrimPath.IsEmpty() || SelectedPrimPath == TEXT("/");
-					const bool bStageUpdated = PrimPath == TEXT("/");
+						const bool bViewingTheUpdatedPrim = SelectedPrimPath.Equals( PrimPath, ESearchCase::IgnoreCase );
+						const bool bViewingStageProperties = SelectedPrimPath.IsEmpty() || SelectedPrimPath == TEXT("/");
+						const bool bStageUpdated = PrimPath == TEXT("/");
 
-					if ( this->UsdPrimInfoWidget &&
-						 ViewModel.UsdStageActor.IsValid() &&
-						 ( bViewingTheUpdatedPrim || ( bViewingStageProperties && bStageUpdated ) ) )
-					{
-						this->UsdPrimInfoWidget->SetPrimPath( GetCurrentStage(), *PrimPath);
-					}
+						if ( this->UsdPrimInfoWidget &&
+							 ViewModel.UsdStageActor.IsValid() &&
+							 ( bViewingTheUpdatedPrim || ( bViewingStageProperties && bStageUpdated ) ) )
+						{
+							this->UsdPrimInfoWidget->SetPrimPath( GetCurrentStage(), *PrimPath);
+						}
 
-					// If we resynced our selected prim or our ancestor and have selection sync enabled, try to refresh it so that we're
-					// still selecting the same actor/component that corresponds to the prim we have currently selected
-					if ( bResync && SelectedPrimPath.StartsWith( PrimPath ) )
-					{
-						OnPrimSelectionChanged( { SelectedPrimPath } );
-					}
-				});
+						// If we resynced our selected prim or our ancestor and have selection sync enabled, try to refresh it so that we're
+						// still selecting the same actor/component that corresponds to the prim we have currently selected
+						if ( bResync && SelectedPrimPath.StartsWith( PrimPath ) )
+						{
+							OnPrimSelectionChanged( { SelectedPrimPath } );
+						}
+
+						// Returning false means this is a one-off, and won't repeat
+						return false;
+					})
+				);
 			}
 		);
 
@@ -346,29 +351,34 @@ void SUsdStage::SetupStageActorDelegates()
 		OnStageChangedHandle = ViewModel.UsdStageActor->OnStageChanged.AddLambda(
 			[ this ]()
 			{
-				AsyncTask(ENamedThreads::GameThread, [this]()
-				{
-					// So we can reset even if our actor is being destroyed right now
-					const bool bEvenIfPendingKill = true;
-					if ( ViewModel.UsdStageActor.IsValid( bEvenIfPendingKill ) )
+				FTSTicker::GetCoreTicker().AddTicker(
+					FTickerDelegate::CreateLambda( [this]( float Time )
 					{
-						// Reset our selection to the stage root
-						SelectedPrimPath = TEXT( "/" );
-
-						if ( this->UsdPrimInfoWidget )
+						// So we can reset even if our actor is being destroyed right now
+						const bool bEvenIfPendingKill = true;
+						if ( ViewModel.UsdStageActor.IsValid( bEvenIfPendingKill ) )
 						{
-							this->UsdPrimInfoWidget->SetPrimPath( GetCurrentStage(), TEXT("/"));
-						}
+							// Reset our selection to the stage root
+							SelectedPrimPath = TEXT( "/" );
 
-						if ( this->UsdStageTreeView )
-						{
-							this->UsdStageTreeView->ClearSelection();
-							this->UsdStageTreeView->RequestTreeRefresh();
+							if ( this->UsdPrimInfoWidget )
+							{
+								this->UsdPrimInfoWidget->SetPrimPath( GetCurrentStage(), TEXT( "/" ) );
+							}
+
+							if ( this->UsdStageTreeView )
+							{
+								this->UsdStageTreeView->ClearSelection();
+								this->UsdStageTreeView->RequestTreeRefresh();
+							}
 						}
-					}
 
-					this->Refresh();
-				});
+						this->Refresh();
+
+						// Returning false means this is a one-off, and won't repeat
+						return false;
+					})
+				);
 			}
 		);
 
@@ -377,54 +387,69 @@ void SUsdStage::SetupStageActorDelegates()
 			{
 				// Refresh widgets on game thread, but close the stage right away. In some contexts this is important, for example when
 				// running a Python script: If our USD Stage Editor is open and our script deletes an actor, it will trigger OnActorDestroyed.
-				// If we had this CloseStage() call inside the AsyncTask, it would only take place when the script has finished running
+				// If we had this CloseStage() call inside the ticker, it would only take place when the script has finished running
 				// (and control flow returned to the game thread) which can lead to some weird results (and break automated tests).
 				// We could get around this on the Python script's side by just yielding, but there may be other scenarios
 				ClearStageActorDelegates();
 				this->ViewModel.CloseStage();
 
-				AsyncTask( ENamedThreads::GameThread, [this]()
-				{
-					this->Refresh();
-				});
+				FTSTicker::GetCoreTicker().AddTicker(
+					FTickerDelegate::CreateLambda( [this]( float Time )
+					{
+						this->Refresh();
+
+						// Returning false means this is a one-off, and won't repeat
+						return false;
+					})
+				);
 			}
 		);
 
 		OnStageEditTargetChangedHandle = ViewModel.UsdStageActor->GetUsdListener().GetOnStageEditTargetChanged().AddLambda(
 			[ this ]()
 			{
-				AsyncTask( ENamedThreads::GameThread, [this]()
-				{
-					AUsdStageActor* StageActor = ViewModel.UsdStageActor.Get();
-					if ( this->UsdLayersTreeView && StageActor )
+				FTSTicker::GetCoreTicker().AddTicker(
+					FTickerDelegate::CreateLambda( [this]( float Time )
 					{
-						constexpr bool bResync = false;
-						UsdLayersTreeView->Refresh(
-							StageActor->GetBaseUsdStage(),
-							StageActor->GetIsolatedUsdStage(),
-							bResync
-						);
-					}
-				});
+						AUsdStageActor* StageActor = ViewModel.UsdStageActor.Get();
+						if ( this->UsdLayersTreeView && StageActor )
+						{
+							constexpr bool bResync = false;
+							UsdLayersTreeView->Refresh(
+								StageActor->GetBaseUsdStage(),
+								StageActor->GetIsolatedUsdStage(),
+								bResync
+							);
+						}
+
+						// Returning false means this is a one-off, and won't repeat
+						return false;
+					})
+				);
 			}
 		);
 
 		OnLayersChangedHandle = ViewModel.UsdStageActor->GetUsdListener().GetOnLayersChanged().AddLambda(
 			[ this ]( const TArray< FString >& LayersNames )
 			{
-				AsyncTask( ENamedThreads::GameThread, [this]()
-				{
-					AUsdStageActor* StageActor = ViewModel.UsdStageActor.Get();
-					if ( this->UsdLayersTreeView && StageActor )
+				FTSTicker::GetCoreTicker().AddTicker(
+					FTickerDelegate::CreateLambda( [this]( float Time )
 					{
-						constexpr bool bResync = false;
-						UsdLayersTreeView->Refresh(
-							StageActor->GetBaseUsdStage(),
-							StageActor->GetIsolatedUsdStage(),
-							bResync
-						);
-					}
-				});
+						AUsdStageActor* StageActor = ViewModel.UsdStageActor.Get();
+						if ( this->UsdLayersTreeView && StageActor )
+						{
+							constexpr bool bResync = false;
+							UsdLayersTreeView->Refresh(
+								StageActor->GetBaseUsdStage(),
+								StageActor->GetIsolatedUsdStage(),
+								bResync
+							);
+						}
+
+						// Returning false means this is a one-off, and won't repeat
+						return false;
+					})
+				);
 			}
 		);
 	}
@@ -1979,6 +2004,11 @@ void SUsdStage::Refresh()
 
 		UsdStageTreeView->RequestTreeRefresh();
 	}
+
+	if ( UsdPrimInfoWidget )
+	{
+		UsdPrimInfoWidget->SetPrimPath( GetCurrentStage(), *SelectedPrimPath );
+	}
 }
 
 void SUsdStage::OnStageActorLoaded( AUsdStageActor* InUsdStageActor )
