USD: Use exported level as sublayer by default.

diff --git a/Engine/Plugins/Importers/USDImporter/Source/USDExporter/Public/LevelSequenceExporterUSDOptions.h b/Engine/Plugins/Importers/USDImporter/Source/USDExporter/Public/LevelSequenceExporterUSDOptions.h
--- a/Engine/Plugins/Importers/USDImporter/Source/USDExporter/Public/LevelSequenceExporterUSDOptions.h
+++ b/Engine/Plugins/Importers/USDImporter/Source/USDExporter/Public/LevelSequenceExporterUSDOptions.h
@@ -62,8 +62,8 @@ public:
 	/**
 	 * If checked this will also add the exported level as a sublayer to the USD files emitted for all exported level sequences
 	 */
-	UPROPERTY( EditAnywhere, config, BlueprintReadWrite, Category = "Level Sequence", meta = ( EditCondition = "bExportLevel" ) )
-	bool bUseExportedLevelAsSublayer = false;
+	UPROPERTY( EditAnywhere, config, BlueprintReadWrite, Category = "Level Sequence", meta = ( EditCondition = "bExportLevel", DisplayName = "Use Exported Level as USD Sublayer" ))
+	bool bUseExportedLevelAsSublayer = true;
 
 	/**
 	 * Whether to export levels and LevelSequences even if the existing files already describe the same versions of compatible assets.
