USD: Emit warnings and try repairing USD normals instead of silently recomputing them.

diff --git a/Engine/Plugins/Importers/USDImporter/Source/USDSchemas/Private/USDGeomMeshTranslator.cpp b/Engine/Plugins/Importers/USDImporter/Source/USDSchemas/Private/USDGeomMeshTranslator.cpp
--- a/Engine/Plugins/Importers/USDImporter/Source/USDSchemas/Private/USDGeomMeshTranslator.cpp
+++ b/Engine/Plugins/Importers/USDImporter/Source/USDSchemas/Private/USDGeomMeshTranslator.cpp
@@ -69,6 +69,12 @@ static FAutoConsoleVariableRef CVarUsdUseGeometryCache(
 
 #endif // WITH_EDITOR
 
+static float GMeshNormalRepairThreshold = 0.05f;
+static FAutoConsoleVariableRef CVarMeshNormalRepairThreshold(
+	TEXT( "USD.MeshNormalRepairThreshold" ),
+	GMeshNormalRepairThreshold,
+	TEXT( "We will try repairing up to this fraction of a Mesh's normals when invalid. If a Mesh has more invalid normals than this, we will recompute all of them. Defaults to 0.05 (5% of all normals)." ) );
+
 namespace UsdGeomMeshTranslatorImpl
 {
 	bool ShouldEnableNanite(
@@ -595,6 +601,60 @@ namespace UsdGeomMeshTranslatorImpl
 		}
 	}
 
+	void RepairNormalsAndTangents( const FString& PrimPath, FMeshDescription& MeshDescription )
+	{
+		FStaticMeshConstAttributes Attributes{ MeshDescription };
+		TArrayView<const FVector3f> VertexInstanceNormals = Attributes.GetVertexInstanceNormals().GetRawArray();
+
+		// Similar to FStaticMeshOperations::AreNormalsAndTangentsValid but we don't care about tangents since we never
+		// read those from USD
+		uint64 InvalidNormalCount = 0;
+		for ( const FVertexInstanceID VertexInstanceID : MeshDescription.VertexInstances().GetElementIDs() )
+		{
+			if ( VertexInstanceNormals[ VertexInstanceID ].IsNearlyZero() || VertexInstanceNormals[ VertexInstanceID ].ContainsNaN() )
+			{
+				++InvalidNormalCount;
+			}
+		}
+		if ( InvalidNormalCount == 0 )
+		{
+			return;
+		}
+
+		const float InvalidNormalFraction = ( float ) InvalidNormalCount / ( float ) VertexInstanceNormals.Num();
+
+		// We always need to do this at this point as ComputeTangentsAndNormals will end up computing tangents anyway
+		// and our triangle tangents are always invalid
+		FStaticMeshOperations::ComputeTriangleTangentsAndNormals( MeshDescription );
+
+		const static FString MeshNormalRepairThresholdText = TEXT( "USD.MeshNormalRepairThreshold" );
+
+		// Make sure our normals can be rebuilt from MeshDescription::InitializeAutoGeneratedAttributes in case some tool needs them.
+		// Always force-compute tangents here as we never have them anyway. If we don't force them to be recomputed we'll get
+		// the worst of both worlds as some of these will be arbitrarily recomputed anyway, and some will be left invalid
+		EComputeNTBsFlags Options = EComputeNTBsFlags::UseMikkTSpace | EComputeNTBsFlags::Tangents;
+		if ( InvalidNormalFraction >= GMeshNormalRepairThreshold )
+		{
+			Options |= EComputeNTBsFlags::Normals;
+			UE_LOG( LogUsd, Warning, TEXT( "%f%% of the normals from Mesh prim '%s' are invalid. This is at or above the threshold of '%f%%' (configurable via the cvar '%s'), so normals will be discarded and fully recomputed." ),
+				InvalidNormalFraction * 100.0f,
+				*PrimPath,
+				GMeshNormalRepairThreshold * 100.0f,
+				*MeshNormalRepairThresholdText
+			);
+		}
+		else if ( InvalidNormalFraction > 0 )
+		{
+			UE_LOG( LogUsd, Warning, TEXT( "%f%% of the normals from Mesh prim '%s' are invalid. This is below the threshold of '%f%%' (configurable via the cvar '%s'), so the invalid normals will be repaired." ),
+				InvalidNormalFraction * 100.0f,
+				*PrimPath,
+				GMeshNormalRepairThreshold * 100.0f,
+				*MeshNormalRepairThresholdText
+			);
+		}
+		FStaticMeshOperations::ComputeTangentsAndNormals( MeshDescription, Options );
+	}
+
 	UStaticMesh* CreateStaticMesh( TArray<FMeshDescription>& LODIndexToMeshDescription, FUsdSchemaTranslationContext& Context, const FString& MeshName, const bool bShouldEnableNanite, bool& bOutIsNew )
 	{
 		UStaticMesh* StaticMesh = nullptr;
@@ -632,6 +692,9 @@ namespace UsdGeomMeshTranslatorImpl
 			// these build steps are async/thread-pool based.
 			SHA1.Update( reinterpret_cast< const uint8* >( &bShouldEnableNanite ), sizeof( bShouldEnableNanite ) );
 
+			// Hash the threshhold so that if we update it and reload we'll regenerate static meshes
+			SHA1.Update( reinterpret_cast< const uint8* >( &GMeshNormalRepairThreshold ), sizeof( GMeshNormalRepairThreshold ) );
+
 			SHA1.Final();
 			SHA1.GetHash(&AllLODHash.Hash[0]);
 		}
@@ -650,14 +713,12 @@ namespace UsdGeomMeshTranslatorImpl
 			{
 				FMeshDescription& MeshDescription = LODIndexToMeshDescription[LODIndex];
 
-				// Make sure our normals/tangents can be rebuilt from MeshDescription::InitializeAutoGeneratedAttributes in case some tool needs them
-				bool bHasInvalidNormals, bHasInvalidTangents;
-				FStaticMeshOperations::AreNormalsAndTangentsValid( MeshDescription, bHasInvalidNormals, bHasInvalidTangents );
+				RepairNormalsAndTangents( MeshName, MeshDescription );
 
 				FStaticMeshSourceModel& SourceModel = StaticMesh->AddSourceModel();
 				SourceModel.BuildSettings.bGenerateLightmapUVs = false;
-				SourceModel.BuildSettings.bRecomputeNormals = bHasInvalidNormals;
-				SourceModel.BuildSettings.bRecomputeTangents = bHasInvalidTangents;
+				SourceModel.BuildSettings.bRecomputeNormals = false;
+				SourceModel.BuildSettings.bRecomputeTangents = false;
 				SourceModel.BuildSettings.bBuildReversedIndexBuffer = false;
 
 				FMeshDescription* StaticMeshDescription = StaticMesh->CreateMeshDescription(LODIndex);
