USD: Allow exporting animations exclusively from selected actor and component bindings.

diff --git a/Engine/Plugins/Importers/USDImporter/Source/USDExporter/Private/LevelExporterUSDOptions.cpp b/Engine/Plugins/Importers/USDImporter/Source/USDExporter/Private/LevelExporterUSDOptions.cpp
--- a/Engine/Plugins/Importers/USDImporter/Source/USDExporter/Private/LevelExporterUSDOptions.cpp
+++ b/Engine/Plugins/Importers/USDImporter/Source/USDExporter/Private/LevelExporterUSDOptions.cpp
@@ -3,6 +3,7 @@
 #include "LevelExporterUSDOptions.h"
 
 #include "USDAssetOptions.h"
+#include "USDExporterModule.h"
 
 #include "AnalyticsEventAttribute.h"
 
@@ -50,6 +51,11 @@ void UsdUtils::AddAnalyticsAttributes(
 void UsdUtils::HashForLevelExport( const FLevelExporterUSDOptionsInner& Options, FSHA1& HashToUpdate )
 {
 	HashToUpdate.Update( reinterpret_cast< const uint8* >( &Options.bSelectionOnly ), sizeof( Options.bSelectionOnly ) );
+	if ( Options.bSelectionOnly )
+	{
+		IUsdExporterModule::HashEditorSelection( HashToUpdate );
+	}
+
 	HashToUpdate.Update( reinterpret_cast< const uint8* >( &Options.bExportActorFolders ), sizeof( Options.bExportActorFolders ) );
 	HashToUpdate.Update( reinterpret_cast< const uint8* >( &Options.bIgnoreSequencerAnimations ), sizeof( Options.bIgnoreSequencerAnimations ) );
 	HashToUpdate.Update( reinterpret_cast< const uint8* >( &Options.bExportFoliageOnActorsLayer ), sizeof( Options.bExportFoliageOnActorsLayer ) );
diff --git a/Engine/Plugins/Importers/USDImporter/Source/USDExporter/Private/LevelSequenceExporterUSD.cpp b/Engine/Plugins/Importers/USDImporter/Source/USDExporter/Private/LevelSequenceExporterUSD.cpp
--- a/Engine/Plugins/Importers/USDImporter/Source/USDExporter/Private/LevelSequenceExporterUSD.cpp
+++ b/Engine/Plugins/Importers/USDImporter/Source/USDExporter/Private/LevelSequenceExporterUSD.cpp
@@ -36,6 +36,7 @@
 #include "MovieSceneTimeHelpers.h"
 #include "MovieSceneTrack.h"
 #include "Sections/MovieSceneSubSection.h"
+#include "Selection.h"
 #include "Sequencer/MovieSceneControlRigParameterTrack.h"
 #include "Subsystems/AssetEditorSubsystem.h"
 #include "Tracks/MovieScenePropertyTrack.h"
@@ -404,6 +405,12 @@ namespace UE
 
 				ULevelSequenceExporterUsdOptions* ExportOptions;
 
+				// If ExportOptions->bSelectionOnly is true, this specifies the components whose bindings we should export
+				TSet<UActorComponent*> SelectedComponents;
+
+				// If ExportOptions->bSelectionOnly is true, this specifies the actors whose bindings we should export
+				TSet<AActor*> SelectedActors;
+
 				// Where we store our ExportTask's bReplaceIdentical, which indicates if we should overwrite files or not
 				bool bReplaceIdentical;
 
@@ -653,10 +660,29 @@ namespace UE
 					// We always use components here because when exporting actors and components to USD we basically
 					// just ignore actors altogether and export the component attachment hierarchy instead
 					USceneComponent* BoundComponent = Cast<USceneComponent>( BoundObject );
-					if ( !BoundComponent )
+					if ( BoundComponent )
+					{
+						// Only ignore component bindings if we're actually selecting some components (including this one),
+						// and our parent actor is not selected. This so that if select some components of actor X but
+						// also actor Y directly (which has component bindings) we'll export all of Y's component bindings, but
+						// only the selected component bindings of X
+						if ( Context.ExportOptions->bSelectionOnly &&
+							 Context.SelectedComponents.Num() > 0 &&
+							 !Context.SelectedActors.Contains( BoundComponent->GetOwner() ) &&
+							 !Context.SelectedComponents.Contains( BoundComponent ) )
+						{
+							continue;
+						}
+					}
+					else
 					{
 						if ( const AActor* Actor = Cast<AActor>( BoundObject ) )
 						{
+							if ( Context.ExportOptions->bSelectionOnly && !Context.SelectedActors.Contains( Actor ) )
+							{
+								continue;
+							}
+
 							BoundComponent = Actor->GetRootComponent();
 						}
 					}
@@ -1223,7 +1249,7 @@ bool ULevelSequenceExporterUsd::ExportBinary( UObject* Object, const TCHAR* Type
 			LevelExportTask->Options = LevelOptions;
 			LevelExportTask->Exporter = nullptr;
 			LevelExportTask->Filename = Context.LevelFilePath;
-			LevelExportTask->bSelected = false;
+			LevelExportTask->bSelected = LevelOptions->Inner.bSelectionOnly; // Move this as the level exporter will favor bSelected
 			LevelExportTask->bReplaceIdentical = ExportTask->bReplaceIdentical;
 			LevelExportTask->bPrompt = false;
 			LevelExportTask->bUseFileArchive = false;
@@ -1269,6 +1295,28 @@ bool ULevelSequenceExporterUsd::ExportBinary( UObject* Object, const TCHAR* Type
 		}
 	}
 
+	// Get selected objects
+	if ( GEditor && Options->bSelectionOnly )
+	{
+		USelection* ComponentSelection= GEditor->GetSelectedComponents();
+		TArray<UActorComponent*> Components;
+		ComponentSelection->GetSelectedObjects( Components );
+		Context.SelectedComponents = TSet<UActorComponent*>{ Components };
+
+		USelection* ActorSelection = GEditor->GetSelectedActors();
+		TArray<AActor*> Actors;
+		ActorSelection->GetSelectedObjects( Actors );
+		Context.SelectedActors = TSet<AActor*>{ Actors };
+		for ( const UActorComponent* Component : Components )
+		{
+			// UE will ensure that the actor selection artificially includes all owners of all selected components, so
+			// we can't tell if an actor is intentionally selected or not. To provide some control, let's make it so that
+			// if components of an actor are selected we'll select only those particular components. If the user wants
+			// the whole actor they can just select the actor itself and none of its components
+			Context.SelectedActors.Remove( Component->GetOwner() );
+		}
+	}
+
 	TArray<UE::FUsdStage> ExportedStages;
 	LevelSequenceExporterImpl::ExportMovieSceneSequence( Context, *LevelSequence, TargetFileName, ExportedStages );
 
diff --git a/Engine/Plugins/Importers/USDImporter/Source/USDExporter/Private/LevelSequenceExporterUSDOptions.cpp b/Engine/Plugins/Importers/USDImporter/Source/USDExporter/Private/LevelSequenceExporterUSDOptions.cpp
--- a/Engine/Plugins/Importers/USDImporter/Source/USDExporter/Private/LevelSequenceExporterUSDOptions.cpp
+++ b/Engine/Plugins/Importers/USDImporter/Source/USDExporter/Private/LevelSequenceExporterUSDOptions.cpp
@@ -3,6 +3,7 @@
 #include "LevelSequenceExporterUSDOptions.h"
 
 #include "LevelExporterUSDOptions.h"
+#include "USDExporterModule.h"
 
 #include "AnalyticsEventAttribute.h"
 
@@ -18,6 +19,7 @@ void UsdUtils::AddAnalyticsAttributes(
 		InOutAttributes.Emplace( TEXT( "StartFrame" ), LexToString( Options.StartFrame ) );
 		InOutAttributes.Emplace( TEXT( "EndFrame" ), LexToString( Options.EndFrame ) );
 	}
+	InOutAttributes.Emplace( TEXT( "SelectionOnly" ), Options.bSelectionOnly );
 	InOutAttributes.Emplace( TEXT( "ExportSubsequencesAsLayers" ), Options.bExportSubsequencesAsLayers );
 	InOutAttributes.Emplace( TEXT( "ExportLevel" ), Options.bExportLevel );
 	if ( Options.bExportLevel )
@@ -45,6 +47,13 @@ void UsdUtils::HashForLevelSequenceExport( const ULevelSequenceExporterUsdOption
 		HashToUpdate.Update( reinterpret_cast< const uint8* >( &Options.StartFrame ), sizeof( Options.StartFrame ) );
 		HashToUpdate.Update( reinterpret_cast< const uint8* >( &Options.EndFrame ), sizeof( Options.EndFrame ) );
 	}
+
+	HashToUpdate.Update( reinterpret_cast< const uint8* >( &Options.bSelectionOnly ), sizeof( Options.bSelectionOnly ) );
+	if ( Options.bSelectionOnly )
+	{
+		IUsdExporterModule::HashEditorSelection( HashToUpdate );
+	}
+
 	HashToUpdate.Update( reinterpret_cast< const uint8* >( &Options.bExportSubsequencesAsLayers ), sizeof( Options.bExportSubsequencesAsLayers ) );
 
 	const bool bUsingLevelSublayer = Options.bExportLevel && Options.bUseExportedLevelAsSublayer;
diff --git a/Engine/Plugins/Importers/USDImporter/Source/USDExporter/Private/USDExporterModule.cpp b/Engine/Plugins/Importers/USDImporter/Source/USDExporter/Private/USDExporterModule.cpp
--- a/Engine/Plugins/Importers/USDImporter/Source/USDExporter/Private/USDExporterModule.cpp
+++ b/Engine/Plugins/Importers/USDImporter/Source/USDExporter/Private/USDExporterModule.cpp
@@ -6,8 +6,10 @@
 #include "USDAssetOptions.h"
 #include "USDMemory.h"
 
+#include "Editor.h"
 #include "Modules/ModuleManager.h"
 #include "PropertyEditorModule.h"
+#include "Selection.h"
 #include "UObject/ObjectSaveContext.h"
 
 class FUsdExporterModule : public IUsdExporterModule
@@ -49,5 +51,44 @@ public:
 	}
 };
 
+void IUsdExporterModule::HashEditorSelection( FSHA1& HashToUpdate )
+{
+	if ( !GEditor )
+	{
+		return;
+	}
+
+	TSet<FString> SelectedPaths;
+
+	USelection* ComponentSelection = GEditor->GetSelectedComponents();
+	TArray<UActorComponent*> Components;
+	ComponentSelection->GetSelectedObjects( Components );
+	for ( const UActorComponent* Component : Components )
+	{
+		if ( Component )
+		{
+			SelectedPaths.Add( Component->GetFullName() );
+		}
+	}
+
+	USelection* ActorSelection = GEditor->GetSelectedActors();
+	TArray<AActor*> Actors;
+	ActorSelection->GetSelectedObjects( Actors );
+	for ( const AActor* Actor : Actors )
+	{
+		if ( Actor )
+		{
+			SelectedPaths.Add( Actor->GetFullName() );
+		}
+	}
+
+	TArray<FString> SelectedPathsArray = SelectedPaths.Array();
+	SelectedPathsArray.Sort();
+	for ( const FString& Path : SelectedPathsArray )
+	{
+		HashToUpdate.UpdateWithString( *Path, Path.Len() );
+	}
+}
+
 IMPLEMENT_MODULE_USD( FUsdExporterModule, USDExporter );
 
diff --git a/Engine/Plugins/Importers/USDImporter/Source/USDExporter/Public/LevelSequenceExporterUSDOptions.h b/Engine/Plugins/Importers/USDImporter/Source/USDExporter/Public/LevelSequenceExporterUSDOptions.h
--- a/Engine/Plugins/Importers/USDImporter/Source/USDExporter/Public/LevelSequenceExporterUSDOptions.h
+++ b/Engine/Plugins/Importers/USDImporter/Source/USDExporter/Public/LevelSequenceExporterUSDOptions.h
@@ -46,6 +46,12 @@ public:
 	UPROPERTY( EditAnywhere, config, BlueprintReadWrite, Category = "Level Sequence", meta = ( EditCondition = "bOverrideExportRange" ) )
 	int32 EndFrame = 0;
 
+	/**
+	 * Whether to export animations exclusively from selected actors or components
+	 */
+	UPROPERTY( EditAnywhere, config, BlueprintReadWrite, Category = "Level Sequence" )
+	bool bSelectionOnly = false;
+
 	/**
 	 * The effect of subsequences is always included on the main exported layer, but if this option is true we will also
 	 * export individual sublayers for each subsequence, so that they can be used by themselves if needed.
diff --git a/Engine/Plugins/Importers/USDImporter/Source/USDExporter/Public/USDExporterModule.h b/Engine/Plugins/Importers/USDImporter/Source/USDExporter/Public/USDExporterModule.h
--- a/Engine/Plugins/Importers/USDImporter/Source/USDExporter/Public/USDExporterModule.h
+++ b/Engine/Plugins/Importers/USDImporter/Source/USDExporter/Public/USDExporterModule.h
@@ -7,4 +7,6 @@
 
 class IUsdExporterModule : public IModuleInterface
 {
+public:
+	static void HashEditorSelection( FSHA1& HashToUpdate );
 };
