USD: Check for timeSamples directly as XMightBeTimeVarying ignored attributes with only one timeSample.

diff --git a/Engine/Plugins/Importers/USDImporter/Source/USDSchemas/Private/USDSkelRootTranslator.cpp b/Engine/Plugins/Importers/USDImporter/Source/USDSchemas/Private/USDSkelRootTranslator.cpp
--- a/Engine/Plugins/Importers/USDImporter/Source/USDSchemas/Private/USDSkelRootTranslator.cpp
+++ b/Engine/Plugins/Importers/USDImporter/Source/USDSchemas/Private/USDSkelRootTranslator.cpp
@@ -1404,8 +1404,10 @@ namespace UsdSkelRootTranslatorImpl
 						}
 						FString SkelAnimationPrimPath = UsdToUnreal::ConvertPath( SkelAnimationPrim.GetPath() );
 
-						if ( !AnimQuery.JointTransformsMightBeTimeVarying() &&
-							( NewBlendShapes.Num() == 0 || !AnimQuery.BlendShapeWeightsMightBeTimeVarying() ) )
+						std::vector<double> JointTimeSamples;
+						std::vector<double> BlendShapeTimeSamples;
+						if ( ( !AnimQuery.GetJointTransformTimeSamples( &JointTimeSamples ) || JointTimeSamples.size() == 0 ) &&
+							( NewBlendShapes.Num() == 0 || ( !AnimQuery.GetBlendShapeWeightTimeSamples( &BlendShapeTimeSamples ) || BlendShapeTimeSamples.size() == 0 ) ) )
 						{
 							continue;
 						}
diff --git a/Engine/Plugins/Importers/USDImporter/Source/USDStage/Private/USDLevelSequenceHelper.cpp b/Engine/Plugins/Importers/USDImporter/Source/USDStage/Private/USDLevelSequenceHelper.cpp
--- a/Engine/Plugins/Importers/USDImporter/Source/USDStage/Private/USDLevelSequenceHelper.cpp
+++ b/Engine/Plugins/Importers/USDImporter/Source/USDStage/Private/USDLevelSequenceHelper.cpp
@@ -1363,15 +1363,12 @@ void FUsdLevelSequenceHelperImpl::AddCommonTracks( const UUsdPrimTwin& PrimTwin,
 		{
 			if ( UE::FUsdGeomXformable AncestorXformable{ AncestorPrim } )
 			{
-				if ( AncestorXformable.TransformMightBeTimeVarying() )
+				TArray<double> TimeSamples;
+				if ( AncestorXformable.GetTimeSamples( &TimeSamples ) && TimeSamples.Num() > 0 )
 				{
 					bNeedTrackToCompensateResetXformOp = true;
 
-					TArray<double> TimeSamples;
-					if ( AncestorXformable.GetTimeSamples( &TimeSamples ) )
-					{
-						AncestorTimeSamples.Append( TimeSamples );
-					}
+					AncestorTimeSamples.Append( TimeSamples );
 				}
 
 				// The exception is if our ancestor also wants to reset its xform stack (i.e. its transform is meant to be
@@ -1401,8 +1398,8 @@ void FUsdLevelSequenceHelperImpl::AddCommonTracks( const UUsdPrimTwin& PrimTwin,
 		}
 	}
 
-	// Test that transform might be time varying and not TransformAttribute since it will check each xform ops
-	if ( Xformable.TransformMightBeTimeVarying() || bNeedTrackToCompensateResetXformOp )
+	TArray<double> TransformTimeSamples;
+	if ( ( Xformable.GetTimeSamples( &TransformTimeSamples ) && TransformTimeSamples.Num() > 0 ) || bNeedTrackToCompensateResetXformOp )
 	{
 		TArray<UE::FUsdAttribute> Attrs = UnrealToUsd::GetAttributesForProperty( Prim, UnrealIdentifiers::TransformPropertyName );
 		if ( Attrs.Num() > 0 )
@@ -1454,8 +1451,6 @@ void FUsdLevelSequenceHelperImpl::AddCommonTracks( const UUsdPrimTwin& PrimTwin,
 		UE::FUsdAttribute VisibilityAttribute = Attrs[ 0 ];
 		if ( VisibilityAttribute )
 		{
-			const bool bAttrMightBeTimeVarying = VisibilityAttribute.ValueMightBeTimeVarying();
-
 			// Collect all the time samples we'll need to sample our visibility at (USD has inherited visibilities, so every time
 			// a parent has a key, we need to recompute the child visibility at that moment too)
 			TArray<double> TotalVisibilityTimeSamples;
@@ -1466,10 +1461,9 @@ void FUsdLevelSequenceHelperImpl::AddCommonTracks( const UUsdPrimTwin& PrimTwin,
 			// many parents and specs on multiple layers, but this is hopefully at least *a* reasonable answer.
 			UE::FUsdAttribute FirstAnimatedVisibilityParentAttr;
 
-			if ( bAttrMightBeTimeVarying || bForceVisibilityTracks )
+			if ( ( VisibilityAttribute.GetTimeSamples( TotalVisibilityTimeSamples ) && TotalVisibilityTimeSamples.Num() > 0 )
+				|| bForceVisibilityTracks )
 			{
-				VisibilityAttribute.GetTimeSamples( TotalVisibilityTimeSamples );
-
 				// TODO: Improve this, as this is extremely inefficient since we'll be parsing this tree for the root down and repeatedly
 				// redoing this one child at a time...
 				UE::FUsdPrim ParentPrim = Prim.GetParent();
@@ -1481,18 +1475,16 @@ void FUsdLevelSequenceHelperImpl::AddCommonTracks( const UUsdPrimTwin& PrimTwin,
 						if ( ParentAttrs.Num() > 0 )
 						{
 							UE::FUsdAttribute ParentVisAttr = ParentAttrs[ 0 ];
-							if ( ParentVisAttr && ParentVisAttr.ValueMightBeTimeVarying() )
+
+							TArray<double> TimeSamples;
+							if ( ParentVisAttr && ( ParentVisAttr.GetTimeSamples( TimeSamples ) && TimeSamples.Num() ) )
 							{
 								if ( !FirstAnimatedVisibilityParentAttr )
 								{
 									FirstAnimatedVisibilityParentAttr = ParentVisAttr;
 								}
 
-								TArray<double> TimeSamples;
-								if ( ParentVisAttr.GetTimeSamples( TimeSamples ) )
-								{
-									TotalVisibilityTimeSamples.Append( TimeSamples );
-								}
+								TotalVisibilityTimeSamples.Append( TimeSamples );
 							}
 						}
 					}
@@ -1507,7 +1499,7 @@ void FUsdLevelSequenceHelperImpl::AddCommonTracks( const UUsdPrimTwin& PrimTwin,
 
 			// Pick which attribute we will use to fetch the target LevelSequence to put our baked tracks
 			UE::FUsdAttribute AttributeForSequence;
-			if ( bAttrMightBeTimeVarying )
+			if ( VisibilityAttribute.GetNumTimeSamples() > 0 )
 			{
 				AttributeForSequence = VisibilityAttribute;
 			}
@@ -1583,7 +1575,7 @@ void FUsdLevelSequenceHelperImpl::AddCameraTracks( const UUsdPrimTwin& PrimTwin,
 
 		// Camera attributes should always match UE properties 1-to-1 here so just get the first
 		const UE::FUsdAttribute& Attr = Attrs[0];
-		if ( !Attr || !Attr.ValueMightBeTimeVarying() )
+		if ( !Attr || Attr.GetNumTimeSamples() == 0 )
 		{
 			continue;
 		}
@@ -1707,7 +1699,7 @@ void FUsdLevelSequenceHelperImpl::AddLightTracks( const UUsdPrimTwin& PrimTwin,
 		// The main attribute is the first one, and that will dictate whether the track is muted or not
 		// This because we don't want to mute the intensity track if just our rect light width track is muted, for example
 		UE::FUsdAttribute MainAttr = Attrs[0];
-		const bool bIsMuted = MainAttr && MainAttr.ValueMightBeTimeVarying() && UsdUtils::IsAttributeMuted( MainAttr, UsdStage );
+		const bool bIsMuted = MainAttr && MainAttr.GetNumTimeSamples() > 0 && UsdUtils::IsAttributeMuted( MainAttr, UsdStage );
 
 		// Remove attributes we failed to find on this prim (no authored data)
 		// As long as we have at least one attribute with timesamples we can carry on, because we can rely on fallback/default values for the others
@@ -1715,9 +1707,8 @@ void FUsdLevelSequenceHelperImpl::AddLightTracks( const UUsdPrimTwin& PrimTwin,
 		{
 			UE::FUsdAttribute& Attr = Attrs[AttrIndex];
 			FString AttrPath = Attr.GetPath().GetString();
-			const bool bMightBeTimeVarying = Attr.ValueMightBeTimeVarying();
 
-			if ( !Attr || !bMightBeTimeVarying )
+			if ( !Attr || Attr.GetNumTimeSamples() == 0 )
 			{
 				Attrs.RemoveAt( AttrIndex );
 			}
@@ -1947,7 +1938,7 @@ void FUsdLevelSequenceHelperImpl::AddPrim( UUsdPrimTwin& PrimTwin, bool bForceVi
 
 	for ( const UE::FUsdAttribute& PrimAttribute : PrimAttributes )
 	{
-		if ( PrimAttribute.ValueMightBeTimeVarying() )
+		if ( PrimAttribute.GetNumTimeSamples() > 0 )
 		{
 			if ( ULevelSequence* AttributeSequence = FindOrAddSequenceForAttribute( PrimAttribute ) )
 			{
@@ -2840,7 +2831,7 @@ void FUsdLevelSequenceHelperImpl::HandleMovieSceneChange( UMovieScene& MovieScen
 
 	auto RemoveTimeSamplesForAttr = []( const UE::FUsdAttribute& Attr )
 	{
-		if ( !Attr || !Attr.ValueMightBeTimeVarying() )
+		if ( !Attr || Attr.GetNumTimeSamples() == 0 )
 		{
 			return;
 		}
diff --git a/Engine/Plugins/Importers/USDImporter/Source/USDUtilities/Private/USDConversionUtils.cpp b/Engine/Plugins/Importers/USDImporter/Source/USDUtilities/Private/USDConversionUtils.cpp
--- a/Engine/Plugins/Importers/USDImporter/Source/USDUtilities/Private/USDConversionUtils.cpp
+++ b/Engine/Plugins/Importers/USDImporter/Source/USDUtilities/Private/USDConversionUtils.cpp
@@ -853,7 +853,8 @@ bool UsdUtils::IsAnimated( const pxr::UsdPrim& Prim )
 			{
 				if ( pxr::UsdGeomXformable AncestorXformable{ AncestorPrim } )
 				{
-					if ( AncestorXformable.TransformMightBeTimeVarying() )
+					std::vector<double> AncestorTimeSamples;
+					if ( AncestorXformable.GetTimeSamples( &AncestorTimeSamples ) && AncestorTimeSamples.size() > 0 )
 					{
 						return true;
 					}
@@ -875,7 +876,8 @@ bool UsdUtils::IsAnimated( const pxr::UsdPrim& Prim )
 	const std::vector< pxr::UsdAttribute >& Attributes = Prim.GetAttributes();
 	for ( const pxr::UsdAttribute& Attribute : Attributes )
 	{
-		if ( Attribute.ValueMightBeTimeVarying() )
+		std::vector<double> TimeSamples;
+		if ( Attribute.GetTimeSamples( &TimeSamples ) && TimeSamples.size() > 0 )
 		{
 			return true;
 		}
@@ -899,7 +901,10 @@ bool UsdUtils::IsAnimated( const pxr::UsdPrim& Prim )
 				continue;
 			}
 
-			if ( AnimQuery.JointTransformsMightBeTimeVarying() || AnimQuery.BlendShapeWeightsMightBeTimeVarying() )
+			std::vector<double> JointTimeSamples;
+			std::vector<double> BlendShapeTimeSamples;
+			if ( ( AnimQuery.GetJointTransformTimeSamples( &JointTimeSamples ) && JointTimeSamples.size() > 0 ) ||
+				 ( AnimQuery.GetBlendShapeWeightTimeSamples( &BlendShapeTimeSamples ) && BlendShapeTimeSamples.size() > 0 ) )
 			{
 				return true;
 			}
@@ -927,7 +932,8 @@ bool UsdUtils::HasAnimatedVisibility( const pxr::UsdPrim& Prim )
 	{
 		if ( pxr::UsdAttribute Attr = Imageable.GetVisibilityAttr() )
 		{
-			if ( Attr.ValueMightBeTimeVarying() )
+			std::vector<double> TimeSamples;
+			if ( Attr.GetTimeSamples( &TimeSamples ) && TimeSamples.size() > 0 )
 			{
 				return true;
 			}
diff --git a/Engine/Plugins/Importers/USDImporter/Source/USDUtilities/Private/USDPrimConversion.cpp b/Engine/Plugins/Importers/USDImporter/Source/USDUtilities/Private/USDPrimConversion.cpp
--- a/Engine/Plugins/Importers/USDImporter/Source/USDUtilities/Private/USDPrimConversion.cpp
+++ b/Engine/Plugins/Importers/USDImporter/Source/USDUtilities/Private/USDPrimConversion.cpp
@@ -3566,11 +3566,9 @@ bool UnrealToUsd::ConvertXformable( const UMovieScene3DTransformTrack& MovieScen
 		FFrameTime UsdStartTime = FFrameRate::TransformTime( PlaybackRange.GetLowerBoundValue(), Resolution, StageFrameRate );
 		FFrameTime UsdEndTime = FFrameRate::TransformTime( PlaybackRange.GetUpperBoundValue(), Resolution, StageFrameRate );
 
-		if ( LocationValuesX.Num() > 0 || Xformable.TransformMightBeTimeVarying() )
+		std::vector< double > UsdTimeSamples;
+		if ( LocationValuesX.Num() > 0 || ( Xformable.GetTimeSamples( &UsdTimeSamples ) && UsdTimeSamples.size() > 0 ) )
 		{
-			std::vector< double > UsdTimeSamples;
-			Xformable.GetTimeSamples( &UsdTimeSamples );
-
 			bIsDataOutOfSync = ( UsdTimeSamples.size() != LocationValuesX.Num() );
 
 			if ( !bIsDataOutOfSync )
diff --git a/Engine/Plugins/Importers/USDImporter/Source/UnrealUSDWrapper/Private/UsdWrappers/UsdAttribute.cpp b/Engine/Plugins/Importers/USDImporter/Source/UnrealUSDWrapper/Private/UsdWrappers/UsdAttribute.cpp
--- a/Engine/Plugins/Importers/USDImporter/Source/UnrealUSDWrapper/Private/UsdWrappers/UsdAttribute.cpp
+++ b/Engine/Plugins/Importers/USDImporter/Source/UnrealUSDWrapper/Private/UsdWrappers/UsdAttribute.cpp
@@ -228,6 +228,15 @@ namespace UE
 #endif // #if USE_USD_SDK
 	}
 
+	size_t FUsdAttribute::GetNumTimeSamples() const
+	{
+#if USE_USD_SDK
+		return Impl->PxrUsdAttribute.Get().GetNumTimeSamples();
+#else
+		return 0;
+#endif // #if USE_USD_SDK
+	}
+
 	bool FUsdAttribute::HasValue() const
 	{
 #if USE_USD_SDK
diff --git a/Engine/Plugins/Importers/USDImporter/Source/UnrealUSDWrapper/Public/UsdWrappers/UsdAttribute.h b/Engine/Plugins/Importers/USDImporter/Source/UnrealUSDWrapper/Public/UsdWrappers/UsdAttribute.h
--- a/Engine/Plugins/Importers/USDImporter/Source/UnrealUSDWrapper/Public/UsdWrappers/UsdAttribute.h
+++ b/Engine/Plugins/Importers/USDImporter/Source/UnrealUSDWrapper/Public/UsdWrappers/UsdAttribute.h
@@ -74,6 +74,7 @@ namespace UE
 		FName GetTypeName() const;
 
 		bool GetTimeSamples( TArray<double>& Times ) const;
+		size_t GetNumTimeSamples() const;
 
 		bool HasValue() const;
 		bool HasFallbackValue() const;
