USD: Fix ambient occlusion handling.

Export ambient occlusion under the standard "occlusion" property instead of "ambientOcclusion";
Add support for reading the "occlusion" channel off the USD material;
Upgrade base materials to have ambient occlusion parameters.

diff --git a/Engine/Plugins/Importers/USDImporter/Source/USDUtilities/Private/USDShadeConversion.cpp b/Engine/Plugins/Importers/USDImporter/Source/USDUtilities/Private/USDShadeConversion.cpp
--- a/Engine/Plugins/Importers/USDImporter/Source/USDUtilities/Private/USDShadeConversion.cpp
+++ b/Engine/Plugins/Importers/USDImporter/Source/USDUtilities/Private/USDShadeConversion.cpp
@@ -1621,7 +1621,7 @@ namespace UE
 						FallbackValue = pxr::GfVec4f{ 0.0, 0.0, 0.0, 1.0f };
 						break;
 					case MP_AmbientOcclusion:
-						InputToken = UnrealIdentifiers::AmbientOcclusion;
+						InputToken = UnrealIdentifiers::Occlusion;
 						InputType = pxr::SdfValueTypeNames->Float;
 						ConstantValue = UserConstantValue ? *UserConstantValue : ( *Samples )[ 0 ].R / 255.0f;
 						FallbackValue = pxr::GfVec4f{ 1.0, 1.0, 1.0, 1.0f };
@@ -1914,6 +1914,16 @@ bool UsdToUnreal::ConvertMaterial( const pxr::UsdShadeMaterial& UsdShadeMaterial
 		}
 	}
 
+	// Occlusion
+	{
+		if ( UsdShadeConversionImpl::GetFloatParameterValue( Connectable, UnrealIdentifiers::Occlusion, 1.0f, ParameterValue, &Material, TexturesCache, &PrimvarToUVIndex, bForceVirtualTextures ) )
+		{
+			UsdShadeConversionImpl::SetParameterValue( Material, TEXT( "AmbientOcclusion" ), ParameterValue, bForUsdPreviewSurface );
+			SetTextureComponentParamForScalarInput( ParameterValue, TEXT( "AmbientOcclusionTextureComponent" ) );
+			bHasMaterialInfo = true;
+		}
+	}
+
 	// Handle world space normals
 	pxr::UsdPrim UsdPrim = UsdShadeMaterial.GetPrim();
 	if ( pxr::UsdAttribute Attr = UsdPrim.GetAttribute( UnrealIdentifiers::WorldSpaceNormals ) )
@@ -2070,6 +2080,18 @@ bool UsdToUnreal::ConvertMaterial( const pxr::UsdShadeMaterial& UsdShadeMaterial
 		}
 	}
 
+	// Occlusion
+	if ( UsdShadeConversionImpl::GetFloatParameterValue( Connectable, UnrealIdentifiers::Occlusion, 1.f, ParameterValue, &Material, TexturesCache, &PrimvarToUVIndex ) )
+	{
+		if ( UMaterialExpression* Expression = UsdShadeConversionImpl::GetExpressionForValue( Material, ParameterValue ) )
+		{
+			EditorOnly->AmbientOcclusion.Expression = Expression;
+			SetOutputIndex( ParameterValue, EditorOnly->AmbientOcclusion.OutputIndex );
+
+			bHasMaterialInfo = true;
+		}
+	}
+
 	// Handle world space normals
 	pxr::UsdPrim UsdPrim = UsdShadeMaterial.GetPrim();
 	if ( pxr::UsdAttribute Attr = UsdPrim.GetAttribute( UnrealIdentifiers::WorldSpaceNormals ) )
diff --git a/Engine/Plugins/Importers/USDImporter/Source/UnrealUSDWrapper/Private/UnrealUSDWrapper.cpp b/Engine/Plugins/Importers/USDImporter/Source/UnrealUSDWrapper/Private/UnrealUSDWrapper.cpp
--- a/Engine/Plugins/Importers/USDImporter/Source/UnrealUSDWrapper/Private/UnrealUSDWrapper.cpp
+++ b/Engine/Plugins/Importers/USDImporter/Source/UnrealUSDWrapper/Private/UnrealUSDWrapper.cpp
@@ -113,7 +113,7 @@ namespace UnrealIdentifiers
 	const TfToken Anisotropy = TfToken("anisotropy");
 	const TfToken Tangent = TfToken("tangent");
 	const TfToken SubsurfaceColor = TfToken("subsurfaceColor");
-	const TfToken AmbientOcclusion = TfToken("ambientOcclusion");
+	const TfToken Occlusion = TfToken("occlusion");
 	const TfToken Refraction = TfToken("ior");
 
 	const TfToken Surface = TfToken("surface");
diff --git a/Engine/Plugins/Importers/USDImporter/Source/UnrealUSDWrapper/Public/UnrealUSDWrapper.h b/Engine/Plugins/Importers/USDImporter/Source/UnrealUSDWrapper/Public/UnrealUSDWrapper.h
--- a/Engine/Plugins/Importers/USDImporter/Source/UnrealUSDWrapper/Public/UnrealUSDWrapper.h
+++ b/Engine/Plugins/Importers/USDImporter/Source/UnrealUSDWrapper/Public/UnrealUSDWrapper.h
@@ -421,7 +421,7 @@ namespace UnrealIdentifiers
 	extern UNREALUSDWRAPPER_API const pxr::TfToken Anisotropy;
 	extern UNREALUSDWRAPPER_API const pxr::TfToken Tangent;
 	extern UNREALUSDWRAPPER_API const pxr::TfToken SubsurfaceColor;
-	extern UNREALUSDWRAPPER_API const pxr::TfToken AmbientOcclusion;
+	extern UNREALUSDWRAPPER_API const pxr::TfToken Occlusion;
 	extern UNREALUSDWRAPPER_API const pxr::TfToken Refraction;
 
 	// Tokens used mostly for shade material conversion
