Sequencer: Follow-up fixes to 22233823

diff --git a/Engine/Plugins/Enterprise/GLTFExporter/Source/GLTFExporter/Private/Tasks/GLTFDelayedAnimationTasks.cpp b/Engine/Plugins/Enterprise/GLTFExporter/Source/GLTFExporter/Private/Tasks/GLTFDelayedAnimationTasks.cpp
--- a/Engine/Plugins/Enterprise/GLTFExporter/Source/GLTFExporter/Private/Tasks/GLTFDelayedAnimationTasks.cpp
+++ b/Engine/Plugins/Enterprise/GLTFExporter/Source/GLTFExporter/Private/Tasks/GLTFDelayedAnimationTasks.cpp
@@ -149,9 +149,6 @@ void FGLTFDelayedLevelSequenceTask::Process()
 	UWorld* World = Level->GetWorld();
 	ULevelSequencePlayer* Player = ULevelSequencePlayer::CreateLevelSequencePlayer(World, Sequence, FMovieSceneSequencePlaybackSettings(), LevelSequenceActor);
 
-	Player->Initialize(Sequence, LevelSequenceActor->GetLevel(), FMovieSceneSequencePlaybackSettings(), FLevelSequenceCameraSettings());
-	Player->State.AssignSequence(MovieSceneSequenceID::Root, *Sequence, *Player);
-
 	FFrameRate TickResolution = MovieScene->GetTickResolution();
 	FFrameRate DisplayRate = MovieScene->GetDisplayRate(); // TODO: add option to switch between DisplayRate and TickResolution?
 	TRange<FFrameNumber> PlaybackRange = MovieScene->GetPlaybackRange();
diff --git a/Engine/Plugins/Experimental/RenderGrid/Source/RenderGridEditor/Private/UI/SRenderGridViewerLive.cpp b/Engine/Plugins/Experimental/RenderGrid/Source/RenderGridEditor/Private/UI/SRenderGridViewerLive.cpp
--- a/Engine/Plugins/Experimental/RenderGrid/Source/RenderGridEditor/Private/UI/SRenderGridViewerLive.cpp
+++ b/Engine/Plugins/Experimental/RenderGrid/Source/RenderGridEditor/Private/UI/SRenderGridViewerLive.cpp
@@ -179,9 +179,6 @@ ULevelSequencePlayer* UE::RenderGrid::Private::SRenderGridEditorViewport::GetSeq
 		{
 			if (IsValid(PlayerActor))
 			{
-				Player->Initialize(LevelSequence, World->PersistentLevel, PlaybackSettings, CameraSettings);
-				Player->State.AssignSequence(MovieSceneSequenceID::Root, *LevelSequence, *Player);
-
 				LevelSequencePlayerWorld = World;
 				LevelSequencePlayerActor = PlayerActor;
 				LevelSequencePlayer = Player;
diff --git a/Engine/Plugins/Importers/USDImporter/Source/USDStage/Private/USDLevelSequenceHelper.cpp b/Engine/Plugins/Importers/USDImporter/Source/USDStage/Private/USDLevelSequenceHelper.cpp
--- a/Engine/Plugins/Importers/USDImporter/Source/USDStage/Private/USDLevelSequenceHelper.cpp
+++ b/Engine/Plugins/Importers/USDImporter/Source/USDStage/Private/USDLevelSequenceHelper.cpp
@@ -283,10 +283,6 @@ namespace UsdLevelSequenceHelperImpl
 				goto cleanup;
 			}
 
-			FLevelSequenceCameraSettings CameraSettings;
-			LevelPlayer->Initialize( LevelSequence, World->PersistentLevel, Settings, CameraSettings );
-			LevelPlayer->State.AssignSequence( MovieSceneSequenceID::Root, *LevelSequence, *Player );
-
 			// Evaluate at the beginning of the subscene time to ensure that spawnables are created before export
 			FFrameTime StartTime = FFrameRate::TransformTime(UE::MovieScene::DiscreteInclusiveLower(MovieScene->GetPlaybackRange()).Value, MovieScene->GetTickResolution(), MovieScene->GetDisplayRate());
 			LevelPlayer->SetPlaybackPosition( FMovieSceneSequencePlaybackParams( StartTime, EUpdatePositionMethod::Play ) );
@@ -3043,11 +3039,6 @@ void FUsdLevelSequenceHelperImpl::HandleControlRigSectionChange( UMovieSceneCont
 			ALevelSequenceActor* OutActor = nullptr;
 			FMovieSceneSequencePlaybackSettings Settings;
 			Player = LevelPlayer = ULevelSequencePlayer::CreateLevelSequencePlayer( World, LevelSequence, Settings, OutActor);
-
-			FLevelSequenceCameraSettings CameraSettings;
-			LevelPlayer->Initialize( LevelSequence, World->PersistentLevel, Settings, CameraSettings );
-
-			LevelPlayer->State.AssignSequence( MovieSceneSequenceID::Root, *LevelSequence, *Player );
 		}
 
 		if ( !Player )
