USD: Support multiple bound components per prim twin;

Add automated tests for cameras with both camera and transform bindings;
Fix automated tests leaving persistent, save assets during automated tests.

diff --git a/Engine/Plugins/Importers/USDImporter/Source/USDStage/Private/USDLevelSequenceHelper.cpp b/Engine/Plugins/Importers/USDImporter/Source/USDStage/Private/USDLevelSequenceHelper.cpp
--- a/Engine/Plugins/Importers/USDImporter/Source/USDStage/Private/USDLevelSequenceHelper.cpp
+++ b/Engine/Plugins/Importers/USDImporter/Source/USDStage/Private/USDLevelSequenceHelper.cpp
@@ -541,7 +541,7 @@ private:
 	template<typename TrackType>
 	TrackType* AddTrack( const FName& PropertyPath, const UUsdPrimTwin& PrimTwin, USceneComponent& ComponentToBind, ULevelSequence& Sequence, bool bIsMuted = false  );
 
-	void RemovePossessable( ULevelSequence& Sequence, const UUsdPrimTwin& PrimTwin );
+	void RemovePossessable( const UUsdPrimTwin& PrimTwin );
 
 // Prims handling
 public:
@@ -558,13 +558,19 @@ private:
 	// Sequence Name to Prim Path. Relationship: 1 Sequence to N Prim Path.
 	TMultiMap<FName, FString> PrimPathByLevelSequenceName;
 
-	TMap< TWeakObjectPtr< const UUsdPrimTwin >, TPair< ULevelSequence*, FGuid > > SceneComponentsBindings;
+	struct FPrimTwinBindings
+	{
+		ULevelSequence* Sequence = nullptr;
+
+		// For now we support one binding per component type (mostly so we can fit a binding to a scene component and
+		// camera component for a Camera prim twin)
+		TMap< const UClass*, FGuid > ComponentClassToBindingGuid;
+	};
+
+	TMap< TWeakObjectPtr< const UUsdPrimTwin >, FPrimTwinBindings > PrimTwinToBindings;
 
 // Time codes handling
 private:
-	/** Returns the FLayerTimeInfo corresponding to the root layer */
-	FLayerTimeInfo* GetRootLayerInfo();
-
 	FLayerTimeInfo& FindOrAddLayerTimeInfo(const UE::FSdfLayer& Layer);
 	FLayerTimeInfo* FindLayerTimeInfo(const UE::FSdfLayer& Layer);
 
@@ -744,7 +750,7 @@ void FUsdLevelSequenceHelperImpl::Clear()
 	LayerTimeInfosByLayerIdentifier.Empty();
 	PrimPathByLevelSequenceName.Empty();
 	SequencesID.Empty();
-	SceneComponentsBindings.Empty();
+	PrimTwinToBindings.Empty();
 	SequenceHierarchyCache = FMovieSceneSequenceHierarchy();
 }
 
@@ -2055,10 +2061,7 @@ void FUsdLevelSequenceHelperImpl::RemovePrim( const UUsdPrimTwin& PrimTwin )
 		}
 	}
 
-	if ( ULevelSequence* Sequence = SceneComponentsBindings.FindRef( &PrimTwin ).Key )
-	{
-		RemovePossessable( *Sequence, PrimTwin );
-	}
+	RemovePossessable( PrimTwin );
 
 	for ( ULevelSequence* SequenceToRemoveForPrim : SequencesToRemoveForPrim )
 	{
@@ -2291,9 +2294,20 @@ void FUsdLevelSequenceHelperImpl::RemoveSequenceForPrim( ULevelSequence& Sequenc
 	}
 }
 
-void FUsdLevelSequenceHelperImpl::RemovePossessable( ULevelSequence& Sequence, const UUsdPrimTwin& PrimTwin )
+void FUsdLevelSequenceHelperImpl::RemovePossessable( const UUsdPrimTwin& PrimTwin )
 {
-	UMovieScene* MovieScene = Sequence.GetMovieScene();
+	FPrimTwinBindings* Bindings = PrimTwinToBindings.Find( &PrimTwin );
+	if ( !Bindings )
+	{
+		return;
+	}
+
+	if( !Bindings->Sequence )
+	{
+		return;
+	}
+
+	UMovieScene* MovieScene = Bindings->Sequence->GetMovieScene();
 	if ( !MovieScene )
 	{
 		return;
@@ -2301,11 +2315,11 @@ void FUsdLevelSequenceHelperImpl::RemovePossessable( ULevelSequence& Sequence, c
 
 	// The RemovePossessable calls Modify the MovieScene already, but the UnbindPossessableObject
 	// ones don't modify the Sequence and change properties, so we must modify them here
-	Sequence.Modify();
+	Bindings->Sequence->Modify();
 
-	if ( const TPair< ULevelSequence*, FGuid >* SceneComponentBinding = SceneComponentsBindings.Find( &PrimTwin ) )
+	for ( const TPair< const UClass*, FGuid >& Pair : Bindings->ComponentClassToBindingGuid )
 	{
-		const FGuid& ComponentPossessableGuid = SceneComponentBinding->Value;
+		const FGuid& ComponentPossessableGuid = Pair.Value;
 
 		FGuid ActorPossessableGuid;
 		if ( FMovieScenePossessable* ComponentPossessable = MovieScene->FindPossessable( ComponentPossessableGuid ) )
@@ -2316,7 +2330,7 @@ void FUsdLevelSequenceHelperImpl::RemovePossessable( ULevelSequence& Sequence, c
 		// This will also remove all tracks bound to this guid
 		if ( MovieScene->RemovePossessable( ComponentPossessableGuid ) )
 		{
-			Sequence.UnbindPossessableObjects( ComponentPossessableGuid );
+			Bindings->Sequence->UnbindPossessableObjects( ComponentPossessableGuid );
 		}
 
 		// If our parent binding has nothing else in it, we should remove it too
@@ -2336,11 +2350,11 @@ void FUsdLevelSequenceHelperImpl::RemovePossessable( ULevelSequence& Sequence, c
 		if ( bRemoveActorBinding )
 		{
 			MovieScene->RemovePossessable( ActorPossessableGuid );
-			Sequence.UnbindPossessableObjects( ActorPossessableGuid );
+			Bindings->Sequence->UnbindPossessableObjects( ActorPossessableGuid );
 		}
-
-		SceneComponentsBindings.Remove( &PrimTwin );
 	}
+
+	PrimTwinToBindings.Remove( &PrimTwin );
 }
 
 void FUsdLevelSequenceHelperImpl::RefreshSequencer()
@@ -2700,58 +2714,60 @@ FGuid FUsdLevelSequenceHelperImpl::GetOrCreateComponentBinding( const UUsdPrimTw
 		return {};
 	}
 
-	if ( TPair< ULevelSequence*, FGuid >* SceneComponentBinding = SceneComponentsBindings.Find( &PrimTwin ) )
+	FPrimTwinBindings& Bindings = PrimTwinToBindings.FindOrAdd( &PrimTwin );
+
+	ensure( Bindings.Sequence == nullptr || Bindings.Sequence == &Sequence );
+	Bindings.Sequence = &Sequence;
+
+	if ( FGuid* ExistingGuid = Bindings.ComponentClassToBindingGuid.Find( ComponentToBind.GetClass() ) )
 	{
-		return SceneComponentBinding->Value;
+		return *ExistingGuid;
 	}
-	else
-	{
-		FGuid ComponentBinding;
-		FGuid ActorBinding;
-		UObject* ComponentContext = ComponentToBind.GetWorld();
 
-		FString PrimName = FPaths::GetBaseFilename( PrimTwin.PrimPath );
+	FGuid ComponentBinding;
+	FGuid ActorBinding;
+	UObject* ComponentContext = ComponentToBind.GetWorld();
+
+	FString PrimName = FPaths::GetBaseFilename( PrimTwin.PrimPath );
 
-		// Make sure we always bind the parent actor too
-		if ( AActor* Actor = ComponentToBind.GetOwner() )
+	// Make sure we always bind the parent actor too
+	if ( AActor* Actor = ComponentToBind.GetOwner() )
+	{
+		ActorBinding = Sequence.FindBindingFromObject( Actor, Actor->GetWorld() );
+		if ( !ActorBinding.IsValid() )
 		{
-			ActorBinding = Sequence.FindBindingFromObject( Actor, Actor->GetWorld() );
-			if ( !ActorBinding.IsValid() )
-			{
-				// We use the label here because that will always be named after the prim that caused the actor
-				// to be generated. If we just used our own PrimName in here we may run into situations where a child Camera prim
-				// of a decomposed camera ends up naming the actor binding after itself, even though the parent Xform prim, and the
-				// actor on the level, maybe named something else
-				ActorBinding = MovieScene->AddPossessable(
+			// We use the label here because that will always be named after the prim that caused the actor
+			// to be generated. If we just used our own PrimName in here we may run into situations where a child Camera prim
+			// of a decomposed camera ends up naming the actor binding after itself, even though the parent Xform prim, and the
+			// actor on the level, maybe named something else
+			ActorBinding = MovieScene->AddPossessable(
 #if WITH_EDITOR
-					Actor->GetActorLabel(),
+				Actor->GetActorLabel(),
 #else
-					Actor->GetName(),
+				Actor->GetName(),
 #endif // WITH_EDITOR
-					Actor->GetClass()
-				);
-				Sequence.BindPossessableObject( ActorBinding, *Actor, Actor->GetWorld() );
-			}
-
-			ComponentContext = Actor;
+				Actor->GetClass()
+			);
+			Sequence.BindPossessableObject( ActorBinding, *Actor, Actor->GetWorld() );
 		}
 
-		ComponentBinding = MovieScene->AddPossessable( PrimName, ComponentToBind.GetClass() );
+		ComponentContext = Actor;
+	}
+
+	ComponentBinding = MovieScene->AddPossessable( PrimName, ComponentToBind.GetClass() );
 
-		if ( ActorBinding.IsValid() && ComponentBinding.IsValid() )
+	if ( ActorBinding.IsValid() && ComponentBinding.IsValid() )
+	{
+		if ( FMovieScenePossessable* ComponentPossessable = MovieScene->FindPossessable( ComponentBinding ) )
 		{
-			if ( FMovieScenePossessable* ComponentPossessable = MovieScene->FindPossessable( ComponentBinding ) )
-			{
-				ComponentPossessable->SetParent( ActorBinding, MovieScene );
-			}
+			ComponentPossessable->SetParent( ActorBinding, MovieScene );
 		}
-
-		// Bind component
-		Sequence.BindPossessableObject( ComponentBinding, ComponentToBind, ComponentContext );
-
-		SceneComponentsBindings.Emplace( &PrimTwin ) = TPair< ULevelSequence*, FGuid >( &Sequence, ComponentBinding );
-		return ComponentBinding;
 	}
+
+	// Bind component
+	Sequence.BindPossessableObject( ComponentBinding, ComponentToBind, ComponentContext );
+	Bindings.ComponentClassToBindingGuid.Emplace( ComponentToBind.GetClass(), ComponentBinding );
+	return ComponentBinding;
 }
 
 void FUsdLevelSequenceHelperImpl::HandleMovieSceneChange( UMovieScene& MovieScene )
@@ -2854,75 +2870,83 @@ void FUsdLevelSequenceHelperImpl::HandleMovieSceneChange( UMovieScene& MovieScen
 	};
 
 	// Check if we deleted things
-	for ( TMap< TWeakObjectPtr< const UUsdPrimTwin >, TPair< ULevelSequence*, FGuid > >::TIterator SceneComponentBindingIt = SceneComponentsBindings.CreateIterator();
-		SceneComponentBindingIt;
-		++SceneComponentBindingIt )
+	for ( TMap< TWeakObjectPtr< const UUsdPrimTwin >, FPrimTwinBindings >::TIterator PrimTwinIt = PrimTwinToBindings.CreateIterator();
+		PrimTwinIt;
+		++PrimTwinIt )
 	{
-		if ( SceneComponentBindingIt->Value.Key != Sequence )
+		const UUsdPrimTwin* UsdPrimTwin = PrimTwinIt->Key.Get();
+		FPrimTwinBindings& Bindings = PrimTwinIt->Value;
+
+		if ( Bindings.Sequence != Sequence )
 		{
 			continue;
 		}
 
-		const FGuid& Guid = SceneComponentBindingIt->Value.Value;
-
-		// Deleted the entire possessable
-		if ( !MovieScene.FindPossessable( Guid ) )
+		for ( TMap< const UClass*, FGuid >::TIterator BindingIt = Bindings.ComponentClassToBindingGuid.CreateIterator();
+			BindingIt;
+			++BindingIt )
 		{
-			SceneComponentBindingIt.RemoveCurrent();
-		}
+			const FGuid& Guid = BindingIt->Value;
 
-		// Check if we have an animated attribute and no track for it --> We may have deleted the track, so clear that attribute
-		// We could keep track of these when adding in some kind of map, but while slower this is likely more robust due to the need to support undo/redo
-		if ( const UUsdPrimTwin* UsdPrimTwin = SceneComponentBindingIt->Key.Get() )
-		{
-			USceneComponent* BoundComponent = UsdPrimTwin->GetSceneComponent();
-			if ( !BoundComponent )
+			// Deleted the possessable
+			if ( !MovieScene.FindPossessable( Guid ) )
 			{
-				continue;
+				BindingIt.RemoveCurrent();
 			}
 
-			const bool bIsCamera = BoundComponent->GetOwner()->IsA<ACineCameraActor>();
-			const bool bIsLight = BoundComponent->IsA<ULightComponentBase>();
-			const bool bIsSkeletal = BoundComponent->IsA<USkeletalMeshComponent>();
-
-			if ( UE::FUsdPrim UsdPrim = UsdStage.GetPrimAtPath( UE::FSdfPath( *UsdPrimTwin->PrimPath ) ) )
+			// Check if we have an animated attribute and no track for it --> We may have deleted the track, so clear that attribute
+			// We could keep track of these when adding in some kind of map, but while slower this is likely more robust due to the need to support undo/redo
+			if ( UsdPrimTwin )
 			{
-				RemoveTimeSamplesForPropertyIfNeeded( UsdPrim, Guid, UnrealIdentifiers::TransformPropertyName );
-				RemoveTimeSamplesForPropertyIfNeeded( UsdPrim, Guid, UnrealIdentifiers::HiddenInGamePropertyName );
-
-				if ( bIsCamera )
-				{
-					RemoveTimeSamplesForPropertyIfNeeded( UsdPrim, Guid, UnrealIdentifiers::CurrentFocalLengthPropertyName );
-					RemoveTimeSamplesForPropertyIfNeeded( UsdPrim, Guid, UnrealIdentifiers::ManualFocusDistancePropertyName );
-					RemoveTimeSamplesForPropertyIfNeeded( UsdPrim, Guid, UnrealIdentifiers::CurrentAperturePropertyName );
-					RemoveTimeSamplesForPropertyIfNeeded( UsdPrim, Guid, UnrealIdentifiers::SensorWidthPropertyName );
-					RemoveTimeSamplesForPropertyIfNeeded( UsdPrim, Guid, UnrealIdentifiers::SensorHeightPropertyName );
-				}
-				else if ( bIsLight )
+				USceneComponent* BoundComponent = UsdPrimTwin->GetSceneComponent();
+				if ( !BoundComponent )
 				{
-					RemoveTimeSamplesForPropertyIfNeeded( UsdPrim, Guid, UnrealIdentifiers::IntensityPropertyName );
-					RemoveTimeSamplesForPropertyIfNeeded( UsdPrim, Guid, UnrealIdentifiers::LightColorPropertyName );
-					RemoveTimeSamplesForPropertyIfNeeded( UsdPrim, Guid, UnrealIdentifiers::UseTemperaturePropertyName );
-					RemoveTimeSamplesForPropertyIfNeeded( UsdPrim, Guid, UnrealIdentifiers::TemperaturePropertyName );
-					RemoveTimeSamplesForPropertyIfNeeded( UsdPrim, Guid, UnrealIdentifiers::SourceRadiusPropertyName );
-					RemoveTimeSamplesForPropertyIfNeeded( UsdPrim, Guid, UnrealIdentifiers::SourceWidthPropertyName );
-					RemoveTimeSamplesForPropertyIfNeeded( UsdPrim, Guid, UnrealIdentifiers::SourceHeightPropertyName );
-					RemoveTimeSamplesForPropertyIfNeeded( UsdPrim, Guid, UnrealIdentifiers::OuterConeAnglePropertyName );
-					RemoveTimeSamplesForPropertyIfNeeded( UsdPrim, Guid, UnrealIdentifiers::InnerConeAnglePropertyName );
-					RemoveTimeSamplesForPropertyIfNeeded( UsdPrim, Guid, UnrealIdentifiers::LightSourceAnglePropertyName );
+					continue;
 				}
-				else if ( bIsSkeletal )
+
+				const bool bIsCamera = BoundComponent->GetOwner()->IsA<ACineCameraActor>();
+				const bool bIsLight = BoundComponent->IsA<ULightComponentBase>();
+				const bool bIsSkeletal = BoundComponent->IsA<USkeletalMeshComponent>();
+
+				if ( UE::FUsdPrim UsdPrim = UsdStage.GetPrimAtPath( UE::FSdfPath( *UsdPrimTwin->PrimPath ) ) )
 				{
-					if ( !MovieScene.FindTrack( UMovieSceneSkeletalAnimationTrack::StaticClass(), Guid ) )
+					RemoveTimeSamplesForPropertyIfNeeded( UsdPrim, Guid, UnrealIdentifiers::TransformPropertyName );
+					RemoveTimeSamplesForPropertyIfNeeded( UsdPrim, Guid, UnrealIdentifiers::HiddenInGamePropertyName );
+
+					if ( bIsCamera )
 					{
-						if ( UE::FUsdPrim SkelAnimationPrim = UsdUtils::FindFirstAnimationSource( UsdPrim ) )
+						RemoveTimeSamplesForPropertyIfNeeded( UsdPrim, Guid, UnrealIdentifiers::CurrentFocalLengthPropertyName );
+						RemoveTimeSamplesForPropertyIfNeeded( UsdPrim, Guid, UnrealIdentifiers::ManualFocusDistancePropertyName );
+						RemoveTimeSamplesForPropertyIfNeeded( UsdPrim, Guid, UnrealIdentifiers::CurrentAperturePropertyName );
+						RemoveTimeSamplesForPropertyIfNeeded( UsdPrim, Guid, UnrealIdentifiers::SensorWidthPropertyName );
+						RemoveTimeSamplesForPropertyIfNeeded( UsdPrim, Guid, UnrealIdentifiers::SensorHeightPropertyName );
+					}
+					else if ( bIsLight )
+					{
+						RemoveTimeSamplesForPropertyIfNeeded( UsdPrim, Guid, UnrealIdentifiers::IntensityPropertyName );
+						RemoveTimeSamplesForPropertyIfNeeded( UsdPrim, Guid, UnrealIdentifiers::LightColorPropertyName );
+						RemoveTimeSamplesForPropertyIfNeeded( UsdPrim, Guid, UnrealIdentifiers::UseTemperaturePropertyName );
+						RemoveTimeSamplesForPropertyIfNeeded( UsdPrim, Guid, UnrealIdentifiers::TemperaturePropertyName );
+						RemoveTimeSamplesForPropertyIfNeeded( UsdPrim, Guid, UnrealIdentifiers::SourceRadiusPropertyName );
+						RemoveTimeSamplesForPropertyIfNeeded( UsdPrim, Guid, UnrealIdentifiers::SourceWidthPropertyName );
+						RemoveTimeSamplesForPropertyIfNeeded( UsdPrim, Guid, UnrealIdentifiers::SourceHeightPropertyName );
+						RemoveTimeSamplesForPropertyIfNeeded( UsdPrim, Guid, UnrealIdentifiers::OuterConeAnglePropertyName );
+						RemoveTimeSamplesForPropertyIfNeeded( UsdPrim, Guid, UnrealIdentifiers::InnerConeAnglePropertyName );
+						RemoveTimeSamplesForPropertyIfNeeded( UsdPrim, Guid, UnrealIdentifiers::LightSourceAnglePropertyName );
+					}
+					else if ( bIsSkeletal )
+					{
+						if ( !MovieScene.FindTrack( UMovieSceneSkeletalAnimationTrack::StaticClass(), Guid ) )
 						{
-							if ( UE::FSdfLayer SkelAnimationLayer = UsdUtils::FindLayerForPrim( SkelAnimationPrim ) )
+							if ( UE::FUsdPrim SkelAnimationPrim = UsdUtils::FindFirstAnimationSource( UsdPrim ) )
 							{
-								RemoveTimeSamplesForAttr( SkelAnimationPrim.GetAttribute( TEXT( "blendShapeWeights" ) ) );
-								RemoveTimeSamplesForAttr( SkelAnimationPrim.GetAttribute( TEXT( "rotations" ) ) );
-								RemoveTimeSamplesForAttr( SkelAnimationPrim.GetAttribute( TEXT( "translations" ) ) );
-								RemoveTimeSamplesForAttr( SkelAnimationPrim.GetAttribute( TEXT( "scales" ) ) );
+								if ( UE::FSdfLayer SkelAnimationLayer = UsdUtils::FindLayerForPrim( SkelAnimationPrim ) )
+								{
+									RemoveTimeSamplesForAttr( SkelAnimationPrim.GetAttribute( TEXT( "blendShapeWeights" ) ) );
+									RemoveTimeSamplesForAttr( SkelAnimationPrim.GetAttribute( TEXT( "rotations" ) ) );
+									RemoveTimeSamplesForAttr( SkelAnimationPrim.GetAttribute( TEXT( "translations" ) ) );
+									RemoveTimeSamplesForAttr( SkelAnimationPrim.GetAttribute( TEXT( "scales" ) ) );
+								}
 							}
 						}
 					}
@@ -3201,7 +3225,15 @@ void FUsdLevelSequenceHelperImpl::HandleTrackChange( const UMovieSceneTrack& Tra
 		FScopedBlockNoticeListening BlockNotices( StageActor.Get() );
 		UE::FUsdPrim UsdPrim = UsdStage.GetPrimAtPath( UE::FSdfPath( *PrimTwin->PrimPath ) );
 
-		SceneComponentsBindings.Emplace( PrimTwin ) = TPair< ULevelSequence*, FGuid >( Sequence, PossessableGuid ); // Make sure we track this binding
+		FPrimTwinBindings& Bindings = PrimTwinToBindings.FindOrAdd(PrimTwin);
+		ensure( Bindings.Sequence == nullptr || Bindings.Sequence == Sequence );
+		Bindings.Sequence = Sequence;
+
+		// Make sure we track this binding
+		const UClass* ComponentClass = BoundSceneComponent->GetClass();
+		FGuid* FoundExistingGuid = Bindings.ComponentClassToBindingGuid.Find( ComponentClass );
+		ensure( !FoundExistingGuid || *FoundExistingGuid == PossessableGuid );
+		Bindings.ComponentClassToBindingGuid.Emplace( ComponentClass, PossessableGuid );
 
 		if ( bIsMuteChange )
 		{
