USD: Fix old detection of Light prims not working after the refactor into UsdLuxLightAPI.

diff --git a/Engine/Plugins/Importers/USDImporter/Source/USDStage/Private/USDLevelSequenceHelper.cpp b/Engine/Plugins/Importers/USDImporter/Source/USDStage/Private/USDLevelSequenceHelper.cpp
--- a/Engine/Plugins/Importers/USDImporter/Source/USDStage/Private/USDLevelSequenceHelper.cpp
+++ b/Engine/Plugins/Importers/USDImporter/Source/USDStage/Private/USDLevelSequenceHelper.cpp
@@ -1962,7 +1962,7 @@ void FUsdLevelSequenceHelperImpl::AddPrim( UUsdPrimTwin& PrimTwin, bool bForceVi
 	{
 		AddCameraTracks( PrimTwin, UsdPrim );
 	}
-	else if ( UsdPrim.IsA( TEXT( "Light" ) ) )
+	else if ( UsdPrim.HasAPI( TEXT( "UsdLuxLightAPI" ) ) )
 	{
 		AddLightTracks( PrimTwin, UsdPrim );
 	}
