USD: Don't remove referenced prims when removing the referencer prim.

diff --git a/Engine/Plugins/Importers/USDImporter/Source/USDExporter/Private/USDConversionBlueprintLibrary.cpp b/Engine/Plugins/Importers/USDImporter/Source/USDExporter/Private/USDConversionBlueprintLibrary.cpp
--- a/Engine/Plugins/Importers/USDImporter/Source/USDExporter/Private/USDConversionBlueprintLibrary.cpp
+++ b/Engine/Plugins/Importers/USDImporter/Source/USDExporter/Private/USDConversionBlueprintLibrary.cpp
@@ -525,7 +525,7 @@ void UUsdConversionBlueprintLibrary::RemoveAllPrimSpecs( const FString& StageRoo
 		return;
 	}
 
-	UsdUtils::RemoveAllPrimSpecs(
+	UsdUtils::RemoveAllLocalPrimSpecs(
 		Stage.GetPrimAtPath( UE::FSdfPath{ *PrimPath } ),
 		UE::FSdfLayer::FindOrOpen( *TargetLayer )
 	);
diff --git a/Engine/Plugins/Importers/USDImporter/Source/USDStage/Private/USDTransactor.cpp b/Engine/Plugins/Importers/USDImporter/Source/USDStage/Private/USDTransactor.cpp
--- a/Engine/Plugins/Importers/USDImporter/Source/USDStage/Private/USDTransactor.cpp
+++ b/Engine/Plugins/Importers/USDImporter/Source/USDStage/Private/USDTransactor.cpp
@@ -294,7 +294,7 @@ namespace UsdUtils
 				*PrimChange.PrimTypeName
 			);
 
-			UsdUtils::RemoveAllPrimSpecs( Stage.GetPrimAtPath( PrimPath ), Stage.GetEditTarget() );
+			UsdUtils::RemoveAllLocalPrimSpecs( Stage.GetPrimAtPath( PrimPath ), Stage.GetEditTarget() );
 			return true;
 		}
 		else if ( bRename )
diff --git a/Engine/Plugins/Importers/USDImporter/Source/USDStageEditor/Private/SUSDStageTreeView.cpp b/Engine/Plugins/Importers/USDImporter/Source/USDStageEditor/Private/SUSDStageTreeView.cpp
--- a/Engine/Plugins/Importers/USDImporter/Source/USDStageEditor/Private/SUSDStageTreeView.cpp
+++ b/Engine/Plugins/Importers/USDImporter/Source/USDStageEditor/Private/SUSDStageTreeView.cpp
@@ -1054,7 +1054,7 @@ void SUsdStageTreeView::OnDeletePrim()
 
 	for ( FUsdPrimViewModelRef SelectedItem : MySelectedItems )
 	{
-		UsdUtils::RemoveAllPrimSpecs( SelectedItem->UsdPrim );
+		UsdUtils::RemoveAllLocalPrimSpecs( SelectedItem->UsdPrim );
 	}
 }
 
diff --git a/Engine/Plugins/Importers/USDImporter/Source/USDUtilities/Private/USDConversionUtils.cpp b/Engine/Plugins/Importers/USDImporter/Source/USDUtilities/Private/USDConversionUtils.cpp
--- a/Engine/Plugins/Importers/USDImporter/Source/USDUtilities/Private/USDConversionUtils.cpp
+++ b/Engine/Plugins/Importers/USDImporter/Source/USDUtilities/Private/USDConversionUtils.cpp
@@ -1833,7 +1833,7 @@ UE::FSdfPath UsdUtils::GetPrimSpecPathForLayer( const UE::FUsdPrim& Prim, const
 	return Result;
 }
 
-USDUTILITIES_API void UsdUtils::RemoveAllPrimSpecs( const UE::FUsdPrim& Prim, const UE::FSdfLayer& Layer )
+USDUTILITIES_API void UsdUtils::RemoveAllLocalPrimSpecs( const UE::FUsdPrim& Prim, const UE::FSdfLayer& Layer )
 {
 #if USE_USD_SDK
 	FScopedUsdAllocs Allocs;
@@ -1844,8 +1844,6 @@ USDUTILITIES_API void UsdUtils::RemoveAllPrimSpecs( const UE::FUsdPrim& Prim, co
 		return;
 	}
 
-	pxr::SdfChangeBlock Block;
-
 	pxr::SdfLayerRefPtr UsdLayer{ Layer };
 	pxr::UsdStageRefPtr UsdStage = UsdPrim.GetStage();
 
@@ -1863,6 +1861,8 @@ USDUTILITIES_API void UsdUtils::RemoveAllPrimSpecs( const UE::FUsdPrim& Prim, co
 		}
 	}
 
+	const pxr::SdfPath TargetPath = UsdPrim.GetPrimPath();
+
 	for ( const pxr::SdfPrimSpecHandle& Spec : UsdPrim.GetPrimStack() )
 	{
 		// For whatever reason sometimes there are invalid specs in the layer stack, so we need to be careful
@@ -1872,7 +1872,14 @@ USDUTILITIES_API void UsdUtils::RemoveAllPrimSpecs( const UE::FUsdPrim& Prim, co
 		}
 
 		pxr::SdfPath SpecPath = Spec->GetPath();
-		if ( !SpecPath.IsPrimPath() )
+
+		// Filtering by the target path is important because if X references Y, we'll actually get Y's specs within
+		// X.GetPrimStack(), and we don't want to remove the referenced specs when removing the referencer.
+		// We strip variant selections here because when removing something inside the variant, SpecPath will contain
+		// the variant selection and look like '/PrimWithVarSet{VarSet=SomeVar}ChildPrim', but our TargetPath will
+		// just look like '/PrimWithVarSet/ChildPrim' instead. These do refer to the exact same prim on the stage
+		// though (when SomeVar is active at least), so we do want to remove both
+		if ( !SpecPath.IsPrimPath() || SpecPath.StripAllVariantSelections() != TargetPath )
 		{
 			continue;
 		}
@@ -1906,7 +1913,7 @@ bool UsdUtils::CutPrims( const TArray<UE::FUsdPrim>& Prims )
 
 	for ( const UE::FUsdPrim& Prim : Prims )
 	{
-		UsdUtils::RemoveAllPrimSpecs( Prim );
+		UsdUtils::RemoveAllLocalPrimSpecs( Prim );
 	}
 
 	return true;
diff --git a/Engine/Plugins/Importers/USDImporter/Source/USDUtilities/Private/USDShadeConversion.cpp b/Engine/Plugins/Importers/USDImporter/Source/USDUtilities/Private/USDShadeConversion.cpp
--- a/Engine/Plugins/Importers/USDImporter/Source/USDUtilities/Private/USDShadeConversion.cpp
+++ b/Engine/Plugins/Importers/USDImporter/Source/USDUtilities/Private/USDShadeConversion.cpp
@@ -2491,7 +2491,7 @@ bool UsdUtils::RemoveUnrealSurfaceOutput( pxr::UsdPrim& MaterialPrim, const UE::
 		if ( pxr::UsdShadeShader SurfaceShader = ShadeMaterial.ComputeSurfaceSource( UnrealIdentifiers::Unreal ) )
 		{
 			// Fully remove the UnrealShader
-			UsdUtils::RemoveAllPrimSpecs( UE::FUsdPrim{ SurfaceShader.GetPrim() }, LayerToAuthorIn );
+			UsdUtils::RemoveAllLocalPrimSpecs( UE::FUsdPrim{ SurfaceShader.GetPrim() }, LayerToAuthorIn );
 		}
 
 		// Disconnect would author something like `token outputs:unreal:surface.connect = None`, which is not quite what we want:
diff --git a/Engine/Plugins/Importers/USDImporter/Source/USDUtilities/Public/USDConversionUtils.h b/Engine/Plugins/Importers/USDImporter/Source/USDUtilities/Public/USDConversionUtils.h
--- a/Engine/Plugins/Importers/USDImporter/Source/USDUtilities/Public/USDConversionUtils.h
+++ b/Engine/Plugins/Importers/USDImporter/Source/USDUtilities/Public/USDConversionUtils.h
@@ -283,7 +283,7 @@ namespace UsdUtils
 	USDUTILITIES_API UE::FSdfPath GetPrimSpecPathForLayer( const UE::FUsdPrim& Prim, const UE::FSdfLayer& Layer );
 
 	/**
-	 * Removes all the prim specs for Prim on the given Layer.
+	 * Removes all the prim specs for Prim on the given Layer, if the layer belongs to the stage's local layer stack.
 	 *
 	 * This function is useful in case the prim is inside a variant set: In that case, just calling FUsdStage::RemovePrim()
 	 * will attempt to remove the "/Root/Example/Child", which wouldn't remove the "/Root{Varset=Var}Example/Child" spec,
@@ -295,7 +295,7 @@ namespace UsdUtils
 	 * @param Layer - Layer to remove prim specs from. This can be left with the invalid layer (default) in order to remove all
 	 *				  specs from the entire stage's local layer stack.
 	 */
-	USDUTILITIES_API void RemoveAllPrimSpecs( const UE::FUsdPrim& Prim, const UE::FSdfLayer& Layer = UE::FSdfLayer{} );
+	USDUTILITIES_API void RemoveAllLocalPrimSpecs( const UE::FUsdPrim& Prim, const UE::FSdfLayer& Layer = UE::FSdfLayer{} );
 
 	/**
 	 * Copies flattened versions of the input prims onto the clipboard stage and removes all the prim specs for Prims from their stages.
