USD: Fix crash when exporting skeletal meshes.

diff --git a/Engine/Plugins/Importers/USDImporter/Source/USDExporter/Private/SkeletalMeshExporterUSD.cpp b/Engine/Plugins/Importers/USDImporter/Source/USDExporter/Private/SkeletalMeshExporterUSD.cpp
--- a/Engine/Plugins/Importers/USDImporter/Source/USDExporter/Private/SkeletalMeshExporterUSD.cpp
+++ b/Engine/Plugins/Importers/USDImporter/Source/USDExporter/Private/SkeletalMeshExporterUSD.cpp
@@ -241,8 +241,6 @@ bool USkeletalMeshExporterUsd::ExportBinary( UObject* Object, const TCHAR* Type,
 
 	FString RootPrimPath = ( TEXT( "/" ) + UsdUtils::SanitizeUsdIdentifier( *SkeletalMesh->GetName() ) );
 
-	FScopedUsdAllocs Allocs;
-
 	UE::FUsdPrim RootPrim = UsdStage.DefinePrim( UE::FSdfPath( *RootPrimPath ), TEXT("SkelRoot") );
 	if ( !RootPrim )
 	{
diff --git a/Engine/Plugins/Importers/USDImporter/Source/USDStageEditor/Private/SUSDLayersTreeView.cpp b/Engine/Plugins/Importers/USDImporter/Source/USDStageEditor/Private/SUSDLayersTreeView.cpp
--- a/Engine/Plugins/Importers/USDImporter/Source/USDStageEditor/Private/SUSDLayersTreeView.cpp
+++ b/Engine/Plugins/Importers/USDImporter/Source/USDStageEditor/Private/SUSDLayersTreeView.cpp
@@ -1562,14 +1562,10 @@ void SUsdLayersTreeView::OnNewSubLayer()
 	}
 
 	TArray< FUsdLayerViewModelRef > MySelectedItems = GetSelectedItems();
-
+	for ( FUsdLayerViewModelRef SelectedItem : MySelectedItems )
 	{
-		FScopedUsdAllocs UsdAllocs;
-		for ( FUsdLayerViewModelRef SelectedItem : MySelectedItems )
-		{
-			SelectedItem->NewSubLayer( *SubLayerFile.GetValue() );
-			break;
-		}
+		SelectedItem->NewSubLayer( *SubLayerFile.GetValue() );
+		break;
 	}
 
 	RequestTreeRefresh();
@@ -1613,23 +1609,19 @@ void SUsdLayersTreeView::OnRemoveSelectedLayers()
 			continue;
 		}
 
+		int32 SubLayerIndex = 0;
+		for ( FUsdLayerViewModelRef Child : SelectedLayer->ParentItem->Children )
 		{
-			FScopedUsdAllocs UsdAllocs;
-
-			int32 SubLayerIndex = 0;
-			for ( FUsdLayerViewModelRef Child : SelectedLayer->ParentItem->Children )
+			if ( Child->LayerIdentifier == SelectedLayer->LayerIdentifier )
 			{
-				if ( Child->LayerIdentifier == SelectedLayer->LayerIdentifier )
+				if ( SelectedLayer->ParentItem )
 				{
-					if ( SelectedLayer->ParentItem )
-					{
-						bLayerRemoved = SelectedLayer->ParentItem->RemoveSubLayer( SubLayerIndex );
-					}
-					break;
+					bLayerRemoved = SelectedLayer->ParentItem->RemoveSubLayer( SubLayerIndex );
 				}
-
-				++SubLayerIndex;
+				break;
 			}
+
+			++SubLayerIndex;
 		}
 	}
 
diff --git a/Engine/Plugins/Importers/USDImporter/Source/USDStageEditor/Private/Widgets/SUSDIntegrationsPanel.cpp b/Engine/Plugins/Importers/USDImporter/Source/USDStageEditor/Private/Widgets/SUSDIntegrationsPanel.cpp
--- a/Engine/Plugins/Importers/USDImporter/Source/USDStageEditor/Private/Widgets/SUSDIntegrationsPanel.cpp
+++ b/Engine/Plugins/Importers/USDImporter/Source/USDStageEditor/Private/Widgets/SUSDIntegrationsPanel.cpp
@@ -56,8 +56,6 @@ TSharedRef< SWidget > SUsdIntegrationsPanelRow::GenerateWidgetForColumn( const F
 	bool bIsLeftColumn = true;
 	FOptionalSize RowHeight = FUsdStageEditorStyle::Get()->GetFloat( "UsdStageEditor.ListItemHeight" );
 
-	FScopedUsdAllocs Allocs;
-
 	FName AttributeName = Attribute->GetName();
 
 	using DisplayTextForPropertiesEntry = TPairInitializer<const FName&, const FText&>;
