USD: Fix additional save dialog showing up during the process of saving a new stage.

Fix save dialogs being shown during undo/redo operations.

diff --git a/Engine/Plugins/Importers/USDImporter/Source/USDStageEditor/Private/USDStageEditorModule.cpp b/Engine/Plugins/Importers/USDImporter/Source/USDStageEditor/Private/USDStageEditorModule.cpp
--- a/Engine/Plugins/Importers/USDImporter/Source/USDStageEditor/Private/USDStageEditorModule.cpp
+++ b/Engine/Plugins/Importers/USDImporter/Source/USDStageEditor/Private/USDStageEditorModule.cpp
@@ -11,6 +11,8 @@
 #include "USDProjectSettings.h"
 #include "USDStageActor.h"
 
+#include "Editor/TransBuffer.h"
+#include "Editor/UnrealEdEngine.h"
 #include "EngineUtils.h"
 #include "Framework/Docking/TabManager.h"
 #include "IAssetTools.h"
@@ -18,6 +20,7 @@
 #include "LevelEditor.h"
 #include "Modules/ModuleManager.h"
 #include "Styling/AppStyle.h"
+#include "UnrealEdGlobals.h"
 #include "UObject/ObjectSaveContext.h"
 #include "Widgets/Docking/SDockTab.h"
 #include "WorkspaceMenuStructure.h"
@@ -108,7 +111,10 @@ namespace UE::UsdStageEditorModule::Private
 
 			for ( const UE::FSdfLayer& UsedLayer : UsedLayers )
 			{
-				if ( UsedLayer.IsDirty() )
+				// This comment is written to the layer when we're in the process of saving a memory-only
+				// stage, and indicates that this layer is already saved (even though it will show as dirty and
+				// anonymous)
+				if ( UsedLayer.IsDirty() && UsedLayer.GetComment() != UnrealIdentifiers::LayerSavedComment )
 				{
 					FUsdSaveDialogRowData& RowData = RowsByIdentifier.FindOrAdd( UsedLayer.GetIdentifier() );
 					RowData.Layer = UsedLayer;
@@ -283,16 +289,38 @@ public:
 		}));
 
 		// Prompt to save modified USD Layers when closing stageactor stages
-		StageActorLoadedHandle = AUsdStageActor::OnActorLoaded.AddLambda([](AUsdStageActor* StageActor)
+		StageActorLoadedHandle = AUsdStageActor::OnActorLoaded.AddLambda([this](AUsdStageActor* StageActor)
 		{
 			if ( !StageActor )
 			{
 				return;
 			}
 
-			StageActor->OnPreStageChanged.AddLambda( [StageActor]()
+			// We never want to prompt when undoing or redoing.
+			// We have to subscribe to this here as the UTransBuffer doesn't exist by the time the module is
+			// initializing
+			if ( UTransBuffer* TransBuffer = GUnrealEd ? Cast<UTransBuffer>( GUnrealEd->Trans ) : nullptr )
 			{
-				if ( StageActor && static_cast< const AUsdStageActor* >( StageActor )->GetUsdStage() )
+				if ( !OnTransactionStateChangedHandle.IsValid() )
+				{
+					OnTransactionStateChangedHandle = TransBuffer->OnTransactionStateChanged().AddLambda(
+						[this]( const FTransactionContext& InTransactionContext, const ETransactionStateEventType InTransactionState )
+						{
+							if ( InTransactionState == ETransactionStateEventType::UndoRedoStarted )
+							{
+								bUndoRedoing = true;
+							}
+							else if ( InTransactionState == ETransactionStateEventType::UndoRedoFinalized )
+							{
+								bUndoRedoing = false;
+							}
+						});
+				}
+			}
+
+			StageActor->OnPreStageChanged.AddLambda( [this, StageActor]()
+			{
+				if ( !bUndoRedoing && StageActor && static_cast< const AUsdStageActor* >( StageActor )->GetUsdStage() )
 				{
 					const bool bForClosing = true;
 					UE::UsdStageEditorModule::Private::SaveStageActorLayersForWorld( StageActor->GetWorld(), bForClosing, StageActor );
@@ -333,10 +361,13 @@ public:
 	}
 
 private:
+	bool bUndoRedoing = false;
+
 	FDelegateHandle LevelEditorTabManagerChangedHandle;
 	FDelegateHandle PreSaveWorldEditorDelegateHandle;
 	FDelegateHandle EditorCanCloseDelegate;
 	FDelegateHandle StageActorLoadedHandle;
+	FDelegateHandle OnTransactionStateChangedHandle;
 #endif // #if USE_USD_SDK
 };
 
diff --git a/Engine/Plugins/Importers/USDImporter/Source/USDStageEditorViewModels/Private/USDStageViewModel.cpp b/Engine/Plugins/Importers/USDImporter/Source/USDStageEditorViewModels/Private/USDStageViewModel.cpp
--- a/Engine/Plugins/Importers/USDImporter/Source/USDStageEditorViewModels/Private/USDStageViewModel.cpp
+++ b/Engine/Plugins/Importers/USDImporter/Source/USDStageEditorViewModels/Private/USDStageViewModel.cpp
@@ -220,7 +220,8 @@ void FUsdStageViewModel::SaveStageAs( const TCHAR* FilePath )
 		FText::FromString( FilePath )
 	) );
 
-	if ( UsdStageActor.IsValid() )
+	AUsdStageActor* StageActorPtr = UsdStageActor.Get();
+	if ( StageActorPtr )
 	{
 		if ( UE::FUsdStage UsdStage = UsdStageActor->GetOrLoadUsdStage() )
 		{
@@ -232,8 +233,26 @@ void FUsdStageViewModel::SaveStageAs( const TCHAR* FilePath )
 				{
 					FScopedUnrealAllocs UEAllocs;
 
+					// In the process of opening FilePath below we'll close our previous stage, which is an anonymous
+					// layer and marked dirty by default. Even though we just saved (exported) it to disk, we'd end
+					// up getting the "do you want to save these dirty USD layers?" dialog by default...
+					// Here we'll write out a comment on the layer that we can easily check for in
+					// USDStageEditorModule::SaveStageActorLayersForWorld to know to skip showing that dialog. Note
+					// how this comment iself doesn't actually get saved to disk though.
+					//
+					// We also block listening here because we don't want to record to the transactor that we wrote
+					// this comment, because we can't record restoring it to what it was either, given that we'll stop
+					// listening to the previous stage after we open the next one. If we just recorded writing the
+					// comment, undoing/redoing through this operation would have left the temp stage with the comment
+					// in it permanently.
+					FScopedBlockNoticeListening BlockListening( StageActorPtr );
+					FString OldComment = RootLayer.GetComment();
+					RootLayer.SetComment( UnrealIdentifiers::LayerSavedComment );
+
 					OpenStage( FilePath );
 
+					RootLayer.SetComment( *OldComment );
+
 					UsdViewModelImpl::SaveUEStateLayer( UsdStage );
 				}
 			}
diff --git a/Engine/Plugins/Importers/USDImporter/Source/UnrealUSDWrapper/Private/UnrealUSDWrapper.cpp b/Engine/Plugins/Importers/USDImporter/Source/UnrealUSDWrapper/Private/UnrealUSDWrapper.cpp
--- a/Engine/Plugins/Importers/USDImporter/Source/UnrealUSDWrapper/Private/UnrealUSDWrapper.cpp
+++ b/Engine/Plugins/Importers/USDImporter/Source/UnrealUSDWrapper/Private/UnrealUSDWrapper.cpp
@@ -730,6 +730,8 @@ EUsdGeomOrientation IUsdPrim::GetGeometryOrientation(const pxr::UsdGeomMesh& Mes
 }
 #endif // USE_USD_SDK
 
+const TCHAR* UnrealIdentifiers::LayerSavedComment = TEXT("unreal:layerSaved");
+
 const TCHAR* UnrealIdentifiers::Invisible = TEXT("invisible");
 const TCHAR* UnrealIdentifiers::Inherited = TEXT("inherited");
 const TCHAR* UnrealIdentifiers::IdentifierPrefix = TEXT("@identifier:");
diff --git a/Engine/Plugins/Importers/USDImporter/Source/UnrealUSDWrapper/Private/UsdWrappers/SdfLayer.cpp b/Engine/Plugins/Importers/USDImporter/Source/UnrealUSDWrapper/Private/UsdWrappers/SdfLayer.cpp
--- a/Engine/Plugins/Importers/USDImporter/Source/UnrealUSDWrapper/Private/UsdWrappers/SdfLayer.cpp
+++ b/Engine/Plugins/Importers/USDImporter/Source/UnrealUSDWrapper/Private/UsdWrappers/SdfLayer.cpp
@@ -229,6 +229,29 @@ namespace UE
 	}
 #endif // #if USE_USD_SDK
 
+	template<typename PtrType>
+	FString FSdfLayerBase<PtrType>::GetComment() const
+	{
+#if USE_USD_SDK
+		if ( const PtrType& Ptr = Impl->GetInner() )
+		{
+			return ANSI_TO_TCHAR( Ptr->GetComment().c_str() );
+		}
+#endif // #if USE_USD_SDK
+		return {};
+	}
+
+	template<typename PtrType>
+	void FSdfLayerBase<PtrType>::SetComment( const TCHAR* Comment ) const
+	{
+#if USE_USD_SDK
+		if ( const PtrType& Ptr = Impl->GetInner() )
+		{
+			Ptr->SetComment( TCHAR_TO_ANSI( Comment ) );
+		}
+#endif // #if USE_USD_SDK
+	}
+
 	template<typename PtrType>
 	void FSdfLayerBase<PtrType>::TransferContent( const FSdfLayer& SourceLayer )
 	{
diff --git a/Engine/Plugins/Importers/USDImporter/Source/UnrealUSDWrapper/Public/UnrealUSDWrapper.h b/Engine/Plugins/Importers/USDImporter/Source/UnrealUSDWrapper/Public/UnrealUSDWrapper.h
--- a/Engine/Plugins/Importers/USDImporter/Source/UnrealUSDWrapper/Public/UnrealUSDWrapper.h
+++ b/Engine/Plugins/Importers/USDImporter/Source/UnrealUSDWrapper/Public/UnrealUSDWrapper.h
@@ -454,6 +454,8 @@ namespace UnrealIdentifiers
 	extern UNREALUSDWRAPPER_API const pxr::TfToken UnrealEngineVersion;
 #endif // #if USE_USD_SDK
 
+	extern UNREALUSDWRAPPER_API const TCHAR* LayerSavedComment;
+
 	extern UNREALUSDWRAPPER_API const TCHAR* Invisible;
 	extern UNREALUSDWRAPPER_API const TCHAR* Inherited;
 	extern UNREALUSDWRAPPER_API const TCHAR* IdentifierPrefix;
diff --git a/Engine/Plugins/Importers/USDImporter/Source/UnrealUSDWrapper/Public/UsdWrappers/SdfLayer.h b/Engine/Plugins/Importers/USDImporter/Source/UnrealUSDWrapper/Public/UsdWrappers/SdfLayer.h
--- a/Engine/Plugins/Importers/USDImporter/Source/UnrealUSDWrapper/Public/UsdWrappers/SdfLayer.h
+++ b/Engine/Plugins/Importers/USDImporter/Source/UnrealUSDWrapper/Public/UsdWrappers/SdfLayer.h
@@ -81,6 +81,9 @@ namespace UE
 
 	// Wrapped pxr::SdfLayer functions, refer to the USD SDK documentation
 	public:
+		FString GetComment() const;
+		void SetComment( const TCHAR* Comment ) const;
+
 		void TransferContent( const FSdfLayer& SourceLayer );
 
 		static FSdfLayer FindOrOpen( const TCHAR* Identifier );
