USD: Automatically upgrade textures and base materials to VT when required.

diff --git a/Engine/Plugins/Importers/USDImporter/Source/USDClasses/Private/USDAssetCache.cpp b/Engine/Plugins/Importers/USDImporter/Source/USDClasses/Private/USDAssetCache.cpp
--- a/Engine/Plugins/Importers/USDImporter/Source/USDClasses/Private/USDAssetCache.cpp
+++ b/Engine/Plugins/Importers/USDImporter/Source/USDClasses/Private/USDAssetCache.cpp
@@ -183,6 +183,28 @@ FString UUsdAssetCache::GetPrimForAsset( UObject* Asset ) const
 	return PrimForAsset;
 }
 
+FString UUsdAssetCache::GetHashForAsset( UObject* Asset ) const
+{
+	if ( !Asset || !OwnedAssets.Contains( Asset ) )
+	{
+		return {};
+	}
+
+	FScopeLock Lock( &CriticalSection );
+
+	if ( const FString* KeyPtr = TransientStorage.FindKey( Asset ) )
+	{
+		return *KeyPtr;
+	}
+
+	if ( const FString* KeyPtr = PersistentStorage.FindKey( Asset ) )
+	{
+		return *KeyPtr;
+	}
+
+	return {};
+}
+
 void UUsdAssetCache::Reset()
 {
 	FScopeLock Lock( &CriticalSection );
diff --git a/Engine/Plugins/Importers/USDImporter/Source/USDClasses/Public/USDAssetCache.h b/Engine/Plugins/Importers/USDImporter/Source/USDClasses/Public/USDAssetCache.h
--- a/Engine/Plugins/Importers/USDImporter/Source/USDClasses/Public/USDAssetCache.h
+++ b/Engine/Plugins/Importers/USDImporter/Source/USDClasses/Public/USDAssetCache.h
@@ -38,6 +38,12 @@ public:
 	/** Returns the first prim path associated with an asset (see TMap::FindKey) */
 	FString GetPrimForAsset( UObject* Asset ) const;
 
+	/**
+	 * Returns the hash used as a key for an asset, in case we own it.
+	 * Returns the empty string otherwise.
+	 */
+	FString GetHashForAsset( UObject* Asset ) const;
+
 	TMap< FString, TWeakObjectPtr<UObject> > GetAssetPrimLinks() const { return PrimPathToAssets; }; // Can't return a reference as it wouldn't be thread-safe
 
 	bool IsAssetOwnedByCache( UObject* Asset ) const { return OwnedAssets.Contains( Asset ); }
diff --git a/Engine/Plugins/Importers/USDImporter/Source/USDSchemas/Private/USDShadeMaterialTranslator.cpp b/Engine/Plugins/Importers/USDImporter/Source/USDSchemas/Private/USDShadeMaterialTranslator.cpp
--- a/Engine/Plugins/Importers/USDImporter/Source/USDSchemas/Private/USDShadeMaterialTranslator.cpp
+++ b/Engine/Plugins/Importers/USDImporter/Source/USDSchemas/Private/USDShadeMaterialTranslator.cpp
@@ -4,6 +4,7 @@
 
 #include "USDAssetCache.h"
 #include "USDAssetImportData.h"
+#include "USDLog.h"
 #include "USDShadeConversion.h"
 #include "USDTypesConversion.h"
 
@@ -23,6 +24,249 @@
 	#include "pxr/usd/usdShade/tokens.h"
 #include "USDIncludesEnd.h"
 
+namespace UE::UsdShadeTranslator::Private
+{
+	const static FString BaseMaterialPathTranslucentVT = TEXT( "/USDImporter/Materials/UsdPreviewSurfaceTranslucentVT.UsdPreviewSurfaceTranslucentVT" );
+	const static FString BaseMaterialPathTranslucent = TEXT( "/USDImporter/Materials/UsdPreviewSurfaceTranslucent.UsdPreviewSurfaceTranslucent" );
+	const static FString BaseMaterialPathVT = TEXT( "/USDImporter/Materials/UsdPreviewSurfaceVT.UsdPreviewSurfaceVT" );
+	const static FString BaseMaterialPath = TEXT( "/USDImporter/Materials/UsdPreviewSurface.UsdPreviewSurface" );
+
+	UMaterialInterface* GetBaseMaterial( bool bIsTranslucent, bool bIsVT )
+	{
+		const FString* TargetMaterialPath = nullptr;
+		if ( bIsTranslucent )
+		{
+			if ( bIsVT )
+			{
+				TargetMaterialPath = &BaseMaterialPathTranslucentVT;
+			}
+			else
+			{
+				TargetMaterialPath = &BaseMaterialPathTranslucent;
+			}
+		}
+		else
+		{
+			if ( bIsVT )
+			{
+				TargetMaterialPath = &BaseMaterialPathVT;
+			}
+			else
+			{
+				TargetMaterialPath = &BaseMaterialPath;
+			}
+		}
+
+		if ( !TargetMaterialPath )
+		{
+			return nullptr;
+		}
+
+		return Cast< UMaterialInterface >( FSoftObjectPath( *TargetMaterialPath ).TryLoad() );
+	}
+
+	UMaterialInterface* GetVTVersionOfBaseMaterial( UMaterialInterface* BaseMaterial )
+	{
+		if ( !BaseMaterial )
+		{
+			return nullptr;
+		}
+
+		const FString PathName = BaseMaterial->GetPathName();
+		if ( PathName == BaseMaterialPathTranslucentVT || PathName == BaseMaterialPathVT )
+		{
+			return BaseMaterial;
+		}
+		else if ( PathName == BaseMaterialPath )
+		{
+			return Cast< UMaterialInterface >( FSoftObjectPath{ BaseMaterialPathVT }.TryLoad() );
+		}
+		else if ( PathName == BaseMaterialPathTranslucent )
+		{
+			return Cast< UMaterialInterface >( FSoftObjectPath{ BaseMaterialPathTranslucentVT }.TryLoad() );
+		}
+
+		// We should only ever call this function with a BaseMaterial that matches one of the above paths
+		ensure( false );
+		return nullptr;
+	}
+
+	void RecursiveUpgradeMaterialsAndTexturesToVT(
+		const TSet<UTexture*>& TexturesToUpgrade,
+		const TSharedRef< FUsdSchemaTranslationContext >& Context,
+		TSet<UMaterialInterface*>& VisitedMaterials,
+		TSet<UMaterialInterface*>& NewMaterials
+	)
+	{
+		for ( UTexture* Texture : TexturesToUpgrade )
+		{
+			if ( Texture->VirtualTextureStreaming )
+			{
+				continue;
+			}
+
+			UE_LOG( LogUsd, Log, TEXT( "Upgrading texture '%s' to VT as it is used by a material that must be VT" ),
+				*Texture->GetName()
+			);
+			FPropertyChangedEvent PropertyChangeEvent( UTexture::StaticClass()->FindPropertyByName( GET_MEMBER_NAME_CHECKED( UTexture, VirtualTextureStreaming ) ) );
+			Texture->Modify();
+			Texture->VirtualTextureStreaming = true;
+
+#if WITH_EDITOR
+			Texture->PostEditChangeProperty( PropertyChangeEvent );
+#endif // WITH_EDITOR
+
+			// Now that our texture is VT, all materials that use the texture must be VT
+			if ( const TSet<UMaterialInterface*>* UserMaterials = Context->TextureToUserMaterials.Find( Texture ) )
+			{
+				for ( UMaterialInterface* UserMaterial : *UserMaterials )
+				{
+					if ( VisitedMaterials.Contains( UserMaterial ) )
+					{
+						continue;
+					}
+
+					if ( UMaterialInstance* MaterialInstance = Cast<UMaterialInstance>( UserMaterial ) )
+					{
+						// Important to not use GetBaseMaterial() here because if our parent is the translucent we'll
+						// get the base UsdPreviewSurface instead, as that is also *its* base
+						UMaterialInterface* BaseMaterial = MaterialInstance->Parent.Get();
+						UMaterialInterface* BaseMaterialVT = GetVTVersionOfBaseMaterial( BaseMaterial );
+						if ( BaseMaterial == BaseMaterialVT )
+						{
+							// Material is already VT, we're good
+							continue;
+						}
+
+						// Visit it before we start recursing. We need this because we must convert textures to VT
+						// before materials (or else we get a warning) but we'll only actually swap the base material
+						// at the end of this scope
+						VisitedMaterials.Add( UserMaterial );
+
+						// If we're going to update this material to VT, all of *its* textures need to be VT too
+						TSet<UTexture*> OtherUsedTextures;
+						for ( const FTextureParameterValue& TextureValue : MaterialInstance->TextureParameterValues )
+						{
+							if ( UTexture* OtherTexture = TextureValue.ParameterValue )
+							{
+								OtherUsedTextures.Add( OtherTexture );
+							}
+						}
+
+						RecursiveUpgradeMaterialsAndTexturesToVT(
+							OtherUsedTextures,
+							Context,
+							VisitedMaterials,
+							NewMaterials
+						);
+
+						// We can't blindly recreate all component render states when a level is being added, because
+						// we may end up first creating render states for some components, and UWorld::AddToWorld
+						// calls FScene::AddPrimitive which expects the component to not have primitives yet
+						FMaterialUpdateContext::EOptions::Type Options = FMaterialUpdateContext::EOptions::Default;
+						if ( Context->Level->bIsAssociatingLevel )
+						{
+							Options = ( FMaterialUpdateContext::EOptions::Type ) (
+								Options & ~FMaterialUpdateContext::EOptions::RecreateRenderStates
+							);
+						}
+
+						UE_LOG( LogUsd, Log, TEXT( "Upgrading material instance '%s' to being based on VT as texture '%s' requires it" ),
+							*MaterialInstance->GetName(),
+							*Texture->GetName()
+						);
+
+#if WITH_EDITOR
+						if ( UMaterialInstanceConstant* MIC = Cast<UMaterialInstanceConstant>( MaterialInstance ) )
+						{
+							FMaterialUpdateContext UpdateContext( Options, GMaxRHIShaderPlatform );
+							UpdateContext.AddMaterialInstance( MIC );
+							MIC->PreEditChange( nullptr );
+							MIC->SetParentEditorOnly( BaseMaterialVT );
+							MIC->PostEditChange();
+						}
+						else
+#endif // WITH_EDITOR
+						if ( UMaterialInstanceDynamic* MID = Cast<UMaterialInstanceDynamic>( UserMaterial ) )
+						{
+							if ( Context->AssetCache )
+							{
+								// This is super slow as it will essentially linear search the asset cache twice, but
+								// not much we can do at runtime without asset import data
+								const FString PrimPath = Context->AssetCache->GetPrimForAsset( MID );
+								const FString Hash = Context->AssetCache->GetHashForAsset( MID );
+
+								const FName NewInstanceName = MakeUniqueObjectName(
+									GetTransientPackage(),
+									UMaterialInstance::StaticClass(),
+									*FPaths::GetBaseFilename( PrimPath )
+								);
+
+								// For MID we can't swap the base material, so we need to create a brand new one and copy
+								// the overrides
+								UMaterialInstanceDynamic* NewMID = UMaterialInstanceDynamic::Create(
+									BaseMaterial,
+									GetTransientPackage(),
+									NewInstanceName
+								);
+								if ( !ensure( NewMID ) )
+								{
+									continue;
+								}
+								NewMID->CopyParameterOverrides( MID );
+
+								Context->AssetCache->DiscardAsset( Hash );
+								Context->AssetCache->CacheAsset( Hash, MID, PrimPath );
+								NewMaterials.Add( NewMID );
+							}
+						}
+						else
+						{
+							// This should never happen
+							ensure( false );
+						}
+					}
+				}
+			}
+		}
+	}
+
+	void UpgradeMaterialsAndTexturesToVT( TSet<UTexture*> TexturesToUpgrade, TSharedRef< FUsdSchemaTranslationContext >& Context )
+	{
+		TSet<UMaterialInterface*> VisitedMaterials;
+		TSet<UMaterialInterface*> NewMaterials;
+		RecursiveUpgradeMaterialsAndTexturesToVT( TexturesToUpgrade, Context, VisitedMaterials, NewMaterials );
+
+		// When we "visit" a MID we'll create a brand new instance of it and discard the old one, so let's drop the old ones
+		// from Context->TextureToUserMaterials too
+		for ( UMaterialInterface* Material : VisitedMaterials )
+		{
+			if ( UMaterialInstanceDynamic* OldMID = Cast<UMaterialInstanceDynamic>( Material ) )
+			{
+				for ( TPair<UTexture*, TSet<UMaterialInterface*>>& Pair : Context->TextureToUserMaterials )
+				{
+					Pair.Value.Remove( Material );
+				}
+			}
+		}
+
+		// Additionally now we need to add those new MIDs we created back into Context->TextureToUserMaterials
+		for ( UMaterialInterface* Material : NewMaterials )
+		{
+			UMaterialInstanceDynamic* MID = Cast<UMaterialInstanceDynamic>( Material );
+			if ( ensure( MID ) )
+			{
+				for ( const FTextureParameterValue& TextureValue : MID->TextureParameterValues )
+				{
+					if ( UTexture* Texture = TextureValue.ParameterValue )
+					{
+						Context->TextureToUserMaterials.FindOrAdd( Texture ).Add( MID );
+					}
+				}
+			}
+		}
+	}
+}
 
 void FUsdShadeMaterialTranslator::CreateAssets()
 {
@@ -51,68 +295,158 @@ void FUsdShadeMaterialTranslator::CreateAssets()
 	if ( !ConvertedMaterial )
 	{
 		const bool bIsTranslucent = UsdUtils::IsMaterialTranslucent( ShadeMaterial );
-		const bool bNeedsVirtualTextures = UsdUtils::IsMaterialUsingUDIMs( ShadeMaterial );
+		const FName InstanceName = MakeUniqueObjectName(
+			GetTransientPackage(),
+			UMaterialInstance::StaticClass(),
+			*FPaths::GetBaseFilename( PrimPath.GetString() )
+		);
 
-		FString BaseMaterialName = TEXT("UsdPreviewSurface");
+		TMap<FString, int32> Unused;
+		TMap<FString, int32>& PrimvarToUVIndex = Context->MaterialToPrimvarToUVIndex
+			? Context->MaterialToPrimvarToUVIndex->FindOrAdd( PrimPath.GetString() )
+			: Unused;
 
-		if ( bIsTranslucent )
+#if WITH_EDITOR
+		if ( GIsEditor ) // Also have to prevent Standalone game from going with MaterialInstanceConstants
 		{
-			BaseMaterialName += TEXT("Translucent");
-		}
+			if ( UMaterialInstanceConstant* NewMaterial = NewObject<UMaterialInstanceConstant>( GetTransientPackage(), InstanceName, Context->ObjectFlags ) )
+			{
+				UUsdAssetImportData* ImportData = NewObject< UUsdAssetImportData >( NewMaterial, TEXT( "USDAssetImportData" ) );
+				ImportData->PrimPath = PrimPath.GetString();
+				NewMaterial->AssetImportData = ImportData;
 
-		if ( bNeedsVirtualTextures )
-		{
-			BaseMaterialName += TEXT("VT");
-		}
+				const bool bSuccess = UsdToUnreal::ConvertMaterial(
+					ShadeMaterial,
+					*NewMaterial,
+					Context->AssetCache.Get(),
+					PrimvarToUVIndex,
+					*Context->RenderContext.ToString()
+				);
+				if ( !bSuccess )
+				{
+					return;
+				}
+
+				TSet<UTexture*> VTTextures;
+				TSet<UTexture*> NonVTTextures;
+				for ( const FTextureParameterValue& TextureValue : NewMaterial->TextureParameterValues )
+				{
+					if ( UTexture* Texture = TextureValue.ParameterValue )
+					{
+						if ( Texture->VirtualTextureStreaming )
+						{
+							UsdUtils::NotifyIfVirtualTexturesNeeded( Texture );
+							VTTextures.Add( Texture );
+						}
+						else
+						{
+							NonVTTextures.Add( Texture );
+						}
+					}
+				}
 
-		const FString BaseMaterialPath = FString::Printf( TEXT("Material'/USDImporter/Materials/%s.%s'"), *BaseMaterialName, *BaseMaterialName );
+				// Our VT material only has VT texture samplers, so *all* of its textures must be VT
+				if ( VTTextures.Num() && NonVTTextures.Num() )
+				{
+					UE_LOG( LogUsd, Log, TEXT( "Upgrading textures used by material instance '%s' to VT as the material must be VT" ),
+						*NewMaterial->GetName()
+					);
+					UE::UsdShadeTranslator::Private::UpgradeMaterialsAndTexturesToVT( NonVTTextures, Context );
+				}
 
-		if ( UMaterialInterface* BaseMaterial = Cast< UMaterialInterface >( FSoftObjectPath( BaseMaterialPath ).TryLoad() ) )
-		{
-			FName InstanceName = MakeUniqueObjectName( GetTransientPackage(), UMaterialInstance::StaticClass(), *FPaths::GetBaseFilename( PrimPath.GetString() ) );
-			if ( GIsEditor ) // Also have to prevent Standalone game from going with MaterialInstanceConstants
-			{
-#if WITH_EDITOR
-				if ( UMaterialInstanceConstant* NewMaterial = NewObject<UMaterialInstanceConstant>( GetTransientPackage(), InstanceName, Context->ObjectFlags ) )
+				const bool bNeedsVT = VTTextures.Num() > 0;
+				UMaterialInterface* BaseMaterial = UE::UsdShadeTranslator::Private::GetBaseMaterial(
+					bIsTranslucent,
+					bNeedsVT
+				);
+
+				if ( ensure( BaseMaterial ) )
 				{
 					NewMaterial->SetParentEditorOnly( BaseMaterial );
 
-					UUsdAssetImportData* ImportData = NewObject< UUsdAssetImportData >( NewMaterial, TEXT( "USDAssetImportData" ) );
-					ImportData->PrimPath = PrimPath.GetString();
-					NewMaterial->AssetImportData = ImportData;
+					// We can't blindly recreate all component render states when a level is being added, because we may end up first creating
+					// render states for some components, and UWorld::AddToWorld calls FScene::AddPrimitive which expects the component to not have
+					// primitives yet
+					FMaterialUpdateContext::EOptions::Type Options = FMaterialUpdateContext::EOptions::Default;
+					if ( Context->Level->bIsAssociatingLevel )
+					{
+						Options = ( FMaterialUpdateContext::EOptions::Type ) ( Options & ~FMaterialUpdateContext::EOptions::RecreateRenderStates );
+					}
 
-					TMap<FString, int32> Unused;
-					TMap<FString, int32>& PrimvarToUVIndex = Context->MaterialToPrimvarToUVIndex ? Context->MaterialToPrimvarToUVIndex->FindOrAdd( PrimPath.GetString() ) : Unused;
+					FMaterialUpdateContext UpdateContext( Options, GMaxRHIShaderPlatform );
+					UpdateContext.AddMaterialInstance( NewMaterial );
+					NewMaterial->PreEditChange( nullptr );
+					NewMaterial->PostEditChange();
+
+					ConvertedMaterial = NewMaterial;
+
+					for ( UTexture* Texture : VTTextures.Union( NonVTTextures ) )
+					{
+						Context->TextureToUserMaterials.FindOrAdd( Texture ).Add( NewMaterial );
+					}
+				}
+			}
+		}
+		else
+#endif // WITH_EDITOR
+		{
+			// At runtime we always start with a non-VT base and if we discover we need one we just create a new MID
+			// using the VT base and copy the overrides. Not much else we can do as we need a base to call
+			// UMaterialInstanceDynamic::Create and get our instance, but we an instance to call UsdToUnreal::ConvertMaterial
+			// to create our textures and decide on our base.
+			const bool bNeedsVT = false;
+			UMaterialInterface* BaseMaterial = UE::UsdShadeTranslator::Private::GetBaseMaterial(
+				bIsTranslucent,
+				bNeedsVT
+			);
 
+			if ( ensure( BaseMaterial ) )
+			{
+				if ( UMaterialInstanceDynamic* NewMaterial = UMaterialInstanceDynamic::Create( BaseMaterial, GetTransientPackage(), InstanceName ) )
+				{
 					if ( UsdToUnreal::ConvertMaterial( ShadeMaterial, *NewMaterial, Context->AssetCache.Get(), PrimvarToUVIndex, *Context->RenderContext.ToString() ) )
 					{
-						// We can't blindly recreate all component render states when a level is being added, because we may end up first creating
-						// render states for some components, and UWorld::AddToWorld calls FScene::AddPrimitive which expects the component to not have
-						// primitives yet
-						FMaterialUpdateContext::EOptions::Type Options = FMaterialUpdateContext::EOptions::Default;
-						if ( Context->Level->bIsAssociatingLevel )
+						TSet<UTexture*> VTTextures;
+						TSet<UTexture*> NonVTTextures;
+						for ( const FTextureParameterValue& TextureValue : NewMaterial->TextureParameterValues )
 						{
-							Options = ( FMaterialUpdateContext::EOptions::Type ) ( Options & ~FMaterialUpdateContext::EOptions::RecreateRenderStates );
+							if ( UTexture* Texture = TextureValue.ParameterValue )
+							{
+								if ( Texture->VirtualTextureStreaming )
+								{
+									VTTextures.Add( Texture );
+								}
+								else
+								{
+									NonVTTextures.Add( Texture );
+								}
+							}
 						}
 
-						FMaterialUpdateContext UpdateContext( Options, GMaxRHIShaderPlatform );
-						UpdateContext.AddMaterialInstance( NewMaterial );
-						NewMaterial->PreEditChange( nullptr );
-						NewMaterial->PostEditChange();
+						// We must stash our material and textures *before* we call UpgradeMaterialsAndTexturesToVT, as that
+						// is what will actually swap our base with a VT one if needed
+						if ( Context->AssetCache )
+						{
+							Context->AssetCache->CacheAsset( MaterialHashString, NewMaterial, PrimPath.GetString() );
+						}
+						for ( UTexture* Texture : VTTextures.Union( NonVTTextures ) )
+						{
+							Context->TextureToUserMaterials.FindOrAdd( Texture ).Add( NewMaterial );
+						}
 
-						ConvertedMaterial = NewMaterial;
-					}
-				}
-#endif // WITH_EDITOR
-			}
-			else if ( UMaterialInstanceDynamic* NewMaterial = UMaterialInstanceDynamic::Create( BaseMaterial, GetTransientPackage(), InstanceName ) )
-			{
-				TMap<FString, int32> Unused;
-				TMap<FString, int32>& PrimvarToUVIndex = Context->MaterialToPrimvarToUVIndex ? Context->MaterialToPrimvarToUVIndex->FindOrAdd( PrimPath.GetString() ) : Unused;
+						// Our VT material only has VT texture samplers, so *all* of its textures must be VT
+						if ( VTTextures.Num() && NonVTTextures.Num() )
+						{
+							UE_LOG( LogUsd, Log, TEXT( "Upgrading textures used by material instance '%s' to VT as the material must be VT" ),
+								*NewMaterial->GetName()
+							);
+							UE::UsdShadeTranslator::Private::UpgradeMaterialsAndTexturesToVT( NonVTTextures, Context );
+						}
 
-				if ( UsdToUnreal::ConvertMaterial( ShadeMaterial, *NewMaterial, Context->AssetCache.Get(), PrimvarToUVIndex, *Context->RenderContext.ToString() ) )
-				{
-					ConvertedMaterial = NewMaterial;
+						// We must go through the cache to fetch our result material here as UpgradeMaterialsAndTexturesToVT
+						// may have created a new MID for this material with a VT base
+						ConvertedMaterial = Cast<UMaterialInterface>( Context->AssetCache->GetCachedAsset( MaterialHashString ) );
+					}
 				}
 			}
 		}
diff --git a/Engine/Plugins/Importers/USDImporter/Source/USDSchemas/Public/USDSchemaTranslator.h b/Engine/Plugins/Importers/USDImporter/Source/USDSchemas/Public/USDSchemaTranslator.h
--- a/Engine/Plugins/Importers/USDImporter/Source/USDSchemas/Public/USDSchemaTranslator.h
+++ b/Engine/Plugins/Importers/USDImporter/Source/USDSchemas/Public/USDSchemaTranslator.h
@@ -176,6 +176,15 @@ struct USDSCHEMAS_API FUsdSchemaTranslationContext : public TSharedFromThis< FUs
 	 */
 	TMap< FString, TMap< FString, int32 > >* MaterialToPrimvarToUVIndex = nullptr;
 
+	/**
+	 * Sometimes we must upgrade a material from non-VT to VT, and so upgrade all of its textures to VT (and then
+	 * upgrade all materials that use them to VT, etc.).
+	 * This member lets us cache which generated materials use which generated textures in order to help with that.
+	 * Material parsing is synchronous. If we ever upgrade it to paralllel/async-task-based, we'll need a mutex around
+	 * this member.
+	 */
+	TMap<UTexture*, TSet<UMaterialInterface*>> TextureToUserMaterials;
+
 	/**
 	 * Whether to try to combine individual assets and components of the same type on a kind-per-kind basis,
 	 * like multiple Mesh prims into a single Static Mesh
diff --git a/Engine/Plugins/Importers/USDImporter/Source/USDUtilities/Private/USDShadeConversion.cpp b/Engine/Plugins/Importers/USDImporter/Source/USDUtilities/Private/USDShadeConversion.cpp
--- a/Engine/Plugins/Importers/USDImporter/Source/USDUtilities/Private/USDShadeConversion.cpp
+++ b/Engine/Plugins/Importers/USDImporter/Source/USDUtilities/Private/USDShadeConversion.cpp
@@ -473,7 +473,7 @@ namespace UE
 			using FParameterValue = TVariant<float, FVector, FTextureParameterValue, FPrimvarReaderParameterValue, bool>;
 
 			bool GetTextureParameterValue( pxr::UsdShadeInput& ShadeInput, TextureGroup LODGroup, FParameterValue& OutValue,
-				UMaterialInterface* Material = nullptr, UUsdAssetCache* TexturesCache = nullptr, TMap<FString, int32>* PrimvarToUVIndex = nullptr, bool bForceVirtualTextures = false )
+				UMaterialInterface* Material = nullptr, UUsdAssetCache* TexturesCache = nullptr, TMap<FString, int32>* PrimvarToUVIndex = nullptr )
 			{
 				FScopedUsdAllocs UsdAllocs;
 
@@ -587,13 +587,6 @@ namespace UE
 								{
 									Texture->SRGB = bSRGB;
 									Texture->CompressionSettings = CompressionSettings;
-
-									if ( bForceVirtualTextures )
-									{
-										Texture->VirtualTextureStreaming = true;
-										UsdUtils::NotifyIfVirtualTexturesNeeded( Texture );
-									}
-
 									Texture->UpdateResource();
 
 									TexturesCache->CacheAsset( TextureHash, Texture );
@@ -673,7 +666,7 @@ namespace UE
 			}
 
 			bool GetFloatParameterValue( pxr::UsdShadeConnectableAPI& Connectable, const pxr::TfToken& InputName, float DefaultValue, FParameterValue& OutValue,
-				UMaterialInterface* Material = nullptr, UUsdAssetCache* TexturesCache = nullptr, TMap<FString, int32>* PrimvarToUVIndex = nullptr, bool bForceVirtualTextures = false)
+				UMaterialInterface* Material = nullptr, UUsdAssetCache* TexturesCache = nullptr, TMap<FString, int32>* PrimvarToUVIndex = nullptr)
 			{
 				FScopedUsdAllocs Allocs;
 
@@ -689,7 +682,7 @@ namespace UE
 				pxr::UsdShadeAttributeType AttributeType;
 				if ( pxr::UsdShadeConnectableAPI::GetConnectedSource( Input.GetAttr(), &Source, &SourceName, &AttributeType ) )
 				{
-					if ( !GetTextureParameterValue( Input, TEXTUREGROUP_WorldSpecular, OutValue, Material, TexturesCache, PrimvarToUVIndex, bForceVirtualTextures ) )
+					if ( !GetTextureParameterValue( Input, TEXTUREGROUP_WorldSpecular, OutValue, Material, TexturesCache, PrimvarToUVIndex ) )
 					{
 						// Check if we have a fallback input that we can use instead, since we don't have a valid texture value
 						if ( const pxr::UsdShadeInput FallbackInput = Source.GetInput( UnrealIdentifiers::Fallback ) )
@@ -704,7 +697,7 @@ namespace UE
 
 						// Recurse because the attribute may just be pointing at some other attribute that has the data
 						// (e.g. when shader input is just "hoisted" and connected to the parent material input)
-						return GetFloatParameterValue( Source, SourceName, DefaultValue, OutValue, Material, TexturesCache, PrimvarToUVIndex, bForceVirtualTextures );
+						return GetFloatParameterValue( Source, SourceName, DefaultValue, OutValue, Material, TexturesCache, PrimvarToUVIndex );
 					}
 				}
 				// No other node connected, so we must have some value
@@ -792,7 +785,7 @@ namespace UE
 			}
 
 			bool GetVec3ParameterValue( pxr::UsdShadeConnectableAPI& Connectable, const pxr::TfToken& InputName, const FLinearColor& DefaultValue, FParameterValue& OutValue,
-				bool bIsNormalMap = false, UMaterialInterface* Material = nullptr, UUsdAssetCache* TexturesCache = nullptr, TMap<FString, int32>* PrimvarToUVIndex = nullptr, bool bForceVirtualTextures = false )
+				bool bIsNormalMap = false, UMaterialInterface* Material = nullptr, UUsdAssetCache* TexturesCache = nullptr, TMap<FString, int32>* PrimvarToUVIndex = nullptr )
 			{
 				FScopedUsdAllocs Allocs;
 
@@ -808,7 +801,7 @@ namespace UE
 				pxr::UsdShadeAttributeType AttributeType;
 				if ( pxr::UsdShadeConnectableAPI::GetConnectedSource( Input.GetAttr(), &Source, &SourceName, &AttributeType ) )
 				{
-					if ( !GetTextureParameterValue( Input, bIsNormalMap ? TEXTUREGROUP_WorldNormalMap : TEXTUREGROUP_World, OutValue, Material, TexturesCache, PrimvarToUVIndex, bForceVirtualTextures ) )
+					if ( !GetTextureParameterValue( Input, bIsNormalMap ? TEXTUREGROUP_WorldNormalMap : TEXTUREGROUP_World, OutValue, Material, TexturesCache, PrimvarToUVIndex ) )
 					{
 						// Check whether this input receives its value through a connection to a
 						// primvar reader shader.
@@ -843,7 +836,7 @@ namespace UE
 						}
 
 						// This shader doesn't have anything: Traverse into the input connectable itself
-						return GetVec3ParameterValue( Source, SourceName, DefaultValue, OutValue, bIsNormalMap, Material, TexturesCache, PrimvarToUVIndex, bForceVirtualTextures );
+						return GetVec3ParameterValue( Source, SourceName, DefaultValue, OutValue, bIsNormalMap, Material, TexturesCache, PrimvarToUVIndex );
 					}
 				}
 				// No other node connected, so we must have some value
@@ -1789,10 +1782,6 @@ bool UsdToUnreal::ConvertMaterial( const pxr::UsdShadeMaterial& UsdShadeMaterial
 
 	bool bHasMaterialInfo = false;
 
-	// We assume that all samplers for a material instance that supports UDIMs would be virtual so we'll force all the textures to be virtual.
-	// The idea being that when building a base material, we can't know which textures will be virtual so it's safer to set the samplers as if they were all virtual.
-	const bool bForceVirtualTextures = UsdUtils::IsMaterialUsingUDIMs( UsdShadeMaterial );
-
 	pxr::TfToken RenderContextToken = pxr::UsdShadeTokens->universalRenderContext;
 	if ( RenderContext )
 	{
@@ -1847,7 +1836,7 @@ bool UsdToUnreal::ConvertMaterial( const pxr::UsdShadeMaterial& UsdShadeMaterial
 	// Base color
 	{
 		const bool bIsNormalMap = false;
-		if ( UsdShadeConversionImpl::GetVec3ParameterValue( Connectable, UnrealIdentifiers::DiffuseColor, FLinearColor( 0, 0, 0 ), ParameterValue, bIsNormalMap, &Material, TexturesCache, &PrimvarToUVIndex, bForceVirtualTextures ) )
+		if ( UsdShadeConversionImpl::GetVec3ParameterValue( Connectable, UnrealIdentifiers::DiffuseColor, FLinearColor( 0, 0, 0 ), ParameterValue, bIsNormalMap, &Material, TexturesCache, &PrimvarToUVIndex ) )
 		{
 			UsdShadeConversionImpl::SetParameterValue( Material, TEXT( "BaseColor" ), ParameterValue, bForUsdPreviewSurface );
 			bHasMaterialInfo = true;
@@ -1857,7 +1846,7 @@ bool UsdToUnreal::ConvertMaterial( const pxr::UsdShadeMaterial& UsdShadeMaterial
 	// Emissive color
 	{
 		const bool bIsNormalMap = false;
-		if ( UsdShadeConversionImpl::GetVec3ParameterValue( Connectable, UnrealIdentifiers::EmissiveColor, FLinearColor( 0, 0, 0 ), ParameterValue, bIsNormalMap, &Material, TexturesCache, &PrimvarToUVIndex, bForceVirtualTextures ) )
+		if ( UsdShadeConversionImpl::GetVec3ParameterValue( Connectable, UnrealIdentifiers::EmissiveColor, FLinearColor( 0, 0, 0 ), ParameterValue, bIsNormalMap, &Material, TexturesCache, &PrimvarToUVIndex ) )
 		{
 			UsdShadeConversionImpl::SetParameterValue( Material, TEXT( "EmissiveColor" ), ParameterValue, bForUsdPreviewSurface );
 			bHasMaterialInfo = true;
@@ -1866,7 +1855,7 @@ bool UsdToUnreal::ConvertMaterial( const pxr::UsdShadeMaterial& UsdShadeMaterial
 
 	// Metallic
 	{
-		if ( UsdShadeConversionImpl::GetFloatParameterValue( Connectable, UnrealIdentifiers::Metallic, 0.f, ParameterValue, &Material, TexturesCache, &PrimvarToUVIndex, bForceVirtualTextures ) )
+		if ( UsdShadeConversionImpl::GetFloatParameterValue( Connectable, UnrealIdentifiers::Metallic, 0.f, ParameterValue, &Material, TexturesCache, &PrimvarToUVIndex ) )
 		{
 			UsdShadeConversionImpl::SetParameterValue( Material, TEXT( "Metallic" ), ParameterValue, bForUsdPreviewSurface );
 			SetTextureComponentParamForScalarInput( ParameterValue, TEXT( "MetallicTextureComponent" ) );
@@ -1876,7 +1865,7 @@ bool UsdToUnreal::ConvertMaterial( const pxr::UsdShadeMaterial& UsdShadeMaterial
 
 	// Roughness
 	{
-		if ( UsdShadeConversionImpl::GetFloatParameterValue( Connectable, UnrealIdentifiers::Roughness, 1.f, ParameterValue, &Material, TexturesCache, &PrimvarToUVIndex, bForceVirtualTextures ) )
+		if ( UsdShadeConversionImpl::GetFloatParameterValue( Connectable, UnrealIdentifiers::Roughness, 1.f, ParameterValue, &Material, TexturesCache, &PrimvarToUVIndex ) )
 		{
 			UsdShadeConversionImpl::SetParameterValue( Material, TEXT( "Roughness" ), ParameterValue, bForUsdPreviewSurface );
 			SetTextureComponentParamForScalarInput( ParameterValue, TEXT( "RoughnessTextureComponent" ) );
@@ -1886,7 +1875,7 @@ bool UsdToUnreal::ConvertMaterial( const pxr::UsdShadeMaterial& UsdShadeMaterial
 
 	// Opacity
 	{
-		if ( UsdShadeConversionImpl::GetFloatParameterValue( Connectable, UnrealIdentifiers::Opacity, 1.f, ParameterValue, &Material, TexturesCache, &PrimvarToUVIndex, bForceVirtualTextures ) )
+		if ( UsdShadeConversionImpl::GetFloatParameterValue( Connectable, UnrealIdentifiers::Opacity, 1.f, ParameterValue, &Material, TexturesCache, &PrimvarToUVIndex ) )
 		{
 			UsdShadeConversionImpl::SetParameterValue( Material, TEXT( "Opacity" ), ParameterValue, bForUsdPreviewSurface );
 			SetTextureComponentParamForScalarInput( ParameterValue, TEXT( "OpacityTextureComponent" ) );
@@ -1897,7 +1886,7 @@ bool UsdToUnreal::ConvertMaterial( const pxr::UsdShadeMaterial& UsdShadeMaterial
 	// Normal
 	{
 		const bool bIsNormalMap = true;
-		if ( UsdShadeConversionImpl::GetVec3ParameterValue( Connectable, UnrealIdentifiers::Normal, FLinearColor( 0, 0, 1 ), ParameterValue, bIsNormalMap, &Material, TexturesCache, &PrimvarToUVIndex, bForceVirtualTextures ) )
+		if ( UsdShadeConversionImpl::GetVec3ParameterValue( Connectable, UnrealIdentifiers::Normal, FLinearColor( 0, 0, 1 ), ParameterValue, bIsNormalMap, &Material, TexturesCache, &PrimvarToUVIndex ) )
 		{
 			UsdShadeConversionImpl::SetParameterValue( Material, TEXT( "Normal" ), ParameterValue, bForUsdPreviewSurface );
 			bHasMaterialInfo = true;
@@ -1906,7 +1895,7 @@ bool UsdToUnreal::ConvertMaterial( const pxr::UsdShadeMaterial& UsdShadeMaterial
 
 	// Refraction
 	{
-		if ( UsdShadeConversionImpl::GetFloatParameterValue( Connectable, UnrealIdentifiers::Refraction, 1.5f, ParameterValue, &Material, TexturesCache, &PrimvarToUVIndex, bForceVirtualTextures ) )
+		if ( UsdShadeConversionImpl::GetFloatParameterValue( Connectable, UnrealIdentifiers::Refraction, 1.5f, ParameterValue, &Material, TexturesCache, &PrimvarToUVIndex ) )
 		{
 			UsdShadeConversionImpl::SetParameterValue( Material, TEXT( "Refraction" ), ParameterValue, bForUsdPreviewSurface );
 			SetTextureComponentParamForScalarInput( ParameterValue, TEXT( "RefractionTextureComponent" ) );
@@ -1916,7 +1905,7 @@ bool UsdToUnreal::ConvertMaterial( const pxr::UsdShadeMaterial& UsdShadeMaterial
 
 	// Occlusion
 	{
-		if ( UsdShadeConversionImpl::GetFloatParameterValue( Connectable, UnrealIdentifiers::Occlusion, 1.0f, ParameterValue, &Material, TexturesCache, &PrimvarToUVIndex, bForceVirtualTextures ) )
+		if ( UsdShadeConversionImpl::GetFloatParameterValue( Connectable, UnrealIdentifiers::Occlusion, 1.0f, ParameterValue, &Material, TexturesCache, &PrimvarToUVIndex ) )
 		{
 			UsdShadeConversionImpl::SetParameterValue( Material, TEXT( "AmbientOcclusion" ), ParameterValue, bForUsdPreviewSurface );
 			SetTextureComponentParamForScalarInput( ParameterValue, TEXT( "AmbientOcclusionTextureComponent" ) );
diff --git a/Engine/Plugins/Importers/USDImporter/Source/USDUtilities/Public/USDShadeConversion.h b/Engine/Plugins/Importers/USDImporter/Source/USDUtilities/Public/USDShadeConversion.h
--- a/Engine/Plugins/Importers/USDImporter/Source/USDUtilities/Public/USDShadeConversion.h
+++ b/Engine/Plugins/Importers/USDImporter/Source/USDUtilities/Public/USDShadeConversion.h
@@ -130,6 +130,7 @@ namespace UsdUtils
 	* Returns whether the material uses UDIMs textures.
 	* This function exists because we need this information *before* we pick the right parent for a material instance and properly convert it.
 	*/
+	UE_DEPRECATED( 5.1, "No longer used: Textures and materials will be upgraded to VT-based on-demand" )
 	USDUTILITIES_API bool IsMaterialUsingUDIMs( const pxr::UsdShadeMaterial& UsdShadeMaterial );
 
 	USDUTILITIES_API FSHAHash HashShadeMaterial( const pxr::UsdShadeMaterial& UsdShadeMaterial, const pxr::TfToken& RenderContext = pxr::UsdShadeTokens->universalRenderContext );
