USD: Export skeletons with the guide purpose.

diff --git a/Engine/Plugins/Importers/USDImporter/Source/USDUtilities/Private/USDSkeletalDataConversion.cpp b/Engine/Plugins/Importers/USDImporter/Source/USDUtilities/Private/USDSkeletalDataConversion.cpp
--- a/Engine/Plugins/Importers/USDImporter/Source/USDUtilities/Private/USDSkeletalDataConversion.cpp
+++ b/Engine/Plugins/Importers/USDImporter/Source/USDUtilities/Private/USDSkeletalDataConversion.cpp
@@ -2883,6 +2883,15 @@ bool UnrealToUsd::ConvertSkeleton( const FReferenceSkeleton& ReferenceSkeleton,
 		BindTransformsAttr.Set( WorldSpaceJointTransforms );
 	}
 
+	// Use Guide purpose on skeletons by default, unless it has some specific purpose set already
+	if ( pxr::UsdAttribute PurposeAttr = UsdSkeleton.GetPurposeAttr() )
+	{
+		if ( !PurposeAttr.HasAuthoredValue() )
+		{
+			PurposeAttr.Set( pxr::UsdGeomTokens->guide );
+		}
+	}
+
 	return true;
 }
 
