USD: Fix sequencer not being able to animate spawned light componets as they were being set to Static mobility.

diff --git a/Engine/Plugins/Importers/USDImporter/Source/USDSchemas/Private/USDGeomXformableTranslator.cpp b/Engine/Plugins/Importers/USDImporter/Source/USDSchemas/Private/USDGeomXformableTranslator.cpp
--- a/Engine/Plugins/Importers/USDImporter/Source/USDSchemas/Private/USDGeomXformableTranslator.cpp
+++ b/Engine/Plugins/Importers/USDImporter/Source/USDSchemas/Private/USDGeomXformableTranslator.cpp
@@ -16,6 +16,7 @@
 #include "USDTypesConversion.h"
 
 #include "Components/HierarchicalInstancedStaticMeshComponent.h"
+#include "Components/LightComponentBase.h"
 #include "Components/SceneComponent.h"
 #include "Components/SkeletalMeshComponent.h"
 #include "Components/StaticMeshComponent.h"
@@ -512,7 +513,11 @@ USceneComponent* FUsdGeomXformableTranslator::CreateComponentsEx( TOptional< TSu
 		}
 		else
 		{
-			SceneComponent->Mobility = UsdUtils::IsAnimated( Prim ) ? EComponentMobility::Movable : EComponentMobility::Static;
+			SceneComponent->Mobility = UsdUtils::IsAnimated( Prim )
+				? EComponentMobility::Movable
+				: SceneComponent->IsA<ULightComponentBase>()	// Lights need to be stationary by default
+					? EComponentMobility::Stationary
+					: EComponentMobility::Static;
 		}
 
 		// Attach to parent
diff --git a/Engine/Plugins/Importers/USDImporter/Source/USDStage/Private/USDLevelSequenceHelper.cpp b/Engine/Plugins/Importers/USDImporter/Source/USDStage/Private/USDLevelSequenceHelper.cpp
--- a/Engine/Plugins/Importers/USDImporter/Source/USDStage/Private/USDLevelSequenceHelper.cpp
+++ b/Engine/Plugins/Importers/USDImporter/Source/USDStage/Private/USDLevelSequenceHelper.cpp
@@ -3020,6 +3020,7 @@ void FUsdLevelSequenceHelperImpl::HandleControlRigSectionChange( UMovieSceneCont
 	{
 		return;
 	}
+	ensure( BoundComponent->Mobility != EComponentMobility::Static );
 
 	USkeleton* Skeleton = BoundComponent->GetSkeletalMeshAsset() ? BoundComponent->GetSkeletalMeshAsset()->GetSkeleton() : nullptr;
 	if ( !Skeleton )
@@ -3190,6 +3191,7 @@ void FUsdLevelSequenceHelperImpl::HandleTrackChange( const UMovieSceneTrack& Tra
 	{
 		return;
 	}
+	ensure( BoundSceneComponent->Mobility != EComponentMobility::Static );
 
 	UUsdPrimTwin* PrimTwin = StageActor->RootUsdTwin->Find( BoundSceneComponent );
 
