USD: Add render context option for MaterialX materials.

diff --git a/Engine/Plugins/Importers/USDImporter/Source/USDSchemas/Private/USDSchemasModule.cpp b/Engine/Plugins/Importers/USDImporter/Source/USDSchemas/Private/USDSchemasModule.cpp
--- a/Engine/Plugins/Importers/USDImporter/Source/USDSchemas/Private/USDSchemasModule.cpp
+++ b/Engine/Plugins/Importers/USDImporter/Source/USDSchemas/Private/USDSchemasModule.cpp
@@ -23,6 +23,11 @@
 #include "Custom/MDLUSDShadeMaterialTranslator.h"
 #endif // #if USE_USD_SDK
 
+namespace UE::UsdSchemasModule::Private
+{
+	FName MaterialXRenderContext = TEXT("mtlx");
+}
+
 class FUsdSchemasModule : public IUsdSchemasModule
 {
 public:
@@ -56,12 +61,14 @@ public:
 		}
 #endif // WITH_EDITOR
 
+		GetRenderContextRegistry().Register( UE::UsdSchemasModule::Private::MaterialXRenderContext );
 #endif // #if USE_USD_SDK
 	}
 
 	virtual void ShutdownModule() override
 	{
 #if USE_USD_SDK
+		GetRenderContextRegistry().Unregister( UE::UsdSchemasModule::Private::MaterialXRenderContext );
 		GetTranslatorRegistry().Unregister( UsdGeomCameraTranslatorHandle );
 		GetTranslatorRegistry().Unregister( UsdGeomMeshTranslatorHandle );
 		GetTranslatorRegistry().Unregister( UsdGeomPointInstancerTranslatorHandle );
