USD: Fix bugs submitted with VT upgrade support on 22323421

diff --git a/Engine/Plugins/Importers/USDImporter/Source/USDSchemas/Private/USDShadeMaterialTranslator.cpp b/Engine/Plugins/Importers/USDImporter/Source/USDSchemas/Private/USDShadeMaterialTranslator.cpp
--- a/Engine/Plugins/Importers/USDImporter/Source/USDSchemas/Private/USDShadeMaterialTranslator.cpp
+++ b/Engine/Plugins/Importers/USDImporter/Source/USDSchemas/Private/USDShadeMaterialTranslator.cpp
@@ -177,7 +177,8 @@ namespace UE::UsdShadeTranslator::Private
 						);
 
 #if WITH_EDITOR
-						if ( UMaterialInstanceConstant* MIC = Cast<UMaterialInstanceConstant>( MaterialInstance ) )
+						UMaterialInstanceConstant* MIC = Cast<UMaterialInstanceConstant>( MaterialInstance );
+						if ( GIsEditor && MIC )
 						{
 							FMaterialUpdateContext UpdateContext( Options, GMaxRHIShaderPlatform );
 							UpdateContext.AddMaterialInstance( MIC );
@@ -205,7 +206,7 @@ namespace UE::UsdShadeTranslator::Private
 								// For MID we can't swap the base material, so we need to create a brand new one and copy
 								// the overrides
 								UMaterialInstanceDynamic* NewMID = UMaterialInstanceDynamic::Create(
-									BaseMaterial,
+									BaseMaterialVT,
 									GetTransientPackage(),
 									NewInstanceName
 								);
@@ -216,7 +217,7 @@ namespace UE::UsdShadeTranslator::Private
 								NewMID->CopyParameterOverrides( MID );
 
 								Context->AssetCache->DiscardAsset( Hash );
-								Context->AssetCache->CacheAsset( Hash, MID, PrimPath );
+								Context->AssetCache->CacheAsset( Hash, NewMID, PrimPath );
 								NewMaterials.Add( NewMID );
 							}
 						}
