USD: Support for reading/writing texture wrapping modes.

Note that previously we never wrote our desired wrapping mode, but also never read it either. By default we use "repeat" in UE, and I think it's relatively rare for textures to actually use texture wrapping, and they instead just use the [0, 1] UV range anyway, which is why we never noticed this.

diff --git a/Engine/Plugins/Importers/USDImporter/Source/USDUtilities/Private/USDShadeConversion.cpp b/Engine/Plugins/Importers/USDImporter/Source/USDUtilities/Private/USDShadeConversion.cpp
--- a/Engine/Plugins/Importers/USDImporter/Source/USDUtilities/Private/USDShadeConversion.cpp
+++ b/Engine/Plugins/Importers/USDImporter/Source/USDUtilities/Private/USDShadeConversion.cpp
@@ -315,7 +315,13 @@ namespace UE
 			// Computes and returns the hash string for the texture at the given path.
 			// Handles regular texture asset paths as well as asset paths identifying textures inside Usdz archives.
 			// Returns an empty string if the texture could not be hashed.
-			FString GetTextureHash( const FString& ResolvedTexturePath, bool bSRGB, TextureCompressionSettings CompressionSettings )
+			FString GetTextureHash(
+				const FString& ResolvedTexturePath,
+				bool bSRGB,
+				TextureCompressionSettings CompressionSettings,
+				TextureAddress AddressX,
+				TextureAddress AddressY
+			)
 			{
 				FMD5 MD5;
 
@@ -359,6 +365,8 @@ namespace UE
 				// Hash the additional data
 				MD5.Update( reinterpret_cast< uint8* >( &bSRGB ), sizeof( bool ) );
 				MD5.Update( reinterpret_cast< uint8* >( &CompressionSettings ), sizeof( CompressionSettings ) );
+				MD5.Update( reinterpret_cast< uint8* >( &AddressX ), sizeof( AddressX ) );
+				MD5.Update( reinterpret_cast< uint8* >( &AddressY ), sizeof( AddressY ) );
 
 				FMD5Hash Hash;
 				Hash.Set( MD5 );
@@ -529,6 +537,36 @@ namespace UE
 						}
 					}
 
+					// Get desired texture wrapping modes
+					TextureAddress AddressX = TextureAddress::TA_Wrap;
+					TextureAddress AddressY = TextureAddress::TA_Wrap;
+					if ( pxr::UsdAttribute WrapSAttr = UsdUVTextureSource.GetInput( UnrealIdentifiers::WrapS ) )
+					{
+						pxr::TfToken WrapS;
+						if ( WrapSAttr.Get( &WrapS ) )
+						{
+							AddressX = WrapS == UnrealIdentifiers::Repeat
+								? TextureAddress::TA_Wrap
+								: WrapS == UnrealIdentifiers::Mirror
+									? TextureAddress::TA_Mirror
+									: TextureAddress::TA_Clamp;  // We also consider the "black" wrap mode as clamp
+																 // as that is the closest we can get
+						}
+					}
+					if ( pxr::UsdAttribute WrapTAttr = UsdUVTextureSource.GetInput( UnrealIdentifiers::WrapT ) )
+					{
+						pxr::TfToken WrapT;
+						if ( WrapTAttr.Get( &WrapT ) )
+						{
+							AddressY = WrapT == UnrealIdentifiers::Repeat
+								? TextureAddress::TA_Wrap
+								: WrapT == UnrealIdentifiers::Mirror
+									? TextureAddress::TA_Mirror
+									: TextureAddress::TA_Clamp;  // We also consider the "black" wrap mode as clamp
+																 // as that is the closest we can get
+						}
+					}
+
 					if ( FileInput && FileInput.GetTypeName() == pxr::SdfValueTypeNames->Asset ) // Check that FileInput is of type Asset
 					{
 						const FString TexturePath = UsdUtils::GetResolvedTexturePath( FileInput.GetAttr() );
@@ -561,7 +599,7 @@ namespace UE
 							bSRGB = false;
 						}
 
-						const FString TextureHash = GetTextureHash( TexturePath, bSRGB, CompressionSettings );
+						const FString TextureHash = GetTextureHash( TexturePath, bSRGB, CompressionSettings, AddressX, AddressY );
 
 						// We only actually want to retrieve the textures if we have a cache to put them in
 						if ( TexturesCache )
@@ -587,6 +625,11 @@ namespace UE
 								{
 									Texture->SRGB = bSRGB;
 									Texture->CompressionSettings = CompressionSettings;
+									if ( UTexture2D* Texture2D = Cast<UTexture2D>( Texture ) )
+									{
+										Texture2D->AddressX = AddressX;
+										Texture2D->AddressY = AddressY;
+									}
 									Texture->UpdateResource();
 
 									TexturesCache->CacheAsset( TextureHash, Texture );
@@ -1685,6 +1728,22 @@ namespace UE
 						pxr::UsdShadeInput TextureFallbackInput = UsdUVTextureShader.CreateInput( UnrealIdentifiers::Fallback, pxr::SdfValueTypeNames->Float4 );
 						TextureFallbackInput.Set( FallbackValue );
 
+						// In the general case it's impossible to set a "correct" wrapping value here because the
+						// material we just baked may be using 3 different textures with UV transforms an all different
+						// texture wrapping modes, and we're forced to pick a single value to wrap the baked texture
+						// with, but let's at least write "repeat" out as that is the default for textures in UE and
+						// the more likely to be correct, in case the mesh does things like have UVs outside [0, 1]
+						pxr::UsdShadeInput TextureFileWrapSInput = UsdUVTextureShader.CreateInput(
+							UnrealIdentifiers::WrapS,
+							pxr::SdfValueTypeNames->Token
+						);
+						TextureFileWrapSInput.Set( UnrealIdentifiers::Repeat );
+						pxr::UsdShadeInput TextureFileWrapTInput = UsdUVTextureShader.CreateInput(
+							UnrealIdentifiers::WrapT,
+							pxr::SdfValueTypeNames->Token
+						);
+						TextureFileWrapTInput.Set( UnrealIdentifiers::Repeat );
+
 						pxr::UsdShadeOutput TextureOutput = UsdUVTextureShader.CreateOutput(
 							InputType == pxr::SdfValueTypeNames->Float ? UnrealIdentifiers::R : UnrealIdentifiers::RGB,
 							InputType
diff --git a/Engine/Plugins/Importers/USDImporter/Source/UnrealUSDWrapper/Private/UnrealUSDWrapper.cpp b/Engine/Plugins/Importers/USDImporter/Source/UnrealUSDWrapper/Private/UnrealUSDWrapper.cpp
--- a/Engine/Plugins/Importers/USDImporter/Source/UnrealUSDWrapper/Private/UnrealUSDWrapper.cpp
+++ b/Engine/Plugins/Importers/USDImporter/Source/UnrealUSDWrapper/Private/UnrealUSDWrapper.cpp
@@ -121,6 +121,11 @@ namespace UnrealIdentifiers
 	const TfToken Varname = TfToken("varname");
 	const TfToken Result = TfToken("result");
 	const TfToken File = TfToken("file");
+	const TfToken WrapT = TfToken( "wrapT" );
+	const TfToken WrapS = TfToken( "wrapS" );
+	const TfToken Repeat = TfToken( "repeat" );
+	const TfToken Mirror = TfToken( "mirror" );
+	const TfToken Clamp = TfToken( "clamp" );
 	const TfToken Fallback = TfToken("fallback");
 	const TfToken R = TfToken("r");
 	const TfToken RGB = TfToken("rgb");
diff --git a/Engine/Plugins/Importers/USDImporter/Source/UnrealUSDWrapper/Public/UnrealUSDWrapper.h b/Engine/Plugins/Importers/USDImporter/Source/UnrealUSDWrapper/Public/UnrealUSDWrapper.h
--- a/Engine/Plugins/Importers/USDImporter/Source/UnrealUSDWrapper/Public/UnrealUSDWrapper.h
+++ b/Engine/Plugins/Importers/USDImporter/Source/UnrealUSDWrapper/Public/UnrealUSDWrapper.h
@@ -430,6 +430,11 @@ namespace UnrealIdentifiers
 	extern UNREALUSDWRAPPER_API const pxr::TfToken Varname;
 	extern UNREALUSDWRAPPER_API const pxr::TfToken Result;
 	extern UNREALUSDWRAPPER_API const pxr::TfToken File;
+	extern UNREALUSDWRAPPER_API const pxr::TfToken WrapS;
+	extern UNREALUSDWRAPPER_API const pxr::TfToken WrapT;
+	extern UNREALUSDWRAPPER_API const pxr::TfToken Repeat;
+	extern UNREALUSDWRAPPER_API const pxr::TfToken Mirror;
+	extern UNREALUSDWRAPPER_API const pxr::TfToken Clamp;
 	extern UNREALUSDWRAPPER_API const pxr::TfToken Fallback;
 	extern UNREALUSDWRAPPER_API const pxr::TfToken R;
 	extern UNREALUSDWRAPPER_API const pxr::TfToken RGB;
