USD: Deprecate option to "remove unreal materials".

diff --git a/Engine/Plugins/Importers/USDImporter/Source/USDExporter/Private/LevelExporterUSDOptions.cpp b/Engine/Plugins/Importers/USDImporter/Source/USDExporter/Private/LevelExporterUSDOptions.cpp
--- a/Engine/Plugins/Importers/USDImporter/Source/USDExporter/Private/LevelExporterUSDOptions.cpp
+++ b/Engine/Plugins/Importers/USDImporter/Source/USDExporter/Private/LevelExporterUSDOptions.cpp
@@ -65,7 +65,6 @@ void UsdUtils::HashForLevelExport( const FLevelExporterUSDOptionsInner& Options,
 
 	// This affects how we author material overrides
 	HashToUpdate.Update( reinterpret_cast< const uint8* >( &Options.AssetOptions.bBakeMaterials ), sizeof( Options.AssetOptions.bBakeMaterials ) );
-	HashToUpdate.Update( reinterpret_cast< const uint8* >( &Options.AssetOptions.bRemoveUnrealMaterials ), sizeof( Options.AssetOptions.bRemoveUnrealMaterials ) );
 
 	// This affects how we author material overrides and landscapes
 	HashToUpdate.Update( reinterpret_cast< const uint8* >( &Options.AssetOptions.bUsePayload ), sizeof( Options.AssetOptions.bUsePayload ) );
diff --git a/Engine/Plugins/Importers/USDImporter/Source/USDExporter/Private/MaterialExporterUSD.cpp b/Engine/Plugins/Importers/USDImporter/Source/USDExporter/Private/MaterialExporterUSD.cpp
--- a/Engine/Plugins/Importers/USDImporter/Source/USDExporter/Private/MaterialExporterUSD.cpp
+++ b/Engine/Plugins/Importers/USDImporter/Source/USDExporter/Private/MaterialExporterUSD.cpp
@@ -401,7 +401,6 @@ bool UMaterialExporterUsd::ExportMaterialsForStage(
 	const FString& StageRootLayerPath,
 	bool bIsAssetLayer,
 	bool bUsePayload,
-	bool bRemoveUnrealMaterials,
 	bool bReplaceIdentical,
 	bool bReExportIdenticalAssets,
 	bool bIsAutomated
@@ -492,8 +491,7 @@ bool UMaterialExporterUsd::ExportMaterialsForStage(
 		RootLayer,
 		MaterialPathNameToFilePath,
 		bIsAssetLayer,
-		bUsePayload,
-		bRemoveUnrealMaterials
+		bUsePayload
 	);
 
 	RootLayer.Save();
diff --git a/Engine/Plugins/Importers/USDImporter/Source/USDExporter/Private/SkeletalMeshExporterUSD.cpp b/Engine/Plugins/Importers/USDImporter/Source/USDExporter/Private/SkeletalMeshExporterUSD.cpp
--- a/Engine/Plugins/Importers/USDImporter/Source/USDExporter/Private/SkeletalMeshExporterUSD.cpp
+++ b/Engine/Plugins/Importers/USDImporter/Source/USDExporter/Private/SkeletalMeshExporterUSD.cpp
@@ -212,7 +212,6 @@ bool USkeletalMeshExporterUsd::ExportBinary( UObject* Object, const TCHAR* Type,
 						UExporter::CurrentFilename,
 						bIsAssetLayer,
 						Options->MeshAssetOptions.bUsePayload,
-						Options->MeshAssetOptions.bRemoveUnrealMaterials,
 						ExportTask->bReplaceIdentical,
 						Options->bReExportIdenticalAssets,
 						ExportTask->bAutomated
@@ -308,7 +307,6 @@ bool USkeletalMeshExporterUsd::ExportBinary( UObject* Object, const TCHAR* Type,
 			AssetStage.GetRootLayer().GetRealPath(),
 			bIsAssetLayer,
 			Options->MeshAssetOptions.bUsePayload,
-			Options->MeshAssetOptions.bRemoveUnrealMaterials,
 			ExportTask->bReplaceIdentical,
 			Options->bReExportIdenticalAssets,
 			ExportTask->bAutomated
diff --git a/Engine/Plugins/Importers/USDImporter/Source/USDExporter/Private/StaticMeshExporterUSD.cpp b/Engine/Plugins/Importers/USDImporter/Source/USDExporter/Private/StaticMeshExporterUSD.cpp
--- a/Engine/Plugins/Importers/USDImporter/Source/USDExporter/Private/StaticMeshExporterUSD.cpp
+++ b/Engine/Plugins/Importers/USDImporter/Source/USDExporter/Private/StaticMeshExporterUSD.cpp
@@ -227,7 +227,6 @@ bool UStaticMeshExporterUsd::ExportBinary( UObject* Object, const TCHAR* Type, F
 						UExporter::CurrentFilename,
 						bIsAssetLayer,
 						Options->MeshAssetOptions.bUsePayload,
-						Options->MeshAssetOptions.bRemoveUnrealMaterials,
 						ExportTask->bReplaceIdentical,
 						Options->bReExportIdenticalAssets,
 						ExportTask->bAutomated
@@ -321,7 +320,6 @@ bool UStaticMeshExporterUsd::ExportBinary( UObject* Object, const TCHAR* Type, F
 			AssetStage.GetRootLayer().GetRealPath(),
 			bIsAssetLayer,
 			Options->MeshAssetOptions.bUsePayload,
-			Options->MeshAssetOptions.bRemoveUnrealMaterials,
 			ExportTask->bReplaceIdentical,
 			Options->bReExportIdenticalAssets,
 			ExportTask->bAutomated
diff --git a/Engine/Plugins/Importers/USDImporter/Source/USDExporter/Private/USDAssetOptions.cpp b/Engine/Plugins/Importers/USDImporter/Source/USDExporter/Private/USDAssetOptions.cpp
--- a/Engine/Plugins/Importers/USDImporter/Source/USDExporter/Private/USDAssetOptions.cpp
+++ b/Engine/Plugins/Importers/USDImporter/Source/USDExporter/Private/USDAssetOptions.cpp
@@ -37,7 +37,6 @@ void UsdUtils::AddAnalyticsAttributes(
 		InOutAttributes.Emplace( TEXT( "PayloadFormat" ), Options.PayloadFormat );
 	}
 	InOutAttributes.Emplace( TEXT( "BakeMaterials" ), Options.bBakeMaterials );
-	InOutAttributes.Emplace( TEXT( "RemoveUnrealMaterials" ), Options.bRemoveUnrealMaterials );
 	if ( Options.bBakeMaterials )
 	{
 		UsdUtils::AddAnalyticsAttributes( Options.MaterialBakingOptions, InOutAttributes );
@@ -76,10 +75,6 @@ void UsdUtils::HashForMeshExport( const FUsdMeshAssetOptions& Options, FSHA1& Ha
 		reinterpret_cast< const uint8* >( &Options.bBakeMaterials ),
 		sizeof( Options.bBakeMaterials )
 	);
-	HashToUpdate.Update(
-		reinterpret_cast< const uint8* >( &Options.bRemoveUnrealMaterials ),
-		sizeof( Options.bRemoveUnrealMaterials )
-	);
 
 	HashToUpdate.Update(
 		reinterpret_cast< const uint8* >( &Options.LowestMeshLOD ),
diff --git a/Engine/Plugins/Importers/USDImporter/Source/USDExporter/Private/USDConversionBlueprintContext.cpp b/Engine/Plugins/Importers/USDImporter/Source/USDExporter/Private/USDConversionBlueprintContext.cpp
--- a/Engine/Plugins/Importers/USDImporter/Source/USDExporter/Private/USDConversionBlueprintContext.cpp
+++ b/Engine/Plugins/Importers/USDImporter/Source/USDExporter/Private/USDConversionBlueprintContext.cpp
@@ -436,7 +436,7 @@ void UUsdConversionBlueprintContext::ReplaceUnrealMaterialsWithBaked( const FFil
 		return;
 	}
 
-	UsdUtils::ReplaceUnrealMaterialsWithBaked( Stage, Layer, BakedMaterials, bIsAssetLayer, bUsePayload, bRemoveUnrealMaterials );
+	UsdUtils::ReplaceUnrealMaterialsWithBaked( Stage, Layer, BakedMaterials, bIsAssetLayer, bUsePayload );
 #endif // USE_USD_SDK
 }
 
@@ -460,7 +460,9 @@ bool UUsdConversionBlueprintContext::RemoveUnrealSurfaceOutput( const FString& P
 		Layer = UE::FSdfLayer::FindOrOpen( *LayerToAuthorIn.FilePath );
 	}
 
+	PRAGMA_DISABLE_DEPRECATION_WARNINGS
 	return UsdUtils::RemoveUnrealSurfaceOutput( Prim, Layer );
+	PRAGMA_ENABLE_DEPRECATION_WARNINGS
 #else
 	return false;
 #endif // USE_USD_SDK
diff --git a/Engine/Plugins/Importers/USDImporter/Source/USDExporter/Public/MaterialExporterUSD.h b/Engine/Plugins/Importers/USDImporter/Source/USDExporter/Public/MaterialExporterUSD.h
--- a/Engine/Plugins/Importers/USDImporter/Source/USDExporter/Public/MaterialExporterUSD.h
+++ b/Engine/Plugins/Importers/USDImporter/Source/USDExporter/Public/MaterialExporterUSD.h
@@ -56,8 +56,6 @@ public:
 	 *                        when authoring the material binding relationships, e.g. whether we author them inside variants or not
 	 * @param bUsePayload - Should be True if the Stage was exported using payload files to store the actual Mesh prims. Also dictates minor
 	 *                      behaviors when authoring the material binding relationships.
-	 * @param bRemoveUnrealMaterials - Whether to remove the `unrealMaterial` attributes after replacing them with material bindings.
-	 *                                 Important because the `unrealMaterial` attributes will be used as a higher priority when determining material assignments
 	 * @param ReplaceIdentical - Whether to overwrite the destination files in case they already exist
 	 * @param bReExportIdenticalAssets - Whether to overwrite the destination files even if they already describe up-to-date versions of Materials
 	 * @param bIsAutomated - Whether the export is being done by a script or not. Just used for analytics
@@ -69,7 +67,6 @@ public:
 		const FString& StageRootLayerPath,
 		bool bIsAssetLayer,
 		bool bUsePayload,
-		bool bRemoveUnrealMaterials,
 		bool bReplaceIdentical = true,
 		bool bReExportIdenticalAssets = false,
 		bool bIsAutomated = false
diff --git a/Engine/Plugins/Importers/USDImporter/Source/USDExporter/Public/USDAssetOptions.h b/Engine/Plugins/Importers/USDImporter/Source/USDExporter/Public/USDAssetOptions.h
--- a/Engine/Plugins/Importers/USDImporter/Source/USDExporter/Public/USDAssetOptions.h
+++ b/Engine/Plugins/Importers/USDImporter/Source/USDExporter/Public/USDAssetOptions.h
@@ -46,7 +46,8 @@ struct USDEXPORTER_API FUsdMeshAssetOptions
 	bool bBakeMaterials = false;
 
 	/** Whether to remove the 'unrealMaterial' attribute after binding the corresponding baked material */
-	UPROPERTY( EditAnywhere, config, BlueprintReadWrite, Category = "Material options", meta = ( EditCondition = bBakeMaterials ) )
+	UE_DEPRECATED( 5.2, "This option is now deprecated as UE material assignments are only visible in the 'unreal' render context anyway" )
+	UPROPERTY( config, meta = ( EditCondition = bBakeMaterials ) )
 	bool bRemoveUnrealMaterials = false;
 
 	UPROPERTY( EditAnywhere, config, BlueprintReadWrite, Category = "Material options", meta = ( EditCondition = bBakeMaterials ) )
@@ -59,6 +60,13 @@ struct USDEXPORTER_API FUsdMeshAssetOptions
 	/** Highest of the LOD indices to export static and skeletal meshes with */
 	UPROPERTY( EditAnywhere, config, BlueprintReadWrite, Category = "Mesh options", meta = ( ClampMin = "0" ) )
 	int32 HighestMeshLOD = MAX_MESH_LOD_COUNT - 1;
+
+	// We need this as clang would otherwise emit some warnings when generating these due to the usages of bRemoveUnrealMaterials
+	PRAGMA_DISABLE_DEPRECATION_WARNINGS
+	FUsdMeshAssetOptions() = default;
+	FUsdMeshAssetOptions( const FUsdMeshAssetOptions& ) = default;
+	FUsdMeshAssetOptions& operator=( const FUsdMeshAssetOptions& ) = default;
+	PRAGMA_ENABLE_DEPRECATION_WARNINGS
 };
 
 namespace UsdUtils
diff --git a/Engine/Plugins/Importers/USDImporter/Source/USDExporter/Public/USDConversionBlueprintContext.h b/Engine/Plugins/Importers/USDImporter/Source/USDExporter/Public/USDConversionBlueprintContext.h
--- a/Engine/Plugins/Importers/USDImporter/Source/USDExporter/Public/USDConversionBlueprintContext.h
+++ b/Engine/Plugins/Importers/USDImporter/Source/USDExporter/Public/USDConversionBlueprintContext.h
@@ -146,11 +146,10 @@ public:
 	 *                        when authoring the material binding relationships, e.g. whether we author them inside variants or not
 	 * @param bUsePayload - Should be True if the Stage was exported using payload files to store the actual Mesh prims. Also dictates minor
 	 *                      behaviors when authoring the material binding relationships.
-	 * @param bRemoveUnrealMaterials - Whether to remove the `unrealMaterial` attributes after replacing them with material bindings.
-	 *                                 Important because the `unrealMaterial` attributes will be used as a higher priority when determining material assignments on import
+	 * @param bRemoveUnrealMaterials - No longer used.
 	 */
 	UFUNCTION( BlueprintCallable, Category = "Conversion utils" )
-	void ReplaceUnrealMaterialsWithBaked( const FFilePath& LayerToAuthorIn, const TMap<FString, FString>& BakedMaterials, bool bIsAssetLayer, bool bUsePayload, bool bRemoveUnrealMaterials );
+	void ReplaceUnrealMaterialsWithBaked( const FFilePath& LayerToAuthorIn, const TMap<FString, FString>& BakedMaterials, bool bIsAssetLayer, bool bUsePayload, bool bRemoveUnrealMaterials = false );
 
 	/**
 	 * Clears any opinions for the 'unreal' render context surface output of MaterialPrim within LayerToAuthorIn.
@@ -160,6 +159,7 @@ public:
 	 * @param LayerToAuthorIn - Layer to clear the opinions in. Can be the empty string to clear opinions from all layers in the layer stack
 	 * @return Whether we successfully cleared the opinions or not
 	 */
+	UE_DEPRECATED( 5.2, "No longer used as UE material assignments are only visible in the 'unreal' render context anyway" )
 	UFUNCTION( BlueprintCallable, Category = "Conversion utils" )
 	bool RemoveUnrealSurfaceOutput( const FString& PrimPath, const FFilePath& LayerToAuthorIn );
 
diff --git a/Engine/Plugins/Importers/USDImporter/Source/USDUtilities/Private/USDGeomMeshConversion.cpp b/Engine/Plugins/Importers/USDImporter/Source/USDUtilities/Private/USDGeomMeshConversion.cpp
--- a/Engine/Plugins/Importers/USDImporter/Source/USDUtilities/Private/USDGeomMeshConversion.cpp
+++ b/Engine/Plugins/Importers/USDImporter/Source/USDUtilities/Private/USDGeomMeshConversion.cpp
@@ -2323,8 +2323,7 @@ void UsdUtils::ReplaceUnrealMaterialsWithBaked(
 	const UE::FSdfLayer& LayerToAuthorIn,
 	const TMap<FString, FString>& BakedMaterials,
 	bool bIsAssetLayer,
-	bool bUsePayload,
-	bool bRemoveUnrealMaterials
+	bool bUsePayload
 )
 {
 	FScopedUsdAllocs Allocs;
@@ -2371,7 +2370,6 @@ void UsdUtils::ReplaceUnrealMaterialsWithBaked(
 			&BakedMaterials,
 			bIsAssetLayer,
 			bUsePayload,
-			bRemoveUnrealMaterials,
 			&StageMatScope
 		]
 		(
@@ -2523,25 +2521,6 @@ void UsdUtils::ReplaceUnrealMaterialsWithBaked(
 
 				pxr::UsdEditContext Context{ StageToTraverse, Layer };
 
-				if ( bRemoveUnrealMaterials )
-				{
-					TOptional<pxr::UsdEditContext> VarContext;
-					if ( bAuthorInsideVariants )
-					{
-						VarContext.Emplace( OuterVariantSet.GetValue().GetVariantEditContext() );
-					}
-
-					if ( UnrealMaterialAttrSpec )
-					{
-						Prim.RemoveProperty( UnrealIdentifiers::MaterialAssignment );
-					}
-
-					if ( UnrealMaterialPrimSpec )
-					{
-						UsdUtils::RemoveUnrealSurfaceOutput( UnrealMaterialPrim, UE::FSdfLayer{ Layer } );
-					}
-				}
-
 				// It was just an empty UE asset path, so just cancel now as our BakedFilename can't possibly be useful
 				if ( UnrealMaterialAssetPath.IsEmpty() )
 				{
@@ -2674,6 +2653,24 @@ void UsdUtils::ReplaceUnrealMaterialsWithBaked(
 	TraverseForMaterialReplacement( UsdStage, Root, Empty, {} );
 }
 
+void UsdUtils::ReplaceUnrealMaterialsWithBaked(
+	const UE::FUsdStage& Stage,
+	const UE::FSdfLayer& LayerToAuthorIn,
+	const TMap<FString, FString>& BakedMaterials,
+	bool bIsAssetLayer,
+	bool bUsePayload,
+	bool bRemoveUnrealMaterials
+)
+{
+	ReplaceUnrealMaterialsWithBaked(
+		Stage,
+		LayerToAuthorIn,
+		BakedMaterials,
+		bIsAssetLayer,
+		bUsePayload
+	);
+}
+
 FString UsdUtils::HashGeomMeshPrim( const UE::FUsdStage& Stage, const FString& PrimPath, double TimeCode )
 {
 	TRACE_CPUPROFILER_EVENT_SCOPE( UsdUtils::HashGeomMeshPrim );
diff --git a/Engine/Plugins/Importers/USDImporter/Source/USDUtilities/Public/USDGeomMeshConversion.h b/Engine/Plugins/Importers/USDImporter/Source/USDUtilities/Public/USDGeomMeshConversion.h
--- a/Engine/Plugins/Importers/USDImporter/Source/USDUtilities/Public/USDGeomMeshConversion.h
+++ b/Engine/Plugins/Importers/USDImporter/Source/USDUtilities/Public/USDGeomMeshConversion.h
@@ -311,10 +311,23 @@ namespace UsdUtils
 	 *                        when authoring the material binding relationships, e.g. whether we author them inside variants or not
 	 * @param bUsePayload - Should be True if the Stage was exported using payload files to store the actual Mesh prims. Also dictates minor
 	 *                      behaviors when authoring the material binding relationships.
-	 * @param bRemoveUnrealMaterials - Whether to remove the `unrealMaterial` attributes after replacing them with material bindings.
-	 *                                 Important because the `unrealMaterial` attributes will be used as a higher priority when determining material assignments
 	 */
-	USDUTILITIES_API void ReplaceUnrealMaterialsWithBaked( const UE::FUsdStage& Stage, const UE::FSdfLayer& LayerToAuthorIn, const TMap<FString, FString>& BakedMaterials, bool bIsAssetLayer, bool bUsePayload, bool bRemoveUnrealMaterials );
+	USDUTILITIES_API void ReplaceUnrealMaterialsWithBaked(
+		const UE::FUsdStage& Stage,
+		const UE::FSdfLayer& LayerToAuthorIn,
+		const TMap<FString, FString>& BakedMaterials,
+		bool bIsAssetLayer,
+		bool bUsePayload
+	);
+	UE_DEPRECATED( 5.2, "The bRemoveUnrealMaterials parameter is now deprecated as removing the UE material assignments is no longer needed" )
+	USDUTILITIES_API void ReplaceUnrealMaterialsWithBaked(
+		const UE::FUsdStage& Stage,
+		const UE::FSdfLayer& LayerToAuthorIn,
+		const TMap<FString, FString>& BakedMaterials,
+		bool bIsAssetLayer,
+		bool bUsePayload,
+		bool bRemoveUnrealMaterials
+	);
 
 	/**
 	 * Hashes the attributes of the GeomMesh at the given prim path on the Stage at TimeCode
diff --git a/Engine/Plugins/Importers/USDImporter/Source/USDUtilities/Public/USDShadeConversion.h b/Engine/Plugins/Importers/USDImporter/Source/USDUtilities/Public/USDShadeConversion.h
--- a/Engine/Plugins/Importers/USDImporter/Source/USDUtilities/Public/USDShadeConversion.h
+++ b/Engine/Plugins/Importers/USDImporter/Source/USDUtilities/Public/USDShadeConversion.h
@@ -118,6 +118,7 @@ namespace UsdUtils
 	 * @param LayerToAuthorIn - Layer to clear the opinions in, or an invalid layer (e.g. UE::FSdfLayer{}, which is the default)
 	 * @return Whether we successfully cleared the opinions or not
 	 */
+	UE_DEPRECATED( 5.2, "No longer used as UE material assignments are only visible in the 'unreal' render context anyway" )
 	USDUTILITIES_API bool RemoveUnrealSurfaceOutput( pxr::UsdPrim& MaterialPrim, const UE::FSdfLayer& LayerToAuthorIn = UE::FSdfLayer{} );
 
 	/**
