USD: Add support for generating two-sided preview surface materials when the meshes that bind them are double-sided.

diff --git a/Engine/Plugins/Importers/USDImporter/Source/USDSchemas/Private/MeshTranslationImpl.cpp b/Engine/Plugins/Importers/USDImporter/Source/USDSchemas/Private/MeshTranslationImpl.cpp
--- a/Engine/Plugins/Importers/USDImporter/Source/USDSchemas/Private/MeshTranslationImpl.cpp
+++ b/Engine/Plugins/Importers/USDImporter/Source/USDSchemas/Private/MeshTranslationImpl.cpp
@@ -26,8 +26,32 @@
 	#include "pxr/usd/usdShade/material.h"
 #include "USDIncludesEnd.h"
 
-TMap<const UsdUtils::FUsdPrimMaterialSlot*, UMaterialInterface*> MeshTranslationImpl::ResolveMaterialAssignmentInfo( const pxr::UsdPrim& UsdPrim, const TArray<UsdUtils::FUsdPrimMaterialAssignmentInfo>& AssignmentInfo, const TArray<UMaterialInterface*>& ExistingAssignments, UUsdAssetCache& AssetCache, float Time, EObjectFlags Flags )
+namespace UE::MeshTranslationImplInternal::Private
 {
+	const static FString BaseMaterialPath = TEXT( "/USDImporter/Materials/UsdPreviewSurface.UsdPreviewSurface" );
+	const static FString BaseMaterialPathVT = TEXT( "/USDImporter/Materials/UsdPreviewSurfaceVT.UsdPreviewSurfaceVT" );
+
+	const static FString BaseMaterialPathTranslucent = TEXT( "/USDImporter/Materials/UsdPreviewSurfaceTranslucent.UsdPreviewSurfaceTranslucent" );
+	const static FString BaseMaterialPathTranslucentVT = TEXT( "/USDImporter/Materials/UsdPreviewSurfaceTranslucentVT.UsdPreviewSurfaceTranslucentVT" );
+
+	const static FString BaseMaterialPathTwoSided = TEXT( "/USDImporter/Materials/UsdPreviewSurfaceTwoSided.UsdPreviewSurfaceTwoSided" );
+	const static FString BaseMaterialPathTwoSidedVT = TEXT( "/USDImporter/Materials/UsdPreviewSurfaceTwoSidedVT.UsdPreviewSurfaceTwoSidedVT" );
+
+	const static FString BaseMaterialPathTranslucentTwoSided = TEXT( "/USDImporter/Materials/UsdPreviewSurfaceTranslucentTwoSided.UsdPreviewSurfaceTranslucentTwoSided" );
+	const static FString BaseMaterialPathTranslucentTwoSidedVT = TEXT( "/USDImporter/Materials/UsdPreviewSurfaceTranslucentTwoSidedVT.UsdPreviewSurfaceTranslucentTwoSidedVT" );
+}
+
+TMap<const UsdUtils::FUsdPrimMaterialSlot*, UMaterialInterface*> MeshTranslationImpl::ResolveMaterialAssignmentInfo(
+	const pxr::UsdPrim& UsdPrim,
+	const TArray<UsdUtils::FUsdPrimMaterialAssignmentInfo>& AssignmentInfo,
+	const TArray<UMaterialInterface*>& ExistingAssignments,
+	UUsdAssetCache& AssetCache,
+	float Time,
+	EObjectFlags Flags
+)
+{
+	FScopedUnrealAllocs Allocs;
+
 	TMap<const UsdUtils::FUsdPrimMaterialSlot*, UMaterialInterface*> ResolvedMaterials;
 
 	uint32 GlobalResolvedMaterialIndex = 0;
@@ -44,8 +68,6 @@ TMap<const UsdUtils::FUsdPrimMaterialSlot*, UMaterialInterface*> MeshTranslation
 			{
 			case UsdUtils::EPrimAssignmentType::DisplayColor:
 			{
-				FScopedUsdAllocs Allocs;
-
 				// Try reusing an already created DisplayColor material
 				if ( UMaterialInterface* ExistingMaterial = Cast<UMaterialInterface>( AssetCache.GetCachedAsset( Slot.MaterialSource ) ) )
 				{
@@ -82,18 +104,86 @@ TMap<const UsdUtils::FUsdPrimMaterialSlot*, UMaterialInterface*> MeshTranslation
 			}
 			case UsdUtils::EPrimAssignmentType::MaterialPrim:
 			{
-				FScopedUsdAllocs Allocs;
+				FString PrimPath = Slot.MaterialSource;
+				const static FString TwoSidedToken = TEXT( "!twosided" );
+				if ( Slot.bMeshIsDoubleSided )
+				{
+					PrimPath += TwoSidedToken;
+				}
+
+				Material = Cast< UMaterialInterface >( AssetCache.GetAssetForPrim( PrimPath ) );
 
-				// Check first or else we may get a warning
-				if ( pxr::SdfPath::IsValidPathString( UnrealToUsd::ConvertString( *Slot.MaterialSource ).Get() ) )
+				// Need to create a two-sided material on-demand
+				if ( !Material && Slot.bMeshIsDoubleSided )
 				{
-					pxr::SdfPath MaterialPrimPath = UnrealToUsd::ConvertPath( *Slot.MaterialSource ).Get();
+					// By now we parsed all materials so we must have the single-sided version of this material
+					UMaterialInstance* OneSidedMat = Cast< UMaterialInstance >(
+						AssetCache.GetAssetForPrim( Slot.MaterialSource )
+					);
+					if ( !ensure( OneSidedMat ) )
+					{
+						continue;
+					}
 
-					// TODO: This may break if MaterialPrimPath targets a prim inside a LOD variant that is disabled...
-					TUsdStore< pxr::UsdPrim > MaterialPrim = UsdPrim.GetStage()->GetPrimAtPath( MaterialPrimPath );
-					if ( MaterialPrim.Get() )
+					// Important to not use GetBaseMaterial() here because if our parent is the translucent we'll
+					// get the base UsdPreviewSurface instead, as that is also *its* base
+					UMaterialInterface* BaseMaterial = OneSidedMat->Parent.Get();
+					UMaterialInterface* BaseMaterialTwoSided =
+						MeshTranslationImpl::GetTwoSidedVersionOfBasePreviewSurfaceMaterial( BaseMaterial );
+					if ( !ensure( BaseMaterialTwoSided && BaseMaterialTwoSided != BaseMaterial ) )
+					{
+						continue;
+					}
+
+					const FName NewInstanceName = MakeUniqueObjectName(
+						GetTransientPackage(),
+						UMaterialInstance::StaticClass(),
+						*( FPaths::GetBaseFilename( Slot.MaterialSource ) + UnrealIdentifiers::TwoSidedMaterialSuffix )
+					);
+
+#if WITH_EDITOR
+					UMaterialInstanceConstant* MIC = Cast<UMaterialInstanceConstant>( OneSidedMat );
+					if ( GIsEditor && MIC )
 					{
-						Material = Cast< UMaterialInterface >( AssetCache.GetAssetForPrim( UsdToUnreal::ConvertPath( MaterialPrim.Get().GetPrimPath() ) ) );
+						UMaterialInstanceConstant* TwoSidedMat = NewObject<UMaterialInstanceConstant>(
+							GetTransientPackage(),
+							NewInstanceName,
+							Flags
+						);
+						if ( TwoSidedMat )
+						{
+							UUsdAssetImportData* ImportData = NewObject< UUsdAssetImportData >(
+								TwoSidedMat,
+								TEXT( "USDAssetImportData" )
+							);
+							ImportData->PrimPath = Slot.MaterialSource;
+							TwoSidedMat->AssetImportData = ImportData;
+						}
+
+						TwoSidedMat->SetParentEditorOnly( BaseMaterialTwoSided );
+						TwoSidedMat->CopyMaterialUniformParametersEditorOnly( OneSidedMat );
+
+						AssetCache.CacheAsset( PrimPath, TwoSidedMat, PrimPath );
+						Material = TwoSidedMat;
+					}
+					else
+#endif // WITH_EDITOR
+					if ( UMaterialInstanceDynamic* MID = Cast<UMaterialInstanceDynamic>( OneSidedMat ) )
+					{
+						UMaterialInstanceDynamic* TwoSidedMat = UMaterialInstanceDynamic::Create(
+							BaseMaterialTwoSided,
+							GetTransientPackage(),
+							NewInstanceName
+						);
+						if ( !ensure( TwoSidedMat ) )
+						{
+							continue;
+						}
+
+						TwoSidedMat->CopyParameterOverrides( MID );
+
+						AssetCache.CacheAsset( PrimPath, TwoSidedMat, PrimPath );
+						Material = TwoSidedMat;
 					}
 				}
 
@@ -102,6 +192,16 @@ TMap<const UsdUtils::FUsdPrimMaterialSlot*, UMaterialInterface*> MeshTranslation
 			case UsdUtils::EPrimAssignmentType::UnrealMaterial:
 			{
 				Material = Cast< UMaterialInterface >( FSoftObjectPath( Slot.MaterialSource ).TryLoad() );
+				if ( !Material->IsTwoSided() && Slot.bMeshIsDoubleSided )
+				{
+					// TODO: Update this message with the proper source prim paths when UE-138122 is submitted,
+					// as 'UsdPrim' may just be e.g. a SkelRoot
+					UE_LOG( LogUsd, Warning, TEXT( "Using one-sided UE material '%s' for doubleSided prim '%s'" ),
+						*Slot.MaterialSource,
+						*UsdToUnreal::ConvertPath( UsdPrim.GetPrimPath() )
+					);
+				}
+
 				break;
 			}
 			case UsdUtils::EPrimAssignmentType::None:
@@ -246,4 +346,144 @@ void MeshTranslationImpl::SetMaterialOverrides(
 	}
 }
 
+UMaterialInterface* MeshTranslationImpl::GetBasePreviewSurfaceMaterial( EUsdBaseMaterialProperties BaseMaterialProperties )
+{
+	using namespace UE::MeshTranslationImplInternal::Private;
+
+	const bool bIsTranslucent = EnumHasAnyFlags( BaseMaterialProperties, EUsdBaseMaterialProperties::Translucent );
+	const bool bIsVT = EnumHasAnyFlags( BaseMaterialProperties, EUsdBaseMaterialProperties::VT );
+	const bool bIsTwoSided = EnumHasAnyFlags( BaseMaterialProperties, EUsdBaseMaterialProperties::TwoSided );
+
+	const FString* TargetMaterialPath = nullptr;
+	if ( bIsTranslucent )
+	{
+		if ( bIsVT )
+		{
+			if ( bIsTwoSided )
+			{
+				TargetMaterialPath = &BaseMaterialPathTranslucentTwoSidedVT;
+			}
+			else
+			{
+				TargetMaterialPath = &BaseMaterialPathTranslucentVT;
+			}
+		}
+		else
+		{
+			if ( bIsTwoSided )
+			{
+				TargetMaterialPath = &BaseMaterialPathTranslucentTwoSided;
+			}
+			else
+			{
+				TargetMaterialPath = &BaseMaterialPathTranslucent;
+			}
+		}
+	}
+	else
+	{
+		if ( bIsVT )
+		{
+			if ( bIsTwoSided )
+			{
+				TargetMaterialPath = &BaseMaterialPathTwoSidedVT;
+			}
+			else
+			{
+				TargetMaterialPath = &BaseMaterialPathVT;
+			}
+		}
+		else
+		{
+			if ( bIsTwoSided )
+			{
+				TargetMaterialPath = &BaseMaterialPathTwoSided;
+			}
+			else
+			{
+				TargetMaterialPath = &BaseMaterialPath;
+			}
+		}
+	}
+
+	if ( !TargetMaterialPath )
+	{
+		return nullptr;
+	}
+
+	return Cast< UMaterialInterface >( FSoftObjectPath( *TargetMaterialPath ).TryLoad() );
+}
+
+UMaterialInterface* MeshTranslationImpl::GetVTVersionOfBasePreviewSurfaceMaterial( UMaterialInterface* BaseMaterial )
+{
+	using namespace UE::MeshTranslationImplInternal::Private;
+
+	if ( !BaseMaterial )
+	{
+		return nullptr;
+	}
+
+	const FString PathName = BaseMaterial->GetPathName();
+	if ( PathName.Contains( TEXT( "VT" ), ESearchCase::CaseSensitive, ESearchDir::FromEnd ) )
+	{
+		return BaseMaterial;
+	}
+	else if ( PathName == BaseMaterialPath )
+	{
+		return Cast< UMaterialInterface >( FSoftObjectPath{ BaseMaterialPathVT }.TryLoad() );
+	}
+	else if ( PathName == BaseMaterialPathTwoSided )
+	{
+		return Cast< UMaterialInterface >( FSoftObjectPath{ BaseMaterialPathTwoSidedVT }.TryLoad() );
+	}
+	else if ( PathName == BaseMaterialPathTranslucent )
+	{
+		return Cast< UMaterialInterface >( FSoftObjectPath{ BaseMaterialPathTranslucentVT }.TryLoad() );
+	}
+	else if ( PathName == BaseMaterialPathTranslucentTwoSided )
+	{
+		return Cast< UMaterialInterface >( FSoftObjectPath{ BaseMaterialPathTranslucentTwoSidedVT }.TryLoad() );
+	}
+
+	// We should only ever call this function with a BaseMaterial that matches one of the above paths
+	ensure( false );
+	return nullptr;
+}
+
+UMaterialInterface* MeshTranslationImpl::GetTwoSidedVersionOfBasePreviewSurfaceMaterial( UMaterialInterface* BaseMaterial )
+{
+	using namespace UE::MeshTranslationImplInternal::Private;
+
+	if ( !BaseMaterial )
+	{
+		return nullptr;
+	}
+
+	const FString PathName = BaseMaterial->GetPathName();
+	if ( PathName.Contains( TEXT( "TwoSided" ), ESearchCase::CaseSensitive, ESearchDir::FromEnd ) )
+	{
+		return BaseMaterial;
+	}
+	else if ( PathName == BaseMaterialPath )
+	{
+		return Cast< UMaterialInterface >( FSoftObjectPath{ BaseMaterialPathTwoSided }.TryLoad() );
+	}
+	else if ( PathName == BaseMaterialPathTranslucent )
+	{
+		return Cast< UMaterialInterface >( FSoftObjectPath{ BaseMaterialPathTranslucentTwoSided }.TryLoad() );
+	}
+	else if ( PathName == BaseMaterialPathVT )
+	{
+		return Cast< UMaterialInterface >( FSoftObjectPath{ BaseMaterialPathTwoSidedVT }.TryLoad() );
+	}
+	else if ( PathName == BaseMaterialPathTranslucentVT )
+	{
+		return Cast< UMaterialInterface >( FSoftObjectPath{ BaseMaterialPathTranslucentTwoSidedVT }.TryLoad() );
+	}
+
+	// We should only ever call this function with a BaseMaterial that matches one of the above paths
+	ensure( false );
+	return nullptr;
+}
+
 #endif // #if USE_USD_SDK
\ No newline at end of file
diff --git a/Engine/Plugins/Importers/USDImporter/Source/USDSchemas/Private/MeshTranslationImpl.h b/Engine/Plugins/Importers/USDImporter/Source/USDSchemas/Private/MeshTranslationImpl.h
--- a/Engine/Plugins/Importers/USDImporter/Source/USDSchemas/Private/MeshTranslationImpl.h
+++ b/Engine/Plugins/Importers/USDImporter/Source/USDSchemas/Private/MeshTranslationImpl.h
@@ -52,6 +52,35 @@ namespace MeshTranslationImpl
 		const FName& RenderContext,
 		const FName& MaterialPurpose
 	);
+
+	enum class EUsdBaseMaterialProperties
+	{
+		None = 0,
+		Translucent = 1,
+		VT = 2,
+		TwoSided = 4
+	};
+	ENUM_CLASS_FLAGS( EUsdBaseMaterialProperties )
+
+	// Returns one of the alternatives of the UsdPreviewSurface base material depending on the material overrides
+	// provided, and nullptr otherwise
+	UMaterialInterface* GetBasePreviewSurfaceMaterial( EUsdBaseMaterialProperties BaseMaterialProperties );
+
+	// Returns the VT version of the provided UsdPreviewSurface BaseMaterial. Returns the provided BaseMaterial back if
+	// it is already a VT-capable base material, and returns nullptr if BaseMaterial isn't one of our base material
+	// alternatives.
+	// Example: Receives UsdPreviewSurfaceTwoSided -> Returns UsdPreviewSurfaceTwoSidedVT
+	// Example: Receives UsdPreviewSurfaceTwoSidedVT -> Returns UsdPreviewSurfaceTwoSidedVT
+	// Example: Receives SomeOtherBaseMaterial -> Returns nullptr
+	UMaterialInterface* GetVTVersionOfBasePreviewSurfaceMaterial( UMaterialInterface* BaseMaterial );
+
+	// Returns the two-sided version of the provided UsdPreviewSurface BaseMaterial. Returns the provided BaseMaterial
+	// back if it is already a two-sided-capable base material, and returns nullptr if BaseMaterial isn't one of our base
+	// material alternatives.
+	// Example: Receives UsdPreviewSurfaceTranslucent -> Returns UsdPreviewSurfaceTwoSidedTranslucent
+	// Example: Receives UsdPreviewSurfaceTwoSidedTranslucent -> Returns UsdPreviewSurfaceTwoSidedTranslucent
+	// Example: Receives SomeOtherBaseMaterial -> Returns nullptr
+	UMaterialInterface* GetTwoSidedVersionOfBasePreviewSurfaceMaterial( UMaterialInterface* BaseMaterial );
 }
 
 #endif // #if USE_USD_SDK
\ No newline at end of file
diff --git a/Engine/Plugins/Importers/USDImporter/Source/USDSchemas/Private/USDShadeMaterialTranslator.cpp b/Engine/Plugins/Importers/USDImporter/Source/USDSchemas/Private/USDShadeMaterialTranslator.cpp
--- a/Engine/Plugins/Importers/USDImporter/Source/USDSchemas/Private/USDShadeMaterialTranslator.cpp
+++ b/Engine/Plugins/Importers/USDImporter/Source/USDSchemas/Private/USDShadeMaterialTranslator.cpp
@@ -2,6 +2,7 @@
 
 #include "USDShadeMaterialTranslator.h"
 
+#include "MeshTranslationImpl.h"
 #include "USDAssetCache.h"
 #include "USDAssetImportData.h"
 #include "USDLog.h"
@@ -26,71 +27,6 @@
 
 namespace UE::UsdShadeTranslator::Private
 {
-	const static FString BaseMaterialPathTranslucentVT = TEXT( "/USDImporter/Materials/UsdPreviewSurfaceTranslucentVT.UsdPreviewSurfaceTranslucentVT" );
-	const static FString BaseMaterialPathTranslucent = TEXT( "/USDImporter/Materials/UsdPreviewSurfaceTranslucent.UsdPreviewSurfaceTranslucent" );
-	const static FString BaseMaterialPathVT = TEXT( "/USDImporter/Materials/UsdPreviewSurfaceVT.UsdPreviewSurfaceVT" );
-	const static FString BaseMaterialPath = TEXT( "/USDImporter/Materials/UsdPreviewSurface.UsdPreviewSurface" );
-
-	UMaterialInterface* GetBaseMaterial( bool bIsTranslucent, bool bIsVT )
-	{
-		const FString* TargetMaterialPath = nullptr;
-		if ( bIsTranslucent )
-		{
-			if ( bIsVT )
-			{
-				TargetMaterialPath = &BaseMaterialPathTranslucentVT;
-			}
-			else
-			{
-				TargetMaterialPath = &BaseMaterialPathTranslucent;
-			}
-		}
-		else
-		{
-			if ( bIsVT )
-			{
-				TargetMaterialPath = &BaseMaterialPathVT;
-			}
-			else
-			{
-				TargetMaterialPath = &BaseMaterialPath;
-			}
-		}
-
-		if ( !TargetMaterialPath )
-		{
-			return nullptr;
-		}
-
-		return Cast< UMaterialInterface >( FSoftObjectPath( *TargetMaterialPath ).TryLoad() );
-	}
-
-	UMaterialInterface* GetVTVersionOfBaseMaterial( UMaterialInterface* BaseMaterial )
-	{
-		if ( !BaseMaterial )
-		{
-			return nullptr;
-		}
-
-		const FString PathName = BaseMaterial->GetPathName();
-		if ( PathName == BaseMaterialPathTranslucentVT || PathName == BaseMaterialPathVT )
-		{
-			return BaseMaterial;
-		}
-		else if ( PathName == BaseMaterialPath )
-		{
-			return Cast< UMaterialInterface >( FSoftObjectPath{ BaseMaterialPathVT }.TryLoad() );
-		}
-		else if ( PathName == BaseMaterialPathTranslucent )
-		{
-			return Cast< UMaterialInterface >( FSoftObjectPath{ BaseMaterialPathTranslucentVT }.TryLoad() );
-		}
-
-		// We should only ever call this function with a BaseMaterial that matches one of the above paths
-		ensure( false );
-		return nullptr;
-	}
-
 	void RecursiveUpgradeMaterialsAndTexturesToVT(
 		const TSet<UTexture*>& TexturesToUpgrade,
 		const TSharedRef< FUsdSchemaTranslationContext >& Context,
@@ -131,7 +67,8 @@ namespace UE::UsdShadeTranslator::Private
 						// Important to not use GetBaseMaterial() here because if our parent is the translucent we'll
 						// get the base UsdPreviewSurface instead, as that is also *its* base
 						UMaterialInterface* BaseMaterial = MaterialInstance->Parent.Get();
-						UMaterialInterface* BaseMaterialVT = GetVTVersionOfBaseMaterial( BaseMaterial );
+						UMaterialInterface* BaseMaterialVT =
+							MeshTranslationImpl::GetVTVersionOfBasePreviewSurfaceMaterial( BaseMaterial );
 						if ( BaseMaterial == BaseMaterialVT )
 						{
 							// Material is already VT, we're good
@@ -355,11 +292,16 @@ void FUsdShadeMaterialTranslator::CreateAssets()
 					UE::UsdShadeTranslator::Private::UpgradeMaterialsAndTexturesToVT( NonVTTextures, Context );
 				}
 
-				const bool bNeedsVT = VTTextures.Num() > 0;
-				UMaterialInterface* BaseMaterial = UE::UsdShadeTranslator::Private::GetBaseMaterial(
-					bIsTranslucent,
-					bNeedsVT
-				);
+				MeshTranslationImpl::EUsdBaseMaterialProperties Properties = MeshTranslationImpl::EUsdBaseMaterialProperties::None;
+				if ( bIsTranslucent )
+				{
+					Properties |= MeshTranslationImpl::EUsdBaseMaterialProperties::Translucent;
+				}
+				if ( VTTextures.Num() > 0 )
+				{
+					Properties |= MeshTranslationImpl::EUsdBaseMaterialProperties::VT;
+				}
+				UMaterialInterface* BaseMaterial = MeshTranslationImpl::GetBasePreviewSurfaceMaterial( Properties );
 
 				if ( ensure( BaseMaterial ) )
 				{
@@ -395,11 +337,12 @@ void FUsdShadeMaterialTranslator::CreateAssets()
 			// using the VT base and copy the overrides. Not much else we can do as we need a base to call
 			// UMaterialInstanceDynamic::Create and get our instance, but we an instance to call UsdToUnreal::ConvertMaterial
 			// to create our textures and decide on our base.
-			const bool bNeedsVT = false;
-			UMaterialInterface* BaseMaterial = UE::UsdShadeTranslator::Private::GetBaseMaterial(
-				bIsTranslucent,
-				bNeedsVT
-			);
+			MeshTranslationImpl::EUsdBaseMaterialProperties Properties = MeshTranslationImpl::EUsdBaseMaterialProperties::None;
+			if ( bIsTranslucent )
+			{
+				Properties |= MeshTranslationImpl::EUsdBaseMaterialProperties::Translucent;
+			}
+			UMaterialInterface* BaseMaterial = MeshTranslationImpl::GetBasePreviewSurfaceMaterial( Properties );
 
 			if ( ensure( BaseMaterial ) )
 			{
diff --git a/Engine/Plugins/Importers/USDImporter/Source/USDStageImporter/Private/USDStageImporter.cpp b/Engine/Plugins/Importers/USDImporter/Source/USDStageImporter/Private/USDStageImporter.cpp
--- a/Engine/Plugins/Importers/USDImporter/Source/USDStageImporter/Private/USDStageImporter.cpp
+++ b/Engine/Plugins/Importers/USDImporter/Source/USDStageImporter/Private/USDStageImporter.cpp
@@ -503,8 +503,23 @@ namespace UsdStageImporterImpl
 			if (UUsdAssetImportData* AssetImportData = Cast<UUsdAssetImportData>(Material->AssetImportData))
 			{
 				// The only materials with no prim path are our auto-generated displayColor materials
-				AssetPath = AssetImportData->PrimPath.IsEmpty()? TEXT("DisplayColor") : AssetImportData->PrimPath;
-				AssetName = FPaths::GetBaseFilename( AssetPath );
+				if ( AssetImportData->PrimPath.IsEmpty() )
+				{
+					AssetPath = TEXT("DisplayColor");
+					AssetName = FPaths::GetBaseFilename( AssetPath );
+				}
+				else
+				{
+					AssetPath = AssetImportData->PrimPath;
+					AssetName = FPaths::GetBaseFilename( AssetPath );
+
+					// If we have a preview surface two-sided material we'll also have a one-sided with the same name,
+					// so add a suffix here so we can clearly tell which is which
+					if ( Material->IsTwoSided() )
+					{
+						AssetName += UnrealIdentifiers::TwoSidedMaterialSuffix;
+					}
+				}
 			}
 		}
 		else if (UTexture* Texture = Cast<UTexture>(Asset))
diff --git a/Engine/Plugins/Importers/USDImporter/Source/USDUtilities/Private/USDGeomMeshConversion.cpp b/Engine/Plugins/Importers/USDImporter/Source/USDUtilities/Private/USDGeomMeshConversion.cpp
--- a/Engine/Plugins/Importers/USDImporter/Source/USDUtilities/Private/USDGeomMeshConversion.cpp
+++ b/Engine/Plugins/Importers/USDImporter/Source/USDUtilities/Private/USDGeomMeshConversion.cpp
@@ -1599,6 +1599,19 @@ UsdUtils::FUsdPrimMaterialAssignmentInfo UsdUtils::GetPrimMaterialAssignments(
 		}
 	}
 
+	bool bIsDoubleSided = false;
+	if ( pxr::UsdGeomMesh Mesh = pxr::UsdGeomMesh{ UsdPrim } )
+	{
+		if ( pxr::UsdAttribute Attr = Mesh.GetDoubleSidedAttr() )
+		{
+			pxr::VtValue AttrValue;
+			if ( Attr.Get( &AttrValue ) && AttrValue.IsHolding<bool>() )
+			{
+				bIsDoubleSided = AttrValue.UncheckedGet<bool>();
+			}
+		}
+	}
+
 	// Priority 1: Material is an unreal asset
 	if ( RenderContext == UnrealIdentifiers::Unreal )
 	{
@@ -1613,6 +1626,7 @@ UsdUtils::FUsdPrimMaterialAssignmentInfo UsdUtils::GetPrimMaterialAssignments(
 					FUsdPrimMaterialSlot& Slot = Result.Slots.Emplace_GetRef();
 					Slot.MaterialSource = UnrealMaterial.GetValue();
 					Slot.AssignmentType = UsdUtils::EPrimAssignmentType::UnrealMaterial;
+					Slot.bMeshIsDoubleSided = bIsDoubleSided;
 
 					return Result;
 				}
@@ -1625,6 +1639,7 @@ UsdUtils::FUsdPrimMaterialAssignmentInfo UsdUtils::GetPrimMaterialAssignments(
 			FUsdPrimMaterialSlot& Slot = Result.Slots.Emplace_GetRef();
 			Slot.MaterialSource = UnrealMaterial.GetValue();
 			Slot.AssignmentType = UsdUtils::EPrimAssignmentType::UnrealMaterial;
+			Slot.bMeshIsDoubleSided = bIsDoubleSided;
 
 			return Result;
 		}
@@ -1636,6 +1651,7 @@ UsdUtils::FUsdPrimMaterialAssignmentInfo UsdUtils::GetPrimMaterialAssignments(
 		FUsdPrimMaterialSlot& Slot = Result.Slots.Emplace_GetRef();
 		Slot.MaterialSource = BoundMaterial.GetValue();
 		Slot.AssignmentType = UsdUtils::EPrimAssignmentType::MaterialPrim;
+		Slot.bMeshIsDoubleSided = bIsDoubleSided;
 
 		return Result;
 	}
@@ -1646,6 +1662,7 @@ UsdUtils::FUsdPrimMaterialAssignmentInfo UsdUtils::GetPrimMaterialAssignments(
 		FUsdPrimMaterialSlot& Slot = Result.Slots.Emplace_GetRef();
 		Slot.MaterialSource = TargetMaterial.GetValue();
 		Slot.AssignmentType = UsdUtils::EPrimAssignmentType::MaterialPrim;
+		Slot.bMeshIsDoubleSided = bIsDoubleSided;
 
 		return Result;
 	}
@@ -1688,6 +1705,7 @@ UsdUtils::FUsdPrimMaterialAssignmentInfo UsdUtils::GetPrimMaterialAssignments(
 								FUsdPrimMaterialSlot& Slot = Result.Slots.Emplace_GetRef();
 								Slot.MaterialSource = UnrealMaterial.GetValue();
 								Slot.AssignmentType = UsdUtils::EPrimAssignmentType::UnrealMaterial;
+								Slot.bMeshIsDoubleSided = bIsDoubleSided;
 								bHasAssignment = true;
 							}
 						}
@@ -1702,6 +1720,7 @@ UsdUtils::FUsdPrimMaterialAssignmentInfo UsdUtils::GetPrimMaterialAssignments(
 						FUsdPrimMaterialSlot& Slot = Result.Slots.Emplace_GetRef();
 						Slot.MaterialSource = UnrealMaterial.GetValue();
 						Slot.AssignmentType = UsdUtils::EPrimAssignmentType::UnrealMaterial;
+						Slot.bMeshIsDoubleSided = bIsDoubleSided;
 						bHasAssignment = true;
 					}
 				}
@@ -1715,6 +1734,7 @@ UsdUtils::FUsdPrimMaterialAssignmentInfo UsdUtils::GetPrimMaterialAssignments(
 					FUsdPrimMaterialSlot& Slot = Result.Slots.Emplace_GetRef();
 					Slot.MaterialSource = BoundMaterial.GetValue();
 					Slot.AssignmentType = UsdUtils::EPrimAssignmentType::MaterialPrim;
+					Slot.bMeshIsDoubleSided = bIsDoubleSided;
 					bHasAssignment = true;
 				}
 			}
@@ -1727,6 +1747,7 @@ UsdUtils::FUsdPrimMaterialAssignmentInfo UsdUtils::GetPrimMaterialAssignments(
 					FUsdPrimMaterialSlot& Slot = Result.Slots.Emplace_GetRef();
 					Slot.MaterialSource = TargetMaterial.GetValue();
 					Slot.AssignmentType = UsdUtils::EPrimAssignmentType::MaterialPrim;
+					Slot.bMeshIsDoubleSided = bIsDoubleSided;
 					bHasAssignment = true;
 				}
 			}
@@ -1739,6 +1760,7 @@ UsdUtils::FUsdPrimMaterialAssignmentInfo UsdUtils::GetPrimMaterialAssignments(
 				{
 					Slot.MaterialSource = DisplayColor.GetValue().ToString();
 					Slot.AssignmentType = UsdUtils::EPrimAssignmentType::DisplayColor;
+					Slot.bMeshIsDoubleSided = bIsDoubleSided;
 				}
 				bHasAssignment = true;
 			}
@@ -1764,6 +1786,7 @@ UsdUtils::FUsdPrimMaterialAssignmentInfo UsdUtils::GetPrimMaterialAssignments(
 			{
 				Slot.MaterialSource = DisplayColor.GetValue().ToString();
 				Slot.AssignmentType = UsdUtils::EPrimAssignmentType::DisplayColor;
+				Slot.bMeshIsDoubleSided = bIsDoubleSided;
 			}
 
 			if ( bProvideMaterialIndices )
@@ -1785,6 +1808,7 @@ UsdUtils::FUsdPrimMaterialAssignmentInfo UsdUtils::GetPrimMaterialAssignments(
 		FUsdPrimMaterialSlot& Slot = Result.Slots.Emplace_GetRef();
 		Slot.MaterialSource = DisplayColor.GetValue().ToString();
 		Slot.AssignmentType = UsdUtils::EPrimAssignmentType::DisplayColor;
+		Slot.bMeshIsDoubleSided = bIsDoubleSided;
 
 		return Result;
 	}
diff --git a/Engine/Plugins/Importers/USDImporter/Source/USDUtilities/Public/USDGeomMeshConversion.h b/Engine/Plugins/Importers/USDImporter/Source/USDUtilities/Public/USDGeomMeshConversion.h
--- a/Engine/Plugins/Importers/USDImporter/Source/USDUtilities/Public/USDGeomMeshConversion.h
+++ b/Engine/Plugins/Importers/USDImporter/Source/USDUtilities/Public/USDGeomMeshConversion.h
@@ -223,21 +223,31 @@ namespace UsdUtils
 		/** What this represents depends on AssignedMaterialType */
 		FString MaterialSource;
 		EPrimAssignmentType AssignmentType = EPrimAssignmentType::None;
+		bool bMeshIsDoubleSided = false;
 
 		friend bool operator==( const FUsdPrimMaterialSlot& Lhs, const FUsdPrimMaterialSlot& Rhs )
 		{
-			return Lhs.AssignmentType == Rhs.AssignmentType && Lhs.MaterialSource.Equals( Rhs.MaterialSource, ESearchCase::CaseSensitive );
+			return Lhs.AssignmentType == Rhs.AssignmentType &&
+				   Lhs.MaterialSource.Equals( Rhs.MaterialSource, ESearchCase::CaseSensitive ) &&
+				   Lhs.bMeshIsDoubleSided == Rhs.bMeshIsDoubleSided;
 		}
 
 		friend uint32 GetTypeHash( const FUsdPrimMaterialSlot& Slot )
 		{
-			return HashCombine( GetTypeHash( Slot.MaterialSource ), static_cast<uint32>(Slot.AssignmentType) );
+			return HashCombine(
+				Slot.bMeshIsDoubleSided,
+				HashCombine(
+					GetTypeHash( Slot.MaterialSource ),
+					static_cast<uint32>(Slot.AssignmentType)
+				)
+			);
 		}
 
 		friend FArchive& operator<<( FArchive& Ar, FUsdPrimMaterialSlot& Slot )
 		{
 			Ar << Slot.MaterialSource;
 			Ar << Slot.AssignmentType;
+			Ar << Slot.bMeshIsDoubleSided;
 
 			return Ar;
 		}
diff --git a/Engine/Plugins/Importers/USDImporter/Source/UnrealUSDWrapper/Private/UnrealUSDWrapper.cpp b/Engine/Plugins/Importers/USDImporter/Source/UnrealUSDWrapper/Private/UnrealUSDWrapper.cpp
--- a/Engine/Plugins/Importers/USDImporter/Source/UnrealUSDWrapper/Private/UnrealUSDWrapper.cpp
+++ b/Engine/Plugins/Importers/USDImporter/Source/UnrealUSDWrapper/Private/UnrealUSDWrapper.cpp
@@ -736,6 +736,7 @@ EUsdGeomOrientation IUsdPrim::GetGeometryOrientation(const pxr::UsdGeomMesh& Mes
 #endif // USE_USD_SDK
 
 const TCHAR* UnrealIdentifiers::LayerSavedComment = TEXT("unreal:layerSaved");
+const TCHAR* UnrealIdentifiers::TwoSidedMaterialSuffix = TEXT("_TwoSided");
 
 const TCHAR* UnrealIdentifiers::Invisible = TEXT("invisible");
 const TCHAR* UnrealIdentifiers::Inherited = TEXT("inherited");
diff --git a/Engine/Plugins/Importers/USDImporter/Source/UnrealUSDWrapper/Public/UnrealUSDWrapper.h b/Engine/Plugins/Importers/USDImporter/Source/UnrealUSDWrapper/Public/UnrealUSDWrapper.h
--- a/Engine/Plugins/Importers/USDImporter/Source/UnrealUSDWrapper/Public/UnrealUSDWrapper.h
+++ b/Engine/Plugins/Importers/USDImporter/Source/UnrealUSDWrapper/Public/UnrealUSDWrapper.h
@@ -460,6 +460,7 @@ namespace UnrealIdentifiers
 #endif // #if USE_USD_SDK
 
 	extern UNREALUSDWRAPPER_API const TCHAR* LayerSavedComment;
+	extern UNREALUSDWRAPPER_API const TCHAR* TwoSidedMaterialSuffix;
 
 	extern UNREALUSDWRAPPER_API const TCHAR* Invisible;
 	extern UNREALUSDWRAPPER_API const TCHAR* Inherited;
