USD: Fix missing scoped allocators causing crashes when switching levels.

diff --git a/Engine/Plugins/Importers/USDImporter/Source/UnrealUSDWrapper/Private/UsdWrappers/SdfLayer.cpp b/Engine/Plugins/Importers/USDImporter/Source/UnrealUSDWrapper/Private/UsdWrappers/SdfLayer.cpp
--- a/Engine/Plugins/Importers/USDImporter/Source/UnrealUSDWrapper/Private/UsdWrappers/SdfLayer.cpp
+++ b/Engine/Plugins/Importers/USDImporter/Source/UnrealUSDWrapper/Private/UsdWrappers/SdfLayer.cpp
@@ -235,6 +235,7 @@ namespace UE
 #if USE_USD_SDK
 		if ( const PtrType& Ptr = Impl->GetInner() )
 		{
+			FScopedUsdAllocs UsdAllocs;
 			return ANSI_TO_TCHAR( Ptr->GetComment().c_str() );
 		}
 #endif // #if USE_USD_SDK
@@ -247,6 +248,7 @@ namespace UE
 #if USE_USD_SDK
 		if ( const PtrType& Ptr = Impl->GetInner() )
 		{
+			FScopedUsdAllocs UsdAllocs;
 			Ptr->SetComment( TCHAR_TO_ANSI( Comment ) );
 		}
 #endif // #if USE_USD_SDK
