USD: Fix incorrect bake to control rig offset in case source animation playback range does not start at zero.

diff --git a/Engine/Plugins/Importers/USDImporter/Source/USDStage/Private/USDLevelSequenceHelper.cpp b/Engine/Plugins/Importers/USDImporter/Source/USDStage/Private/USDLevelSequenceHelper.cpp
--- a/Engine/Plugins/Importers/USDImporter/Source/USDStage/Private/USDLevelSequenceHelper.cpp
+++ b/Engine/Plugins/Importers/USDImporter/Source/USDStage/Private/USDLevelSequenceHelper.cpp
@@ -384,13 +384,20 @@ namespace UsdLevelSequenceHelperImpl
 				ControlRigSectionStartFrame = 0;
 			}
 
-			bool bSequencerOwnsControlRig = true;
-			UMovieSceneSection* NewSection = Track->CreateControlRigSection( ControlRigSectionStartFrame, ControlRig, bSequencerOwnsControlRig );
+			// This is unused
+			const FFrameNumber StartTime = 0;
+			const bool bSequencerOwnsControlRig = true;
+			UMovieSceneSection* NewSection = Track->CreateControlRigSection( StartTime, ControlRig, bSequencerOwnsControlRig );
 			UMovieSceneControlRigParameterSection* ParamSection = Cast<UMovieSceneControlRigParameterSection>( NewSection );
 
 			Track->SetTrackName( FName( *ObjectName ) );
 			Track->SetDisplayName( FText::FromString( ObjectName ) );
 
+			// We want the baked keyframe times to match exactly the source animation keyframe times. The way that
+			// LoadAnimSequenceIntoThisSection uses this however indicates it's meant to be relative to the start
+			// of the playback range, which we do here
+			ControlRigSectionStartFrame -= UE::MovieScene::DiscreteInclusiveLower( MovieScene->GetPlaybackRange() );
+
 			ParamSection->LoadAnimSequenceIntoThisSection( AnimSequence, MovieScene, SkeletalMeshComp, bReduceKeys, Tolerance, ControlRigSectionStartFrame );
 
 			// Disable Skeletal Animation Tracks
