USD: Fix assert when trying to update a reflection capture for a pending kill actor.

diff --git a/Engine/Plugins/Importers/USDImporter/Source/USDExporter/Private/LevelSequenceExporterUSD.cpp b/Engine/Plugins/Importers/USDImporter/Source/USDExporter/Private/LevelSequenceExporterUSD.cpp
--- a/Engine/Plugins/Importers/USDImporter/Source/USDExporter/Private/LevelSequenceExporterUSD.cpp
+++ b/Engine/Plugins/Importers/USDImporter/Source/USDExporter/Private/LevelSequenceExporterUSD.cpp
@@ -23,6 +23,7 @@
 #include "Compilation/MovieSceneCompiledDataManager.h"
 #include "Components/SkeletalMeshComponent.h"
 #include "Editor.h"
+#include "Engine/SphereReflectionCapture.h"
 #include "EngineAnalytics.h"
 #include "Evaluation/MovieSceneSequenceHierarchy.h"
 #include "ISequencer.h"
@@ -100,6 +101,15 @@ namespace UE
 
 				virtual UObject* SpawnObject( FMovieSceneSpawnable& Spawnable, FMovieSceneSequenceIDRef TemplateID, IMovieScenePlayer& Player ) override
 				{
+					// Never spawn ASphereReflectionCapture actors. These are useless in USD anyway, and we run into
+					// trouble after we're done exporting them because on the tick where they're destroyed the editor
+					// will still attempt to update their captures and some downstream code doesn't like that their
+					// components are pending kill (check UE-167593 for more info)
+					if ( const ASphereReflectionCapture* ReflectionCapture = Cast< const ASphereReflectionCapture >( Spawnable.GetObjectTemplate() ) )
+					{
+						return nullptr;
+					}
+
 					ISequencer* Sequencer = static_cast<ISequencer*>(&Player);
 					UMovieSceneSequence* RootSequence = Sequencer->GetRootMovieSceneSequence();
 					if ( !RootSequence )
@@ -635,6 +645,13 @@ namespace UE
 						FMovieSceneSpawnable& Spawnable = MovieScene->GetSpawnable( Index );
 						const FGuid& Guid = Spawnable.GetGuid();
 
+						// We won't have spawned ASphereReflectionCapture here.
+						// See the comment inside FLevelSequenceHidingSpawnRegister::SpawnObject and UE-167593 for more info
+						if ( ASphereReflectionCapture* ReflectionCapture = Cast< ASphereReflectionCapture>( Spawnable.GetObjectTemplate() ) )
+						{
+							continue;
+						}
+
 						UObject* BoundObject = Context.SpawnRegister->GetExistingSpawn( *RootSequence, SequenceInstance, Guid );
 						if ( !BoundObject )
 						{
