USD: Fix skeletal meshes generated from multiple Mesh prims ending up with incorrect material assignments in some scenarios.

diff --git a/Engine/Plugins/Importers/USDImporter/Source/USDUtilities/Private/USDSkeletalDataConversion.cpp b/Engine/Plugins/Importers/USDImporter/Source/USDUtilities/Private/USDSkeletalDataConversion.cpp
--- a/Engine/Plugins/Importers/USDImporter/Source/USDUtilities/Private/USDSkeletalDataConversion.cpp
+++ b/Engine/Plugins/Importers/USDImporter/Source/USDUtilities/Private/USDSkeletalDataConversion.cpp
@@ -1634,10 +1634,10 @@ bool UsdToUnreal::ConvertSkinnedMesh(
 
 		// Manage materials
 		int32 LocalMaterialIndex = 0;
-		if ( FaceMaterialIndices.IsValidIndex( PolygonIndex ) )
+		if ( FaceMaterialIndices.IsValidIndex( LocalIndex ) )
 		{
-			LocalMaterialIndex = FaceMaterialIndices[ PolygonIndex ];
-			if ( !LocalMaterialSlots.IsValidIndex(LocalMaterialIndex) )
+			LocalMaterialIndex = FaceMaterialIndices[ LocalIndex ];
+			if ( !LocalMaterialSlots.IsValidIndex( LocalMaterialIndex ) )
 			{
 				LocalMaterialIndex = 0;
 			}
