USD: Support tracking material overrides in collapsed meshes.

diff --git a/Engine/Plugins/Importers/USDImporter/Source/USDClasses/Public/USDAssetImportData.h b/Engine/Plugins/Importers/USDImporter/Source/USDClasses/Public/USDAssetImportData.h
--- a/Engine/Plugins/Importers/USDImporter/Source/USDClasses/Public/USDAssetImportData.h
+++ b/Engine/Plugins/Importers/USDImporter/Source/USDClasses/Public/USDAssetImportData.h
@@ -37,4 +37,30 @@ public:
 	// This should be applied *after* any offset/scale conversions on the time coordinate.
 	UPROPERTY()
 	float LayerStartOffsetSeconds = 0.0f;
+};
+
+/** Simple wrapper because we're not allowed to have TMap properties with TArray<FString> as values */
+USTRUCT()
+struct FUsdPrimPathList
+{
+	GENERATED_BODY()
+
+	UPROPERTY()
+	TArray<FString> PrimPaths;
+};
+
+/** We assign these to UStaticMeshes or USkeletalMeshes generated from USD */
+UCLASS( config = EditorPerProjectUserSettings, AutoExpandCategories = ( Options ), MinimalAPI )
+class UUsdMeshAssetImportData : public UUsdAssetImportData
+{
+	GENERATED_BODY()
+
+public:
+	/**
+	 * Maps from a material slot index of an UStaticMesh or USkeletalMesh to a list of source prims that contain this
+	 * assignment. It can contain multiple prims in case we combine material slots and/or collapse prims
+	 * (e.g. {0: ['/Root/mesh', '/Root/othermesh/geomsubset0', '/Root/othermesh/geomsubset1'] }).
+	 */
+	UPROPERTY()
+	TMap< int32, FUsdPrimPathList > MaterialSlotToPrimPaths;
 };
\ No newline at end of file
diff --git a/Engine/Plugins/Importers/USDImporter/Source/USDSchemas/Private/MeshTranslationImpl.cpp b/Engine/Plugins/Importers/USDImporter/Source/USDSchemas/Private/MeshTranslationImpl.cpp
--- a/Engine/Plugins/Importers/USDImporter/Source/USDSchemas/Private/MeshTranslationImpl.cpp
+++ b/Engine/Plugins/Importers/USDImporter/Source/USDSchemas/Private/MeshTranslationImpl.cpp
@@ -44,9 +44,7 @@ namespace UE::MeshTranslationImplInternal::Private
 TMap<const UsdUtils::FUsdPrimMaterialSlot*, UMaterialInterface*> MeshTranslationImpl::ResolveMaterialAssignmentInfo(
 	const pxr::UsdPrim& UsdPrim,
 	const TArray<UsdUtils::FUsdPrimMaterialAssignmentInfo>& AssignmentInfo,
-	const TArray<UMaterialInterface*>& ExistingAssignments,
 	UUsdAssetCache& AssetCache,
-	float Time,
 	EObjectFlags Flags
 )
 {
@@ -307,9 +305,7 @@ void MeshTranslationImpl::SetMaterialOverrides(
 	TMap<const UsdUtils::FUsdPrimMaterialSlot*, UMaterialInterface*> ResolvedMaterials = MeshTranslationImpl::ResolveMaterialAssignmentInfo(
 		ValidPrim,
 		LODIndexToAssignments,
-		ExistingAssignments,
 		AssetCache,
-		Time,
 		Flags
 	);
 
diff --git a/Engine/Plugins/Importers/USDImporter/Source/USDSchemas/Private/MeshTranslationImpl.h b/Engine/Plugins/Importers/USDImporter/Source/USDSchemas/Private/MeshTranslationImpl.h
--- a/Engine/Plugins/Importers/USDImporter/Source/USDSchemas/Private/MeshTranslationImpl.h
+++ b/Engine/Plugins/Importers/USDImporter/Source/USDSchemas/Private/MeshTranslationImpl.h
@@ -27,13 +27,11 @@ namespace UsdUtils
 /** Implementation that can be shared between the SkelRoot translator and GeomMesh translators */
 namespace MeshTranslationImpl
 {
-	/** Retrieves the target materials described on AssignmentInfo, considering that the previous material assignment on the mesh was ExistingAssignments */
+	/** Resolves the material assignments in AssignmentInfo, returning an UMaterialInterface for each material slot */
 	TMap<const UsdUtils::FUsdPrimMaterialSlot*, UMaterialInterface*> ResolveMaterialAssignmentInfo(
 		const pxr::UsdPrim& UsdPrim,
 		const TArray<UsdUtils::FUsdPrimMaterialAssignmentInfo>& AssignmentInfo,
-		const TArray<UMaterialInterface*>& ExistingAssignments,
 		UUsdAssetCache& AssetCache,
-		float Time,
 		EObjectFlags Flags
 	);
 
diff --git a/Engine/Plugins/Importers/USDImporter/Source/USDSchemas/Private/USDGeomMeshTranslator.cpp b/Engine/Plugins/Importers/USDImporter/Source/USDSchemas/Private/USDGeomMeshTranslator.cpp
--- a/Engine/Plugins/Importers/USDImporter/Source/USDSchemas/Private/USDGeomMeshTranslator.cpp
+++ b/Engine/Plugins/Importers/USDImporter/Source/USDSchemas/Private/USDGeomMeshTranslator.cpp
@@ -228,7 +228,16 @@ namespace UsdGeomMeshTranslatorImpl
 			ExistingAssignments.Add(StaticMaterial.MaterialInterface);
 		}
 
-		TMap<const UsdUtils::FUsdPrimMaterialSlot*, UMaterialInterface*> ResolvedMaterials = MeshTranslationImpl::ResolveMaterialAssignmentInfo(UsdPrim, LODIndexToMaterialInfo, ExistingAssignments, AssetCache, Time, Flags );
+#if WITH_EDITOR
+		UUsdMeshAssetImportData* ImportData = Cast< UUsdMeshAssetImportData>( StaticMesh.AssetImportData.Get() );
+#endif // WITH_EDITOR
+
+		TMap<const UsdUtils::FUsdPrimMaterialSlot*, UMaterialInterface*> ResolvedMaterials = MeshTranslationImpl::ResolveMaterialAssignmentInfo(
+			UsdPrim,
+			LODIndexToMaterialInfo,
+			AssetCache,
+			Flags
+		);
 
 		uint32 StaticMeshSlotIndex = 0;
 		for ( int32 LODIndex = 0; LODIndex < LODIndexToMaterialInfo.Num(); ++LODIndex )
@@ -264,6 +273,11 @@ namespace UsdGeomMeshTranslatorImpl
 				}
 
 #if WITH_EDITOR
+				if ( ImportData )
+				{
+					ImportData->MaterialSlotToPrimPaths.FindOrAdd( StaticMeshSlotIndex ).PrimPaths = Slot.PrimPaths.Array();
+				}
+
 				// Setup the section map so that our LOD material index is properly mapped to the static mesh material index
 				// At runtime we don't ever parse these variants as LODs so we don't need this
 				if ( StaticMesh.GetSectionInfoMap().IsValidSection( LODIndex, LODSlotIndex ) )
@@ -304,6 +318,10 @@ namespace UsdGeomMeshTranslatorImpl
 	{
 		bool bMaterialAssignementsHaveChanged = false;
 
+#if WITH_EDITOR
+		UUsdMeshAssetImportData* MeshImportData = Cast< UUsdMeshAssetImportData>( GeometryCache.AssetImportData.Get() );
+#endif // WITH_EDITOR
+
 		uint32 StaticMeshSlotIndex = 0;
 		for ( int32 LODIndex = 0; LODIndex < LODIndexToMaterialInfo.Num(); ++LODIndex )
 		{
@@ -432,6 +450,10 @@ namespace UsdGeomMeshTranslatorImpl
 					GeometryCache.Materials[ StaticMeshSlotIndex ] = Material;
 					bMaterialAssignementsHaveChanged = true;
 				}
+
+#if WITH_EDITOR
+				MeshImportData->MaterialSlotToPrimPaths.FindOrAdd( StaticMeshSlotIndex ).PrimPaths = Slot.PrimPaths.Array();
+#endif // WITH_EDITOR
 			}
 		}
 
@@ -497,6 +519,7 @@ namespace UsdGeomMeshTranslatorImpl
 				if ( bSuccess )
 				{
 					OptionsCopy.AdditionalTransform = MeshTransform * Options.AdditionalTransform;
+					OptionsCopy.bMergeIdenticalMaterialSlots = false;  // We only merge slots when collapsing, and we never collapse LODs
 
 					bSuccess &= UsdToUnreal::ConvertGeomMesh(
 						LODMesh,
@@ -584,6 +607,7 @@ namespace UsdGeomMeshTranslatorImpl
 			{
 				UsdToUnreal::FUsdMeshConversionOptions OptionsCopy = Options;
 				OptionsCopy.AdditionalTransform = MeshTransform * Options.AdditionalTransform;
+				OptionsCopy.bMergeIdenticalMaterialSlots = false;  // We only merge for collapsed meshes
 
 				bSuccess &= UsdToUnreal::ConvertGeomMesh(
 					pxr::UsdGeomMesh{ UsdMesh },
@@ -1181,7 +1205,7 @@ void FBuildStaticMeshTaskChain::SetupTasks()
 
 				if ( bIsNew )
 				{
-					UUsdAssetImportData* ImportData = NewObject<UUsdAssetImportData>( StaticMesh, TEXT( "UUSDAssetImportData" ) );
+					UUsdMeshAssetImportData* ImportData = NewObject<UUsdMeshAssetImportData>( StaticMesh, TEXT( "UUSDAssetImportData" ) );
 					ImportData->PrimPath = PrimPathString;
 					StaticMesh->AssetImportData = ImportData;
 				}
@@ -1431,7 +1455,7 @@ void FGeometryCacheCreateAssetsTaskChain::SetupTasks()
 
 			if ( bIsNew && GeometryCache )
 			{
-				UUsdAssetImportData* ImportData = NewObject< UUsdAssetImportData >( GeometryCache, TEXT( "UUSDAssetImportData" ) );
+				UUsdMeshAssetImportData* ImportData = NewObject< UUsdMeshAssetImportData >( GeometryCache, TEXT( "UUSDAssetImportData" ) );
 				ImportData->PrimPath = PrimPathString;
 				GeometryCache->AssetImportData = ImportData;
 			}
@@ -1442,7 +1466,7 @@ void FGeometryCacheCreateAssetsTaskChain::SetupTasks()
 				Context->AssetCache->LinkAssetToPrim( PrimPathString, GeometryCache );
 
 				// Only process the materials if we own the GeometryCache. If it's new we know we do
-				UUsdAssetImportData* ImportData = Cast< UUsdAssetImportData >( GeometryCache->AssetImportData );
+				UUsdMeshAssetImportData* ImportData = Cast< UUsdMeshAssetImportData >( GeometryCache->AssetImportData );
 				if ( ImportData && ImportData->PrimPath == PrimPathString )
 				{
 					bMaterialsHaveChanged = UsdGeomMeshTranslatorImpl::ProcessGeometryCacheMaterials( GetPrim(), LODIndexToMaterialInfo, *GeometryCache, *Context->AssetCache.Get(), Context->Time, Context->ObjectFlags );
diff --git a/Engine/Plugins/Importers/USDImporter/Source/USDSchemas/Private/USDSkelRootTranslator.cpp b/Engine/Plugins/Importers/USDImporter/Source/USDSchemas/Private/USDSkelRootTranslator.cpp
--- a/Engine/Plugins/Importers/USDImporter/Source/USDSchemas/Private/USDSkelRootTranslator.cpp
+++ b/Engine/Plugins/Importers/USDImporter/Source/USDSchemas/Private/USDSkelRootTranslator.cpp
@@ -95,7 +95,16 @@ namespace UsdSkelRootTranslatorImpl
 			ExistingAssignments.Add( SkeletalMaterial.MaterialInterface );
 		}
 
-		TMap<const UsdUtils::FUsdPrimMaterialSlot*, UMaterialInterface*> ResolvedMaterials = MeshTranslationImpl::ResolveMaterialAssignmentInfo( UsdPrim, LODIndexToMaterialInfo, ExistingAssignments, AssetCache, Time, Flags );
+#if WITH_EDITOR
+		UUsdMeshAssetImportData* ImportData = Cast< UUsdMeshAssetImportData>( SkeletalMesh->GetAssetImportData() );
+#endif // WITH_EDITOR
+
+		TMap<const UsdUtils::FUsdPrimMaterialSlot*, UMaterialInterface*> ResolvedMaterials = MeshTranslationImpl::ResolveMaterialAssignmentInfo(
+			UsdPrim,
+			LODIndexToMaterialInfo,
+			AssetCache,
+			Flags
+		);
 
 		bool bMaterialsHaveChanged = false;
 
@@ -166,6 +175,13 @@ namespace UsdSkelRootTranslatorImpl
 					bMaterialsHaveChanged = true;
 				}
 
+#if WITH_EDITOR
+				if ( ImportData )
+				{
+					ImportData->MaterialSlotToPrimPaths.FindOrAdd( SkeletalMeshSlotIndex ).PrimPaths = Slot.PrimPaths.Array();
+				}
+#endif // WITH_EDITOR
+
 				// Already have a material at that LOD remap slot, need to reassign
 				if ( LODMaterialMap.IsValidIndex( LODSlotIndex ) )
 				{
@@ -709,9 +725,7 @@ namespace UsdSkelRootTranslatorImpl
 		TMap<const UsdUtils::FUsdPrimMaterialSlot*, UMaterialInterface*> ResolvedMaterials = MeshTranslationImpl::ResolveMaterialAssignmentInfo(
 			ValidSkelRootPrim,
 			LODIndexToAssignments,
-			ExistingAssignments,
 			AssetCache,
-			Time,
 			Flags
 		);
 
@@ -1241,6 +1255,10 @@ namespace UsdSkelRootTranslatorImpl
 				{
 					if ( bIsNew )
 					{
+						UUsdMeshAssetImportData* ImportData = NewObject< UUsdMeshAssetImportData >( SkeletalMesh, TEXT( "USDAssetImportData" ) );
+						ImportData->PrimPath = SkelRootPath;
+						SkeletalMesh->SetAssetImportData(ImportData);
+
 						const bool bMaterialsHaveChanged = UsdSkelRootTranslatorImpl::ProcessMaterials(
 							GetPrim(),
 							LODIndexToMaterialInfo,
@@ -1257,10 +1275,6 @@ namespace UsdSkelRootTranslatorImpl
 							SkeletalMesh->UpdateUVChannelData( bRebuildAll );
 						}
 
-						UUsdAssetImportData* ImportData = NewObject< UUsdAssetImportData >( SkeletalMesh, TEXT( "USDAssetImportData" ) );
-						ImportData->PrimPath = SkelRootPath;
-						SkeletalMesh->SetAssetImportData(ImportData);
-
 						Context->AssetCache->CacheAsset( SkeletalMeshHashString, SkeletalMesh );
 						Context->AssetCache->CacheAsset( SkeletalMeshHashString + TEXT( "_Skeleton" ), SkeletalMesh->GetSkeleton() );
 					}
diff --git a/Engine/Plugins/Importers/USDImporter/Source/USDUtilities/Private/USDGeomMeshConversion.cpp b/Engine/Plugins/Importers/USDImporter/Source/USDUtilities/Private/USDGeomMeshConversion.cpp
--- a/Engine/Plugins/Importers/USDImporter/Source/USDUtilities/Private/USDGeomMeshConversion.cpp
+++ b/Engine/Plugins/Importers/USDImporter/Source/USDUtilities/Private/USDGeomMeshConversion.cpp
@@ -678,22 +678,50 @@ bool UsdToUnreal::ConvertGeomMesh(
 		for ( int32 Index = 0; Index < OutMaterialAssignments.Slots.Num(); ++Index )
 		{
 			const UsdUtils::FUsdPrimMaterialSlot& Slot = OutMaterialAssignments.Slots[ Index ];
-			CombinedMaterialSlotsToIndex.Add( Slot, Index );
+
+			// Combine entries in this way so that we can append PrimPaths
+			TMap< UsdUtils::FUsdPrimMaterialSlot, int32 >::TKeyIterator KeyIt = CombinedMaterialSlotsToIndex.CreateKeyIterator( Slot );
+			if ( KeyIt )
+			{
+				KeyIt.Key().PrimPaths.Append( Slot.PrimPaths );
+				KeyIt.Value() = Index;
+			}
+			else
+			{
+				CombinedMaterialSlotsToIndex.Add( Slot, Index );
+			}
 		}
 
+		// Combine our LocalSlots into CombinedMaterialSlotsToIndex
 		for ( int32 LocalIndex = 0; LocalIndex < LocalInfo.Slots.Num(); ++LocalIndex )
 		{
 			const UsdUtils::FUsdPrimMaterialSlot& LocalSlot = LocalInfo.Slots[ LocalIndex ];
-			if ( int32* ExistingCombinedIndex = CombinedMaterialSlotsToIndex.Find( LocalSlot ) )
+
+			// Combine entries in this way so that we can append PrimPaths
+			TMap< UsdUtils::FUsdPrimMaterialSlot, int32 >::TKeyIterator KeyIt = CombinedMaterialSlotsToIndex.CreateKeyIterator( LocalSlot );
+			if ( KeyIt )
 			{
-				LocalToCombinedMaterialSlotIndices[ LocalIndex ] = *ExistingCombinedIndex;
+				KeyIt.Key().PrimPaths.Append( LocalSlot.PrimPaths );
+
+				const int32 ExistingCombinedIndex = KeyIt.Value();
+				LocalToCombinedMaterialSlotIndices[ LocalIndex ] = ExistingCombinedIndex;
 			}
 			else
 			{
-				OutMaterialAssignments.Slots.Add( LocalSlot );
-				LocalToCombinedMaterialSlotIndices[ LocalIndex ] = OutMaterialAssignments.Slots.Num() - 1;
+				int32 NewIndex = OutMaterialAssignments.Slots.Add( LocalSlot );
+				CombinedMaterialSlotsToIndex.Add( LocalSlot, NewIndex );
+				LocalToCombinedMaterialSlotIndices[ LocalIndex ] = NewIndex;
 			}
 		}
+
+		// Now that we merged all prim paths into they keys of CombinedMaterialSlotsToIndex, let's copy them back into
+		// our output
+		for ( UsdUtils::FUsdPrimMaterialSlot& Slot : OutMaterialAssignments.Slots )
+		{
+			TMap< UsdUtils::FUsdPrimMaterialSlot, int32 >::TKeyIterator KeyIt = CombinedMaterialSlotsToIndex.CreateKeyIterator( Slot );
+			ensure(KeyIt);
+			Slot.PrimPaths = KeyIt.Key().PrimPaths;
+		}
 	}
 	else
 	{
@@ -1612,6 +1640,8 @@ UsdUtils::FUsdPrimMaterialAssignmentInfo UsdUtils::GetPrimMaterialAssignments(
 		}
 	}
 
+	FString MeshPrimPath = UsdToUnreal::ConvertPath( UsdPrim.GetPath() );
+
 	// Priority 1: Material is an unreal asset
 	if ( RenderContext == UnrealIdentifiers::Unreal )
 	{
@@ -1627,6 +1657,7 @@ UsdUtils::FUsdPrimMaterialAssignmentInfo UsdUtils::GetPrimMaterialAssignments(
 					Slot.MaterialSource = UnrealMaterial.GetValue();
 					Slot.AssignmentType = UsdUtils::EPrimAssignmentType::UnrealMaterial;
 					Slot.bMeshIsDoubleSided = bIsDoubleSided;
+					Slot.PrimPaths.Add( MeshPrimPath );
 
 					return Result;
 				}
@@ -1640,6 +1671,7 @@ UsdUtils::FUsdPrimMaterialAssignmentInfo UsdUtils::GetPrimMaterialAssignments(
 			Slot.MaterialSource = UnrealMaterial.GetValue();
 			Slot.AssignmentType = UsdUtils::EPrimAssignmentType::UnrealMaterial;
 			Slot.bMeshIsDoubleSided = bIsDoubleSided;
+			Slot.PrimPaths.Add( MeshPrimPath );
 
 			return Result;
 		}
@@ -1652,6 +1684,7 @@ UsdUtils::FUsdPrimMaterialAssignmentInfo UsdUtils::GetPrimMaterialAssignments(
 		Slot.MaterialSource = BoundMaterial.GetValue();
 		Slot.AssignmentType = UsdUtils::EPrimAssignmentType::MaterialPrim;
 		Slot.bMeshIsDoubleSided = bIsDoubleSided;
+		Slot.PrimPaths.Add( MeshPrimPath );
 
 		return Result;
 	}
@@ -1663,6 +1696,7 @@ UsdUtils::FUsdPrimMaterialAssignmentInfo UsdUtils::GetPrimMaterialAssignments(
 		Slot.MaterialSource = TargetMaterial.GetValue();
 		Slot.AssignmentType = UsdUtils::EPrimAssignmentType::MaterialPrim;
 		Slot.bMeshIsDoubleSided = bIsDoubleSided;
+		Slot.PrimPaths.Add( MeshPrimPath );
 
 		return Result;
 	}
@@ -1689,13 +1723,16 @@ UsdUtils::FUsdPrimMaterialAssignmentInfo UsdUtils::GetPrimMaterialAssignments(
 			const pxr::UsdGeomSubset& GeomSubset = GeomSubsets[ GeomSubsetIndex ];
 			bool bHasAssignment = false;
 
+			pxr::UsdPrim GeomSubsetPrim = GeomSubset.GetPrim();
+			FString GeomSubsetPath = UsdToUnreal::ConvertPath( GeomSubsetPrim.GetPath() );
+
 			// Priority 4.1: Material is an unreal asset
 			if ( RenderContext == UnrealIdentifiers::Unreal )
 			{
 				// Priority 4.1.1: Partition has an unreal rendercontext material prim binding
 				if ( !bHasAssignment )
 				{
-					pxr::UsdShadeMaterialBindingAPI BindingAPI( GeomSubset.GetPrim() );
+					pxr::UsdShadeMaterialBindingAPI BindingAPI( GeomSubsetPrim );
 					if ( BindingAPI )
 					{
 						if ( pxr::UsdShadeMaterial ShadeMaterial = BindingAPI.ComputeBoundMaterial( MaterialPurpose ) )
@@ -1706,6 +1743,7 @@ UsdUtils::FUsdPrimMaterialAssignmentInfo UsdUtils::GetPrimMaterialAssignments(
 								Slot.MaterialSource = UnrealMaterial.GetValue();
 								Slot.AssignmentType = UsdUtils::EPrimAssignmentType::UnrealMaterial;
 								Slot.bMeshIsDoubleSided = bIsDoubleSided;
+								Slot.PrimPaths.Add( GeomSubsetPath );
 								bHasAssignment = true;
 							}
 						}
@@ -1715,12 +1753,13 @@ UsdUtils::FUsdPrimMaterialAssignmentInfo UsdUtils::GetPrimMaterialAssignments(
 				// Priority 4.1.2: Partitition has an unrealMaterial attribute directly on it
 				if ( !bHasAssignment )
 				{
-					if ( TOptional<FString> UnrealMaterial = FetchFirstUEMaterialFromAttribute( GeomSubset.GetPrim(), TimeCode ) )
+					if ( TOptional<FString> UnrealMaterial = FetchFirstUEMaterialFromAttribute( GeomSubsetPrim, TimeCode ) )
 					{
 						FUsdPrimMaterialSlot& Slot = Result.Slots.Emplace_GetRef();
 						Slot.MaterialSource = UnrealMaterial.GetValue();
 						Slot.AssignmentType = UsdUtils::EPrimAssignmentType::UnrealMaterial;
 						Slot.bMeshIsDoubleSided = bIsDoubleSided;
+						Slot.PrimPaths.Add( GeomSubsetPath );
 						bHasAssignment = true;
 					}
 				}
@@ -1729,12 +1768,13 @@ UsdUtils::FUsdPrimMaterialAssignmentInfo UsdUtils::GetPrimMaterialAssignments(
 			// Priority 4.2: computing bound material
 			if ( !bHasAssignment )
 			{
-				if ( TOptional<FString> BoundMaterial = FetchMaterialByComputingBoundMaterial( GeomSubset.GetPrim() ) )
+				if ( TOptional<FString> BoundMaterial = FetchMaterialByComputingBoundMaterial( GeomSubsetPrim ) )
 				{
 					FUsdPrimMaterialSlot& Slot = Result.Slots.Emplace_GetRef();
 					Slot.MaterialSource = BoundMaterial.GetValue();
 					Slot.AssignmentType = UsdUtils::EPrimAssignmentType::MaterialPrim;
 					Slot.bMeshIsDoubleSided = bIsDoubleSided;
+					Slot.PrimPaths.Add( GeomSubsetPath );
 					bHasAssignment = true;
 				}
 			}
@@ -1742,12 +1782,13 @@ UsdUtils::FUsdPrimMaterialAssignmentInfo UsdUtils::GetPrimMaterialAssignments(
 			// Priority 4.3: material:binding relationship
 			if ( !bHasAssignment )
 			{
-				if ( TOptional<FString> TargetMaterial = FetchMaterialByMaterialRelationship( GeomSubset.GetPrim() ) )
+				if ( TOptional<FString> TargetMaterial = FetchMaterialByMaterialRelationship( GeomSubsetPrim ) )
 				{
 					FUsdPrimMaterialSlot& Slot = Result.Slots.Emplace_GetRef();
 					Slot.MaterialSource = TargetMaterial.GetValue();
 					Slot.AssignmentType = UsdUtils::EPrimAssignmentType::MaterialPrim;
 					Slot.bMeshIsDoubleSided = bIsDoubleSided;
+					Slot.PrimPaths.Add( GeomSubsetPath );
 					bHasAssignment = true;
 				}
 			}
@@ -1762,6 +1803,7 @@ UsdUtils::FUsdPrimMaterialAssignmentInfo UsdUtils::GetPrimMaterialAssignments(
 					Slot.AssignmentType = UsdUtils::EPrimAssignmentType::DisplayColor;
 					Slot.bMeshIsDoubleSided = bIsDoubleSided;
 				}
+				Slot.PrimPaths.Add( GeomSubsetPath );
 				bHasAssignment = true;
 			}
 
@@ -1809,6 +1851,7 @@ UsdUtils::FUsdPrimMaterialAssignmentInfo UsdUtils::GetPrimMaterialAssignments(
 		Slot.MaterialSource = DisplayColor.GetValue().ToString();
 		Slot.AssignmentType = UsdUtils::EPrimAssignmentType::DisplayColor;
 		Slot.bMeshIsDoubleSided = bIsDoubleSided;
+		Slot.PrimPaths.Add( MeshPrimPath );
 
 		return Result;
 	}
diff --git a/Engine/Plugins/Importers/USDImporter/Source/USDUtilities/Private/USDPrimConversion.cpp b/Engine/Plugins/Importers/USDImporter/Source/USDUtilities/Private/USDPrimConversion.cpp
--- a/Engine/Plugins/Importers/USDImporter/Source/USDUtilities/Private/USDPrimConversion.cpp
+++ b/Engine/Plugins/Importers/USDImporter/Source/USDUtilities/Private/USDPrimConversion.cpp
@@ -3,6 +3,7 @@
 #include "USDPrimConversion.h"
 
 #include "UnrealUSDWrapper.h"
+#include "USDAssetImportData.h"
 #include "USDAttributeUtils.h"
 #include "USDConversionUtils.h"
 #include "USDLayerUtils.h"
@@ -2126,30 +2127,48 @@ bool UnrealToUsd::ConvertMeshComponent( const pxr::UsdStageRefPtr& Stage, const
 
 	FScopedUsdAllocs Allocs;
 
+#if WITH_EDITOR
+	// Handle material overrides
 	if ( const UGeometryCacheComponent* GeometryCacheComponent = Cast<const UGeometryCacheComponent>( MeshComponent ) )
 	{
-		const TArray<UMaterialInterface*>& Overrides = MeshComponent->OverrideMaterials;
-		for ( int32 MatIndex = 0; MatIndex < Overrides.Num(); ++MatIndex )
+		if ( const UGeometryCache* GeometryCache = GeometryCacheComponent->GetGeometryCache() )
 		{
-			UMaterialInterface* Override = Overrides[ MatIndex ];
-			if ( !Override )
+			const TArray<UMaterialInterface*>& Overrides = MeshComponent->OverrideMaterials;
+			for ( int32 MatIndex = 0; MatIndex < Overrides.Num(); ++MatIndex )
 			{
-				continue;
-			}
+				UMaterialInterface* Override = Overrides[ MatIndex ];
+				if ( !Override )
+				{
+					continue;
+				}
 
-			pxr::SdfPath OverridePrimPath = UsdPrim.GetPath();
+				UUsdMeshAssetImportData* ImportData = Cast<UUsdMeshAssetImportData>( GeometryCache->AssetImportData.Get() );
+				if ( !ImportData )
+				{
+					return false;
+				}
 
-			pxr::UsdPrim MeshPrim = Stage->OverridePrim( OverridePrimPath );
-			UE::USDPrimConversionImpl::Private::AuthorMaterialOverride( MeshPrim, Override->GetPathName() );
+				if ( const FUsdPrimPathList* SourcePrimPaths = ImportData->MaterialSlotToPrimPaths.Find( MatIndex ) )
+				{
+					for ( const FString& PrimPath : SourcePrimPaths->PrimPaths )
+					{
+						pxr::SdfPath OverridePrimPath = UnrealToUsd::ConvertPath( *PrimPath ).Get();
+						pxr::UsdPrim MeshPrim = Stage->OverridePrim( OverridePrimPath );
+						UE::USDPrimConversionImpl::Private::AuthorMaterialOverride( MeshPrim, Override->GetPathName() );
+					}
+				}
+			}
 		}
 	}
 	else if ( const UStaticMeshComponent* StaticMeshComponent = Cast<const UStaticMeshComponent>( MeshComponent ) )
 	{
-#if WITH_EDITOR
 		if ( UStaticMesh* Mesh = StaticMeshComponent->GetStaticMesh() )
 		{
-			int32 NumLODs = Mesh->GetNumLODs();
-			const bool bHasLODs = NumLODs > 1;
+			UUsdMeshAssetImportData* ImportData = Cast<UUsdMeshAssetImportData>( Mesh->AssetImportData.Get() );
+			if ( !ImportData )
+			{
+				return false;
+			}
 
 			const TArray<UMaterialInterface*>& Overrides = MeshComponent->OverrideMaterials;
 			for ( int32 MatIndex = 0; MatIndex < Overrides.Num(); ++MatIndex )
@@ -2160,7 +2179,7 @@ bool UnrealToUsd::ConvertMeshComponent( const pxr::UsdStageRefPtr& Stage, const
 					continue;
 				}
 
-				for ( int32 LODIndex = 0; LODIndex < NumLODs; ++LODIndex )
+				for ( int32 LODIndex = 0; LODIndex < Mesh->GetNumLODs(); ++LODIndex )
 				{
 					int32 NumSections = Mesh->GetNumSections( LODIndex );
 					const bool bHasSubsets = NumSections > 1;
@@ -2173,40 +2192,19 @@ bool UnrealToUsd::ConvertMeshComponent( const pxr::UsdStageRefPtr& Stage, const
 							continue;
 						}
 
-						pxr::SdfPath OverridePrimPath = UsdPrim.GetPath();
-
-						// If we have only 1 LOD, the asset's DefaultPrim will be the Mesh prim directly.
-						// If we have multiple, the default prim won't have any schema, but will contain separate
-						// Mesh prims for each LOD named "LOD0", "LOD1", etc., switched via a "LOD" variant set
-						if ( bHasLODs )
+						if ( const FUsdPrimPathList* SourcePrimPaths = ImportData->MaterialSlotToPrimPaths.Find( SectionMatIndex ) )
 						{
-							OverridePrimPath = OverridePrimPath.AppendPath( UnrealToUsd::ConvertPath( *FString::Printf( TEXT( "LOD%d" ), LODIndex ) ).Get() );
-						}
-
-						// If our LOD has only one section, its material assignment will be authored directly on the Mesh prim.
-						// If it has more than one material slot, we'll author UsdGeomSubset for each LOD Section, and author the material
-						// assignment there
-						if ( bHasSubsets )
-						{
-							// Assume the UE sections are in the same order as the USD ones
-							std::vector<pxr::UsdGeomSubset> GeomSubsets = pxr::UsdShadeMaterialBindingAPI( UsdPrim ).GetMaterialBindSubsets();
-							if ( SectionIndex < GeomSubsets.size() )
-							{
-								OverridePrimPath = OverridePrimPath.AppendChild( GeomSubsets[ SectionIndex ].GetPrim().GetName() );
-							}
-							else
+							for ( const FString& PrimPath : SourcePrimPaths->PrimPaths )
 							{
-								OverridePrimPath = OverridePrimPath.AppendPath( UnrealToUsd::ConvertPath( *FString::Printf( TEXT( "Section%d" ), SectionIndex ) ).Get() );
+								pxr::SdfPath OverridePrimPath = UnrealToUsd::ConvertPath( *PrimPath ).Get();
+								pxr::UsdPrim MeshPrim = Stage->OverridePrim( OverridePrimPath );
+								UE::USDPrimConversionImpl::Private::AuthorMaterialOverride( MeshPrim, Override->GetPathName() );
 							}
 						}
-
-						pxr::UsdPrim MeshPrim = Stage->OverridePrim( OverridePrimPath );
-						UE::USDPrimConversionImpl::Private::AuthorMaterialOverride( MeshPrim, Override->GetPathName() );
 					}
 				}
 			}
 		}
-#endif // WITH_EDITOR
 	}
 	else if ( const USkinnedMeshComponent* SkinnedMeshComponent = Cast<const USkinnedMeshComponent>( MeshComponent ) )
 	{
@@ -2216,33 +2214,24 @@ bool UnrealToUsd::ConvertMeshComponent( const pxr::UsdStageRefPtr& Stage, const
 			return false;
 		}
 
-		if ( const USkinnedAsset* SkinnedAsset = SkinnedMeshComponent->GetSkinnedAsset())
+		if ( const USkeletalMesh* SkeletalMesh = Cast<USkeletalMesh>( SkinnedMeshComponent->GetSkinnedAsset() ) )
 		{
-			FSkeletalMeshRenderData* RenderData = SkinnedAsset->GetResourceForRendering();
+			FSkeletalMeshRenderData* RenderData = SkeletalMesh->GetResourceForRendering();
 			if ( !RenderData )
 			{
 				return false;
 			}
 
-			TIndirectArray<FSkeletalMeshLODRenderData>& LodRenderData = RenderData->LODRenderData;
-			if ( LodRenderData.Num() == 0 )
+			UUsdMeshAssetImportData* ImportData = Cast<UUsdMeshAssetImportData>( SkeletalMesh->GetAssetImportData() );
+			if ( !ImportData )
 			{
 				return false;
 			}
 
-			int32 NumLODs = SkinnedAsset->GetLODNum();
-			const bool bHasLODs = NumLODs > 1;
-
-			FString MeshName;
-			if ( !bHasLODs )
+			TIndirectArray<FSkeletalMeshLODRenderData>& LodRenderData = RenderData->LODRenderData;
+			if ( LodRenderData.Num() == 0 )
 			{
-				for ( const pxr::UsdPrim& Child : UsdPrim.GetChildren() )
-				{
-					if ( pxr::UsdGeomMesh Mesh{ Child } )
-					{
-						MeshName = UsdToUnreal::ConvertToken( Child.GetName() );
-					}
-				}
+				return false;
 			}
 
 			const TArray<UMaterialInterface*>& Overrides = MeshComponent->OverrideMaterials;
@@ -2254,14 +2243,14 @@ bool UnrealToUsd::ConvertMeshComponent( const pxr::UsdStageRefPtr& Stage, const
 					continue;
 				}
 
-				for ( int32 LODIndex = 0; LODIndex < NumLODs; ++LODIndex )
+				for ( int32 LODIndex = 0; LODIndex < SkeletalMesh->GetLODNum(); ++LODIndex )
 				{
 					if ( !LodRenderData.IsValidIndex( LODIndex ) )
 					{
 						continue;
 					}
 
-					const FSkeletalMeshLODInfo* LODInfo = SkinnedAsset->GetLODInfo( LODIndex );
+					const FSkeletalMeshLODInfo* LODInfo = SkeletalMesh->GetLODInfo( LODIndex );
 
 					const TArray<FSkelMeshRenderSection>& Sections = LodRenderData[ LODIndex ].RenderSections;
 					int32 NumSections = Sections.Num();
@@ -2282,44 +2271,21 @@ bool UnrealToUsd::ConvertMeshComponent( const pxr::UsdStageRefPtr& Stage, const
 							continue;
 						}
 
-						pxr::SdfPath OverridePrimPath = UsdPrim.GetPath();
-
-						// If we have only 1 LOD, the asset's DefaultPrim will be a SkelRoot, and the Mesh will be a subprim
-						// with the same name.If we have the default prim is also the SkelRoot, but will contain separate
-						// Mesh prims for each LOD named "LOD0", "LOD1", etc., switched via a "LOD" variant set
-						if ( bHasLODs )
-						{
-							OverridePrimPath = OverridePrimPath.AppendPath( UnrealToUsd::ConvertPath( *FString::Printf( TEXT( "LOD%d" ), LODIndex ) ).Get() );
-						}
-						else
-						{
-							OverridePrimPath = OverridePrimPath.AppendElementString( UnrealToUsd::ConvertString( *MeshName ).Get() );
-						}
-
-						// If our LOD has only one section, its material assignment will be authored directly on the Mesh prim.
-						// If it has more than one material slot, we'll author UsdGeomSubset for each LOD Section, and author the material
-						// assignment there
-						if ( bHasSubsets )
+						if ( const FUsdPrimPathList* SourcePrimPaths = ImportData->MaterialSlotToPrimPaths.Find( SectionMatIndex ) )
 						{
-							// Assume the UE sections are in the same order as the USD ones
-							std::vector<pxr::UsdGeomSubset> GeomSubsets = pxr::UsdShadeMaterialBindingAPI( UsdPrim ).GetMaterialBindSubsets();
-							if ( SectionIndex < GeomSubsets.size() )
-							{
-								OverridePrimPath = OverridePrimPath.AppendChild( GeomSubsets[ SectionIndex ].GetPrim().GetName() );
-							}
-							else
+							for ( const FString& PrimPath : SourcePrimPaths->PrimPaths )
 							{
-								OverridePrimPath = OverridePrimPath.AppendPath( UnrealToUsd::ConvertPath( *FString::Printf( TEXT( "Section%d" ), SectionIndex ) ).Get() );
+								pxr::SdfPath OverridePrimPath = UnrealToUsd::ConvertPath( *PrimPath ).Get();
+								pxr::UsdPrim MeshPrim = Stage->OverridePrim( OverridePrimPath );
+								UE::USDPrimConversionImpl::Private::AuthorMaterialOverride( MeshPrim, Override->GetPathName() );
 							}
 						}
-
-						pxr::UsdPrim MeshPrim = Stage->OverridePrim( OverridePrimPath );
-						UE::USDPrimConversionImpl::Private::AuthorMaterialOverride( MeshPrim, Override->GetPathName() );
 					}
 				}
 			}
 		}
 	}
+#endif // WITH_EDITOR
 
 	return true;
 }
diff --git a/Engine/Plugins/Importers/USDImporter/Source/USDUtilities/Private/USDSkeletalDataConversion.cpp b/Engine/Plugins/Importers/USDImporter/Source/USDUtilities/Private/USDSkeletalDataConversion.cpp
--- a/Engine/Plugins/Importers/USDImporter/Source/USDUtilities/Private/USDSkeletalDataConversion.cpp
+++ b/Engine/Plugins/Importers/USDImporter/Source/USDUtilities/Private/USDSkeletalDataConversion.cpp
@@ -1393,27 +1393,54 @@ bool UsdToUnreal::ConvertSkinnedMesh(
 	// We want to combine identical slots for skeletal meshes, which is different to static meshes, where each section gets a slot
 	// Note: This is a different index remapping to the one that happens for LODs, using LODMaterialMap! Here we're combining meshes of the same LOD
 	TMap<UsdUtils::FUsdPrimMaterialSlot, int32> SlotToCombinedMaterialIndex;
-	TMap<int32, int32> LocalToCombinedMaterialIndex;
+
+	// Position 3 in this has the value 6 --> Local material slot #3 is actually the combined material slot #6
+	TArray<int32> LocalToCombinedMaterialIndex;
+	LocalToCombinedMaterialIndex.SetNumZeroed( LocalInfo.Slots.Num() );
+
 	for (int32 Index = 0; Index < MaterialAssignments.Num(); ++Index)
 	{
-		SlotToCombinedMaterialIndex.Add( MaterialAssignments[ Index ], Index );
+		const UsdUtils::FUsdPrimMaterialSlot& Slot = MaterialAssignments[ Index ];
+
+		// Combine entries in this way so that we can append PrimPaths
+		TMap< UsdUtils::FUsdPrimMaterialSlot, int32 >::TKeyIterator KeyIt = SlotToCombinedMaterialIndex.CreateKeyIterator( Slot );
+		if ( KeyIt )
+		{
+			KeyIt.Key().PrimPaths.Append( Slot.PrimPaths );
+			KeyIt.Value() = Index;
+		}
+		else
+		{
+			SlotToCombinedMaterialIndex.Add( Slot, Index );
+		}
 	}
 	for (int32 LocalIndex = 0; LocalIndex < LocalInfo.Slots.Num(); ++LocalIndex)
 	{
-		UsdUtils::FUsdPrimMaterialSlot& LocalSlot = LocalInfo.Slots[LocalIndex];
+		const UsdUtils::FUsdPrimMaterialSlot& LocalSlot = LocalInfo.Slots[ LocalIndex ];
 
-		int32 CombinedMaterialIndex = INDEX_NONE;
-		if ( int32* FoundCombinedMaterialIndex = SlotToCombinedMaterialIndex.Find( LocalSlot ) )
+		// Combine entries in this way so that we can append PrimPaths
+		TMap< UsdUtils::FUsdPrimMaterialSlot, int32 >::TKeyIterator KeyIt = SlotToCombinedMaterialIndex.CreateKeyIterator( LocalSlot );
+		if ( KeyIt )
 		{
-			CombinedMaterialIndex = *FoundCombinedMaterialIndex;
+			KeyIt.Key().PrimPaths.Append( LocalSlot.PrimPaths );
+
+			const int32 ExistingCombinedIndex = KeyIt.Value();
+			LocalToCombinedMaterialIndex[ LocalIndex ] = ExistingCombinedIndex;
 		}
 		else
 		{
-			CombinedMaterialIndex = MaterialAssignments.Add( LocalSlot );
-			SlotToCombinedMaterialIndex.Add( LocalSlot, CombinedMaterialIndex );
+			int32 NewIndex = MaterialAssignments.Add( LocalSlot );
+			SlotToCombinedMaterialIndex.Add( LocalSlot, NewIndex );
+			LocalToCombinedMaterialIndex[ LocalIndex ] = NewIndex;
 		}
-
-		LocalToCombinedMaterialIndex.Add( LocalIndex, CombinedMaterialIndex );
+	}
+	// Now that we merged all prim paths into they keys of CombinedMaterialSlotsToIndex, let's copy them back into
+	// our output
+	for ( UsdUtils::FUsdPrimMaterialSlot& Slot : MaterialAssignments )
+	{
+		TMap< UsdUtils::FUsdPrimMaterialSlot, int32 >::TKeyIterator KeyIt = SlotToCombinedMaterialIndex.CreateKeyIterator( Slot );
+		ensure( KeyIt );
+		Slot.PrimPaths = KeyIt.Key().PrimPaths;
 	}
 
 	// Retrieve vertex colors
diff --git a/Engine/Plugins/Importers/USDImporter/Source/USDUtilities/Public/USDGeomMeshConversion.h b/Engine/Plugins/Importers/USDImporter/Source/USDUtilities/Public/USDGeomMeshConversion.h
--- a/Engine/Plugins/Importers/USDImporter/Source/USDUtilities/Public/USDGeomMeshConversion.h
+++ b/Engine/Plugins/Importers/USDImporter/Source/USDUtilities/Public/USDGeomMeshConversion.h
@@ -220,6 +220,12 @@ namespace UsdUtils
 	 */
 	struct FUsdPrimMaterialSlot
 	{
+		/**
+		 * Path to the prims that contain this material assignment (e.g. '/Root/my_cube' or '/Root/my_cube/geomsubset_0'.
+		 * This is a set because this single material slot in UE may be merged from N identical prims.
+		 */
+		TSet<FString> PrimPaths;
+
 		/** What this represents depends on AssignedMaterialType */
 		FString MaterialSource;
 		EPrimAssignmentType AssignmentType = EPrimAssignmentType::None;
@@ -248,6 +254,7 @@ namespace UsdUtils
 			Ar << Slot.MaterialSource;
 			Ar << Slot.AssignmentType;
 			Ar << Slot.bMeshIsDoubleSided;
+			Ar << Slot.PrimPaths;
 
 			return Ar;
 		}
