USD: Fix USD error when exporting files with meshes on Linux.

diff --git a/Engine/Plugins/Importers/USDImporter/Source/USDUtilities/Private/USDGeomMeshConversion.cpp b/Engine/Plugins/Importers/USDImporter/Source/USDUtilities/Private/USDGeomMeshConversion.cpp
--- a/Engine/Plugins/Importers/USDImporter/Source/USDUtilities/Private/USDGeomMeshConversion.cpp
+++ b/Engine/Plugins/Importers/USDImporter/Source/USDUtilities/Private/USDGeomMeshConversion.cpp
@@ -157,7 +157,10 @@ namespace UE::UsdGeomMeshConversion::Private
 				// actually subdivide (like usdview) will just discard authored normals and fully recompute them
 				// on-demand in case they have a valid subdivision scheme (which is the default state).
 				// Reference: https://graphics.pixar.com/usd/release/api/class_usd_geom_mesh.html#UsdGeom_Mesh_Normals
-				ensure( UsdMesh.CreateSubdivisionSchemeAttr( pxr::VtValue{ pxr::UsdGeomTokens->none } ) );
+				if ( pxr::UsdAttribute SubdivisionAttr = UsdMesh.CreateSubdivisionSchemeAttr() )
+				{
+					ensure( SubdivisionAttr.Set( pxr::UsdGeomTokens->none ) );
+				}
 
 				pxr::UsdAttribute NormalsAttribute = UsdMesh.CreateNormalsAttr();
 				if ( NormalsAttribute )
diff --git a/Engine/Plugins/Importers/USDImporter/Source/USDUtilities/Private/USDSkeletalDataConversion.cpp b/Engine/Plugins/Importers/USDImporter/Source/USDUtilities/Private/USDSkeletalDataConversion.cpp
--- a/Engine/Plugins/Importers/USDImporter/Source/USDUtilities/Private/USDSkeletalDataConversion.cpp
+++ b/Engine/Plugins/Importers/USDImporter/Source/USDUtilities/Private/USDSkeletalDataConversion.cpp
@@ -749,7 +749,10 @@ namespace UnrealToUsdImpl
 				// actually subdivide (like usdview) will just discard authored normals and fully recompute them
 				// on-demand in case they have a valid subdivision scheme (which is the default state).
 				// Reference: https://graphics.pixar.com/usd/release/api/class_usd_geom_mesh.html#UsdGeom_Mesh_Normals
-				ensure( UsdLODPrimGeomMesh.CreateSubdivisionSchemeAttr( pxr::VtValue{ pxr::UsdGeomTokens->none } ) );
+				if ( pxr::UsdAttribute SubdivisionAttr = UsdLODPrimGeomMesh.CreateSubdivisionSchemeAttr() )
+				{
+					ensure( SubdivisionAttr.Set( pxr::UsdGeomTokens->none ) );
+				}
 
 				pxr::UsdAttribute NormalsAttribute = UsdLODPrimGeomMesh.CreateNormalsAttr();
 				if ( NormalsAttribute )
