USD: Convert the "collapse top level point instancers" option into a cvar.

diff --git a/Engine/Plugins/Importers/USDImporter/Source/USDSchemas/Private/USDGeomPointInstancerTranslator.cpp b/Engine/Plugins/Importers/USDImporter/Source/USDSchemas/Private/USDGeomPointInstancerTranslator.cpp
--- a/Engine/Plugins/Importers/USDImporter/Source/USDSchemas/Private/USDGeomPointInstancerTranslator.cpp
+++ b/Engine/Plugins/Importers/USDImporter/Source/USDSchemas/Private/USDGeomPointInstancerTranslator.cpp
@@ -39,6 +39,12 @@
 
 #define LOCTEXT_NAMESPACE "USDGeomPointInstancer"
 
+static bool GCollapseTopLevelPointInstancers = false;
+static FAutoConsoleVariableRef CVarCollapseTopLevelPointInstancers(
+	TEXT( "USD.CollapseTopLevelPointInstancers" ),
+	GCollapseTopLevelPointInstancers,
+	TEXT( "If this is true will cause any point instancer to be collapsed to a single static mesh. Point instancers that are used as prototypes for other point instancers will always be collapsed." ) );
+
 namespace UsdGeomPointInstancerTranslatorImpl
 {
 	void ApplyPointInstanceTransforms( UInstancedStaticMeshComponent* Component, TArray<FTransform>& InstanceTransforms )
@@ -146,7 +152,7 @@ void FUsdGeomPointInstancerTranslator::CreateAssets()
 	// If another FUsdGeomXformableTranslator is collapsing the point instancer prim, it will do so by calling
 	// UsdToUnreal::ConvertGeomMeshHierarchy which will consume the prim directly.
 	// This case right here is if we're collapsing *ourselves*, where we'll essentially pretend we're a single static mesh.
-	if ( Context->bCollapseTopLevelPointInstancers )
+	if ( GCollapseTopLevelPointInstancers )
 	{
 		// Don't bake our actual point instancer's transform or visibility into the mesh as its nice to have these on the static mesh component instead
 		const bool bIgnoreTopLevelTransformAndVisibility = true;
@@ -225,7 +231,7 @@ USceneComponent* FUsdGeomPointInstancerTranslator::CreateComponents()
 	TRACE_CPUPROFILER_EVENT_SCOPE( FUsdGeomPointInstancerTranslator::CreateComponents );
 
 	// If we're collapsing ourselves, we're really just a collapsed Xform prim, so let that translator handle it
-	if ( Context->bCollapseTopLevelPointInstancers )
+	if ( GCollapseTopLevelPointInstancers )
 	{
 		return FUsdGeomXformableTranslator::CreateComponents();
 	}
@@ -373,6 +379,11 @@ void FUsdGeomPointInstancerTranslator::UpdateComponents( USceneComponent* PointI
 	Super::UpdateComponents( PointInstancerRootComponent );
 }
 
+bool FUsdGeomPointInstancerTranslator::CanBeCollapsed( ECollapsingType CollapsingType ) const
+{
+	return GCollapseTopLevelPointInstancers;
+}
+
 #undef LOCTEXT_NAMESPACE
 
 #endif // #if USE_USD_SDK
diff --git a/Engine/Plugins/Importers/USDImporter/Source/USDSchemas/Private/USDGeomXformableTranslator.cpp b/Engine/Plugins/Importers/USDImporter/Source/USDSchemas/Private/USDGeomXformableTranslator.cpp
--- a/Engine/Plugins/Importers/USDImporter/Source/USDSchemas/Private/USDGeomXformableTranslator.cpp
+++ b/Engine/Plugins/Importers/USDImporter/Source/USDSchemas/Private/USDGeomXformableTranslator.cpp
@@ -462,9 +462,14 @@ USceneComponent* FUsdGeomXformableTranslator::CreateComponentsEx( TOptional< TSu
 				}
 				// If this is a component for a point instancer that just collapsed itself into a static mesh, just make
 				// a static mesh component that can receive it
-				else if ( Context->bCollapseTopLevelPointInstancers && pxr::UsdPrim{ Prim }.IsA<pxr::UsdGeomPointInstancer>() )
+				else if ( pxr::UsdPrim{ Prim }.IsA<pxr::UsdGeomPointInstancer>() )
 				{
-					ComponentType = UStaticMeshComponent::StaticClass();
+					static IConsoleVariable* CollapseCvar =
+						IConsoleManager::Get().FindConsoleVariable( TEXT( "USD.CollapseTopLevelPointInstancers" ) );
+					if ( CollapseCvar && CollapseCvar->GetBool() )
+					{
+						ComponentType = UStaticMeshComponent::StaticClass();
+					}
 				}
 			}
 		}
diff --git a/Engine/Plugins/Importers/USDImporter/Source/USDSchemas/Public/USDGeomPointInstancerTranslator.h b/Engine/Plugins/Importers/USDImporter/Source/USDSchemas/Public/USDGeomPointInstancerTranslator.h
--- a/Engine/Plugins/Importers/USDImporter/Source/USDSchemas/Public/USDGeomPointInstancerTranslator.h
+++ b/Engine/Plugins/Importers/USDImporter/Source/USDSchemas/Public/USDGeomPointInstancerTranslator.h
@@ -37,7 +37,7 @@ public:
 	// We only return true from CanBeCollapsed if e.g. a generic Xform prim should be able to collapse us too (like when we're
 	// collapsing even simple, not-nested point instancers).
 	virtual bool CollapsesChildren( ECollapsingType CollapsingType ) const override { return true; }
-	virtual bool CanBeCollapsed( ECollapsingType CollapsingType ) const override { return Context->bCollapseTopLevelPointInstancers; }
+	virtual bool CanBeCollapsed( ECollapsingType CollapsingType ) const override;
 };
 
 #endif // #if USE_USD_SDK
diff --git a/Engine/Plugins/Importers/USDImporter/Source/USDSchemas/Public/USDSchemaTranslator.h b/Engine/Plugins/Importers/USDImporter/Source/USDSchemas/Public/USDSchemaTranslator.h
--- a/Engine/Plugins/Importers/USDImporter/Source/USDSchemas/Public/USDSchemaTranslator.h
+++ b/Engine/Plugins/Importers/USDImporter/Source/USDSchemas/Public/USDSchemaTranslator.h
@@ -202,6 +202,7 @@ struct USDSCHEMAS_API FUsdSchemaTranslationContext : public TSharedFromThis< FUs
 	 * If false, will cause us to use HierarchicalInstancedStaticMeshComponents to replicate the instancing behavior.
 	 * Point instancers inside other point instancer prototypes are *always* collapsed into the prototype's static mesh.
 	 */
+	UE_DEPRECATED( 5.2, "This option is now controlled via the cvar 'USD.CollapseTopLevelPointInstancers'" )
 	bool bCollapseTopLevelPointInstancers = false;
 
 	/**
diff --git a/Engine/Plugins/Importers/USDImporter/Source/USDStage/Private/USDStageActor.cpp b/Engine/Plugins/Importers/USDImporter/Source/USDStage/Private/USDStageActor.cpp
--- a/Engine/Plugins/Importers/USDImporter/Source/USDStage/Private/USDStageActor.cpp
+++ b/Engine/Plugins/Importers/USDImporter/Source/USDStage/Private/USDStageActor.cpp
@@ -118,7 +118,6 @@ struct FUsdStageActorImpl
 
 		TranslationContext->KindsToCollapse = (EUsdDefaultKind) StageActor->KindsToCollapse;
 		TranslationContext->bMergeIdenticalMaterialSlots = StageActor->bMergeIdenticalMaterialSlots;
-		TranslationContext->bCollapseTopLevelPointInstancers = StageActor->bCollapseTopLevelPointInstancers;
 
 		UE::FSdfPath UsdPrimPath( *PrimPath );
 		UUsdPrimTwin* ParentUsdPrimTwin = StageActor->GetRootPrimTwin()->Find( UsdPrimPath.GetParentPath().GetString() );
@@ -535,7 +534,6 @@ struct FUsdStageActorImpl
 			EventAttributes.Emplace( TEXT( "InterpolationType" ), LexToString( (uint8)StageActor->InterpolationType) );
 			EventAttributes.Emplace( TEXT( "KindsToCollapse" ), LexToString( StageActor->KindsToCollapse ) );
 			EventAttributes.Emplace( TEXT( "MergeIdenticalMaterialSlots" ), LexToString( StageActor->bMergeIdenticalMaterialSlots ) );
-			EventAttributes.Emplace( TEXT( "CollapseTopLevelPointInstancers" ), LexToString( StageActor->bCollapseTopLevelPointInstancers ) );
 			EventAttributes.Emplace( TEXT( "PurposesToLoad" ), LexToString( StageActor->PurposesToLoad ) );
 			EventAttributes.Emplace( TEXT( "NaniteTriangleThreshold" ), LexToString( StageActor->NaniteTriangleThreshold ) );
 			EventAttributes.Emplace( TEXT( "RenderContext" ), StageActor->RenderContext.ToString() );
@@ -622,7 +620,6 @@ AUsdStageActor::AUsdStageActor()
 	, InterpolationType( EUsdInterpolationType::Linear )
 	, KindsToCollapse( ( int32 ) ( EUsdDefaultKind::Component | EUsdDefaultKind::Subcomponent ) )
 	, bMergeIdenticalMaterialSlots( true )
-	, bCollapseTopLevelPointInstancers( false )
 	, PurposesToLoad( (int32) EUsdPurpose::Proxy )
 	, NaniteTriangleThreshold( (uint64) 1000000 )
 	, MaterialPurpose( *UnrealIdentifiers::MaterialAllPurpose )
@@ -1488,10 +1485,6 @@ void AUsdStageActor::SetMergeIdenticalMaterialSlots( bool bMerge )
 
 void AUsdStageActor::SetCollapseTopLevelPointInstancers( bool bCollapse )
 {
-	Modify();
-
-	bCollapseTopLevelPointInstancers = bCollapse;
-	LoadUsdStage();
 }
 
 void AUsdStageActor::SetPurposesToLoad( int32 NewPurposesToLoad )
@@ -2834,10 +2827,6 @@ void AUsdStageActor::HandlePropertyChangedEvent( FPropertyChangedEvent& Property
 	{
 		SetMergeIdenticalMaterialSlots( bMergeIdenticalMaterialSlots );
 	}
-	else if ( PropertyName == GET_MEMBER_NAME_CHECKED( AUsdStageActor, bCollapseTopLevelPointInstancers ) )
-	{
-		SetCollapseTopLevelPointInstancers( bCollapseTopLevelPointInstancers );
-	}
 	else if ( PropertyName == GET_MEMBER_NAME_CHECKED( AUsdStageActor, PurposesToLoad ) )
 	{
 		SetPurposesToLoad( PurposesToLoad );
diff --git a/Engine/Plugins/Importers/USDImporter/Source/USDStage/Public/USDStageActor.h b/Engine/Plugins/Importers/USDImporter/Source/USDStage/Public/USDStageActor.h
--- a/Engine/Plugins/Importers/USDImporter/Source/USDStage/Public/USDStageActor.h
+++ b/Engine/Plugins/Importers/USDImporter/Source/USDStage/Public/USDStageActor.h
@@ -68,7 +68,8 @@ public:
 	 * If false, will cause us to use HierarchicalInstancedStaticMeshComponents to replicate the instancing behavior.
 	 * Point instancers inside other point instancer prototypes are *always* collapsed into the prototype's static mesh.
 	 */
-	UPROPERTY( EditAnywhere, BlueprintReadOnly, Category = "USD" )
+	UE_DEPRECATED( 5.2, "This option is now controlled via the cvar 'USD.CollapseTopLevelPointInstancers'" )
+	UPROPERTY()
 	bool bCollapseTopLevelPointInstancers;
 
 	/* Only load prims with these specific purposes from the USD file */
@@ -107,6 +108,7 @@ public:
 	UFUNCTION( BlueprintCallable, Category = "USD", meta = ( CallInEditor = "true" ) )
 	USDSTAGE_API void SetMergeIdenticalMaterialSlots( bool bMerge );
 
+	UE_DEPRECATED( 5.2, "This option is now controlled via the cvar 'USD.CollapseTopLevelPointInstancers'" )
 	UFUNCTION( BlueprintCallable, Category = "USD", meta = ( CallInEditor = "true" ) )
 	USDSTAGE_API void SetCollapseTopLevelPointInstancers( bool bCollapse );
 
diff --git a/Engine/Plugins/Importers/USDImporter/Source/USDStageEditor/Private/SUSDStage.cpp b/Engine/Plugins/Importers/USDImporter/Source/USDStageEditor/Private/SUSDStage.cpp
--- a/Engine/Plugins/Importers/USDImporter/Source/USDStageEditor/Private/SUSDStage.cpp
+++ b/Engine/Plugins/Importers/USDImporter/Source/USDStageEditor/Private/SUSDStage.cpp
@@ -1303,41 +1303,6 @@ void SUsdStage::FillCollapsingSubMenu( FMenuBuilder& MenuBuilder )
 		EUserInterfaceActionType::ToggleButton
 	);
 
-	MenuBuilder.AddMenuEntry(
-		LOCTEXT( "CollapseTopLevelPointInstancers", "Collapse top-level point instancers" ),
-		LOCTEXT( "CollapseTopLevelPointInstancersToolTip", "If true, will cause us to collapse any point instancer prim into a single static mesh and static mesh component.\nIf false, will cause us to use HierarchicalInstancedStaticMeshComponents to replicate the instancing behavior.\nPoint instancers inside other point instancer prototypes are * always * collapsed into the prototype's static mesh." ),
-		FSlateIcon(),
-		FUIAction(
-			FExecuteAction::CreateLambda( [this]()
-			{
-				if ( AUsdStageActor* StageActor = ViewModel.UsdStageActor.Get() )
-				{
-					FScopedTransaction Transaction( FText::Format(
-						LOCTEXT( "CollapseTopLevelPointInstancersTransaction", "Toggle bCollapseTopLevelPointInstancers on USD stage actor '{1}'" ),
-						FText::FromString( StageActor->GetActorLabel() )
-					) );
-
-					TGuardValue<bool> MaintainSelectionGuard( bUpdatingViewportSelection, true );
-					StageActor->SetCollapseTopLevelPointInstancers( !StageActor->bCollapseTopLevelPointInstancers );
-				}
-			}),
-			FCanExecuteAction::CreateLambda( [this]()
-			{
-				return ViewModel.UsdStageActor.Get() != nullptr;
-			}),
-			FIsActionChecked::CreateLambda( [this]()
-			{
-				if ( AUsdStageActor* StageActor = ViewModel.UsdStageActor.Get() )
-				{
-					return StageActor->bCollapseTopLevelPointInstancers;
-				}
-				return false;
-			})
-		),
-		NAME_None,
-		EUserInterfaceActionType::ToggleButton
-	);
-
 	auto AddKindToCollapseEntry = [&]( const EUsdDefaultKind Kind, const FText& Text, FCanExecuteAction CanExecuteAction )
 	{
 		MenuBuilder.AddMenuEntry(
diff --git a/Engine/Plugins/Importers/USDImporter/Source/USDStageImporter/Private/USDStageImportOptions.cpp b/Engine/Plugins/Importers/USDImporter/Source/USDStageImporter/Private/USDStageImportOptions.cpp
--- a/Engine/Plugins/Importers/USDImporter/Source/USDStageImporter/Private/USDStageImportOptions.cpp
+++ b/Engine/Plugins/Importers/USDImporter/Source/USDStageImporter/Private/USDStageImportOptions.cpp
@@ -33,7 +33,6 @@ UUsdStageImportOptions::UUsdStageImportOptions(const FObjectInitializer& ObjectI
 	bPrimPathFolderStructure = false;
 	KindsToCollapse = ( int32 ) ( EUsdDefaultKind::Component | EUsdDefaultKind::Subcomponent );
 	bMergeIdenticalMaterialSlots = true;
-	bCollapseTopLevelPointInstancers = false;
 	bInterpretLODs = true;
 }
 
@@ -100,6 +99,5 @@ void UsdUtils::AddAnalyticsAttributes(
 	InOutAttributes.Emplace( TEXT( "PrimPathFolderStructure" ), LexToString( Options.bPrimPathFolderStructure ) );
 	InOutAttributes.Emplace( TEXT( "KindsToCollapse" ), LexToString( Options.KindsToCollapse ) );
 	InOutAttributes.Emplace( TEXT( "MergeIdenticalMaterialSlots" ), LexToString( Options.bMergeIdenticalMaterialSlots ) );
-	InOutAttributes.Emplace( TEXT( "CollapseTopLevelPointInstancers" ), LexToString( Options.bCollapseTopLevelPointInstancers ) );
 	InOutAttributes.Emplace( TEXT( "InterpretLODs" ), LexToString( Options.bInterpretLODs ) );
 }
\ No newline at end of file
diff --git a/Engine/Plugins/Importers/USDImporter/Source/USDStageImporter/Private/USDStageImporter.cpp b/Engine/Plugins/Importers/USDImporter/Source/USDStageImporter/Private/USDStageImporter.cpp
--- a/Engine/Plugins/Importers/USDImporter/Source/USDStageImporter/Private/USDStageImporter.cpp
+++ b/Engine/Plugins/Importers/USDImporter/Source/USDStageImporter/Private/USDStageImporter.cpp
@@ -1643,7 +1643,6 @@ void UUsdStageImporter::ImportFromFile(FUsdStageImportContext& ImportContext)
 	TranslationContext->ParentComponent = ImportContext.SceneActor ? ImportContext.SceneActor->GetRootComponent() : nullptr;
 	TranslationContext->KindsToCollapse = ( EUsdDefaultKind ) ImportContext.ImportOptions->KindsToCollapse;
 	TranslationContext->bMergeIdenticalMaterialSlots = ImportContext.ImportOptions->bMergeIdenticalMaterialSlots;
-	TranslationContext->bCollapseTopLevelPointInstancers = ImportContext.ImportOptions->bCollapseTopLevelPointInstancers;
 	TranslationContext->bAllowInterpretingLODs = ImportContext.ImportOptions->bInterpretLODs;
 	TranslationContext->bAllowParsingSkeletalAnimations = ImportContext.ImportOptions->bImportGeometry && ImportContext.ImportOptions->bImportSkeletalAnimations;
 	TranslationContext->MaterialToPrimvarToUVIndex = &ImportContext.MaterialToPrimvarToUVIndex;
@@ -1743,7 +1742,6 @@ bool UUsdStageImporter::ReimportSingleAsset(FUsdStageImportContext& ImportContex
 	TranslationContext->RootMotionHandling = ImportContext.ImportOptions->RootMotionHandling;
 	TranslationContext->KindsToCollapse = ( EUsdDefaultKind ) ImportContext.ImportOptions->KindsToCollapse;
 	TranslationContext->bMergeIdenticalMaterialSlots = ImportContext.ImportOptions->bMergeIdenticalMaterialSlots;
-	TranslationContext->bCollapseTopLevelPointInstancers = ImportContext.ImportOptions->bCollapseTopLevelPointInstancers;
 	TranslationContext->bAllowInterpretingLODs = ImportContext.ImportOptions->bInterpretLODs;
 	TranslationContext->bAllowParsingSkeletalAnimations = ImportContext.ImportOptions->bImportGeometry && ImportContext.ImportOptions->bImportSkeletalAnimations;
 	TranslationContext->MaterialToPrimvarToUVIndex = &ImportContext.MaterialToPrimvarToUVIndex;
diff --git a/Engine/Plugins/Importers/USDImporter/Source/USDStageImporter/Public/USDStageImportOptions.h b/Engine/Plugins/Importers/USDImporter/Source/USDStageImporter/Public/USDStageImportOptions.h
--- a/Engine/Plugins/Importers/USDImporter/Source/USDStageImporter/Public/USDStageImportOptions.h
+++ b/Engine/Plugins/Importers/USDImporter/Source/USDStageImporter/Public/USDStageImportOptions.h
@@ -144,7 +144,8 @@ public:
 	 * If false, will cause us to use HierarchicalInstancedStaticMeshComponents to replicate the instancing behavior.
 	 * Point instancers inside other point instancer prototypes are *always* collapsed into the prototype's static mesh.
 	 */
-	UPROPERTY( BlueprintReadWrite, config, EditAnywhere, Category = "Processing" )
+	UE_DEPRECATED( 5.2, "This option is now controlled via the cvar 'USD.CollapseTopLevelPointInstancers'" )
+	UPROPERTY()
 	bool bCollapseTopLevelPointInstancers;
 
 	/** When true, if a prim has a "LOD" variant set with variants named "LOD0", "LOD1", etc. where each contains a UsdGeomMesh, the importer will attempt to parse the meshes as separate LODs of a single UStaticMesh. When false, only the selected variant will be parsed as LOD0 of the UStaticMesh.  */
