USD: Limit skeletal mesh vertex influences to 8, but only resize influences if we can't rely on unlimited mode.

diff --git a/Engine/Plugins/Importers/USDImporter/Source/USDUtilities/Private/USDSkeletalDataConversion.cpp b/Engine/Plugins/Importers/USDImporter/Source/USDUtilities/Private/USDSkeletalDataConversion.cpp
--- a/Engine/Plugins/Importers/USDImporter/Source/USDUtilities/Private/USDSkeletalDataConversion.cpp
+++ b/Engine/Plugins/Importers/USDImporter/Source/USDUtilities/Private/USDSkeletalDataConversion.cpp
@@ -1789,13 +1789,15 @@ bool UsdToUnreal::ConvertSkinnedMesh(
 	VtArray<float> JointWeights;
 	SkinningQuery.ComputeVaryingJointInfluences(NumPoints, &JointIndices, &JointWeights);
 
-	// Recompute the joint influences if it's above the limit
+	// Recompute the joint influences if we need to
 	uint32 NumInfluencesPerComponent = SkinningQuery.GetNumInfluencesPerComponent();
-	if (NumInfluencesPerComponent > MAX_INFLUENCES_PER_STREAM)
+	const uint32 MaxAllowedInfluences = EXTRA_BONE_INFLUENCES;
+	const bool bUseUnlimitedBoneInfluences = FGPUBaseSkinVertexFactory::UseUnlimitedBoneInfluences( NumInfluencesPerComponent );
+	if ( NumInfluencesPerComponent > MaxAllowedInfluences && !bUseUnlimitedBoneInfluences )
 	{
-		UsdSkelResizeInfluences(&JointIndices, NumInfluencesPerComponent, MAX_INFLUENCES_PER_STREAM);
-		UsdSkelResizeInfluences(&JointWeights, NumInfluencesPerComponent, MAX_INFLUENCES_PER_STREAM);
-		NumInfluencesPerComponent = MAX_INFLUENCES_PER_STREAM;
+		UsdSkelResizeInfluences( &JointIndices, NumInfluencesPerComponent, MaxAllowedInfluences );
+		UsdSkelResizeInfluences( &JointWeights, NumInfluencesPerComponent, MaxAllowedInfluences );
+		NumInfluencesPerComponent = MaxAllowedInfluences;
 	}
 
 	// We keep track of which influences we added because we combine many Mesh prim (each with potentially a different
