USD: Fix prims with other schemas collapsing into meshes and disappearing.

diff --git a/Engine/Plugins/Importers/USDImporter/Source/USDSchemas/Private/USDGeomCameraTranslator.cpp b/Engine/Plugins/Importers/USDImporter/Source/USDSchemas/Private/USDGeomCameraTranslator.cpp
--- a/Engine/Plugins/Importers/USDImporter/Source/USDSchemas/Private/USDGeomCameraTranslator.cpp
+++ b/Engine/Plugins/Importers/USDImporter/Source/USDSchemas/Private/USDGeomCameraTranslator.cpp
@@ -105,4 +105,14 @@ void FUsdGeomCameraTranslator::UpdateComponents( USceneComponent* SceneComponent
 	}
 }
 
+bool FUsdGeomCameraTranslator::CollapsesChildren( ECollapsingType CollapsingType ) const
+{
+	return false;
+}
+
+bool FUsdGeomCameraTranslator::CanBeCollapsed( ECollapsingType CollapsingType ) const
+{
+	return false;
+}
+
 #endif // #if USE_USD_SDK
diff --git a/Engine/Plugins/Importers/USDImporter/Source/USDSchemas/Private/USDLuxLightTranslator.cpp b/Engine/Plugins/Importers/USDImporter/Source/USDSchemas/Private/USDLuxLightTranslator.cpp
--- a/Engine/Plugins/Importers/USDImporter/Source/USDSchemas/Private/USDLuxLightTranslator.cpp
+++ b/Engine/Plugins/Importers/USDImporter/Source/USDSchemas/Private/USDLuxLightTranslator.cpp
@@ -98,4 +98,14 @@ void FUsdLuxLightTranslator::UpdateComponents( USceneComponent* SceneComponent )
 	}
 }
 
+bool FUsdLuxLightTranslator::CollapsesChildren( ECollapsingType CollapsingType ) const
+{
+	return false;
+}
+
+bool FUsdLuxLightTranslator::CanBeCollapsed( ECollapsingType CollapsingType ) const
+{
+	return false;
+}
+
 #endif // #if USE_USD_SDK
\ No newline at end of file
diff --git a/Engine/Plugins/Importers/USDImporter/Source/USDSchemas/Private/USDShadeMaterialTranslator.cpp b/Engine/Plugins/Importers/USDImporter/Source/USDSchemas/Private/USDShadeMaterialTranslator.cpp
--- a/Engine/Plugins/Importers/USDImporter/Source/USDSchemas/Private/USDShadeMaterialTranslator.cpp
+++ b/Engine/Plugins/Importers/USDImporter/Source/USDSchemas/Private/USDShadeMaterialTranslator.cpp
@@ -411,4 +411,14 @@ void FUsdShadeMaterialTranslator::CreateAssets()
 	}
 }
 
+bool FUsdShadeMaterialTranslator::CollapsesChildren( ECollapsingType CollapsingType ) const
+{
+	return false;
+}
+
+bool FUsdShadeMaterialTranslator::CanBeCollapsed( ECollapsingType CollapsingType ) const
+{
+	return false;
+}
+
 #endif // #if USE_USD_SDK
diff --git a/Engine/Plugins/Importers/USDImporter/Source/USDSchemas/Public/USDGeomCameraTranslator.h b/Engine/Plugins/Importers/USDImporter/Source/USDSchemas/Public/USDGeomCameraTranslator.h
--- a/Engine/Plugins/Importers/USDImporter/Source/USDSchemas/Public/USDGeomCameraTranslator.h
+++ b/Engine/Plugins/Importers/USDImporter/Source/USDSchemas/Public/USDGeomCameraTranslator.h
@@ -21,6 +21,9 @@ public:
 
 	virtual USceneComponent* CreateComponents() override;
 	virtual void UpdateComponents( USceneComponent* SceneComponent ) override;
+
+	virtual bool CollapsesChildren( ECollapsingType CollapsingType ) const override;
+	virtual bool CanBeCollapsed( ECollapsingType CollapsingType ) const override;
 };
 
 #endif // #if USE_USD_SDK
diff --git a/Engine/Plugins/Importers/USDImporter/Source/USDSchemas/Public/USDLuxLightTranslator.h b/Engine/Plugins/Importers/USDImporter/Source/USDSchemas/Public/USDLuxLightTranslator.h
--- a/Engine/Plugins/Importers/USDImporter/Source/USDSchemas/Public/USDLuxLightTranslator.h
+++ b/Engine/Plugins/Importers/USDImporter/Source/USDSchemas/Public/USDLuxLightTranslator.h
@@ -26,6 +26,9 @@ public:
 
 	virtual USceneComponent* CreateComponents() override;
 	virtual void UpdateComponents( USceneComponent* SceneComponent );
+
+	virtual bool CollapsesChildren( ECollapsingType CollapsingType ) const override;
+	virtual bool CanBeCollapsed( ECollapsingType CollapsingType ) const override;
 };
 
 #endif // #if USE_USD_SDK
diff --git a/Engine/Plugins/Importers/USDImporter/Source/USDSchemas/Public/USDShadeMaterialTranslator.h b/Engine/Plugins/Importers/USDImporter/Source/USDSchemas/Public/USDShadeMaterialTranslator.h
--- a/Engine/Plugins/Importers/USDImporter/Source/USDSchemas/Public/USDShadeMaterialTranslator.h
+++ b/Engine/Plugins/Importers/USDImporter/Source/USDSchemas/Public/USDShadeMaterialTranslator.h
@@ -17,6 +17,8 @@ public:
 
 	virtual void CreateAssets() override;
 
+	virtual bool CollapsesChildren( ECollapsingType CollapsingType ) const override;
+	virtual bool CanBeCollapsed( ECollapsingType CollapsingType ) const override;
 };
 
 #endif // #if USE_USD_SDK
