USD: Fix prim properties panel not updating when selecting actors and components via selection sync.

diff --git a/Engine/Plugins/Importers/USDImporter/Source/USDStageEditor/Private/SUSDStage.cpp b/Engine/Plugins/Importers/USDImporter/Source/USDStageEditor/Private/SUSDStage.cpp
--- a/Engine/Plugins/Importers/USDImporter/Source/USDStageEditor/Private/SUSDStage.cpp
+++ b/Engine/Plugins/Importers/USDImporter/Source/USDStageEditor/Private/SUSDStage.cpp
@@ -2046,6 +2046,11 @@ void SUsdStage::OnViewportSelectionChanged( UObject* NewSelection )
 	if ( PrimPaths.Num() > 0 )
 	{
 		UsdStageTreeView->SelectPrims( PrimPaths );
+
+		if ( UsdPrimInfoWidget )
+		{
+			UsdPrimInfoWidget->SetPrimPath( GetCurrentStage(), *PrimPaths[ PrimPaths.Num() - 1 ] );
+		}
 	}
 }
 
