USD: Remember import/export options dialog sizes.

diff --git a/Engine/Plugins/Importers/USDImporter/Source/USDStageImporter/Private/USDOptionsWindow.cpp b/Engine/Plugins/Importers/USDImporter/Source/USDStageImporter/Private/USDOptionsWindow.cpp
--- a/Engine/Plugins/Importers/USDImporter/Source/USDStageImporter/Private/USDOptionsWindow.cpp
+++ b/Engine/Plugins/Importers/USDImporter/Source/USDStageImporter/Private/USDOptionsWindow.cpp
@@ -9,10 +9,12 @@
 #include "UsdWrappers/UsdStage.h"
 #include "Widgets/SUSDStagePreviewTree.h"
 
+#include "CoreGlobals.h"
 #include "CoreMinimal.h"
 #include "Framework/Application/SlateApplication.h"
 #include "HAL/PlatformApplicationMisc.h"
 #include "Interfaces/IMainFrameModule.h"
+#include "Misc/ConfigCacheIni.h"
 #include "Misc/ScopedSlowTask.h"
 #include "Modules/ModuleManager.h"
 #include "PropertyEditorModule.h"
@@ -48,14 +50,38 @@ bool SUsdOptionsWindow::ShowOptions(
 		DisplayMetrics.PrimaryDisplayWorkAreaRect.Top
 	);
 
-	const FVector2D ClientSize{
-		320.0f * DPIScaleFactor * ( Stage == nullptr ? 1 : 2 ), // Make it wider if we're going to show two panes
-		320.0f * DPIScaleFactor
-	};
+	const static FString ImporterName = TEXT( "USDImporter" );
+	const FString ClassName = OptionsObject.GetClass()->GetName();
 
+	const FString WidthConfigName = ClassName + TEXT("_DialogWidth");
+	int32 Width = 100;
+	if ( GConfig->GetInt( *ImporterName, *WidthConfigName, Width, GEditorPerProjectIni ) )
+	{
+		// Just to make sure we don't load something unusable
+		Width = FMath::Min( FMath::Max( Width, 100 ), DisplayMetrics.PrimaryDisplayWidth );
+	}
+	else
+	{
+		// Make it wider if we're going to show two panes
+		Width = 500 * ( Stage == nullptr ? 1 : 2 );
+	}
+
+	const FString HeightConfigName = ClassName + TEXT( "_DialogHeight" );
+	int32 Height = 100;
+	if ( GConfig->GetInt( *ImporterName, *HeightConfigName, Height, GEditorPerProjectIni ) )
+	{
+		Height = FMath::Min( FMath::Max( Height, 100 ), DisplayMetrics.PrimaryDisplayHeight );
+	}
+	else
+	{
+		Height = 600;
+	}
+
+	const FVector2D ClientSize{ static_cast< float >( Width ), static_cast< float >( Height ) };
 	TSharedRef<SWindow> Window = SNew( SWindow )
 		.Title( WindowTitle )
 		.SizingRule( ESizingRule::UserSized )
+		.AdjustInitialSizeAndPositionForDPIScale( false )
 		.ClientSize( ClientSize );
 
 	TSharedPtr<SUsdOptionsWindow> OptionsWindow;
@@ -68,6 +94,23 @@ bool SUsdOptionsWindow::ShowOptions(
 		.Stage( Stage )
 	);
 
+	Window->GetOnWindowClosedEvent().AddLambda( [WidthConfigName, HeightConfigName, OptionsWindow]( const TSharedRef<SWindow>& Window )
+	{
+		// We use the geometry here because any sort of size we can extract from the window includes some added margins
+		// that are not easy to account for, even with GetWindowSizeFromClientSize. The original provided ClientSize
+		// seems to always match the content geometry though, so we're saving and loading the same thing
+		FGeometry Geometry = Window->GetContent()->GetCachedGeometry();
+		FVector2D ContentSize = Geometry.Size * Geometry.Scale;
+		const int32 Width = static_cast< int32 >( ContentSize.X );
+		const int32 Height = static_cast< int32 >( ContentSize.Y );
+
+		GConfig->SetInt( *ImporterName, *WidthConfigName, Width, GEditorPerProjectIni );
+		GConfig->SetInt( *ImporterName, *HeightConfigName, Height, GEditorPerProjectIni );
+
+		const bool bRemoveFromCache = false;
+		GConfig->Flush( bRemoveFromCache, GEditorPerProjectIni );
+	});
+
 	// Preemptively make sure we have a progress dialog created before showing our modal. This because the progress
 	// dialog itself is also modal. If it doesn't exist yet, and our options dialog causes a progress dialog
 	// to be spawned (e.g. when switching the Level to export via the LevelSequenceUSDExporter), the progress dialog
