USD: Automatically set the proper edit target for each subsequence edit.

diff --git a/Engine/Plugins/Importers/USDImporter/Source/USDStage/Private/USDLevelSequenceHelper.cpp b/Engine/Plugins/Importers/USDImporter/Source/USDStage/Private/USDLevelSequenceHelper.cpp
--- a/Engine/Plugins/Importers/USDImporter/Source/USDStage/Private/USDLevelSequenceHelper.cpp
+++ b/Engine/Plugins/Importers/USDImporter/Source/USDStage/Private/USDLevelSequenceHelper.cpp
@@ -99,35 +99,6 @@ namespace UsdLevelSequenceHelperImpl
 		return SanitizedText;
 	}
 
-	/** Sets the readonly value of the scene on construction and reverts it on destruction */
-	class FMovieSceneReadonlyGuard
-	{
-	public:
-#if WITH_EDITOR
-		explicit FMovieSceneReadonlyGuard( UMovieScene& InMovieScene, const bool bNewReadonlyValue )
-			: MovieScene( InMovieScene )
-			, bWasReadonly( InMovieScene.IsReadOnly() )
-		{
-			MovieScene.SetReadOnly( bNewReadonlyValue );
-		}
-
-		~FMovieSceneReadonlyGuard()
-		{
-			MovieScene.SetReadOnly( bWasReadonly );
-		}
-#else
-		explicit FMovieSceneReadonlyGuard( UMovieScene& InMovieScene, const bool bNewReadonlyValue )
-			: MovieScene( InMovieScene )
-			, bWasReadonly( true )
-		{
-		}
-#endif // WITH_EDITOR
-
-	private:
-		UMovieScene& MovieScene;
-		bool bWasReadonly;
-	};
-
 	/**
 	 * Similar to FrameRate.AsFrameNumber(TimeSeconds) except that it uses RoundToDouble instead of FloorToDouble, to
 	 * prevent issues with floating point precision
@@ -510,6 +481,7 @@ private:
 	ULevelSequence* FindOrAddSequenceForAttribute( const UE::FUsdAttribute& Attribute );
 	ULevelSequence* FindSequenceForIdentifier( const FString& SequenceIdentitifer );
 	ULevelSequence* FindOrAddSequenceForLayer( const UE::FSdfLayer& Layer, const FString& SequenceIdentifier, const FString& SequenceDisplayName );
+	UE::FSdfLayer FindEditTargetForSubsequence( ULevelSequence* Sequence );
 
 	/** Removes PrimTwin as a user of Sequence. If Sequence is now unused, remove its subsection and itself. */
 	void RemoveSequenceForPrim( ULevelSequence& Sequence, const UUsdPrimTwin& PrimTwin );
@@ -517,6 +489,7 @@ private:
 private:
 	ULevelSequence* MainLevelSequence;
 	TMap<FString, ULevelSequence*> LevelSequencesByIdentifier;
+	TMap<ULevelSequence*, FString> IdentifierByLevelSequence;
 
 	TSet< FName > LocalLayersSequences; // List of sequences associated with sublayers
 
@@ -623,11 +596,6 @@ private:
 	FDelegateHandle OnStageEditTargetChangedHandle;
 	FDelegateHandle OnUsdObjectsChangedHandle;
 
-// Readonly handling
-private:
-	void UpdateMovieSceneReadonlyFlags();
-	void UpdateMovieSceneReadonlyFlag( UMovieScene& MovieScene, const FString& LayerIdentifier );
-
 private:
 	void RefreshSequencer();
 
@@ -752,6 +720,7 @@ void FUsdLevelSequenceHelperImpl::Clear()
 {
 	MainLevelSequence = nullptr;
 	LevelSequencesByIdentifier.Empty();
+	IdentifierByLevelSequence.Empty();
 	LocalLayersSequences.Empty();
 	LayerIdentifierByLevelSequenceName.Empty();
 	LayerTimeInfosByLayerIdentifier.Empty();
@@ -845,12 +814,6 @@ void FUsdLevelSequenceHelperImpl::BindToUsdStageActor( AUsdStageActor* InStageAc
 		return;
 	}
 
-	OnStageEditTargetChangedHandle = StageActor->GetUsdListener().GetOnStageEditTargetChanged().AddLambda(
-		[ this ]()
-		{
-			UpdateMovieSceneReadonlyFlags();
-		});
-
 	OnUsdObjectsChangedHandle = StageActor->GetUsdListener().GetOnObjectsChanged().AddRaw( this, &FUsdLevelSequenceHelperImpl::OnUsdObjectsChanged );
 
 	// Bind stage actor
@@ -1042,11 +1005,11 @@ ULevelSequence* FUsdLevelSequenceHelperImpl::FindOrAddSequenceForLayer( const UE
 
 		LayerIdentifierByLevelSequenceName.Add( Sequence->GetFName(), Layer.GetIdentifier() );
 		LevelSequencesByIdentifier.Add( SequenceIdentifier, Sequence );
+		IdentifierByLevelSequence.Add( Sequence, SequenceIdentifier );
 
 		const FLayerTimeInfo LayerTimeInfo = FindOrAddLayerTimeInfo( Layer );
 
 		UpdateMovieSceneTimeRanges( *MovieScene, LayerTimeInfo );
-		UpdateMovieSceneReadonlyFlag( *MovieScene, LayerTimeInfo.Identifier );
 
 		UE_LOG( LogUsd, Verbose, TEXT("Created Sequence for identifier: '%s'"), *SequenceIdentifier );
 	}
@@ -1054,6 +1017,12 @@ ULevelSequence* FUsdLevelSequenceHelperImpl::FindOrAddSequenceForLayer( const UE
 	return Sequence;
 }
 
+UE::FSdfLayer FUsdLevelSequenceHelperImpl::FindEditTargetForSubsequence( ULevelSequence* Sequence )
+{
+	FString LayerIdentifier = IdentifierByLevelSequence.FindRef(Sequence);
+	return UE::FSdfLayer::FindOrOpen( *LayerIdentifier );
+}
+
 UMovieSceneSubSection* FUsdLevelSequenceHelperImpl::FindSubSequenceSection( ULevelSequence& Sequence, ULevelSequence& SubSequence )
 {
 	UMovieScene* MovieScene = Sequence.GetMovieScene();
@@ -1110,9 +1079,6 @@ void FUsdLevelSequenceHelperImpl::CreateSubSequenceSection( ULevelSequence& Sequ
 		return;
 	}
 
-	const bool bReadonly = false;
-	UsdLevelSequenceHelperImpl::FMovieSceneReadonlyGuard MovieSceneReadonlyGuard{ *MovieScene, bReadonly };
-
 	FFrameRate TickResolution = MovieScene->GetTickResolution();
 
 	UMovieSceneSubTrack* SubTrack = MovieScene->FindTrack< UMovieSceneSubTrack >();
@@ -1428,9 +1394,6 @@ void FUsdLevelSequenceHelperImpl::AddCommonTracks( const UUsdPrimTwin& PrimTwin,
 
 					if ( UMovieScene* MovieScene = AttributeSequence->GetMovieScene() )
 					{
-						const bool bReadonly = false;
-						UsdLevelSequenceHelperImpl::FMovieSceneReadonlyGuard MovieSceneReadonlyGuard{ *MovieScene, bReadonly };
-
 						TArray<double> TimeSamples;
 						if ( Xformable.GetTimeSamples( &TimeSamples ) )
 						{
@@ -1532,9 +1495,6 @@ void FUsdLevelSequenceHelperImpl::AddCommonTracks( const UUsdPrimTwin& PrimTwin,
 
 					if ( UMovieScene* MovieScene = AttributeSequence->GetMovieScene() )
 					{
-						const bool bReadonly = false;
-						UsdLevelSequenceHelperImpl::FMovieSceneReadonlyGuard MovieSceneReadonlyGuard{ *MovieScene, bReadonly };
-
 						if ( UMovieSceneVisibilityTrack* VisibilityTrack = AddTrack<UMovieSceneVisibilityTrack>( UnrealIdentifiers::HiddenInGamePropertyName, PrimTwin, *ComponentToBind, *AttributeSequence, bIsMuted ) )
 						{
 							UsdToUnreal::FPropertyTrackReader Reader = UsdToUnreal::CreatePropertyTrackReader( Prim, UnrealIdentifiers::HiddenInGamePropertyName );
@@ -1607,9 +1567,6 @@ void FUsdLevelSequenceHelperImpl::AddCameraTracks( const UUsdPrimTwin& PrimTwin,
 				continue;
 			}
 
-			const bool bReadonly = false;
-			UsdLevelSequenceHelperImpl::FMovieSceneReadonlyGuard MovieSceneReadonlyGuard{ *MovieScene, bReadonly };
-
 			TArray<double> TimeSamples;
 			if ( !Attr.GetTimeSamples( TimeSamples ) )
 			{
@@ -1742,9 +1699,6 @@ void FUsdLevelSequenceHelperImpl::AddLightTracks( const UUsdPrimTwin& PrimTwin,
 			continue;
 		}
 
-		const bool bReadonly = false;
-		UsdLevelSequenceHelperImpl::FMovieSceneReadonlyGuard MovieSceneReadonlyGuard{ *MovieScene, bReadonly };
-
 		UsdToUnreal::FPropertyTrackReader Reader = UsdToUnreal::CreatePropertyTrackReader( Prim, PropertyPath );
 
 		switch ( TrackType )
@@ -1837,9 +1791,6 @@ void FUsdLevelSequenceHelperImpl::AddSkeletalTracks( const UUsdPrimTwin& PrimTwi
 		return;
 	}
 
-	const bool bReadonly = false;
-	UsdLevelSequenceHelperImpl::FMovieSceneReadonlyGuard MovieSceneReadonlyGuard{ *MovieScene, bReadonly };
-
 	// We will mute all SkelAnimation attributes if we mute, so here let's only consider something muted
 	// if it has all attributes muted as well.
 	// We know at least one of these attributes ones is valid and animated because we have an UAnimSequence
@@ -1914,9 +1865,6 @@ void FUsdLevelSequenceHelperImpl::AddGroomTracks( const UUsdPrimTwin& PrimTwin,
 		return;
 	}
 
-	const bool bReadonly = false;
-	UsdLevelSequenceHelperImpl::FMovieSceneReadonlyGuard MovieSceneReadonlyGuard{ *MovieScene, bReadonly };
-
 	const bool bIsMuted = false;
 	if ( UMovieSceneGroomCacheTrack* GroomCacheTrack = AddTrack< UMovieSceneGroomCacheTrack >( Prim.GetName(), PrimTwin, *ComponentToBind, *GroomAnimationSequence, bIsMuted ) )
 	{
@@ -2003,9 +1951,6 @@ TrackType* FUsdLevelSequenceHelperImpl::AddTrack( const FName& TrackName, const
 
 	const FGuid ComponentBinding = GetOrCreateComponentBinding( PrimTwin, ComponentToBind, Sequence );
 
-	const bool bReadonly = false;
-	UsdLevelSequenceHelperImpl::FMovieSceneReadonlyGuard MovieSceneReadonlyGuard{ *MovieScene, bReadonly };
-
 	TrackType* Track = MovieScene->FindTrack< TrackType >( ComponentBinding, TrackName );
 	if ( Track )
 	{
@@ -2165,9 +2110,6 @@ void FUsdLevelSequenceHelperImpl::UpdateControlRigTracks( UUsdPrimTwin& PrimTwin
 		return;
 	}
 
-	const bool bReadonly = false;
-	UsdLevelSequenceHelperImpl::FMovieSceneReadonlyGuard MovieSceneReadonlyGuard{ *MovieScene, bReadonly };
-
 	const FGuid ComponentBinding = GetOrCreateComponentBinding( PrimTwin, *ComponentToBind, *SkelAnimationSequence );
 
 	// NOTE: We are fetching the first skel track we find, since we can't actually use SkelAnimationPrim.GetName() here at all!
@@ -2296,6 +2238,7 @@ void FUsdLevelSequenceHelperImpl::RemoveSequenceForPrim( ULevelSequence& Sequenc
 			}
 
 			LevelSequencesByIdentifier.Remove( PrimTwin.PrimPath );
+			IdentifierByLevelSequence.Remove( &Sequence );
 			SequencesID.Remove( &Sequence );
 		}
 	}
@@ -2488,37 +2431,6 @@ void FUsdLevelSequenceHelperImpl::UpdateUsdLayerOffsetFromSection(const UMovieSc
 	UE_LOG(LogUsd, Verbose, TEXT("\tNew OffsetScale: %f, %f"), NewLayerOffset.Offset, NewLayerOffset.Scale);
 }
 
-void FUsdLevelSequenceHelperImpl::UpdateMovieSceneReadonlyFlags()
-{
-	for ( const TPair< FString, ULevelSequence* >& SequenceIndentifierToSequence : LevelSequencesByIdentifier )
-	{
-		if ( ULevelSequence* Sequence = SequenceIndentifierToSequence.Value )
-		{
-			if ( FString* LayerIdentifier = LayerIdentifierByLevelSequenceName.Find( Sequence->GetFName() ) )
-			{
-				if ( UMovieScene* MovieScene = Sequence->GetMovieScene() )
-				{
-					UpdateMovieSceneReadonlyFlag( *MovieScene, *LayerIdentifier );
-				}
-			}
-		}
-	}
-}
-
-void FUsdLevelSequenceHelperImpl::UpdateMovieSceneReadonlyFlag( UMovieScene& MovieScene, const FString& LayerIdentifier )
-{
-#if WITH_EDITOR
-	if ( !UsdStage )
-	{
-		return;
-	}
-
-	UE::FSdfLayer Layer = UE::FSdfLayer::FindOrOpen( *LayerIdentifier );
-	const bool bIsReadOnly = ( Layer != UsdStage.GetEditTarget() );
-	MovieScene.SetReadOnly( bIsReadOnly );
-#endif // WITH_EDITOR
-}
-
 void FUsdLevelSequenceHelperImpl::UpdateMovieSceneTimeRanges( UMovieScene& MovieScene, const FLayerTimeInfo& LayerTimeInfo )
 {
 	const double FramesPerSecond = GetFramesPerSecond();
@@ -3055,6 +2967,15 @@ void FUsdLevelSequenceHelperImpl::HandleControlRigSectionChange( UMovieSceneCont
 		return;
 	}
 
+	// Each sequence corresponds to a specific USD layer. If we're editing something in a sequence, then we must make
+	// that layer the edit target too
+	UE::FSdfLayer EditTargetLayer = FindEditTargetForSubsequence( LevelSequence );
+	if ( !EditTargetLayer )
+	{
+		return;
+	}
+	UE::FUsdEditContext EditContext{ UsdStage, EditTargetLayer };
+
 	TSharedPtr<ISequencer> PinnedSequencer = UsdLevelSequenceHelperImpl::GetOpenedSequencerForLevelSequence( MainLevelSequence );
 
 	// Fetch a sequence player we can use. We'll almost always have the sequencer opened here (we are responding to a transaction
@@ -3229,6 +3150,15 @@ void FUsdLevelSequenceHelperImpl::HandleTrackChange( const UMovieSceneTrack& Tra
 		}
 	}
 
+	// Each sequence corresponds to a specific USD layer. If we're editing something in a sequence, then we must make
+	// that layer the edit target too
+	UE::FSdfLayer EditTargetLayer = FindEditTargetForSubsequence( Sequence );
+	if ( !EditTargetLayer )
+	{
+		return;
+	}
+	UE::FUsdEditContext EditContext{ UsdStage, EditTargetLayer };
+
 	if ( PrimTwin )
 	{
 		FScopedBlockNoticeListening BlockNotices( StageActor.Get() );
