Added missing include "Engine/Level.h"

diff --git a/Engine/Plugins/Importers/USDImporter/Source/USDStage/Private/USDStageActor.cpp b/Engine/Plugins/Importers/USDImporter/Source/USDStage/Private/USDStageActor.cpp
--- a/Engine/Plugins/Importers/USDImporter/Source/USDStage/Private/USDStageActor.cpp
+++ b/Engine/Plugins/Importers/USDImporter/Source/USDStage/Private/USDStageActor.cpp
@@ -44,6 +44,7 @@
 #include "Components/SpotLightComponent.h"
 #include "Components/StaticMeshComponent.h"
 #include "Engine/Engine.h"
+#include "Engine/Level.h"
 #include "Engine/Light.h"
 #include "Engine/SkeletalMesh.h"
 #include "Engine/StaticMesh.h"
