USD: Merge cl 19383968 from Partner-Martmot427 to fix double-transformation being applied to foliage actors.

diff --git a/Engine/Plugins/Importers/USDImporter/Source/USDUtilities/Private/USDPrimConversion.cpp b/Engine/Plugins/Importers/USDImporter/Source/USDUtilities/Private/USDPrimConversion.cpp
--- a/Engine/Plugins/Importers/USDImporter/Source/USDUtilities/Private/USDPrimConversion.cpp
+++ b/Engine/Plugins/Importers/USDImporter/Source/USDUtilities/Private/USDPrimConversion.cpp
@@ -2421,6 +2421,7 @@ namespace UE
 			void ConvertFoliageInstances(
 				const FFoliageInfo& Info,
 				const TSet<int32>& UEInstances,
+				const FTransform& UEWorldToFoliageActor,
 				const FUsdStageInfo& StageInfo,
 				int PrototypeIndex,
 				pxr::VtArray<int>& ProtoIndices,
@@ -2443,8 +2444,8 @@ namespace UE
 					const FFoliageInstancePlacementInfo* Instance = &Info.Instances[ InstanceIndex ];
 
 					// Convert axes
-					FTransform UETransform{ Instance->Rotation, (FVector)Instance->Location, (FVector)Instance->DrawScale3D };
-					FTransform USDTransform = UsdUtils::ConvertAxes( StageInfo.UpAxis == EUsdUpAxis::ZAxis, UETransform );
+					FTransform UEWorldTransform{ Instance->Rotation, (FVector)Instance->Location, (FVector)Instance->DrawScale3D };
+					FTransform USDTransform = UsdUtils::ConvertAxes( StageInfo.UpAxis == EUsdUpAxis::ZAxis, UEWorldTransform * UEWorldToFoliageActor );
 
 					FVector Translation = USDTransform.GetTranslation();
 					FQuat Rotation = USDTransform.GetRotation();
@@ -2491,6 +2492,10 @@ bool UnrealToUsd::ConvertInstancedFoliageActor( const AInstancedFoliageActor& Ac
 
 	TSet<FFoliageInstanceBaseId> HandledComponents;
 
+	// It seems like the foliage instance transforms are actually world transforms, so to get them into the coordinate space of the generated
+	// point instancer, we'll have to concatenate with the inverse the foliage actor's ActorToWorld transform
+	FTransform UEWorldToFoliageActor = Actor.GetTransform().Inverse();
+
 	int PrototypeIndex = 0;
 	for ( const TPair<UFoliageType*, TUniqueObj<FFoliageInfo>>& FoliagePair : Actor.GetFoliageInfos() )
 	{
@@ -2511,7 +2516,17 @@ bool UnrealToUsd::ConvertInstancedFoliageActor( const AInstancedFoliageActor& Ac
 
 			if ( const TSet<int32>* InstanceSet = Info.ComponentHash.Find( ComponentId ) )
 			{
-				UE::USDPrimConversionImpl::Private::ConvertFoliageInstances( Info, *InstanceSet, StageInfo, PrototypeIndex, ProtoIndices, Positions, Orientations, Scales );
+				UE::USDPrimConversionImpl::Private::ConvertFoliageInstances(
+					Info,
+					*InstanceSet,
+					UEWorldToFoliageActor,
+					StageInfo,
+					PrototypeIndex,
+					ProtoIndices,
+					Positions,
+					Orientations,
+					Scales
+				);
 			}
 		}
 
@@ -2530,7 +2545,17 @@ bool UnrealToUsd::ConvertInstancedFoliageActor( const AInstancedFoliageActor& Ac
 				}
 
 				const TSet<int32>& InstanceSet = Pair.Value;
-				UE::USDPrimConversionImpl::Private::ConvertFoliageInstances( Info, InstanceSet, StageInfo, PrototypeIndex, ProtoIndices, Positions, Orientations, Scales );
+				UE::USDPrimConversionImpl::Private::ConvertFoliageInstances(
+					Info,
+					InstanceSet,
+					UEWorldToFoliageActor,
+					StageInfo,
+					PrototypeIndex,
+					ProtoIndices,
+					Positions,
+					Orientations,
+					Scales
+				);
 			}
 		}
 
