USD: Remove usage of SequencerActor tag.

diff --git a/Engine/Plugins/Importers/USDImporter/Source/USDExporter/Private/LevelSequenceExporterUSD.cpp b/Engine/Plugins/Importers/USDImporter/Source/USDExporter/Private/LevelSequenceExporterUSD.cpp
--- a/Engine/Plugins/Importers/USDImporter/Source/USDExporter/Private/LevelSequenceExporterUSD.cpp
+++ b/Engine/Plugins/Importers/USDImporter/Source/USDExporter/Private/LevelSequenceExporterUSD.cpp
@@ -200,13 +200,6 @@ namespace UE
 
 						if ( AActor* SpawnedActor = Cast<AActor>( Object ) )
 						{
-							// HACK: Remove the SequencerActor tag from the spawned actor because the USD level exporter
-							// uses these as an indicator that it shouldn't export a given actor, because the transient actors spawned by the
-							// sequencer and the AUsdStageActor contain them.
-							// We will manually delete all of those actors with our spawn register later and never actually use them
-							// for anything else though, so this should be fine
-							SpawnedActor->Tags.Remove( TEXT( "SequencerActor" ) );
-
 							// Rename the spawn to a unique name or else in case of name collisions they will overwrite each other when writing animation data.
 							// The level exporter will rename actors to unique prims by itself though.
 							FString NewLabel = UsdUtils::GetUniqueName( SpawnedActor->GetActorLabel(), UsedActorLabels );
diff --git a/Engine/Plugins/Importers/USDImporter/Source/USDExporter/Private/USDConversionBlueprintLibrary.cpp b/Engine/Plugins/Importers/USDImporter/Source/USDExporter/Private/USDConversionBlueprintLibrary.cpp
--- a/Engine/Plugins/Importers/USDImporter/Source/USDExporter/Private/USDConversionBlueprintLibrary.cpp
+++ b/Engine/Plugins/Importers/USDImporter/Source/USDExporter/Private/USDConversionBlueprintLibrary.cpp
@@ -288,6 +288,20 @@ TSet<AActor*> UUsdConversionBlueprintLibrary::GetActorsToConvert( UWorld* World
 	}
 
 	Result.Remove( nullptr );
+
+	// Remove transient actors here because it is not possible to do this via Python
+	TSet<AActor*> ActorsToRemove;
+	ActorsToRemove.Reserve( Result.Num() );
+	for ( AActor* Actor : Result )
+	{
+		// Actors marked with this tag are transient because they're spawnables: We still want to export those, in case
+		// we're exporting a level for a LevelSequence export with spawnables.
+		if ( Actor->HasAnyFlags( EObjectFlags::RF_Transient ) && !Actor->Tags.Contains( TEXT( "SequencerActor" ) ) )
+		{
+			ActorsToRemove.Add( Actor );
+		}
+	}
+	Result = Result.Difference( ActorsToRemove );
 	return Result;
 }
 
diff --git a/Engine/Plugins/Importers/USDImporter/Source/USDSchemas/Private/USDGeomXformableTranslator.cpp b/Engine/Plugins/Importers/USDImporter/Source/USDSchemas/Private/USDGeomXformableTranslator.cpp
--- a/Engine/Plugins/Importers/USDImporter/Source/USDSchemas/Private/USDGeomXformableTranslator.cpp
+++ b/Engine/Plugins/Importers/USDImporter/Source/USDSchemas/Private/USDGeomXformableTranslator.cpp
@@ -413,12 +413,6 @@ USceneComponent* FUsdGeomXformableTranslator::CreateComponentsEx( TOptional< TSu
 			}
 #endif // WITH_EDITOR
 
-			// Hack to show transient actors in world outliner
-			if (SpawnedActor->HasAnyFlags(EObjectFlags::RF_Transient))
-			{
-				SpawnedActor->Tags.AddUnique( TEXT("SequencerActor") );
-			}
-
 			SceneComponent = SpawnedActor->GetRootComponent();
 
 			ComponentOuter = SpawnedActor;
