USD: Fix comparison operators betwee UE::FSdfLayer and UE::FSdfLayerWeak not working on Clang C++20 toolchain.

diff --git a/Engine/Plugins/Importers/USDImporter/Source/UnrealUSDWrapper/Private/UsdWrappers/SdfLayer.cpp b/Engine/Plugins/Importers/USDImporter/Source/UnrealUSDWrapper/Private/UsdWrappers/SdfLayer.cpp
--- a/Engine/Plugins/Importers/USDImporter/Source/UnrealUSDWrapper/Private/UsdWrappers/SdfLayer.cpp
+++ b/Engine/Plugins/Importers/USDImporter/Source/UnrealUSDWrapper/Private/UsdWrappers/SdfLayer.cpp
@@ -126,16 +126,26 @@ namespace UE
 	}
 
 	template<typename PtrType>
-	bool FSdfLayerBase<PtrType>::operator==( const FSdfLayerBase& Other ) const
+	template<typename OtherPtrType>
+	bool FSdfLayerBase<PtrType>::operator==( const FSdfLayerBase<OtherPtrType>& Other ) const
 	{
 		return Impl->GetInner() == Other.Impl->GetInner();
 	}
+	template bool UNREALUSDWRAPPER_API FSdfLayer::operator==( const FSdfLayer& Other ) const;
+	template bool UNREALUSDWRAPPER_API FSdfLayer::operator==( const FSdfLayerWeak& Other ) const;
+	template bool UNREALUSDWRAPPER_API FSdfLayerWeak::operator==( const FSdfLayer& Other ) const;
+	template bool UNREALUSDWRAPPER_API FSdfLayerWeak::operator==( const FSdfLayerWeak& Other ) const;
 
 	template<typename PtrType>
-	bool FSdfLayerBase<PtrType>::operator!=( const FSdfLayerBase& Other ) const
+	template<typename OtherPtrType>
+	bool FSdfLayerBase<PtrType>::operator!=( const FSdfLayerBase<OtherPtrType>& Other ) const
 	{
 		return !( *this == Other );
 	}
+	template bool UNREALUSDWRAPPER_API FSdfLayer::operator!=( const FSdfLayer& Other ) const;
+	template bool UNREALUSDWRAPPER_API FSdfLayer::operator!=( const FSdfLayerWeak& Other ) const;
+	template bool UNREALUSDWRAPPER_API FSdfLayerWeak::operator!=( const FSdfLayer& Other ) const;
+	template bool UNREALUSDWRAPPER_API FSdfLayerWeak::operator!=( const FSdfLayerWeak& Other ) const;
 
 	template<typename PtrType>
 	FSdfLayerBase<PtrType>::operator bool() const
diff --git a/Engine/Plugins/Importers/USDImporter/Source/UnrealUSDWrapper/Public/UsdWrappers/SdfLayer.h b/Engine/Plugins/Importers/USDImporter/Source/UnrealUSDWrapper/Public/UsdWrappers/SdfLayer.h
--- a/Engine/Plugins/Importers/USDImporter/Source/UnrealUSDWrapper/Public/UsdWrappers/SdfLayer.h
+++ b/Engine/Plugins/Importers/USDImporter/Source/UnrealUSDWrapper/Public/UsdWrappers/SdfLayer.h
@@ -51,8 +51,10 @@ namespace UE
 		FSdfLayerBase& operator=( const FSdfLayerWeak& Other );
 		FSdfLayerBase& operator=( FSdfLayerWeak&& Other );
 
-		bool operator==( const FSdfLayerBase& Other ) const;
-		bool operator!=( const FSdfLayerBase& Other ) const;
+		template<typename OtherPtrType>
+		bool operator==( const FSdfLayerBase<OtherPtrType>& Other ) const;
+		template<typename OtherPtrType>
+		bool operator!=( const FSdfLayerBase<OtherPtrType>& Other ) const;
 
 		explicit operator bool() const;
 
