USD: Silence some warnings generated on the USD headers on C++20 targets.

diff --git a/Engine/Plugins/Importers/USDImporter/Source/UnrealUSDWrapper/Public/USDIncludesStart.h b/Engine/Plugins/Importers/USDImporter/Source/UnrealUSDWrapper/Public/USDIncludesStart.h
--- a/Engine/Plugins/Importers/USDImporter/Source/UnrealUSDWrapper/Public/USDIncludesStart.h
+++ b/Engine/Plugins/Importers/USDImporter/Source/UnrealUSDWrapper/Public/USDIncludesStart.h
@@ -11,6 +11,8 @@ THIRD_PARTY_INCLUDES_START
 #pragma warning(disable: 4003) /* pxr/usd/sdf/fileFormat.h BOOST_PP_SEQ_DETAIL_IS_NOT_EMPTY during static analysis */
 #pragma warning(disable: 5033) /* 'register' is no longer a supported storage class */
 #pragma warning(disable: 6319)
+#pragma warning(disable: 5232) /* Silence some warnings about operator==(const pxr::TfRefPtr<pxr::SdfLayer> &,const pxr::TfWeakPtr<pxr::SdfLayer> &) */
+#pragma warning(disable: 4717) /* calling itself recursively within the USD SDK's weakPtrFacade.h that we get when compiling C++20 targets (UE-169305) */
 
 #pragma push_macro("check")
 #undef check // Boost function is named 'check' in boost\python\detail\convertible.hpp
