USD: Read/write textures according to a specified color space.

Export EmissiveColor as EXR;
Remove usage of EmissiveScale;

diff --git a/Engine/Plugins/Importers/USDImporter/Source/USDUtilities/Private/USDShadeConversion.cpp b/Engine/Plugins/Importers/USDImporter/Source/USDUtilities/Private/USDShadeConversion.cpp
--- a/Engine/Plugins/Importers/USDImporter/Source/USDUtilities/Private/USDShadeConversion.cpp
+++ b/Engine/Plugins/Importers/USDImporter/Source/USDUtilities/Private/USDShadeConversion.cpp
@@ -91,22 +91,26 @@ namespace UE
 			struct FBakedMaterialView
 			{
 				TMap<EMaterialProperty, TArray<FColor>*> PropertyData;
+				TMap<EMaterialProperty, TArray<FFloat16Color>*> HDRPropertyData;
 				TMap<EMaterialProperty, FIntPoint> PropertySizes;
-				float EmissiveScale;
 
 				explicit FBakedMaterialView( FBakeOutput& BakeOutput )
 					: PropertySizes( BakeOutput.PropertySizes )
-					, EmissiveScale( BakeOutput.EmissiveScale )
 				{
+					HDRPropertyData.Reserve( BakeOutput.HDRPropertyData.Num() );
+					for ( TPair<EMaterialProperty, TArray<FFloat16Color>>& BakedPair : BakeOutput.HDRPropertyData )
+					{
+						HDRPropertyData.Add( BakedPair.Key, &BakedPair.Value );
+					}
+
 					PropertyData.Reserve( BakeOutput.PropertyData.Num() );
 					for ( TPair<EMaterialProperty, TArray<FColor>>& BakedPair : BakeOutput.PropertyData )
 					{
-						PropertyData.Add(BakedPair.Key, &BakedPair.Value);
+						PropertyData.Add( BakedPair.Key, &BakedPair.Value );
 					}
 				}
 
 				explicit FBakedMaterialView( FFlattenMaterial& FlattenMaterial )
-					: EmissiveScale( FlattenMaterial.EmissiveScale )
 				{
 					PropertyData.Reserve( static_cast< uint8 >( EFlattenMaterialProperties::NumFlattenMaterialProperties ) );
 					for ( const TPair< EFlattenMaterialProperties, EMaterialProperty>& FlattenToProperty : FlattenToMaterialProperty )
@@ -595,9 +599,26 @@ namespace UE
 							bSRGB = false;
 							CompressionSettings = TextureCompressionSettings::TC_Normalmap;
 						}
-						else if ( LODGroup == TEXTUREGROUP_WorldSpecular )
+						else
 						{
-							bSRGB = false;
+							if ( LODGroup == TEXTUREGROUP_WorldSpecular )
+							{
+								bSRGB = false;
+							}
+
+							// Override the sRGB flag with whatever is specified on the texture shader, if it has anything specific.
+							if ( pxr::UsdAttribute SourceColorSpaceAttr = UsdUVTextureSource.GetInput( UnrealIdentifiers::SourceColorSpaceToken ) )
+							{
+								pxr::TfToken SourceColorSpaceValue;
+								if ( SourceColorSpaceAttr.HasAuthoredValue() && SourceColorSpaceAttr.Get( &SourceColorSpaceValue ) )
+								{
+									bSRGB = SourceColorSpaceValue == UnrealIdentifiers::RawColorSpaceToken
+										? false
+										: SourceColorSpaceValue == UnrealIdentifiers::SRGBColorSpaceToken
+											? true
+											: bSRGB;
+								}
+							}
 						}
 
 						const FString TextureHash = GetTextureHash( TexturePath, bSRGB, CompressionSettings, AddressX, AddressY );
@@ -1393,8 +1414,10 @@ namespace UE
 
 				TArray<FBakeOutput> BakeOutputs;
 				IMaterialBakingModule& Module = FModuleManager::Get().LoadModuleChecked<IMaterialBakingModule>( "MaterialBaking" );
-				bool bLinearBake = true;
+				const bool bLinearBake = true;
 				Module.SetLinearBake( bLinearBake );
+				const bool bEmissiveHDR = true;
+				Module.SetEmissiveHDR( bEmissiveHDR );
 				Module.BakeMaterials( { &MatSet }, { &MeshSettings }, BakeOutputs );
 
 				if ( BakeOutputs.Num() < 1 )
@@ -1410,15 +1433,74 @@ namespace UE
 			TMap<EMaterialProperty, FString> WriteTextures( FBakedMaterialView& BakedSamples, const FString& TextureNamePrefix, const FDirectoryPath& TexturesFolder )
 			{
 				IImageWrapperModule& ImageWrapperModule = FModuleManager::Get().LoadModuleChecked<IImageWrapperModule>( TEXT( "ImageWrapper" ) );
+				TSharedPtr<IImageWrapper> EXRImageWrapper = ImageWrapperModule.CreateImageWrapper( EImageFormat::EXR );
 				TSharedPtr<IImageWrapper> PNGImageWrapper = ImageWrapperModule.CreateImageWrapper( EImageFormat::PNG );
 
 				TMap<EMaterialProperty, FString> WrittenTexturesPerChannel;
 
+				// Write textures for HDR baked properties larger than 1x1
+				for ( TPair<EMaterialProperty, TArray<FFloat16Color>*>& BakedDataPair : BakedSamples.HDRPropertyData )
+				{
+					const EMaterialProperty Property = BakedDataPair.Key;
+
+					TArray<FFloat16Color>* Samples = BakedDataPair.Value;
+					if ( !Samples || Samples->Num() < 2 )
+					{
+						continue;
+					}
+
+					const FIntPoint& FinalSize = BakedSamples.PropertySizes.FindChecked( Property );
+					if ( FinalSize.GetMin() < 2 )
+					{
+						continue;
+					}
+
+					const UEnum* PropertyEnum = StaticEnum<EMaterialProperty>();
+					FName PropertyName = PropertyEnum->GetNameByValue( Property );
+					FString TrimmedPropertyName = PropertyName.ToString();
+					TrimmedPropertyName.RemoveFromStart( TEXT( "MP_" ) );
+
+					// Final FilePath will be something like "C:/TexturesFolder/Game_ContentFolder_Materials_Red_BaseColor.png", which automatically guarantees
+					// it won't overwrite another texture from the same export, but will overwrite old textures from previous exports
+					FString TextureFileName = ObjectTools::SanitizeObjectName( FPaths::ChangeExtension( TextureNamePrefix, TEXT( "" ) ) );
+					TextureFileName.RemoveFromStart( TEXT( "_" ) );
+					FString TextureFilePath = FPaths::Combine( TexturesFolder.Path, FString::Printf( TEXT( "%s_%s.exr" ), *ObjectTools::SanitizeObjectName( TextureFileName ), *TrimmedPropertyName ) );
+
+					// For some reason the baked samples always have zero alpha and there is nothing we can do about it... It seems like the material baking module is made
+					// with the intent that the data ends up in UTexture2Ds, where they can be set to be compressed without alpha and have the value ignored.
+					// Since we need to write these to file, we must set them back up to full alpha. This is potentially useless as USD handles these at most as color3f, but
+					// it could be annoying for the user if they intend on using the textures for anything else
+					for ( FFloat16Color& Sample : *Samples )
+					{
+						Sample.A = 1.0f;
+					}
+
+					EXRImageWrapper->SetRaw( Samples->GetData(), Samples->GetAllocatedSize(), FinalSize.X, FinalSize.Y, ERGBFormat::RGBAF, 16 );
+					const TArray64<uint8> Data = EXRImageWrapper->GetCompressed( 100 );
+
+					bool bWroteFile = FFileHelper::SaveArrayToFile( Data, *TextureFilePath );
+					if ( bWroteFile )
+					{
+						WrittenTexturesPerChannel.Add( Property, TextureFilePath );
+					}
+					else
+					{
+						UE_LOG( LogUsd, Warning, TEXT( "Failed to write texture '%s', baked channel will be ignored." ), *TextureFilePath );
+					}
+				}
+
 				// Write textures for baked properties larger than 1x1
 				for ( TPair<EMaterialProperty, TArray<FColor>*>& BakedDataPair : BakedSamples.PropertyData )
 				{
 					const EMaterialProperty Property = BakedDataPair.Key;
 
+					// The material baking module still generates and sends an SDR version of any HDR channel it also bakes,
+					// so here let's skip this one in case we already generated an HDR texture for the channel
+					if ( WrittenTexturesPerChannel.Contains( Property ) )
+					{
+						continue;
+					}
+
 					TArray<FColor>* Samples = BakedDataPair.Value;
 					if ( !Samples || Samples->Num() < 2 )
 					{
@@ -1502,42 +1584,110 @@ namespace UE
 				// Collect all the properties we'll process, as some data is baked and some comes from values the user input as export options
 				TSet<EMaterialProperty> PropertiesToProcess;
 				{
-					PropertiesToProcess.Reserve( BakedData.PropertyData.Num() + UserConstantValues.Num() );
+					PropertiesToProcess.Reserve( BakedData.PropertyData.Num() + UserConstantValues.Num() + BakedData.HDRPropertyData.Num() );
 					BakedData.PropertyData.GetKeys( PropertiesToProcess );
 
-					TSet<EMaterialProperty> PropertiesWithUserConstantData;
-					UserConstantValues.GetKeys( PropertiesWithUserConstantData );
+					TSet<EMaterialProperty> UsedProperties;
+					UserConstantValues.GetKeys( UsedProperties );
+					PropertiesToProcess.Append( UsedProperties );
 
-					PropertiesToProcess.Append(PropertiesWithUserConstantData);
+					BakedData.HDRPropertyData.GetKeys( UsedProperties );
+					PropertiesToProcess.Append( UsedProperties );
 				}
 
-				bool bHasEmissive = false;
-
 				// Fill in outputs
 				for ( const EMaterialProperty& Property : PropertiesToProcess )
 				{
-					const TArray<FColor>* Samples = BakedData.PropertyData.FindRef( Property );
-					const FIntPoint* SampleSize = BakedData.PropertySizes.Find( Property );
-					const float* UserConstantValue = UserConstantValues.Find( Property );
 					const FString* TextureFilePath = WrittenTextures.Find( Property );
-					const bool bPropertyValueIsConstant = ( UserConstantValue != nullptr || ( Samples && Samples->Num() == 1 ) );
+					const float* UserConstantValue = UserConstantValues.Find( Property );
+					const FIntPoint* SampleSize = BakedData.PropertySizes.Find( Property );
+
+					uint64 NumSamples = 0;
+
+					bool bPropertyValueIsConstant = false;
+					pxr::GfVec3f ConstantLinearValue;
 
-					if ( ( !Samples || !SampleSize ) && !UserConstantValue )
+					// Try HDR first: When baking a channel as HDR (like emissive) the MaterialBaking module
+					// will still bake the SDR version of the channel too. We keep both in our FBakedMaterialView
+					// because it only does the mechanism of decaying to a single value on the SDR data array
+					bool bParsedHDR = false;
+					if ( BakedData.HDRPropertyData.Contains( Property ) )
 					{
-						UE_LOG( LogUsd, Warning, TEXT( "Skipping material property %d as we have no valid data to use." ), Property );
-						continue;
+						const TArray<FColor>* SDRSamples = BakedData.PropertyData.FindRef( Property );
+						NumSamples = SDRSamples ? SDRSamples->Num() : 0;
+
+						const bool bDecayedToSingleSample = NumSamples == 1;
+						bPropertyValueIsConstant = ( UserConstantValue != nullptr || bDecayedToSingleSample );
+
+						const TArray<FFloat16Color>* Samples = BakedData.HDRPropertyData.FindRef( Property );
+						if ( bDecayedToSingleSample && Samples )
+						{
+							// If it decayed to single sample we know our SDR array only has one value,
+							// and so should our HDR array. Get the value from the HDR one to avoid an
+							// unnecessary quantization though
+							const FFloat16Color& Sample = ( *Samples )[ 0 ];
+							ConstantLinearValue = pxr::GfVec3f{ Sample.R, Sample.G, Sample.B };
+						}
+
+						if ( Samples && Samples->Num() > 0 )
+						{
+							bParsedHDR = true;
+						}
 					}
 
-					if ( !UserConstantValue && Samples && SampleSize )
+					if ( !bParsedHDR )
 					{
-						if ( Samples->Num() != SampleSize->X * SampleSize->Y )
+						const TArray<FColor>* Samples = BakedData.PropertyData.FindRef( Property );
+						NumSamples = Samples ? Samples->Num() : 0;
+
+						bPropertyValueIsConstant = ( UserConstantValue != nullptr || NumSamples == 1 );
+						if ( NumSamples == 1 )
 						{
-							UE_LOG( LogUsd, Warning, TEXT( "Skipping material property %d as it has an unexpected number of samples (%d instead of %d)." ), Property, Samples->Num(), SampleSize->X * SampleSize->Y );
-							continue;
+							switch ( Property )
+							{
+								case EMaterialProperty::MP_BaseColor:
+								case EMaterialProperty::MP_SubsurfaceColor:
+								{
+									pxr::GfVec4f ConvertedColor = UnrealToUsd::ConvertColor( ( *Samples )[ 0 ] );
+									ConstantLinearValue = pxr::GfVec3f{
+										ConvertedColor[ 0 ],
+										ConvertedColor[ 1 ],
+										ConvertedColor[ 2 ]
+									};
+									break;
+								}
+								case EMaterialProperty::MP_Normal:
+								case EMaterialProperty::MP_Tangent:
+								{
+									FVector ConvertedNormal{ ( *Samples )[ 0 ].ReinterpretAsLinear() };
+									ConstantLinearValue = UnrealToUsd::ConvertVector( StageInfo, ConvertedNormal );
+									break;
+								}
+								default:
+								{
+									const FColor& Sample = ( *Samples )[ 0 ];
+									ConstantLinearValue = pxr::GfVec3f{
+										Sample.R / 255.0f,
+										Sample.G / 255.0f,
+										Sample.B / 255.0f
+									};
+									break;
+								}
+							}
 						}
-						else if ( Samples->Num() == 0 )
+					}
+
+					if ( ( NumSamples == 0 || !SampleSize ) && !UserConstantValue )
+					{
+						UE_LOG( LogUsd, Warning, TEXT( "Skipping material property %d as we have no valid data to use." ), Property );
+						continue;
+					}
+
+					if ( !UserConstantValue && NumSamples > 0 && SampleSize )
+					{
+						if ( NumSamples != SampleSize->X * SampleSize->Y )
 						{
-							// Silently ignore this channel if we have an expected number of zero samples
+							UE_LOG( LogUsd, Warning, TEXT( "Skipping material property %d as it has an unexpected number of samples (%d instead of %d)." ), Property, NumSamples, SampleSize->X * SampleSize->Y );
 							continue;
 						}
 					}
@@ -1552,6 +1702,7 @@ namespace UE
 					pxr::SdfValueTypeName InputType;
 					pxr::VtValue ConstantValue;
 					pxr::GfVec4f FallbackValue;
+					pxr::TfToken ColorSpaceToken = UnrealIdentifiers::RawColorSpaceToken;
 					switch ( Property )
 					{
 					case MP_BaseColor:
@@ -1565,28 +1716,28 @@ namespace UE
 							}
 							else
 							{
-								pxr::GfVec4f ConvertedColor = UnrealToUsd::ConvertColor( ( *Samples )[ 0 ] );
-								ConstantValue = pxr::GfVec3f( ConvertedColor[ 0 ], ConvertedColor[ 1 ], ConvertedColor[ 2 ] );
+								ConstantValue = ConstantLinearValue;
 							}
 						}
 						FallbackValue = pxr::GfVec4f{ 0.0, 0.0, 0.0, 1.0f };
+						ColorSpaceToken = UnrealIdentifiers::SRGBColorSpaceToken;
 						break;
 					case MP_Specular:
 						InputToken = UnrealIdentifiers::Specular;
 						InputType = pxr::SdfValueTypeNames->Float;
-						ConstantValue = UserConstantValue ? *UserConstantValue : ( *Samples )[ 0 ].R / 255.0f;
+						ConstantValue = UserConstantValue ? *UserConstantValue : ConstantLinearValue[ 0 ];
 						FallbackValue = pxr::GfVec4f{ 0.5, 0.5, 0.5, 1.0f };
 						break;
 					case MP_Metallic:
 						InputToken = UnrealIdentifiers::Metallic;
 						InputType = pxr::SdfValueTypeNames->Float;
-						ConstantValue = UserConstantValue ? *UserConstantValue : ( *Samples )[ 0 ].R / 255.0f;
+						ConstantValue = UserConstantValue ? *UserConstantValue : ConstantLinearValue[ 0 ];
 						FallbackValue = pxr::GfVec4f{ 0.0, 0.0, 0.0, 1.0f };
 						break;
 					case MP_Roughness:
 						InputToken = UnrealIdentifiers::Roughness;
 						InputType = pxr::SdfValueTypeNames->Float;
-						ConstantValue = UserConstantValue ? *UserConstantValue : ( *Samples )[ 0 ].R / 255.0f;
+						ConstantValue = UserConstantValue ? *UserConstantValue : ConstantLinearValue[ 0 ];
 						FallbackValue = pxr::GfVec4f{ 0.5, 0.5, 0.5, 1.0f };
 						break;
 					case MP_Normal:
@@ -1594,15 +1745,14 @@ namespace UE
 						InputType = pxr::SdfValueTypeNames->Normal3f;
 						if ( bPropertyValueIsConstant )
 						{
-							// This doesn't make much sense here but for it's an available option, so here we go
+							// This doesn't make much sense here but it's an available option, so here we go
 							if ( UserConstantValue )
 							{
 								ConstantValue = pxr::GfVec3f( *UserConstantValue, *UserConstantValue, *UserConstantValue );
 							}
 							else
 							{
-								FVector ConvertedNormal{ ( *Samples )[ 0 ].ReinterpretAsLinear() };
-								ConstantValue = UnrealToUsd::ConvertVector( StageInfo, ConvertedNormal );
+								ConstantValue = ConstantLinearValue;
 							}
 						}
 						FallbackValue = pxr::GfVec4f{ 0.0, 1.0, 0.0, 1.0f };
@@ -1612,15 +1762,14 @@ namespace UE
 						InputType = pxr::SdfValueTypeNames->Normal3f;
 						if ( bPropertyValueIsConstant )
 						{
-							// This doesn't make much sense here but for it's an available option, so here we go
+							// This doesn't make much sense here but it's an available option, so here we go
 							if ( UserConstantValue )
 							{
 								ConstantValue = pxr::GfVec3f( *UserConstantValue, *UserConstantValue, *UserConstantValue );
 							}
 							else
 							{
-								FVector ConvertedNormal{ ( *Samples )[ 0 ].ReinterpretAsLinear() };
-								ConstantValue = UnrealToUsd::ConvertVector( StageInfo, ConvertedNormal );
+								ConstantValue = ConstantLinearValue;
 							}
 						}
 						FallbackValue = pxr::GfVec4f{ 0.0, 1.0, 0.0, 1.0f };
@@ -1630,37 +1779,36 @@ namespace UE
 						InputType = pxr::SdfValueTypeNames->Color3f;
 						if ( bPropertyValueIsConstant )
 						{
-							// This doesn't make much sense here but for it's an available option, so here we go
+							// This doesn't make much sense here but it's an available option, so here we go
 							if ( UserConstantValue )
 							{
 								ConstantValue = pxr::GfVec3f( *UserConstantValue, *UserConstantValue, *UserConstantValue );
 							}
 							else
 							{
-								pxr::GfVec4f ConvertedColor = UnrealToUsd::ConvertColor( ( *Samples )[ 0 ] );
-								ConstantValue = pxr::GfVec3f( ConvertedColor[ 0 ], ConvertedColor[ 1 ], ConvertedColor[ 2 ] );
+								ConstantValue = ConstantLinearValue;
 							}
 						}
 						FallbackValue = pxr::GfVec4f{ 0.0, 0.0, 0.0, 1.0f };
-						bHasEmissive = true;
+						// Emissive is also written out with RawColorSpaceToken as we write them as EXR files now
 						break;
 					case MP_Opacity: // It's OK that we use the same for both as these are mutually exclusive blend modes
 					case MP_OpacityMask:
 						InputToken = UnrealIdentifiers::Opacity;
 						InputType = pxr::SdfValueTypeNames->Float;
-						ConstantValue = UserConstantValue ? *UserConstantValue : ( *Samples )[ 0 ].R / 255.0f;
+						ConstantValue = UserConstantValue ? *UserConstantValue : ConstantLinearValue[ 0 ];
 						FallbackValue = pxr::GfVec4f{ 1.0, 1.0, 1.0, 1.0f };
 						break;
 					case MP_Anisotropy:
 						InputToken = UnrealIdentifiers::Anisotropy;
 						InputType = pxr::SdfValueTypeNames->Float;
-						ConstantValue = UserConstantValue ? *UserConstantValue : ( *Samples )[ 0 ].R / 255.0f;
+						ConstantValue = UserConstantValue ? *UserConstantValue : ConstantLinearValue[ 0 ];
 						FallbackValue = pxr::GfVec4f{ 0.0, 0.0, 0.0, 1.0f };
 						break;
 					case MP_AmbientOcclusion:
 						InputToken = UnrealIdentifiers::Occlusion;
 						InputType = pxr::SdfValueTypeNames->Float;
-						ConstantValue = UserConstantValue ? *UserConstantValue : ( *Samples )[ 0 ].R / 255.0f;
+						ConstantValue = UserConstantValue ? *UserConstantValue : ConstantLinearValue[ 0 ];
 						FallbackValue = pxr::GfVec4f{ 1.0, 1.0, 1.0, 1.0f };
 						break;
 					case MP_SubsurfaceColor:
@@ -1674,11 +1822,11 @@ namespace UE
 							}
 							else
 							{
-								pxr::GfVec4f ConvertedColor = UnrealToUsd::ConvertColor( ( *Samples )[ 0 ] );
-								ConstantValue = pxr::GfVec3f( ConvertedColor[ 0 ], ConvertedColor[ 1 ], ConvertedColor[ 2 ] );
+								ConstantValue = ConstantLinearValue;
 							}
 						}
 						FallbackValue = pxr::GfVec4f{ 1.0, 1.0, 1.0, 1.0f };
+						ColorSpaceToken = UnrealIdentifiers::SRGBColorSpaceToken;
 						break;
 					default:
 						continue;
@@ -1726,6 +1874,9 @@ namespace UE
 						pxr::UsdShadeInput TextureStInput = UsdUVTextureShader.CreateInput( UnrealIdentifiers::St, pxr::SdfValueTypeNames->Float2 );
 						TextureStInput.ConnectToSource( PrimvarReaderShader.GetOutput( UnrealIdentifiers::Result ) );
 
+						pxr::UsdShadeInput TextureColorSpaceInput = UsdUVTextureShader.CreateInput( UnrealIdentifiers::SourceColorSpaceToken, pxr::SdfValueTypeNames->Token );
+						TextureColorSpaceInput.Set( ColorSpaceToken );
+
 						pxr::UsdShadeInput TextureFallbackInput = UsdUVTextureShader.CreateInput( UnrealIdentifiers::Fallback, pxr::SdfValueTypeNames->Float4 );
 						TextureFallbackInput.Set( FallbackValue );
 
@@ -1754,17 +1905,6 @@ namespace UE
 					}
 				}
 
-				// The emissive texture is just remapped to [0, 255] on bake, and is intended to be scaled by the emissiveScale before usage
-				if ( bHasEmissive )
-				{
-					pxr::UsdShadeInput EmissiveScaleInput = OutUsdShadeMaterial.CreateInput(
-						UnrealToUsd::ConvertToken( TEXT( "emissiveScale" ) ).Get(),
-						pxr::SdfValueTypeNames->Float
-					);
-
-					EmissiveScaleInput.Set( BakedData.EmissiveScale );
-				}
-
 				return true;
 			}
 #endif // WITH_EDITOR
diff --git a/Engine/Plugins/Importers/USDImporter/Source/UnrealUSDWrapper/Private/UnrealUSDWrapper.cpp b/Engine/Plugins/Importers/USDImporter/Source/UnrealUSDWrapper/Private/UnrealUSDWrapper.cpp
--- a/Engine/Plugins/Importers/USDImporter/Source/UnrealUSDWrapper/Private/UnrealUSDWrapper.cpp
+++ b/Engine/Plugins/Importers/USDImporter/Source/UnrealUSDWrapper/Private/UnrealUSDWrapper.cpp
@@ -129,6 +129,9 @@ namespace UnrealIdentifiers
 	const TfToken Fallback = TfToken("fallback");
 	const TfToken R = TfToken("r");
 	const TfToken RGB = TfToken("rgb");
+	const TfToken RawColorSpaceToken = TfToken{ "raw" };
+	const TfToken SRGBColorSpaceToken = TfToken{ "sRGB" };
+	const TfToken SourceColorSpaceToken = TfToken{ "sourceColorSpace" };
 
 	const TfToken UsdPreviewSurface = TfToken( "UsdPreviewSurface" );
 	const TfToken UsdPrimvarReader_float2 = TfToken( "UsdPrimvarReader_float2" );
diff --git a/Engine/Plugins/Importers/USDImporter/Source/UnrealUSDWrapper/Public/UnrealUSDWrapper.h b/Engine/Plugins/Importers/USDImporter/Source/UnrealUSDWrapper/Public/UnrealUSDWrapper.h
--- a/Engine/Plugins/Importers/USDImporter/Source/UnrealUSDWrapper/Public/UnrealUSDWrapper.h
+++ b/Engine/Plugins/Importers/USDImporter/Source/UnrealUSDWrapper/Public/UnrealUSDWrapper.h
@@ -438,6 +438,9 @@ namespace UnrealIdentifiers
 	extern UNREALUSDWRAPPER_API const pxr::TfToken Fallback;
 	extern UNREALUSDWRAPPER_API const pxr::TfToken R;
 	extern UNREALUSDWRAPPER_API const pxr::TfToken RGB;
+	extern UNREALUSDWRAPPER_API const pxr::TfToken RawColorSpaceToken;
+	extern UNREALUSDWRAPPER_API const pxr::TfToken SRGBColorSpaceToken;
+	extern UNREALUSDWRAPPER_API const pxr::TfToken SourceColorSpaceToken;
 
 	// Tokens copied from usdImaging, because at the moment it's all we need from it
 	extern UNREALUSDWRAPPER_API const pxr::TfToken UsdPreviewSurface;
