USD: Display opened stage on the stage actor picker menu.

diff --git a/Engine/Plugins/Importers/USDImporter/Source/USDStageEditor/Private/SUSDStage.cpp b/Engine/Plugins/Importers/USDImporter/Source/USDStageEditor/Private/SUSDStage.cpp
--- a/Engine/Plugins/Importers/USDImporter/Source/USDStageEditor/Private/SUSDStage.cpp
+++ b/Engine/Plugins/Importers/USDImporter/Source/USDStageEditor/Private/SUSDStage.cpp
@@ -33,6 +33,8 @@
 #include "Engine/World.h"
 #include "EngineAnalytics.h"
 #include "Framework/MultiBox/MultiBoxBuilder.h"
+#include "ISceneOutliner.h"
+#include "ISceneOutlinerColumn.h"
 #include "LevelEditor.h"
 #include "Modules/ModuleManager.h"
 #include "SceneOutlinerModule.h"
@@ -180,6 +182,90 @@ namespace SUSDStageImpl
 			GEditor->SelectComponent( Component, bSelected, bNotifySelectionChanged );
 		}
 	}
+
+	class FStageNameOutlinerColumn : public ISceneOutlinerColumn
+	{
+	public:
+		FStageNameOutlinerColumn( ISceneOutliner& SceneOutliner )
+			: WeakSceneOutliner( StaticCastSharedRef<ISceneOutliner>( SceneOutliner.AsShared() ) )
+		{}
+
+		virtual ~FStageNameOutlinerColumn() {}
+		static FName GetID()
+		{
+			static FName ID( "Stage" );
+			return ID;
+		}
+		static FString GetStageForItem( const ISceneOutlinerTreeItem* Item )
+		{
+			if ( Item )
+			{
+				if ( const FActorTreeItem* ActorTreeItem = Item->CastTo<FActorTreeItem>() )
+				{
+					if ( const AUsdStageActor* StageActor = Cast<const AUsdStageActor>( ActorTreeItem->Actor.Get() ) )
+					{
+						if ( const UE::FUsdStage& Stage = StageActor->GetUsdStage() )
+						{
+							return Stage.GetRootLayer().GetDisplayName();
+						}
+					}
+				}
+			}
+
+			return {};
+		}
+
+		// Begin ISceneOutlinerColumn Implementation
+		virtual FName GetColumnID() override
+		{
+			return GetID();
+		}
+		virtual SHeaderRow::FColumn::FArguments ConstructHeaderRowColumn() override
+		{
+			return SHeaderRow::Column( FName( TEXT( "Stage" ) ) );
+		}
+		virtual const TSharedRef<SWidget> ConstructRowWidget(
+			FSceneOutlinerTreeItemRef TreeItem,
+			const STableRow<FSceneOutlinerTreeItemPtr>& Row
+		) override
+		{
+			return SNew( STextBlock ).Text( FText::FromString( GetStageForItem( &TreeItem.Get() ) ) );
+		}
+		virtual bool SupportsSorting() const override
+		{
+			return true;
+		}
+		virtual void SortItems( TArray<FSceneOutlinerTreeItemPtr>& OutItems, const EColumnSortMode::Type SortMode ) const
+		{
+			if ( SortMode == EColumnSortMode::Ascending )
+			{
+				OutItems.Sort(
+					[]( const FSceneOutlinerTreeItemPtr& First, const FSceneOutlinerTreeItemPtr& Second )
+					{
+						const ISceneOutlinerTreeItem* FirstPtr = First.Get();
+						const ISceneOutlinerTreeItem* SecondPtr = Second.Get();
+						return GetStageForItem( FirstPtr ) < GetStageForItem( SecondPtr );
+					}
+				);
+			}
+
+			if ( SortMode == EColumnSortMode::Descending )
+			{
+				OutItems.Sort(
+					[]( const FSceneOutlinerTreeItemPtr& First, const FSceneOutlinerTreeItemPtr& Second )
+					{
+						const ISceneOutlinerTreeItem* FirstPtr = First.Get();
+						const ISceneOutlinerTreeItem* SecondPtr = Second.Get();
+						return GetStageForItem( FirstPtr ) > GetStageForItem( SecondPtr );
+					}
+				);
+			}
+		}
+		// End ISceneOutlinerColumn Implementation
+
+	private:
+		TWeakPtr<ISceneOutliner> WeakSceneOutliner;
+	};
 }
 
 void SUsdStage::Construct( const FArguments& InArgs )
@@ -538,12 +624,41 @@ TSharedRef< SWidget > SUsdStage::MakeActorPickerMenuContent()
 {
 	if ( !ActorPickerMenu.IsValid() )
 	{
+		const int32 NamePriority = 0;
+		const int32 StagePriority = 1;
+		const bool bCanBeHidden = true;
+		const float FillSize = 1.0f;
+
+		FSceneOutlinerColumnInfo ActorColumnInfo{ ESceneOutlinerColumnVisibility::Visible, NamePriority };
+		ActorColumnInfo.ColumnLabel = LOCTEXT( "ActorColumnText", "Actor" );
+		ActorColumnInfo.FillSize = FillSize;
+
+		FSceneOutlinerColumnInfo StageColumnInfo{
+			ESceneOutlinerColumnVisibility::Visible,
+			StagePriority,
+			FCreateSceneOutlinerColumn::CreateLambda(
+				[]( ISceneOutliner& InSceneOutliner )
+				{
+					return MakeShareable( new SUSDStageImpl::FStageNameOutlinerColumn( InSceneOutliner ) );
+				}
+			),
+			bCanBeHidden,
+			FillSize
+		};
+
 		FSceneOutlinerInitializationOptions InitOptions;
-		InitOptions.bShowHeaderRow = false;
+		InitOptions.bShowHeaderRow = true;  // The header lets us resize/hide the columns if we want
 		InitOptions.bShowSearchBox = true;
 		InitOptions.bShowCreateNewFolder = false;
 		InitOptions.bFocusSearchBoxWhenOpened = true;
-		InitOptions.ColumnMap.Add( FSceneOutlinerBuiltInColumnTypes::Label(), FSceneOutlinerColumnInfo( ESceneOutlinerColumnVisibility::Visible, 0 ) );
+		InitOptions.ColumnMap.Add(
+			FSceneOutlinerBuiltInColumnTypes::Label(),
+			ActorColumnInfo
+		);
+		InitOptions.ColumnMap.Add(
+			SUSDStageImpl::FStageNameOutlinerColumn::GetID(),
+			StageColumnInfo
+		);
 		InitOptions.Filters->AddFilterPredicate<FActorTreeItem>(
 			FActorTreeItem::FFilterPredicate::CreateLambda(
 				[]( const AActor* Actor )
@@ -576,7 +691,7 @@ TSharedRef< SWidget > SUsdStage::MakeActorPickerMenuContent()
 
 		return SNew(SBox)
 			.Padding( FMargin( 1 ) )    // Add a small margin or else we'll get dark gray on dark gray which can look a bit confusing
-			.MinDesiredWidth( 400.0f )  // Force a min width or else the tree view item text will run up right to the very edge pixel of the menu
+			.MinDesiredWidth( 600.0f )  // Force a min width or else the tree view item text will run up right to the very edge pixel of the menu
 			.HAlign( HAlign_Fill )
 			[
 				ActorPickerMenu.ToSharedRef()
