USD: Allow collapsing Xform hierarchies with a single Mesh.

diff --git a/Engine/Plugins/Importers/USDImporter/Source/USDSchemas/Private/USDGeomXformableTranslator.cpp b/Engine/Plugins/Importers/USDImporter/Source/USDSchemas/Private/USDGeomXformableTranslator.cpp
--- a/Engine/Plugins/Importers/USDImporter/Source/USDSchemas/Private/USDGeomXformableTranslator.cpp
+++ b/Engine/Plugins/Importers/USDImporter/Source/USDSchemas/Private/USDGeomXformableTranslator.cpp
@@ -712,13 +712,7 @@ bool FUsdGeomXformableTranslator::CollapsesChildren( ECollapsingType CollapsingT
 			}
 		}
 
-
-		// We only support collapsing GeomMeshes for now and we only want to do it when there are multiple meshes as the resulting mesh is considered unique
-		if ( ChildGeomMeshes.Num() < 2 )
-		{
-			bCollapsesChildren = false;
-		}
-		else if ( Context->InfoCache.IsValid() )
+		if ( Context->InfoCache.IsValid() )
 		{
 			TOptional<uint64> NumExpectedVertices = Context->InfoCache->GetSubtreeVertexCount( PrimPath );
 			if ( !NumExpectedVertices.IsSet() || NumExpectedVertices.GetValue() > GMaxNumVerticesCollapsedMesh )
@@ -760,8 +754,6 @@ bool FUsdGeomXformableTranslator::CanBeCollapsed( ECollapsingType CollapsingType
 {
 	FScopedUsdAllocs UsdAllocs;
 
-	FString PrimPathStr = PrimPath.GetString();
-
 	pxr::UsdPrim UsdPrim{ GetPrim() };
 	if ( !UsdPrim )
 	{
