USD: Allow setting stage editor options without a stage actor bound.

diff --git a/Engine/Plugins/Importers/USDImporter/Source/USDStage/Public/USDStageActor.h b/Engine/Plugins/Importers/USDImporter/Source/USDStage/Public/USDStageActor.h
--- a/Engine/Plugins/Importers/USDImporter/Source/USDStage/Public/USDStageActor.h
+++ b/Engine/Plugins/Importers/USDImporter/Source/USDStage/Public/USDStageActor.h
@@ -31,7 +31,7 @@ enum class EUsdPurpose : int32;
 struct FMeshDescription;
 struct FUsdSchemaTranslationContext;
 
-UCLASS( MinimalAPI )
+UCLASS( MinimalAPI, config = USDImporter )
 class AUsdStageActor : public AActor
 {
 	GENERATED_BODY()
@@ -43,24 +43,28 @@ public:
 	UPROPERTY(EditAnywhere, BlueprintReadOnly, Category = "USD", meta = (RelativeToGameDir))
 	FFilePath RootLayer;
 
-	UPROPERTY(EditAnywhere, BlueprintReadOnly, Category = "USD")
+	// These properties are configs so that spawned actors read them from the CDO when spawned.
+	// This allows the defaults for them to be configured on EditorPerProjectUserSettings.ini, and allows us to write
+	// to that config from the USD Stage Editor, specifying our options before the editor is attached to any stage actor.
+
+	UPROPERTY(EditAnywhere, BlueprintReadOnly, Category = "USD", config )
 	EUsdInitialLoadSet InitialLoadSet;
 
-	UPROPERTY( EditAnywhere, BlueprintReadOnly, Category = "USD" )
+	UPROPERTY( EditAnywhere, BlueprintReadOnly, Category = "USD", config )
 	EUsdInterpolationType InterpolationType;
 
 	/**
 	 * Whether to try to combine individual assets and components of the same type on a kind-per-kind basis,
 	 * like multiple Mesh prims into a single Static Mesh
 	 */
-	UPROPERTY(EditAnywhere, BlueprintReadOnly, Category = "USD", meta = (Bitmask, BitmaskEnum="/Script/UnrealUSDWrapper.EUsdDefaultKind"))
+	UPROPERTY(EditAnywhere, BlueprintReadOnly, Category = "USD", config, meta = (Bitmask, BitmaskEnum="/Script/UnrealUSDWrapper.EUsdDefaultKind"))
 	int32 KindsToCollapse;
 
 	/**
 	 * If enabled, when multiple mesh prims are collapsed into a single static mesh, identical material slots are merged into one slot.
 	 * Otherwise, materials slots are simply appended to the list.
 	 */
-	UPROPERTY(EditAnywhere, BlueprintReadOnly, Category = "USD")
+	UPROPERTY(EditAnywhere, BlueprintReadOnly, Category = "USD", config )
 	bool bMergeIdenticalMaterialSlots;
 
 	/**
@@ -73,23 +77,23 @@ public:
 	bool bCollapseTopLevelPointInstancers;
 
 	/* Only load prims with these specific purposes from the USD file */
-	UPROPERTY(EditAnywhere, BlueprintReadOnly, Category = "USD", meta = (Bitmask, BitmaskEnum="/Script/UnrealUSDWrapper.EUsdPurpose"))
+	UPROPERTY(EditAnywhere, BlueprintReadOnly, Category = "USD", config, meta = (Bitmask, BitmaskEnum="/Script/UnrealUSDWrapper.EUsdPurpose"))
 	int32 PurposesToLoad;
 
 	/** Try enabling Nanite for static meshes that are generated with at least this many triangles */
-	UPROPERTY( EditAnywhere, BlueprintReadOnly, Category = "USD", meta = ( NoSpinbox = "true", UIMin = "0", ClampMin = "0" ))
+	UPROPERTY( EditAnywhere, BlueprintReadOnly, Category = "USD", config, meta = ( NoSpinbox = "true", UIMin = "0", ClampMin = "0" ))
 	int32 NaniteTriangleThreshold;
 
 	/** Specifies which set of shaders to use when parsing USD materials, in addition to the universal render context. */
-	UPROPERTY(EditAnywhere, BlueprintReadOnly, Category = "USD")
+	UPROPERTY(EditAnywhere, BlueprintReadOnly, Category = "USD", config )
 	FName RenderContext;
 
 	/** Specifies which material purbose to use when parsing USD material bindings, in addition to the "allPurpose" fallback. */
-	UPROPERTY( EditAnywhere, BlueprintReadOnly, Category = "USD" )
+	UPROPERTY( EditAnywhere, BlueprintReadOnly, Category = "USD", config )
 	FName MaterialPurpose;
 
 	// Describes what to add to the root bone animation within generated AnimSequences, if anything
-	UPROPERTY( EditAnywhere, BlueprintReadOnly, Category = "USD" )
+	UPROPERTY( EditAnywhere, BlueprintReadOnly, Category = "USD", config )
 	EUsdRootMotionHandling RootMotionHandling = EUsdRootMotionHandling::NoAdditionalRootMotion;
 
 public:
diff --git a/Engine/Plugins/Importers/USDImporter/Source/USDStageEditor/Private/SUSDStage.cpp b/Engine/Plugins/Importers/USDImporter/Source/USDStageEditor/Private/SUSDStage.cpp
--- a/Engine/Plugins/Importers/USDImporter/Source/USDStageEditor/Private/SUSDStage.cpp
+++ b/Engine/Plugins/Importers/USDImporter/Source/USDStageEditor/Private/SUSDStage.cpp
@@ -1032,7 +1032,7 @@ void SUsdStage::FillPayloadsSubMenu(FMenuBuilder& MenuBuilder)
 		FUIAction(
 			FExecuteAction::CreateLambda([this]()
 			{
-				if(AUsdStageActor* StageActor = ViewModel.UsdStageActor.Get())
+				if ( AUsdStageActor* StageActor = GetStageActorOrCDO() )
 				{
 					FScopedTransaction Transaction(FText::Format(
 						LOCTEXT("SetLoadAllTransaction", "Set USD stage actor '{0}' actor to load all payloads initially"),
@@ -1041,15 +1041,16 @@ void SUsdStage::FillPayloadsSubMenu(FMenuBuilder& MenuBuilder)
 
 					StageActor->Modify();
 					StageActor->InitialLoadSet = EUsdInitialLoadSet::LoadAll;
+					if ( StageActor->IsTemplate() )
+					{
+						StageActor->SaveConfig();
+					}
 				}
 			}),
-			FCanExecuteAction::CreateLambda([this]()
-			{
-				return ViewModel.UsdStageActor.Get() != nullptr;
-			}),
+			FCanExecuteAction{},
 			FIsActionChecked::CreateLambda([this]()
 			{
-				if(AUsdStageActor* StageActor = ViewModel.UsdStageActor.Get())
+				if ( AUsdStageActor* StageActor = GetStageActorOrCDO() )
 				{
 					return StageActor->InitialLoadSet == EUsdInitialLoadSet::LoadAll;
 				}
@@ -1067,7 +1068,7 @@ void SUsdStage::FillPayloadsSubMenu(FMenuBuilder& MenuBuilder)
 		FUIAction(
 			FExecuteAction::CreateLambda([this]()
 			{
-				if(AUsdStageActor* StageActor = ViewModel.UsdStageActor.Get())
+				if ( AUsdStageActor* StageActor = GetStageActorOrCDO() )
 				{
 					FScopedTransaction Transaction(FText::Format(
 						LOCTEXT("SetLoadNoneTransaction", "Set USD stage actor '{0}' actor to load no payloads initially"),
@@ -1076,15 +1077,16 @@ void SUsdStage::FillPayloadsSubMenu(FMenuBuilder& MenuBuilder)
 
 					StageActor->Modify();
 					StageActor->InitialLoadSet = EUsdInitialLoadSet::LoadNone;
+					if ( StageActor->IsTemplate() )
+					{
+						StageActor->SaveConfig();
+					}
 				}
 			}),
-			FCanExecuteAction::CreateLambda([this]()
-			{
-				return ViewModel.UsdStageActor.Get() != nullptr;
-			}),
+			FCanExecuteAction{},
 			FIsActionChecked::CreateLambda([this]()
 			{
-				if(AUsdStageActor* StageActor = ViewModel.UsdStageActor.Get())
+				if ( AUsdStageActor* StageActor = GetStageActorOrCDO() )
 				{
 					return StageActor->InitialLoadSet == EUsdInitialLoadSet::LoadNone;
 				}
@@ -1107,7 +1109,7 @@ void SUsdStage::FillPurposesToLoadSubMenu(FMenuBuilder& MenuBuilder)
 			FUIAction(
 				FExecuteAction::CreateLambda([this, Purpose]()
 				{
-					if(AUsdStageActor* StageActor = ViewModel.UsdStageActor.Get())
+					if ( AUsdStageActor* StageActor = GetStageActorOrCDO() )
 					{
 						FScopedTransaction Transaction(FText::Format(
 							LOCTEXT("PurposesToLoadTransaction", "Change purposes to load for USD stage actor '{0}'"),
@@ -1124,15 +1126,17 @@ void SUsdStage::FillPurposesToLoadSubMenu(FMenuBuilder& MenuBuilder)
 							FindFieldChecked< FProperty >( StageActor->GetClass(), GET_MEMBER_NAME_CHECKED( AUsdStageActor, PurposesToLoad ) )
 						);
 						StageActor->PostEditChangeProperty(PropertyChangedEvent);
+
+						if ( StageActor->IsTemplate() )
+						{
+							StageActor->SaveConfig();
+						}
 					}
 				}),
-				FCanExecuteAction::CreateLambda([this]()
-				{
-					return ViewModel.UsdStageActor.Get() != nullptr;
-				}),
+				FCanExecuteAction{},
 				FIsActionChecked::CreateLambda([this, Purpose]()
 				{
-					if(AUsdStageActor* StageActor = ViewModel.UsdStageActor.Get())
+					if ( AUsdStageActor* StageActor = GetStageActorOrCDO() )
 					{
 						return EnumHasAllFlags((EUsdPurpose)StageActor->PurposesToLoad, Purpose);
 					}
@@ -1166,7 +1170,7 @@ void SUsdStage::FillRenderContextSubMenu( FMenuBuilder& MenuBuilder )
 			FUIAction(
 				FExecuteAction::CreateLambda([this, RenderContext]()
 				{
-					if ( AUsdStageActor* StageActor = ViewModel.UsdStageActor.Get() )
+					if ( AUsdStageActor* StageActor = GetStageActorOrCDO() )
 					{
 						FScopedTransaction Transaction(FText::Format(
 							LOCTEXT("RenderContextToLoadTransaction", "Change render context to load for USD stage actor '{0}'"),
@@ -1183,15 +1187,17 @@ void SUsdStage::FillRenderContextSubMenu( FMenuBuilder& MenuBuilder )
 							FindFieldChecked< FProperty >( StageActor->GetClass(), GET_MEMBER_NAME_CHECKED( AUsdStageActor, RenderContext ) )
 						);
 						StageActor->PostEditChangeProperty(PropertyChangedEvent);
+
+						if ( StageActor->IsTemplate() )
+						{
+							StageActor->SaveConfig();
+						}
 					}
 				}),
-				FCanExecuteAction::CreateLambda([this]()
-				{
-					return ViewModel.UsdStageActor.Get() != nullptr;
-				}),
+				FCanExecuteAction{},
 				FIsActionChecked::CreateLambda([this, RenderContext]()
 				{
-					if ( AUsdStageActor* StageActor = ViewModel.UsdStageActor.Get() )
+					if ( AUsdStageActor* StageActor = GetStageActorOrCDO() )
 					{
 						return StageActor->RenderContext == RenderContext;
 					}
@@ -1270,9 +1276,13 @@ void SUsdStage::FillMaterialPurposeSubMenu( FMenuBuilder& MenuBuilder )
 		{
 			if ( ChosenOption )
 			{
-				if ( AUsdStageActor* StageActor = ViewModel.UsdStageActor.Get() )
+				if ( AUsdStageActor* StageActor = GetStageActorOrCDO() )
 				{
 					StageActor->SetMaterialPurpose( **ChosenOption );
+					if ( StageActor->IsTemplate() )
+					{
+						StageActor->SaveConfig();
+					}
 				}
 			}
 		})
@@ -1280,7 +1290,7 @@ void SUsdStage::FillMaterialPurposeSubMenu( FMenuBuilder& MenuBuilder )
 			SNew( SEditableTextBox )
 			.Text_Lambda([this]() -> FText
 			{
-				if ( AUsdStageActor* StageActor = ViewModel.UsdStageActor.Get() )
+				if ( AUsdStageActor* StageActor = GetStageActorOrCDO() )
 				{
 					return FText::FromString( StageActor->MaterialPurpose == *UnrealIdentifiers::MaterialAllPurpose
 						? UnrealIdentifiers::MaterialAllPurposeText
@@ -1320,9 +1330,13 @@ void SUsdStage::FillMaterialPurposeSubMenu( FMenuBuilder& MenuBuilder )
 					}
 				}
 
-				if ( AUsdStageActor* StageActor = ViewModel.UsdStageActor.Get() )
+				if ( AUsdStageActor* StageActor = GetStageActorOrCDO() )
 				{
 					StageActor->SetMaterialPurpose( NewPurpose );
+					if ( StageActor->IsTemplate() )
+					{
+						StageActor->SaveConfig();
+					}
 				}
 			})
 		]
@@ -1346,7 +1360,7 @@ void SUsdStage::FillRootMotionSubMenu( FMenuBuilder& MenuBuilder )
 			FUIAction(
 				FExecuteAction::CreateLambda([this, HandlingStrategy]()
 				{
-					if(AUsdStageActor* StageActor = ViewModel.UsdStageActor.Get())
+					if ( AUsdStageActor* StageActor = GetStageActorOrCDO() )
 					{
 						FScopedTransaction Transaction(FText::Format(
 							LOCTEXT("RootMotionHandlingTransaction", "Change root motion handling strategy for USD stage actor '{0}'"),
@@ -1358,15 +1372,16 @@ void SUsdStage::FillRootMotionSubMenu( FMenuBuilder& MenuBuilder )
 
 						StageActor->Modify();
 						StageActor->SetRootMotionHandling( HandlingStrategy );
+						if ( StageActor->IsTemplate() )
+						{
+							StageActor->SaveConfig();
+						}
 					}
 				}),
-				FCanExecuteAction::CreateLambda([this]()
-				{
-					return ViewModel.UsdStageActor.Get() != nullptr;
-				}),
+				FCanExecuteAction{},
 				FIsActionChecked::CreateLambda([this, HandlingStrategy]()
 				{
-					if(AUsdStageActor* StageActor = ViewModel.UsdStageActor.Get())
+					if ( AUsdStageActor* StageActor = GetStageActorOrCDO() )
 					{
 						return StageActor->RootMotionHandling == HandlingStrategy;
 					}
@@ -1392,7 +1407,7 @@ void SUsdStage::FillCollapsingSubMenu( FMenuBuilder& MenuBuilder )
 		FUIAction(
 			FExecuteAction::CreateLambda( [this]()
 			{
-				if ( AUsdStageActor* StageActor = ViewModel.UsdStageActor.Get() )
+				if ( AUsdStageActor* StageActor = GetStageActorOrCDO() )
 				{
 					FScopedTransaction Transaction( FText::Format(
 						LOCTEXT( "MergeIdenticalMaterialSlotsTransaction", "Toggle bMergeIdenticalMaterialSlots on USD stage actor '{1}'" ),
@@ -1401,15 +1416,16 @@ void SUsdStage::FillCollapsingSubMenu( FMenuBuilder& MenuBuilder )
 
 					TGuardValue<bool> MaintainSelectionGuard( bUpdatingViewportSelection, true );
 					StageActor->SetMergeIdenticalMaterialSlots( !StageActor->bMergeIdenticalMaterialSlots );
+					if ( StageActor->IsTemplate() )
+					{
+						StageActor->SaveConfig();
+					}
 				}
 			}),
-			FCanExecuteAction::CreateLambda( [this]()
-			{
-				return ViewModel.UsdStageActor.Get() != nullptr;
-			}),
+			FCanExecuteAction{},
 			FIsActionChecked::CreateLambda( [this]()
 			{
-				if ( AUsdStageActor* StageActor = ViewModel.UsdStageActor.Get() )
+				if ( AUsdStageActor* StageActor = GetStageActorOrCDO() )
 				{
 					return StageActor->bMergeIdenticalMaterialSlots;
 				}
@@ -1429,7 +1445,7 @@ void SUsdStage::FillCollapsingSubMenu( FMenuBuilder& MenuBuilder )
 			FUIAction(
 				FExecuteAction::CreateLambda( [this, Text, Kind]()
 				{
-					if ( AUsdStageActor* StageActor = ViewModel.UsdStageActor.Get() )
+					if ( AUsdStageActor* StageActor = GetStageActorOrCDO() )
 					{
 						FScopedTransaction Transaction( FText::Format(
 							LOCTEXT( "ToggleCollapsingTransaction", "Toggle asset and component collapsing for kind '{0}' on USD stage actor '{1}'" ),
@@ -1447,12 +1463,16 @@ void SUsdStage::FillCollapsingSubMenu( FMenuBuilder& MenuBuilder )
 
 						int32 NewKindsToCollapse = ( int32 ) ( ( EUsdDefaultKind ) StageActor->KindsToCollapse ^ Kind );
 						StageActor->SetKindsToCollapse( NewKindsToCollapse );
+						if ( StageActor->IsTemplate() )
+						{
+							StageActor->SaveConfig();
+						}
 					}
 				}),
 				CanExecuteAction,
 				FIsActionChecked::CreateLambda( [this, Kind]()
 				{
-					if ( AUsdStageActor* StageActor = ViewModel.UsdStageActor.Get() )
+					if ( AUsdStageActor* StageActor = GetStageActorOrCDO() )
 					{
 						return EnumHasAllFlags( ( EUsdDefaultKind ) StageActor->KindsToCollapse, Kind );
 					}
@@ -1469,10 +1489,7 @@ void SUsdStage::FillCollapsingSubMenu( FMenuBuilder& MenuBuilder )
 	AddKindToCollapseEntry(
 		EUsdDefaultKind::Model,
 		LOCTEXT( "ModelKind", "Model" ),
-		FCanExecuteAction::CreateLambda( [this]()
-		{
-			return ViewModel.UsdStageActor.Get() != nullptr;
-		})
+		FCanExecuteAction{}
 	);
 
 	AddKindToCollapseEntry(
@@ -1480,7 +1497,7 @@ void SUsdStage::FillCollapsingSubMenu( FMenuBuilder& MenuBuilder )
 		LOCTEXT( "ModelComponent", "   Component" ),
 		FCanExecuteAction::CreateLambda( [this]()
 		{
-			if ( AUsdStageActor* StageActor = ViewModel.UsdStageActor.Get() )
+			if ( AUsdStageActor* StageActor = GetStageActorOrCDO() )
 			{
 				// If we're collapsing all "model" kinds, the all "component"s should be collapsed anyway
 				if ( !EnumHasAnyFlags( ( EUsdDefaultKind ) StageActor->KindsToCollapse, EUsdDefaultKind::Model ) )
@@ -1498,7 +1515,7 @@ void SUsdStage::FillCollapsingSubMenu( FMenuBuilder& MenuBuilder )
 		LOCTEXT( "ModelGroup", "   Group" ),
 		FCanExecuteAction::CreateLambda( [this]()
 		{
-			if ( AUsdStageActor* StageActor = ViewModel.UsdStageActor.Get() )
+			if ( AUsdStageActor* StageActor = GetStageActorOrCDO() )
 			{
 				if ( !EnumHasAnyFlags( ( EUsdDefaultKind ) StageActor->KindsToCollapse, EUsdDefaultKind::Model ) )
 				{
@@ -1515,7 +1532,7 @@ void SUsdStage::FillCollapsingSubMenu( FMenuBuilder& MenuBuilder )
 		LOCTEXT( "ModelAssembly", "      Assembly" ),
 		FCanExecuteAction::CreateLambda( [this]()
 		{
-			if ( AUsdStageActor* StageActor = ViewModel.UsdStageActor.Get() )
+			if ( AUsdStageActor* StageActor = GetStageActorOrCDO() )
 			{
 				if ( !EnumHasAnyFlags( ( EUsdDefaultKind ) StageActor->KindsToCollapse, EUsdDefaultKind::Model ) && !EnumHasAnyFlags( ( EUsdDefaultKind ) StageActor->KindsToCollapse, EUsdDefaultKind::Group ) )
 				{
@@ -1530,10 +1547,7 @@ void SUsdStage::FillCollapsingSubMenu( FMenuBuilder& MenuBuilder )
 	AddKindToCollapseEntry(
 		EUsdDefaultKind::Subcomponent,
 		LOCTEXT( "ModelSubcomponent", "Subcomponent" ),
-		FCanExecuteAction::CreateLambda( [this]()
-		{
-			return ViewModel.UsdStageActor.Get() != nullptr;
-		})
+		FCanExecuteAction{}
 	);
 
 	MenuBuilder.EndSection();
@@ -1548,7 +1562,7 @@ void SUsdStage::FillInterpolationTypeSubMenu(FMenuBuilder& MenuBuilder)
 		FUIAction(
 			FExecuteAction::CreateLambda([this]()
 			{
-				if(AUsdStageActor* StageActor = ViewModel.UsdStageActor.Get())
+				if ( AUsdStageActor* StageActor = GetStageActorOrCDO() )
 				{
 					FScopedTransaction Transaction(FText::Format(
 						LOCTEXT("SetLinearInterpolationType", "Set USD stage actor '{0}' to linear interpolation"),
@@ -1559,15 +1573,17 @@ void SUsdStage::FillInterpolationTypeSubMenu(FMenuBuilder& MenuBuilder)
 					TGuardValue<bool> MaintainSelectionGuard( bUpdatingViewportSelection, true );
 
 					StageActor->SetInterpolationType( EUsdInterpolationType::Linear );
+
+					if ( StageActor->IsTemplate() )
+					{
+						StageActor->SaveConfig();
+					}
 				}
 			}),
-			FCanExecuteAction::CreateLambda([this]()
-			{
-				return ViewModel.UsdStageActor.Get() != nullptr;
-			}),
+			FCanExecuteAction{},
 			FIsActionChecked::CreateLambda([this]()
 			{
-				if(AUsdStageActor* StageActor = ViewModel.UsdStageActor.Get())
+				if ( AUsdStageActor* StageActor = GetStageActorOrCDO() )
 				{
 					return StageActor->InterpolationType == EUsdInterpolationType::Linear;
 				}
@@ -1585,7 +1601,7 @@ void SUsdStage::FillInterpolationTypeSubMenu(FMenuBuilder& MenuBuilder)
 		FUIAction(
 			FExecuteAction::CreateLambda([this]()
 			{
-				if(AUsdStageActor* StageActor = ViewModel.UsdStageActor.Get())
+				if ( AUsdStageActor* StageActor = GetStageActorOrCDO() )
 				{
 					FScopedTransaction Transaction(FText::Format(
 						LOCTEXT("SetHeldInterpolationType", "Set USD stage actor '{0}' to held interpolation"),
@@ -1596,15 +1612,17 @@ void SUsdStage::FillInterpolationTypeSubMenu(FMenuBuilder& MenuBuilder)
 					TGuardValue<bool> MaintainSelectionGuard( bUpdatingViewportSelection, true );
 
 					StageActor->SetInterpolationType( EUsdInterpolationType::Held );
+
+					if ( StageActor->IsTemplate() )
+					{
+						StageActor->SaveConfig();
+					}
 				}
 			}),
-			FCanExecuteAction::CreateLambda([this]()
-			{
-				return ViewModel.UsdStageActor.Get() != nullptr;
-			}),
+			FCanExecuteAction{},
 			FIsActionChecked::CreateLambda([this]()
 			{
-				if(AUsdStageActor* StageActor = ViewModel.UsdStageActor.Get())
+				if ( AUsdStageActor* StageActor = GetStageActorOrCDO() )
 				{
 					return StageActor->InterpolationType == EUsdInterpolationType::Held;
 				}
@@ -1670,7 +1688,7 @@ void SUsdStage::FillSelectionSubMenu( FMenuBuilder& MenuBuilder )
 
 void SUsdStage::FillNaniteThresholdSubMenu( FMenuBuilder& MenuBuilder )
 {
-	if ( AUsdStageActor* StageActor = ViewModel.UsdStageActor.Get() )
+	if ( AUsdStageActor* StageActor = GetStageActorOrCDO() )
 	{
 		CurrentNaniteThreshold = StageActor->NaniteTriangleThreshold;
 	}
@@ -1681,10 +1699,6 @@ void SUsdStage::FillNaniteThresholdSubMenu( FMenuBuilder& MenuBuilder )
 		.Value( this, &SUsdStage::GetNaniteTriangleThresholdValue )
 		.OnValueChanged(this, &SUsdStage::OnNaniteTriangleThresholdValueChanged )
 		.SupportDynamicSliderMaxValue( true )
-		.IsEnabled_Lambda( [this]() -> bool
-		{
-			return ViewModel.UsdStageActor.Get() != nullptr;
-		})
 		.OnValueCommitted( this, &SUsdStage::OnNaniteTriangleThresholdValueCommitted );
 
 	const bool bNoIndent = true;
@@ -2181,7 +2195,7 @@ void SUsdStage::OnNaniteTriangleThresholdValueChanged( int32 InValue )
 
 void SUsdStage::OnNaniteTriangleThresholdValueCommitted( int32 InValue, ETextCommit::Type InCommitType )
 {
-	AUsdStageActor* StageActor = ViewModel.UsdStageActor.Get();
+	AUsdStageActor* StageActor = GetStageActorOrCDO();
 	if ( !StageActor )
 	{
 		return;
@@ -2197,6 +2211,11 @@ void SUsdStage::OnNaniteTriangleThresholdValueCommitted( int32 InValue, ETextCom
 
 	StageActor->SetNaniteTriangleThreshold( InValue );
 	CurrentNaniteThreshold = InValue;
+
+	if ( StageActor->IsTemplate() )
+	{
+		StageActor->SaveConfig();
+	}
 }
 
 UE::FUsdStageWeak SUsdStage::GetCurrentStage() const
@@ -2210,6 +2229,17 @@ UE::FUsdStageWeak SUsdStage::GetCurrentStage() const
 	return UE::FUsdStageWeak{};
 }
 
+AUsdStageActor* SUsdStage::GetStageActorOrCDO() const
+{
+	AUsdStageActor* StageActor = ViewModel.UsdStageActor.Get();
+	if ( !StageActor )
+	{
+		StageActor = GetMutableDefault<AUsdStageActor>();
+	}
+	ensure( StageActor );
+	return StageActor;
+}
+
 #endif // #if USE_USD_SDK
 
 #undef LOCTEXT_NAMESPACE
diff --git a/Engine/Plugins/Importers/USDImporter/Source/USDStageEditor/Private/SUSDStage.h b/Engine/Plugins/Importers/USDImporter/Source/USDStageEditor/Private/SUSDStage.h
--- a/Engine/Plugins/Importers/USDImporter/Source/USDStageEditor/Private/SUSDStage.h
+++ b/Engine/Plugins/Importers/USDImporter/Source/USDStageEditor/Private/SUSDStage.h
@@ -86,6 +86,8 @@ protected:
 
 	UE::FUsdStageWeak GetCurrentStage() const;
 
+	AUsdStageActor* GetStageActorOrCDO() const;
+
 protected:
 	TSharedPtr< class SUsdStageTreeView > UsdStageTreeView;
 	TSharedPtr< class SUsdPrimInfo > UsdPrimInfoWidget;
