USD: Cleanup USDInfoCache.

Use pimpl to hide implementation and allow tweaks on patch versions;
Use a single map to a struct instead of 4 maps containing the same keys.

diff --git a/Engine/Plugins/Importers/USDImporter/Source/USDSchemas/Private/USDInfoCache.cpp b/Engine/Plugins/Importers/USDImporter/Source/USDSchemas/Private/USDInfoCache.cpp
--- a/Engine/Plugins/Importers/USDImporter/Source/USDSchemas/Private/USDInfoCache.cpp
+++ b/Engine/Plugins/Importers/USDImporter/Source/USDSchemas/Private/USDInfoCache.cpp
@@ -23,87 +23,139 @@
 #include "USDIncludesEnd.h"
 #endif // USE_USD_SDK
 
+namespace UE::UsdInfoCache::Private
+{
+	struct FUsdPrimInfo
+	{
+		UE::FSdfPath AssetCollapsedRoot;
+		UE::FSdfPath ComponentCollapsedRoot;
+		uint64 ExpectedVertexCountForSubtree;
+		uint64 ExpectedMaterialSlotCountForSubtree;
+	};
+}
+
+FArchive& operator <<( FArchive& Ar, UE::UsdInfoCache::Private::FUsdPrimInfo& Info )
+{
+	Ar << Info.AssetCollapsedRoot;
+	Ar << Info.ComponentCollapsedRoot;
+	Ar << Info.ExpectedVertexCountForSubtree;
+	Ar << Info.ExpectedMaterialSlotCountForSubtree;
+
+	return Ar;
+}
+
+struct FUsdInfoCache::FUsdInfoCacheImpl
+{
+	mutable FRWLock Lock;
+	TMap< UE::FSdfPath, UE::UsdInfoCache::Private::FUsdPrimInfo > InfoMap;
+};
+
+FUsdInfoCache::FUsdInfoCache()
+{
+	Impl = MakeUnique< FUsdInfoCache::FUsdInfoCacheImpl>();
+}
+
+FUsdInfoCache::~FUsdInfoCache()
+{
+}
+
 bool FUsdInfoCache::Serialize( FArchive& Ar )
 {
-	FWriteScopeLock ScopeLock( Lock );
-	Ar << AssetPathsToCollapsedRoot;
-	Ar << ComponentPathsToCollapsedRoot;
-	Ar << ExpectedVertexCountPerSubtree;
-	Ar << ExpectedMaterialSlotCountPerSubtree;
+	if ( FUsdInfoCacheImpl* ImplPtr = Impl.Get() )
+	{
+		FWriteScopeLock ScopeLock( ImplPtr->Lock );
+		Ar << ImplPtr->InfoMap;
+	}
 
 	return true;
 }
 
 bool FUsdInfoCache::ContainsInfoAboutPrim( const UE::FSdfPath& Path ) const
 {
-	return AssetPathsToCollapsedRoot.Contains( Path )
-		|| ComponentPathsToCollapsedRoot.Contains( Path )
-		|| ExpectedVertexCountPerSubtree.Contains( Path )
-		|| ExpectedMaterialSlotCountPerSubtree.Contains( Path );
+	if ( FUsdInfoCacheImpl* ImplPtr = Impl.Get() )
+	{
+		FReadScopeLock ScopeLock( ImplPtr->Lock );
+		return ImplPtr->InfoMap.Contains( Path );
+	}
+
+	return false;
 }
 
 bool FUsdInfoCache::IsPathCollapsed( const UE::FSdfPath& Path, ECollapsingType CollapsingType ) const
 {
-	FReadScopeLock ScopeLock( Lock );
-	const TMap< UE::FSdfPath, UE::FSdfPath >* MapToUse =
-		CollapsingType == ECollapsingType::Assets
-			? &AssetPathsToCollapsedRoot
-			: &ComponentPathsToCollapsedRoot;
-
-	if ( const UE::FSdfPath* FoundResult = MapToUse->Find( Path ) )
+	if ( FUsdInfoCacheImpl* ImplPtr = Impl.Get() )
 	{
-		// A non-empty path to another prim means this prim is collapsed into that one
-		return !FoundResult->IsEmpty() && ( *FoundResult ) != Path;
+		FReadScopeLock ScopeLock( ImplPtr->Lock );
+
+		if ( const UE::UsdInfoCache::Private::FUsdPrimInfo* FoundInfo = ImplPtr->InfoMap.Find( Path ) )
+		{
+			const UE::FSdfPath* CollapsedRoot = CollapsingType == ECollapsingType::Assets
+				? &FoundInfo->AssetCollapsedRoot
+				: &FoundInfo->ComponentCollapsedRoot;
+
+			// A non-empty path to another prim means this prim is collapsed into that one
+			return !CollapsedRoot->IsEmpty() && *CollapsedRoot != Path;
+		}
+
+		// This should never happen: We should have cached the entire tree
+		ensureMsgf( false, TEXT( "Prim path '%s' has not been cached!" ), *Path.GetString() );
 	}
 
-	// This should never happen: We should have cached the entire tree
-	ensureMsgf( false, TEXT( "Prim path '%s' has not been cached!" ), *Path.GetString() );
 	return false;
 }
 
 bool FUsdInfoCache::DoesPathCollapseChildren( const UE::FSdfPath& Path, ECollapsingType CollapsingType ) const
 {
-	FReadScopeLock ScopeLock( Lock );
-	const TMap< UE::FSdfPath, UE::FSdfPath >* MapToUse =
-		CollapsingType == ECollapsingType::Assets
-			? &AssetPathsToCollapsedRoot
-			: &ComponentPathsToCollapsedRoot;
-
-	if ( const UE::FSdfPath* FoundResult = MapToUse->Find( Path ) )
+	if ( FUsdInfoCacheImpl* ImplPtr = Impl.Get() )
 	{
-		// We store our own Path in there when we collapse children. Otherwise we hold the path of our collapse root, or empty (in case nothing is collapsed up to here)
-		return (*FoundResult) == Path;
+		FReadScopeLock ScopeLock( ImplPtr->Lock );
+
+		if ( const UE::UsdInfoCache::Private::FUsdPrimInfo* FoundInfo = ImplPtr->InfoMap.Find( Path ) )
+		{
+			const UE::FSdfPath* CollapsedRoot = CollapsingType == ECollapsingType::Assets
+				? &FoundInfo->AssetCollapsedRoot
+				: &FoundInfo->ComponentCollapsedRoot;
+
+			// We store our own Path in there when we collapse children.
+			// Otherwise we hold the path of our collapse root, or empty (in case nothing is collapsed up to here)
+			return *CollapsedRoot == Path;
+		}
+
+		// This should never happen: We should have cached the entire tree
+		ensureMsgf( false, TEXT( "Prim path '%s' has not been cached!" ), *Path.GetString() );
 	}
 
-	// This should never happen: We should have cached the entire tree
-	ensureMsgf( false, TEXT( "Prim path '%s' has not been cached!" ), *Path.GetString() );
 	return false;
 }
 
 UE::FSdfPath FUsdInfoCache::UnwindToNonCollapsedPath( const UE::FSdfPath& Path, ECollapsingType CollapsingType ) const
 {
-	FReadScopeLock ScopeLock( Lock );
-	const TMap< UE::FSdfPath, UE::FSdfPath >* MapToUse =
-		CollapsingType == ECollapsingType::Assets
-			? &AssetPathsToCollapsedRoot
-			: &ComponentPathsToCollapsedRoot;
-
-	if ( const UE::FSdfPath* FoundResult = MapToUse->Find( Path ) )
+	if ( FUsdInfoCacheImpl* ImplPtr = Impl.Get() )
 	{
-		// An empty path here means that we are not collapsed at all
-		if ( FoundResult->IsEmpty() )
-		{
-			return Path;
-		}
-		// Otherwise we have our own path in there (in case we collapse children) or the path to the prim that collapsed us
-		else
+		FReadScopeLock ScopeLock( ImplPtr->Lock );
+
+		if ( const UE::UsdInfoCache::Private::FUsdPrimInfo* FoundInfo = ImplPtr->InfoMap.Find( Path ) )
 		{
-			return *FoundResult;
+			const UE::FSdfPath* CollapsedRoot = CollapsingType == ECollapsingType::Assets
+				? &FoundInfo->AssetCollapsedRoot
+				: &FoundInfo->ComponentCollapsedRoot;
+
+			// An empty path here means that we are not collapsed at all
+			if ( CollapsedRoot->IsEmpty() )
+			{
+				return Path;
+			}
+			// Otherwise we have our own path in there (in case we collapse children) or the path to the prim that collapsed us
+			else
+			{
+				return *CollapsedRoot;
+			}
 		}
+
+		// This should never happen: We should have cached the entire tree
+		ensureMsgf( false, TEXT( "Prim path '%s' has not been cached!" ), *Path.GetString() );
 	}
 
-	// This should never happen: We should have cached the entire tree
-	ensureMsgf( false, TEXT( "Prim path '%s' has not been cached!" ), *Path.GetString() );
 	return Path;
 }
 
@@ -113,11 +165,10 @@ namespace UE::USDInfoCacheImpl::Private
 	void GetPrimVertexCountAndSlots(
 		const pxr::UsdPrim& UsdPrim,
 		const FUsdSchemaTranslationContext& Context,
-		const TMap< UE::FSdfPath, uint64 >& InSubtreeToVertexCounts,
+		const FUsdInfoCache::FUsdInfoCacheImpl& Impl,
 		const TMap< UE::FSdfPath, TArray<UsdUtils::FUsdPrimMaterialSlot> >& InSubtreeToMaterialSlots,
 		uint64& OutVertexCount,
-		TArray<UsdUtils::FUsdPrimMaterialSlot>& OutMaterialSlots,
-		FRWLock& Lock
+		TArray<UsdUtils::FUsdPrimMaterialSlot>& OutMaterialSlots
 	)
 	{
 		TRACE_CPUPROFILER_EVENT_SCOPE( UE::USDInfoCacheImpl::Private::GetPrimVertexCountAndSlots );
@@ -168,16 +219,16 @@ namespace UE::USDInfoCacheImpl::Private
 				PrototypeVertexCounts.SetNumZeroed( PrototypePaths.size() );
 
 				{
-					FReadScopeLock ScopeLock( Lock );
+					FReadScopeLock ScopeLock( Impl.Lock );
 					for ( uint32 PrototypeIndex = 0; PrototypeIndex < PrototypePaths.size(); ++PrototypeIndex )
 					{
 						const pxr::SdfPath& PrototypePath = PrototypePaths[ PrototypeIndex ];
 
 						// If we're calling this for a point instancer we should have parsed the results for our
 						// prototype subtrees already
-						if ( const uint64* FoundPrototypeVertexCount = InSubtreeToVertexCounts.Find( UE::FSdfPath{ PrototypePath } ) )
+						if ( const UE::UsdInfoCache::Private::FUsdPrimInfo* FoundInfo = Impl.InfoMap.Find( UE::FSdfPath{ PrototypePath } ) )
 						{
-							PrototypeVertexCounts[ PrototypeIndex ] = *FoundPrototypeVertexCount;
+							PrototypeVertexCounts[ PrototypeIndex ] = FoundInfo->ExpectedVertexCountForSubtree;
 						}
 
 						if ( const TArray<UsdUtils::FUsdPrimMaterialSlot>* FoundPrototypeSlots = InSubtreeToMaterialSlots.Find( UE::FSdfPath{ PrototypePath } ) )
@@ -205,11 +256,8 @@ namespace UE::USDInfoCacheImpl::Private
 	void RecursiveRebuildCache(
 		const pxr::UsdPrim& UsdPrim,
 		FUsdSchemaTranslationContext& Context,
+		FUsdInfoCache::FUsdInfoCacheImpl& Impl,
 		FUsdSchemaTranslatorRegistry& Registry,
-		TMap< UE::FSdfPath, UE::FSdfPath >& AssetPathsToCollapsedRoot,
-		TMap< UE::FSdfPath, UE::FSdfPath >& ComponentPathsToCollapsedRoot,
-		FRWLock& Lock,
-		TMap< UE::FSdfPath, uint64 >& InOutSubtreeToVertexCounts,
 		TMap< UE::FSdfPath, TArray<UsdUtils::FUsdPrimMaterialSlot>>& InOutSubtreeToMaterialSlots,
 		TArray< FString >& InOutPointInstancerPaths,
 		uint64& OutSubtreeVertexCount,
@@ -276,11 +324,8 @@ namespace UE::USDInfoCacheImpl::Private
 				RecursiveRebuildCache(
 					Prims[ Index ],
 					Context,
+					Impl,
 					Registry,
-					AssetPathsToCollapsedRoot,
-					ComponentPathsToCollapsedRoot,
-					Lock,
-					InOutSubtreeToVertexCounts,
 					InOutSubtreeToMaterialSlots,
 					InOutPointInstancerPaths,
 					ChildSubtreeVertexCounts[ Index ],
@@ -304,11 +349,10 @@ namespace UE::USDInfoCacheImpl::Private
 			GetPrimVertexCountAndSlots(
 				UsdPrim,
 				Context,
-				InOutSubtreeToVertexCounts,
+				Impl,
 				InOutSubtreeToMaterialSlots,
 				OutSubtreeVertexCount,
-				OutSubtreeSlots,
-				Lock
+				OutSubtreeSlots
 			);
 
 			for ( uint32 ChildIndex = 0; ChildIndex < NumChildren; ++ChildIndex )
@@ -319,7 +363,9 @@ namespace UE::USDInfoCacheImpl::Private
 		}
 
 		{
-			FWriteScopeLock ScopeLock( Lock );
+			FWriteScopeLock ScopeLock( Impl.Lock );
+
+			UE::UsdInfoCache::Private::FUsdPrimInfo& Info = Impl.InfoMap.FindOrAdd( UE::FSdfPath{ UsdPrimPath } );
 
 			// For point instancers we can't guarantee we parsed the prototypes yet because they
 			// could technically be anywhere, so store them here for a later pass
@@ -333,16 +379,15 @@ namespace UE::USDInfoCacheImpl::Private
 			// stuff, and then we'll updadte all ancestors
 			else
 			{
-				InOutSubtreeToVertexCounts.Emplace( UsdPrimPath, OutSubtreeVertexCount );
+				Info.ExpectedVertexCountForSubtree = OutSubtreeVertexCount;
 				InOutSubtreeToMaterialSlots.Emplace( UsdPrimPath, OutSubtreeSlots );
 			}
 
 			// These paths will be still empty in case nothing has collapsed yet, hold UsdPrimPath in case UsdPrim
 			// collapses that type, or hold the path to the collapsed root passed in via our caller, in case we're
 			// collapsed
-			AssetPathsToCollapsedRoot.Emplace( UsdPrimPath, *AssetCollapsedRootOverride );
-			ComponentPathsToCollapsedRoot.Emplace( UsdPrimPath, *ComponentCollapsedRootOverride );
-
+			Info.AssetCollapsedRoot = *AssetCollapsedRootOverride;
+			Info.ComponentCollapsedRoot = *ComponentCollapsedRootOverride;
 		}
 	}
 
@@ -357,11 +402,10 @@ namespace UE::USDInfoCacheImpl::Private
 	 */
 	void UpdateInfoForPointInstancers(
 		const pxr::UsdStageRefPtr& Stage,
-		TArray< FString >& PointInstancerPaths,
 		const FUsdSchemaTranslationContext& Context,
-		TMap< UE::FSdfPath, uint64 >& InOutExpectedVertexCountPerSubtree,
-		TMap< UE::FSdfPath, TArray<UsdUtils::FUsdPrimMaterialSlot>>& InOutSubtreeMaterialSlots,
-		FRWLock& Lock
+		FUsdInfoCache::FUsdInfoCacheImpl& Impl,
+		TArray< FString >& PointInstancerPaths,
+		TMap< UE::FSdfPath, TArray<UsdUtils::FUsdPrimMaterialSlot>>& InOutSubtreeMaterialSlots
 	)
 	{
 		// We must sort point instancers in a particular order in case they depend on each other.
@@ -417,14 +461,16 @@ namespace UE::USDInfoCacheImpl::Private
 				GetPrimVertexCountAndSlots(
 					PointInstancer,
 					Context,
-					InOutExpectedVertexCountPerSubtree,
+					Impl,
 					InOutSubtreeMaterialSlots,
 					PointInstancerVertexCount,
-					PointInstancerMaterialSlots,
-					Lock
+					PointInstancerMaterialSlots
 				);
 
-				InOutExpectedVertexCountPerSubtree.Emplace( UsdPointInstancerPath, PointInstancerVertexCount );
+				FWriteScopeLock Lock{ Impl.Lock };
+				UE::UsdInfoCache::Private::FUsdPrimInfo& Info = Impl.InfoMap.FindOrAdd( UsdPointInstancerPath );
+
+				Info.ExpectedVertexCountForSubtree = PointInstancerVertexCount;
 				InOutSubtreeMaterialSlots.Emplace( UsdPointInstancerPath, PointInstancerMaterialSlots );
 
 				// Now that we have info on the point instancer itself, update the counts of all ancestors.
@@ -442,7 +488,9 @@ namespace UE::USDInfoCacheImpl::Private
 						break;
 					}
 
-					InOutExpectedVertexCountPerSubtree[ ParentPath ] += PointInstancerVertexCount;
+					UE::UsdInfoCache::Private::FUsdPrimInfo& ParentInfo = Impl.InfoMap.FindOrAdd( ParentPath );
+					ParentInfo.ExpectedVertexCountForSubtree += PointInstancerVertexCount;
+
 					InOutSubtreeMaterialSlots[ ParentPath ].Append( PointInstancerMaterialSlots );
 
 					// Break only here so we update the pseudoroot too
@@ -467,24 +515,28 @@ namespace UE::USDInfoCacheImpl::Private
 	 * only here.
 	 */
 	void CollectMaterialSlotCounts(
+		FUsdInfoCache::FUsdInfoCacheImpl& Impl,
 		const TMap< UE::FSdfPath, TArray<UsdUtils::FUsdPrimMaterialSlot>>& SubtreeMaterialSlots,
-		TMap< UE::FSdfPath, uint64 >& OutMaterialSlotCounts,
 		bool bMergeIdenticalSlots
 	)
 	{
+		FWriteScopeLock Lock{ Impl.Lock };
+
 		if ( bMergeIdenticalSlots )
 		{
 			for ( const TPair< UE::FSdfPath, TArray<UsdUtils::FUsdPrimMaterialSlot>>& Pair : SubtreeMaterialSlots )
 			{
 				TSet<UsdUtils::FUsdPrimMaterialSlot> SlotsSet{ Pair.Value };
-				OutMaterialSlotCounts.FindOrAdd( Pair.Key ) = SlotsSet.Num();
+				UE::UsdInfoCache::Private::FUsdPrimInfo& Info = Impl.InfoMap.FindOrAdd( Pair.Key );
+				Info.ExpectedMaterialSlotCountForSubtree = SlotsSet.Num();
 			}
 		}
 		else
 		{
 			for ( const TPair< UE::FSdfPath, TArray<UsdUtils::FUsdPrimMaterialSlot>>& Pair : SubtreeMaterialSlots )
 			{
-				OutMaterialSlotCounts.FindOrAdd( Pair.Key ) = Pair.Value.Num();
+				UE::UsdInfoCache::Private::FUsdPrimInfo& Info = Impl.InfoMap.FindOrAdd( Pair.Key );
+				Info.ExpectedMaterialSlotCountForSubtree = Pair.Value.Num();
 			}
 		}
 	}
@@ -493,9 +545,17 @@ namespace UE::USDInfoCacheImpl::Private
 
 TOptional<uint64> FUsdInfoCache::GetSubtreeVertexCount( const UE::FSdfPath& Path )
 {
-	if ( uint64* FoundCount = ExpectedVertexCountPerSubtree.Find( Path ) )
+	if ( FUsdInfoCacheImpl* ImplPtr = Impl.Get() )
 	{
-		return *FoundCount;
+		FReadScopeLock ScopeLock( ImplPtr->Lock );
+
+		if ( const UE::UsdInfoCache::Private::FUsdPrimInfo* FoundInfo = ImplPtr->InfoMap.Find( Path ) )
+		{
+			return FoundInfo->ExpectedVertexCountForSubtree;
+		}
+
+		// This should never happen: We should have cached the entire tree
+		ensureMsgf( false, TEXT( "Prim path '%s' has not been cached!" ), *Path.GetString() );
 	}
 
 	return {};
@@ -503,9 +563,17 @@ TOptional<uint64> FUsdInfoCache::GetSubtreeVertexCount( const UE::FSdfPath& Path
 
 TOptional<uint64> FUsdInfoCache::GetSubtreeMaterialSlotCount( const UE::FSdfPath& Path )
 {
-	if ( uint64* FoundCount = ExpectedMaterialSlotCountPerSubtree.Find( Path ) )
+	if ( FUsdInfoCacheImpl* ImplPtr = Impl.Get() )
 	{
-		return *FoundCount;
+		FReadScopeLock ScopeLock( ImplPtr->Lock );
+
+		if ( const UE::UsdInfoCache::Private::FUsdPrimInfo* FoundInfo = ImplPtr->InfoMap.Find( Path ) )
+		{
+			return FoundInfo->ExpectedMaterialSlotCountForSubtree;
+		}
+
+		// This should never happen: We should have cached the entire tree
+		ensureMsgf( false, TEXT( "Prim path '%s' has not been cached!" ), *Path.GetString() );
 	}
 
 	return {};
@@ -516,6 +584,12 @@ void FUsdInfoCache::RebuildCacheForSubtree( const UE::FUsdPrim& Prim, FUsdSchema
 #if USE_USD_SDK
 	TRACE_CPUPROFILER_EVENT_SCOPE( FUsdInfoCache::RebuildCacheForSubtree );
 
+	FUsdInfoCacheImpl* ImplPtr = Impl.Get();
+	if ( !ImplPtr )
+	{
+		return;
+	}
+
 	// We can't deallocate our info cache pointer with the Usd allocator
 	FScopedUnrealAllocs UEAllocs;
 
@@ -543,11 +617,8 @@ void FUsdInfoCache::RebuildCacheForSubtree( const UE::FUsdPrim& Prim, FUsdSchema
 		UE::USDInfoCacheImpl::Private::RecursiveRebuildCache(
 			UsdPrim,
 			Context,
+			*ImplPtr,
 			Registry,
-			AssetPathsToCollapsedRoot,
-			ComponentPathsToCollapsedRoot,
-			Lock,
-			ExpectedVertexCountPerSubtree,
 			TempSubtreeSlots,
 			PointInstancerPaths,
 			SubtreeVertexCount,
@@ -556,16 +627,15 @@ void FUsdInfoCache::RebuildCacheForSubtree( const UE::FUsdPrim& Prim, FUsdSchema
 
 		UE::USDInfoCacheImpl::Private::UpdateInfoForPointInstancers(
 			UsdPrim.GetStage(),
-			PointInstancerPaths,
 			Context,
-			ExpectedVertexCountPerSubtree,
-			TempSubtreeSlots,
-			Lock
+			*ImplPtr,
+			PointInstancerPaths,
+			TempSubtreeSlots
 		);
 
 		UE::USDInfoCacheImpl::Private::CollectMaterialSlotCounts(
+			*ImplPtr,
 			TempSubtreeSlots,
-			ExpectedMaterialSlotCountPerSubtree,
 			Context.bMergeIdenticalMaterialSlots
 		);
 	}
@@ -574,20 +644,20 @@ void FUsdInfoCache::RebuildCacheForSubtree( const UE::FUsdPrim& Prim, FUsdSchema
 
 void FUsdInfoCache::Clear()
 {
-	FWriteScopeLock ScopeLock( Lock );
-
-	AssetPathsToCollapsedRoot.Empty();
-	ComponentPathsToCollapsedRoot.Empty();
-	ExpectedVertexCountPerSubtree.Empty();
-	ExpectedMaterialSlotCountPerSubtree.Empty();
+	if ( FUsdInfoCacheImpl* ImplPtr = Impl.Get() )
+	{
+		FWriteScopeLock ScopeLock( ImplPtr->Lock );
+		ImplPtr->InfoMap.Empty();
+	}
 }
 
 bool FUsdInfoCache::IsEmpty()
 {
-	FReadScopeLock ScopeLock( Lock );
+	if ( FUsdInfoCacheImpl* ImplPtr = Impl.Get() )
+	{
+		FReadScopeLock ScopeLock( ImplPtr->Lock );
+		return ImplPtr->InfoMap.IsEmpty();
+	}
 
-	return AssetPathsToCollapsedRoot.IsEmpty()
-		&& ComponentPathsToCollapsedRoot.IsEmpty()
-		&& ExpectedVertexCountPerSubtree.IsEmpty()
-		&& ExpectedMaterialSlotCountPerSubtree.IsEmpty();
+	return true;
 }
\ No newline at end of file
diff --git a/Engine/Plugins/Importers/USDImporter/Source/USDSchemas/Public/USDInfoCache.h b/Engine/Plugins/Importers/USDImporter/Source/USDSchemas/Public/USDInfoCache.h
--- a/Engine/Plugins/Importers/USDImporter/Source/USDSchemas/Public/USDInfoCache.h
+++ b/Engine/Plugins/Importers/USDImporter/Source/USDSchemas/Public/USDInfoCache.h
@@ -26,6 +26,9 @@ enum class ECollapsingType
  */
 struct USDSCHEMAS_API FUsdInfoCache
 {
+	FUsdInfoCache();
+	virtual ~FUsdInfoCache();
+
 	bool Serialize( FArchive& Ar );
 
 	// Returns whether we contain any info about prim at 'Path' at all
@@ -49,13 +52,7 @@ struct USDSCHEMAS_API FUsdInfoCache
 	void Clear();
 	bool IsEmpty();
 
+	struct FUsdInfoCacheImpl;
 private:
-	mutable FRWLock Lock;
-
-	// TODO: Maybe use an impl class instead so that we can also clean up some of the parameters of the implementations
-	// TODO: Combine these maps into a single map to a struct with all this info
-	TMap< UE::FSdfPath, UE::FSdfPath > AssetPathsToCollapsedRoot;
-	TMap< UE::FSdfPath, UE::FSdfPath > ComponentPathsToCollapsedRoot;
-	TMap< UE::FSdfPath, uint64 > ExpectedVertexCountPerSubtree;
-	TMap< UE::FSdfPath, uint64 > ExpectedMaterialSlotCountPerSubtree;
+	TUniquePtr<FUsdInfoCacheImpl> Impl;
 };
\ No newline at end of file
