USD: Fix UE::FSdfPath assignment operator from pxr::SdfPath.

diff --git a/Engine/Plugins/Importers/USDImporter/Source/UnrealUSDWrapper/Private/UsdWrappers/SdfPath.cpp b/Engine/Plugins/Importers/USDImporter/Source/UnrealUSDWrapper/Private/UsdWrappers/SdfPath.cpp
--- a/Engine/Plugins/Importers/USDImporter/Source/UnrealUSDWrapper/Private/UsdWrappers/SdfPath.cpp
+++ b/Engine/Plugins/Importers/USDImporter/Source/UnrealUSDWrapper/Private/UsdWrappers/SdfPath.cpp
@@ -143,12 +143,14 @@ namespace UE
 
 	FSdfPath& FSdfPath::operator=(  const pxr::SdfPath& InSdfPath )
 	{
+		FScopedUnrealAllocs UnrealAllocs;
 		Impl = MakeUnique< Internal::FSdfPathImpl >( InSdfPath );
 		return *this;
 	}
 
 	FSdfPath& FSdfPath::operator=( pxr::SdfPath&& InSdfPath )
 	{
+		FScopedUnrealAllocs UnrealAllocs;
 		Impl = MakeUnique< Internal::FSdfPathImpl >( MoveTemp( InSdfPath ) );
 		return *this;
 	}
