USD: Remove stuff deprecated in 5.0.

diff --git a/Engine/Plugins/Importers/USDImporter/Source/USDUtilities/Private/USDLightConversion.cpp b/Engine/Plugins/Importers/USDImporter/Source/USDUtilities/Private/USDLightConversion.cpp
--- a/Engine/Plugins/Importers/USDImporter/Source/USDUtilities/Private/USDLightConversion.cpp
+++ b/Engine/Plugins/Importers/USDImporter/Source/USDUtilities/Private/USDLightConversion.cpp
@@ -87,41 +87,6 @@ namespace LightConversionImpl
 	}
 }
 
-bool UsdToUnreal::ConvertLight( const pxr::UsdLuxLightAPI& LightAPI, ULightComponentBase& LightComponentBase, double TimeCode )
-{
-	return UsdToUnreal::ConvertLight( LightAPI.GetPrim(), LightComponentBase, TimeCode );
-}
-
-bool UsdToUnreal::ConvertDistantLight( const pxr::UsdLuxDistantLight& DistantLight, UDirectionalLightComponent& LightComponent, double TimeCode )
-{
-	return UsdToUnreal::ConvertDistantLight( DistantLight.GetPrim(), LightComponent, TimeCode );
-}
-
-bool UsdToUnreal::ConvertRectLight( const FUsdStageInfo& StageInfo, const pxr::UsdLuxRectLight& RectLight, URectLightComponent& LightComponent, double TimeCode )
-{
-	return UsdToUnreal::ConvertRectLight( RectLight.GetPrim(), LightComponent, TimeCode );
-}
-
-bool UsdToUnreal::ConvertDiskLight( const FUsdStageInfo& StageInfo, const pxr::UsdLuxDiskLight& DiskLight, URectLightComponent& LightComponent, double TimeCode )
-{
-	return UsdToUnreal::ConvertDiskLight( DiskLight.GetPrim(), LightComponent, TimeCode );
-}
-
-bool UsdToUnreal::ConvertSphereLight( const FUsdStageInfo& StageInfo, const pxr::UsdLuxSphereLight& SphereLight, UPointLightComponent& LightComponent, double TimeCode )
-{
-	return UsdToUnreal::ConvertSphereLight( SphereLight.GetPrim(), LightComponent, TimeCode );
-}
-
-bool UsdToUnreal::ConvertDomeLight( const FUsdStageInfo& StageInfo, const pxr::UsdLuxDomeLight& DomeLight, USkyLightComponent& LightComponent, UUsdAssetCache* TexturesCache, double TimeCode )
-{
-	return UsdToUnreal::ConvertDomeLight( DomeLight.GetPrim(), LightComponent, TexturesCache );
-}
-
-bool UsdToUnreal::ConvertLuxShapingAPI( const FUsdStageInfo& StageInfo, const pxr::UsdLuxShapingAPI& ShapingAPI, USpotLightComponent& LightComponent, double TimeCode )
-{
-	return UsdToUnreal::ConvertLuxShapingAPI( ShapingAPI.GetPrim(), LightComponent, TimeCode );
-}
-
 bool UsdToUnreal::ConvertLight( const pxr::UsdPrim& Prim, ULightComponentBase& LightComponentBase, double UsdTimeCode )
 {
 	FScopedUsdAllocs Allocs;
diff --git a/Engine/Plugins/Importers/USDImporter/Source/USDUtilities/Private/USDPrimConversion.cpp b/Engine/Plugins/Importers/USDImporter/Source/USDUtilities/Private/USDPrimConversion.cpp
--- a/Engine/Plugins/Importers/USDImporter/Source/USDUtilities/Private/USDPrimConversion.cpp
+++ b/Engine/Plugins/Importers/USDImporter/Source/USDUtilities/Private/USDPrimConversion.cpp
@@ -432,136 +432,6 @@ bool UsdToUnreal::ConvertXformable( const pxr::UsdStageRefPtr& Stage, const pxr:
 	return true;
 }
 
-bool UsdToUnreal::ConvertXformable( const pxr::UsdTyped& Schema, UMovieScene3DTransformTrack& MovieSceneTrack, const FMovieSceneSequenceTransform& SequenceTransform )
-{
-	const UMovieScene* MovieScene = MovieSceneTrack.GetTypedOuter< UMovieScene >();
-
-	if ( !MovieScene )
-	{
-		return false;
-	}
-
-	FScopedUsdAllocs UsdAllocs;
-
-	pxr::UsdGeomXformable Xformable( Schema );
-
-	if ( !Xformable )
-	{
-		return false;
-	}
-
-	std::vector< double > UsdTimeSamples;
-	Xformable.GetTimeSamples( &UsdTimeSamples );
-
-	if ( UsdTimeSamples.empty() )
-	{
-		return false;
-	}
-
-	const FFrameRate Resolution = MovieScene->GetTickResolution();
-	const FFrameRate DisplayRate = MovieScene->GetDisplayRate();
-
-	TArray< FFrameNumber > FrameNumbers;
-	FrameNumbers.Reserve( UsdTimeSamples.size() );
-
-	TArray< FMovieSceneDoubleValue > TranslationValuesX;
-	TranslationValuesX.Reserve( UsdTimeSamples.size() );
-
-	TArray< FMovieSceneDoubleValue > TranslationValuesY;
-	TranslationValuesY.Reserve( UsdTimeSamples.size() );
-
-	TArray< FMovieSceneDoubleValue > TranslationValuesZ;
-	TranslationValuesZ.Reserve( UsdTimeSamples.size() );
-
-	TArray< FMovieSceneDoubleValue > RotationValuesX;
-	RotationValuesX.Reserve( UsdTimeSamples.size() );
-
-	TArray< FMovieSceneDoubleValue > RotationValuesY;
-	RotationValuesY.Reserve( UsdTimeSamples.size() );
-
-	TArray< FMovieSceneDoubleValue > RotationValuesZ;
-	RotationValuesZ.Reserve( UsdTimeSamples.size() );
-
-	TArray< FMovieSceneDoubleValue > ScaleValuesX;
-	ScaleValuesX.Reserve( UsdTimeSamples.size() );
-
-	TArray< FMovieSceneDoubleValue > ScaleValuesY;
-	ScaleValuesY.Reserve( UsdTimeSamples.size() );
-
-	TArray< FMovieSceneDoubleValue > ScaleValuesZ;
-	ScaleValuesZ.Reserve( UsdTimeSamples.size() );
-
-	pxr::UsdStageRefPtr Stage = Schema.GetPrim().GetStage();
-	const double StageTimeCodesPerSecond = Stage->GetTimeCodesPerSecond();
-	const FFrameRate StageFrameRate( StageTimeCodesPerSecond, 1 );
-
-	const ERichCurveInterpMode InterpMode = ( Stage->GetInterpolationType() == pxr::UsdInterpolationTypeLinear ) ? ERichCurveInterpMode::RCIM_Linear : ERichCurveInterpMode::RCIM_Constant;
-
-	for ( double UsdTimeSample : UsdTimeSamples )
-	{
-		// Frame Number
-		int32 FrameNumber = FMath::FloorToInt( UsdTimeSample );
-		float SubFrameNumber = UsdTimeSample - FrameNumber;
-
-		FFrameTime FrameTime( FrameNumber, SubFrameNumber );
-
-		FFrameTime KeyFrameTime = FFrameRate::TransformTime( FrameTime, StageFrameRate, Resolution );
-		KeyFrameTime *= SequenceTransform;
-		FrameNumbers.Add( KeyFrameTime.GetFrame() );
-
-		// Frame Values
-		FTransform Transform;
-		UsdToUnreal::ConvertXformable( Stage, Xformable, Transform, UsdTimeSample );
-
-		// Location
-		TranslationValuesX.Emplace_GetRef( Transform.GetLocation().X ).InterpMode = InterpMode;
-		TranslationValuesY.Emplace_GetRef( Transform.GetLocation().Y ).InterpMode = InterpMode;
-		TranslationValuesZ.Emplace_GetRef( Transform.GetLocation().Z ).InterpMode = InterpMode;
-
-		// Rotation
-		FRotator Rotator = Transform.Rotator();
-		RotationValuesX.Emplace_GetRef( Rotator.Roll ).InterpMode = InterpMode;
-		RotationValuesY.Emplace_GetRef( Rotator.Pitch ).InterpMode = InterpMode;
-		RotationValuesZ.Emplace_GetRef( Rotator.Yaw ).InterpMode = InterpMode;
-
-		// Scale
-		ScaleValuesX.Emplace_GetRef( Transform.GetScale3D().X ).InterpMode = InterpMode;
-		ScaleValuesY.Emplace_GetRef( Transform.GetScale3D().Y ).InterpMode = InterpMode;
-		ScaleValuesZ.Emplace_GetRef( Transform.GetScale3D().Z ).InterpMode = InterpMode;
-	}
-
-	bool bSectionAdded = false;
-	UMovieScene3DTransformSection* TransformSection = Cast< UMovieScene3DTransformSection >( MovieSceneTrack.FindOrAddSection( 0, bSectionAdded ) );
-	TransformSection->EvalOptions.CompletionMode = EMovieSceneCompletionMode::KeepState;
-	TransformSection->SetRange( TRange< FFrameNumber >::All() );
-
-	TArrayView< FMovieSceneDoubleChannel* > Channels = TransformSection->GetChannelProxy().GetChannels< FMovieSceneDoubleChannel >();
-
-	check( Channels.Num() >= 9 );
-
-	// Translation
-	Channels[0]->Set( FrameNumbers, TranslationValuesX );
-	Channels[1]->Set( FrameNumbers, TranslationValuesY );
-	Channels[2]->Set( FrameNumbers, TranslationValuesZ );
-
-	// Rotation
-	Channels[3]->Set( FrameNumbers, RotationValuesX );
-	Channels[4]->Set( FrameNumbers, RotationValuesY );
-	Channels[5]->Set( FrameNumbers, RotationValuesZ );
-
-	// Scale
-	Channels[6]->Set( FrameNumbers, ScaleValuesX );
-	Channels[7]->Set( FrameNumbers, ScaleValuesY );
-	Channels[8]->Set( FrameNumbers, ScaleValuesZ );
-
-	return true;
-}
-
-bool UsdToUnreal::ConvertGeomCamera( const pxr::UsdStageRefPtr& Stage, const pxr::UsdGeomCamera& GeomCamera, UCineCameraComponent& CameraComponent, double EvalTime )
-{
-	return ConvertGeomCamera( UE::FUsdPrim{ GeomCamera.GetPrim() }, CameraComponent, EvalTime );
-}
-
 bool UsdToUnreal::ConvertGeomCamera( const UE::FUsdPrim& Prim, UCineCameraComponent& CameraComponent, double UsdTimeCode )
 {
 	FScopedUsdAllocs Allocs;
@@ -2369,16 +2239,6 @@ bool UnrealToUsd::ConvertHierarchicalInstancedStaticMeshComponent( const UHierar
 	return true;
 }
 
-bool UnrealToUsd::ConvertCameraComponent( const pxr::UsdStageRefPtr& Stage, const UCineCameraComponent* CameraComponent, pxr::UsdPrim& UsdPrim, double TimeCode )
-{
-	if ( CameraComponent )
-	{
-		return ConvertCameraComponent( *CameraComponent, UE::FUsdPrim{ UsdPrim }, TimeCode );
-	}
-
-	return false;
-}
-
 bool UnrealToUsd::ConvertXformable( const FTransform& RelativeTransform, pxr::UsdPrim& UsdPrim, double TimeCode )
 {
 	if ( !UsdPrim )
diff --git a/Engine/Plugins/Importers/USDImporter/Source/USDUtilities/Public/USDLightConversion.h b/Engine/Plugins/Importers/USDImporter/Source/USDUtilities/Public/USDLightConversion.h
--- a/Engine/Plugins/Importers/USDImporter/Source/USDUtilities/Public/USDLightConversion.h
+++ b/Engine/Plugins/Importers/USDImporter/Source/USDUtilities/Public/USDLightConversion.h
@@ -54,21 +54,6 @@ struct FUsdStageInfo;
  */
 namespace UsdToUnreal
 {
-	UE_DEPRECATED( 5.0, "Prefer the overload that receives a pxr::UsdPrim" )
-	USDUTILITIES_API bool ConvertLight( const pxr::UsdLuxLightAPI& LightAPI, ULightComponentBase& LightComponentBase, double TimeCode );
-	UE_DEPRECATED( 5.0, "Prefer the overload that receives a pxr::UsdPrim" )
-	USDUTILITIES_API bool ConvertDistantLight( const pxr::UsdLuxDistantLight& DistantLight, UDirectionalLightComponent& LightComponent, double TimeCode );
-	UE_DEPRECATED( 5.0, "Prefer the overload that receives a pxr::UsdPrim" )
-	USDUTILITIES_API bool ConvertRectLight( const FUsdStageInfo& StageInfo, const pxr::UsdLuxRectLight& RectLight, URectLightComponent& LightComponent, double TimeCode );
-	UE_DEPRECATED( 5.0, "Prefer the overload that receives a pxr::UsdPrim" )
-	USDUTILITIES_API bool ConvertDiskLight( const FUsdStageInfo& StageInfo, const pxr::UsdLuxDiskLight& DiskLight, URectLightComponent& LightComponent, double TimeCode );
-	UE_DEPRECATED( 5.0, "Prefer the overload that receives a pxr::UsdPrim" )
-	USDUTILITIES_API bool ConvertSphereLight( const FUsdStageInfo& StageInfo, const pxr::UsdLuxSphereLight& SphereLight, UPointLightComponent& LightComponent, double TimeCode );
-	UE_DEPRECATED( 5.0, "Prefer the overload that receives a pxr::UsdPrim" )
-	USDUTILITIES_API bool ConvertDomeLight( const FUsdStageInfo& StageInfo, const pxr::UsdLuxDomeLight& DomeLight, USkyLightComponent& LightComponent, UUsdAssetCache* TexturesCache, double TimeCode );
-	UE_DEPRECATED( 5.0, "Prefer the overload that receives a pxr::UsdPrim" )
-	USDUTILITIES_API bool ConvertLuxShapingAPI( const FUsdStageInfo& StageInfo, const pxr::UsdLuxShapingAPI& ShapingAPI, USpotLightComponent& LightComponent, double TimeCode );
-
 	USDUTILITIES_API bool ConvertLight( const pxr::UsdPrim& Prim, ULightComponentBase& LightComponent, double UsdTimeCode = UsdUtils::GetDefaultTimeCode() );
 	USDUTILITIES_API bool ConvertDistantLight( const pxr::UsdPrim& Prim, UDirectionalLightComponent& LightComponent, double UsdTimeCode = UsdUtils::GetDefaultTimeCode() );
 	USDUTILITIES_API bool ConvertRectLight( const pxr::UsdPrim& Prim, URectLightComponent& LightComponent, double UsdTimeCode = UsdUtils::GetDefaultTimeCode() );
diff --git a/Engine/Plugins/Importers/USDImporter/Source/USDUtilities/Public/USDPrimConversion.h b/Engine/Plugins/Importers/USDImporter/Source/USDUtilities/Public/USDPrimConversion.h
--- a/Engine/Plugins/Importers/USDImporter/Source/USDUtilities/Public/USDPrimConversion.h
+++ b/Engine/Plugins/Importers/USDImporter/Source/USDUtilities/Public/USDPrimConversion.h
@@ -51,15 +51,6 @@ namespace UE
 
 namespace UsdToUnreal
 {
-	/**
-	 * Converts a time varying UsdGeomXformable to a UMovieScene3DTransformTrack.
-	 * @param Schema               The Xformable to read from
-	 * @param MovieSceneTrack      The track to add the time sampled transform to
-	 * @param SequenceTransform    The time transform to apply to the track keys to get them from Usd Stage time to track time (in other words: from main sequence to subsequence)
-	 */
-	UE_DEPRECATED( 5.0, "Use CreatePropertyTrackReader and ConvertTransformTimeSamples (c.f. FUsdLevelSequenceHelperImpl::AddCommonTracks)" )
-	USDUTILITIES_API bool ConvertXformable( const pxr::UsdTyped& Schema, UMovieScene3DTransformTrack& MovieSceneTrack, const FMovieSceneSequenceTransform& SequenceTransform );
-
 	/**
 	 * Converts a pxr::UsdGeomXformable's attribute values into an USceneComponent's property values
 	 * @param Stage - Stage that contains the prim to convert
@@ -82,9 +73,6 @@ namespace UsdToUnreal
 	 */
 	USDUTILITIES_API bool ConvertXformable( const pxr::UsdStageRefPtr& Stage, const pxr::UsdTyped& Schema, FTransform& OutTransform, double EvalTime, bool* bOutResetTransformStack = nullptr );
 
-	UE_DEPRECATED( 5.0, "Prefer the overload that receives an UE::FUsdPrim" )
-	USDUTILITIES_API bool ConvertGeomCamera( const pxr::UsdStageRefPtr& Stage, const pxr::UsdGeomCamera& GeomCamera, UCineCameraComponent& CameraComponent, double EvalTime );
-
 	USDUTILITIES_API bool ConvertGeomCamera( const UE::FUsdPrim& Prim, UCineCameraComponent& CameraComponent, double UsdTimeCode = UsdUtils::GetDefaultTimeCode() );
 
 	/**
@@ -121,9 +109,6 @@ namespace UsdToUnreal
 
 namespace UnrealToUsd
 {
-	UE_DEPRECATED( 5.0, "Prefer the overload that receives a pxr::UsdPrim" )
-	USDUTILITIES_API bool ConvertCameraComponent( const pxr::UsdStageRefPtr& Stage, const UCineCameraComponent* CameraComponent, pxr::UsdPrim& UsdPrim, double TimeCode = UsdUtils::GetDefaultTimeCode() );
-
 	USDUTILITIES_API bool ConvertCameraComponent( const UCineCameraComponent& CameraComponent, pxr::UsdPrim& Prim, double UsdTimeCode = UsdUtils::GetDefaultTimeCode() );
 
 	/**
diff --git a/Engine/Plugins/Importers/USDImporter/Source/UnrealUSDWrapper/Private/UnrealUSDWrapper.cpp b/Engine/Plugins/Importers/USDImporter/Source/UnrealUSDWrapper/Private/UnrealUSDWrapper.cpp
--- a/Engine/Plugins/Importers/USDImporter/Source/UnrealUSDWrapper/Private/UnrealUSDWrapper.cpp
+++ b/Engine/Plugins/Importers/USDImporter/Source/UnrealUSDWrapper/Private/UnrealUSDWrapper.cpp
@@ -151,236 +151,6 @@ namespace UnrealIdentifiers
 	const TfToken UnrealEngineVersion = TfToken( "unreal:engineVersion" );
 }
 
-std::string FUsdAttribute::GetUnrealPropertyPath( const pxr::UsdAttribute& Attribute )
-{
-	std::string UnrealPropertyPath;
-
-	VtValue CustomData = Attribute.GetCustomDataByKey(UnrealIdentifiers::PropertyPath);
-
-	if (CustomData.IsHolding<std::string>())
-	{
-		UnrealPropertyPath = CustomData.Get<std::string>();
-	}
-
-	return UnrealPropertyPath;
-}
-
-template<typename T>
-bool GetValue(T& OutVal, const pxr::UsdAttribute& Attrib, int ArrayIndex, double Time)
-{
-	bool bResult = false;
-
-	if (ArrayIndex != -1)
-	{
-		// Note: VtArray is copy on write so this is cheap
-		VtArray<T> Array;
-		if (Attrib.Get(&Array, Time))
-		{
-			OutVal = Array[ArrayIndex];
-			bResult = true;
-		}
-	}
-	else
-	{
-		bResult = Attrib.Get(&OutVal, Time);
-	}
-
-	return bResult;
-}
-
-template<typename T>
-bool IsHolding(const VtValue& Value)
-{
-	return Value.IsHolding<T>() || Value.IsHolding<VtArray<T>>();
-}
-
-bool FUsdAttribute::AsInt(int64_t& OutVal, const pxr::UsdAttribute& Attribute, int ArrayIndex, double Time)
-{
-	// We test multiple types of ints here. int64 is always returned as it can hold all other types
-	// Unreal expects this
-	VtValue Value;
-	bool bResult = Attribute.Get(&Value, Time);
-	if (IsHolding<int8_t>(Value))
-	{
-		uint8_t Val = 0;
-		bResult = GetValue(Val, Attribute, ArrayIndex, Time);
-		OutVal = Val;
-	}
-	else if (IsHolding<int32_t>(Value))
-	{
-		int32_t Val = 0;
-		bResult = GetValue(Val, Attribute, ArrayIndex, Time);
-		OutVal = Val;
-	}
-	else if (IsHolding<int64_t>(Value))
-	{
-		int64_t Val = 0;
-		bResult = GetValue(Val, Attribute, ArrayIndex, Time);
-		OutVal = Val;
-	}
-
-	return bResult;
-}
-
-bool FUsdAttribute::AsUnsignedInt(uint64_t& OutVal, const pxr::UsdAttribute& Attribute, int ArrayIndex, double Time)
-{
-	// We test multiple types of ints here. uint64 is always returned as it can hold all other types
-	// Unreal expects this
-	VtValue Value;
-	bool bResult = Attribute.Get(&Value, Time);
-	if (IsHolding<uint8_t>(Value))
-	{
-		uint8_t Val;
-		bResult = GetValue(Val, Attribute, ArrayIndex, Time);
-		OutVal = Val;
-	}
-	else if (IsHolding<uint32_t>(Value))
-	{
-		uint32_t Val;
-		bResult = GetValue(Val, Attribute, ArrayIndex, Time);
-		OutVal = Val;
-	}
-	else if (IsHolding<uint64_t>(Value))
-	{
-		uint64_t Val;
-		bResult = GetValue(Val, Attribute, ArrayIndex, Time);
-		OutVal = Val;
-	}
-
-	return bResult;
-}
-
-bool FUsdAttribute::AsDouble(double& OutVal, const pxr::UsdAttribute& Attribute, int ArrayIndex, double Time)
-{
-	bool bResult = false;
-
-	bResult = GetValue<double>(OutVal, Attribute, ArrayIndex, Time);
-
-	if (!bResult)
-	{
-		float Val = 0.0f;
-		bResult = GetValue<float>(Val, Attribute, ArrayIndex, Time);
-		OutVal = Val;
-	}
-
-	return bResult;
-}
-
-bool FUsdAttribute::AsString(const char*& OutVal, const pxr::UsdAttribute& Attribute, int ArrayIndex, double Time)
-{
-	// this method is very hacky to return temp strings
-	// designed to have the string copied immediately
-	bool bResult = false;
-
-	VtValue Value;
-	Attribute.Get(&Value);
-	// mem leak
-	static std::string Temp;
-	if (IsHolding<std::string>(Value))
-	{
-		bResult = GetValue(Temp, Attribute, ArrayIndex, Time);
-
-		OutVal = Temp.c_str();
-	}
-	else if (IsHolding<TfToken>(Value))
-	{
-		TfToken Token;
-		bResult = GetValue(Token, Attribute, ArrayIndex, Time);
-
-		Temp = Token.GetString();
-
-		OutVal = Temp.c_str();
-	}
-
-	return bResult;
-}
-
-bool FUsdAttribute::AsBool(bool& OutVal, const pxr::UsdAttribute& Attribute, int ArrayIndex, double Time)
-{
-	return GetValue(OutVal, Attribute, ArrayIndex, Time);
-}
-
-bool FUsdAttribute::AsVector2(FUsdVector2Data& OutVal, const pxr::UsdAttribute& Attribute, int ArrayIndex, double Time)
-{
-	GfVec2f Value;
-	const bool bResult = GetValue(Value, Attribute, ArrayIndex, Time);
-
-	OutVal.X = Value[0];
-	OutVal.Y = Value[1];
-
-	return bResult;
-}
-
-bool FUsdAttribute::AsVector3(FUsdVectorData& OutVal, const pxr::UsdAttribute& Attribute, int ArrayIndex, double Time)
-{
-	GfVec3f Value;
-	const bool bResult = GetValue(Value, Attribute, ArrayIndex, Time);
-
-	OutVal.X = Value[0];
-	OutVal.Y = Value[1];
-	OutVal.Z = Value[2];
-
-	return bResult;
-}
-
-bool FUsdAttribute::AsVector4(FUsdVector4Data& OutVal, const pxr::UsdAttribute& Attribute, int ArrayIndex, double Time)
-{
-	GfVec4f Value;
-	const bool bResult = GetValue(Value, Attribute, ArrayIndex, Time);
-
-	OutVal.X = Value[0];
-	OutVal.Y = Value[1];
-	OutVal.Z = Value[2];
-	OutVal.W = Value[3];
-
-	return bResult;
-}
-
-bool FUsdAttribute::AsColor(FUsdVector4Data& OutVal, const pxr::UsdAttribute& Attribute, int ArrayIndex, double Time)
-{
-	GfVec4f Value;
-	bool bResult = GetValue(Value, Attribute, ArrayIndex, Time);
-
-	if (bResult)
-	{
-		OutVal.X = Value[0];
-		OutVal.Y = Value[1];
-		OutVal.Z = Value[2];
-		OutVal.W = Value[3];
-	}
-	else
-	{
-		// Try color 3 with a = 1;
-		GfVec3f Value3;
-		bResult = GetValue<GfVec3f>(Value3, Attribute, ArrayIndex, Time);
-		OutVal.X = Value3[0];
-		OutVal.Y = Value3[1];
-		OutVal.Z = Value3[2];
-		OutVal.W = 1;
-	}
-
-	return bResult;
-}
-
-bool FUsdAttribute::IsUnsigned(const pxr::UsdAttribute& Attribute)
-{
-	VtValue Value;
-	Attribute.Get(&Value);
-
-	return IsHolding<uint8_t>(Value)
-		|| IsHolding<uint32_t>(Value)
-		|| IsHolding<uint64_t>(Value);
-}
-
-int FUsdAttribute::GetArraySize( const pxr::UsdAttribute& Attribute )
-{
-	VtValue Value;
-	Attribute.Get(&Value);
-
-	return Value.IsArrayValued() ? (int)Value.GetArraySize() : -1;
-
-}
-
 bool IUsdPrim::IsValidPrimName(const FString& Name, FText& OutReason)
 {
 	if (Name.IsEmpty())
@@ -577,112 +347,6 @@ pxr::GfMatrix4d IUsdPrim::GetLocalToWorldTransform(const pxr::UsdPrim& Prim, dou
 	return AccumulatedTransform;
 }
 
-std::string IUsdPrim::GetUnrealPropertyPath(const pxr::UsdPrim& Prim)
-{
-	VtValue CustomData = Prim.GetCustomDataByKey(UnrealIdentifiers::PropertyPath);
-	if (CustomData.IsHolding<std::string>())
-	{
-		return CustomData.Get<std::string>();
-	}
-
-	return {};
-}
-
-TArray< UE::FUsdAttribute > PrivateGetAttributes(const pxr::UsdPrim& Prim, const TfToken& ByMetadata)
-{
-	FScopedUsdAllocs UsdAllocs;
-
-	std::vector<UsdAttribute> Attributes = Prim.GetAttributes();
-
-	TArray<UE::FUsdAttribute> OutAttributes;
-	OutAttributes.Reserve(Attributes.size());
-
-	for (UsdAttribute& Attr : Attributes)
-	{
-		if (ByMetadata.IsEmpty() || Attr.HasCustomDataKey(ByMetadata))
-		{
-			OutAttributes.Emplace( Attr);
-		}
-	}
-
-	return OutAttributes;
-}
-
-TArray< UE::FUsdAttribute > IUsdPrim::GetUnrealPropertyAttributes(const pxr::UsdPrim& Prim)
-{
-	return PrivateGetAttributes(Prim, UnrealIdentifiers::PropertyPath);
-}
-
-std::string IUsdPrim::GetUnrealAssetPath(const pxr::UsdPrim& Prim)
-{
-	std::string UnrealAssetPath;
-
-	UsdAttribute UnrealAssetPathAttr = Prim.GetAttribute(UnrealIdentifiers::AssetPath);
-	if (UnrealAssetPathAttr.HasValue())
-	{
-		UnrealAssetPathAttr.Get(&UnrealAssetPath);
-	}
-
-	return UnrealAssetPath;
-}
-
-std::string IUsdPrim::GetUnrealActorClass(const pxr::UsdPrim& Prim)
-{
-	std::string UnrealActorClass;
-
-	UsdAttribute UnrealActorClassAttr = Prim.GetAttribute(UnrealIdentifiers::ActorClass);
-	if (UnrealActorClassAttr.HasValue())
-	{
-		UnrealActorClassAttr.Get(&UnrealActorClass);
-	}
-
-	return UnrealActorClass;
-}
-
-namespace Internal
-{
-	TArray< FString > FillMaterialInfo(const SdfPath& Path, UsdStageWeakPtr Stage)
-	{
-		TArray< FString > MaterialNames;
-
-		// load each material at the material path;
-		UsdPrim MaterialPrim = Stage->Load(Path);
-
-		if(MaterialPrim)
-		{
-			// Default to using the prim path name as the path for this material in Unreal
-			FString MaterialName = ANSI_TO_TCHAR( MaterialPrim.GetName().GetString().c_str() ) ;
-
-			std::string UsdMaterialName;
-
-			// See if the material has an "unrealAssetPath" attribute.  This should be the full name of the material
-			static const TfToken AssetPathToken = TfToken(UnrealIdentifiers::AssetPath);
-			UsdAttribute UnrealAssetPathAttr = MaterialPrim.GetAttribute(AssetPathToken);
-			if (UnrealAssetPathAttr && UnrealAssetPathAttr.HasValue())
-			{
-				UnrealAssetPathAttr.Get(&UsdMaterialName);
-
-
-			}
-
-			MaterialNames.Add( MoveTemp( MaterialName ) );
-		}
-
-		return MaterialNames;
-	}
-}
-
-std::string DiscoverInformationAboutUsdMaterial(const UsdShadeMaterial& ShadeMaterial, const UsdGeomGprim& boundPrim)
-{
-	std::string ShadingEngineName = (ShadeMaterial ? ShadeMaterial.GetPrim() : boundPrim.GetPrim()).GetPrimPath().GetString();
-	return ShadingEngineName;
-}
-
-bool IUsdPrim::IsUnrealProperty(const pxr::UsdPrim& Prim)
-{
-	return Prim.HasCustomDataKey(UnrealIdentifiers::PropertyPath);
-}
-
 bool IUsdPrim::HasTransform(const pxr::UsdPrim& Prim)
 {
 	return UsdGeomXformable(Prim) ? true : false;
diff --git a/Engine/Plugins/Importers/USDImporter/Source/UnrealUSDWrapper/Public/UnrealUSDWrapper.h b/Engine/Plugins/Importers/USDImporter/Source/UnrealUSDWrapper/Public/UnrealUSDWrapper.h
--- a/Engine/Plugins/Importers/USDImporter/Source/UnrealUSDWrapper/Public/UnrealUSDWrapper.h
+++ b/Engine/Plugins/Importers/USDImporter/Source/UnrealUSDWrapper/Public/UnrealUSDWrapper.h
@@ -94,78 +94,6 @@ enum class EUsdDefaultKind : int32
 };
 ENUM_CLASS_FLAGS( EUsdDefaultKind );
 
-
-struct FUsdVector2Data
-{
-	UE_DEPRECATED( 5.0, "Not used anymore" )
-	FUsdVector2Data(float InX = 0, float InY = 0)
-		: X(InX)
-		, Y(InY)
-	{}
-
-	float X;
-	float Y;
-};
-
-struct FUsdVectorData
-{
-	UE_DEPRECATED( 5.0, "Not used anymore" )
-	FUsdVectorData(float InX = 0, float InY = 0, float InZ = 0)
-		: X(InX)
-		, Y(InY)
-		, Z(InZ)
-	{}
-
-	float X;
-	float Y;
-	float Z;
-};
-
-struct FUsdVector4Data
-{
-	UE_DEPRECATED( 5.0, "Not used anymore" )
-	FUsdVector4Data(float InX = 0, float InY = 0, float InZ = 0, float InW = 0)
-		: X(InX)
-		, Y(InY)
-		, Z(InZ)
-		, W(InW)
-	{}
-
-	float X;
-	float Y;
-	float Z;
-	float W;
-};
-
-struct FUsdUVData
-{
-	UE_DEPRECATED( 5.0, "Not used anymore" )
-	FUsdUVData()
-	{}
-
-	/** Defines how UVs are mapped to faces */
-	EUsdInterpolationMethod UVInterpMethod;
-
-	/** Raw UVs */
-	std::vector<FUsdVector2Data> Coords;
-};
-
-struct FUsdQuatData
-{
-	UE_DEPRECATED( 5.0, "Not used anymore" )
-	FUsdQuatData(float InX=0, float InY = 0, float InZ = 0, float InW = 0)
-		: X(InX)
-		, Y(InY)
-		, Z(InZ)
-		, W(InW)
-	{}
-
-	float X;
-	float Y;
-	float Z;
-	float W;
-};
-
 UENUM()
 enum class EUsdInitialLoadSet : uint8
 {
@@ -296,53 +224,6 @@ private:
 	static TUniquePtr<FUsdDiagnosticDelegate> Delegate;
 };
 
-class FUsdAttribute
-{
-public:
-#if USE_USD_SDK
-
-	UE_DEPRECATED( 5.0, "Prefer using the utilities in USDConversionUtils.h or the actual pxr::UsdAttribute wrapper with the same name at UsdWrappers/UsdAttribute.h" )
-	UNREALUSDWRAPPER_API static std::string GetUnrealPropertyPath(const pxr::UsdAttribute& Attribute);
-
-	// Get the number of elements in the array if it is an array.  Otherwise -1
-	UE_DEPRECATED( 5.0, "Prefer using the utilities in USDConversionUtils.h or the actual pxr::UsdAttribute wrapper with the same name at UsdWrappers/UsdAttribute.h" )
-	UNREALUSDWRAPPER_API static int GetArraySize(const pxr::UsdAttribute& Attribute);
-
-	UE_DEPRECATED( 5.0, "Prefer using the utilities in USDConversionUtils.h or the actual pxr::UsdAttribute wrapper with the same name at UsdWrappers/UsdAttribute.h" )
-	UNREALUSDWRAPPER_API static bool AsInt(int64_t& OutVal, const pxr::UsdAttribute& Attribute, int ArrayIndex = -1, double Time = UnrealUSDWrapper::GetDefaultTimeCode());
-
-	UE_DEPRECATED( 5.0, "Prefer using the utilities in USDConversionUtils.h or the actual pxr::UsdAttribute wrapper with the same name at UsdWrappers/UsdAttribute.h" )
-	UNREALUSDWRAPPER_API static bool AsUnsignedInt(uint64_t& OutVal, const pxr::UsdAttribute& Attribute, int ArrayIndex = -1, double Time = UnrealUSDWrapper::GetDefaultTimeCode());
-
-	UE_DEPRECATED( 5.0, "Prefer using the utilities in USDConversionUtils.h or the actual pxr::UsdAttribute wrapper with the same name at UsdWrappers/UsdAttribute.h" )
-	UNREALUSDWRAPPER_API static bool AsDouble(double& OutVal, const pxr::UsdAttribute& Attribute, int ArrayIndex = -1, double Time = UnrealUSDWrapper::GetDefaultTimeCode());
-
-	UE_DEPRECATED( 5.0, "Prefer using the utilities in USDConversionUtils.h or the actual pxr::UsdAttribute wrapper with the same name at UsdWrappers/UsdAttribute.h" )
-	UNREALUSDWRAPPER_API static bool AsString(const char*& OutVal, const pxr::UsdAttribute& Attribute, int ArrayIndex = -1, double Time = UnrealUSDWrapper::GetDefaultTimeCode());
-
-	UE_DEPRECATED( 5.0, "Prefer using the utilities in USDConversionUtils.h or the actual pxr::UsdAttribute wrapper with the same name at UsdWrappers/UsdAttribute.h" )
-	UNREALUSDWRAPPER_API static bool AsBool(bool& OutVal, const pxr::UsdAttribute& Attribute, int ArrayIndex = -1, double Time = UnrealUSDWrapper::GetDefaultTimeCode());
-
-	UE_DEPRECATED( 5.0, "Prefer using the utilities in USDConversionUtils.h or the actual pxr::UsdAttribute wrapper with the same name at UsdWrappers/UsdAttribute.h" )
-	UNREALUSDWRAPPER_API static bool AsVector2(FUsdVector2Data& OutVal, const pxr::UsdAttribute& Attribute, int ArrayIndex = -1, double Time = UnrealUSDWrapper::GetDefaultTimeCode());
-
-	UE_DEPRECATED( 5.0, "Prefer using the utilities in USDConversionUtils.h or the actual pxr::UsdAttribute wrapper with the same name at UsdWrappers/UsdAttribute.h" )
-	UNREALUSDWRAPPER_API static bool AsVector3(FUsdVectorData& OutVal, const pxr::UsdAttribute& Attribute, int ArrayIndex = -1, double Time = UnrealUSDWrapper::GetDefaultTimeCode());
-
-	UE_DEPRECATED( 5.0, "Prefer using the utilities in USDConversionUtils.h or the actual pxr::UsdAttribute wrapper with the same name at UsdWrappers/UsdAttribute.h" )
-	UNREALUSDWRAPPER_API static bool AsVector4(FUsdVector4Data& OutVal, const pxr::UsdAttribute& Attribute, int ArrayIndex = -1, double Time = UnrealUSDWrapper::GetDefaultTimeCode());
-
-	UE_DEPRECATED( 5.0, "Prefer using the utilities in USDConversionUtils.h or the actual pxr::UsdAttribute wrapper with the same name at UsdWrappers/UsdAttribute.h" )
-	UNREALUSDWRAPPER_API static bool AsColor(FUsdVector4Data& OutVal, const pxr::UsdAttribute& Attribute, int ArrayIndex = -1, double Time = UnrealUSDWrapper::GetDefaultTimeCode());
-
-	UE_DEPRECATED( 5.0, "Prefer using the utilities in USDConversionUtils.h or the actual pxr::UsdAttribute wrapper with the same name at UsdWrappers/UsdAttribute.h" )
-	UNREALUSDWRAPPER_API static bool IsUnsigned(const pxr::UsdAttribute& Attribute);
-
-#endif // #if USE_USD_SDK
-};
-
-
-
 class IUsdPrim
 {
 public:
@@ -366,21 +247,6 @@ public:
 	static UNREALUSDWRAPPER_API pxr::GfMatrix4d GetLocalToWorldTransform(const pxr::UsdPrim& Prim, double Time, const pxr::SdfPath& AbsoluteRootPath);
 	static UNREALUSDWRAPPER_API bool HasTransform(const pxr::UsdPrim& Prim);
 
-	UE_DEPRECATED( 5.0, "Not used anymore" )
-	static UNREALUSDWRAPPER_API bool IsUnrealProperty(const pxr::UsdPrim& Prim);
-
-	UE_DEPRECATED( 5.0, "Not used anymore" )
-	static UNREALUSDWRAPPER_API std::string GetUnrealPropertyPath(const pxr::UsdPrim& Prim);
-
-	UE_DEPRECATED( 5.0, "Not used anymore" )
-	static UNREALUSDWRAPPER_API TArray< UE::FUsdAttribute > GetUnrealPropertyAttributes(const pxr::UsdPrim& Prim);
-
-	UE_DEPRECATED( 5.0, "Not used anymore" )
-	static UNREALUSDWRAPPER_API std::string GetUnrealAssetPath(const pxr::UsdPrim& Prim);
-
-	UE_DEPRECATED( 5.0, "Not used anymore" )
-	static UNREALUSDWRAPPER_API std::string GetUnrealActorClass(const pxr::UsdPrim& Prim);
-
 	static UNREALUSDWRAPPER_API bool SetActiveLODIndex(const pxr::UsdPrim& Prim, int LODIndex);
 
 	static UNREALUSDWRAPPER_API EUsdGeomOrientation GetGeometryOrientation(const pxr::UsdGeomMesh& Mesh);
