USD: Allow picking custom base materials to use for generated material instances.

Rename the display color base materials from "DoubleSided" to "TwoSided" for consistency;
Remove unused configs from USDAssetImportData objects;
Change project settings config to use a separate config file that carries the used base material configs into packaged games;
Remove code that added the USDImporter base material folder to directories to always cook as it's no longer needed.

diff --git a/Engine/Plugins/Importers/USDImporter/Source/USDClasses/Public/USDAssetImportData.h b/Engine/Plugins/Importers/USDImporter/Source/USDClasses/Public/USDAssetImportData.h
--- a/Engine/Plugins/Importers/USDImporter/Source/USDClasses/Public/USDAssetImportData.h
+++ b/Engine/Plugins/Importers/USDImporter/Source/USDClasses/Public/USDAssetImportData.h
@@ -7,7 +7,7 @@
 
 #include "USDAssetImportData.generated.h"
 
-UCLASS(config = EditorPerProjectUserSettings, AutoExpandCategories = (Options), MinimalAPI)
+UCLASS(AutoExpandCategories = (Options), MinimalAPI)
 class UUsdAssetImportData : public UAssetImportData
 {
 	GENERATED_BODY()
@@ -22,7 +22,7 @@ public:
 	TObjectPtr<class UObject> ImportOptions;
 };
 
-UCLASS(config = EditorPerProjectUserSettings, AutoExpandCategories = (Options), MinimalAPI)
+UCLASS(AutoExpandCategories = (Options), MinimalAPI)
 class UUsdAnimSequenceAssetImportData : public UUsdAssetImportData
 {
 	GENERATED_BODY()
@@ -50,7 +50,7 @@ struct FUsdPrimPathList
 };
 
 /** We assign these to UStaticMeshes or USkeletalMeshes generated from USD */
-UCLASS( config = EditorPerProjectUserSettings, AutoExpandCategories = ( Options ), MinimalAPI )
+UCLASS( AutoExpandCategories = ( Options ), MinimalAPI )
 class UUsdMeshAssetImportData : public UUsdAssetImportData
 {
 	GENERATED_BODY()
diff --git a/Engine/Plugins/Importers/USDImporter/Source/USDClasses/Public/USDProjectSettings.h b/Engine/Plugins/Importers/USDImporter/Source/USDClasses/Public/USDProjectSettings.h
--- a/Engine/Plugins/Importers/USDImporter/Source/USDClasses/Public/USDProjectSettings.h
+++ b/Engine/Plugins/Importers/USDImporter/Source/USDClasses/Public/USDProjectSettings.h
@@ -16,7 +16,9 @@ enum class EUsdSaveDialogBehavior : uint8
 	ShowPrompt
 };
 
-UCLASS(config=Editor, meta=(DisplayName=USDImporter), MinimalAPI)
+// USDImporter and defaultconfig here so this ends up at DefaultUSDImporter.ini in the editor, and is sent to the
+// packaged game as well
+UCLASS(config=USDImporter, defaultconfig, meta=(DisplayName=USDImporter), MinimalAPI)
 class UUsdProjectSettings : public UDeveloperSettings
 {
 	GENERATED_BODY()
@@ -52,4 +54,114 @@ public:
 	// Whether to display the pop up dialog asking what to do about dirty USD layers when closing USD stages
 	UPROPERTY(config, EditAnywhere, Category = "USD|Dialogs" )
 	EUsdSaveDialogBehavior ShowSaveLayersDialogWhenClosing = EUsdSaveDialogBehavior::ShowPrompt;
+
+	// Note that the below properties being FSoftObjectPath ensure that these materials are cooked into packaged games
+
+	/**
+	 * What material to use as base material when creating material instances from USD materials.
+	 * You can swap these with your own materials, but make sure that the replacement materials have parameters with
+	 * the same names and types as the ones provided by the default material, otherwise the instances will not have
+	 * the parameters filled with values extracted from the USD material when parsing.
+	 */
+	UPROPERTY( config, EditAnywhere, Category = "USD|Base Materials|UsdPreviewSurface", meta = ( AllowedClasses = "/Script/Engine.MaterialInterface" ) )
+	FSoftObjectPath BasePreviewSurfaceMaterial = FSoftObjectPath{ TEXT( "/USDImporter/Materials/UsdPreviewSurface.UsdPreviewSurface" ) };
+
+	/**
+	 * What material to use as base material when creating material instances from USD materials.
+	 * You can swap these with your own materials, but make sure that the replacement materials have parameters with
+	 * the same names and types as the ones provided by the default material, otherwise the instances will not have
+	 * the parameters filled with values extracted from the USD material when parsing.
+	 */
+	UPROPERTY( config, EditAnywhere, Category = "USD|Base Materials|UsdPreviewSurface", meta = ( AllowedClasses = "/Script/Engine.MaterialInterface" ) )
+	FSoftObjectPath BasePreviewSurfaceTranslucentMaterial = FSoftObjectPath{ TEXT( "/USDImporter/Materials/UsdPreviewSurfaceTranslucent.UsdPreviewSurfaceTranslucent" ) };
+
+	/**
+	 * What material to use as base material when creating material instances from USD materials.
+	 * You can swap these with your own materials, but make sure that the replacement materials have parameters with
+	 * the same names and types as the ones provided by the default material, otherwise the instances will not have
+	 * the parameters filled with values extracted from the USD material when parsing.
+	 */
+	UPROPERTY( config, EditAnywhere, Category = "USD|Base Materials|UsdPreviewSurface", meta = ( AllowedClasses = "/Script/Engine.MaterialInterface" ) )
+	FSoftObjectPath BasePreviewSurfaceTwoSidedMaterial = FSoftObjectPath{ TEXT( "/USDImporter/Materials/UsdPreviewSurfaceTwoSided.UsdPreviewSurfaceTwoSided" ) };
+
+	/**
+	 * What material to use as base material when creating material instances from USD materials.
+	 * You can swap these with your own materials, but make sure that the replacement materials have parameters with
+	 * the same names and types as the ones provided by the default material, otherwise the instances will not have
+	 * the parameters filled with values extracted from the USD material when parsing.
+	 */
+	UPROPERTY( config, EditAnywhere, Category = "USD|Base Materials|UsdPreviewSurface", meta = ( AllowedClasses = "/Script/Engine.MaterialInterface" ) )
+	FSoftObjectPath BasePreviewSurfaceTranslucentTwoSidedMaterial = FSoftObjectPath{ TEXT( "/USDImporter/Materials/UsdPreviewSurfaceTranslucentTwoSided.UsdPreviewSurfaceTranslucentTwoSided" ) };
+
+	/**
+	 * What material to use as base material when creating material instances from USD materials.
+	 * You can swap these with your own materials, but make sure that the replacement materials have parameters with
+	 * the same names and types as the ones provided by the default material, otherwise the instances will not have
+	 * the parameters filled with values extracted from the USD material when parsing.
+	 */
+	UPROPERTY( config, EditAnywhere, Category = "USD|Base Materials|UsdPreviewSurface with Virtual Textures", meta = ( AllowedClasses = "/Script/Engine.MaterialInterface" ) )
+	FSoftObjectPath BasePreviewSurfaceVTMaterial = FSoftObjectPath{ TEXT( "/USDImporter/Materials/UsdPreviewSurfaceVT.UsdPreviewSurfaceVT" ) };
+
+	/**
+	 * What material to use as base material when creating material instances from USD materials.
+	 * You can swap these with your own materials, but make sure that the replacement materials have parameters with
+	 * the same names and types as the ones provided by the default material, otherwise the instances will not have
+	 * the parameters filled with values extracted from the USD material when parsing.
+	 */
+	UPROPERTY( config, EditAnywhere, Category = "USD|Base Materials|UsdPreviewSurface with Virtual Textures", meta = ( AllowedClasses = "/Script/Engine.MaterialInterface" ) )
+	FSoftObjectPath BasePreviewSurfaceTranslucentVTMaterial = FSoftObjectPath{ TEXT( "/USDImporter/Materials/UsdPreviewSurfaceTranslucentVT.UsdPreviewSurfaceTranslucentVT" ) };
+
+	/**
+	 * What material to use as base material when creating material instances from USD materials.
+	 * You can swap these with your own materials, but make sure that the replacement materials have parameters with
+	 * the same names and types as the ones provided by the default material, otherwise the instances will not have
+	 * the parameters filled with values extracted from the USD material when parsing.
+	 */
+	UPROPERTY( config, EditAnywhere, Category = "USD|Base Materials|UsdPreviewSurface with Virtual Textures", meta = ( AllowedClasses = "/Script/Engine.MaterialInterface" ) )
+	FSoftObjectPath BasePreviewSurfaceTwoSidedVTMaterial = FSoftObjectPath{ TEXT( "/USDImporter/Materials/UsdPreviewSurfaceTwoSidedVT.UsdPreviewSurfaceTwoSidedVT" ) };
+
+	/**
+	 * What material to use as base material when creating material instances from USD materials.
+	 * You can swap these with your own materials, but make sure that the replacement materials have parameters with
+	 * the same names and types as the ones provided by the default material, otherwise the instances will not have
+	 * the parameters filled with values extracted from the USD material when parsing.
+	 */
+	UPROPERTY( config, EditAnywhere, Category = "USD|Base Materials|UsdPreviewSurface with Virtual Textures", meta = ( AllowedClasses = "/Script/Engine.MaterialInterface" ) )
+	FSoftObjectPath BasePreviewSurfaceTranslucentTwoSidedVTMaterial = FSoftObjectPath{ TEXT( "/USDImporter/Materials/UsdPreviewSurfaceTranslucentTwoSidedVT.UsdPreviewSurfaceTranslucentTwoSidedVT" ) };
+
+	/**
+	 * What material to use as base material when creating material instances from USD materials.
+	 * You can swap these with your own materials, but make sure that the replacement materials have parameters with
+	 * the same names and types as the ones provided by the default material, otherwise the instances will not have
+	 * the parameters filled with values extracted from the USD material when parsing.
+	 */
+	UPROPERTY( config, EditAnywhere, Category = "USD|Base Materials|DisplayColor", meta = ( AllowedClasses = "/Script/Engine.MaterialInterface" ) )
+	FSoftObjectPath BaseDisplayColorMaterial = FSoftObjectPath{ TEXT( "/USDImporter/Materials/DisplayColor.DisplayColor" ) };
+
+	/**
+	 * What material to use as base material when creating material instances from USD materials.
+	 * You can swap these with your own materials, but make sure that the replacement materials have parameters with
+	 * the same names and types as the ones provided by the default material, otherwise the instances will not have
+	 * the parameters filled with values extracted from the USD material when parsing.
+	 */
+	UPROPERTY( config, EditAnywhere, Category = "USD|Base Materials|DisplayColor", meta = ( AllowedClasses = "/Script/Engine.MaterialInterface" ) )
+	FSoftObjectPath BaseDisplayColorAndOpacityMaterial = FSoftObjectPath{ TEXT( "/USDImporter/Materials/DisplayColorAndOpacity.DisplayColorAndOpacity" ) };
+
+	/**
+	 * What material to use as base material when creating material instances from USD materials.
+	 * You can swap these with your own materials, but make sure that the replacement materials have parameters with
+	 * the same names and types as the ones provided by the default material, otherwise the instances will not have
+	 * the parameters filled with values extracted from the USD material when parsing.
+	 */
+	UPROPERTY( config, EditAnywhere, Category = "USD|Base Materials|DisplayColor", meta = ( AllowedClasses = "/Script/Engine.MaterialInterface" ) )
+	FSoftObjectPath BaseDisplayColorTwoSidedMaterial = FSoftObjectPath{ TEXT( "/USDImporter/Materials/DisplayColorTwoSided.DisplayColorTwoSided" ) };
+
+	/**
+	 * What material to use as base material when creating material instances from USD materials.
+	 * You can swap these with your own materials, but make sure that the replacement materials have parameters with
+	 * the same names and types as the ones provided by the default material, otherwise the instances will not have
+	 * the parameters filled with values extracted from the USD material when parsing.
+	 */
+	UPROPERTY( config, EditAnywhere, Category = "USD|Base Materials|DisplayColor", meta = ( AllowedClasses = "/Script/Engine.MaterialInterface" ) )
+	FSoftObjectPath BaseDisplayColorAndOpacityTwoSidedMaterial = FSoftObjectPath{ TEXT( "/USDImporter/Materials/DisplayColorAndOpacityTwoSided.DisplayColorAndOpacityTwoSided" ) };
 };
\ No newline at end of file
diff --git a/Engine/Plugins/Importers/USDImporter/Source/USDSchemas/Private/MeshTranslationImpl.cpp b/Engine/Plugins/Importers/USDImporter/Source/USDSchemas/Private/MeshTranslationImpl.cpp
--- a/Engine/Plugins/Importers/USDImporter/Source/USDSchemas/Private/MeshTranslationImpl.cpp
+++ b/Engine/Plugins/Importers/USDImporter/Source/USDSchemas/Private/MeshTranslationImpl.cpp
@@ -8,6 +8,7 @@
 #include "USDGeomMeshConversion.h"
 #include "USDLog.h"
 #include "USDMemory.h"
+#include "USDProjectSettings.h"
 #include "USDTypesConversion.h"
 
 #include "Components/MeshComponent.h"
@@ -15,6 +16,7 @@
 #include "Materials/MaterialInstanceConstant.h"
 #include "Materials/MaterialInstanceDynamic.h"
 #include "Materials/MaterialInterface.h"
+#include "Misc/Paths.h"
 
 #if USE_USD_SDK
 
@@ -27,21 +29,6 @@
 	#include "pxr/usd/usdShade/material.h"
 #include "USDIncludesEnd.h"
 
-namespace UE::MeshTranslationImplInternal::Private
-{
-	const static FString BaseMaterialPath = TEXT( "/USDImporter/Materials/UsdPreviewSurface.UsdPreviewSurface" );
-	const static FString BaseMaterialPathVT = TEXT( "/USDImporter/Materials/UsdPreviewSurfaceVT.UsdPreviewSurfaceVT" );
-
-	const static FString BaseMaterialPathTranslucent = TEXT( "/USDImporter/Materials/UsdPreviewSurfaceTranslucent.UsdPreviewSurfaceTranslucent" );
-	const static FString BaseMaterialPathTranslucentVT = TEXT( "/USDImporter/Materials/UsdPreviewSurfaceTranslucentVT.UsdPreviewSurfaceTranslucentVT" );
-
-	const static FString BaseMaterialPathTwoSided = TEXT( "/USDImporter/Materials/UsdPreviewSurfaceTwoSided.UsdPreviewSurfaceTwoSided" );
-	const static FString BaseMaterialPathTwoSidedVT = TEXT( "/USDImporter/Materials/UsdPreviewSurfaceTwoSidedVT.UsdPreviewSurfaceTwoSidedVT" );
-
-	const static FString BaseMaterialPathTranslucentTwoSided = TEXT( "/USDImporter/Materials/UsdPreviewSurfaceTranslucentTwoSided.UsdPreviewSurfaceTranslucentTwoSided" );
-	const static FString BaseMaterialPathTranslucentTwoSidedVT = TEXT( "/USDImporter/Materials/UsdPreviewSurfaceTranslucentTwoSidedVT.UsdPreviewSurfaceTranslucentTwoSidedVT" );
-}
-
 TMap<const UsdUtils::FUsdPrimMaterialSlot*, UMaterialInterface*> MeshTranslationImpl::ResolveMaterialAssignmentInfo(
 	const pxr::UsdPrim& UsdPrim,
 	const TArray<UsdUtils::FUsdPrimMaterialAssignmentInfo>& AssignmentInfo,
@@ -345,35 +332,39 @@ void MeshTranslationImpl::SetMaterialOverrides(
 
 UMaterialInterface* MeshTranslationImpl::GetBasePreviewSurfaceMaterial( EUsdBaseMaterialProperties BaseMaterialProperties )
 {
-	using namespace UE::MeshTranslationImplInternal::Private;
+	const UUsdProjectSettings* Settings = GetDefault<UUsdProjectSettings>();
+	if ( !Settings )
+	{
+		return nullptr;
+	}
 
 	const bool bIsTranslucent = EnumHasAnyFlags( BaseMaterialProperties, EUsdBaseMaterialProperties::Translucent );
 	const bool bIsVT = EnumHasAnyFlags( BaseMaterialProperties, EUsdBaseMaterialProperties::VT );
 	const bool bIsTwoSided = EnumHasAnyFlags( BaseMaterialProperties, EUsdBaseMaterialProperties::TwoSided );
 
-	const FString* TargetMaterialPath = nullptr;
+	const FSoftObjectPath* TargetMaterialPath = nullptr;
 	if ( bIsTranslucent )
 	{
 		if ( bIsVT )
 		{
 			if ( bIsTwoSided )
 			{
-				TargetMaterialPath = &BaseMaterialPathTranslucentTwoSidedVT;
+				TargetMaterialPath = &Settings->BasePreviewSurfaceTranslucentTwoSidedVTMaterial;
 			}
 			else
 			{
-				TargetMaterialPath = &BaseMaterialPathTranslucentVT;
+				TargetMaterialPath = &Settings->BasePreviewSurfaceTranslucentVTMaterial;
 			}
 		}
 		else
 		{
 			if ( bIsTwoSided )
 			{
-				TargetMaterialPath = &BaseMaterialPathTranslucentTwoSided;
+				TargetMaterialPath = &Settings->BasePreviewSurfaceTranslucentTwoSidedMaterial;
 			}
 			else
 			{
-				TargetMaterialPath = &BaseMaterialPathTranslucent;
+				TargetMaterialPath = &Settings->BasePreviewSurfaceTranslucentMaterial;
 			}
 		}
 	}
@@ -383,22 +374,22 @@ UMaterialInterface* MeshTranslationImpl::GetBasePreviewSurfaceMaterial( EUsdBase
 		{
 			if ( bIsTwoSided )
 			{
-				TargetMaterialPath = &BaseMaterialPathTwoSidedVT;
+				TargetMaterialPath = &Settings->BasePreviewSurfaceTwoSidedVTMaterial;
 			}
 			else
 			{
-				TargetMaterialPath = &BaseMaterialPathVT;
+				TargetMaterialPath = &Settings->BasePreviewSurfaceVTMaterial;
 			}
 		}
 		else
 		{
 			if ( bIsTwoSided )
 			{
-				TargetMaterialPath = &BaseMaterialPathTwoSided;
+				TargetMaterialPath = &Settings->BasePreviewSurfaceTwoSidedMaterial;
 			}
 			else
 			{
-				TargetMaterialPath = &BaseMaterialPath;
+				TargetMaterialPath = &Settings->BasePreviewSurfaceMaterial;
 			}
 		}
 	}
@@ -408,38 +399,42 @@ UMaterialInterface* MeshTranslationImpl::GetBasePreviewSurfaceMaterial( EUsdBase
 		return nullptr;
 	}
 
-	return Cast< UMaterialInterface >( FSoftObjectPath( *TargetMaterialPath ).TryLoad() );
+	return Cast< UMaterialInterface >( TargetMaterialPath->TryLoad() );
 }
 
 UMaterialInterface* MeshTranslationImpl::GetVTVersionOfBasePreviewSurfaceMaterial( UMaterialInterface* BaseMaterial )
 {
-	using namespace UE::MeshTranslationImplInternal::Private;
-
 	if ( !BaseMaterial )
 	{
 		return nullptr;
 	}
 
-	const FString PathName = BaseMaterial->GetPathName();
-	if ( PathName.Contains( TEXT( "VT" ), ESearchCase::CaseSensitive, ESearchDir::FromEnd ) )
+	const UUsdProjectSettings* Settings = GetDefault<UUsdProjectSettings>();
+	if ( !Settings )
+	{
+		return nullptr;
+	}
+
+	const FSoftObjectPath PathName = BaseMaterial->GetPathName();
+	if ( PathName.ToString().Contains(TEXT("VT"), ESearchCase::CaseSensitive, ESearchDir::FromEnd) )
 	{
 		return BaseMaterial;
 	}
-	else if ( PathName == BaseMaterialPath )
+	else if ( PathName == Settings->BasePreviewSurfaceMaterial )
 	{
-		return Cast< UMaterialInterface >( FSoftObjectPath{ BaseMaterialPathVT }.TryLoad() );
+		return Cast< UMaterialInterface >( Settings->BasePreviewSurfaceVTMaterial.TryLoad() );
 	}
-	else if ( PathName == BaseMaterialPathTwoSided )
+	else if ( PathName == Settings->BasePreviewSurfaceTwoSidedMaterial )
 	{
-		return Cast< UMaterialInterface >( FSoftObjectPath{ BaseMaterialPathTwoSidedVT }.TryLoad() );
+		return Cast< UMaterialInterface >( Settings->BasePreviewSurfaceTwoSidedVTMaterial.TryLoad() );
 	}
-	else if ( PathName == BaseMaterialPathTranslucent )
+	else if ( PathName == Settings->BasePreviewSurfaceTranslucentMaterial )
 	{
-		return Cast< UMaterialInterface >( FSoftObjectPath{ BaseMaterialPathTranslucentVT }.TryLoad() );
+		return Cast< UMaterialInterface >( Settings->BasePreviewSurfaceTranslucentVTMaterial.TryLoad() );
 	}
-	else if ( PathName == BaseMaterialPathTranslucentTwoSided )
+	else if ( PathName == Settings->BasePreviewSurfaceTranslucentTwoSidedMaterial )
 	{
-		return Cast< UMaterialInterface >( FSoftObjectPath{ BaseMaterialPathTranslucentTwoSidedVT }.TryLoad() );
+		return Cast< UMaterialInterface >( Settings->BasePreviewSurfaceTranslucentTwoSidedVTMaterial.TryLoad() );
 	}
 
 	// We should only ever call this function with a BaseMaterial that matches one of the above paths
@@ -449,33 +444,37 @@ UMaterialInterface* MeshTranslationImpl::GetVTVersionOfBasePreviewSurfaceMateria
 
 UMaterialInterface* MeshTranslationImpl::GetTwoSidedVersionOfBasePreviewSurfaceMaterial( UMaterialInterface* BaseMaterial )
 {
-	using namespace UE::MeshTranslationImplInternal::Private;
-
 	if ( !BaseMaterial )
 	{
 		return nullptr;
 	}
 
-	const FString PathName = BaseMaterial->GetPathName();
-	if ( PathName.Contains( TEXT( "TwoSided" ), ESearchCase::CaseSensitive, ESearchDir::FromEnd ) )
+	const UUsdProjectSettings* Settings = GetDefault<UUsdProjectSettings>();
+	if ( !Settings )
+	{
+		return nullptr;
+	}
+
+	const FSoftObjectPath PathName = BaseMaterial->GetPathName();
+	if ( PathName.ToString().Contains(TEXT("TwoSided"), ESearchCase::CaseSensitive, ESearchDir::FromEnd) )
 	{
 		return BaseMaterial;
 	}
-	else if ( PathName == BaseMaterialPath )
+	else if ( PathName == Settings->BasePreviewSurfaceMaterial )
 	{
-		return Cast< UMaterialInterface >( FSoftObjectPath{ BaseMaterialPathTwoSided }.TryLoad() );
+		return Cast< UMaterialInterface >( Settings->BasePreviewSurfaceTwoSidedMaterial.TryLoad() );
 	}
-	else if ( PathName == BaseMaterialPathTranslucent )
+	else if ( PathName == Settings->BasePreviewSurfaceTranslucentMaterial )
 	{
-		return Cast< UMaterialInterface >( FSoftObjectPath{ BaseMaterialPathTranslucentTwoSided }.TryLoad() );
+		return Cast< UMaterialInterface >( Settings->BasePreviewSurfaceTranslucentTwoSidedMaterial.TryLoad() );
 	}
-	else if ( PathName == BaseMaterialPathVT )
+	else if ( PathName == Settings->BasePreviewSurfaceVTMaterial )
 	{
-		return Cast< UMaterialInterface >( FSoftObjectPath{ BaseMaterialPathTwoSidedVT }.TryLoad() );
+		return Cast< UMaterialInterface >( Settings->BasePreviewSurfaceTwoSidedVTMaterial.TryLoad() );
 	}
-	else if ( PathName == BaseMaterialPathTranslucentVT )
+	else if ( PathName == Settings->BasePreviewSurfaceTranslucentVTMaterial )
 	{
-		return Cast< UMaterialInterface >( FSoftObjectPath{ BaseMaterialPathTranslucentTwoSidedVT }.TryLoad() );
+		return Cast< UMaterialInterface >( Settings->BasePreviewSurfaceTranslucentTwoSidedVTMaterial.TryLoad() );
 	}
 
 	// We should only ever call this function with a BaseMaterial that matches one of the above paths
diff --git a/Engine/Plugins/Importers/USDImporter/Source/USDStage/Private/USDStageModule.cpp b/Engine/Plugins/Importers/USDImporter/Source/USDStage/Private/USDStageModule.cpp
--- a/Engine/Plugins/Importers/USDImporter/Source/USDStage/Private/USDStageModule.cpp
+++ b/Engine/Plugins/Importers/USDImporter/Source/USDStage/Private/USDStageModule.cpp
@@ -24,38 +24,6 @@ public:
 #if WITH_EDITOR
 		LLM_SCOPE_BYTAG(Usd);
 
-		// If don't have any active references to our materials they won't be packaged into monolithic builds, and we wouldn't
-		// be able to create dynamic material instances at runtime.
-		TSharedPtr<IPlugin> Plugin = IPluginManager::Get().FindPlugin( TEXT( "USDImporter" ) );
-		if ( Plugin.IsValid() )
-		{
-			FString MaterialsPath = TEXT( "/USDImporter/Materials" );
-
-			UProjectPackagingSettings* PackagingSettings = Cast<UProjectPackagingSettings>( UProjectPackagingSettings::StaticClass()->GetDefaultObject() );
-			if ( PackagingSettings )
-			{
-				bool bAlreadyInPath = false;
-
-				TArray<FDirectoryPath>& DirectoriesToCook = PackagingSettings->DirectoriesToAlwaysCook;
-				for ( int32 Index = DirectoriesToCook.Num() - 1; Index >= 0; --Index )
-				{
-					if ( FPaths::IsSamePath( DirectoriesToCook[ Index ].Path, MaterialsPath ) )
-					{
-						bAlreadyInPath = true;
-						break;
-					}
-				}
-
-				if ( !bAlreadyInPath )
-				{
-					FDirectoryPath MaterialsDirectory;
-					MaterialsDirectory.Path = MaterialsPath;
-
-					PackagingSettings->DirectoriesToAlwaysCook.Add( MaterialsDirectory );
-				}
-			}
-		}
-
 		FPropertyEditorModule& PropertyModule = FModuleManager::LoadModuleChecked<FPropertyEditorModule>( TEXT( "PropertyEditor" ) );
 		PropertyModule.RegisterCustomClassLayout( TEXT( "UsdStageActor" ), FOnGetDetailCustomizationInstance::CreateStatic( &FUsdStageActorCustomization::MakeInstance ) );
 #endif // WITH_EDITOR
@@ -64,25 +32,6 @@ public:
 	virtual void ShutdownModule() override
 	{
 #if WITH_EDITOR
-		TSharedPtr<IPlugin> Plugin = IPluginManager::Get().FindPlugin( TEXT( "USDImporter" ) );
-		if ( Plugin.IsValid() && !IsEngineExitRequested() ) // If we're shutting down because the engine is exiting then this may crash, and there's no point anyway
-		{
-			FString Path = TEXT( "/USDImporter/Materials" );
-
-			UProjectPackagingSettings* PackagingSettings = Cast<UProjectPackagingSettings>( UProjectPackagingSettings::StaticClass()->GetDefaultObject() );
-			if ( PackagingSettings )
-			{
-				TArray<FDirectoryPath>& DirectoriesToCook = PackagingSettings->DirectoriesToAlwaysCook;
-				for ( int32 Index = DirectoriesToCook.Num() - 1; Index >= 0; --Index )
-				{
-					if ( FPaths::IsSamePath( DirectoriesToCook[ Index ].Path, Path ) )
-					{
-						DirectoriesToCook.RemoveAt(Index);
-					}
-				}
-			}
-		}
-
 		FPropertyEditorModule& PropertyModule = FModuleManager::LoadModuleChecked< FPropertyEditorModule >( TEXT( "PropertyEditor" ) );
 		PropertyModule.UnregisterCustomClassLayout( TEXT( "UsdStageActor" ) );
 #endif // WITH_EDITOR
diff --git a/Engine/Plugins/Importers/USDImporter/Source/USDUtilities/Private/USDGeomMeshConversion.cpp b/Engine/Plugins/Importers/USDImporter/Source/USDUtilities/Private/USDGeomMeshConversion.cpp
--- a/Engine/Plugins/Importers/USDImporter/Source/USDUtilities/Private/USDGeomMeshConversion.cpp
+++ b/Engine/Plugins/Importers/USDImporter/Source/USDUtilities/Private/USDGeomMeshConversion.cpp
@@ -12,6 +12,7 @@
 #include "USDLog.h"
 #include "USDMemory.h"
 #include "USDPrimConversion.h"
+#include "USDProjectSettings.h"
 #include "USDShadeConversion.h"
 #include "USDTypesConversion.h"
 
@@ -1380,31 +1381,41 @@ bool UsdToUnreal::ConvertGeomMeshHierarchy(
 
 UMaterialInstanceDynamic* UsdUtils::CreateDisplayColorMaterialInstanceDynamic( const UsdUtils::FDisplayColorMaterial& DisplayColorDescription )
 {
-	FString ParentPath;
+	const UUsdProjectSettings* Settings = GetDefault<UUsdProjectSettings>();
+	if ( !Settings )
+	{
+		return nullptr;
+	}
+
+	const FSoftObjectPath* ParentPathPtr = nullptr;
 	if ( DisplayColorDescription.bHasOpacity )
 	{
 		if ( DisplayColorDescription.bIsDoubleSided )
 		{
-			ParentPath = TEXT( "Material'/USDImporter/Materials/DisplayColorAndOpacityDoubleSided.DisplayColorAndOpacityDoubleSided'" );
+			ParentPathPtr = &Settings->BaseDisplayColorAndOpacityTwoSidedMaterial;
 		}
 		else
 		{
-			ParentPath = TEXT( "Material'/USDImporter/Materials/DisplayColorAndOpacity.DisplayColorAndOpacity'" );
+			ParentPathPtr = &Settings->BaseDisplayColorAndOpacityMaterial;
 		}
 	}
 	else
 	{
 		if ( DisplayColorDescription.bIsDoubleSided )
 		{
-			ParentPath = TEXT( "Material'/USDImporter/Materials/DisplayColorDoubleSided.DisplayColorDoubleSided'" );
+			ParentPathPtr = &Settings->BaseDisplayColorTwoSidedMaterial;
 		}
 		else
 		{
-			ParentPath = TEXT( "Material'/USDImporter/Materials/DisplayColor.DisplayColor'" );
+			ParentPathPtr = &Settings->BaseDisplayColorMaterial;
 		}
 	}
+	if ( !ParentPathPtr )
+	{
+		return nullptr;
+	}
 
-	if ( UMaterialInterface* ParentMaterial = Cast< UMaterialInterface >( FSoftObjectPath( ParentPath ).TryLoad() ) )
+	if ( UMaterialInterface* ParentMaterial = Cast< UMaterialInterface >( ParentPathPtr->TryLoad() ) )
 	{
 		FName AssetName = MakeUniqueObjectName(
 			GetTransientPackage(),
@@ -1427,31 +1438,37 @@ UMaterialInstanceDynamic* UsdUtils::CreateDisplayColorMaterialInstanceDynamic( c
 UMaterialInstanceConstant* UsdUtils::CreateDisplayColorMaterialInstanceConstant( const FDisplayColorMaterial& DisplayColorDescription )
 {
 #if WITH_EDITOR
-	FString ParentPath;
+	const UUsdProjectSettings* Settings = GetDefault<UUsdProjectSettings>();
+	if ( !Settings )
+	{
+		return nullptr;
+	}
+
+	const FSoftObjectPath* ParentPathPtr = nullptr;
 	if ( DisplayColorDescription.bHasOpacity )
 	{
 		if ( DisplayColorDescription.bIsDoubleSided )
 		{
-			ParentPath = TEXT( "Material'/USDImporter/Materials/DisplayColorAndOpacityDoubleSided.DisplayColorAndOpacityDoubleSided'" );
+			ParentPathPtr = &Settings->BaseDisplayColorAndOpacityTwoSidedMaterial;
 		}
 		else
 		{
-			ParentPath = TEXT( "Material'/USDImporter/Materials/DisplayColorAndOpacity.DisplayColorAndOpacity'" );
+			ParentPathPtr = &Settings->BaseDisplayColorAndOpacityMaterial;
 		}
 	}
 	else
 	{
 		if ( DisplayColorDescription.bIsDoubleSided )
 		{
-			ParentPath = TEXT( "Material'/USDImporter/Materials/DisplayColorDoubleSided.DisplayColorDoubleSided'" );
+			ParentPathPtr = &Settings->BaseDisplayColorTwoSidedMaterial;
 		}
 		else
 		{
-			ParentPath = TEXT( "Material'/USDImporter/Materials/DisplayColor.DisplayColor'" );
+			ParentPathPtr = &Settings->BaseDisplayColorMaterial;
 		}
 	}
 
-	if ( UMaterialInterface* ParentMaterial = Cast< UMaterialInterface >( FSoftObjectPath( ParentPath ).TryLoad() ) )
+	if ( UMaterialInterface* ParentMaterial = Cast< UMaterialInterface >( ParentPathPtr->TryLoad() ) )
 	{
 		FName AssetName = MakeUniqueObjectName(
 			GetTransientPackage(),
