USD: Show an error toast in case an export would have been impossible due to the layer already being open.

diff --git a/Engine/Plugins/Importers/USDImporter/Source/USDExporter/Private/AnimSequenceExporterUSD.cpp b/Engine/Plugins/Importers/USDImporter/Source/USDExporter/Private/AnimSequenceExporterUSD.cpp
--- a/Engine/Plugins/Importers/USDImporter/Source/USDExporter/Private/AnimSequenceExporterUSD.cpp
+++ b/Engine/Plugins/Importers/USDImporter/Source/USDExporter/Private/AnimSequenceExporterUSD.cpp
@@ -11,6 +11,7 @@
 #include "UnrealUSDWrapper.h"
 #include "USDClassesModule.h"
 #include "USDErrorUtils.h"
+#include "USDExporterModule.h"
 #include "USDLayerUtils.h"
 #include "USDLog.h"
 #include "USDOptionsWindow.h"
@@ -123,6 +124,11 @@ bool UAnimSequenceExporterUSD::ExportBinary( UObject* Object, const TCHAR* Type,
 		return false;
 	}
 
+	if ( !IUsdExporterModule::CanExportToLayer( UExporter::CurrentFilename ) )
+	{
+		return false;
+	}
+
 	// See comment on the analogous line within StaticMeshExporterUSD.cpp
 	ExportTask->bPrompt = false;
 
@@ -362,4 +368,4 @@ bool UAnimSequenceExporterUSD::ExportBinary( UObject* Object, const TCHAR* Type,
 #else
 	return false;
 #endif // #if USE_USD_SDK
-}
\ No newline at end of file
+}
diff --git a/Engine/Plugins/Importers/USDImporter/Source/USDExporter/Private/LevelSequenceExporterUSD.cpp b/Engine/Plugins/Importers/USDImporter/Source/USDExporter/Private/LevelSequenceExporterUSD.cpp
--- a/Engine/Plugins/Importers/USDImporter/Source/USDExporter/Private/LevelSequenceExporterUSD.cpp
+++ b/Engine/Plugins/Importers/USDImporter/Source/USDExporter/Private/LevelSequenceExporterUSD.cpp
@@ -9,6 +9,7 @@
 #include "USDConversionUtils.h"
 #include "USDLayerUtils.h"
 #include "USDLog.h"
+#include "USDExporterModule.h"
 #include "USDOptionsWindow.h"
 #include "USDPrimConversion.h"
 #include "USDUnrealAssetInfo.h"
@@ -968,6 +969,11 @@ namespace UE
 					}
 				}
 
+				if ( !IUsdExporterModule::CanExportToLayer( UniqueFilePath ) )
+				{
+					return;
+				}
+
 				FString LevelSequenceVersion = MovieSceneSequence.GetSignature().ToString();
 				{
 					// We could just use the GUID directly but all other asset types end up with SHA hash size so lets be
diff --git a/Engine/Plugins/Importers/USDImporter/Source/USDExporter/Private/MaterialExporterUSD.cpp b/Engine/Plugins/Importers/USDImporter/Source/USDExporter/Private/MaterialExporterUSD.cpp
--- a/Engine/Plugins/Importers/USDImporter/Source/USDExporter/Private/MaterialExporterUSD.cpp
+++ b/Engine/Plugins/Importers/USDImporter/Source/USDExporter/Private/MaterialExporterUSD.cpp
@@ -8,6 +8,7 @@
 #include "UnrealUSDWrapper.h"
 #include "USDClassesModule.h"
 #include "USDConversionUtils.h"
+#include "USDExporterModule.h"
 #include "USDGeomMeshConversion.h"
 #include "USDLog.h"
 #include "USDMemory.h"
@@ -287,6 +288,11 @@ bool UMaterialExporterUsd::ExportMaterial(
 	SHA1.GetHash( &Hash.Hash[ 0 ] );
 	FString MaterialHashString = Hash.ToString();
 
+	if ( !IUsdExporterModule::CanExportToLayer( FilePath.FilePath ) )
+	{
+		return false;
+	}
+
 	// Check if we already have exported what we plan on exporting anyway
 	if ( FPaths::FileExists( FilePath.FilePath ) )
 	{
diff --git a/Engine/Plugins/Importers/USDImporter/Source/USDExporter/Private/SkeletalMeshExporterUSD.cpp b/Engine/Plugins/Importers/USDImporter/Source/USDExporter/Private/SkeletalMeshExporterUSD.cpp
--- a/Engine/Plugins/Importers/USDImporter/Source/USDExporter/Private/SkeletalMeshExporterUSD.cpp
+++ b/Engine/Plugins/Importers/USDImporter/Source/USDExporter/Private/SkeletalMeshExporterUSD.cpp
@@ -11,6 +11,7 @@
 #include "SkeletalMeshExporterUSDOptions.h"
 #include "USDClassesModule.h"
 #include "USDConversionUtils.h"
+#include "USDExporterModule.h"
 #include "USDLog.h"
 #include "USDMemory.h"
 #include "USDOptionsWindow.h"
@@ -139,6 +140,12 @@ bool USkeletalMeshExporterUsd::ExportBinary( UObject* Object, const TCHAR* Type,
 		PayloadFilename = FPaths::Combine( PathPart, FilenamePart + TEXT( "_payload." ) + ExtensionPart );
 	}
 
+	if ( !IUsdExporterModule::CanExportToLayer( UExporter::CurrentFilename ) ||
+		( Options->MeshAssetOptions.bUsePayload && !IUsdExporterModule::CanExportToLayer( PayloadFilename ) ) )
+	{
+		return false;
+	}
+
 	// Get a simple GUID hash/identifier of our mesh
 	FString DDCKeyHash;
 	if( FSkeletalMeshRenderData* RenderData = SkeletalMesh->GetResourceForRendering() )
diff --git a/Engine/Plugins/Importers/USDImporter/Source/USDExporter/Private/StaticMeshExporterUSD.cpp b/Engine/Plugins/Importers/USDImporter/Source/USDExporter/Private/StaticMeshExporterUSD.cpp
--- a/Engine/Plugins/Importers/USDImporter/Source/USDExporter/Private/StaticMeshExporterUSD.cpp
+++ b/Engine/Plugins/Importers/USDImporter/Source/USDExporter/Private/StaticMeshExporterUSD.cpp
@@ -9,6 +9,7 @@
 #include "StaticMeshResources.h"
 #include "USDClassesModule.h"
 #include "USDConversionUtils.h"
+#include "USDExporterModule.h"
 #include "USDGeomMeshConversion.h"
 #include "USDLog.h"
 #include "USDMemory.h"
@@ -152,6 +153,12 @@ bool UStaticMeshExporterUsd::ExportBinary( UObject* Object, const TCHAR* Type, F
 		PayloadFilename = FPaths::Combine( PathPart, FilenamePart + TEXT( "_payload." ) + ExtensionPart );
 	}
 
+	if ( !IUsdExporterModule::CanExportToLayer( UExporter::CurrentFilename ) ||
+		( Options->MeshAssetOptions.bUsePayload && !IUsdExporterModule::CanExportToLayer( PayloadFilename ) ) )
+	{
+		return false;
+	}
+
 	// Get a simple GUID hash/identifier of our mesh
 	FString DDCKeyHash;
 	if ( FStaticMeshRenderData* RenderData = StaticMesh->GetRenderData() )
diff --git a/Engine/Plugins/Importers/USDImporter/Source/USDExporter/Private/USDConversionBlueprintLibrary.cpp b/Engine/Plugins/Importers/USDImporter/Source/USDExporter/Private/USDConversionBlueprintLibrary.cpp
--- a/Engine/Plugins/Importers/USDImporter/Source/USDExporter/Private/USDConversionBlueprintLibrary.cpp
+++ b/Engine/Plugins/Importers/USDImporter/Source/USDExporter/Private/USDConversionBlueprintLibrary.cpp
@@ -330,6 +330,11 @@ FString UUsdConversionBlueprintLibrary::GenerateObjectVersionString( const UObje
 	return Hash.ToString();
 }
 
+bool UUsdConversionBlueprintLibrary::CanExportToLayer( const FString& TargetFilePath )
+{
+	return IUsdExporterModule::CanExportToLayer( TargetFilePath );
+}
+
 FString UUsdConversionBlueprintLibrary::MakePathRelativeToLayer( const FString& AnchorLayerPath, const FString& PathToMakeRelative )
 {
 #if USE_USD_SDK
diff --git a/Engine/Plugins/Importers/USDImporter/Source/USDExporter/Private/USDExporterModule.cpp b/Engine/Plugins/Importers/USDImporter/Source/USDExporter/Private/USDExporterModule.cpp
--- a/Engine/Plugins/Importers/USDImporter/Source/USDExporter/Private/USDExporterModule.cpp
+++ b/Engine/Plugins/Importers/USDImporter/Source/USDExporter/Private/USDExporterModule.cpp
@@ -4,13 +4,20 @@
 
 #include "LevelExporterUSDOptionsCustomization.h"
 #include "USDAssetOptions.h"
+#include "USDLog.h"
 #include "USDMemory.h"
 
+#include "UsdWrappers/SdfLayer.h"
+
 #include "Editor.h"
+#include "Framework/Notifications/NotificationManager.h"
 #include "Modules/ModuleManager.h"
 #include "PropertyEditorModule.h"
 #include "Selection.h"
 #include "UObject/ObjectSaveContext.h"
+#include "Widgets/Notifications/SNotificationList.h"
+
+#define LOCTEXT_NAMESPACE "USDExporterModule"
 
 class FUsdExporterModule : public IUsdExporterModule
 {
@@ -90,5 +97,31 @@ void IUsdExporterModule::HashEditorSelection( FSHA1& HashToUpdate )
 	}
 }
 
+bool IUsdExporterModule::CanExportToLayer( const FString& TargetFilePath )
+{
+	for ( const UE::FSdfLayerWeak& LoadedLayer : UE::FSdfLayerWeak::GetLoadedLayers() )
+	{
+		if ( LoadedLayer.GetIdentifier() == TargetFilePath )
+		{
+			FText ErrorMessage = FText::Format(
+				LOCTEXT( "FailedExportLayerOpenSubText", "Failed to export to the USD layer '{0}' as a layer with this identifier is already open in another USD Stage. Please use a different file path or try closing that other USD stage and exporting again." ),
+				FText::FromString( TargetFilePath ) );
+			UE_LOG( LogUsd, Error, TEXT( "%s" ), *ErrorMessage.ToString() );
+
+			FNotificationInfo ErrorToast( LOCTEXT( "FailedExportLayerOpenText", "USD: Export failure" ) );
+			ErrorToast.ExpireDuration = 10.0f;
+			ErrorToast.bFireAndForget = true;
+			ErrorToast.Image = FCoreStyle::Get().GetBrush( TEXT( "MessageLog.Error" ) );
+			ErrorToast.SubText = ErrorMessage;
+			FSlateNotificationManager::Get().AddNotification( ErrorToast );
+
+			return false;
+		}
+	}
+
+	return true;
+}
+
 IMPLEMENT_MODULE_USD( FUsdExporterModule, USDExporter );
 
+#undef LOCTEXT_NAMESPACE
\ No newline at end of file
diff --git a/Engine/Plugins/Importers/USDImporter/Source/USDExporter/Public/USDConversionBlueprintLibrary.h b/Engine/Plugins/Importers/USDImporter/Source/USDExporter/Public/USDConversionBlueprintLibrary.h
--- a/Engine/Plugins/Importers/USDImporter/Source/USDExporter/Public/USDConversionBlueprintLibrary.h
+++ b/Engine/Plugins/Importers/USDImporter/Source/USDExporter/Public/USDConversionBlueprintLibrary.h
@@ -71,6 +71,10 @@ public:
 	UFUNCTION( BlueprintCallable, Category = "World utils" )
 	static FString GenerateObjectVersionString( const UObject* ObjectToExport, UObject* ExportOptions );
 
+	/** Checks whether we can create a USD Layer with "TargetFilePath" as identifier and export to it */
+	UFUNCTION( BlueprintCallable, Category = "World utils" )
+	static bool CanExportToLayer( const FString& TargetFilePath );
+
 	UFUNCTION( BlueprintCallable, Category = "Layer utils" )
 	static FString MakePathRelativeToLayer( const FString& AnchorLayerPath, const FString& PathToMakeRelative );
 
diff --git a/Engine/Plugins/Importers/USDImporter/Source/USDExporter/Public/USDExporterModule.h b/Engine/Plugins/Importers/USDImporter/Source/USDExporter/Public/USDExporterModule.h
--- a/Engine/Plugins/Importers/USDImporter/Source/USDExporter/Public/USDExporterModule.h
+++ b/Engine/Plugins/Importers/USDImporter/Source/USDExporter/Public/USDExporterModule.h
@@ -9,4 +9,7 @@ class IUsdExporterModule : public IModuleInterface
 {
 public:
 	static void HashEditorSelection( FSHA1& HashToUpdate );
+
+	/** Checks whether we can create a USD Layer with "TargetFilePath" as identifier and export to it */
+	static bool CanExportToLayer( const FString& TargetFilePath );
 };
diff --git a/Engine/Plugins/Importers/USDImporter/Source/UnrealUSDWrapper/Private/UsdWrappers/SdfLayer.cpp b/Engine/Plugins/Importers/USDImporter/Source/UnrealUSDWrapper/Private/UsdWrappers/SdfLayer.cpp
--- a/Engine/Plugins/Importers/USDImporter/Source/UnrealUSDWrapper/Private/UsdWrappers/SdfLayer.cpp
+++ b/Engine/Plugins/Importers/USDImporter/Source/UnrealUSDWrapper/Private/UsdWrappers/SdfLayer.cpp
@@ -153,6 +153,21 @@ namespace UE
 		return ( bool ) Impl->GetInner();
 	}
 
+	uint32 GetTypeHash( const UE::FSdfLayerWeak& Layer )
+	{
+		uint32 Result = 0;
+#if USE_USD_SDK
+		FScopedUnrealAllocs UnrealAllocs;
+
+		// Using TfHash will defer back to TfHashAppend for the type, which should use the implementation for
+		// TfWeakPtrFacade, which is how TfWeakPtr seem to be hashed in USD.
+		// We have to clamp the size_t to uint32 here but deferring to USD is likely the best approach to guarantee
+		// that objects with non-colliding hashes in USD also not collide in UE.
+		Result = static_cast< uint32 >( pxr::TfHash()( Layer.Impl->GetInner() ) );
+#endif // #if USE_USD_SDK
+		return Result;
+	}
+
 #if USE_USD_SDK
 	template<typename PtrType>
 	FSdfLayerBase<PtrType>::FSdfLayerBase( const pxr::SdfLayerRefPtr& InSdfLayer )
@@ -275,6 +290,24 @@ namespace UE
 #endif // #if USE_USD_SDK
 	}
 
+	template<typename PtrType>
+	TSet<UE::FSdfLayerWeak> FSdfLayerBase<PtrType>::GetLoadedLayers()
+	{
+		TSet<UE::FSdfLayerWeak> Result;
+#if USE_USD_SDK
+		FScopedUsdAllocs UsdAllocs;
+
+		pxr::SdfLayerHandleSet LoadedSdfLayers = pxr::SdfLayer::GetLoadedLayers();
+		Result.Reserve( LoadedSdfLayers.size() );
+
+		for ( const pxr::SdfLayerHandle& LoadedSdfLayer : LoadedSdfLayers )
+		{
+			Result.Add( UE::FSdfLayerWeak{ LoadedSdfLayer } );
+		}
+#endif // #if USE_USD_SDK
+		return Result;
+	}
+
 	template<typename PtrType>
 	FSdfLayer FSdfLayerBase<PtrType>::FindOrOpen( const TCHAR* Identifier )
 	{
diff --git a/Engine/Plugins/Importers/USDImporter/Source/UnrealUSDWrapper/Public/UsdWrappers/SdfLayer.h b/Engine/Plugins/Importers/USDImporter/Source/UnrealUSDWrapper/Public/UsdWrappers/SdfLayer.h
--- a/Engine/Plugins/Importers/USDImporter/Source/UnrealUSDWrapper/Public/UsdWrappers/SdfLayer.h
+++ b/Engine/Plugins/Importers/USDImporter/Source/UnrealUSDWrapper/Public/UsdWrappers/SdfLayer.h
@@ -58,6 +58,8 @@ namespace UE
 
 		explicit operator bool() const;
 
+		friend UNREALUSDWRAPPER_API uint32 GetTypeHash( const FSdfLayerWeak& Layer );
+
 		// Auto conversion from/to PtrType. We use concrete pointer types here
 		// because we should also be able to convert between them
 	public:
@@ -88,6 +90,7 @@ namespace UE
 
 		void TransferContent( const FSdfLayer& SourceLayer );
 
+		static TSet<FSdfLayerWeak> GetLoadedLayers();
 		static FSdfLayer FindOrOpen( const TCHAR* Identifier );
 		static FSdfLayer CreateNew( const TCHAR* Identifier );
 
