USD: Restore option to import stage at a specific time code;

Allow importing GeometryCaches as StaticMeshes.

diff --git a/Engine/Plugins/Importers/USDImporter/Source/USDSchemas/Private/USDGeomMeshTranslator.cpp b/Engine/Plugins/Importers/USDImporter/Source/USDSchemas/Private/USDGeomMeshTranslator.cpp
--- a/Engine/Plugins/Importers/USDImporter/Source/USDSchemas/Private/USDGeomMeshTranslator.cpp
+++ b/Engine/Plugins/Importers/USDImporter/Source/USDSchemas/Private/USDGeomMeshTranslator.cpp
@@ -1487,7 +1487,7 @@ void FUsdGeomMeshTranslator::CreateAssets()
 	TRACE_CPUPROFILER_EVENT_SCOPE( FUsdGeomMeshTranslator::CreateAssets );
 
 #if WITH_EDITOR
-	if ( GUseGeometryCacheUSD && UsdGeomMeshTranslatorImpl::IsAnimated( GetPrim() ) )
+	if ( !Context->bIsImporting && GUseGeometryCacheUSD && UsdGeomMeshTranslatorImpl::IsAnimated( GetPrim() ) )
 	{
 		// Create the GeometryCache TaskChain
 		TSharedRef< FGeometryCacheCreateAssetsTaskChain > AssetsTaskChain = MakeShared< FGeometryCacheCreateAssetsTaskChain >( Context, PrimPath );
@@ -1509,7 +1509,7 @@ USceneComponent* FUsdGeomMeshTranslator::CreateComponents()
 
 #if WITH_EDITOR
 	// Force animated meshes as GeometryCache
-	if ( GUseGeometryCacheUSD && UsdGeomMeshTranslatorImpl::IsAnimated( GetPrim() ) )
+	if ( !Context->bIsImporting && GUseGeometryCacheUSD && UsdGeomMeshTranslatorImpl::IsAnimated( GetPrim() ) )
 	{
 		ComponentType = UGeometryCacheUsdComponent::StaticClass();
 	}
@@ -1547,6 +1547,9 @@ USceneComponent* FUsdGeomMeshTranslator::CreateComponents()
 	}
 	else if ( UGeometryCacheComponent* Component = Cast<UGeometryCacheComponent>( SceneComponent ) )
 	{
+		// We don't support importing geometry cache stuff yet so we should not be in here
+		ensure( !Context->bIsImporting );
+
 		if ( UGeometryCache* GeometryCache = Cast< UGeometryCache >( Context->AssetCache->GetAssetForPrim( PrimPathString ) ) )
 		{
 			// Geometry caches don't support LODs
@@ -1612,6 +1615,8 @@ void FUsdGeomMeshTranslator::UpdateComponents( USceneComponent* SceneComponent )
 	// Set the initial GeometryCache on the GeometryCacheUsdComponent
 	if ( UGeometryCacheUsdComponent* GeometryCacheUsdComponent = Cast< UGeometryCacheUsdComponent >( SceneComponent ) )
 	{
+		ensure( !Context->bIsImporting );
+
 		UGeometryCache* GeometryCache = Cast< UGeometryCache >( Context->AssetCache->GetAssetForPrim( PrimPath.GetString() ) );
 
 		bool bShouldRegister = false;
diff --git a/Engine/Plugins/Importers/USDImporter/Source/USDStageEditorViewModels/Private/USDStageViewModel.cpp b/Engine/Plugins/Importers/USDImporter/Source/USDStageEditorViewModels/Private/USDStageViewModel.cpp
--- a/Engine/Plugins/Importers/USDImporter/Source/USDStageEditorViewModels/Private/USDStageViewModel.cpp
+++ b/Engine/Plugins/Importers/USDImporter/Source/USDStageEditorViewModels/Private/USDStageViewModel.cpp
@@ -288,6 +288,7 @@ void FUsdStageViewModel::ImportStage()
 		ImportContext.ImportOptions->MaterialPurpose = StageActor->MaterialPurpose;
 		ImportContext.ImportOptions->StageOptions.MetersPerUnit = UsdUtils::GetUsdStageMetersPerUnit( UsdStage );
 		ImportContext.ImportOptions->StageOptions.UpAxis = UsdUtils::GetUsdStageUpAxisAsEnum( UsdStage );
+		ImportContext.ImportOptions->ImportTimeCode = StageActor->GetTime();
 		ImportContext.bReadFromStageCache = true; // So that we import whatever the user has open right now, even if the file has changes
 
 		const FString RootPath = UsdStage.GetRootLayer().GetRealPath();
diff --git a/Engine/Plugins/Importers/USDImporter/Source/USDStageImporter/Private/USDStageImportOptions.cpp b/Engine/Plugins/Importers/USDImporter/Source/USDStageImporter/Private/USDStageImportOptions.cpp
--- a/Engine/Plugins/Importers/USDImporter/Source/USDStageImporter/Private/USDStageImportOptions.cpp
+++ b/Engine/Plugins/Importers/USDImporter/Source/USDStageImporter/Private/USDStageImportOptions.cpp
@@ -26,6 +26,8 @@ UUsdStageImportOptions::UUsdStageImportOptions(const FObjectInitializer& ObjectI
 	bOverrideStageOptions = false;
 	StageOptions.MetersPerUnit = 0.01f;
 	StageOptions.UpAxis = EUsdUpAxis::ZAxis;
+	bImportAtSpecificTimeCode = false;
+	ImportTimeCode = 0.0f;
 
 	ExistingActorPolicy = EReplaceActorPolicy::Replace;
 	ExistingAssetPolicy = EReplaceAssetPolicy::Replace;
@@ -93,6 +95,11 @@ void UsdUtils::AddAnalyticsAttributes(
 	{
 		UsdUtils::AddAnalyticsAttributes( Options.StageOptions, InOutAttributes );
 	}
+	InOutAttributes.Emplace( TEXT( "ImportAtSpecificTimeCode" ), Options.bImportAtSpecificTimeCode );
+	if ( Options.bImportAtSpecificTimeCode )
+	{
+		InOutAttributes.Emplace( TEXT( "ImportTimeCode" ), LexToString( Options.ImportTimeCode ) );
+	}
 	InOutAttributes.Emplace( TEXT( "NumGroomInterpolationSettings" ), LexToString( Options.GroomInterpolationSettings.Num() ) );
 	InOutAttributes.Emplace( TEXT( "ReplaceActorPolicy" ), LexToString( ( uint8 ) Options.ExistingActorPolicy ) );
 	InOutAttributes.Emplace( TEXT( "ReplaceAssetPolicy" ), LexToString( ( uint8 ) Options.ExistingAssetPolicy ) );
diff --git a/Engine/Plugins/Importers/USDImporter/Source/USDStageImporter/Private/USDStageImporter.cpp b/Engine/Plugins/Importers/USDImporter/Source/USDStageImporter/Private/USDStageImporter.cpp
--- a/Engine/Plugins/Importers/USDImporter/Source/USDStageImporter/Private/USDStageImporter.cpp
+++ b/Engine/Plugins/Importers/USDImporter/Source/USDStageImporter/Private/USDStageImporter.cpp
@@ -1541,7 +1541,7 @@ namespace UsdStageImporterImpl
 	// Also, skeletal mesh components need to be manually ticked, or else they may be showing an animated state of an animation that
 	// we chose not to import, and wouldn't update otherwise until manually ticked by the user (or after save/reload), which may look
 	// like a bug
-	void RefreshComponents( AActor* RootSceneActor )
+	void RefreshComponents( AActor* RootSceneActor, bool bImportAtSpecificTimeCode )
 	{
 		if ( !RootSceneActor )
 		{
@@ -1566,6 +1566,30 @@ namespace UsdStageImporterImpl
 					SkeletalMeshComponent->MarkRenderTransformDirty();
 					SkeletalMeshComponent->MarkRenderDynamicDataDirty();
 				}
+				else if ( bImportAtSpecificTimeCode )
+				{
+					// The asset we return from the import factories may lead to
+					// USkeletalMesh::PostEditChangeProperty being called. The FMultiComponentReregisterContext in
+					// there will call USkeletalMeshComponent::InitAnim on all components that use the mesh, which will
+					// wipe out the Position value on SkeletalMeshComponents. This value describes the current state of
+					// animating the AnimSequence according to the import TimeCode, and we need to keep it.
+					// If we store the position via this function however, it will store it inside the position within
+					// the AnimationData member, so that when USkeletalMeshComponent::InitAnim calls
+					// USkeletalMeshComponent::InitializeAnimScriptInstance, it will use that AnimationData to
+					// initialize the new AnimScriptInstance and retain our desired Position value.
+					// This is also nice because it will serialize this animation position and restore it whenever the
+					// component updates, which should provide some persistent to this animation state (which we assume
+					// was desirable since the user imported at a particular time)
+					const bool bIsLooping = false;
+					const bool bIsPlaying = false;
+					const float Position = SkeletalMeshComponent->GetPosition();
+					SkeletalMeshComponent->OverrideAnimationData(
+						SkeletalMeshComponent->AnimationData.AnimToPlay,
+						bIsLooping,
+						bIsPlaying,
+						Position
+					);
+				}
 
 				// It does need us to manually set this to dirty regardless or else it won't update in case we changed material
 				// assignments
@@ -1638,7 +1662,9 @@ void UUsdStageImporter::ImportFromFile(FUsdStageImportContext& ImportContext)
 	TranslationContext->bIsImporting = true;
 	TranslationContext->Level = ImportContext.World->GetCurrentLevel();
 	TranslationContext->ObjectFlags = ImportContext.ImportObjectFlags;
-	TranslationContext->Time = static_cast< float >( UsdUtils::GetDefaultTimeCode() );
+	TranslationContext->Time = ImportContext.ImportOptions->bImportAtSpecificTimeCode
+		? ImportContext.ImportOptions->ImportTimeCode
+		: static_cast< float >( UsdUtils::GetDefaultTimeCode() );
 	TranslationContext->PurposesToLoad = ( EUsdPurpose ) ImportContext.ImportOptions->PurposesToImport;
 	TranslationContext->NaniteTriangleThreshold = ImportContext.ImportOptions->NaniteTriangleThreshold;
 	TranslationContext->RenderContext = ImportContext.ImportOptions->RenderContextToImport;
@@ -1670,7 +1696,7 @@ void UUsdStageImporter::ImportFromFile(FUsdStageImportContext& ImportContext)
 	UsdStageImporterImpl::RemapSoftReferences( ImportContext, UsedAssetsAndDependencies, SoftObjectsToRemap );
 	UsdStageImporterImpl::Cleanup( ImportContext.SceneActor, ExistingSceneActor, ImportContext.ImportOptions->ExistingActorPolicy );
 	UsdStageImporterImpl::NotifyAssetRegistry( UsedAssetsAndDependencies );
-	UsdStageImporterImpl::RefreshComponents( ImportContext.SceneActor );
+	UsdStageImporterImpl::RefreshComponents( ImportContext.SceneActor, ImportContext.ImportOptions->bImportAtSpecificTimeCode );
 
 	FUsdDelegates::OnPostUsdImport.Broadcast( ImportContext.FilePath );
 
@@ -1738,7 +1764,9 @@ bool UUsdStageImporter::ReimportSingleAsset(FUsdStageImportContext& ImportContex
 	TranslationContext->bIsImporting = true;
 	TranslationContext->Level = ImportContext.World->GetCurrentLevel();
 	TranslationContext->ObjectFlags = ImportContext.ImportObjectFlags;
-	TranslationContext->Time = static_cast< float >( UsdUtils::GetDefaultTimeCode() );
+	TranslationContext->Time = ImportContext.ImportOptions->bImportAtSpecificTimeCode
+		? ImportContext.ImportOptions->ImportTimeCode
+		: static_cast< float >( UsdUtils::GetDefaultTimeCode() );
 	TranslationContext->PurposesToLoad = ( EUsdPurpose ) ImportContext.ImportOptions->PurposesToImport;
 	TranslationContext->NaniteTriangleThreshold = ImportContext.ImportOptions->NaniteTriangleThreshold;
 	TranslationContext->RenderContext = ImportContext.ImportOptions->RenderContextToImport;
@@ -1797,7 +1825,7 @@ bool UUsdStageImporter::ReimportSingleAsset(FUsdStageImportContext& ImportContex
 
 	UsdStageImporterImpl::Cleanup( ImportContext.SceneActor, nullptr, ImportContext.ImportOptions->ExistingActorPolicy );
 	UsdStageImporterImpl::NotifyAssetRegistry( { ReimportedObject } );
-	UsdStageImporterImpl::RefreshComponents( ImportContext.SceneActor );
+	UsdStageImporterImpl::RefreshComponents( ImportContext.SceneActor, ImportContext.ImportOptions->bImportAtSpecificTimeCode );
 
 	FUsdDelegates::OnPostUsdImport.Broadcast(ImportContext.FilePath);
 
diff --git a/Engine/Plugins/Importers/USDImporter/Source/USDStageImporter/Public/USDStageImportOptions.h b/Engine/Plugins/Importers/USDImporter/Source/USDStageImporter/Public/USDStageImportOptions.h
--- a/Engine/Plugins/Importers/USDImporter/Source/USDStageImporter/Public/USDStageImportOptions.h
+++ b/Engine/Plugins/Importers/USDImporter/Source/USDStageImporter/Public/USDStageImportOptions.h
@@ -102,6 +102,17 @@ public:
 	UPROPERTY( EditAnywhere, config, BlueprintReadWrite, Category = "USD options", meta = ( EditCondition = bOverrideStageOptions ) )
 	FUsdStageOptions StageOptions;
 
+	/**
+	 * When true the stage will be evaluated at ImportTimeCode for the import.
+	 * When false, the stage will be evaluated at the default (non-animated) timecode
+	 */
+	UPROPERTY( BlueprintReadWrite, config, EditAnywhere, Category = "USD options" )
+	bool bImportAtSpecificTimeCode;
+
+	/** TimeCode to evaluate the stage for import, in case bImportAtSpecificTimeCode is enabled */
+	UPROPERTY( BlueprintReadWrite, config, EditAnywhere, Category = "USD options", meta = ( EditCondition = bImportAtSpecificTimeCode ) )
+	float ImportTimeCode;
+
 
 	/** Groom group interpolation settings */
 	UPROPERTY(EditAnywhere, config, BlueprintReadWrite, Category = "Groom")
