USD: Rename usages of "base material" to "reference material"

diff --git a/Engine/Plugins/Importers/USDImporter/Source/USDClasses/Public/USDProjectSettings.h b/Engine/Plugins/Importers/USDImporter/Source/USDClasses/Public/USDProjectSettings.h
--- a/Engine/Plugins/Importers/USDImporter/Source/USDClasses/Public/USDProjectSettings.h
+++ b/Engine/Plugins/Importers/USDImporter/Source/USDClasses/Public/USDProjectSettings.h
@@ -58,110 +58,110 @@ public:
 	// Note that the below properties being FSoftObjectPath ensure that these materials are cooked into packaged games
 
 	/**
-	 * What material to use as base material when creating material instances from USD materials.
+	 * What material to use as reference material when creating material instances from USD materials.
 	 * You can swap these with your own materials, but make sure that the replacement materials have parameters with
 	 * the same names and types as the ones provided by the default material, otherwise the instances will not have
 	 * the parameters filled with values extracted from the USD material when parsing.
 	 */
-	UPROPERTY( config, EditAnywhere, Category = "USD|Base Materials|UsdPreviewSurface", meta = ( AllowedClasses = "/Script/Engine.MaterialInterface" ) )
-	FSoftObjectPath BasePreviewSurfaceMaterial = FSoftObjectPath{ TEXT( "/USDImporter/Materials/UsdPreviewSurface.UsdPreviewSurface" ) };
+	UPROPERTY( config, EditAnywhere, Category = "USD|Reference Materials|UsdPreviewSurface", meta = ( AllowedClasses = "/Script/Engine.MaterialInterface" ) )
+	FSoftObjectPath ReferencePreviewSurfaceMaterial = FSoftObjectPath{ TEXT( "/USDImporter/Materials/UsdPreviewSurface.UsdPreviewSurface" ) };
 
 	/**
-	 * What material to use as base material when creating material instances from USD materials.
+	 * What material to use as reference material when creating material instances from USD materials.
 	 * You can swap these with your own materials, but make sure that the replacement materials have parameters with
 	 * the same names and types as the ones provided by the default material, otherwise the instances will not have
 	 * the parameters filled with values extracted from the USD material when parsing.
 	 */
-	UPROPERTY( config, EditAnywhere, Category = "USD|Base Materials|UsdPreviewSurface", meta = ( AllowedClasses = "/Script/Engine.MaterialInterface" ) )
-	FSoftObjectPath BasePreviewSurfaceTranslucentMaterial = FSoftObjectPath{ TEXT( "/USDImporter/Materials/UsdPreviewSurfaceTranslucent.UsdPreviewSurfaceTranslucent" ) };
+	UPROPERTY( config, EditAnywhere, Category = "USD|Reference Materials|UsdPreviewSurface", meta = ( AllowedClasses = "/Script/Engine.MaterialInterface" ) )
+	FSoftObjectPath ReferencePreviewSurfaceTranslucentMaterial = FSoftObjectPath{ TEXT( "/USDImporter/Materials/UsdPreviewSurfaceTranslucent.UsdPreviewSurfaceTranslucent" ) };
 
 	/**
-	 * What material to use as base material when creating material instances from USD materials.
+	 * What material to use as reference material when creating material instances from USD materials.
 	 * You can swap these with your own materials, but make sure that the replacement materials have parameters with
 	 * the same names and types as the ones provided by the default material, otherwise the instances will not have
 	 * the parameters filled with values extracted from the USD material when parsing.
 	 */
-	UPROPERTY( config, EditAnywhere, Category = "USD|Base Materials|UsdPreviewSurface", meta = ( AllowedClasses = "/Script/Engine.MaterialInterface" ) )
-	FSoftObjectPath BasePreviewSurfaceTwoSidedMaterial = FSoftObjectPath{ TEXT( "/USDImporter/Materials/UsdPreviewSurfaceTwoSided.UsdPreviewSurfaceTwoSided" ) };
+	UPROPERTY( config, EditAnywhere, Category = "USD|Reference Materials|UsdPreviewSurface", meta = ( AllowedClasses = "/Script/Engine.MaterialInterface" ) )
+	FSoftObjectPath ReferencePreviewSurfaceTwoSidedMaterial = FSoftObjectPath{ TEXT( "/USDImporter/Materials/UsdPreviewSurfaceTwoSided.UsdPreviewSurfaceTwoSided" ) };
 
 	/**
-	 * What material to use as base material when creating material instances from USD materials.
+	 * What material to use as reference material when creating material instances from USD materials.
 	 * You can swap these with your own materials, but make sure that the replacement materials have parameters with
 	 * the same names and types as the ones provided by the default material, otherwise the instances will not have
 	 * the parameters filled with values extracted from the USD material when parsing.
 	 */
-	UPROPERTY( config, EditAnywhere, Category = "USD|Base Materials|UsdPreviewSurface", meta = ( AllowedClasses = "/Script/Engine.MaterialInterface" ) )
-	FSoftObjectPath BasePreviewSurfaceTranslucentTwoSidedMaterial = FSoftObjectPath{ TEXT( "/USDImporter/Materials/UsdPreviewSurfaceTranslucentTwoSided.UsdPreviewSurfaceTranslucentTwoSided" ) };
+	UPROPERTY( config, EditAnywhere, Category = "USD|Reference Materials|UsdPreviewSurface", meta = ( AllowedClasses = "/Script/Engine.MaterialInterface" ) )
+	FSoftObjectPath ReferencePreviewSurfaceTranslucentTwoSidedMaterial = FSoftObjectPath{ TEXT( "/USDImporter/Materials/UsdPreviewSurfaceTranslucentTwoSided.UsdPreviewSurfaceTranslucentTwoSided" ) };
 
 	/**
-	 * What material to use as base material when creating material instances from USD materials.
+	 * What material to use as reference material when creating material instances from USD materials.
 	 * You can swap these with your own materials, but make sure that the replacement materials have parameters with
 	 * the same names and types as the ones provided by the default material, otherwise the instances will not have
 	 * the parameters filled with values extracted from the USD material when parsing.
 	 */
-	UPROPERTY( config, EditAnywhere, Category = "USD|Base Materials|UsdPreviewSurface with Virtual Textures", meta = ( AllowedClasses = "/Script/Engine.MaterialInterface" ) )
-	FSoftObjectPath BasePreviewSurfaceVTMaterial = FSoftObjectPath{ TEXT( "/USDImporter/Materials/UsdPreviewSurfaceVT.UsdPreviewSurfaceVT" ) };
+	UPROPERTY( config, EditAnywhere, Category = "USD|Reference Materials|UsdPreviewSurface with Virtual Textures", meta = ( AllowedClasses = "/Script/Engine.MaterialInterface" ) )
+	FSoftObjectPath ReferencePreviewSurfaceVTMaterial = FSoftObjectPath{ TEXT( "/USDImporter/Materials/UsdPreviewSurfaceVT.UsdPreviewSurfaceVT" ) };
 
 	/**
-	 * What material to use as base material when creating material instances from USD materials.
+	 * What material to use as reference material when creating material instances from USD materials.
 	 * You can swap these with your own materials, but make sure that the replacement materials have parameters with
 	 * the same names and types as the ones provided by the default material, otherwise the instances will not have
 	 * the parameters filled with values extracted from the USD material when parsing.
 	 */
-	UPROPERTY( config, EditAnywhere, Category = "USD|Base Materials|UsdPreviewSurface with Virtual Textures", meta = ( AllowedClasses = "/Script/Engine.MaterialInterface" ) )
-	FSoftObjectPath BasePreviewSurfaceTranslucentVTMaterial = FSoftObjectPath{ TEXT( "/USDImporter/Materials/UsdPreviewSurfaceTranslucentVT.UsdPreviewSurfaceTranslucentVT" ) };
+	UPROPERTY( config, EditAnywhere, Category = "USD|Reference Materials|UsdPreviewSurface with Virtual Textures", meta = ( AllowedClasses = "/Script/Engine.MaterialInterface" ) )
+	FSoftObjectPath ReferencePreviewSurfaceTranslucentVTMaterial = FSoftObjectPath{ TEXT( "/USDImporter/Materials/UsdPreviewSurfaceTranslucentVT.UsdPreviewSurfaceTranslucentVT" ) };
 
 	/**
-	 * What material to use as base material when creating material instances from USD materials.
+	 * What material to use as reference material when creating material instances from USD materials.
 	 * You can swap these with your own materials, but make sure that the replacement materials have parameters with
 	 * the same names and types as the ones provided by the default material, otherwise the instances will not have
 	 * the parameters filled with values extracted from the USD material when parsing.
 	 */
-	UPROPERTY( config, EditAnywhere, Category = "USD|Base Materials|UsdPreviewSurface with Virtual Textures", meta = ( AllowedClasses = "/Script/Engine.MaterialInterface" ) )
-	FSoftObjectPath BasePreviewSurfaceTwoSidedVTMaterial = FSoftObjectPath{ TEXT( "/USDImporter/Materials/UsdPreviewSurfaceTwoSidedVT.UsdPreviewSurfaceTwoSidedVT" ) };
+	UPROPERTY( config, EditAnywhere, Category = "USD|Reference Materials|UsdPreviewSurface with Virtual Textures", meta = ( AllowedClasses = "/Script/Engine.MaterialInterface" ) )
+	FSoftObjectPath ReferencePreviewSurfaceTwoSidedVTMaterial = FSoftObjectPath{ TEXT( "/USDImporter/Materials/UsdPreviewSurfaceTwoSidedVT.UsdPreviewSurfaceTwoSidedVT" ) };
 
 	/**
-	 * What material to use as base material when creating material instances from USD materials.
+	 * What material to use as reference material when creating material instances from USD materials.
 	 * You can swap these with your own materials, but make sure that the replacement materials have parameters with
 	 * the same names and types as the ones provided by the default material, otherwise the instances will not have
 	 * the parameters filled with values extracted from the USD material when parsing.
 	 */
-	UPROPERTY( config, EditAnywhere, Category = "USD|Base Materials|UsdPreviewSurface with Virtual Textures", meta = ( AllowedClasses = "/Script/Engine.MaterialInterface" ) )
-	FSoftObjectPath BasePreviewSurfaceTranslucentTwoSidedVTMaterial = FSoftObjectPath{ TEXT( "/USDImporter/Materials/UsdPreviewSurfaceTranslucentTwoSidedVT.UsdPreviewSurfaceTranslucentTwoSidedVT" ) };
+	UPROPERTY( config, EditAnywhere, Category = "USD|Reference Materials|UsdPreviewSurface with Virtual Textures", meta = ( AllowedClasses = "/Script/Engine.MaterialInterface" ) )
+	FSoftObjectPath ReferencePreviewSurfaceTranslucentTwoSidedVTMaterial = FSoftObjectPath{ TEXT( "/USDImporter/Materials/UsdPreviewSurfaceTranslucentTwoSidedVT.UsdPreviewSurfaceTranslucentTwoSidedVT" ) };
 
 	/**
-	 * What material to use as base material when creating material instances from USD materials.
+	 * What material to use as reference material when creating material instances from USD materials.
 	 * You can swap these with your own materials, but make sure that the replacement materials have parameters with
 	 * the same names and types as the ones provided by the default material, otherwise the instances will not have
 	 * the parameters filled with values extracted from the USD material when parsing.
 	 */
-	UPROPERTY( config, EditAnywhere, Category = "USD|Base Materials|DisplayColor", meta = ( AllowedClasses = "/Script/Engine.MaterialInterface" ) )
-	FSoftObjectPath BaseDisplayColorMaterial = FSoftObjectPath{ TEXT( "/USDImporter/Materials/DisplayColor.DisplayColor" ) };
+	UPROPERTY( config, EditAnywhere, Category = "USD|Reference Materials|DisplayColor", meta = ( AllowedClasses = "/Script/Engine.MaterialInterface" ) )
+	FSoftObjectPath ReferenceDisplayColorMaterial = FSoftObjectPath{ TEXT( "/USDImporter/Materials/DisplayColor.DisplayColor" ) };
 
 	/**
-	 * What material to use as base material when creating material instances from USD materials.
+	 * What material to use as reference material when creating material instances from USD materials.
 	 * You can swap these with your own materials, but make sure that the replacement materials have parameters with
 	 * the same names and types as the ones provided by the default material, otherwise the instances will not have
 	 * the parameters filled with values extracted from the USD material when parsing.
 	 */
-	UPROPERTY( config, EditAnywhere, Category = "USD|Base Materials|DisplayColor", meta = ( AllowedClasses = "/Script/Engine.MaterialInterface" ) )
-	FSoftObjectPath BaseDisplayColorAndOpacityMaterial = FSoftObjectPath{ TEXT( "/USDImporter/Materials/DisplayColorAndOpacity.DisplayColorAndOpacity" ) };
+	UPROPERTY( config, EditAnywhere, Category = "USD|Reference Materials|DisplayColor", meta = ( AllowedClasses = "/Script/Engine.MaterialInterface" ) )
+	FSoftObjectPath ReferenceDisplayColorAndOpacityMaterial = FSoftObjectPath{ TEXT( "/USDImporter/Materials/DisplayColorAndOpacity.DisplayColorAndOpacity" ) };
 
 	/**
-	 * What material to use as base material when creating material instances from USD materials.
+	 * What material to use as reference material when creating material instances from USD materials.
 	 * You can swap these with your own materials, but make sure that the replacement materials have parameters with
 	 * the same names and types as the ones provided by the default material, otherwise the instances will not have
 	 * the parameters filled with values extracted from the USD material when parsing.
 	 */
-	UPROPERTY( config, EditAnywhere, Category = "USD|Base Materials|DisplayColor", meta = ( AllowedClasses = "/Script/Engine.MaterialInterface" ) )
-	FSoftObjectPath BaseDisplayColorTwoSidedMaterial = FSoftObjectPath{ TEXT( "/USDImporter/Materials/DisplayColorTwoSided.DisplayColorTwoSided" ) };
+	UPROPERTY( config, EditAnywhere, Category = "USD|Reference Materials|DisplayColor", meta = ( AllowedClasses = "/Script/Engine.MaterialInterface" ) )
+	FSoftObjectPath ReferenceDisplayColorTwoSidedMaterial = FSoftObjectPath{ TEXT( "/USDImporter/Materials/DisplayColorTwoSided.DisplayColorTwoSided" ) };
 
 	/**
-	 * What material to use as base material when creating material instances from USD materials.
+	 * What material to use as reference material when creating material instances from USD materials.
 	 * You can swap these with your own materials, but make sure that the replacement materials have parameters with
 	 * the same names and types as the ones provided by the default material, otherwise the instances will not have
 	 * the parameters filled with values extracted from the USD material when parsing.
 	 */
-	UPROPERTY( config, EditAnywhere, Category = "USD|Base Materials|DisplayColor", meta = ( AllowedClasses = "/Script/Engine.MaterialInterface" ) )
-	FSoftObjectPath BaseDisplayColorAndOpacityTwoSidedMaterial = FSoftObjectPath{ TEXT( "/USDImporter/Materials/DisplayColorAndOpacityTwoSided.DisplayColorAndOpacityTwoSided" ) };
+	UPROPERTY( config, EditAnywhere, Category = "USD|Reference Materials|DisplayColor", meta = ( AllowedClasses = "/Script/Engine.MaterialInterface" ) )
+	FSoftObjectPath ReferenceDisplayColorAndOpacityTwoSidedMaterial = FSoftObjectPath{ TEXT( "/USDImporter/Materials/DisplayColorAndOpacityTwoSided.DisplayColorAndOpacityTwoSided" ) };
 };
\ No newline at end of file
diff --git a/Engine/Plugins/Importers/USDImporter/Source/USDSchemas/Private/MeshTranslationImpl.cpp b/Engine/Plugins/Importers/USDImporter/Source/USDSchemas/Private/MeshTranslationImpl.cpp
--- a/Engine/Plugins/Importers/USDImporter/Source/USDSchemas/Private/MeshTranslationImpl.cpp
+++ b/Engine/Plugins/Importers/USDImporter/Source/USDSchemas/Private/MeshTranslationImpl.cpp
@@ -112,11 +112,11 @@ TMap<const UsdUtils::FUsdPrimMaterialSlot*, UMaterialInterface*> MeshTranslation
 					}
 
 					// Important to not use GetBaseMaterial() here because if our parent is the translucent we'll
-					// get the base UsdPreviewSurface instead, as that is also *its* base
-					UMaterialInterface* BaseMaterial = OneSidedMat->Parent.Get();
-					UMaterialInterface* BaseMaterialTwoSided =
-						MeshTranslationImpl::GetTwoSidedVersionOfBasePreviewSurfaceMaterial( BaseMaterial );
-					if ( !ensure( BaseMaterialTwoSided && BaseMaterialTwoSided != BaseMaterial ) )
+					// get the reference UsdPreviewSurface instead, as that is also *its* reference
+					UMaterialInterface* ReferenceMaterial = OneSidedMat->Parent.Get();
+					UMaterialInterface* ReferenceMaterialTwoSided =
+						MeshTranslationImpl::GetTwoSidedVersionOfReferencePreviewSurfaceMaterial( ReferenceMaterial );
+					if ( !ensure( ReferenceMaterialTwoSided && ReferenceMaterialTwoSided != ReferenceMaterial ) )
 					{
 						continue;
 					}
@@ -146,7 +146,7 @@ TMap<const UsdUtils::FUsdPrimMaterialSlot*, UMaterialInterface*> MeshTranslation
 							TwoSidedMat->AssetImportData = ImportData;
 						}
 
-						TwoSidedMat->SetParentEditorOnly( BaseMaterialTwoSided );
+						TwoSidedMat->SetParentEditorOnly( ReferenceMaterialTwoSided );
 						TwoSidedMat->CopyMaterialUniformParametersEditorOnly( OneSidedMat );
 
 						AssetCache.CacheAsset( PrimPath, TwoSidedMat, PrimPath );
@@ -157,7 +157,7 @@ TMap<const UsdUtils::FUsdPrimMaterialSlot*, UMaterialInterface*> MeshTranslation
 					if ( UMaterialInstanceDynamic* MID = Cast<UMaterialInstanceDynamic>( OneSidedMat ) )
 					{
 						UMaterialInstanceDynamic* TwoSidedMat = UMaterialInstanceDynamic::Create(
-							BaseMaterialTwoSided,
+							ReferenceMaterialTwoSided,
 							GetTransientPackage(),
 							NewInstanceName
 						);
@@ -330,7 +330,7 @@ void MeshTranslationImpl::SetMaterialOverrides(
 	}
 }
 
-UMaterialInterface* MeshTranslationImpl::GetBasePreviewSurfaceMaterial( EUsdBaseMaterialProperties BaseMaterialProperties )
+UMaterialInterface* MeshTranslationImpl::GetReferencePreviewSurfaceMaterial( EUsdReferenceMaterialProperties ReferenceMaterialProperties )
 {
 	const UUsdProjectSettings* Settings = GetDefault<UUsdProjectSettings>();
 	if ( !Settings )
@@ -338,9 +338,9 @@ UMaterialInterface* MeshTranslationImpl::GetBasePreviewSurfaceMaterial( EUsdBase
 		return nullptr;
 	}
 
-	const bool bIsTranslucent = EnumHasAnyFlags( BaseMaterialProperties, EUsdBaseMaterialProperties::Translucent );
-	const bool bIsVT = EnumHasAnyFlags( BaseMaterialProperties, EUsdBaseMaterialProperties::VT );
-	const bool bIsTwoSided = EnumHasAnyFlags( BaseMaterialProperties, EUsdBaseMaterialProperties::TwoSided );
+	const bool bIsTranslucent = EnumHasAnyFlags( ReferenceMaterialProperties, EUsdReferenceMaterialProperties::Translucent );
+	const bool bIsVT = EnumHasAnyFlags( ReferenceMaterialProperties, EUsdReferenceMaterialProperties::VT );
+	const bool bIsTwoSided = EnumHasAnyFlags( ReferenceMaterialProperties, EUsdReferenceMaterialProperties::TwoSided );
 
 	const FSoftObjectPath* TargetMaterialPath = nullptr;
 	if ( bIsTranslucent )
@@ -349,22 +349,22 @@ UMaterialInterface* MeshTranslationImpl::GetBasePreviewSurfaceMaterial( EUsdBase
 		{
 			if ( bIsTwoSided )
 			{
-				TargetMaterialPath = &Settings->BasePreviewSurfaceTranslucentTwoSidedVTMaterial;
+				TargetMaterialPath = &Settings->ReferencePreviewSurfaceTranslucentTwoSidedVTMaterial;
 			}
 			else
 			{
-				TargetMaterialPath = &Settings->BasePreviewSurfaceTranslucentVTMaterial;
+				TargetMaterialPath = &Settings->ReferencePreviewSurfaceTranslucentVTMaterial;
 			}
 		}
 		else
 		{
 			if ( bIsTwoSided )
 			{
-				TargetMaterialPath = &Settings->BasePreviewSurfaceTranslucentTwoSidedMaterial;
+				TargetMaterialPath = &Settings->ReferencePreviewSurfaceTranslucentTwoSidedMaterial;
 			}
 			else
 			{
-				TargetMaterialPath = &Settings->BasePreviewSurfaceTranslucentMaterial;
+				TargetMaterialPath = &Settings->ReferencePreviewSurfaceTranslucentMaterial;
 			}
 		}
 	}
@@ -374,22 +374,22 @@ UMaterialInterface* MeshTranslationImpl::GetBasePreviewSurfaceMaterial( EUsdBase
 		{
 			if ( bIsTwoSided )
 			{
-				TargetMaterialPath = &Settings->BasePreviewSurfaceTwoSidedVTMaterial;
+				TargetMaterialPath = &Settings->ReferencePreviewSurfaceTwoSidedVTMaterial;
 			}
 			else
 			{
-				TargetMaterialPath = &Settings->BasePreviewSurfaceVTMaterial;
+				TargetMaterialPath = &Settings->ReferencePreviewSurfaceVTMaterial;
 			}
 		}
 		else
 		{
 			if ( bIsTwoSided )
 			{
-				TargetMaterialPath = &Settings->BasePreviewSurfaceTwoSidedMaterial;
+				TargetMaterialPath = &Settings->ReferencePreviewSurfaceTwoSidedMaterial;
 			}
 			else
 			{
-				TargetMaterialPath = &Settings->BasePreviewSurfaceMaterial;
+				TargetMaterialPath = &Settings->ReferencePreviewSurfaceMaterial;
 			}
 		}
 	}
@@ -402,9 +402,9 @@ UMaterialInterface* MeshTranslationImpl::GetBasePreviewSurfaceMaterial( EUsdBase
 	return Cast< UMaterialInterface >( TargetMaterialPath->TryLoad() );
 }
 
-UMaterialInterface* MeshTranslationImpl::GetVTVersionOfBasePreviewSurfaceMaterial( UMaterialInterface* BaseMaterial )
+UMaterialInterface* MeshTranslationImpl::GetVTVersionOfReferencePreviewSurfaceMaterial( UMaterialInterface* ReferenceMaterial )
 {
-	if ( !BaseMaterial )
+	if ( !ReferenceMaterial )
 	{
 		return nullptr;
 	}
@@ -415,36 +415,36 @@ UMaterialInterface* MeshTranslationImpl::GetVTVersionOfBasePreviewSurfaceMateria
 		return nullptr;
 	}
 
-	const FSoftObjectPath PathName = BaseMaterial->GetPathName();
+	const FSoftObjectPath PathName = ReferenceMaterial->GetPathName();
 	if ( PathName.ToString().Contains(TEXT("VT"), ESearchCase::CaseSensitive, ESearchDir::FromEnd) )
 	{
-		return BaseMaterial;
+		return ReferenceMaterial;
 	}
-	else if ( PathName == Settings->BasePreviewSurfaceMaterial )
+	else if ( PathName == Settings->ReferencePreviewSurfaceMaterial )
 	{
-		return Cast< UMaterialInterface >( Settings->BasePreviewSurfaceVTMaterial.TryLoad() );
+		return Cast< UMaterialInterface >( Settings->ReferencePreviewSurfaceVTMaterial.TryLoad() );
 	}
-	else if ( PathName == Settings->BasePreviewSurfaceTwoSidedMaterial )
+	else if ( PathName == Settings->ReferencePreviewSurfaceTwoSidedMaterial )
 	{
-		return Cast< UMaterialInterface >( Settings->BasePreviewSurfaceTwoSidedVTMaterial.TryLoad() );
+		return Cast< UMaterialInterface >( Settings->ReferencePreviewSurfaceTwoSidedVTMaterial.TryLoad() );
 	}
-	else if ( PathName == Settings->BasePreviewSurfaceTranslucentMaterial )
+	else if ( PathName == Settings->ReferencePreviewSurfaceTranslucentMaterial )
 	{
-		return Cast< UMaterialInterface >( Settings->BasePreviewSurfaceTranslucentVTMaterial.TryLoad() );
+		return Cast< UMaterialInterface >( Settings->ReferencePreviewSurfaceTranslucentVTMaterial.TryLoad() );
 	}
-	else if ( PathName == Settings->BasePreviewSurfaceTranslucentTwoSidedMaterial )
+	else if ( PathName == Settings->ReferencePreviewSurfaceTranslucentTwoSidedMaterial )
 	{
-		return Cast< UMaterialInterface >( Settings->BasePreviewSurfaceTranslucentTwoSidedVTMaterial.TryLoad() );
+		return Cast< UMaterialInterface >( Settings->ReferencePreviewSurfaceTranslucentTwoSidedVTMaterial.TryLoad() );
 	}
 
-	// We should only ever call this function with a BaseMaterial that matches one of the above paths
+	// We should only ever call this function with a ReferenceMaterial that matches one of the above paths
 	ensure( false );
 	return nullptr;
 }
 
-UMaterialInterface* MeshTranslationImpl::GetTwoSidedVersionOfBasePreviewSurfaceMaterial( UMaterialInterface* BaseMaterial )
+UMaterialInterface* MeshTranslationImpl::GetTwoSidedVersionOfReferencePreviewSurfaceMaterial( UMaterialInterface* ReferenceMaterial )
 {
-	if ( !BaseMaterial )
+	if ( !ReferenceMaterial )
 	{
 		return nullptr;
 	}
@@ -455,29 +455,29 @@ UMaterialInterface* MeshTranslationImpl::GetTwoSidedVersionOfBasePreviewSurfaceM
 		return nullptr;
 	}
 
-	const FSoftObjectPath PathName = BaseMaterial->GetPathName();
+	const FSoftObjectPath PathName = ReferenceMaterial->GetPathName();
 	if ( PathName.ToString().Contains(TEXT("TwoSided"), ESearchCase::CaseSensitive, ESearchDir::FromEnd) )
 	{
-		return BaseMaterial;
+		return ReferenceMaterial;
 	}
-	else if ( PathName == Settings->BasePreviewSurfaceMaterial )
+	else if ( PathName == Settings->ReferencePreviewSurfaceMaterial )
 	{
-		return Cast< UMaterialInterface >( Settings->BasePreviewSurfaceTwoSidedMaterial.TryLoad() );
+		return Cast< UMaterialInterface >( Settings->ReferencePreviewSurfaceTwoSidedMaterial.TryLoad() );
 	}
-	else if ( PathName == Settings->BasePreviewSurfaceTranslucentMaterial )
+	else if ( PathName == Settings->ReferencePreviewSurfaceTranslucentMaterial )
 	{
-		return Cast< UMaterialInterface >( Settings->BasePreviewSurfaceTranslucentTwoSidedMaterial.TryLoad() );
+		return Cast< UMaterialInterface >( Settings->ReferencePreviewSurfaceTranslucentTwoSidedMaterial.TryLoad() );
 	}
-	else if ( PathName == Settings->BasePreviewSurfaceVTMaterial )
+	else if ( PathName == Settings->ReferencePreviewSurfaceVTMaterial )
 	{
-		return Cast< UMaterialInterface >( Settings->BasePreviewSurfaceTwoSidedVTMaterial.TryLoad() );
+		return Cast< UMaterialInterface >( Settings->ReferencePreviewSurfaceTwoSidedVTMaterial.TryLoad() );
 	}
-	else if ( PathName == Settings->BasePreviewSurfaceTranslucentVTMaterial )
+	else if ( PathName == Settings->ReferencePreviewSurfaceTranslucentVTMaterial )
 	{
-		return Cast< UMaterialInterface >( Settings->BasePreviewSurfaceTranslucentTwoSidedVTMaterial.TryLoad() );
+		return Cast< UMaterialInterface >( Settings->ReferencePreviewSurfaceTranslucentTwoSidedVTMaterial.TryLoad() );
 	}
 
-	// We should only ever call this function with a BaseMaterial that matches one of the above paths
+	// We should only ever call this function with a ReferenceMaterial that matches one of the above paths
 	ensure( false );
 	return nullptr;
 }
diff --git a/Engine/Plugins/Importers/USDImporter/Source/USDSchemas/Private/MeshTranslationImpl.h b/Engine/Plugins/Importers/USDImporter/Source/USDSchemas/Private/MeshTranslationImpl.h
--- a/Engine/Plugins/Importers/USDImporter/Source/USDSchemas/Private/MeshTranslationImpl.h
+++ b/Engine/Plugins/Importers/USDImporter/Source/USDSchemas/Private/MeshTranslationImpl.h
@@ -51,34 +51,34 @@ namespace MeshTranslationImpl
 		const FName& MaterialPurpose
 	);
 
-	enum class EUsdBaseMaterialProperties
+	enum class EUsdReferenceMaterialProperties
 	{
 		None = 0,
 		Translucent = 1,
 		VT = 2,
 		TwoSided = 4
 	};
-	ENUM_CLASS_FLAGS( EUsdBaseMaterialProperties )
+	ENUM_CLASS_FLAGS( EUsdReferenceMaterialProperties )
 
-	// Returns one of the alternatives of the UsdPreviewSurface base material depending on the material overrides
+	// Returns one of the alternatives of the UsdPreviewSurface reference material depending on the material overrides
 	// provided, and nullptr otherwise
-	UMaterialInterface* GetBasePreviewSurfaceMaterial( EUsdBaseMaterialProperties BaseMaterialProperties );
+	UMaterialInterface* GetReferencePreviewSurfaceMaterial( EUsdReferenceMaterialProperties ReferenceMaterialProperties );
 
-	// Returns the VT version of the provided UsdPreviewSurface BaseMaterial. Returns the provided BaseMaterial back if
-	// it is already a VT-capable base material, and returns nullptr if BaseMaterial isn't one of our base material
+	// Returns the VT version of the provided UsdPreviewSurface ReferenceMaterial. Returns the provided ReferenceMaterial back if
+	// it is already a VT-capable reference material, and returns nullptr if ReferenceMaterial isn't one of our reference material
 	// alternatives.
 	// Example: Receives UsdPreviewSurfaceTwoSided -> Returns UsdPreviewSurfaceTwoSidedVT
 	// Example: Receives UsdPreviewSurfaceTwoSidedVT -> Returns UsdPreviewSurfaceTwoSidedVT
-	// Example: Receives SomeOtherBaseMaterial -> Returns nullptr
-	UMaterialInterface* GetVTVersionOfBasePreviewSurfaceMaterial( UMaterialInterface* BaseMaterial );
+	// Example: Receives SomeOtherReferenceMaterial -> Returns nullptr
+	UMaterialInterface* GetVTVersionOfReferencePreviewSurfaceMaterial( UMaterialInterface* ReferenceMaterial );
 
-	// Returns the two-sided version of the provided UsdPreviewSurface BaseMaterial. Returns the provided BaseMaterial
-	// back if it is already a two-sided-capable base material, and returns nullptr if BaseMaterial isn't one of our base
+	// Returns the two-sided version of the provided UsdPreviewSurface ReferenceMaterial. Returns the provided ReferenceMaterial
+	// back if it is already a two-sided-capable reference material, and returns nullptr if ReferenceMaterial isn't one of our reference
 	// material alternatives.
 	// Example: Receives UsdPreviewSurfaceTranslucent -> Returns UsdPreviewSurfaceTwoSidedTranslucent
 	// Example: Receives UsdPreviewSurfaceTwoSidedTranslucent -> Returns UsdPreviewSurfaceTwoSidedTranslucent
-	// Example: Receives SomeOtherBaseMaterial -> Returns nullptr
-	UMaterialInterface* GetTwoSidedVersionOfBasePreviewSurfaceMaterial( UMaterialInterface* BaseMaterial );
+	// Example: Receives SomeOtherReferenceMaterial -> Returns nullptr
+	UMaterialInterface* GetTwoSidedVersionOfReferencePreviewSurfaceMaterial( UMaterialInterface* ReferenceMaterial );
 }
 
 #endif // #if USE_USD_SDK
\ No newline at end of file
diff --git a/Engine/Plugins/Importers/USDImporter/Source/USDSchemas/Private/USDShadeMaterialTranslator.cpp b/Engine/Plugins/Importers/USDImporter/Source/USDSchemas/Private/USDShadeMaterialTranslator.cpp
--- a/Engine/Plugins/Importers/USDImporter/Source/USDSchemas/Private/USDShadeMaterialTranslator.cpp
+++ b/Engine/Plugins/Importers/USDImporter/Source/USDSchemas/Private/USDShadeMaterialTranslator.cpp
@@ -66,18 +66,18 @@ namespace UE::UsdShadeTranslator::Private
 					if ( UMaterialInstance* MaterialInstance = Cast<UMaterialInstance>( UserMaterial ) )
 					{
 						// Important to not use GetBaseMaterial() here because if our parent is the translucent we'll
-						// get the base UsdPreviewSurface instead, as that is also *its* base
-						UMaterialInterface* BaseMaterial = MaterialInstance->Parent.Get();
-						UMaterialInterface* BaseMaterialVT =
-							MeshTranslationImpl::GetVTVersionOfBasePreviewSurfaceMaterial( BaseMaterial );
-						if ( BaseMaterial == BaseMaterialVT )
+						// get the reference UsdPreviewSurface instead, as that is also *its* reference
+						UMaterialInterface* ReferenceMaterial = MaterialInstance->Parent.Get();
+						UMaterialInterface* ReferenceMaterialVT =
+							MeshTranslationImpl::GetVTVersionOfReferencePreviewSurfaceMaterial( ReferenceMaterial );
+						if ( ReferenceMaterial == ReferenceMaterialVT )
 						{
 							// Material is already VT, we're good
 							continue;
 						}
 
 						// Visit it before we start recursing. We need this because we must convert textures to VT
-						// before materials (or else we get a warning) but we'll only actually swap the base material
+						// before materials (or else we get a warning) but we'll only actually swap the reference material
 						// at the end of this scope
 						VisitedMaterials.Add( UserMaterial );
 
@@ -109,7 +109,7 @@ namespace UE::UsdShadeTranslator::Private
 							);
 						}
 
-						UE_LOG( LogUsd, Log, TEXT( "Upgrading material instance '%s' to being based on VT as texture '%s' requires it" ),
+						UE_LOG( LogUsd, Log, TEXT( "Upgrading material instance '%s' to having a VT reference as texture '%s' requires it" ),
 							*MaterialInstance->GetName(),
 							*Texture->GetName()
 						);
@@ -121,7 +121,7 @@ namespace UE::UsdShadeTranslator::Private
 							FMaterialUpdateContext UpdateContext( Options, GMaxRHIShaderPlatform );
 							UpdateContext.AddMaterialInstance( MIC );
 							MIC->PreEditChange( nullptr );
-							MIC->SetParentEditorOnly( BaseMaterialVT );
+							MIC->SetParentEditorOnly( ReferenceMaterialVT );
 							MIC->PostEditChange();
 						}
 						else
@@ -141,10 +141,10 @@ namespace UE::UsdShadeTranslator::Private
 									*FPaths::GetBaseFilename( PrimPath )
 								);
 
-								// For MID we can't swap the base material, so we need to create a brand new one and copy
+								// For MID we can't swap the reference material, so we need to create a brand new one and copy
 								// the overrides
 								UMaterialInstanceDynamic* NewMID = UMaterialInstanceDynamic::Create(
-									BaseMaterialVT,
+									ReferenceMaterialVT,
 									GetTransientPackage(),
 									NewInstanceName
 								);
@@ -293,20 +293,20 @@ void FUsdShadeMaterialTranslator::CreateAssets()
 					UE::UsdShadeTranslator::Private::UpgradeMaterialsAndTexturesToVT( NonVTTextures, Context );
 				}
 
-				MeshTranslationImpl::EUsdBaseMaterialProperties Properties = MeshTranslationImpl::EUsdBaseMaterialProperties::None;
+				MeshTranslationImpl::EUsdReferenceMaterialProperties Properties = MeshTranslationImpl::EUsdReferenceMaterialProperties::None;
 				if ( bIsTranslucent )
 				{
-					Properties |= MeshTranslationImpl::EUsdBaseMaterialProperties::Translucent;
+					Properties |= MeshTranslationImpl::EUsdReferenceMaterialProperties::Translucent;
 				}
 				if ( VTTextures.Num() > 0 )
 				{
-					Properties |= MeshTranslationImpl::EUsdBaseMaterialProperties::VT;
+					Properties |= MeshTranslationImpl::EUsdReferenceMaterialProperties::VT;
 				}
-				UMaterialInterface* BaseMaterial = MeshTranslationImpl::GetBasePreviewSurfaceMaterial( Properties );
+				UMaterialInterface* ReferenceMaterial = MeshTranslationImpl::GetReferencePreviewSurfaceMaterial( Properties );
 
-				if ( ensure( BaseMaterial ) )
+				if ( ensure( ReferenceMaterial ) )
 				{
-					NewMaterial->SetParentEditorOnly( BaseMaterial );
+					NewMaterial->SetParentEditorOnly( ReferenceMaterial );
 
 					// We can't blindly recreate all component render states when a level is being added, because we may end up first creating
 					// render states for some components, and UWorld::AddToWorld calls FScene::AddPrimitive which expects the component to not have
@@ -334,20 +334,20 @@ void FUsdShadeMaterialTranslator::CreateAssets()
 		else
 #endif // WITH_EDITOR
 		{
-			// At runtime we always start with a non-VT base and if we discover we need one we just create a new MID
-			// using the VT base and copy the overrides. Not much else we can do as we need a base to call
+			// At runtime we always start with a non-VT reference and if we discover we need one we just create a new MID
+			// using the VT reference and copy the overrides. Not much else we can do as we need a reference to call
 			// UMaterialInstanceDynamic::Create and get our instance, but we an instance to call UsdToUnreal::ConvertMaterial
-			// to create our textures and decide on our base.
-			MeshTranslationImpl::EUsdBaseMaterialProperties Properties = MeshTranslationImpl::EUsdBaseMaterialProperties::None;
+			// to create our textures and decide on our reference.
+			MeshTranslationImpl::EUsdReferenceMaterialProperties Properties = MeshTranslationImpl::EUsdReferenceMaterialProperties::None;
 			if ( bIsTranslucent )
 			{
-				Properties |= MeshTranslationImpl::EUsdBaseMaterialProperties::Translucent;
+				Properties |= MeshTranslationImpl::EUsdReferenceMaterialProperties::Translucent;
 			}
-			UMaterialInterface* BaseMaterial = MeshTranslationImpl::GetBasePreviewSurfaceMaterial( Properties );
+			UMaterialInterface* ReferenceMaterial = MeshTranslationImpl::GetReferencePreviewSurfaceMaterial( Properties );
 
-			if ( ensure( BaseMaterial ) )
+			if ( ensure( ReferenceMaterial ) )
 			{
-				if ( UMaterialInstanceDynamic* NewMaterial = UMaterialInstanceDynamic::Create( BaseMaterial, GetTransientPackage(), InstanceName ) )
+				if ( UMaterialInstanceDynamic* NewMaterial = UMaterialInstanceDynamic::Create( ReferenceMaterial, GetTransientPackage(), InstanceName ) )
 				{
 					if ( UsdToUnreal::ConvertMaterial( ShadeMaterial, *NewMaterial, Context->AssetCache.Get(), PrimvarToUVIndex, *Context->RenderContext.ToString() ) )
 					{
@@ -369,7 +369,7 @@ void FUsdShadeMaterialTranslator::CreateAssets()
 						}
 
 						// We must stash our material and textures *before* we call UpgradeMaterialsAndTexturesToVT, as that
-						// is what will actually swap our base with a VT one if needed
+						// is what will actually swap our reference with a VT one if needed
 						if ( Context->AssetCache )
 						{
 							Context->AssetCache->CacheAsset( MaterialHashString, NewMaterial, PrimPath.GetString() );
@@ -389,7 +389,7 @@ void FUsdShadeMaterialTranslator::CreateAssets()
 						}
 
 						// We must go through the cache to fetch our result material here as UpgradeMaterialsAndTexturesToVT
-						// may have created a new MID for this material with a VT base
+						// may have created a new MID for this material with a VT reference
 						ConvertedMaterial = Cast<UMaterialInterface>( Context->AssetCache->GetCachedAsset( MaterialHashString ) );
 					}
 				}
diff --git a/Engine/Plugins/Importers/USDImporter/Source/USDUtilities/Private/USDGeomMeshConversion.cpp b/Engine/Plugins/Importers/USDImporter/Source/USDUtilities/Private/USDGeomMeshConversion.cpp
--- a/Engine/Plugins/Importers/USDImporter/Source/USDUtilities/Private/USDGeomMeshConversion.cpp
+++ b/Engine/Plugins/Importers/USDImporter/Source/USDUtilities/Private/USDGeomMeshConversion.cpp
@@ -1392,22 +1392,22 @@ UMaterialInstanceDynamic* UsdUtils::CreateDisplayColorMaterialInstanceDynamic( c
 	{
 		if ( DisplayColorDescription.bIsDoubleSided )
 		{
-			ParentPathPtr = &Settings->BaseDisplayColorAndOpacityTwoSidedMaterial;
+			ParentPathPtr = &Settings->ReferenceDisplayColorAndOpacityTwoSidedMaterial;
 		}
 		else
 		{
-			ParentPathPtr = &Settings->BaseDisplayColorAndOpacityMaterial;
+			ParentPathPtr = &Settings->ReferenceDisplayColorAndOpacityMaterial;
 		}
 	}
 	else
 	{
 		if ( DisplayColorDescription.bIsDoubleSided )
 		{
-			ParentPathPtr = &Settings->BaseDisplayColorTwoSidedMaterial;
+			ParentPathPtr = &Settings->ReferenceDisplayColorTwoSidedMaterial;
 		}
 		else
 		{
-			ParentPathPtr = &Settings->BaseDisplayColorMaterial;
+			ParentPathPtr = &Settings->ReferenceDisplayColorMaterial;
 		}
 	}
 	if ( !ParentPathPtr )
@@ -1449,22 +1449,22 @@ UMaterialInstanceConstant* UsdUtils::CreateDisplayColorMaterialInstanceConstant(
 	{
 		if ( DisplayColorDescription.bIsDoubleSided )
 		{
-			ParentPathPtr = &Settings->BaseDisplayColorAndOpacityTwoSidedMaterial;
+			ParentPathPtr = &Settings->ReferenceDisplayColorAndOpacityTwoSidedMaterial;
 		}
 		else
 		{
-			ParentPathPtr = &Settings->BaseDisplayColorAndOpacityMaterial;
+			ParentPathPtr = &Settings->ReferenceDisplayColorAndOpacityMaterial;
 		}
 	}
 	else
 	{
 		if ( DisplayColorDescription.bIsDoubleSided )
 		{
-			ParentPathPtr = &Settings->BaseDisplayColorTwoSidedMaterial;
+			ParentPathPtr = &Settings->ReferenceDisplayColorTwoSidedMaterial;
 		}
 		else
 		{
-			ParentPathPtr = &Settings->BaseDisplayColorMaterial;
+			ParentPathPtr = &Settings->ReferenceDisplayColorMaterial;
 		}
 	}
 
diff --git a/Engine/Plugins/Importers/USDImporter/Source/USDUtilities/Private/USDShadeConversion.cpp b/Engine/Plugins/Importers/USDImporter/Source/USDUtilities/Private/USDShadeConversion.cpp
--- a/Engine/Plugins/Importers/USDImporter/Source/USDUtilities/Private/USDShadeConversion.cpp
+++ b/Engine/Plugins/Importers/USDImporter/Source/USDUtilities/Private/USDShadeConversion.cpp
@@ -1166,7 +1166,7 @@ namespace UE
 			};
 
 			/**
-			 * Specialized version of FSetParameterValueVisitor for UE's UsdPreviewSurface base materials.
+			 * Specialized version of FSetParameterValueVisitor for UE's UsdPreviewSurface reference materials.
 			 */
 			struct FSetPreviewSurfaceParameterValueVisitor : private FSetParameterValueVisitor
 			{
diff --git a/Engine/Plugins/Importers/USDImporter/Source/USDUtilities/Public/USDGeomMeshConversion.h b/Engine/Plugins/Importers/USDImporter/Source/USDUtilities/Public/USDGeomMeshConversion.h
--- a/Engine/Plugins/Importers/USDImporter/Source/USDUtilities/Public/USDGeomMeshConversion.h
+++ b/Engine/Plugins/Importers/USDImporter/Source/USDUtilities/Public/USDGeomMeshConversion.h
@@ -272,7 +272,7 @@ namespace UsdUtils
 	/** Creates a FDisplayColorMaterial object describing the vertex color/opacity data from UsdMesh at time TimeCode */
 	USDUTILITIES_API TOptional<FDisplayColorMaterial> ExtractDisplayColorMaterial( const pxr::UsdGeomMesh& UsdMesh, const pxr::UsdTimeCode TimeCode = pxr::UsdTimeCode::EarliestTime() );
 
-	/** Creates a dynamic material instance using the right base material depending on the given description */
+	/** Creates a dynamic material instance using the right reference material depending on the given description */
 	USDUTILITIES_API UMaterialInstanceDynamic* CreateDisplayColorMaterialInstanceDynamic( const UsdUtils::FDisplayColorMaterial& DisplayColorDescription );
 	USDUTILITIES_API UMaterialInstanceConstant* CreateDisplayColorMaterialInstanceConstant( const UsdUtils::FDisplayColorMaterial& DisplayColorDescription );
 
