USD: GeometryCache from USD now use tracks in LevelSequence instead of the Time property to drive their animations. Added USD.DisableGeoCacheTracks cvar to revert to the previous behavior.

diff --git a/Engine/Plugins/Importers/USDImporter/Source/GeometryCacheUSD/Private/GeometryCacheTrackUSD.cpp b/Engine/Plugins/Importers/USDImporter/Source/GeometryCacheUSD/Private/GeometryCacheTrackUSD.cpp
--- a/Engine/Plugins/Importers/USDImporter/Source/GeometryCacheUSD/Private/GeometryCacheTrackUSD.cpp
+++ b/Engine/Plugins/Importers/USDImporter/Source/GeometryCacheUSD/Private/GeometryCacheTrackUSD.cpp
@@ -6,20 +6,28 @@
 #include "UsdWrappers/SdfLayer.h"
 #include "UsdWrappers/UsdStage.h"
 
+#include "GeometryCacheHelpers.h"
 #include "GeometryCacheUSDStream.h"
+#include "HAL/IConsoleManager.h"
 #include "IGeometryCacheStreamer.h"
 
+static bool GDisableGeoCacheTracks = false;
+static FAutoConsoleVariableRef CVarDisableGeoCacheTracks(
+	TEXT("USD.DisableGeoCacheTracks"),
+	GDisableGeoCacheTracks,
+	TEXT("Set to true to disable geometry cache tracks in Sequencer and drive them by the Time property instead. The stage must be reloaded after changing this value."));
+
 UGeometryCacheTrackUsd::UGeometryCacheTrackUsd()
-: StartFrameIndex(0)
-, EndFrameIndex(0)
-{
-}
+	: FramesPerSecond(24.0)
+	, StartFrameIndex(0)
+	, EndFrameIndex(0)
+{}
 
 void UGeometryCacheTrackUsd::BeginDestroy()
 {
 	UnloadUsdStage();
 
-	IGeometryCacheStreamer::Get().UnregisterTrack( this );
+	IGeometryCacheStreamer::Get().UnregisterTrack(this);
 	UsdStream.Reset();
 
 	Super::BeginDestroy();
@@ -59,19 +67,67 @@ const bool UGeometryCacheTrackUsd::UpdateBoundsData(const float Time, const bool
 
 const int32 UGeometryCacheTrackUsd::FindSampleIndexFromTime(const float Time, const bool bLooping) const
 {
+	if (!GDisableGeoCacheTracks)
+	{
+		// Time is relative to the start of the section on the track
+		float SampleTime = Time;
+		if (bLooping)
+		{
+			SampleTime = GeometyCacheHelpers::WrapAnimationTime(Time, Duration);
+		}
+		// Which is converted to an index
+		int32 FrameIndex = FMath::RoundToInt(Time * FramesPerSecond);
+
+		// The final computed frame index must be offset by the start frame index
+		return FMath::Clamp(FrameIndex + StartFrameIndex, StartFrameIndex, EndFrameIndex - 1);
+	}
 	// Treat the time as the frame index
-	int32 FrameIndex = (int32) Time;
+	int32 FrameIndex = (int32)Time;
 	return FMath::Clamp(FrameIndex, StartFrameIndex, EndFrameIndex - 1);
 }
 
+float UGeometryCacheTrackUsd::GetTimeFromSampleIndex(int32 SampleIndex) const
+{
+	if (!GDisableGeoCacheTracks)
+	{
+		// Time is relative to the start of the section on the track
+		return float((SampleIndex - StartFrameIndex) / FramesPerSecond);
+	}
+	// Treat the frame index as time
+	return SampleIndex;
+}
+
+void UGeometryCacheTrackUsd::GetFractionalFrameIndexFromTime(const float Time, const bool bLooping, int& OutFrameIndex, float& OutFraction) const
+{
+	OutFrameIndex = FindSampleIndexFromTime(Time, bLooping);
+
+	float AdjustedTime = Time;
+	if (!GDisableGeoCacheTracks)
+	{
+		if (bLooping)
+		{
+			AdjustedTime = GeometyCacheHelpers::WrapAnimationTime(Time, Duration);
+		}
+		// Time at ThisFrameIndex with index normalized to 0
+		const float FrameIndexTime = (OutFrameIndex - StartFrameIndex) / FramesPerSecond;
+		OutFraction = float((AdjustedTime - FrameIndexTime) * FramesPerSecond);
+	}
+	else
+	{
+		// Clamp the Time (which is the FrameNumber) to the range of the USD GeometryCache track
+		// EndFrameIndex - 1 since last frame index is not included
+		OutFraction = FMath::Clamp(Time, float(StartFrameIndex), float(EndFrameIndex - 1)) - OutFrameIndex;
+	}
+}
+
 const FGeometryCacheTrackSampleInfo& UGeometryCacheTrackUsd::GetSampleInfo(float Time, bool bLooping)
 {
 	if (SampleInfos.Num() == 0)
 	{
 		if (Duration > 0.f)
 		{
-			// Duration is the number of frames
-			SampleInfos.SetNum((int32)Duration);
+			const int32 NumFrames = EndFrameIndex - StartFrameIndex + 1;
+			SampleInfos.SetNum(NumFrames);
 		}
 		else
 		{
@@ -80,7 +136,7 @@ const FGeometryCacheTrackSampleInfo& UGeometryCacheTrackUsd::GetSampleInfo(float
 	}
 
 	// The sample info index must start from 0, while the sample index is between the range of the animation
-	const int32 SampleIndex = FindSampleIndexFromTime( Time, bLooping );
+	const int32 SampleIndex = FindSampleIndexFromTime(Time, bLooping);
 	const int32 SampleInfoIndex = SampleIndex - StartFrameIndex;
 
 	FGeometryCacheTrackSampleInfo& CurrentSampleInfo = SampleInfos[SampleInfoIndex];
@@ -92,7 +148,7 @@ const FGeometryCacheTrackSampleInfo& UGeometryCacheTrackUsd::GetSampleInfo(float
 		{
 			CurrentSampleInfo = FGeometryCacheTrackSampleInfo(
 				Time,
-				(FBox) TempMeshData.BoundingBox,
+				(FBox)TempMeshData.BoundingBox,
 				TempMeshData.Positions.Num(),
 				TempMeshData.Indices.Num()
 			);
@@ -127,21 +183,21 @@ bool UGeometryCacheTrackUsd::GetMeshData(int32 SampleIndex, FGeometryCacheMeshDa
 
 bool UGeometryCacheTrackUsd::LoadUsdStage()
 {
-	if ( CurrentStagePinned )
+	if (CurrentStagePinned)
 	{
 		// Already loaded
 		return true;
 	}
 
-	if ( CurrentStageWeak )
+	if (CurrentStageWeak)
 	{
 		// Upgrade our weak pointer if its not invalid already
 		CurrentStagePinned = CurrentStageWeak;
 		return true;
 	}
-	else if ( !StageRootLayerPath.IsEmpty() )
+	else if (!StageRootLayerPath.IsEmpty())
 	{
-		UE_LOG( LogUsd, Warning, TEXT( "UGeometryCacheTrackUsd is reopening the stage '%s' to stream in frames for the geometry cache generated for prim '%s'" ), *StageRootLayerPath, *PrimPath );
+		UE_LOG(LogUsd, Warning, TEXT("UGeometryCacheTrackUsd is reopening the stage '%s' to stream in frames for the geometry cache generated for prim '%s'"), *StageRootLayerPath, *PrimPath);
 
 		// Reopen the stage. If our weak pointer is no longer valid then nothing cared about keeping that
 		// stage alive anyway, so it's likely not a problem if we start reading frames from the reopened stage
@@ -153,12 +209,12 @@ bool UGeometryCacheTrackUsd::LoadUsdStage()
 		// after we're done with it, and the UGeometryCacheTrackUsd shouldn't have authority to blindly just remove it
 		// from the cache as the user may have place it there intentionally
 		const bool bUseStageCache = false;
-		CurrentStagePinned = UnrealUSDWrapper::OpenStage( *StageRootLayerPath, EUsdInitialLoadSet::LoadAll, bUseStageCache );
+		CurrentStagePinned = UnrealUSDWrapper::OpenStage(*StageRootLayerPath, EUsdInitialLoadSet::LoadAll, bUseStageCache);
 		CurrentStageWeak = CurrentStagePinned;
 		return true;
 	}
 
-	UE_LOG( LogUsd, Warning, TEXT( "UGeometryCacheTrackUsd track failed to access USD stage to stream requested frames" ) );
+	UE_LOG(LogUsd, Warning, TEXT("UGeometryCacheTrackUsd track failed to access USD stage to stream requested frames"));
 	return false;
 }
 
@@ -186,8 +242,19 @@ void UGeometryCacheTrackUsd::Initialize(
 	MaterialToPrimvarToUVIndex = InMaterialToPrimvarToUVIndex;
 	StartFrameIndex = InStartFrameIndex;
 	EndFrameIndex = InEndFrameIndex;
+	FramesPerSecond = CurrentStagePinned.GetTimeCodesPerSecond();
+
+	if (FramesPerSecond == 0)
+	{
+		ensureMsgf(false, TEXT("Invalid USD GeometryCache FPS detected. Falling back to 1 FPS"));
+		FramesPerSecond = 1;
+	}
 
-	Duration = ( float ) ( EndFrameIndex - StartFrameIndex );
+	// The main difference between GDisableGeoCacheTracks true or false is in how time
+	// is converted to frame index. When true, the time is basically the frame indices
+	// and no conversion is actually needed.
+	const int32 NumFrames = EndFrameIndex - StartFrameIndex + 1;
+	Duration = GDisableGeoCacheTracks ? NumFrames : NumFrames / FramesPerSecond;
 
 	UsdStream.Reset(new FGeometryCacheUsdStream(this, InReadFunc));
 	IGeometryCacheStreamer::Get().RegisterTrack(this, UsdStream.Get());
diff --git a/Engine/Plugins/Importers/USDImporter/Source/GeometryCacheUSD/Private/GeometryCacheUSDSceneProxy.cpp b/Engine/Plugins/Importers/USDImporter/Source/GeometryCacheUSD/Private/GeometryCacheUSDSceneProxy.cpp
--- a/Engine/Plugins/Importers/USDImporter/Source/GeometryCacheUSD/Private/GeometryCacheUSDSceneProxy.cpp
+++ b/Engine/Plugins/Importers/USDImporter/Source/GeometryCacheUSD/Private/GeometryCacheUSDSceneProxy.cpp
@@ -1,8 +1,8 @@
 // Copyright Epic Games, Inc. All Rights Reserved.
 
 #include "GeometryCacheUSDSceneProxy.h"
-#include "GeometryCacheUSDComponent.h"
 #include "GeometryCacheTrackUSD.h"
+#include "GeometryCacheUSDComponent.h"
 
 FGeometryCacheUsdSceneProxy::FGeometryCacheUsdSceneProxy(UGeometryCacheUsdComponent* Component)
 : FGeometryCacheSceneProxy(Component, [this]() { return new FGeomCacheTrackUsdProxy(GetScene().GetFeatureLevel()); })
@@ -11,6 +11,8 @@ FGeometryCacheUsdSceneProxy::FGeometryCacheUsdSceneProxy(UGeometryCacheUsdCompon
 
 bool FGeomCacheTrackUsdProxy::UpdateMeshData(float Time, bool bLooping, int32& InOutMeshSampleIndex, FGeometryCacheMeshData& OutMeshData)
 {
+	TRACE_CPUPROFILER_EVENT_SCOPE(FGeomCacheTrackUsdProxy::UpdateMeshData);
+
 	if (UGeometryCacheTrackUsd* UsdTrack = Cast<UGeometryCacheTrackUsd>(Track))
 	{
 		FGeometryCacheMeshData* PtrMeshData = nullptr;
@@ -26,6 +28,8 @@ bool FGeomCacheTrackUsdProxy::UpdateMeshData(float Time, bool bLooping, int32& I
 
 bool FGeomCacheTrackUsdProxy::GetMeshData(int32 SampleIndex, FGeometryCacheMeshData& OutMeshData)
 {
+	TRACE_CPUPROFILER_EVENT_SCOPE(FGeomCacheTrackUsdProxy::GetMeshData);
+
 	if (UGeometryCacheTrackUsd* UsdTrack = Cast<UGeometryCacheTrackUsd>(Track))
 	{
 		return UsdTrack->GetMeshData(SampleIndex, OutMeshData);
@@ -35,11 +39,16 @@ bool FGeomCacheTrackUsdProxy::GetMeshData(int32 SampleIndex, FGeometryCacheMeshD
 
 bool FGeomCacheTrackUsdProxy::IsTopologyCompatible(int32 SampleIndexA, int32 SampleIndexB)
 {
+	TRACE_CPUPROFILER_EVENT_SCOPE(FGeomCacheTrackUsdProxy::IsTopologyCompatible);
+
 	if (UGeometryCacheTrackUsd* UsdTrack = Cast<UGeometryCacheTrackUsd>(Track))
 	{
 		// The boolean argument is not actually used
-		const int32 NumVerticesA = UsdTrack->GetSampleInfo(SampleIndexA, false).NumVertices;
-		const int32 NumVerticesB = UsdTrack->GetSampleInfo(SampleIndexB, false).NumVertices;
+		const float TimeA = UsdTrack->GetTimeFromSampleIndex(SampleIndexA);
+		const int32 NumVerticesA = UsdTrack->GetSampleInfo(TimeA, false).NumVertices;
+
+		const float TimeB = UsdTrack->GetTimeFromSampleIndex(SampleIndexB);
+		const int32 NumVerticesB = UsdTrack->GetSampleInfo(TimeB, false).NumVertices;
 
 		return NumVerticesA == NumVerticesB;
 	}
@@ -54,14 +63,12 @@ const FVisibilitySample& FGeomCacheTrackUsdProxy::GetVisibilitySample(float Time
 
 void FGeomCacheTrackUsdProxy::FindSampleIndexesFromTime(float Time, bool bLooping, bool bIsPlayingBackwards, int32 &OutFrameIndex, int32 &OutNextFrameIndex, float &InInterpolationFactor)
 {
+	TRACE_CPUPROFILER_EVENT_SCOPE(FGeomCacheTrackUsdProxy::FindSampleIndexesFromTime);
+
 	if (UGeometryCacheTrackUsd* UsdTrack = Cast<UGeometryCacheTrackUsd>(Track))
 	{
-		int32 ThisFrameIndex = UsdTrack->FindSampleIndexFromTime(Time, bLooping);
-		OutFrameIndex = ThisFrameIndex;
+		UsdTrack->GetFractionalFrameIndexFromTime(Time, bLooping, OutFrameIndex, InInterpolationFactor);
 		OutNextFrameIndex = OutFrameIndex + 1;
-		// Clamp the Time (which is the FrameNumber) to the range of the USD GeometryCache track
-		// EndFrameIndex - 1 since last frame index is not included
-		InInterpolationFactor = FMath::Clamp(Time, float(UsdTrack->GetStartFrameIndex()), float(UsdTrack->GetEndFrameIndex() - 1)) - ThisFrameIndex;
 
 		// If playing backwards the logical order of previous and next is reversed
 		if (bIsPlayingBackwards)
diff --git a/Engine/Plugins/Importers/USDImporter/Source/GeometryCacheUSD/Private/GeometryCacheUSDStream.cpp b/Engine/Plugins/Importers/USDImporter/Source/GeometryCacheUSD/Private/GeometryCacheUSDStream.cpp
--- a/Engine/Plugins/Importers/USDImporter/Source/GeometryCacheUSD/Private/GeometryCacheUSDStream.cpp
+++ b/Engine/Plugins/Importers/USDImporter/Source/GeometryCacheUSD/Private/GeometryCacheUSDStream.cpp
@@ -57,11 +57,11 @@ FGeometryCacheUsdStream::FGeometryCacheUsdStream(UGeometryCacheTrackUsd* InUsdTr
 : FGeometryCacheStreamBase(
 	kUsdReadConcurrency,
 	FGeometryCacheStreamDetails{
-		InUsdTrack->GetEndFrameIndex() - InUsdTrack->GetStartFrameIndex() + 1,
-		float((InUsdTrack->GetEndFrameIndex() - InUsdTrack->GetStartFrameIndex() + 1) / InUsdTrack->CurrentStagePinned.GetFramesPerSecond()),
-		float(1.0f / InUsdTrack->CurrentStagePinned.GetFramesPerSecond()),
-		InUsdTrack->GetStartFrameIndex(),
-		InUsdTrack->GetEndFrameIndex()})
+		InUsdTrack->EndFrameIndex - InUsdTrack->StartFrameIndex + 1,
+		float((InUsdTrack->EndFrameIndex - InUsdTrack->StartFrameIndex + 1) / InUsdTrack->FramesPerSecond),
+		float(1.0f / InUsdTrack->FramesPerSecond),
+		InUsdTrack->StartFrameIndex,
+		InUsdTrack->EndFrameIndex})
 , UsdTrack(InUsdTrack)
 , ReadFunc(InReadFunc)
 {
diff --git a/Engine/Plugins/Importers/USDImporter/Source/GeometryCacheUSD/Public/GeometryCacheTrackUSD.h b/Engine/Plugins/Importers/USDImporter/Source/GeometryCacheUSD/Public/GeometryCacheTrackUSD.h
--- a/Engine/Plugins/Importers/USDImporter/Source/GeometryCacheUSD/Public/GeometryCacheTrackUSD.h
+++ b/Engine/Plugins/Importers/USDImporter/Source/GeometryCacheUSD/Public/GeometryCacheTrackUSD.h
@@ -43,9 +43,8 @@ public:
 	//~ End UGeometryCacheTrack Interface.
 
 	const int32 FindSampleIndexFromTime(const float Time, const bool bLooping) const;
-
-	int32 GetStartFrameIndex() const { return StartFrameIndex; }
-	int32 GetEndFrameIndex() const { return EndFrameIndex;  }
+	float GetTimeFromSampleIndex(int32 SampleIndex) const;
+	void GetFractionalFrameIndexFromTime(const float Time, const bool bLooping, int& OutFrameIndex, float& OutFraction) const;
 
 	bool GetMeshData(int32 SampleIndex, FGeometryCacheMeshData& OutMeshData);
 
@@ -57,6 +56,7 @@ public:
 
 public:
 	FName RenderContext;
+	double FramesPerSecond;
 	int32 StartFrameIndex;
 	int32 EndFrameIndex;
 	TMap< FString, TMap< FString, int32 > > MaterialToPrimvarToUVIndex;
