USD: GeometryCache from USD now use tracks in LevelSequence instead of the Time property to drive their animations. Added USD.DisableGeoCacheTracks cvar to revert to the previous behavior.

diff --git a/Engine/Plugins/Importers/USDImporter/Source/USDSchemas/Private/USDGeomMeshTranslator.cpp b/Engine/Plugins/Importers/USDImporter/Source/USDSchemas/Private/USDGeomMeshTranslator.cpp
--- a/Engine/Plugins/Importers/USDImporter/Source/USDSchemas/Private/USDGeomMeshTranslator.cpp
+++ b/Engine/Plugins/Importers/USDImporter/Source/USDSchemas/Private/USDGeomMeshTranslator.cpp
@@ -30,6 +30,7 @@
 #include "GeometryCache.h"
 #include "GeometryCacheComponent.h"
 #include "GroomComponent.h"
+#include "HAL/IConsoleManager.h"
 #include "Interfaces/ITargetPlatform.h"
 #include "Interfaces/ITargetPlatformManagerModule.h"
 #include "Materials/Material.h"
@@ -52,13 +53,13 @@
 #endif // WITH_EDITOR
 
 #include "USDIncludesStart.h"
-	#include "pxr/usd/usd/editContext.h"
-	#include "pxr/usd/usd/prim.h"
-	#include "pxr/usd/usd/stage.h"
-	#include "pxr/usd/usd/typed.h"
-	#include "pxr/usd/usdGeom/mesh.h"
-	#include "pxr/usd/usdGeom/xformable.h"
-	#include "pxr/usd/usdShade/material.h"
+#include "pxr/usd/usd/editContext.h"
+#include "pxr/usd/usd/prim.h"
+#include "pxr/usd/usd/stage.h"
+#include "pxr/usd/usd/typed.h"
+#include "pxr/usd/usdGeom/mesh.h"
+#include "pxr/usd/usdGeom/xformable.h"
+#include "pxr/usd/usdShade/material.h"
 #include "USDIncludesEnd.h"
 
 #if WITH_EDITOR  // The GeometryCacheStreamer module is editor-only
@@ -74,9 +75,9 @@ static FAutoConsoleVariableRef CVarUsdUseGeometryCache(
 
 static float GMeshNormalRepairThreshold = 0.05f;
 static FAutoConsoleVariableRef CVarMeshNormalRepairThreshold(
-	TEXT( "USD.MeshNormalRepairThreshold" ),
+	TEXT("USD.MeshNormalRepairThreshold"),
 	GMeshNormalRepairThreshold,
-	TEXT( "We will try repairing up to this fraction of a Mesh's normals when invalid. If a Mesh has more invalid normals than this, we will recompute all of them. Defaults to 0.05 (5% of all normals)." ) );
+	TEXT("We will try repairing up to this fraction of a Mesh's normals when invalid. If a Mesh has more invalid normals than this, we will recompute all of them. Defaults to 0.05 (5% of all normals)."));
 
 namespace UsdGeomMeshTranslatorImpl
 {
@@ -87,15 +88,15 @@ namespace UsdGeomMeshTranslatorImpl
 		const UE::FUsdPrim& Prim
 	)
 	{
-		if ( LODIndexToMeshDescription.Num() < 1 )
+		if (LODIndexToMeshDescription.Num() < 1)
 		{
 			return false;
 		}
 
 		FScopedUsdAllocs UsdAllocs;
 
-		pxr::UsdPrim UsdPrim{ Prim };
-		if ( !UsdPrim )
+		pxr::UsdPrim UsdPrim{Prim};
+		if (!UsdPrim)
 		{
 			return false;
 		}
@@ -103,27 +104,27 @@ namespace UsdGeomMeshTranslatorImpl
 		bool bHasNaniteOverrideEnabled = false;
 
 		// We want Nanite because of an override on Prim
-		if ( pxr::UsdAttribute NaniteOverride = UsdPrim.GetAttribute( UnrealIdentifiers::UnrealNaniteOverride ) )
+		if (pxr::UsdAttribute NaniteOverride = UsdPrim.GetAttribute(UnrealIdentifiers::UnrealNaniteOverride))
 		{
 			pxr::TfToken OverrideValue;
-			if ( NaniteOverride.Get( &OverrideValue ) )
+			if (NaniteOverride.Get(&OverrideValue))
 			{
-				if ( OverrideValue == UnrealIdentifiers::UnrealNaniteOverrideEnable )
+				if (OverrideValue == UnrealIdentifiers::UnrealNaniteOverrideEnable)
 				{
 					bHasNaniteOverrideEnabled = true;
-					UE_LOG( LogUsd, Log, TEXT( "Trying to enable Nanite for mesh generated for prim '%s' as the '%s' attribute is set to '%s'" ),
+					UE_LOG(LogUsd, Log, TEXT("Trying to enable Nanite for mesh generated for prim '%s' as the '%s' attribute is set to '%s'"),
 						*Prim.GetPrimPath().GetString(),
-						*UsdToUnreal::ConvertToken( UnrealIdentifiers::UnrealNaniteOverride ),
-						*UsdToUnreal::ConvertToken( UnrealIdentifiers::UnrealNaniteOverrideEnable )
+						*UsdToUnreal::ConvertToken(UnrealIdentifiers::UnrealNaniteOverride),
+						*UsdToUnreal::ConvertToken(UnrealIdentifiers::UnrealNaniteOverrideEnable)
 					);
 
 				}
-				else if ( OverrideValue == UnrealIdentifiers::UnrealNaniteOverrideDisable )
+				else if (OverrideValue == UnrealIdentifiers::UnrealNaniteOverrideDisable)
 				{
-					UE_LOG( LogUsd, Log, TEXT( "Not enabling Nanite for mesh generated for prim '%s' as the '%s' attribute is set to '%s'" ),
+					UE_LOG(LogUsd, Log, TEXT("Not enabling Nanite for mesh generated for prim '%s' as the '%s' attribute is set to '%s'"),
 						*Prim.GetPrimPath().GetString(),
-						*UsdToUnreal::ConvertToken( UnrealIdentifiers::UnrealNaniteOverride ),
-						*UsdToUnreal::ConvertToken( UnrealIdentifiers::UnrealNaniteOverrideDisable )
+						*UsdToUnreal::ConvertToken(UnrealIdentifiers::UnrealNaniteOverride),
+						*UsdToUnreal::ConvertToken(UnrealIdentifiers::UnrealNaniteOverrideDisable)
 					);
 					return false;
 				}
@@ -131,12 +132,12 @@ namespace UsdGeomMeshTranslatorImpl
 		}
 
 		// We want Nanite because the mesh is large enough for the threshold, which is set to something valid
-		if ( !bHasNaniteOverrideEnabled )
+		if (!bHasNaniteOverrideEnabled)
 		{
-			const int32 NumTriangles = LODIndexToMeshDescription[ 0 ].Triangles().Num();
-			if ( NumTriangles >= Context.NaniteTriangleThreshold )
+			const int32 NumTriangles = LODIndexToMeshDescription[0].Triangles().Num();
+			if (NumTriangles >= Context.NaniteTriangleThreshold)
 			{
-				UE_LOG( LogUsd, Verbose, TEXT( "Trying to enable Nanite for mesh generated for prim '%s' as it has '%d' triangles, and the threshold is '%d'" ),
+				UE_LOG(LogUsd, Verbose, TEXT("Trying to enable Nanite for mesh generated for prim '%s' as it has '%d' triangles, and the threshold is '%d'"),
 					*Prim.GetPrimPath().GetString(),
 					NumTriangles,
 					Context.NaniteTriangleThreshold
@@ -144,7 +145,7 @@ namespace UsdGeomMeshTranslatorImpl
 			}
 			else
 			{
-				UE_LOG( LogUsd, Verbose, TEXT( "Not enabling Nanite for mesh generated for prim '%s' as it has '%d' triangles, and the threshold is '%d'" ),
+				UE_LOG(LogUsd, Verbose, TEXT("Not enabling Nanite for mesh generated for prim '%s' as it has '%d' triangles, and the threshold is '%d'"),
 					*Prim.GetPrimPath().GetString(),
 					NumTriangles,
 					Context.NaniteTriangleThreshold
@@ -158,22 +159,22 @@ namespace UsdGeomMeshTranslatorImpl
 		// likely don't want Nanite for one of the LOD meshes anyway, as that doesn't really make any sense.
 		// If the user wants to have Nanite within the variant set all they would otherwise need is to name the variant set something
 		// else other than LOD.
-		if ( LODIndexToMeshDescription.Num() > 1 )
+		if (LODIndexToMeshDescription.Num() > 1)
 		{
-			UE_LOG( LogUsd, Warning, TEXT( "Not enabling Nanite for mesh generated for prim '%s' as it has more than one generated LOD (and so came from a LOD variant set setup)" ),
+			UE_LOG(LogUsd, Warning, TEXT("Not enabling Nanite for mesh generated for prim '%s' as it has more than one generated LOD (and so came from a LOD variant set setup)"),
 				*Prim.GetPrimPath().GetString()
 			);
 			return false;
 		}
 
 		// We cannot enable Nanite here due to a limit on the number of material slots
-		if ( LODIndexToMaterialInfo.Num() > 0 )
+		if (LODIndexToMaterialInfo.Num() > 0)
 		{
 			const int32 NumSections = LODIndexToMaterialInfo[0].Slots.Num();
 			const int32 MaxNumSections = 64; // There is no define for this, but it's checked for on NaniteBuilder.cpp, FBuilderModule::Build
-			if ( NumSections > MaxNumSections )
+			if (NumSections > MaxNumSections)
 			{
-				UE_LOG( LogUsd, Warning, TEXT( "Not enabling Nanite for mesh generated for prim '%s' as LOD0 has '%d' material slots, which is above the Nanite limit of '%d'" ),
+				UE_LOG(LogUsd, Warning, TEXT("Not enabling Nanite for mesh generated for prim '%s' as LOD0 has '%d' material slots, which is above the Nanite limit of '%d'"),
 					*Prim.GetPrimPath().GetString(),
 					NumSections,
 					MaxNumSections
@@ -183,7 +184,7 @@ namespace UsdGeomMeshTranslatorImpl
 		}
 
 #if !WITH_EDITOR
-		UE_LOG( LogUsd, Warning, TEXT( "Not enabling Nanite for mesh generated for prim '%s' as we can't setup Nanite during runtime" ),
+		UE_LOG(LogUsd, Warning, TEXT("Not enabling Nanite for mesh generated for prim '%s' as we can't setup Nanite during runtime"),
 			*Prim.GetPrimPath().GetString()
 		);
 		return false;
@@ -192,24 +193,24 @@ namespace UsdGeomMeshTranslatorImpl
 		return true;
 	}
 
-	bool IsAnimated( const pxr::UsdPrim& Prim )
+	bool IsAnimated(const pxr::UsdPrim& Prim)
 	{
 		FScopedUsdAllocs UsdAllocs;
 
 		bool bHasAttributesTimeSamples = false;
 		{
 			constexpr bool bIncludeInherited = false;
-			pxr::TfTokenVector GeomMeshAttributeNames = pxr::UsdGeomMesh::GetSchemaAttributeNames( bIncludeInherited );
-			pxr::TfTokenVector GeomPointBasedAttributeNames = pxr::UsdGeomPointBased::GetSchemaAttributeNames( bIncludeInherited );
+			pxr::TfTokenVector GeomMeshAttributeNames = pxr::UsdGeomMesh::GetSchemaAttributeNames(bIncludeInherited);
+			pxr::TfTokenVector GeomPointBasedAttributeNames = pxr::UsdGeomPointBased::GetSchemaAttributeNames(bIncludeInherited);
 
-			GeomMeshAttributeNames.reserve( GeomMeshAttributeNames.size() + GeomPointBasedAttributeNames.size() );
-			GeomMeshAttributeNames.insert( GeomMeshAttributeNames.end(), GeomPointBasedAttributeNames.begin(), GeomPointBasedAttributeNames.end() );
+			GeomMeshAttributeNames.reserve(GeomMeshAttributeNames.size() + GeomPointBasedAttributeNames.size());
+			GeomMeshAttributeNames.insert(GeomMeshAttributeNames.end(), GeomPointBasedAttributeNames.begin(), GeomPointBasedAttributeNames.end());
 
-			for ( const pxr::TfToken& AttributeName : GeomMeshAttributeNames )
+			for (const pxr::TfToken& AttributeName : GeomMeshAttributeNames)
 			{
-				const pxr::UsdAttribute& Attribute = Prim.GetAttribute( AttributeName );
+				const pxr::UsdAttribute& Attribute = Prim.GetAttribute(AttributeName);
 
-				if ( Attribute && Attribute.ValueMightBeTimeVarying() )
+				if (Attribute && Attribute.ValueMightBeTimeVarying())
 				{
 					bHasAttributesTimeSamples = true;
 					break;
@@ -221,18 +222,18 @@ namespace UsdGeomMeshTranslatorImpl
 	}
 
 	/** Returns true if material infos have changed on the StaticMesh */
-	bool ProcessMaterials( const pxr::UsdPrim& UsdPrim, const TArray<UsdUtils::FUsdPrimMaterialAssignmentInfo>& LODIndexToMaterialInfo, UStaticMesh& StaticMesh, UUsdAssetCache& AssetCache, float Time, EObjectFlags Flags )
+	bool ProcessMaterials(const pxr::UsdPrim& UsdPrim, const TArray<UsdUtils::FUsdPrimMaterialAssignmentInfo>& LODIndexToMaterialInfo, UStaticMesh& StaticMesh, UUsdAssetCache& AssetCache, float Time, EObjectFlags Flags)
 	{
 		bool bMaterialAssignementsHaveChanged = false;
 
 		TArray<UMaterialInterface*> ExistingAssignments;
-		for ( const FStaticMaterial& StaticMaterial : StaticMesh.GetStaticMaterials() )
+		for (const FStaticMaterial& StaticMaterial : StaticMesh.GetStaticMaterials())
 		{
 			ExistingAssignments.Add(StaticMaterial.MaterialInterface);
 		}
 
 #if WITH_EDITOR
-		UUsdMeshAssetImportData* ImportData = Cast< UUsdMeshAssetImportData>( StaticMesh.AssetImportData.Get() );
+		UUsdMeshAssetImportData* ImportData = Cast< UUsdMeshAssetImportData>(StaticMesh.AssetImportData.Get());
 #endif // WITH_EDITOR
 
 		TMap<const UsdUtils::FUsdPrimMaterialSlot*, UMaterialInterface*> ResolvedMaterials = MeshTranslationImpl::ResolveMaterialAssignmentInfo(
@@ -243,16 +244,16 @@ namespace UsdGeomMeshTranslatorImpl
 		);
 
 		uint32 StaticMeshSlotIndex = 0;
-		for ( int32 LODIndex = 0; LODIndex < LODIndexToMaterialInfo.Num(); ++LODIndex )
+		for (int32 LODIndex = 0; LODIndex < LODIndexToMaterialInfo.Num(); ++LODIndex)
 		{
-			const TArray< UsdUtils::FUsdPrimMaterialSlot >& LODSlots = LODIndexToMaterialInfo[ LODIndex ].Slots;
+			const TArray< UsdUtils::FUsdPrimMaterialSlot >& LODSlots = LODIndexToMaterialInfo[LODIndex].Slots;
 
-			for ( int32 LODSlotIndex = 0; LODSlotIndex < LODSlots.Num(); ++LODSlotIndex, ++StaticMeshSlotIndex )
+			for (int32 LODSlotIndex = 0; LODSlotIndex < LODSlots.Num(); ++LODSlotIndex, ++StaticMeshSlotIndex)
 			{
-				const UsdUtils::FUsdPrimMaterialSlot& Slot = LODSlots[ LODSlotIndex ];
+				const UsdUtils::FUsdPrimMaterialSlot& Slot = LODSlots[LODSlotIndex];
 
 				UMaterialInterface* Material = nullptr;
-				if ( UMaterialInterface** FoundMaterial = ResolvedMaterials.Find( &Slot ) )
+				if (UMaterialInterface** FoundMaterial = ResolvedMaterials.Find(&Slot))
 				{
 					Material = *FoundMaterial;
 				}
@@ -263,34 +264,34 @@ namespace UsdGeomMeshTranslatorImpl
 				}
 
 				// Create and set the static material
-				FStaticMaterial StaticMaterial( Material, *LexToString( StaticMeshSlotIndex ) );
-				if ( !StaticMesh.GetStaticMaterials().IsValidIndex( StaticMeshSlotIndex ) )
+				FStaticMaterial StaticMaterial(Material, *LexToString(StaticMeshSlotIndex));
+				if (!StaticMesh.GetStaticMaterials().IsValidIndex(StaticMeshSlotIndex))
 				{
-					StaticMesh.GetStaticMaterials().Add( MoveTemp( StaticMaterial ) );
+					StaticMesh.GetStaticMaterials().Add(MoveTemp(StaticMaterial));
 					bMaterialAssignementsHaveChanged = true;
 				}
-				else if ( !( StaticMesh.GetStaticMaterials()[ StaticMeshSlotIndex ] == StaticMaterial ) )
+				else if (!(StaticMesh.GetStaticMaterials()[StaticMeshSlotIndex] == StaticMaterial))
 				{
-					StaticMesh.GetStaticMaterials()[ StaticMeshSlotIndex ] = MoveTemp( StaticMaterial );
+					StaticMesh.GetStaticMaterials()[StaticMeshSlotIndex] = MoveTemp(StaticMaterial);
 					bMaterialAssignementsHaveChanged = true;
 				}
 
 #if WITH_EDITOR
-				if ( ImportData )
+				if (ImportData)
 				{
-					ImportData->MaterialSlotToPrimPaths.FindOrAdd( StaticMeshSlotIndex ).PrimPaths = Slot.PrimPaths.Array();
+					ImportData->MaterialSlotToPrimPaths.FindOrAdd(StaticMeshSlotIndex).PrimPaths = Slot.PrimPaths.Array();
 				}
 
 				// Setup the section map so that our LOD material index is properly mapped to the static mesh material index
 				// At runtime we don't ever parse these variants as LODs so we don't need this
-				if ( StaticMesh.GetSectionInfoMap().IsValidSection( LODIndex, LODSlotIndex ) )
+				if (StaticMesh.GetSectionInfoMap().IsValidSection(LODIndex, LODSlotIndex))
 				{
-					FMeshSectionInfo MeshSectionInfo = StaticMesh.GetSectionInfoMap().Get( LODIndex, LODSlotIndex );
+					FMeshSectionInfo MeshSectionInfo = StaticMesh.GetSectionInfoMap().Get(LODIndex, LODSlotIndex);
 
-					if ( MeshSectionInfo.MaterialIndex != StaticMeshSlotIndex )
+					if (MeshSectionInfo.MaterialIndex != StaticMeshSlotIndex)
 					{
 						MeshSectionInfo.MaterialIndex = StaticMeshSlotIndex;
-						StaticMesh.GetSectionInfoMap().Set( LODIndex, LODSlotIndex, MeshSectionInfo );
+						StaticMesh.GetSectionInfoMap().Set(LODIndex, LODSlotIndex, MeshSectionInfo);
 
 						bMaterialAssignementsHaveChanged = true;
 					}
@@ -300,7 +301,7 @@ namespace UsdGeomMeshTranslatorImpl
 					FMeshSectionInfo MeshSectionInfo;
 					MeshSectionInfo.MaterialIndex = StaticMeshSlotIndex;
 
-					StaticMesh.GetSectionInfoMap().Set( LODIndex, LODSlotIndex, MeshSectionInfo );
+					StaticMesh.GetSectionInfoMap().Set(LODIndex, LODSlotIndex, MeshSectionInfo);
 
 					bMaterialAssignementsHaveChanged = true;
 				}
@@ -309,7 +310,7 @@ namespace UsdGeomMeshTranslatorImpl
 		}
 
 #if WITH_EDITOR
-		StaticMesh.GetOriginalSectionInfoMap().CopyFrom( StaticMesh.GetSectionInfoMap() );
+		StaticMesh.GetOriginalSectionInfoMap().CopyFrom(StaticMesh.GetSectionInfoMap());
 #endif // WITH_EDITOR
 
 		return bMaterialAssignementsHaveChanged;
@@ -317,57 +318,57 @@ namespace UsdGeomMeshTranslatorImpl
 
 #if WITH_EDITOR
 	// #ueent_todo: Merge the code with ProcessMaterials
-	bool ProcessGeometryCacheMaterials( const pxr::UsdPrim& UsdPrim, const TArray< UsdUtils::FUsdPrimMaterialAssignmentInfo >& LODIndexToMaterialInfo, UGeometryCache& GeometryCache, UUsdAssetCache& AssetCache, float Time, EObjectFlags Flags)
+	bool ProcessGeometryCacheMaterials(const pxr::UsdPrim& UsdPrim, const TArray< UsdUtils::FUsdPrimMaterialAssignmentInfo >& LODIndexToMaterialInfo, UGeometryCache& GeometryCache, UUsdAssetCache& AssetCache, float Time, EObjectFlags Flags)
 	{
 		bool bMaterialAssignementsHaveChanged = false;
 
 #if WITH_EDITOR
-		UUsdMeshAssetImportData* MeshImportData = Cast< UUsdMeshAssetImportData>( GeometryCache.AssetImportData.Get() );
+		UUsdMeshAssetImportData* MeshImportData = Cast< UUsdMeshAssetImportData>(GeometryCache.AssetImportData.Get());
 #endif // WITH_EDITOR
 
 		uint32 StaticMeshSlotIndex = 0;
-		for ( int32 LODIndex = 0; LODIndex < LODIndexToMaterialInfo.Num(); ++LODIndex )
+		for (int32 LODIndex = 0; LODIndex < LODIndexToMaterialInfo.Num(); ++LODIndex)
 		{
-			const TArray< UsdUtils::FUsdPrimMaterialSlot >& LODSlots = LODIndexToMaterialInfo[ LODIndex ].Slots;
+			const TArray< UsdUtils::FUsdPrimMaterialSlot >& LODSlots = LODIndexToMaterialInfo[LODIndex].Slots;
 
-			for ( int32 LODSlotIndex = 0; LODSlotIndex < LODSlots.Num(); ++LODSlotIndex, ++StaticMeshSlotIndex )
+			for (int32 LODSlotIndex = 0; LODSlotIndex < LODSlots.Num(); ++LODSlotIndex, ++StaticMeshSlotIndex)
 			{
-				const UsdUtils::FUsdPrimMaterialSlot& Slot = LODSlots[ LODSlotIndex ];
+				const UsdUtils::FUsdPrimMaterialSlot& Slot = LODSlots[LODSlotIndex];
 				UMaterialInterface* Material = nullptr;
 
-				switch ( Slot.AssignmentType )
+				switch (Slot.AssignmentType)
 				{
 				case UsdUtils::EPrimAssignmentType::DisplayColor:
 				{
 					FScopedUsdAllocs Allocs;
 
 					// Try reusing an already created DisplayColor material
-					if ( UMaterialInstanceConstant* ExistingMaterial = Cast< UMaterialInstanceConstant >( AssetCache.GetCachedAsset( Slot.MaterialSource ) ) )
+					if (UMaterialInstanceConstant* ExistingMaterial = Cast< UMaterialInstanceConstant >(AssetCache.GetCachedAsset(Slot.MaterialSource)))
 					{
 						Material = ExistingMaterial;
 					}
 
 					// Need to create a new DisplayColor material
-					if ( Material == nullptr )
+					if (Material == nullptr)
 					{
-						if ( TOptional< UsdUtils::FDisplayColorMaterial > DisplayColorDesc = UsdUtils::FDisplayColorMaterial::FromString( Slot.MaterialSource ) )
+						if (TOptional< UsdUtils::FDisplayColorMaterial > DisplayColorDesc = UsdUtils::FDisplayColorMaterial::FromString(Slot.MaterialSource))
 						{
 							UMaterialInstance* MaterialInstance = nullptr;
 
-							if ( GIsEditor )  // Editor, PIE => true; Standlone, packaged => false
+							if (GIsEditor)  // Editor, PIE => true; Standlone, packaged => false
 							{
-								MaterialInstance = UsdUtils::CreateDisplayColorMaterialInstanceConstant( DisplayColorDesc.GetValue() );
+								MaterialInstance = UsdUtils::CreateDisplayColorMaterialInstanceConstant(DisplayColorDesc.GetValue());
 
 								// Leave PrimPath as empty as it likely will be reused by many prims
-								UUsdAssetImportData* ImportData = NewObject< UUsdAssetImportData >( MaterialInstance, TEXT( "USDAssetImportData" ) );
+								UUsdAssetImportData* ImportData = NewObject< UUsdAssetImportData >(MaterialInstance, TEXT("USDAssetImportData"));
 								MaterialInstance->AssetImportData = ImportData;
 							}
 							else
 							{
-								MaterialInstance = UsdUtils::CreateDisplayColorMaterialInstanceDynamic( DisplayColorDesc.GetValue() );
+								MaterialInstance = UsdUtils::CreateDisplayColorMaterialInstanceDynamic(DisplayColorDesc.GetValue());
 							}
 
-							AssetCache.CacheAsset( Slot.MaterialSource, MaterialInstance );
+							AssetCache.CacheAsset(Slot.MaterialSource, MaterialInstance);
 							Material = MaterialInstance;
 						}
 					}
@@ -379,15 +380,15 @@ namespace UsdGeomMeshTranslatorImpl
 					FScopedUsdAllocs Allocs;
 
 					// Check first or else we may get a warning
-					if ( pxr::SdfPath::IsValidPathString( UnrealToUsd::ConvertString( *Slot.MaterialSource ).Get() ) )
+					if (pxr::SdfPath::IsValidPathString(UnrealToUsd::ConvertString(*Slot.MaterialSource).Get()))
 					{
-						pxr::SdfPath MaterialPrimPath = UnrealToUsd::ConvertPath( *Slot.MaterialSource ).Get();
+						pxr::SdfPath MaterialPrimPath = UnrealToUsd::ConvertPath(*Slot.MaterialSource).Get();
 
 						// TODO: This may break if MaterialPrimPath targets a prim inside a LOD variant that is disabled...
-						TUsdStore< pxr::UsdPrim > MaterialPrim = UsdPrim.GetStage()->GetPrimAtPath( MaterialPrimPath );
-						if ( MaterialPrim.Get() )
+						TUsdStore< pxr::UsdPrim > MaterialPrim = UsdPrim.GetStage()->GetPrimAtPath(MaterialPrimPath);
+						if (MaterialPrim.Get())
 						{
-							Material = Cast< UMaterialInterface >( AssetCache.GetAssetForPrim( UsdToUnreal::ConvertPath( MaterialPrim.Get().GetPrimPath() ) ) );
+							Material = Cast< UMaterialInterface >(AssetCache.GetAssetForPrim(UsdToUnreal::ConvertPath(MaterialPrim.Get().GetPrimPath())));
 						}
 					}
 
@@ -395,67 +396,67 @@ namespace UsdGeomMeshTranslatorImpl
 				}
 				case UsdUtils::EPrimAssignmentType::UnrealMaterial:
 				{
-					Material = Cast< UMaterialInterface >( FSoftObjectPath( Slot.MaterialSource ).TryLoad() );
+					Material = Cast< UMaterialInterface >(FSoftObjectPath(Slot.MaterialSource).TryLoad());
 					break;
 				}
 				case UsdUtils::EPrimAssignmentType::None:
 				default:
 				{
-					ensure( false );
+					ensure(false);
 					break;
 				}
 				}
 
 				// Fallback to this UsdGeomMesh DisplayColor material if present
-				if ( Material == nullptr )
+				if (Material == nullptr)
 				{
 					FScopedUsdAllocs Allocs;
 
 					// Try reusing an already created DisplayColor material
-					if ( UMaterialInstanceConstant* ExistingMaterial = Cast< UMaterialInstanceConstant >( AssetCache.GetCachedAsset( Slot.MaterialSource ) ) )
+					if (UMaterialInstanceConstant* ExistingMaterial = Cast< UMaterialInstanceConstant >(AssetCache.GetCachedAsset(Slot.MaterialSource)))
 					{
 						Material = ExistingMaterial;
 					}
 
 					// Need to create a new DisplayColor material
-					if ( Material == nullptr )
+					if (Material == nullptr)
 					{
-						if ( TOptional< UsdUtils::FDisplayColorMaterial > DisplayColorDesc = UsdUtils::ExtractDisplayColorMaterial( pxr::UsdGeomMesh( UsdPrim ) ) )
+						if (TOptional< UsdUtils::FDisplayColorMaterial > DisplayColorDesc = UsdUtils::ExtractDisplayColorMaterial(pxr::UsdGeomMesh(UsdPrim)))
 						{
 							UMaterialInstance* MaterialInstance = nullptr;
 
-							if ( GIsEditor )  // Editor, PIE => true; Standlone, packaged => false
+							if (GIsEditor)  // Editor, PIE => true; Standlone, packaged => false
 							{
-								MaterialInstance = UsdUtils::CreateDisplayColorMaterialInstanceConstant( DisplayColorDesc.GetValue() );
+								MaterialInstance = UsdUtils::CreateDisplayColorMaterialInstanceConstant(DisplayColorDesc.GetValue());
 
 								// Leave PrimPath as empty as it likely will be reused by many prims
-								UUsdAssetImportData* ImportData = NewObject< UUsdAssetImportData >( MaterialInstance, TEXT( "USDAssetImportData" ) );
+								UUsdAssetImportData* ImportData = NewObject< UUsdAssetImportData >(MaterialInstance, TEXT("USDAssetImportData"));
 								MaterialInstance->AssetImportData = ImportData;
 							}
 							else
 							{
-								MaterialInstance = UsdUtils::CreateDisplayColorMaterialInstanceDynamic( DisplayColorDesc.GetValue() );
+								MaterialInstance = UsdUtils::CreateDisplayColorMaterialInstanceDynamic(DisplayColorDesc.GetValue());
 							}
 
-							AssetCache.CacheAsset( Slot.MaterialSource, MaterialInstance );
+							AssetCache.CacheAsset(Slot.MaterialSource, MaterialInstance);
 							Material = MaterialInstance;
 						}
 					}
 				}
 
-				if ( !GeometryCache.Materials.IsValidIndex( StaticMeshSlotIndex ) )
+				if (!GeometryCache.Materials.IsValidIndex(StaticMeshSlotIndex))
 				{
-					GeometryCache.Materials.Add( Material );
+					GeometryCache.Materials.Add(Material);
 					bMaterialAssignementsHaveChanged = true;
 				}
-				else if ( !( GeometryCache.Materials[ StaticMeshSlotIndex ] == Material ) )
+				else if (!(GeometryCache.Materials[StaticMeshSlotIndex] == Material))
 				{
-					GeometryCache.Materials[ StaticMeshSlotIndex ] = Material;
+					GeometryCache.Materials[StaticMeshSlotIndex] = Material;
 					bMaterialAssignementsHaveChanged = true;
 				}
 
 #if WITH_EDITOR
-				MeshImportData->MaterialSlotToPrimPaths.FindOrAdd( StaticMeshSlotIndex ).PrimPaths = Slot.PrimPaths.Array();
+				MeshImportData->MaterialSlotToPrimPaths.FindOrAdd(StaticMeshSlotIndex).PrimPaths = Slot.PrimPaths.Array();
 #endif // WITH_EDITOR
 			}
 		}
@@ -478,13 +479,13 @@ namespace UsdGeomMeshTranslatorImpl
 		FScopedUsdAllocs Allocs;
 
 		pxr::UsdPrim UsdMeshPrim = UsdMesh.GetPrim();
-		if ( !UsdMeshPrim )
+		if (!UsdMeshPrim)
 		{
 			return false;
 		}
 
 		pxr::UsdPrim ParentPrim = UsdMeshPrim.GetParent();
-		if ( !ParentPrim )
+		if (!ParentPrim)
 		{
 			return false;
 		}
@@ -494,12 +495,12 @@ namespace UsdGeomMeshTranslatorImpl
 
 		UsdToUnreal::FUsdMeshConversionOptions OptionsCopy = Options;
 
-		TFunction<bool( const pxr::UsdGeomMesh&, int32 )> ConvertLOD = [ & ]( const pxr::UsdGeomMesh& LODMesh, int32 LODIndex )
+		TFunction<bool(const pxr::UsdGeomMesh&, int32)> ConvertLOD = [&](const pxr::UsdGeomMesh& LODMesh, int32 LODIndex)
 		{
 			FMeshDescription TempMeshDescription;
 			UsdUtils::FUsdPrimMaterialAssignmentInfo TempMaterialInfo;
 
-			FStaticMeshAttributes StaticMeshAttributes( TempMeshDescription );
+			FStaticMeshAttributes StaticMeshAttributes(TempMeshDescription);
 			StaticMeshAttributes.Register();
 
 			bool bSuccess = true;
@@ -509,7 +510,7 @@ namespace UsdGeomMeshTranslatorImpl
 			// care if this mesh in particular has been marked as invisible
 			pxr::TfToken Visibility;
 			pxr::UsdAttribute VisibilityAttr = LODMesh.GetVisibilityAttr();
-			if ( VisibilityAttr && VisibilityAttr.Get( &Visibility, Options.TimeCode ) && Visibility == pxr::UsdGeomTokens->inherited )
+			if (VisibilityAttr && VisibilityAttr.Get(&Visibility, Options.TimeCode) && Visibility == pxr::UsdGeomTokens->inherited)
 			{
 				// If we're interpreting LODs we must bake the transform from each LOD Mesh into the vertices, because there's no guarantee
 				// all LODs have the same transform, so we can't just put the transforms directly on the component. If we are not interpreting
@@ -517,9 +518,9 @@ namespace UsdGeomMeshTranslatorImpl
 				// TODO: Handle resetXformOp here
 				bool bResetXformStack = false;
 				FTransform MeshTransform = FTransform::Identity;
-				bSuccess &= UsdToUnreal::ConvertXformable( LODMesh.GetPrim().GetStage(), LODMesh, MeshTransform, Options.TimeCode.GetValue(), &bResetXformStack );
+				bSuccess &= UsdToUnreal::ConvertXformable(LODMesh.GetPrim().GetStage(), LODMesh, MeshTransform, Options.TimeCode.GetValue(), &bResetXformStack);
 
-				if ( bSuccess )
+				if (bSuccess)
 				{
 					OptionsCopy.AdditionalTransform = MeshTransform * Options.AdditionalTransform;
 					OptionsCopy.bMergeIdenticalMaterialSlots = false;  // We only merge slots when collapsing, and we never collapse LODs
@@ -533,25 +534,25 @@ namespace UsdGeomMeshTranslatorImpl
 				}
 			}
 
-			if ( bSuccess )
+			if (bSuccess)
 			{
-				LODIndexToMeshDescriptionMap.Add( LODIndex, MoveTemp( TempMeshDescription ) );
-				LODIndexToMaterialInfoMap.Add( LODIndex, MoveTemp( TempMaterialInfo ) );
+				LODIndexToMeshDescriptionMap.Add(LODIndex, MoveTemp(TempMeshDescription));
+				LODIndexToMaterialInfoMap.Add(LODIndex, MoveTemp(TempMaterialInfo));
 			}
 
 			return true;
 		};
-		bool bFoundLODs = UsdUtils::IterateLODMeshes( ParentPrim, ConvertLOD );
+		bool bFoundLODs = UsdUtils::IterateLODMeshes(ParentPrim, ConvertLOD);
 
 		// Place them in order as we can't have e.g. LOD0 and LOD2 without LOD1, and there's no reason downstream code needs to care about this
-		OutLODIndexToMeshDescription.Reset( LODIndexToMeshDescriptionMap.Num() );
-		OutLODIndexToMaterialInfo.Reset( LODIndexToMaterialInfoMap.Num() );
-		LODIndexToMeshDescriptionMap.KeySort( TLess<int32>() );
-		for ( TPair<int32, FMeshDescription>& Entry : LODIndexToMeshDescriptionMap )
+		OutLODIndexToMeshDescription.Reset(LODIndexToMeshDescriptionMap.Num());
+		OutLODIndexToMaterialInfo.Reset(LODIndexToMaterialInfoMap.Num());
+		LODIndexToMeshDescriptionMap.KeySort(TLess<int32>());
+		for (TPair<int32, FMeshDescription>& Entry : LODIndexToMeshDescriptionMap)
 		{
 			const int32 OldLODIndex = Entry.Key;
-			OutLODIndexToMeshDescription.Add( MoveTemp( Entry.Value ) );
-			OutLODIndexToMaterialInfo.Add( MoveTemp( LODIndexToMaterialInfoMap[ OldLODIndex ] ) );
+			OutLODIndexToMeshDescription.Add(MoveTemp(Entry.Value));
+			OutLODIndexToMaterialInfo.Add(MoveTemp(LODIndexToMaterialInfoMap[OldLODIndex]));
 		}
 
 		return bFoundLODs;
@@ -565,7 +566,7 @@ namespace UsdGeomMeshTranslatorImpl
 		bool bInterpretLODs = false
 	)
 	{
-		if ( !UsdMesh )
+		if (!UsdMesh)
 		{
 			return;
 		}
@@ -577,13 +578,13 @@ namespace UsdGeomMeshTranslatorImpl
 		pxr::SdfPath Path = Prim.GetPrimPath();
 
 		bool bInterpretedLODs = false;
-		if ( bInterpretLODs )
+		if (bInterpretLODs)
 		{
-			bInterpretedLODs = TryLoadingMultipleLODs( UsdMesh, OutLODIndexToMeshDescription, OutLODIndexToMaterialInfo, Options );
+			bInterpretedLODs = TryLoadingMultipleLODs(UsdMesh, OutLODIndexToMeshDescription, OutLODIndexToMaterialInfo, Options);
 
 			// Have to be very careful here as flipping through LODs invalidates prim references, so we need to
 			// re-acquire them
-			UsdMesh = pxr::UsdGeomMesh{ Stage->GetPrimAtPath( Path ) };
+			UsdMesh = pxr::UsdGeomMesh{Stage->GetPrimAtPath(Path)};
 		}
 
 		// If we've managed to interpret LODs, we won't place our mesh transform on the static mesh component itself
@@ -591,98 +592,98 @@ namespace UsdGeomMeshTranslatorImpl
 		// So here we do that
 		bool bSuccess = true;
 		FTransform MeshTransform = FTransform::Identity;
-		if ( bInterpretedLODs && OutLODIndexToMeshDescription.Num() > 1 )
+		if (bInterpretedLODs && OutLODIndexToMeshDescription.Num() > 1)
 		{
 			// TODO: Handle resetXformOp here
 			bool bResetXformStack = false;
-			bSuccess &= UsdToUnreal::ConvertXformable( Stage, UsdMesh, MeshTransform, Options.TimeCode.GetValue(), &bResetXformStack );
+			bSuccess &= UsdToUnreal::ConvertXformable(Stage, UsdMesh, MeshTransform, Options.TimeCode.GetValue(), &bResetXformStack);
 		}
 
-		if ( !bInterpretedLODs )
+		if (!bInterpretedLODs)
 		{
 			FMeshDescription TempMeshDescription;
 			UsdUtils::FUsdPrimMaterialAssignmentInfo TempMaterialInfo;
 
-			FStaticMeshAttributes StaticMeshAttributes( TempMeshDescription );
+			FStaticMeshAttributes StaticMeshAttributes(TempMeshDescription);
 			StaticMeshAttributes.Register();
 
-			if ( bSuccess )
+			if (bSuccess)
 			{
 				UsdToUnreal::FUsdMeshConversionOptions OptionsCopy = Options;
 				OptionsCopy.AdditionalTransform = MeshTransform * Options.AdditionalTransform;
 				OptionsCopy.bMergeIdenticalMaterialSlots = false;  // We only merge for collapsed meshes
 
 				bSuccess &= UsdToUnreal::ConvertGeomMesh(
-					pxr::UsdGeomMesh{ UsdMesh },
+					pxr::UsdGeomMesh{UsdMesh},
 					TempMeshDescription,
 					TempMaterialInfo,
 					OptionsCopy
 				);
 			}
 
-			if ( bSuccess )
+			if (bSuccess)
 			{
-				OutLODIndexToMeshDescription = { MoveTemp( TempMeshDescription ) };
-				OutLODIndexToMaterialInfo = { MoveTemp( TempMaterialInfo ) };
+				OutLODIndexToMeshDescription = {MoveTemp(TempMeshDescription)};
+				OutLODIndexToMaterialInfo = {MoveTemp(TempMaterialInfo)};
 			}
 		}
 	}
 
-	void RepairNormalsAndTangents( const FString& PrimPath, FMeshDescription& MeshDescription )
+	void RepairNormalsAndTangents(const FString& PrimPath, FMeshDescription& MeshDescription)
 	{
-		FStaticMeshConstAttributes Attributes{ MeshDescription };
+		FStaticMeshConstAttributes Attributes{MeshDescription};
 		TArrayView<const FVector3f> VertexInstanceNormals = Attributes.GetVertexInstanceNormals().GetRawArray();
 
 		// Similar to FStaticMeshOperations::AreNormalsAndTangentsValid but we don't care about tangents since we never
 		// read those from USD
 		uint64 InvalidNormalCount = 0;
-		for ( const FVertexInstanceID VertexInstanceID : MeshDescription.VertexInstances().GetElementIDs() )
+		for (const FVertexInstanceID VertexInstanceID : MeshDescription.VertexInstances().GetElementIDs())
 		{
-			if ( VertexInstanceNormals[ VertexInstanceID ].IsNearlyZero() || VertexInstanceNormals[ VertexInstanceID ].ContainsNaN() )
+			if (VertexInstanceNormals[VertexInstanceID].IsNearlyZero() || VertexInstanceNormals[VertexInstanceID].ContainsNaN())
 			{
 				++InvalidNormalCount;
 			}
 		}
-		if ( InvalidNormalCount == 0 )
+		if (InvalidNormalCount == 0)
 		{
 			return;
 		}
 
-		const float InvalidNormalFraction = ( float ) InvalidNormalCount / ( float ) VertexInstanceNormals.Num();
+		const float InvalidNormalFraction = (float)InvalidNormalCount / (float)VertexInstanceNormals.Num();
 
 		// We always need to do this at this point as ComputeTangentsAndNormals will end up computing tangents anyway
 		// and our triangle tangents are always invalid
-		FStaticMeshOperations::ComputeTriangleTangentsAndNormals( MeshDescription );
+		FStaticMeshOperations::ComputeTriangleTangentsAndNormals(MeshDescription);
 
-		const static FString MeshNormalRepairThresholdText = TEXT( "USD.MeshNormalRepairThreshold" );
+		const static FString MeshNormalRepairThresholdText = TEXT("USD.MeshNormalRepairThreshold");
 
 		// Make sure our normals can be rebuilt from MeshDescription::InitializeAutoGeneratedAttributes in case some tool needs them.
 		// Always force-compute tangents here as we never have them anyway. If we don't force them to be recomputed we'll get
 		// the worst of both worlds as some of these will be arbitrarily recomputed anyway, and some will be left invalid
 		EComputeNTBsFlags Options = EComputeNTBsFlags::UseMikkTSpace | EComputeNTBsFlags::Tangents;
-		if ( InvalidNormalFraction >= GMeshNormalRepairThreshold )
+		if (InvalidNormalFraction >= GMeshNormalRepairThreshold)
 		{
 			Options |= EComputeNTBsFlags::Normals;
-			UE_LOG( LogUsd, Warning, TEXT( "%f%% of the normals from Mesh prim '%s' are invalid. This is at or above the threshold of '%f%%' (configurable via the cvar '%s'), so normals will be discarded and fully recomputed." ),
+			UE_LOG(LogUsd, Warning, TEXT("%f%% of the normals from Mesh prim '%s' are invalid. This is at or above the threshold of '%f%%' (configurable via the cvar '%s'), so normals will be discarded and fully recomputed."),
 				InvalidNormalFraction * 100.0f,
 				*PrimPath,
 				GMeshNormalRepairThreshold * 100.0f,
 				*MeshNormalRepairThresholdText
 			);
 		}
-		else if ( InvalidNormalFraction > 0 )
+		else if (InvalidNormalFraction > 0)
 		{
-			UE_LOG( LogUsd, Warning, TEXT( "%f%% of the normals from Mesh prim '%s' are invalid. This is below the threshold of '%f%%' (configurable via the cvar '%s'), so the invalid normals will be repaired." ),
+			UE_LOG(LogUsd, Warning, TEXT("%f%% of the normals from Mesh prim '%s' are invalid. This is below the threshold of '%f%%' (configurable via the cvar '%s'), so the invalid normals will be repaired."),
 				InvalidNormalFraction * 100.0f,
 				*PrimPath,
 				GMeshNormalRepairThreshold * 100.0f,
 				*MeshNormalRepairThresholdText
 			);
 		}
-		FStaticMeshOperations::ComputeTangentsAndNormals( MeshDescription, Options );
+		FStaticMeshOperations::ComputeTangentsAndNormals(MeshDescription, Options);
 	}
 
-	UStaticMesh* CreateStaticMesh( TArray<FMeshDescription>& LODIndexToMeshDescription, FUsdSchemaTranslationContext& Context, const FString& MeshName, const bool bShouldEnableNanite, bool& bOutIsNew )
+	UStaticMesh* CreateStaticMesh(TArray<FMeshDescription>& LODIndexToMeshDescription, FUsdSchemaTranslationContext& Context, const FString& MeshName, const bool bShouldEnableNanite, bool& bOutIsNew)
 	{
 		UStaticMesh* StaticMesh = nullptr;
 
@@ -692,12 +693,12 @@ namespace UsdGeomMeshTranslatorImpl
 		{
 			FSHA1 SHA1;
 
-			for (const FMeshDescription& MeshDescription : LODIndexToMeshDescription )
+			for (const FMeshDescription& MeshDescription : LODIndexToMeshDescription)
 			{
-				FSHAHash LODHash = FStaticMeshOperations::ComputeSHAHash( MeshDescription );
-				SHA1.Update( &LODHash.Hash[0], sizeof(LODHash.Hash) );
+				FSHAHash LODHash = FStaticMeshOperations::ComputeSHAHash(MeshDescription);
+				SHA1.Update(&LODHash.Hash[0], sizeof(LODHash.Hash));
 
-				if ( !MeshDescription.IsEmpty() )
+				if (!MeshDescription.IsEmpty())
 				{
 					bHasValidMeshDescription = true;
 				}
@@ -717,30 +718,30 @@ namespace UsdGeomMeshTranslatorImpl
 			// another prim, but that doesn't seem like a good path to go down
 			// Additionally, hashing this bool also prevents us from having to force-rebuild a mesh to switch its Nanite flag, which could be tricky to do since some of
 			// these build steps are async/thread-pool based.
-			SHA1.Update( reinterpret_cast< const uint8* >( &bShouldEnableNanite ), sizeof( bShouldEnableNanite ) );
+			SHA1.Update(reinterpret_cast<const uint8*>(&bShouldEnableNanite), sizeof(bShouldEnableNanite));
 
 			// Hash the threshhold so that if we update it and reload we'll regenerate static meshes
-			SHA1.Update( reinterpret_cast< const uint8* >( &GMeshNormalRepairThreshold ), sizeof( GMeshNormalRepairThreshold ) );
+			SHA1.Update(reinterpret_cast<const uint8*>(&GMeshNormalRepairThreshold), sizeof(GMeshNormalRepairThreshold));
 
 			SHA1.Final();
 			SHA1.GetHash(&AllLODHash.Hash[0]);
 		}
 
-		StaticMesh = Cast< UStaticMesh >( Context.AssetCache->GetCachedAsset( AllLODHash.ToString() ) );
+		StaticMesh = Cast< UStaticMesh >(Context.AssetCache->GetCachedAsset(AllLODHash.ToString()));
 
-		if ( !StaticMesh && bHasValidMeshDescription )
+		if (!StaticMesh && bHasValidMeshDescription)
 		{
 			bOutIsNew = true;
 
-			FName AssetName = MakeUniqueObjectName( GetTransientPackage(), UStaticMesh::StaticClass(), *FPaths::GetBaseFilename( MeshName ) );
-			StaticMesh = NewObject< UStaticMesh >( GetTransientPackage(), AssetName, Context.ObjectFlags | EObjectFlags::RF_Public );
+			FName AssetName = MakeUniqueObjectName(GetTransientPackage(), UStaticMesh::StaticClass(), *FPaths::GetBaseFilename(MeshName));
+			StaticMesh = NewObject< UStaticMesh >(GetTransientPackage(), AssetName, Context.ObjectFlags | EObjectFlags::RF_Public);
 
 #if WITH_EDITOR
-			for ( int32 LODIndex = 0; LODIndex < LODIndexToMeshDescription.Num(); ++LODIndex )
+			for (int32 LODIndex = 0; LODIndex < LODIndexToMeshDescription.Num(); ++LODIndex)
 			{
 				FMeshDescription& MeshDescription = LODIndexToMeshDescription[LODIndex];
 
-				RepairNormalsAndTangents( MeshName, MeshDescription );
+				RepairNormalsAndTangents(MeshName, MeshDescription);
 
 				FStaticMeshSourceModel& SourceModel = StaticMesh->AddSourceModel();
 				SourceModel.BuildSettings.bGenerateLightmapUVs = false;
@@ -749,14 +750,14 @@ namespace UsdGeomMeshTranslatorImpl
 				SourceModel.BuildSettings.bBuildReversedIndexBuffer = false;
 
 				FMeshDescription* StaticMeshDescription = StaticMesh->CreateMeshDescription(LODIndex);
-				check( StaticMeshDescription );
-				*StaticMeshDescription = MoveTemp( MeshDescription );
+				check(StaticMeshDescription);
+				*StaticMeshDescription = MoveTemp(MeshDescription);
 			}
 #endif // WITH_EDITOR
 
 			StaticMesh->SetLightingGuid();
 
-			Context.AssetCache->CacheAsset( AllLODHash.ToString(), StaticMesh );
+			Context.AssetCache->CacheAsset(AllLODHash.ToString(), StaticMesh);
 		}
 		else
 		{
@@ -767,11 +768,11 @@ namespace UsdGeomMeshTranslatorImpl
 		return StaticMesh;
 	}
 
-	void PreBuildStaticMesh( UStaticMesh& StaticMesh )
+	void PreBuildStaticMesh(UStaticMesh& StaticMesh)
 	{
-		TRACE_CPUPROFILER_EVENT_SCOPE( UsdGeomMeshTranslatorImpl::PreBuildStaticMesh );
+		TRACE_CPUPROFILER_EVENT_SCOPE(UsdGeomMeshTranslatorImpl::PreBuildStaticMesh);
 
-		if ( StaticMesh.GetRenderData())
+		if (StaticMesh.GetRenderData())
 		{
 			StaticMesh.ReleaseResources();
 			StaticMesh.ReleaseResourcesFence.Wait();
@@ -781,11 +782,11 @@ namespace UsdGeomMeshTranslatorImpl
 		StaticMesh.CreateBodySetup();
 	}
 
-	bool BuildStaticMesh( UStaticMesh& StaticMesh, const FStaticFeatureLevel& FeatureLevel, TArray<FMeshDescription>& LODIndexToMeshDescription )
+	bool BuildStaticMesh(UStaticMesh& StaticMesh, const FStaticFeatureLevel& FeatureLevel, TArray<FMeshDescription>& LODIndexToMeshDescription)
 	{
-		TRACE_CPUPROFILER_EVENT_SCOPE( UsdGeomMeshTranslatorImpl::BuildStaticMesh );
+		TRACE_CPUPROFILER_EVENT_SCOPE(UsdGeomMeshTranslatorImpl::BuildStaticMesh);
 
-		if ( LODIndexToMeshDescription.Num() == 0 )
+		if (LODIndexToMeshDescription.Num() == 0)
 		{
 			return false;
 		}
@@ -793,64 +794,64 @@ namespace UsdGeomMeshTranslatorImpl
 #if WITH_EDITOR
 		ITargetPlatformManagerModule& TargetPlatformManager = GetTargetPlatformManagerRef();
 		ITargetPlatform* RunningPlatform = TargetPlatformManager.GetRunningTargetPlatform();
-		check( RunningPlatform );
+		check(RunningPlatform);
 
 		const FStaticMeshLODSettings& LODSettings = RunningPlatform->GetStaticMeshLODSettings();
-		StaticMesh.GetRenderData()->Cache( RunningPlatform, &StaticMesh, LODSettings );
+		StaticMesh.GetRenderData()->Cache(RunningPlatform, &StaticMesh, LODSettings);
 #else
-		StaticMesh.GetRenderData()->AllocateLODResources( LODIndexToMeshDescription.Num() );
+		StaticMesh.GetRenderData()->AllocateLODResources(LODIndexToMeshDescription.Num());
 
 		// Build render data from each mesh description
-		for ( int32 LODIndex = 0; LODIndex < LODIndexToMeshDescription.Num(); ++LODIndex )
+		for (int32 LODIndex = 0; LODIndex < LODIndexToMeshDescription.Num(); ++LODIndex)
 		{
-			FStaticMeshLODResources& LODResources = StaticMesh.GetRenderData()->LODResources[ LODIndex ];
+			FStaticMeshLODResources& LODResources = StaticMesh.GetRenderData()->LODResources[LODIndex];
 
-			FMeshDescription& MeshDescription = LODIndexToMeshDescription[ LODIndex ];
-			TVertexInstanceAttributesConstRef< FVector4f > MeshDescriptionColors = MeshDescription.VertexInstanceAttributes().GetAttributesRef<FVector4f>( MeshAttribute::VertexInstance::Color );
+			FMeshDescription& MeshDescription = LODIndexToMeshDescription[LODIndex];
+			TVertexInstanceAttributesConstRef< FVector4f > MeshDescriptionColors = MeshDescription.VertexInstanceAttributes().GetAttributesRef<FVector4f>(MeshAttribute::VertexInstance::Color);
 
 			// Compute normals here if necessary because they're not going to be computed via the regular static mesh build pipeline at runtime
 			// (i.e. StaticMeshBuilder is not available at runtime)
 			// We need polygon info because ComputeTangentsAndNormals uses it to repair the invalid vertex normals/tangents
 			// Can't calculate just the required polygons as ComputeTangentsAndNormals is parallel and we can't guarantee thread-safe access patterns
-			FStaticMeshOperations::ComputeTriangleTangentsAndNormals( MeshDescription );
-			FStaticMeshOperations::ComputeTangentsAndNormals( MeshDescription, EComputeNTBsFlags::UseMikkTSpace );
+			FStaticMeshOperations::ComputeTriangleTangentsAndNormals(MeshDescription);
+			FStaticMeshOperations::ComputeTangentsAndNormals(MeshDescription, EComputeNTBsFlags::UseMikkTSpace);
 
 			// Manually set this as it seems the UStaticMesh only sets this whenever the mesh is serialized, which we won't do
 			LODResources.bHasColorVertexData = MeshDescriptionColors.GetNumElements() > 0;
 
-			StaticMesh.BuildFromMeshDescription( MeshDescription, LODResources );
+			StaticMesh.BuildFromMeshDescription(MeshDescription, LODResources);
 		}
 
 #endif // WITH_EDITOR
 
-		static IConsoleVariable* CVar = IConsoleManager::Get().FindConsoleVariable( TEXT( "USD.EnableCollision" ) );
+		static IConsoleVariable* CVar = IConsoleManager::Get().FindConsoleVariable(TEXT("USD.EnableCollision"));
 		bool bEnableCollision = CVar && CVar->GetBool();
 
-		if ( StaticMesh.GetBodySetup())
+		if (StaticMesh.GetBodySetup())
 		{
-			if ( bEnableCollision )
+			if (bEnableCollision)
 			{
 				StaticMesh.GetBodySetup()->CreatePhysicsMeshes();
 			}
 			else
 			{
-				StaticMesh.GetBodySetup()->DefaultInstance.SetCollisionEnabled( ECollisionEnabled::NoCollision );
-				StaticMesh.GetBodySetup()->DefaultInstance.SetCollisionProfileName( UCollisionProfile::NoCollision_ProfileName );
+				StaticMesh.GetBodySetup()->DefaultInstance.SetCollisionEnabled(ECollisionEnabled::NoCollision);
+				StaticMesh.GetBodySetup()->DefaultInstance.SetCollisionProfileName(UCollisionProfile::NoCollision_ProfileName);
 			}
 		}
 		return true;
 	}
 
-	void PostBuildStaticMesh( UStaticMesh& StaticMesh, const TArray<FMeshDescription>& LODIndexToMeshDescription )
+	void PostBuildStaticMesh(UStaticMesh& StaticMesh, const TArray<FMeshDescription>& LODIndexToMeshDescription)
 	{
 		// For runtime builds, the analogue for this stuff is already done from within BuildFromMeshDescriptions
-		TRACE_CPUPROFILER_EVENT_SCOPE( UsdGeomMeshTranslatorImpl::PostBuildStaticMesh );
+		TRACE_CPUPROFILER_EVENT_SCOPE(UsdGeomMeshTranslatorImpl::PostBuildStaticMesh);
 
 		StaticMesh.InitResources();
 
 #if WITH_EDITOR
 		// Fetch the MeshDescription from the StaticMesh because we'll have moved it away from LODIndexToMeshDescription CreateStaticMesh
-		if ( const FMeshDescription* MeshDescription = StaticMesh.GetMeshDescription( 0 ) )
+		if (const FMeshDescription* MeshDescription = StaticMesh.GetMeshDescription(0))
 		{
 			StaticMesh.GetRenderData()->Bounds = MeshDescription->GetBounds();
 		}
@@ -858,60 +859,60 @@ namespace UsdGeomMeshTranslatorImpl
 		StaticMesh.ClearMeshDescriptions(); // Clear mesh descriptions to reduce memory usage, they are kept only in bulk data form
 #else
 		// Fetch the MeshDescription from the imported LODIndexToMeshDescription as StaticMesh.GetMeshDescription is editor-only
-		StaticMesh.GetRenderData()->Bounds = LODIndexToMeshDescription[ 0 ].GetBounds();
+		StaticMesh.GetRenderData()->Bounds = LODIndexToMeshDescription[0].GetBounds();
 		StaticMesh.CalculateExtendedBounds();
 #endif // WITH_EDITOR
 	}
 
 #if WITH_EDITOR
-	void GeometryCacheDataForMeshDescription( FGeometryCacheMeshData& OutMeshData, FMeshDescription& MeshDescription );
+	void GeometryCacheDataForMeshDescription(FGeometryCacheMeshData& OutMeshData, FMeshDescription& MeshDescription);
 
-	void GetGeometryCacheDataTimeCodeRange( const UE::FUsdStage& Stage, const FString& PrimPath, int32& OutStartFrame, int32& OutEndFrame )
+	void GetGeometryCacheDataTimeCodeRange(const UE::FUsdStage& Stage, const FString& PrimPath, int32& OutStartFrame, int32& OutEndFrame)
 	{
-		if ( !Stage || PrimPath.IsEmpty() )
+		if (!Stage || PrimPath.IsEmpty())
 		{
 			return;
 		}
 
 		FScopedUsdAllocs Allocs;
 
-		pxr::UsdPrim UsdPrim = pxr::UsdPrim{ Stage.GetPrimAtPath( UE::FSdfPath{ *PrimPath } ) };
-		if ( !UsdPrim )
+		pxr::UsdPrim UsdPrim = pxr::UsdPrim{Stage.GetPrimAtPath(UE::FSdfPath{ *PrimPath })};
+		if (!UsdPrim)
 		{
 			return;
 		}
 
 		constexpr bool bIncludeInherited = false;
-		pxr::TfTokenVector GeomMeshAttributeNames = pxr::UsdGeomMesh::GetSchemaAttributeNames( bIncludeInherited );
-		pxr::TfTokenVector GeomPointBasedAttributeNames = pxr::UsdGeomPointBased::GetSchemaAttributeNames( bIncludeInherited );
+		pxr::TfTokenVector GeomMeshAttributeNames = pxr::UsdGeomMesh::GetSchemaAttributeNames(bIncludeInherited);
+		pxr::TfTokenVector GeomPointBasedAttributeNames = pxr::UsdGeomPointBased::GetSchemaAttributeNames(bIncludeInherited);
 
-		GeomMeshAttributeNames.reserve( GeomMeshAttributeNames.size() + GeomPointBasedAttributeNames.size() );
-		GeomMeshAttributeNames.insert( GeomMeshAttributeNames.end(), GeomPointBasedAttributeNames.begin(), GeomPointBasedAttributeNames.end() );
+		GeomMeshAttributeNames.reserve(GeomMeshAttributeNames.size() + GeomPointBasedAttributeNames.size());
+		GeomMeshAttributeNames.insert(GeomMeshAttributeNames.end(), GeomPointBasedAttributeNames.begin(), GeomPointBasedAttributeNames.end());
 
 		TOptional<double> MinStartTimeCode;
 		TOptional<double> MaxEndTimeCode;
 
-		for ( const pxr::TfToken& AttributeName : GeomMeshAttributeNames )
+		for (const pxr::TfToken& AttributeName : GeomMeshAttributeNames)
 		{
-			if ( const pxr::UsdAttribute& Attribute = UsdPrim.GetAttribute( AttributeName ) )
+			if (const pxr::UsdAttribute& Attribute = UsdPrim.GetAttribute(AttributeName))
 			{
 				std::vector<double> TimeSamples;
-				if ( Attribute.GetTimeSamples( &TimeSamples ) && TimeSamples.size() > 0 )
+				if (Attribute.GetTimeSamples(&TimeSamples) && TimeSamples.size() > 0)
 				{
-					MinStartTimeCode = FMath::Min( MinStartTimeCode.Get( TNumericLimits<double>::Max() ), TimeSamples[ 0 ] );
-					MaxEndTimeCode = FMath::Max( MaxEndTimeCode.Get( TNumericLimits<double>::Lowest() ), TimeSamples[ TimeSamples.size() - 1 ] );
+					MinStartTimeCode = FMath::Min(MinStartTimeCode.Get(TNumericLimits<double>::Max()), TimeSamples[0]);
+					MaxEndTimeCode = FMath::Max(MaxEndTimeCode.Get(TNumericLimits<double>::Lowest()), TimeSamples[TimeSamples.size() - 1]);
 				}
 			}
 		}
 
-		if ( MinStartTimeCode.IsSet() && MaxEndTimeCode.IsSet() )
+		if (MinStartTimeCode.IsSet() && MaxEndTimeCode.IsSet())
 		{
-			OutStartFrame = FMath::FloorToInt( MinStartTimeCode.GetValue() );
-			OutEndFrame = FMath::CeilToInt( MaxEndTimeCode.GetValue() );
+			OutStartFrame = FMath::FloorToInt(MinStartTimeCode.GetValue());
+			OutEndFrame = FMath::CeilToInt(MaxEndTimeCode.GetValue());
 		}
 	}
 
-	UGeometryCache* CreateGeometryCache( const FString& InPrimPath, TArray< FMeshDescription >& LODIndexToMeshDescription, TSharedRef< FUsdSchemaTranslationContext> Context, bool& bOutIsNew )
+	UGeometryCache* CreateGeometryCache(const FString& InPrimPath, TArray< FMeshDescription >& LODIndexToMeshDescription, TSharedRef< FUsdSchemaTranslationContext> Context, bool& bOutIsNew)
 	{
 		UGeometryCache* GeometryCache = nullptr;
 
@@ -919,27 +920,27 @@ namespace UsdGeomMeshTranslatorImpl
 
 		FSHAHash AllLODHash;
 		FSHA1 SHA1;
-		for ( const FMeshDescription& MeshDescription : LODIndexToMeshDescription )
+		for (const FMeshDescription& MeshDescription : LODIndexToMeshDescription)
 		{
-			FSHAHash LODHash = FStaticMeshOperations::ComputeSHAHash( MeshDescription );
-			SHA1.Update( &LODHash.Hash[ 0 ], sizeof( LODHash.Hash ) );
+			FSHAHash LODHash = FStaticMeshOperations::ComputeSHAHash(MeshDescription);
+			SHA1.Update(&LODHash.Hash[0], sizeof(LODHash.Hash));
 
-			if ( !MeshDescription.IsEmpty() )
+			if (!MeshDescription.IsEmpty())
 			{
 				bHasValidMeshDescription = true;
 			}
 		}
 		SHA1.Final();
-		SHA1.GetHash( &AllLODHash.Hash[ 0 ] );
+		SHA1.GetHash(&AllLODHash.Hash[0]);
 
-		GeometryCache = Cast< UGeometryCache >( Context->AssetCache->GetCachedAsset( AllLODHash.ToString() ) );
+		GeometryCache = Cast< UGeometryCache >(Context->AssetCache->GetCachedAsset(AllLODHash.ToString()));
 
-		if ( !GeometryCache && bHasValidMeshDescription )
+		if (!GeometryCache && bHasValidMeshDescription)
 		{
 			bOutIsNew = true;
 
-			const FName AssetName = MakeUniqueObjectName( GetTransientPackage(), UGeometryCache::StaticClass(), *FPaths::GetBaseFilename( InPrimPath ) );
-			GeometryCache = NewObject< UGeometryCache >( GetTransientPackage(), AssetName, Context->ObjectFlags | EObjectFlags::RF_Public );
+			const FName AssetName = MakeUniqueObjectName(GetTransientPackage(), UGeometryCache::StaticClass(), *FPaths::GetBaseFilename(InPrimPath));
+			GeometryCache = NewObject< UGeometryCache >(GetTransientPackage(), AssetName, Context->ObjectFlags | EObjectFlags::RF_Public);
 
 			const UE::FUsdStage& Stage = Context->Stage;
 
@@ -947,17 +948,17 @@ namespace UsdGeomMeshTranslatorImpl
 			TMap< FString, TMap< FString, int32 > >* MaterialToPrimvarToUVIndex = Context->MaterialToPrimvarToUVIndex ? Context->MaterialToPrimvarToUVIndex : &Unused;
 
 			// Fetch the animated mesh start/end frame as they may be different from just the stage's start and end time codes
-			int32 StartFrame = FMath::FloorToInt( Stage.GetStartTimeCode() );
-			int32 EndFrame = FMath::CeilToInt( Stage.GetEndTimeCode() );
-			UsdGeomMeshTranslatorImpl::GetGeometryCacheDataTimeCodeRange( Stage, InPrimPath, StartFrame, EndFrame );
+			int32 StartFrame = FMath::FloorToInt(Stage.GetStartTimeCode());
+			int32 EndFrame = FMath::CeilToInt(Stage.GetEndTimeCode());
+			UsdGeomMeshTranslatorImpl::GetGeometryCacheDataTimeCodeRange(Stage, InPrimPath, StartFrame, EndFrame);
 
 			// The GeometryCache module expects the end frame to be one past the last animation frame
 			EndFrame += 1;
 
 			pxr::TfToken RenderContextToken = pxr::UsdShadeTokens->universalRenderContext;
-			if ( !Context->RenderContext.IsNone() )
+			if (!Context->RenderContext.IsNone())
 			{
-				RenderContextToken = UnrealToUsd::ConvertToken( *Context->RenderContext.ToString() ).Get();
+				RenderContextToken = UnrealToUsd::ConvertToken(*Context->RenderContext.ToString()).Get();
 			}
 
 			UsdToUnreal::FUsdMeshConversionOptions Options;
@@ -966,7 +967,7 @@ namespace UsdGeomMeshTranslatorImpl
 			Options.RenderContext = RenderContextToken;
 
 			// Create and configure a new USDTrack to be added to the GeometryCache
-			UGeometryCacheTrackUsd* UsdTrack = NewObject< UGeometryCacheTrackUsd >( GeometryCache );
+			UGeometryCacheTrackUsd* UsdTrack = NewObject< UGeometryCacheTrackUsd >(GeometryCache);
 			UsdTrack->Initialize(
 				Stage,
 				InPrimPath,
@@ -974,21 +975,21 @@ namespace UsdGeomMeshTranslatorImpl
 				*MaterialToPrimvarToUVIndex,
 				StartFrame,
 				EndFrame,
-				[InPrimPath, Options]( const TWeakObjectPtr<UGeometryCacheTrackUsd> TrackPtr, float Time, FGeometryCacheMeshData& OutMeshData ) mutable
+				[InPrimPath, Options](const TWeakObjectPtr<UGeometryCacheTrackUsd> TrackPtr, float Time, FGeometryCacheMeshData& OutMeshData) mutable
 				{
 					UGeometryCacheTrackUsd* Track = TrackPtr.Get();
-					if ( !Track )
+					if (!Track)
 					{
 						return false;
 					}
 
-					if ( !Track->CurrentStagePinned )
+					if (!Track->CurrentStagePinned)
 					{
 						return false;
 					}
 
-					UE::FUsdPrim Prim = Track->CurrentStagePinned.GetPrimAtPath( UE::FSdfPath{ *Track->PrimPath } );
-					if ( !Prim )
+					UE::FUsdPrim Prim = Track->CurrentStagePinned.GetPrimAtPath(UE::FSdfPath{*Track->PrimPath});
+					if (!Prim)
 					{
 						return false;
 					}
@@ -1001,11 +1002,11 @@ namespace UsdGeomMeshTranslatorImpl
 
 					// Need a local copy of Options to set the TimeCode since this function is called from multiple worker threads
 					UsdToUnreal::FUsdMeshConversionOptions LocalOptions(Options);
-					LocalOptions.TimeCode = pxr::UsdTimeCode{ Time };
+					LocalOptions.TimeCode = pxr::UsdTimeCode{Time};
 					LocalOptions.MaterialToPrimvarToUVIndex = &Track->MaterialToPrimvarToUVIndex;
 
 					UsdGeomMeshTranslatorImpl::LoadMeshDescriptions(
-						pxr::UsdTyped( Prim ),
+						pxr::UsdTyped(Prim),
 						LODIndexToMeshDescription,
 						LODIndexToMaterialInfo,
 						LocalOptions,
@@ -1013,15 +1014,15 @@ namespace UsdGeomMeshTranslatorImpl
 					);
 
 					// Convert the MeshDescription to MeshData, first LOD only
-					for ( FMeshDescription& MeshDescription : LODIndexToMeshDescription )
+					for (FMeshDescription& MeshDescription : LODIndexToMeshDescription)
 					{
-						if ( !MeshDescription.IsEmpty() )
+						if (!MeshDescription.IsEmpty())
 						{
 							// Compute the normals and tangents for the mesh
 							const float ComparisonThreshold = THRESH_POINTS_ARE_SAME;
 
 							// This function make sure the Polygon Normals Tangents Binormals are computed and also remove degenerated triangle from the render mesh description.
-							FStaticMeshOperations::ComputeTriangleTangentsAndNormals( MeshDescription, ComparisonThreshold );
+							FStaticMeshOperations::ComputeTriangleTangentsAndNormals(MeshDescription, ComparisonThreshold);
 
 							// Compute any missing normals or tangents.
 							// Static meshes always blend normals of overlapping corners.
@@ -1029,9 +1030,9 @@ namespace UsdGeomMeshTranslatorImpl
 							ComputeNTBsOptions |= EComputeNTBsFlags::IgnoreDegenerateTriangles;
 							ComputeNTBsOptions |= EComputeNTBsFlags::UseMikkTSpace;
 
-							FStaticMeshOperations::ComputeTangentsAndNormals( MeshDescription, ComputeNTBsOptions );
+							FStaticMeshOperations::ComputeTangentsAndNormals(MeshDescription, ComputeNTBsOptions);
 
-							UsdGeomMeshTranslatorImpl::GeometryCacheDataForMeshDescription( OutMeshData, MeshDescription );
+							UsdGeomMeshTranslatorImpl::GeometryCacheDataForMeshDescription(OutMeshData, MeshDescription);
 
 							return true;
 						}
@@ -1040,18 +1041,18 @@ namespace UsdGeomMeshTranslatorImpl
 				}
 			);
 
-			GeometryCache->AddTrack( UsdTrack );
+			GeometryCache->AddTrack(UsdTrack);
 
 			TArray< FMatrix > Mats;
-			Mats.Add( FMatrix::Identity );
-			Mats.Add( FMatrix::Identity );
+			Mats.Add(FMatrix::Identity);
+			Mats.Add(FMatrix::Identity);
 
 			TArray< float > MatTimes;
-			MatTimes.Add( 0.0f );
-			MatTimes.Add( 0.0f );
-			UsdTrack->SetMatrixSamples( Mats, MatTimes );
+			MatTimes.Add(0.0f);
+			MatTimes.Add(0.0f);
+			UsdTrack->SetMatrixSamples(Mats, MatTimes);
 
-			Context->AssetCache->CacheAsset( AllLODHash.ToString(), GeometryCache );
+			Context->AssetCache->CacheAsset(AllLODHash.ToString(), GeometryCache);
 		}
 		else
 		{
@@ -1062,9 +1063,9 @@ namespace UsdGeomMeshTranslatorImpl
 	}
 
 	// #ueent_todo: Replace MeshDescription with RawMesh and also make it work with StaticMesh
-	void GeometryCacheDataForMeshDescription( FGeometryCacheMeshData& OutMeshData, FMeshDescription& MeshDescription )
+	void GeometryCacheDataForMeshDescription(FGeometryCacheMeshData& OutMeshData, FMeshDescription& MeshDescription)
 	{
-		TRACE_CPUPROFILER_EVENT_SCOPE( GeometryCacheDataForMeshDescription );
+		TRACE_CPUPROFILER_EVENT_SCOPE(GeometryCacheDataForMeshDescription);
 
 		OutMeshData.Positions.Reset();
 		OutMeshData.TextureCoordinates.Reset();
@@ -1083,7 +1084,7 @@ namespace UsdGeomMeshTranslatorImpl
 		OutMeshData.VertexInfo.bHasUV0 = true;
 		OutMeshData.VertexInfo.bHasMotionVectors = false;
 
-		FStaticMeshAttributes MeshDescriptionAttributes( MeshDescription );
+		FStaticMeshAttributes MeshDescriptionAttributes(MeshDescription);
 
 		TVertexAttributesConstRef< FVector3f > VertexPositions = MeshDescriptionAttributes.GetVertexPositions();
 		TVertexInstanceAttributesConstRef< FVector3f > VertexInstanceNormals = MeshDescriptionAttributes.GetVertexInstanceNormals();
@@ -1096,19 +1097,19 @@ namespace UsdGeomMeshTranslatorImpl
 		const int32 NumTriangles = MeshDescription.Triangles().Num();
 		const int32 NumMeshDataVertices = NumTriangles * 3;
 
-		OutMeshData.Positions.Reserve( NumVertices );
-		OutMeshData.Indices.Reserve( NumMeshDataVertices );
-		OutMeshData.TangentsZ.Reserve( NumMeshDataVertices );
-		OutMeshData.Colors.Reserve( NumMeshDataVertices );
-		OutMeshData.TextureCoordinates.Reserve( NumMeshDataVertices );
+		OutMeshData.Positions.Reserve(NumVertices);
+		OutMeshData.Indices.Reserve(NumMeshDataVertices);
+		OutMeshData.TangentsZ.Reserve(NumMeshDataVertices);
+		OutMeshData.Colors.Reserve(NumMeshDataVertices);
+		OutMeshData.TextureCoordinates.Reserve(NumMeshDataVertices);
 
-		FBox BoundingBox( EForceInit::ForceInitToZero );
+		FBox BoundingBox(EForceInit::ForceInitToZero);
 		int32 MaterialIndex = 0;
 		int32 VertexIndex = 0;
-		for ( FPolygonGroupID PolygonGroupID : MeshDescription.PolygonGroups().GetElementIDs() )
+		for (FPolygonGroupID PolygonGroupID : MeshDescription.PolygonGroups().GetElementIDs())
 		{
 			// Skip empty polygon groups
-			if ( MeshDescription.GetNumPolygonGroupPolygons( PolygonGroupID ) == 0 )
+			if (MeshDescription.GetNumPolygonGroupPolygons(PolygonGroupID) == 0)
 			{
 				continue;
 			}
@@ -1118,29 +1119,29 @@ namespace UsdGeomMeshTranslatorImpl
 			BatchInfo.MaterialIndex = MaterialIndex++;
 
 			int32 TriangleCount = 0;
-			for ( FPolygonID PolygonID : MeshDescription.GetPolygonGroupPolygonIDs( PolygonGroupID ) )
+			for (FPolygonID PolygonID : MeshDescription.GetPolygonGroupPolygonIDs(PolygonGroupID))
 			{
-				for ( FTriangleID TriangleID : MeshDescription.GetPolygonTriangles( PolygonID ) )
+				for (FTriangleID TriangleID : MeshDescription.GetPolygonTriangles(PolygonID))
 				{
-					for ( FVertexInstanceID VertexInstanceID : MeshDescription.GetTriangleVertexInstances( TriangleID ) )
+					for (FVertexInstanceID VertexInstanceID : MeshDescription.GetTriangleVertexInstances(TriangleID))
 					{
-						const FVector3f& Position = VertexPositions[ MeshDescription.GetVertexInstanceVertex( VertexInstanceID ) ];
-						OutMeshData.Positions.Add( Position );
+						const FVector3f& Position = VertexPositions[MeshDescription.GetVertexInstanceVertex(VertexInstanceID)];
+						OutMeshData.Positions.Add(Position);
 						BoundingBox += (FVector)Position;
 
-						OutMeshData.Indices.Add( VertexIndex++ );
+						OutMeshData.Indices.Add(VertexIndex++);
 
-						FPackedNormal Normal = VertexInstanceNormals[ VertexInstanceID ];
-						Normal.Vector.W = VertexInstanceBinormalSigns[ VertexInstanceID ] < 0 ? -127 : 127;
-						OutMeshData.TangentsZ.Add( Normal );
-						OutMeshData.TangentsX.Add( VertexInstanceTangents[ VertexInstanceID ] );
+						FPackedNormal Normal = VertexInstanceNormals[VertexInstanceID];
+						Normal.Vector.W = VertexInstanceBinormalSigns[VertexInstanceID] < 0 ? -127 : 127;
+						OutMeshData.TangentsZ.Add(Normal);
+						OutMeshData.TangentsX.Add(VertexInstanceTangents[VertexInstanceID]);
 
 						const bool bSRGB = true;
-						OutMeshData.Colors.Add( FLinearColor( VertexInstanceColors[ VertexInstanceID ] ).ToFColor( bSRGB ) );
+						OutMeshData.Colors.Add(FLinearColor(VertexInstanceColors[VertexInstanceID]).ToFColor(bSRGB));
 
 						// Supporting only one UV channel
 						const int32 UVIndex = 0;
-						OutMeshData.TextureCoordinates.Add( VertexInstanceUVs.Get( VertexInstanceID, UVIndex ) );
+						OutMeshData.TextureCoordinates.Add(VertexInstanceUVs.Get(VertexInstanceID, UVIndex));
 					}
 
 					++TriangleCount;
@@ -1150,33 +1151,33 @@ namespace UsdGeomMeshTranslatorImpl
 			OutMeshData.BoundingBox = (FBox3f)BoundingBox;
 
 			BatchInfo.NumTriangles = TriangleCount;
-			OutMeshData.BatchesInfo.Add( BatchInfo );
+			OutMeshData.BatchesInfo.Add(BatchInfo);
 		}
 	}
 #endif // WITH_EDITOR
 }
 
-FBuildStaticMeshTaskChain::FBuildStaticMeshTaskChain( const TSharedRef< FUsdSchemaTranslationContext >& InContext, const UE::FSdfPath& InPrimPath )
-	: PrimPath( InPrimPath )
-	, Context( InContext )
+FBuildStaticMeshTaskChain::FBuildStaticMeshTaskChain(const TSharedRef< FUsdSchemaTranslationContext >& InContext, const UE::FSdfPath& InPrimPath)
+	: PrimPath(InPrimPath)
+	, Context(InContext)
 {
 }
 
 void FBuildStaticMeshTaskChain::SetupTasks()
 {
 	// Ignore meshes from disabled purposes
-	if ( !EnumHasAllFlags( Context->PurposesToLoad, IUsdPrim::GetPurpose( GetPrim() ) ) )
+	if (!EnumHasAllFlags(Context->PurposesToLoad, IUsdPrim::GetPurpose(GetPrim())))
 	{
 		return;
 	}
 
 	// Create static mesh (Main thread)
-	Do( ESchemaTranslationLaunchPolicy::Sync,
-		[ this ]()
+	Do(ESchemaTranslationLaunchPolicy::Sync,
+		[this]()
 		{
 			// Force load MeshBuilderModule so that it's ready for the async tasks
 #if WITH_EDITOR
-			FModuleManager::LoadModuleChecked< IMeshBuilderModule >( TEXT("MeshBuilder") );
+			FModuleManager::LoadModuleChecked< IMeshBuilderModule >(TEXT("MeshBuilder"));
 #endif // WITH_EDITOR
 
 			const FString PrimPathString = PrimPath.GetString();
@@ -1186,7 +1187,7 @@ void FBuildStaticMeshTaskChain::SetupTasks()
 			// variant or other Meshes outside of the variant), but it's not ideal to have all the generated assets end
 			// up imported as "SM_LOD0_22", "SM_LOD0_23", etc. So here we fetch the parent prim name in case we're a LOD
 			FString MeshName;
-			if ( Context->bAllowInterpretingLODs && UsdUtils::IsGeomMeshALOD( GetPrim() ) )
+			if (Context->bAllowInterpretingLODs && UsdUtils::IsGeomMeshALOD(GetPrim()))
 			{
 				MeshName = PrimPath.GetParentPath().GetString();
 			}
@@ -1196,19 +1197,19 @@ void FBuildStaticMeshTaskChain::SetupTasks()
 			}
 
 			bool bIsNew = true;
-			const bool bShouldEnableNanite = UsdGeomMeshTranslatorImpl::ShouldEnableNanite( LODIndexToMeshDescription, LODIndexToMaterialInfo, *Context, GetPrim() );
-			StaticMesh = UsdGeomMeshTranslatorImpl::CreateStaticMesh( LODIndexToMeshDescription, *Context, MeshName, bShouldEnableNanite, bIsNew );
+			const bool bShouldEnableNanite = UsdGeomMeshTranslatorImpl::ShouldEnableNanite(LODIndexToMeshDescription, LODIndexToMaterialInfo, *Context, GetPrim());
+			StaticMesh = UsdGeomMeshTranslatorImpl::CreateStaticMesh(LODIndexToMeshDescription, *Context, MeshName, bShouldEnableNanite, bIsNew);
 
-			if ( StaticMesh )
+			if (StaticMesh)
 			{
-				Context->AssetCache->LinkAssetToPrim( PrimPathString, StaticMesh );
+				Context->AssetCache->LinkAssetToPrim(PrimPathString, StaticMesh);
 
 #if WITH_EDITOR
 				StaticMesh->NaniteSettings.bEnabled = bShouldEnableNanite;
 
-				if ( bIsNew )
+				if (bIsNew)
 				{
-					UUsdMeshAssetImportData* ImportData = NewObject<UUsdMeshAssetImportData>( StaticMesh, TEXT( "UUSDAssetImportData" ) );
+					UUsdMeshAssetImportData* ImportData = NewObject<UUsdMeshAssetImportData>(StaticMesh, TEXT("UUSDAssetImportData"));
 					ImportData->PrimPath = PrimPathString;
 					StaticMesh->AssetImportData = ImportData;
 				}
@@ -1217,11 +1218,11 @@ void FBuildStaticMeshTaskChain::SetupTasks()
 
 				// Only process the materials if we own the mesh. If it's new we know we do
 #if WITH_EDITOR
-				UUsdAssetImportData* ImportData = Cast<UUsdAssetImportData>( StaticMesh->AssetImportData );
-				if ( ImportData && ImportData->PrimPath == PrimPathString )
+				UUsdAssetImportData* ImportData = Cast<UUsdAssetImportData>(StaticMesh->AssetImportData);
+				if (ImportData && ImportData->PrimPath == PrimPathString)
 #endif // WITH_EDITOR
 				{
-					if ( !bIsNew )
+					if (!bIsNew)
 					{
 						// We may change material assignments
 						StaticMesh->Modify();
@@ -1236,15 +1237,15 @@ void FBuildStaticMeshTaskChain::SetupTasks()
 						Context->ObjectFlags
 					);
 
-					if ( bMaterialsHaveChanged )
+					if (bMaterialsHaveChanged)
 					{
 						const bool bRebuildAll = true;
 
 #if WITH_EDITOR
-						StaticMesh->UpdateUVChannelData( bRebuildAll );
+						StaticMesh->UpdateUVChannelData(bRebuildAll);
 #else
 						// UpdateUVChannelData doesn't do anything without the editor
-						for ( FStaticMaterial& Material : StaticMesh->GetStaticMaterials() )
+						for (FStaticMaterial& Material : StaticMesh->GetStaticMaterials())
 						{
 							Material.UVChannelData.bInitialized = true;
 						}
@@ -1255,54 +1256,54 @@ void FBuildStaticMeshTaskChain::SetupTasks()
 
 			// Only need to continue building the mesh if we just created it
 			return bIsNew;
-		} );
+		});
 
 #if WITH_EDITOR
 	// Commit mesh description (Async)
-	Then( ESchemaTranslationLaunchPolicy::Async,
-		[ this ]()
+	Then(ESchemaTranslationLaunchPolicy::Async,
+		[this]()
 		{
 			UStaticMesh::FCommitMeshDescriptionParams Params;
 			Params.bMarkPackageDirty = false;
 			Params.bUseHashAsGuid = true;
 
-			for ( int32 LODIndex = 0; LODIndex < LODIndexToMeshDescription.Num(); ++LODIndex )
+			for (int32 LODIndex = 0; LODIndex < LODIndexToMeshDescription.Num(); ++LODIndex)
 			{
-				StaticMesh->CommitMeshDescription( LODIndex, Params );
+				StaticMesh->CommitMeshDescription(LODIndex, Params);
 			}
 
 			return true;
-		} );
+		});
 #endif // WITH_EDITOR
 
 	// PreBuild static mesh (Main thread)
-	Then( ESchemaTranslationLaunchPolicy::Sync,
-		[ this ]()
+	Then(ESchemaTranslationLaunchPolicy::Sync,
+		[this]()
 		{
-			RecreateRenderStateContextPtr = MakeShared<FStaticMeshComponentRecreateRenderStateContext>( StaticMesh, true, true );
+			RecreateRenderStateContextPtr = MakeShared<FStaticMeshComponentRecreateRenderStateContext>(StaticMesh, true, true);
 
-			UsdGeomMeshTranslatorImpl::PreBuildStaticMesh( *StaticMesh );
+			UsdGeomMeshTranslatorImpl::PreBuildStaticMesh(*StaticMesh);
 
 			return true;
-		} );
+		});
 
 	// Build static mesh (Async)
-	Then( ESchemaTranslationLaunchPolicy::Async,
-		[ this ]() mutable
+	Then(ESchemaTranslationLaunchPolicy::Async,
+		[this]() mutable
 		{
 			FStaticFeatureLevel FeatureLevel = GMaxRHIFeatureLevel;
 
 			UWorld* World = Context->Level ? Context->Level->GetWorld() : nullptr;
-			if ( !World )
+			if (!World)
 			{
 				World = GWorld;
 			}
-			if ( World )
+			if (World)
 			{
 				FeatureLevel = World->FeatureLevel;
 			}
 
-			if ( !UsdGeomMeshTranslatorImpl::BuildStaticMesh( *StaticMesh, FeatureLevel, LODIndexToMeshDescription ) )
+			if (!UsdGeomMeshTranslatorImpl::BuildStaticMesh(*StaticMesh, FeatureLevel, LODIndexToMeshDescription))
 			{
 				// Build failed, discard the mesh
 				StaticMesh = nullptr;
@@ -1311,23 +1312,23 @@ void FBuildStaticMeshTaskChain::SetupTasks()
 			}
 
 			return true;
-		} );
+		});
 
 	// PostBuild static mesh (Main thread)
-	Then( ESchemaTranslationLaunchPolicy::Sync,
-		[ this ]()
+	Then(ESchemaTranslationLaunchPolicy::Sync,
+		[this]()
 		{
-			UsdGeomMeshTranslatorImpl::PostBuildStaticMesh( *StaticMesh, LODIndexToMeshDescription );
+			UsdGeomMeshTranslatorImpl::PostBuildStaticMesh(*StaticMesh, LODIndexToMeshDescription);
 
 			RecreateRenderStateContextPtr.Reset();
 
 			return true;
-		} );
+		});
 }
 
-FGeomMeshCreateAssetsTaskChain::FGeomMeshCreateAssetsTaskChain( const TSharedRef< FUsdSchemaTranslationContext >& InContext, const UE::FSdfPath& InPrimPath, const FTransform& InAdditionalTransform )
-	: FBuildStaticMeshTaskChain( InContext, InPrimPath )
-	, AdditionalTransform( InAdditionalTransform )
+FGeomMeshCreateAssetsTaskChain::FGeomMeshCreateAssetsTaskChain(const TSharedRef< FUsdSchemaTranslationContext >& InContext, const UE::FSdfPath& InPrimPath, const FTransform& InAdditionalTransform)
+	: FBuildStaticMeshTaskChain(InContext, InPrimPath)
+	, AdditionalTransform(InAdditionalTransform)
 {
 	SetupTasks();
 }
@@ -1339,22 +1340,22 @@ void FGeomMeshCreateAssetsTaskChain::SetupTasks()
 	// To parse all LODs we need to actively switch variant sets to other variants (triggering prim loading/unloading and notices),
 	// which could cause race conditions if other async translation tasks are trying to access those prims
 	ESchemaTranslationLaunchPolicy LaunchPolicy = ESchemaTranslationLaunchPolicy::Async;
-	if ( Context->bAllowInterpretingLODs && UsdUtils::IsGeomMeshALOD( GetPrim() ) )
+	if (Context->bAllowInterpretingLODs && UsdUtils::IsGeomMeshALOD(GetPrim()))
 	{
 		LaunchPolicy = ESchemaTranslationLaunchPolicy::ExclusiveSync;
 	}
 
 	// Create mesh descriptions (Async or ExclusiveSync)
-	Do( LaunchPolicy,
-		[ this ]() -> bool
+	Do(LaunchPolicy,
+		[this]() -> bool
 		{
 			TMap< FString, TMap< FString, int32 > > Unused;
 			TMap< FString, TMap< FString, int32 > >* MaterialToPrimvarToUVIndex = Context->MaterialToPrimvarToUVIndex ? Context->MaterialToPrimvarToUVIndex : &Unused;
 
 			pxr::TfToken RenderContextToken = pxr::UsdShadeTokens->universalRenderContext;
-			if ( !Context->RenderContext.IsNone() )
+			if (!Context->RenderContext.IsNone())
 			{
-				RenderContextToken = UnrealToUsd::ConvertToken( *Context->RenderContext.ToString() ).Get();
+				RenderContextToken = UnrealToUsd::ConvertToken(*Context->RenderContext.ToString()).Get();
 			}
 
 			UsdToUnreal::FUsdMeshConversionOptions Options;
@@ -1366,7 +1367,7 @@ void FGeomMeshCreateAssetsTaskChain::SetupTasks()
 			Options.AdditionalTransform = AdditionalTransform;
 
 			UsdGeomMeshTranslatorImpl::LoadMeshDescriptions(
-				pxr::UsdTyped( GetPrim() ),
+				pxr::UsdTyped(GetPrim()),
 				LODIndexToMeshDescription,
 				LODIndexToMaterialInfo,
 				Options,
@@ -1374,15 +1375,15 @@ void FGeomMeshCreateAssetsTaskChain::SetupTasks()
 			);
 
 			// If we have at least one valid LOD, we should keep going
-			for ( const FMeshDescription& MeshDescription : LODIndexToMeshDescription )
+			for (const FMeshDescription& MeshDescription : LODIndexToMeshDescription)
 			{
-				if ( !MeshDescription.IsEmpty() )
+				if (!MeshDescription.IsEmpty())
 				{
 					return true;
 				}
 			}
 			return false;
-		} );
+		});
 
 	FBuildStaticMeshTaskChain::SetupTasks();
 }
@@ -1391,8 +1392,8 @@ void FGeomMeshCreateAssetsTaskChain::SetupTasks()
 class FGeometryCacheCreateAssetsTaskChain : public FBuildStaticMeshTaskChain
 {
 public:
-	explicit FGeometryCacheCreateAssetsTaskChain( const TSharedRef< FUsdSchemaTranslationContext >& InContext, const UE::FSdfPath& InPrimPath )
-		: FBuildStaticMeshTaskChain( InContext, InPrimPath )
+	explicit FGeometryCacheCreateAssetsTaskChain(const TSharedRef< FUsdSchemaTranslationContext >& InContext, const UE::FSdfPath& InPrimPath)
+		: FBuildStaticMeshTaskChain(InContext, InPrimPath)
 	{
 		SetupTasks();
 	}
@@ -1406,8 +1407,8 @@ void FGeometryCacheCreateAssetsTaskChain::SetupTasks()
 	FScopedUnrealAllocs UnrealAllocs;
 
 	// Create mesh descriptions (Async or ExclusiveSync)
-	Do( ESchemaTranslationLaunchPolicy::Async,
-		[ this ]() -> bool
+	Do(ESchemaTranslationLaunchPolicy::Async,
+		[this]() -> bool
 		{
 			// Always hash the mesh at the same time because it may be animated, and
 			// otherwise we may think it's a new asset just because the context is at a different timecode (e.g. if we reload)
@@ -1417,9 +1418,9 @@ void FGeometryCacheCreateAssetsTaskChain::SetupTasks()
 			TMap< FString, TMap< FString, int32 > >* MaterialToPrimvarToUVIndex = Context->MaterialToPrimvarToUVIndex ? Context->MaterialToPrimvarToUVIndex : &Unused;
 
 			pxr::TfToken RenderContextToken = pxr::UsdShadeTokens->universalRenderContext;
-			if ( !Context->RenderContext.IsNone() )
+			if (!Context->RenderContext.IsNone())
 			{
-				RenderContextToken = UnrealToUsd::ConvertToken( *Context->RenderContext.ToString() ).Get();
+				RenderContextToken = UnrealToUsd::ConvertToken(*Context->RenderContext.ToString()).Get();
 			}
 
 			UsdToUnreal::FUsdMeshConversionOptions Options;
@@ -1430,7 +1431,7 @@ void FGeometryCacheCreateAssetsTaskChain::SetupTasks()
 			Options.bMergeIdenticalMaterialSlots = Context->bMergeIdenticalMaterialSlots;
 
 			UsdGeomMeshTranslatorImpl::LoadMeshDescriptions(
-				pxr::UsdTyped( GetPrim() ),
+				pxr::UsdTyped(GetPrim()),
 				LODIndexToMeshDescription,
 				LODIndexToMaterialInfo,
 				Options,
@@ -1438,68 +1439,68 @@ void FGeometryCacheCreateAssetsTaskChain::SetupTasks()
 			);
 
 			// If we have at least one valid LOD, we should keep going
-			for ( const FMeshDescription& MeshDescription : LODIndexToMeshDescription )
+			for (const FMeshDescription& MeshDescription : LODIndexToMeshDescription)
 			{
-				if ( !MeshDescription.IsEmpty() )
+				if (!MeshDescription.IsEmpty())
 				{
 					return true;
 				}
 			}
 			return false;
-		} );
+		});
 
 	// Create the GeometryCache (Main thread)
-	Then( ESchemaTranslationLaunchPolicy::Sync,
-		[ this ]()
+	Then(ESchemaTranslationLaunchPolicy::Sync,
+		[this]()
 		{
 			bool bIsNew = true;
 			const FString PrimPathString = PrimPath.GetString();
-			UGeometryCache* GeometryCache = UsdGeomMeshTranslatorImpl::CreateGeometryCache( PrimPathString, LODIndexToMeshDescription, Context, bIsNew );
+			UGeometryCache* GeometryCache = UsdGeomMeshTranslatorImpl::CreateGeometryCache(PrimPathString, LODIndexToMeshDescription, Context, bIsNew);
 
-			if ( bIsNew && GeometryCache )
+			if (bIsNew && GeometryCache)
 			{
-				UUsdMeshAssetImportData* ImportData = NewObject< UUsdMeshAssetImportData >( GeometryCache, TEXT( "UUSDAssetImportData" ) );
+				UUsdMeshAssetImportData* ImportData = NewObject< UUsdMeshAssetImportData >(GeometryCache, TEXT("UUSDAssetImportData"));
 				ImportData->PrimPath = PrimPathString;
 				GeometryCache->AssetImportData = ImportData;
 			}
 
 			bool bMaterialsHaveChanged = false;
-			if ( GeometryCache )
+			if (GeometryCache)
 			{
-				Context->AssetCache->LinkAssetToPrim( PrimPathString, GeometryCache );
+				Context->AssetCache->LinkAssetToPrim(PrimPathString, GeometryCache);
 
 				// Only process the materials if we own the GeometryCache. If it's new we know we do
-				UUsdMeshAssetImportData* ImportData = Cast< UUsdMeshAssetImportData >( GeometryCache->AssetImportData );
-				if ( ImportData && ImportData->PrimPath == PrimPathString )
+				UUsdMeshAssetImportData* ImportData = Cast< UUsdMeshAssetImportData >(GeometryCache->AssetImportData);
+				if (ImportData && ImportData->PrimPath == PrimPathString)
 				{
-					bMaterialsHaveChanged = UsdGeomMeshTranslatorImpl::ProcessGeometryCacheMaterials( GetPrim(), LODIndexToMaterialInfo, *GeometryCache, *Context->AssetCache.Get(), Context->Time, Context->ObjectFlags );
+					bMaterialsHaveChanged = UsdGeomMeshTranslatorImpl::ProcessGeometryCacheMaterials(GetPrim(), LODIndexToMaterialInfo, *GeometryCache, *Context->AssetCache.Get(), Context->Time, Context->ObjectFlags);
 				}
 			}
 
-			const bool bContinueTaskChain = ( bIsNew || bMaterialsHaveChanged );
+			const bool bContinueTaskChain = (bIsNew || bMaterialsHaveChanged);
 			return bContinueTaskChain;
-		} );
+		});
 }
 #endif // WITH_EDITOR
 
 void FUsdGeomMeshTranslator::CreateAssets()
 {
-	TRACE_CPUPROFILER_EVENT_SCOPE( FUsdGeomMeshTranslator::CreateAssets );
+	TRACE_CPUPROFILER_EVENT_SCOPE(FUsdGeomMeshTranslator::CreateAssets);
 
 #if WITH_EDITOR
-	if ( !Context->bIsImporting && GUseGeometryCacheUSD && UsdGeomMeshTranslatorImpl::IsAnimated( GetPrim() ) )
+	if (!Context->bIsImporting && GUseGeometryCacheUSD && UsdGeomMeshTranslatorImpl::IsAnimated(GetPrim()))
 	{
 		// Create the GeometryCache TaskChain
-		TSharedRef< FGeometryCacheCreateAssetsTaskChain > AssetsTaskChain = MakeShared< FGeometryCacheCreateAssetsTaskChain >( Context, PrimPath );
+		TSharedRef< FGeometryCacheCreateAssetsTaskChain > AssetsTaskChain = MakeShared< FGeometryCacheCreateAssetsTaskChain >(Context, PrimPath);
 
-		Context->TranslatorTasks.Add( MoveTemp( AssetsTaskChain ) );
+		Context->TranslatorTasks.Add(MoveTemp(AssetsTaskChain));
 	}
 	else
 #endif // WITH_EDITOR
 	{
-		TSharedRef< FGeomMeshCreateAssetsTaskChain > AssetsTaskChain = MakeShared< FGeomMeshCreateAssetsTaskChain >( Context, PrimPath );
+		TSharedRef< FGeomMeshCreateAssetsTaskChain > AssetsTaskChain = MakeShared< FGeomMeshCreateAssetsTaskChain >(Context, PrimPath);
 
-		Context->TranslatorTasks.Add( MoveTemp( AssetsTaskChain ) );
+		Context->TranslatorTasks.Add(MoveTemp(AssetsTaskChain));
 	}
 }
 
@@ -1509,27 +1510,27 @@ USceneComponent* FUsdGeomMeshTranslator::CreateComponents()
 
 #if WITH_EDITOR
 	// Force animated meshes as GeometryCache
-	if ( !Context->bIsImporting && GUseGeometryCacheUSD && UsdGeomMeshTranslatorImpl::IsAnimated( GetPrim() ) )
+	if (!Context->bIsImporting && GUseGeometryCacheUSD && UsdGeomMeshTranslatorImpl::IsAnimated(GetPrim()))
 	{
 		ComponentType = UGeometryCacheUsdComponent::StaticClass();
 	}
 #endif // WITH_EDITOR
 
-	USceneComponent* SceneComponent = CreateComponentsEx( ComponentType, {} );
-	UpdateComponents( SceneComponent );
+	USceneComponent* SceneComponent = CreateComponentsEx(ComponentType, {});
+	UpdateComponents(SceneComponent);
 
 	const FString PrimPathString = PrimPath.GetString();
 
 	// Handle material overrides
 	// Note: This can be here and not in USDGeomXformableTranslator because there is no way that a collapsed mesh prim could end up with a material override
-	if ( UStaticMeshComponent* StaticMeshComponent = Cast<UStaticMeshComponent>( SceneComponent ) )
+	if (UStaticMeshComponent* StaticMeshComponent = Cast<UStaticMeshComponent>(SceneComponent))
 	{
-		if ( UStaticMesh* StaticMesh = Cast< UStaticMesh >( Context->AssetCache->GetAssetForPrim( PrimPathString ) ) )
+		if (UStaticMesh* StaticMesh = Cast< UStaticMesh >(Context->AssetCache->GetAssetForPrim(PrimPathString)))
 		{
 			TArray<UMaterialInterface*> ExistingAssignments;
-			for ( FStaticMaterial& StaticMaterial : StaticMesh->GetStaticMaterials() )
+			for (FStaticMaterial& StaticMaterial : StaticMesh->GetStaticMaterials())
 			{
-				ExistingAssignments.Add( StaticMaterial.MaterialInterface );
+				ExistingAssignments.Add(StaticMaterial.MaterialInterface);
 			}
 
 			MeshTranslationImpl::SetMaterialOverrides(
@@ -1545,12 +1546,12 @@ USceneComponent* FUsdGeomMeshTranslator::CreateComponents()
 			);
 		}
 	}
-	else if ( UGeometryCacheComponent* Component = Cast<UGeometryCacheComponent>( SceneComponent ) )
+	else if (UGeometryCacheComponent* Component = Cast<UGeometryCacheComponent>(SceneComponent))
 	{
 		// We don't support importing geometry cache stuff yet so we should not be in here
-		ensure( !Context->bIsImporting );
+		ensure(!Context->bIsImporting);
 
-		if ( UGeometryCache* GeometryCache = Cast< UGeometryCache >( Context->AssetCache->GetAssetForPrim( PrimPathString ) ) )
+		if (UGeometryCache* GeometryCache = Cast< UGeometryCache >(Context->AssetCache->GetAssetForPrim(PrimPathString)))
 		{
 			// Geometry caches don't support LODs
 			const bool bAllowInterpretingLODs = false;
@@ -1569,18 +1570,18 @@ USceneComponent* FUsdGeomMeshTranslator::CreateComponents()
 
 #if WITH_EDITOR
 			// Check if the prim has the GroomBinding schema and setup the component and assets necessary to bind the groom to the GeometryCache
-			if ( UsdUtils::PrimHasSchema( GetPrim(), UnrealIdentifiers::GroomBindingAPI ) )
+			if (UsdUtils::PrimHasSchema(GetPrim(), UnrealIdentifiers::GroomBindingAPI))
 			{
-				UsdGroomTranslatorUtils::CreateGroomBindingAsset( GetPrim(), *( Context->AssetCache ), Context->ObjectFlags );
+				UsdGroomTranslatorUtils::CreateGroomBindingAsset(GetPrim(), *(Context->AssetCache), Context->ObjectFlags);
 
 				// For the groom binding to work, the GroomComponent must be a child of the SceneComponent
 				// so the Context ParentComponent is set to the SceneComponent temporarily
-				TGuardValue< USceneComponent* > ParentComponentGuard{ Context->ParentComponent, SceneComponent };
+				TGuardValue< USceneComponent* > ParentComponentGuard{Context->ParentComponent, SceneComponent};
 				const bool bNeedsActor = false;
-				UGroomComponent* GroomComponent = Cast< UGroomComponent >( CreateComponentsEx( TSubclassOf< USceneComponent >( UGroomComponent::StaticClass() ), bNeedsActor ) );
-				if ( GroomComponent )
+				UGroomComponent* GroomComponent = Cast< UGroomComponent >(CreateComponentsEx(TSubclassOf< USceneComponent >(UGroomComponent::StaticClass()), bNeedsActor));
+				if (GroomComponent)
 				{
-					UpdateComponents( SceneComponent );
+					UpdateComponents(SceneComponent);
 				}
 			}
 #endif // !WITH_EDITOR
@@ -1590,9 +1591,9 @@ USceneComponent* FUsdGeomMeshTranslator::CreateComponents()
 	return SceneComponent;
 }
 
-void FUsdGeomMeshTranslator::UpdateComponents( USceneComponent* SceneComponent )
+void FUsdGeomMeshTranslator::UpdateComponents(USceneComponent* SceneComponent)
 {
-	if ( SceneComponent )
+	if (SceneComponent)
 	{
 		SceneComponent->Modify();
 	}
@@ -1601,8 +1602,8 @@ void FUsdGeomMeshTranslator::UpdateComponents( USceneComponent* SceneComponent )
 #if WITH_EDITOR
 		!GUseGeometryCacheUSD &&
 #endif // !WITH_EDITOR
-		UsdGeomMeshTranslatorImpl::IsAnimated( GetPrim() )
-	)
+		UsdGeomMeshTranslatorImpl::IsAnimated(GetPrim())
+		)
 	{
 		// The assets might have changed since our attributes are animated
 		// Note that we must wait for these to complete as they make take a while and we want to
@@ -1613,18 +1614,18 @@ void FUsdGeomMeshTranslator::UpdateComponents( USceneComponent* SceneComponent )
 
 #if WITH_EDITOR
 	// Set the initial GeometryCache on the GeometryCacheUsdComponent
-	if ( UGeometryCacheUsdComponent* GeometryCacheUsdComponent = Cast< UGeometryCacheUsdComponent >( SceneComponent ) )
+	if (UGeometryCacheUsdComponent* GeometryCacheUsdComponent = Cast< UGeometryCacheUsdComponent >(SceneComponent))
 	{
-		ensure( !Context->bIsImporting );
+		ensure(!Context->bIsImporting);
 
-		UGeometryCache* GeometryCache = Cast< UGeometryCache >( Context->AssetCache->GetAssetForPrim( PrimPath.GetString() ) );
+		UGeometryCache* GeometryCache = Cast< UGeometryCache >(Context->AssetCache->GetAssetForPrim(PrimPath.GetString()));
 
 		bool bShouldRegister = false;
-		if ( GeometryCache != GeometryCacheUsdComponent->GetGeometryCache() )
+		if (GeometryCache != GeometryCacheUsdComponent->GetGeometryCache())
 		{
 			bShouldRegister = true;
 
-			if ( GeometryCacheUsdComponent->IsRegistered() )
+			if (GeometryCacheUsdComponent->IsRegistered())
 			{
 				GeometryCacheUsdComponent->UnregisterComponent();
 			}
@@ -1633,43 +1634,49 @@ void FUsdGeomMeshTranslator::UpdateComponents( USceneComponent* SceneComponent )
 			GeometryCacheUsdComponent->GeometryCache = GeometryCache;
 		}
 
-		float TimeCode = Context->Time;
-		if ( FMath::IsNaN( TimeCode ) )
+		// Manually tick USD GeometryCache only when their tracks are disabled in Sequencer
+		static IConsoleVariable* CVar = IConsoleManager::Get().FindConsoleVariable(TEXT("USD.DisableGeoCacheTracks"));
+		bool bDisableGeoCacheTracks = CVar && CVar->GetBool();
+		if (bDisableGeoCacheTracks)
 		{
-			int32 StartFrame = FMath::FloorToInt( Context->Stage.GetStartTimeCode() );
-			int32 EndFrame = FMath::CeilToInt( Context->Stage.GetEndTimeCode() );
-			UsdGeomMeshTranslatorImpl::GetGeometryCacheDataTimeCodeRange( Context->Stage, PrimPath.GetString(), StartFrame, EndFrame );
+			float TimeCode = Context->Time;
+			if (FMath::IsNaN(TimeCode))
+			{
+				int32 StartFrame = FMath::FloorToInt(Context->Stage.GetStartTimeCode());
+				int32 EndFrame = FMath::CeilToInt(Context->Stage.GetEndTimeCode());
+				UsdGeomMeshTranslatorImpl::GetGeometryCacheDataTimeCodeRange(Context->Stage, PrimPath.GetString(), StartFrame, EndFrame);
 
-			TimeCode = static_cast< float >( StartFrame );
-		}
+				TimeCode = static_cast<float>(StartFrame);
+			}
 
-		// This is the main call responsible for animating the geometry cache.
-		// It needs to happen after setting the geometry cache and before registering, because we must force the
-		// geometry cache to register itself at Context->Time so that it will synchronously load that frame right away.
-		// Otherwise the geometry cache will start at t=0 regardless of Context->Time
-		GeometryCacheUsdComponent->SetManualTick( true );
-		GeometryCacheUsdComponent->TickAtThisTime( TimeCode, true, false, true );
+			// This is the main call responsible for animating the geometry cache.
+			// It needs to happen after setting the geometry cache and before registering, because we must force the
+			// geometry cache to register itself at Context->Time so that it will synchronously load that frame right away.
+			// Otherwise the geometry cache will start at t=0 regardless of Context->Time
+			GeometryCacheUsdComponent->SetManualTick(true);
+			GeometryCacheUsdComponent->TickAtThisTime(TimeCode, true, false, true);
+		}
 
 		// Note how we should only register if our geometry cache changed: If we did this every time we would
 		// register too early during the process of duplicating into PIE, and that would prevent a future RegisterComponent
 		// call from naturally creating the required render state
-		if ( bShouldRegister && !GeometryCacheUsdComponent->IsRegistered() )
+		if (bShouldRegister && !GeometryCacheUsdComponent->IsRegistered())
 		{
 			GeometryCacheUsdComponent->RegisterComponent();
 		}
 
 		// If the prim has a GroomBinding schema, apply the target groom to its associated GroomComponent
-		if ( UsdUtils::PrimHasSchema( GetPrim(), UnrealIdentifiers::GroomBindingAPI ) )
+		if (UsdUtils::PrimHasSchema(GetPrim(), UnrealIdentifiers::GroomBindingAPI))
 		{
-			UsdGroomTranslatorUtils::SetGroomFromPrim( GetPrim(), *Context->AssetCache, SceneComponent );
+			UsdGroomTranslatorUtils::SetGroomFromPrim(GetPrim(), *Context->AssetCache, SceneComponent);
 		}
 	}
 #endif // WITH_EDITOR
 
-	Super::UpdateComponents( SceneComponent );
+	Super::UpdateComponents(SceneComponent);
 }
 
-bool FUsdGeomMeshTranslator::CollapsesChildren( ECollapsingType CollapsingType ) const
+bool FUsdGeomMeshTranslator::CollapsesChildren(ECollapsingType CollapsingType) const
 {
 	// We can't claim we collapse anything here since we'll just parse the mesh for this prim and that's it,
 	// otherwise the translation context wouldn't spawn translators for our child prims.
@@ -1681,19 +1688,19 @@ bool FUsdGeomMeshTranslator::CollapsesChildren( ECollapsingType CollapsingType )
 	return false;
 }
 
-bool FUsdGeomMeshTranslator::CanBeCollapsed( ECollapsingType CollapsingType ) const
+bool FUsdGeomMeshTranslator::CanBeCollapsed(ECollapsingType CollapsingType) const
 {
 	UE::FUsdPrim Prim = GetPrim();
 
 	// Don't collapse if our final UStaticMesh would have multiple LODs
-	if ( Context->bAllowInterpretingLODs &&
-		 CollapsingType == ECollapsingType::Assets &&
-		 UsdUtils::IsGeomMeshALOD( Prim ) )
+	if (Context->bAllowInterpretingLODs &&
+		CollapsingType == ECollapsingType::Assets &&
+		UsdUtils::IsGeomMeshALOD(Prim))
 	{
 		return false;
 	}
 
-	return Super::CanBeCollapsed( CollapsingType );
+	return Super::CanBeCollapsed(CollapsingType);
 }
 
 #endif // #if USE_USD_SDK
