USD: GeometryCache from USD now use tracks in LevelSequence instead of the Time property to drive their animations. Added USD.DisableGeoCacheTracks cvar to revert to the previous behavior.

diff --git a/Engine/Plugins/Importers/USDImporter/Source/USDStage/Private/USDLevelSequenceHelper.cpp b/Engine/Plugins/Importers/USDImporter/Source/USDStage/Private/USDLevelSequenceHelper.cpp
--- a/Engine/Plugins/Importers/USDImporter/Source/USDStage/Private/USDLevelSequenceHelper.cpp
+++ b/Engine/Plugins/Importers/USDImporter/Source/USDStage/Private/USDLevelSequenceHelper.cpp
@@ -40,14 +40,19 @@
 #include "Components/SpotLightComponent.h"
 #include "ControlRigObjectBinding.h"
 #include "CoreMinimal.h"
+#include "GeometryCacheComponent.h"
+#include "GeometryCacheTrackUSD.h"
 #include "GroomCache.h"
 #include "GroomComponent.h"
+#include "HAL/IConsoleManager.h"
 #include "LevelSequence.h"
 #include "LevelSequenceActor.h"
 #include "LevelSequencePlayer.h"
 #include "Misc/ITransaction.h"
 #include "Misc/TransactionObjectEvent.h"
 #include "MovieScene.h"
+#include "MovieSceneGeometryCacheSection.h"
+#include "MovieSceneGeometryCacheTrack.h"
 #include "MovieSceneGroomCacheSection.h"
 #include "MovieSceneGroomCacheTrack.h"
 #include "MovieSceneTimeHelpers.h"
@@ -86,13 +91,13 @@
 namespace UsdLevelSequenceHelperImpl
 {
 	// Adapted from ObjectTools as it is within an Editor-only module
-	FString SanitizeObjectName( const FString& InObjectName )
+	FString SanitizeObjectName(const FString& InObjectName)
 	{
 		FString SanitizedText = InObjectName;
 		const TCHAR* InvalidChar = INVALID_OBJECTNAME_CHARACTERS;
-		while ( *InvalidChar )
+		while (*InvalidChar)
 		{
-			SanitizedText.ReplaceCharInline( *InvalidChar, TCHAR( '_' ), ESearchCase::CaseSensitive );
+			SanitizedText.ReplaceCharInline(*InvalidChar, TCHAR('_'), ESearchCase::CaseSensitive);
 			++InvalidChar;
 		}
 
@@ -103,33 +108,33 @@ namespace UsdLevelSequenceHelperImpl
 	 * Similar to FrameRate.AsFrameNumber(TimeSeconds) except that it uses RoundToDouble instead of FloorToDouble, to
 	 * prevent issues with floating point precision
 	 */
-	FFrameNumber RoundAsFrameNumber( const FFrameRate& FrameRate, double TimeSeconds )
+	FFrameNumber RoundAsFrameNumber(const FFrameRate& FrameRate, double TimeSeconds)
 	{
-		const double TimeAsFrame = ( double( TimeSeconds ) * FrameRate.Numerator ) / FrameRate.Denominator;
-		return FFrameNumber(static_cast<int32>( FMath::RoundToDouble( TimeAsFrame ) ));
+		const double TimeAsFrame = (double(TimeSeconds) * FrameRate.Numerator) / FrameRate.Denominator;
+		return FFrameNumber(static_cast<int32>(FMath::RoundToDouble(TimeAsFrame)));
 	}
 
 	// Like UMovieScene::FindTrack, except that if we require class T it will return a track of type T or any type that derives from T
 	template<typename TrackType>
-	TrackType* FindTrackTypeOrDerived( const UMovieScene* MovieScene, const FGuid& Guid, const FName& TrackName = NAME_None )
+	TrackType* FindTrackTypeOrDerived(const UMovieScene* MovieScene, const FGuid& Guid, const FName& TrackName = NAME_None)
 	{
-		if ( !MovieScene || !Guid.IsValid() )
+		if (!MovieScene || !Guid.IsValid())
 		{
 			return nullptr;
 		}
 
-		for ( const FMovieSceneBinding& Binding : MovieScene->GetBindings() )
+		for (const FMovieSceneBinding& Binding : MovieScene->GetBindings())
 		{
-			if ( Binding.GetObjectGuid() != Guid )
+			if (Binding.GetObjectGuid() != Guid)
 			{
 				continue;
 			}
 
-			for ( UMovieSceneTrack* Track : Binding.GetTracks() )
+			for (UMovieSceneTrack* Track : Binding.GetTracks())
 			{
-				if ( TrackType* CastTrack = Cast<TrackType>(Track) )
+				if (TrackType* CastTrack = Cast<TrackType>(Track))
 				{
-					if ( TrackName == NAME_None || Track->GetTrackName() == TrackName )
+					if (TrackName == NAME_None || Track->GetTrackName() == TrackName)
 					{
 						return CastTrack;
 					}
@@ -142,10 +147,10 @@ namespace UsdLevelSequenceHelperImpl
 
 	// Returns the UObject that is bound to the track. Will only consider possessables (and ignore spawnables)
 	// since we don't currently have any workflow where an opened USD stage would interact with UE spawnables.
-	UObject* LocateBoundObject( const UMovieSceneSequence& MovieSceneSequence, const FMovieScenePossessable& Possessable )
+	UObject* LocateBoundObject(const UMovieSceneSequence& MovieSceneSequence, const FMovieScenePossessable& Possessable)
 	{
 		UMovieScene* MovieScene = MovieSceneSequence.GetMovieScene();
-		if ( !MovieScene )
+		if (!MovieScene)
 		{
 			return nullptr;
 		}
@@ -156,31 +161,31 @@ namespace UsdLevelSequenceHelperImpl
 		// If we have a parent guid, we must provide the object as a context because really the binding path
 		// will just contain the component name
 		UObject* ParentContext = nullptr;
-		if ( ParentGuid.IsValid() )
+		if (ParentGuid.IsValid())
 		{
-			if ( FMovieScenePossessable* ParentPossessable = MovieScene->FindPossessable( ParentGuid ) )
+			if (FMovieScenePossessable* ParentPossessable = MovieScene->FindPossessable(ParentGuid))
 			{
-				ParentContext = LocateBoundObject( MovieSceneSequence, *ParentPossessable );
+				ParentContext = LocateBoundObject(MovieSceneSequence, *ParentPossessable);
 			}
 		}
 
-		TArray<UObject*, TInlineAllocator<1>> Objects = MovieSceneSequence.LocateBoundObjects( Guid, ParentContext );
-		if ( Objects.Num() > 0 )
+		TArray<UObject*, TInlineAllocator<1>> Objects = MovieSceneSequence.LocateBoundObjects(Guid, ParentContext);
+		if (Objects.Num() > 0)
 		{
-			return Objects[ 0 ];
+			return Objects[0];
 		}
 
 		return nullptr;
 	}
 
-	void MuteTrack( UMovieSceneTrack* Track, UMovieScene* MovieScene, const FString& ComponentBindingString, const FString& TrackName, bool bMute )
+	void MuteTrack(UMovieSceneTrack* Track, UMovieScene* MovieScene, const FString& ComponentBindingString, const FString& TrackName, bool bMute)
 	{
-		if ( !Track || !MovieScene )
+		if (!Track || !MovieScene)
 		{
 			return;
 		}
 
-		if ( Track->IsEvalDisabled() == bMute )
+		if (Track->IsEvalDisabled() == bMute)
 		{
 			return;
 		}
@@ -190,27 +195,27 @@ namespace UsdLevelSequenceHelperImpl
 		// will chose in favor of MuteNodes
 		MovieScene->Modify();
 
-		const FString MuteNode = FString::Printf( TEXT( "%s.%s" ), *ComponentBindingString, *TrackName );
-		if ( bMute )
+		const FString MuteNode = FString::Printf(TEXT("%s.%s"), *ComponentBindingString, *TrackName);
+		if (bMute)
 		{
-			MovieScene->GetMuteNodes().AddUnique( MuteNode );
+			MovieScene->GetMuteNodes().AddUnique(MuteNode);
 		}
 		else
 		{
-			MovieScene->GetMuteNodes().Remove( MuteNode );
+			MovieScene->GetMuteNodes().Remove(MuteNode);
 		}
 #endif // WITH_EDITOR
 
 		Track->Modify();
-		Track->SetEvalDisabled( bMute );
+		Track->SetEvalDisabled(bMute);
 	}
 
 #if WITH_EDITOR
-	TSharedPtr< ISequencer > GetOpenedSequencerForLevelSequence( ULevelSequence* LevelSequence )
+	TSharedPtr< ISequencer > GetOpenedSequencerForLevelSequence(ULevelSequence* LevelSequence)
 	{
 		const bool bFocusIfOpen = false;
-		IAssetEditorInstance* AssetEditor = GEditor->GetEditorSubsystem<UAssetEditorSubsystem>()->FindEditorForAsset( LevelSequence, bFocusIfOpen );
-		ILevelSequenceEditorToolkit* LevelSequenceEditor = static_cast< ILevelSequenceEditorToolkit* >( AssetEditor );
+		IAssetEditorInstance* AssetEditor = GEditor->GetEditorSubsystem<UAssetEditorSubsystem>()->FindEditorForAsset(LevelSequence, bFocusIfOpen);
+		ILevelSequenceEditorToolkit* LevelSequenceEditor = static_cast<ILevelSequenceEditorToolkit*>(AssetEditor);
 		return LevelSequenceEditor ? LevelSequenceEditor->GetSequencer() : nullptr;
 	}
 
@@ -231,7 +236,7 @@ namespace UsdLevelSequenceHelperImpl
 	)
 	{
 		UMovieScene* MovieScene = LevelSequence->GetMovieScene();
-		if ( !MovieScene || !SkeletalMeshComp || !SkeletalMeshComp->GetSkeletalMeshAsset() || !SkeletalMeshComp->GetSkeletalMeshAsset()->GetSkeleton())
+		if (!MovieScene || !SkeletalMeshComp || !SkeletalMeshComp->GetSkeletalMeshAsset() || !SkeletalMeshComp->GetSkeletalMeshAsset()->GetSkeleton())
 		{
 			return false;
 		}
@@ -248,31 +253,31 @@ namespace UsdLevelSequenceHelperImpl
 		IMovieScenePlayer* Player = nullptr;
 		ULevelSequencePlayer* LevelPlayer = nullptr;
 		{
-			Player = LevelPlayer = ULevelSequencePlayer::CreateLevelSequencePlayer( World, LevelSequence, Settings, OutActor );
-			if ( !Player || !LevelPlayer )
+			Player = LevelPlayer = ULevelSequencePlayer::CreateLevelSequencePlayer(World, LevelSequence, Settings, OutActor);
+			if (!Player || !LevelPlayer)
 			{
 				goto cleanup;
 			}
 
 			// Evaluate at the beginning of the subscene time to ensure that spawnables are created before export
 			FFrameTime StartTime = FFrameRate::TransformTime(UE::MovieScene::DiscreteInclusiveLower(MovieScene->GetPlaybackRange()).Value, MovieScene->GetTickResolution(), MovieScene->GetDisplayRate());
-			LevelPlayer->SetPlaybackPosition( FMovieSceneSequencePlaybackParams( StartTime, EUpdatePositionMethod::Play ) );
+			LevelPlayer->SetPlaybackPosition(FMovieSceneSequencePlaybackParams(StartTime, EUpdatePositionMethod::Play));
 		}
 
 		MovieScene->Modify();
 
 		// We allow baking with no AnimSequence (to allow rigging with no previous animation), so if we don't
 		// have an AnimSequence yet we need to bake a temp one
-		if( !AnimSequence )
+		if (!AnimSequence)
 		{
 			bCreatedTempSequence = true;
 			AnimSequence = NewObject<UAnimSequence>();
-			AnimSequence->SetSkeleton( SkeletalMeshComp->GetSkeletalMeshAsset()->GetSkeleton());
+			AnimSequence->SetSkeleton(SkeletalMeshComp->GetSkeletalMeshAsset()->GetSkeleton());
 
 			FMovieSceneSequenceIDRef Template = MovieSceneSequenceID::Root;
 			FMovieSceneSequenceTransform RootToLocalTransform;
-			bResult = MovieSceneToolHelpers::ExportToAnimSequence( AnimSequence, ExportOptions, MovieScene, Player, SkeletalMeshComp, Template, RootToLocalTransform );
-			if ( !bResult )
+			bResult = MovieSceneToolHelpers::ExportToAnimSequence(AnimSequence, ExportOptions, MovieScene, Player, SkeletalMeshComp, Template, RootToLocalTransform);
+			if (!bResult)
 			{
 				goto cleanup;
 			}
@@ -281,46 +286,46 @@ namespace UsdLevelSequenceHelperImpl
 		// Disable any extra existing control rig tracks for this binding.
 		// Reuse one of the control rig parameter tracks if we can
 		{
-			TArray<UMovieSceneTrack*> Tracks = MovieScene->FindTracks( UMovieSceneControlRigParameterTrack::StaticClass(), ComponentBinding, NAME_None );
-			for ( UMovieSceneTrack* AnyOleTrack : Tracks )
+			TArray<UMovieSceneTrack*> Tracks = MovieScene->FindTracks(UMovieSceneControlRigParameterTrack::StaticClass(), ComponentBinding, NAME_None);
+			for (UMovieSceneTrack* AnyOleTrack : Tracks)
 			{
-				UMovieSceneControlRigParameterTrack* ValidTrack = Cast<UMovieSceneControlRigParameterTrack>( AnyOleTrack );
-				if ( ValidTrack )
+				UMovieSceneControlRigParameterTrack* ValidTrack = Cast<UMovieSceneControlRigParameterTrack>(AnyOleTrack);
+				if (ValidTrack)
 				{
 					Track = ValidTrack;
 					Track->Modify();
-					for ( UMovieSceneSection* Section : Track->GetAllSections() )
+					for (UMovieSceneSection* Section : Track->GetAllSections())
 					{
-						Section->SetIsActive( false );
+						Section->SetIsActive(false);
 					}
 				}
 			}
 
-			if ( !Track )
+			if (!Track)
 			{
-				Track = Cast<UMovieSceneControlRigParameterTrack>( MovieScene->AddTrack( UMovieSceneControlRigParameterTrack::StaticClass(), ComponentBinding ) );
+				Track = Cast<UMovieSceneControlRigParameterTrack>(MovieScene->AddTrack(UMovieSceneControlRigParameterTrack::StaticClass(), ComponentBinding));
 				Track->Modify();
 			}
 		}
 
-		if ( Track )
+		if (Track)
 		{
 			FString ObjectName = InClass->GetName();
-			ObjectName.RemoveFromEnd( TEXT( "_C" ) );
-			UControlRig* ControlRig = NewObject<UControlRig>( Track, InClass, FName( *ObjectName ), RF_Transactional );
-			if ( InClass != UFKControlRig::StaticClass() && !ControlRig->SupportsEvent( TEXT( "Backwards Solve" ) ) )
+			ObjectName.RemoveFromEnd(TEXT("_C"));
+			UControlRig* ControlRig = NewObject<UControlRig>(Track, InClass, FName(*ObjectName), RF_Transactional);
+			if (InClass != UFKControlRig::StaticClass() && !ControlRig->SupportsEvent(TEXT("Backwards Solve")))
 			{
-				MovieScene->RemoveTrack( *Track );
+				MovieScene->RemoveTrack(*Track);
 				goto cleanup;
 			}
 
 			ControlRig->Modify();
-			ControlRig->SetObjectBinding( MakeShared<FControlRigObjectBinding>() );
-			ControlRig->GetObjectBinding()->BindToObject( SkeletalMeshComp );
-			ControlRig->GetDataSourceRegistry()->RegisterDataSource( UControlRig::OwnerComponent, ControlRig->GetObjectBinding()->GetBoundObject() );
+			ControlRig->SetObjectBinding(MakeShared<FControlRigObjectBinding>());
+			ControlRig->GetObjectBinding()->BindToObject(SkeletalMeshComp);
+			ControlRig->GetDataSourceRegistry()->RegisterDataSource(UControlRig::OwnerComponent, ControlRig->GetObjectBinding()->GetBoundObject());
 			ControlRig->Initialize();
 			ControlRig->RequestInit();
-			ControlRig->SetBoneInitialTransformsFromSkeletalMeshComponent( SkeletalMeshComp, true );
+			ControlRig->SetBoneInitialTransformsFromSkeletalMeshComponent(SkeletalMeshComp, true);
 			ControlRig->Evaluate_AnyThread();
 
 			// Find the animation section's start frame, or else the baked control rig tracks will always be
@@ -329,28 +334,28 @@ namespace UsdLevelSequenceHelperImpl
 			bool bFoundAtLeastOneSection = false;
 			FFrameNumber ControlRigSectionStartFrame = TNumericLimits<int32>::Max();
 			UMovieSceneSkeletalAnimationTrack* SkelTrack = Cast<UMovieSceneSkeletalAnimationTrack>(
-				MovieScene->FindTrack( UMovieSceneSkeletalAnimationTrack::StaticClass(), ComponentBinding, NAME_None )
-			);
-			if ( SkelTrack )
+				MovieScene->FindTrack(UMovieSceneSkeletalAnimationTrack::StaticClass(), ComponentBinding, NAME_None)
+				);
+			if (SkelTrack)
 			{
-				for ( const UMovieSceneSection* Section : SkelTrack->GetAllSections() )
+				for (const UMovieSceneSection* Section : SkelTrack->GetAllSections())
 				{
-					if ( const UMovieSceneSkeletalAnimationSection* SkelSection = Cast<UMovieSceneSkeletalAnimationSection>( Section ) )
+					if (const UMovieSceneSkeletalAnimationSection* SkelSection = Cast<UMovieSceneSkeletalAnimationSection>(Section))
 					{
-						if ( SkelSection->Params.Animation == AnimSequence )
+						if (SkelSection->Params.Animation == AnimSequence)
 						{
 							TRange<FFrameNumber> Range = SkelSection->ComputeEffectiveRange();
-							if ( Range.HasLowerBound() )
+							if (Range.HasLowerBound())
 							{
 								bFoundAtLeastOneSection = true;
-								ControlRigSectionStartFrame = FMath::Min( ControlRigSectionStartFrame, Range.GetLowerBoundValue() );
+								ControlRigSectionStartFrame = FMath::Min(ControlRigSectionStartFrame, Range.GetLowerBoundValue());
 								break;
 							}
 						}
 					}
 				}
 			}
-			if ( !bFoundAtLeastOneSection )
+			if (!bFoundAtLeastOneSection)
 			{
 				ControlRigSectionStartFrame = 0;
 			}
@@ -358,30 +363,30 @@ namespace UsdLevelSequenceHelperImpl
 			// This is unused
 			const FFrameNumber StartTime = 0;
 			const bool bSequencerOwnsControlRig = true;
-			UMovieSceneSection* NewSection = Track->CreateControlRigSection( StartTime, ControlRig, bSequencerOwnsControlRig );
-			UMovieSceneControlRigParameterSection* ParamSection = Cast<UMovieSceneControlRigParameterSection>( NewSection );
+			UMovieSceneSection* NewSection = Track->CreateControlRigSection(StartTime, ControlRig, bSequencerOwnsControlRig);
+			UMovieSceneControlRigParameterSection* ParamSection = Cast<UMovieSceneControlRigParameterSection>(NewSection);
 
-			Track->SetTrackName( FName( *ObjectName ) );
-			Track->SetDisplayName( FText::FromString( ObjectName ) );
+			Track->SetTrackName(FName(*ObjectName));
+			Track->SetDisplayName(FText::FromString(ObjectName));
 
 			// We want the baked keyframe times to match exactly the source animation keyframe times. The way that
 			// LoadAnimSequenceIntoThisSection uses this however indicates it's meant to be relative to the start
 			// of the playback range, which we do here
-			ControlRigSectionStartFrame -= UE::MovieScene::DiscreteInclusiveLower( MovieScene->GetPlaybackRange() );
+			ControlRigSectionStartFrame -= UE::MovieScene::DiscreteInclusiveLower(MovieScene->GetPlaybackRange());
 
-			ParamSection->LoadAnimSequenceIntoThisSection( AnimSequence, MovieScene, SkeletalMeshComp, bReduceKeys, Tolerance, ControlRigSectionStartFrame );
+			ParamSection->LoadAnimSequenceIntoThisSection(AnimSequence, MovieScene, SkeletalMeshComp, bReduceKeys, Tolerance, ControlRigSectionStartFrame);
 
 			// Disable Skeletal Animation Tracks
-			if ( SkelTrack )
+			if (SkelTrack)
 			{
 				SkelTrack->Modify();
 
-				for ( UMovieSceneSection* Section : SkelTrack->GetAllSections() )
+				for (UMovieSceneSection* Section : SkelTrack->GetAllSections())
 				{
-					if ( Section )
+					if (Section)
 					{
 						Section->TryModify();
-						Section->SetIsActive( false );
+						Section->SetIsActive(false);
 					}
 				}
 			}
@@ -390,19 +395,19 @@ namespace UsdLevelSequenceHelperImpl
 		}
 
 	cleanup:
-		if ( bCreatedTempSequence && AnimSequence )
+		if (bCreatedTempSequence && AnimSequence)
 		{
 			AnimSequence->MarkAsGarbage();
 		}
 
-		if ( LevelPlayer )
+		if (LevelPlayer)
 		{
 			LevelPlayer->Stop();
 		}
 
-		if ( OutActor && World )
+		if (OutActor && World)
 		{
-			World->DestroyActor( OutActor );
+			World->DestroyActor(OutActor);
 		}
 
 		return bResult;
@@ -440,24 +445,24 @@ private:
 
 		bool IsAnimated() const
 		{
-			return !FMath::IsNearlyEqual( StartTimeCode.Get( 0.0 ), EndTimeCode.Get( 0.0 ) );
+			return !FMath::IsNearlyEqual(StartTimeCode.Get(0.0), EndTimeCode.Get(0.0));
 		}
 	};
 
-// FGCObject interface
+	// FGCObject interface
 protected:
-	virtual void AddReferencedObjects( FReferenceCollector& Collector ) override;
+	virtual void AddReferencedObjects(FReferenceCollector& Collector) override;
 	virtual FString GetReferencerName() const override
 	{
 		return TEXT("FUsdLevelSequenceHelperImpl");
 	}
 
-// Sequences handling
+	// Sequences handling
 public:
 
 	/** Creates a Level Sequence and its SubSequenceSection for each layer in the local layer stack (root layer and sub layers) */
 	void CreateLocalLayersSequences();
-	void BindToUsdStageActor( AUsdStageActor* InStageActor );
+	void BindToUsdStageActor(AUsdStageActor* InStageActor);
 	void UnbindFromUsdStageActor();
 	void OnStageActorRenamed();
 
@@ -465,8 +470,8 @@ public:
 	TArray< ULevelSequence* > GetSubSequences() const
 	{
 		TArray< ULevelSequence* > SubSequences;
-		LevelSequencesByIdentifier.GenerateValueArray( SubSequences );
-		SubSequences.Remove( MainLevelSequence );
+		LevelSequencesByIdentifier.GenerateValueArray(SubSequences);
+		SubSequences.Remove(MainLevelSequence);
 
 		return SubSequences;
 	}
@@ -477,14 +482,14 @@ public:
 	}
 
 private:
-	ULevelSequence* FindSequenceForAttribute( const UE::FUsdAttribute& Attribute );
-	ULevelSequence* FindOrAddSequenceForAttribute( const UE::FUsdAttribute& Attribute );
-	ULevelSequence* FindSequenceForIdentifier( const FString& SequenceIdentitifer );
-	ULevelSequence* FindOrAddSequenceForLayer( const UE::FSdfLayer& Layer, const FString& SequenceIdentifier, const FString& SequenceDisplayName );
-	UE::FSdfLayer FindEditTargetForSubsequence( ULevelSequence* Sequence );
+	ULevelSequence* FindSequenceForAttribute(const UE::FUsdAttribute& Attribute);
+	ULevelSequence* FindOrAddSequenceForAttribute(const UE::FUsdAttribute& Attribute);
+	ULevelSequence* FindSequenceForIdentifier(const FString& SequenceIdentitifer);
+	ULevelSequence* FindOrAddSequenceForLayer(const UE::FSdfLayer& Layer, const FString& SequenceIdentifier, const FString& SequenceDisplayName);
+	UE::FSdfLayer FindEditTargetForSubsequence(ULevelSequence* Sequence);
 
 	/** Removes PrimTwin as a user of Sequence. If Sequence is now unused, remove its subsection and itself. */
-	void RemoveSequenceForPrim( ULevelSequence& Sequence, const UUsdPrimTwin& PrimTwin );
+	void RemoveSequenceForPrim(ULevelSequence& Sequence, const UUsdPrimTwin& PrimTwin);
 
 private:
 	ULevelSequence* MainLevelSequence;
@@ -499,40 +504,41 @@ private:
 	// Sequence Name to Layer Identifier Map. Relationship: N Sequences to 1 Layer.
 	TMap<FName, FString> LayerIdentifierByLevelSequenceName;
 
-// Sections handling
+	// Sections handling
 private:
 	/** Returns the UMovieSceneSubSection associated with SubSequence on the Sequence UMovieSceneSubTrack if it exists */
-	UMovieSceneSubSection* FindSubSequenceSection( ULevelSequence& Sequence, ULevelSequence& SubSequence );
-	void CreateSubSequenceSection( ULevelSequence& Sequence, ULevelSequence& Subsequence );
-	void RemoveSubSequenceSection( ULevelSequence& Sequence, ULevelSequence& SubSequence );
+	UMovieSceneSubSection* FindSubSequenceSection(ULevelSequence& Sequence, ULevelSequence& SubSequence);
+	void CreateSubSequenceSection(ULevelSequence& Sequence, ULevelSequence& Subsequence);
+	void RemoveSubSequenceSection(ULevelSequence& Sequence, ULevelSequence& SubSequence);
 
-// Tracks handling
+	// Tracks handling
 private:
 	/** Creates a time track on the ULevelSequence corresponding to Info */
 	void CreateTimeTrack(const FUsdLevelSequenceHelperImpl::FLayerTimeInfo& Info);
 	void RemoveTimeTrack(const FUsdLevelSequenceHelperImpl::FLayerTimeInfo* Info);
 
-	void AddCommonTracks( const UUsdPrimTwin& PrimTwin, const UE::FUsdPrim& Prim, bool bForceVisibilityTracks = false );
-	void AddCameraTracks( const UUsdPrimTwin& PrimTwin, const UE::FUsdPrim& Prim );
-	void AddLightTracks( const UUsdPrimTwin& PrimTwin, const UE::FUsdPrim& Prim, const TSet<FName>& PropertyPathsToRead = {} );
-	void AddSkeletalTracks( const UUsdPrimTwin& PrimTwin, const UE::FUsdPrim& Prim );
-	void AddGroomTracks( const UUsdPrimTwin& PrimTwin, const UE::FUsdPrim& Prim );
+	void AddCommonTracks(const UUsdPrimTwin& PrimTwin, const UE::FUsdPrim& Prim, bool bForceVisibilityTracks = false);
+	void AddCameraTracks(const UUsdPrimTwin& PrimTwin, const UE::FUsdPrim& Prim);
+	void AddLightTracks(const UUsdPrimTwin& PrimTwin, const UE::FUsdPrim& Prim, const TSet<FName>& PropertyPathsToRead = {});
+	void AddSkeletalTracks(const UUsdPrimTwin& PrimTwin, const UE::FUsdPrim& Prim);
+	void AddGroomTracks(const UUsdPrimTwin& PrimTwin, const UE::FUsdPrim& Prim);
+	void AddGeometryCacheTracks(const UUsdPrimTwin& PrimTwin, const UE::FUsdPrim& Prim);
 
 	template<typename TrackType>
-	TrackType* AddTrack( const FName& PropertyPath, const UUsdPrimTwin& PrimTwin, USceneComponent& ComponentToBind, ULevelSequence& Sequence, bool bIsMuted = false  );
+	TrackType* AddTrack(const FName& PropertyPath, const UUsdPrimTwin& PrimTwin, USceneComponent& ComponentToBind, ULevelSequence& Sequence, bool bIsMuted = false);
 
-	void RemovePossessable( const UUsdPrimTwin& PrimTwin );
+	void RemovePossessable(const UUsdPrimTwin& PrimTwin);
 
-// Prims handling
+	// Prims handling
 public:
 	// If bForceVisibilityTracks is true, we'll add and bake the visibility tracks for this prim even if the
 	// prim itself doesn't have animated visibility (so that we can handle its visibility in case one of its
 	// parents does have visibility animations)
-	void AddPrim( UUsdPrimTwin& PrimTwin, bool bForceVisibilityTracks = false );
+	void AddPrim(UUsdPrimTwin& PrimTwin, bool bForceVisibilityTracks = false);
 	void RemovePrim(const UUsdPrimTwin& PrimTwin);
 
 	// These functions assume the skeletal animation tracks (if any) were already added to the level sequence
-	void UpdateControlRigTracks( UUsdPrimTwin& PrimTwin );
+	void UpdateControlRigTracks(UUsdPrimTwin& PrimTwin);
 
 private:
 	// Sequence Name to Prim Path. Relationship: 1 Sequence to N Prim Path.
@@ -549,7 +555,7 @@ private:
 
 	TMap< TWeakObjectPtr< const UUsdPrimTwin >, FPrimTwinBindings > PrimTwinToBindings;
 
-// Time codes handling
+	// Time codes handling
 private:
 	FLayerTimeInfo& FindOrAddLayerTimeInfo(const UE::FSdfLayer& Layer);
 	FLayerTimeInfo* FindLayerTimeInfo(const UE::FSdfLayer& Layer);
@@ -558,15 +564,15 @@ private:
 	void UpdateUsdLayerOffsetFromSection(const UMovieSceneSequence* Sequence, const UMovieSceneSubSection* Section);
 
 	/** Updates LayerTimeInfo with Layer */
-	void UpdateLayerTimeInfoFromLayer( FLayerTimeInfo& LayerTimeInfo, const UE::FSdfLayer& Layer );
+	void UpdateLayerTimeInfoFromLayer(FLayerTimeInfo& LayerTimeInfo, const UE::FSdfLayer& Layer);
 
 	/** Updates MovieScene with LayerTimeInfo */
-	void UpdateMovieSceneTimeRanges( UMovieScene& MovieScene, const FLayerTimeInfo& LayerTimeInfo );
+	void UpdateMovieSceneTimeRanges(UMovieScene& MovieScene, const FLayerTimeInfo& LayerTimeInfo);
 
 	double GetFramesPerSecond() const;
 	double GetTimeCodesPerSecond() const;
 
-	FGuid GetOrCreateComponentBinding( const UUsdPrimTwin& PrimTwin, USceneComponent& ComponentToBind, ULevelSequence& Sequence );
+	FGuid GetOrCreateComponentBinding(const UUsdPrimTwin& PrimTwin, USceneComponent& ComponentToBind, ULevelSequence& Sequence);
 
 	TMap<FString, FLayerTimeInfo> LayerTimeInfosByLayerIdentifier; // Maps a LayerTimeInfo to a given Layer through its identifier
 
@@ -584,12 +590,12 @@ public:
 
 private:
 	void OnObjectTransacted(UObject* Object, const class FTransactionObjectEvent& Event);
-	void OnUsdObjectsChanged( const UsdUtils::FObjectChangesByPath& InfoChanges, const UsdUtils::FObjectChangesByPath& ResyncChanges );
-	void HandleTransactionStateChanged( const FTransactionContext& InTransactionContext, const ETransactionStateEventType InTransactionState );
+	void OnUsdObjectsChanged(const UsdUtils::FObjectChangesByPath& InfoChanges, const UsdUtils::FObjectChangesByPath& ResyncChanges);
+	void HandleTransactionStateChanged(const FTransactionContext& InTransactionContext, const ETransactionStateEventType InTransactionState);
 	void HandleMovieSceneChange(UMovieScene& MovieScene);
 	void HandleSubSectionChange(UMovieSceneSubSection& Section);
-	void HandleControlRigSectionChange( UMovieSceneControlRigParameterSection& Section );
-	void HandleTrackChange( const UMovieSceneTrack& Track, bool bIsMuteChange );
+	void HandleControlRigSectionChange(UMovieSceneControlRigParameterSection& Section);
+	void HandleTrackChange(const UMovieSceneTrack& Track, bool bIsMuteChange);
 	void HandleDisplayRateChange(const double DisplayRate);
 
 	FDelegateHandle OnObjectTransactedHandle;
@@ -627,16 +633,16 @@ const TCHAR* FUsdLevelSequenceHelperImpl::TimeTrackName = TEXT("Time");
 const double FUsdLevelSequenceHelperImpl::EmptySubSectionRange = 10.0;
 
 FUsdLevelSequenceHelperImpl::FUsdLevelSequenceHelperImpl()
-	: MainLevelSequence( nullptr )
+	: MainLevelSequence(nullptr)
 {
 #if WITH_EDITOR
 	OnObjectTransactedHandle = FCoreUObjectDelegates::OnObjectTransacted.AddRaw(this, &FUsdLevelSequenceHelperImpl::OnObjectTransacted);
 
-	if ( GEditor )
+	if (GEditor)
 	{
-		if ( UTransBuffer* Transactor = Cast<UTransBuffer>( GEditor->Trans ) )
+		if (UTransBuffer* Transactor = Cast<UTransBuffer>(GEditor->Trans))
 		{
-			Transactor->OnTransactionStateChanged().AddRaw( this, &FUsdLevelSequenceHelperImpl::HandleTransactionStateChanged );
+			Transactor->OnTransactionStateChanged().AddRaw(this, &FUsdLevelSequenceHelperImpl::HandleTransactionStateChanged);
 		}
 	}
 #endif // WITH_EDITOR
@@ -644,10 +650,10 @@ FUsdLevelSequenceHelperImpl::FUsdLevelSequenceHelperImpl()
 
 FUsdLevelSequenceHelperImpl::~FUsdLevelSequenceHelperImpl()
 {
-	if ( StageActor.IsValid() )
+	if (StageActor.IsValid())
 	{
 		StageActor->GetUsdListener().GetOnStageEditTargetChanged().Remove(OnStageEditTargetChangedHandle);
-		StageActor->GetUsdListener().GetOnObjectsChanged().Remove( OnUsdObjectsChangedHandle );
+		StageActor->GetUsdListener().GetOnObjectsChanged().Remove(OnUsdObjectsChangedHandle);
 		OnStageEditTargetChangedHandle.Reset();
 	}
 
@@ -655,11 +661,11 @@ FUsdLevelSequenceHelperImpl::~FUsdLevelSequenceHelperImpl()
 	FCoreUObjectDelegates::OnObjectTransacted.Remove(OnObjectTransactedHandle);
 	OnObjectTransactedHandle.Reset();
 
-	if ( GEditor )
+	if (GEditor)
 	{
-		if ( UTransBuffer* Transactor = Cast<UTransBuffer>( GEditor->Trans ) )
+		if (UTransBuffer* Transactor = Cast<UTransBuffer>(GEditor->Trans))
 		{
-			Transactor->OnTransactionStateChanged().RemoveAll( this );
+			Transactor->OnTransactionStateChanged().RemoveAll(this);
 		}
 	}
 #endif // WITH_EDITOR
@@ -673,40 +679,40 @@ ULevelSequence* FUsdLevelSequenceHelperImpl::Init(const UE::FUsdStage& InUsdStag
 	return MainLevelSequence;
 }
 
-void FUsdLevelSequenceHelperImpl::SetAssetCache( UUsdAssetCache* InAssetCache )
+void FUsdLevelSequenceHelperImpl::SetAssetCache(UUsdAssetCache* InAssetCache)
 {
 	AssetCache = InAssetCache;
 }
 
 bool FUsdLevelSequenceHelperImpl::HasData() const
 {
-	if ( !MainLevelSequence )
+	if (!MainLevelSequence)
 	{
 		return false;
 	}
 
 	UMovieScene* MovieScene = MainLevelSequence->GetMovieScene();
-	if ( !MovieScene )
+	if (!MovieScene)
 	{
 		return false;
 	}
 
-	if ( MovieScene->GetPossessableCount() > 0 )
+	if (MovieScene->GetPossessableCount() > 0)
 	{
 		return true;
 	}
 
 	UMovieSceneSubTrack* Track = MovieScene->FindTrack<UMovieSceneSubTrack>();
-	if ( !Track )
+	if (!Track)
 	{
 		return false;
 	}
 
-	for ( UMovieSceneSection* Section : Track->GetAllSections() )
+	for (UMovieSceneSection* Section : Track->GetAllSections())
 	{
-		if ( UMovieSceneSubSection* SubSection = Cast<UMovieSceneSubSection>( Section ) )
+		if (UMovieSceneSubSection* SubSection = Cast<UMovieSceneSubSection>(Section))
 		{
-			if ( SubSection->GetSequence() )
+			if (SubSection->GetSequence())
 			{
 				return true;
 			}
@@ -730,67 +736,67 @@ void FUsdLevelSequenceHelperImpl::Clear()
 	SequenceHierarchyCache = FMovieSceneSequenceHierarchy();
 }
 
-void FUsdLevelSequenceHelperImpl::AddReferencedObjects( FReferenceCollector& Collector )
+void FUsdLevelSequenceHelperImpl::AddReferencedObjects(FReferenceCollector& Collector)
 {
-	Collector.AddReferencedObject( MainLevelSequence );
-	Collector.AddReferencedObjects( LevelSequencesByIdentifier );
+	Collector.AddReferencedObject(MainLevelSequence);
+	Collector.AddReferencedObjects(LevelSequencesByIdentifier);
 }
 
 void FUsdLevelSequenceHelperImpl::CreateLocalLayersSequences()
 {
 	Clear();
 
-	if ( !UsdStage )
+	if (!UsdStage)
 	{
 		return;
 	}
 
 	UE::FSdfLayer RootLayer = UsdStage.GetRootLayer();
-	const FLayerTimeInfo& RootLayerInfo = FindOrAddLayerTimeInfo( RootLayer );
+	const FLayerTimeInfo& RootLayerInfo = FindOrAddLayerTimeInfo(RootLayer);
 
 	UE_LOG(LogUsd, Verbose, TEXT("CreateLayerSequences: Initializing level sequence for '%s'"), *RootLayerInfo.Identifier);
 
 	// Create main level sequence for root layer
-	MainLevelSequence = FindOrAddSequenceForLayer( RootLayer, RootLayer.GetIdentifier(), RootLayer.GetDisplayName() );
+	MainLevelSequence = FindOrAddSequenceForLayer(RootLayer, RootLayer.GetIdentifier(), RootLayer.GetDisplayName());
 
-	if ( !MainLevelSequence )
+	if (!MainLevelSequence)
 	{
 		return;
 	}
 
 	UMovieScene* MovieScene = MainLevelSequence->GetMovieScene();
-	if ( !MovieScene )
+	if (!MovieScene)
 	{
 		return;
 	}
 
-	SequencesID.Add( MainLevelSequence ) = MovieSceneSequenceID::Root;
+	SequencesID.Add(MainLevelSequence) = MovieSceneSequenceID::Root;
 
-	LocalLayersSequences.Add( MainLevelSequence->GetFName() );
+	LocalLayersSequences.Add(MainLevelSequence->GetFName());
 
-	TFunction< void( const FLayerTimeInfo* LayerTimeInfo, ULevelSequence& ParentSequence ) > RecursivelyCreateSequencesForLayer;
-	RecursivelyCreateSequencesForLayer = [ &RecursivelyCreateSequencesForLayer, this ]( const FLayerTimeInfo* LayerTimeInfo, ULevelSequence& ParentSequence )
+	TFunction< void(const FLayerTimeInfo* LayerTimeInfo, ULevelSequence& ParentSequence) > RecursivelyCreateSequencesForLayer;
+	RecursivelyCreateSequencesForLayer = [&RecursivelyCreateSequencesForLayer, this](const FLayerTimeInfo* LayerTimeInfo, ULevelSequence& ParentSequence)
 	{
-		if ( !LayerTimeInfo )
+		if (!LayerTimeInfo)
 		{
 			return;
 		}
 
-		if ( UE::FSdfLayer Layer = UE::FSdfLayer::FindOrOpen( *LayerTimeInfo->Identifier ) )
+		if (UE::FSdfLayer Layer = UE::FSdfLayer::FindOrOpen(*LayerTimeInfo->Identifier))
 		{
-			for ( const FString& SubLayerPath : Layer.GetSubLayerPaths() )
+			for (const FString& SubLayerPath : Layer.GetSubLayerPaths())
 			{
-				if ( UE::FSdfLayer SubLayer = UsdUtils::FindLayerForSubLayerPath( Layer, SubLayerPath ) )
+				if (UE::FSdfLayer SubLayer = UsdUtils::FindLayerForSubLayerPath(Layer, SubLayerPath))
 				{
-					if ( ULevelSequence* SubSequence = FindOrAddSequenceForLayer( SubLayer, SubLayer.GetIdentifier(), SubLayer.GetDisplayName() ) )
+					if (ULevelSequence* SubSequence = FindOrAddSequenceForLayer(SubLayer, SubLayer.GetIdentifier(), SubLayer.GetDisplayName()))
 					{
-						if ( !LocalLayersSequences.Contains( SubSequence->GetFName() ) ) // Make sure we don't parse an already parsed layer
+						if (!LocalLayersSequences.Contains(SubSequence->GetFName())) // Make sure we don't parse an already parsed layer
 						{
-							LocalLayersSequences.Add( SubSequence->GetFName() );
+							LocalLayersSequences.Add(SubSequence->GetFName());
 
-							CreateSubSequenceSection( ParentSequence, *SubSequence );
+							CreateSubSequenceSection(ParentSequence, *SubSequence);
 
-							RecursivelyCreateSequencesForLayer( FindLayerTimeInfo( SubLayer ), *SubSequence );
+							RecursivelyCreateSequencesForLayer(FindLayerTimeInfo(SubLayer), *SubSequence);
 						}
 					}
 				}
@@ -799,22 +805,22 @@ void FUsdLevelSequenceHelperImpl::CreateLocalLayersSequences()
 	};
 
 	// Create level sequences for all sub layers (accessible via the main level sequence but otherwise hidden)
-	RecursivelyCreateSequencesForLayer( &RootLayerInfo, *MainLevelSequence );
+	RecursivelyCreateSequencesForLayer(&RootLayerInfo, *MainLevelSequence);
 }
 
-void FUsdLevelSequenceHelperImpl::BindToUsdStageActor( AUsdStageActor* InStageActor )
+void FUsdLevelSequenceHelperImpl::BindToUsdStageActor(AUsdStageActor* InStageActor)
 {
 	UnbindFromUsdStageActor();
 
 	StageActor = InStageActor;
-	SetAssetCache( InStageActor ? InStageActor->GetAssetCache() : nullptr );
+	SetAssetCache(InStageActor ? InStageActor->GetAssetCache() : nullptr);
 
-	if ( !StageActor.IsValid() || !MainLevelSequence || !MainLevelSequence->GetMovieScene() )
+	if (!StageActor.IsValid() || !MainLevelSequence || !MainLevelSequence->GetMovieScene())
 	{
 		return;
 	}
 
-	OnUsdObjectsChangedHandle = StageActor->GetUsdListener().GetOnObjectsChanged().AddRaw( this, &FUsdLevelSequenceHelperImpl::OnUsdObjectsChanged );
+	OnUsdObjectsChangedHandle = StageActor->GetUsdListener().GetOnObjectsChanged().AddRaw(this, &FUsdLevelSequenceHelperImpl::OnUsdObjectsChanged);
 
 	// Bind stage actor
 	StageActorBinding = MainLevelSequence->GetMovieScene()->AddPossessable(
@@ -825,36 +831,36 @@ void FUsdLevelSequenceHelperImpl::BindToUsdStageActor( AUsdStageActor* InStageAc
 #endif // WITH_EDITOR
 		StageActor->GetClass()
 	);
-	MainLevelSequence->BindPossessableObject( StageActorBinding, *StageActor, StageActor->GetWorld() );
+	MainLevelSequence->BindPossessableObject(StageActorBinding, *StageActor, StageActor->GetWorld());
 
-	CreateTimeTrack( FindOrAddLayerTimeInfo( UsdStage.GetRootLayer() ) );
+	CreateTimeTrack(FindOrAddLayerTimeInfo(UsdStage.GetRootLayer()));
 }
 
 void FUsdLevelSequenceHelperImpl::UnbindFromUsdStageActor()
 {
-	if ( UsdStage )
+	if (UsdStage)
 	{
-		RemoveTimeTrack( FindLayerTimeInfo( UsdStage.GetRootLayer() ) );
+		RemoveTimeTrack(FindLayerTimeInfo(UsdStage.GetRootLayer()));
 	}
 
-	if ( MainLevelSequence && MainLevelSequence->GetMovieScene() )
+	if (MainLevelSequence && MainLevelSequence->GetMovieScene())
 	{
-		if ( MainLevelSequence->GetMovieScene()->RemovePossessable( StageActorBinding ) )
+		if (MainLevelSequence->GetMovieScene()->RemovePossessable(StageActorBinding))
 		{
-			MainLevelSequence->UnbindPossessableObjects( StageActorBinding );
+			MainLevelSequence->UnbindPossessableObjects(StageActorBinding);
 		}
 	}
 
 	StageActorBinding = FGuid::NewGuid();
 
-	if ( StageActor.IsValid() )
+	if (StageActor.IsValid())
 	{
-		StageActor->GetUsdListener().GetOnStageEditTargetChanged().Remove( OnStageEditTargetChangedHandle );
-		StageActor->GetUsdListener().GetOnObjectsChanged().Remove( OnUsdObjectsChangedHandle );
+		StageActor->GetUsdListener().GetOnStageEditTargetChanged().Remove(OnStageEditTargetChangedHandle);
+		StageActor->GetUsdListener().GetOnObjectsChanged().Remove(OnUsdObjectsChangedHandle);
 		StageActor.Reset();
 	}
 
-	SetAssetCache( nullptr );
+	SetAssetCache(nullptr);
 
 	OnStageEditTargetChangedHandle.Reset();
 }
@@ -862,12 +868,12 @@ void FUsdLevelSequenceHelperImpl::UnbindFromUsdStageActor()
 void FUsdLevelSequenceHelperImpl::OnStageActorRenamed()
 {
 	AUsdStageActor* StageActorPtr = StageActor.Get();
-	if ( !StageActorPtr )
+	if (!StageActorPtr)
 	{
 		return;
 	}
 
-	FMovieScenePossessable NewPossessable {
+	FMovieScenePossessable NewPossessable{
 	#if WITH_EDITOR
 		StageActorPtr->GetActorLabel(),
 #else
@@ -878,51 +884,51 @@ void FUsdLevelSequenceHelperImpl::OnStageActorRenamed()
 	FGuid NewId = NewPossessable.GetGuid();
 
 	bool bDidSomething = false;
-	for ( const TPair<FString, ULevelSequence*>& Pair : LevelSequencesByIdentifier )
+	for (const TPair<FString, ULevelSequence*>& Pair : LevelSequencesByIdentifier)
 	{
 		ULevelSequence* Sequence = Pair.Value;
-		if ( !Sequence )
+		if (!Sequence)
 		{
 			continue;
 		}
 
 		UMovieScene* MovieScene = Sequence->GetMovieScene();
-		if ( !MovieScene )
+		if (!MovieScene)
 		{
 			continue;
 		}
 
-		const bool bDidRenameMovieScene = MovieScene->ReplacePossessable( StageActorBinding, NewPossessable );
-		if ( bDidRenameMovieScene )
+		const bool bDidRenameMovieScene = MovieScene->ReplacePossessable(StageActorBinding, NewPossessable);
+		if (bDidRenameMovieScene)
 		{
-			Sequence->UnbindPossessableObjects( NewId );
-			Sequence->BindPossessableObject( NewId, *StageActorPtr, StageActorPtr->GetWorld() );
+			Sequence->UnbindPossessableObjects(NewId);
+			Sequence->BindPossessableObject(NewId, *StageActorPtr, StageActorPtr->GetWorld());
 
 			bDidSomething = true;
 		}
 	}
 
-	if ( bDidSomething )
+	if (bDidSomething)
 	{
 		StageActorBinding = NewId;
 	}
 }
 
-ULevelSequence* FUsdLevelSequenceHelperImpl::FindSequenceForAttribute( const UE::FUsdAttribute& Attribute )
+ULevelSequence* FUsdLevelSequenceHelperImpl::FindSequenceForAttribute(const UE::FUsdAttribute& Attribute)
 {
-	if ( !Attribute || !Attribute.GetPrim() )
+	if (!Attribute || !Attribute.GetPrim())
 	{
 		return nullptr;
 	}
 
-	if ( !UsdStage )
+	if (!UsdStage)
 	{
 		return nullptr;
 	}
 
-	UE::FSdfLayer AttributeLayer = UsdUtils::FindLayerForAttribute( Attribute, 0.0 );
+	UE::FSdfLayer AttributeLayer = UsdUtils::FindLayerForAttribute(Attribute, 0.0);
 
-	if ( !AttributeLayer )
+	if (!AttributeLayer)
 	{
 		return nullptr;
 	}
@@ -931,56 +937,56 @@ ULevelSequence* FUsdLevelSequenceHelperImpl::FindSequenceForAttribute( const UE:
 
 	UE::FUsdPrim Prim = Attribute.GetPrim();
 
-	UE::FSdfLayer PrimLayer = UsdUtils::FindLayerForPrim( Prim );
+	UE::FSdfLayer PrimLayer = UsdUtils::FindLayerForPrim(Prim);
 	FString PrimLayerIdentifier = PrimLayer.GetIdentifier();
 
 	ULevelSequence* Sequence = nullptr;
 
 	// If the attribute is on the Root or a SubLayer, return the Sequence associated with that layer
-	if ( AttributeLayer.HasSpec( Prim.GetPrimPath() ) && UsdStage.HasLocalLayer( AttributeLayer ) )
+	if (AttributeLayer.HasSpec(Prim.GetPrimPath()) && UsdStage.HasLocalLayer(AttributeLayer))
 	{
-		Sequence = FindSequenceForIdentifier( AttributeLayer.GetIdentifier() );
+		Sequence = FindSequenceForIdentifier(AttributeLayer.GetIdentifier());
 	}
 	// The prim should have its own sequence, return that
 	else
 	{
-		Sequence = FindSequenceForIdentifier( Prim.GetPrimPath().GetString() );
+		Sequence = FindSequenceForIdentifier(Prim.GetPrimPath().GetString());
 	}
 
 	return Sequence;
 }
 
-ULevelSequence* FUsdLevelSequenceHelperImpl::FindOrAddSequenceForAttribute( const UE::FUsdAttribute& Attribute )
+ULevelSequence* FUsdLevelSequenceHelperImpl::FindOrAddSequenceForAttribute(const UE::FUsdAttribute& Attribute)
 {
-	if ( !Attribute || !Attribute.GetPrim() )
+	if (!Attribute || !Attribute.GetPrim())
 	{
 		return nullptr;
 	}
 
-	ULevelSequence* Sequence = FindSequenceForAttribute( Attribute );
-	if ( !Sequence )
+	ULevelSequence* Sequence = FindSequenceForAttribute(Attribute);
+	if (!Sequence)
 	{
-		if ( UE::FSdfLayer AttributeLayer = UsdUtils::FindLayerForAttribute( Attribute, 0.0 ) )
+		if (UE::FSdfLayer AttributeLayer = UsdUtils::FindLayerForAttribute(Attribute, 0.0))
 		{
 			const FString SequenceIdentifier = Attribute.GetPrim().GetPrimPath().GetString();
 
-			Sequence = FindOrAddSequenceForLayer( AttributeLayer, SequenceIdentifier, SequenceIdentifier );
+			Sequence = FindOrAddSequenceForLayer(AttributeLayer, SequenceIdentifier, SequenceIdentifier);
 		}
 	}
 
 	return Sequence;
 }
 
-ULevelSequence* FUsdLevelSequenceHelperImpl::FindOrAddSequenceForLayer( const UE::FSdfLayer& Layer, const FString& SequenceIdentifier, const FString& SequenceDisplayName )
+ULevelSequence* FUsdLevelSequenceHelperImpl::FindOrAddSequenceForLayer(const UE::FSdfLayer& Layer, const FString& SequenceIdentifier, const FString& SequenceDisplayName)
 {
-	if ( !Layer )
+	if (!Layer)
 	{
 		return nullptr;
 	}
 
-	ULevelSequence* Sequence = FindSequenceForIdentifier( SequenceIdentifier );
+	ULevelSequence* Sequence = FindSequenceForIdentifier(SequenceIdentifier);
 
-	if ( !Sequence )
+	if (!Sequence)
 	{
 		// This needs to be unique, or else when we reload the stage we will end up with a new ULevelSequence with the same class, outer and name as the
 		// previous one. Also note that the previous level sequence, even though unreferenced by the stage actor, is likely still alive and valid due to references
@@ -991,69 +997,69 @@ ULevelSequence* FUsdLevelSequenceHelperImpl::FindOrAddSequenceForLayer( const UE
 		// which also gives us a stage actor. So if we don't have an actor but have a cache, we're importing
 		const bool bIsImporting = StageActor.IsExplicitlyNull() && AssetCache;
 		FName UniqueSequenceName = bIsImporting
-			? *UsdLevelSequenceHelperImpl::SanitizeObjectName( FPaths::GetBaseFilename( SequenceDisplayName ) )
-			: MakeUniqueObjectName( GetTransientPackage(), ULevelSequence::StaticClass(), *UsdLevelSequenceHelperImpl::SanitizeObjectName( FPaths::GetBaseFilename( SequenceDisplayName ) ) );
+			? *UsdLevelSequenceHelperImpl::SanitizeObjectName(FPaths::GetBaseFilename(SequenceDisplayName))
+			: MakeUniqueObjectName(GetTransientPackage(), ULevelSequence::StaticClass(), *UsdLevelSequenceHelperImpl::SanitizeObjectName(FPaths::GetBaseFilename(SequenceDisplayName)));
 
-		Sequence = NewObject< ULevelSequence >( GetTransientPackage(), UniqueSequenceName, FUsdLevelSequenceHelperImpl::DefaultObjFlags );
+		Sequence = NewObject< ULevelSequence >(GetTransientPackage(), UniqueSequenceName, FUsdLevelSequenceHelperImpl::DefaultObjFlags);
 		Sequence->Initialize();
 
 		UMovieScene* MovieScene = Sequence->MovieScene;
-		if ( !MovieScene )
+		if (!MovieScene)
 		{
 			return nullptr;
 		}
 
-		LayerIdentifierByLevelSequenceName.Add( Sequence->GetFName(), Layer.GetIdentifier() );
-		LevelSequencesByIdentifier.Add( SequenceIdentifier, Sequence );
-		IdentifierByLevelSequence.Add( Sequence, SequenceIdentifier );
+		LayerIdentifierByLevelSequenceName.Add(Sequence->GetFName(), Layer.GetIdentifier());
+		LevelSequencesByIdentifier.Add(SequenceIdentifier, Sequence);
+		IdentifierByLevelSequence.Add(Sequence, SequenceIdentifier);
 
-		const FLayerTimeInfo LayerTimeInfo = FindOrAddLayerTimeInfo( Layer );
+		const FLayerTimeInfo LayerTimeInfo = FindOrAddLayerTimeInfo(Layer);
 
-		UpdateMovieSceneTimeRanges( *MovieScene, LayerTimeInfo );
+		UpdateMovieSceneTimeRanges(*MovieScene, LayerTimeInfo);
 
-		UE_LOG( LogUsd, Verbose, TEXT("Created Sequence for identifier: '%s'"), *SequenceIdentifier );
+		UE_LOG(LogUsd, Verbose, TEXT("Created Sequence for identifier: '%s'"), *SequenceIdentifier);
 	}
 
 	return Sequence;
 }
 
-UE::FSdfLayer FUsdLevelSequenceHelperImpl::FindEditTargetForSubsequence( ULevelSequence* Sequence )
+UE::FSdfLayer FUsdLevelSequenceHelperImpl::FindEditTargetForSubsequence(ULevelSequence* Sequence)
 {
 	FString LayerIdentifier = IdentifierByLevelSequence.FindRef(Sequence);
-	return UE::FSdfLayer::FindOrOpen( *LayerIdentifier );
+	return UE::FSdfLayer::FindOrOpen(*LayerIdentifier);
 }
 
-UMovieSceneSubSection* FUsdLevelSequenceHelperImpl::FindSubSequenceSection( ULevelSequence& Sequence, ULevelSequence& SubSequence )
+UMovieSceneSubSection* FUsdLevelSequenceHelperImpl::FindSubSequenceSection(ULevelSequence& Sequence, ULevelSequence& SubSequence)
 {
 	UMovieScene* MovieScene = Sequence.GetMovieScene();
-	if ( !MovieScene )
+	if (!MovieScene)
 	{
 		return nullptr;
 	}
 
 	UMovieSceneSubTrack* SubTrack = MovieScene->FindTrack< UMovieSceneSubTrack >();
 
-	if ( !SubTrack )
+	if (!SubTrack)
 	{
 		return nullptr;
 	}
 
-	UMovieSceneSection* const * SubSection = Algo::FindByPredicate( SubTrack->GetAllSections(),
-	[ &SubSequence ]( UMovieSceneSection* Section ) -> bool
-	{
-		if ( UMovieSceneSubSection* SubSection = Cast< UMovieSceneSubSection >( Section ) )
+	UMovieSceneSection* const* SubSection = Algo::FindByPredicate(SubTrack->GetAllSections(),
+		[&SubSequence](UMovieSceneSection* Section) -> bool
 		{
-			return  ( SubSection->GetSequence() == &SubSequence );
-		}
-		else
-		{
-			return false;
-		}
-	} );
+			if (UMovieSceneSubSection* SubSection = Cast< UMovieSceneSubSection >(Section))
+			{
+				return  (SubSection->GetSequence() == &SubSequence);
+			}
+			else
+			{
+				return false;
+			}
+		});
 
-	if ( SubSection )
+	if (SubSection)
 	{
-		return  Cast< UMovieSceneSubSection >( *SubSection );
+		return  Cast< UMovieSceneSubSection >(*SubSection);
 	}
 	else
 	{
@@ -1061,20 +1067,20 @@ UMovieSceneSubSection* FUsdLevelSequenceHelperImpl::FindSubSequenceSection( ULev
 	}
 }
 
-void FUsdLevelSequenceHelperImpl::CreateSubSequenceSection( ULevelSequence& Sequence, ULevelSequence& SubSequence )
+void FUsdLevelSequenceHelperImpl::CreateSubSequenceSection(ULevelSequence& Sequence, ULevelSequence& SubSequence)
 {
-	if ( &Sequence == &SubSequence )
+	if (&Sequence == &SubSequence)
 	{
 		return;
 	}
 
 	UMovieScene* MovieScene = Sequence.GetMovieScene();
-	if ( !MovieScene )
+	if (!MovieScene)
 	{
 		return;
 	}
 
-	if ( !UsdStage )
+	if (!UsdStage)
 	{
 		return;
 	}
@@ -1082,55 +1088,55 @@ void FUsdLevelSequenceHelperImpl::CreateSubSequenceSection( ULevelSequence& Sequ
 	FFrameRate TickResolution = MovieScene->GetTickResolution();
 
 	UMovieSceneSubTrack* SubTrack = MovieScene->FindTrack< UMovieSceneSubTrack >();
-	if ( !SubTrack )
+	if (!SubTrack)
 	{
 		SubTrack = MovieScene->AddTrack< UMovieSceneSubTrack >();
 	}
 
-	const FString* LayerIdentifier = LayerIdentifierByLevelSequenceName.Find( Sequence.GetFName() );
-	const FString* SubLayerIdentifier = LayerIdentifierByLevelSequenceName.Find( SubSequence.GetFName() );
+	const FString* LayerIdentifier = LayerIdentifierByLevelSequenceName.Find(Sequence.GetFName());
+	const FString* SubLayerIdentifier = LayerIdentifierByLevelSequenceName.Find(SubSequence.GetFName());
 
-  	if ( !LayerIdentifier || !SubLayerIdentifier )
+	if (!LayerIdentifier || !SubLayerIdentifier)
 	{
 		return;
 	}
 
-	FLayerTimeInfo* LayerTimeInfo = LayerTimeInfosByLayerIdentifier.Find( *LayerIdentifier );
-	FLayerTimeInfo* SubLayerTimeInfo = LayerTimeInfosByLayerIdentifier.Find( *SubLayerIdentifier );
+	FLayerTimeInfo* LayerTimeInfo = LayerTimeInfosByLayerIdentifier.Find(*LayerIdentifier);
+	FLayerTimeInfo* SubLayerTimeInfo = LayerTimeInfosByLayerIdentifier.Find(*SubLayerIdentifier);
 
-	if ( !LayerTimeInfo || !SubLayerTimeInfo )
+	if (!LayerTimeInfo || !SubLayerTimeInfo)
 	{
 		return;
 	}
 
 	UE::FSdfLayerOffset SubLayerOffset;
 
-	UE::FSdfLayer Layer = UE::FSdfLayer::FindOrOpen( **LayerIdentifier );
-	UE::FSdfLayer SubLayer = UE::FSdfLayer::FindOrOpen( **SubLayerIdentifier );
+	UE::FSdfLayer Layer = UE::FSdfLayer::FindOrOpen(**LayerIdentifier);
+	UE::FSdfLayer SubLayer = UE::FSdfLayer::FindOrOpen(**SubLayerIdentifier);
 
 	TArray< FString > PrimPathsForSequence;
-	PrimPathByLevelSequenceName.MultiFind( SubSequence.GetFName(), PrimPathsForSequence );
+	PrimPathByLevelSequenceName.MultiFind(SubSequence.GetFName(), PrimPathsForSequence);
 
-	if ( PrimPathsForSequence.Num() > 0 )
+	if (PrimPathsForSequence.Num() > 0)
 	{
-		if ( UE::FUsdPrim SequencePrim = UsdStage.GetPrimAtPath( UE::FSdfPath( *PrimPathsForSequence[0] ) ) )
+		if (UE::FUsdPrim SequencePrim = UsdStage.GetPrimAtPath(UE::FSdfPath(*PrimPathsForSequence[0])))
 		{
-			TArray<UE::FUsdAttribute> Attrs = UnrealToUsd::GetAttributesForProperty( SequencePrim, UnrealIdentifiers::TransformPropertyName );
-			if ( Attrs.Num() > 0 )
+			TArray<UE::FUsdAttribute> Attrs = UnrealToUsd::GetAttributesForProperty(SequencePrim, UnrealIdentifiers::TransformPropertyName);
+			if (Attrs.Num() > 0)
 			{
-				SubLayerOffset = UsdUtils::GetLayerToStageOffset( Attrs[0] );
+				SubLayerOffset = UsdUtils::GetLayerToStageOffset(Attrs[0]);
 			}
 		}
 	}
-	else if ( UsdStage.HasLocalLayer( SubLayer ) )
+	else if (UsdStage.HasLocalLayer(SubLayer))
 	{
-		const FLayerOffsetInfo* SubLayerOffsetPtr = Algo::FindByPredicate( LayerTimeInfo->SubLayersOffsets,
-			[ &SubLayerIdentifier ]( const FLayerOffsetInfo& Other )
-		{
-			return ( Other.LayerIdentifier == *SubLayerIdentifier );
-		} );
+		const FLayerOffsetInfo* SubLayerOffsetPtr = Algo::FindByPredicate(LayerTimeInfo->SubLayersOffsets,
+			[&SubLayerIdentifier](const FLayerOffsetInfo& Other)
+			{
+				return (Other.LayerIdentifier == *SubLayerIdentifier);
+			});
 
-		if ( SubLayerOffsetPtr )
+		if (SubLayerOffsetPtr)
 		{
 			SubLayerOffset = SubLayerOffsetPtr->LayerOffset;
 		}
@@ -1146,27 +1152,27 @@ void FUsdLevelSequenceHelperImpl::CreateSubSequenceSection( ULevelSequence& Sequ
 	const double SubStartTimeSeconds = SubLayerOffset.Offset / TimeCodesPerSecond;
 	const double SubEndTimeSeconds = SubStartTimeSeconds + SubDurationSeconds;
 
-	const FFrameNumber StartFrame = UsdLevelSequenceHelperImpl::RoundAsFrameNumber( TickResolution, SubStartTimeSeconds );
-	const FFrameNumber EndFrame = UsdLevelSequenceHelperImpl::RoundAsFrameNumber( TickResolution, SubEndTimeSeconds );
+	const FFrameNumber StartFrame = UsdLevelSequenceHelperImpl::RoundAsFrameNumber(TickResolution, SubStartTimeSeconds);
+	const FFrameNumber EndFrame = UsdLevelSequenceHelperImpl::RoundAsFrameNumber(TickResolution, SubEndTimeSeconds);
 
 	// Don't clip subsections with their duration, so that the root layer's [startTimeCode, endTimeCode] range is the only thing clipping
 	// anything, as this is how USD seems to behave. Even if a middle sublayer has startTimeCode == endTimeCode, its animations
 	// (or its child sublayers') won't be clipped by it and play according to the stage's range
 	const double StageEndTimeSeconds = UsdStage.GetEndTimeCode() / UsdStage.GetTimeCodesPerSecond();
-	const FFrameNumber StageEndFrame = UsdLevelSequenceHelperImpl::RoundAsFrameNumber( TickResolution, StageEndTimeSeconds );
+	const FFrameNumber StageEndFrame = UsdLevelSequenceHelperImpl::RoundAsFrameNumber(TickResolution, StageEndTimeSeconds);
 
 	// Max here because StartFrame can theoretically be larger than StageEndFrame, which would generate a range where the upper bound is smaller
 	// than the lower bound, which can trigger asserts
-	TRange< FFrameNumber > SubSectionRange{ StartFrame, FMath::Max( StageEndFrame, EndFrame ) };
+	TRange< FFrameNumber > SubSectionRange{StartFrame, FMath::Max(StageEndFrame, EndFrame)};
 
-	UMovieSceneSubSection* SubSection = FindSubSequenceSection( Sequence, SubSequence );
-	if ( SubSection )
+	UMovieSceneSubSection* SubSection = FindSubSequenceSection(Sequence, SubSequence);
+	if (SubSection)
 	{
-		SubSection->SetRange( SubSectionRange );
+		SubSection->SetRange(SubSectionRange);
 	}
 	else
 	{
-		SubSection = SubTrack->AddSequence( &SubSequence, SubSectionRange.GetLowerBoundValue(), SubSectionRange.Size< FFrameNumber >().Value );
+		SubSection = SubTrack->AddSequence(&SubSequence, SubSectionRange.GetLowerBoundValue(), SubSectionRange.Size< FFrameNumber >().Value);
 
 		UE_LOG(LogUsd, Verbose, TEXT("Adding subsection '%s' to sequence '%s'. StartFrame: '%d'"),
 			*SubSection->GetName(),
@@ -1176,19 +1182,19 @@ void FUsdLevelSequenceHelperImpl::CreateSubSequenceSection( ULevelSequence& Sequ
 	}
 
 	const double TimeCodesPerSecondDifference = TimeCodesPerSecond / SubLayer.GetTimeCodesPerSecond();
-	SubSection->Parameters.TimeScale = FMath::IsNearlyZero( SubLayerOffset.Scale ) ? 0.f : 1.f / ( SubLayerOffset.Scale / TimeCodesPerSecondDifference );
+	SubSection->Parameters.TimeScale = FMath::IsNearlyZero(SubLayerOffset.Scale) ? 0.f : 1.f / (SubLayerOffset.Scale / TimeCodesPerSecondDifference);
 
-	if ( MainLevelSequence )
+	if (MainLevelSequence)
 	{
-		UMovieSceneCompiledDataManager::CompileHierarchy( MainLevelSequence, &SequenceHierarchyCache, EMovieSceneServerClientMask::All );
+		UMovieSceneCompiledDataManager::CompileHierarchy(MainLevelSequence, &SequenceHierarchyCache, EMovieSceneServerClientMask::All);
 
-		for ( const TTuple< FMovieSceneSequenceID, FMovieSceneSubSequenceData >& Pair : SequenceHierarchyCache.AllSubSequenceData() )
+		for (const TTuple< FMovieSceneSequenceID, FMovieSceneSubSequenceData >& Pair : SequenceHierarchyCache.AllSubSequenceData())
 		{
-			if ( UMovieSceneSequence* CachedSubSequence = Pair.Value.GetSequence() )
+			if (UMovieSceneSequence* CachedSubSequence = Pair.Value.GetSequence())
 			{
-				if ( CachedSubSequence == &SubSequence )
+				if (CachedSubSequence == &SubSequence)
 				{
-					SequencesID.Add( &SubSequence, Pair.Key );
+					SequencesID.Add(&SubSequence, Pair.Key);
 					break;
 				}
 			}
@@ -1196,19 +1202,19 @@ void FUsdLevelSequenceHelperImpl::CreateSubSequenceSection( ULevelSequence& Sequ
 	}
 }
 
-void FUsdLevelSequenceHelperImpl::RemoveSubSequenceSection( ULevelSequence& Sequence, ULevelSequence& SubSequence )
+void FUsdLevelSequenceHelperImpl::RemoveSubSequenceSection(ULevelSequence& Sequence, ULevelSequence& SubSequence)
 {
-	if ( UMovieSceneSubTrack* SubTrack = Sequence.GetMovieScene()->FindTrack< UMovieSceneSubTrack >() )
+	if (UMovieSceneSubTrack* SubTrack = Sequence.GetMovieScene()->FindTrack< UMovieSceneSubTrack >())
 	{
-		if ( UMovieSceneSection* SubSection = FindSubSequenceSection( Sequence, SubSequence ) )
+		if (UMovieSceneSection* SubSection = FindSubSequenceSection(Sequence, SubSequence))
 		{
-			SequencesID.Remove( &SubSequence );
+			SequencesID.Remove(&SubSequence);
 			SubTrack->Modify();
-			SubTrack->RemoveSection( *SubSection );
+			SubTrack->RemoveSection(*SubSection);
 
-			if ( MainLevelSequence )
+			if (MainLevelSequence)
 			{
-				UMovieSceneCompiledDataManager::CompileHierarchy( MainLevelSequence, &SequenceHierarchyCache, EMovieSceneServerClientMask::All );
+				UMovieSceneCompiledDataManager::CompileHierarchy(MainLevelSequence, &SequenceHierarchyCache, EMovieSceneServerClientMask::All);
 			}
 		}
 	}
@@ -1216,7 +1222,7 @@ void FUsdLevelSequenceHelperImpl::RemoveSubSequenceSection( ULevelSequence& Sequ
 
 void FUsdLevelSequenceHelperImpl::CreateTimeTrack(const FLayerTimeInfo& Info)
 {
-	ULevelSequence* Sequence = FindSequenceForIdentifier( Info.Identifier );
+	ULevelSequence* Sequence = FindSequenceForIdentifier(Info.Identifier);
 
 	if (!Sequence || !StageActorBinding.IsValid())
 	{
@@ -1248,28 +1254,28 @@ void FUsdLevelSequenceHelperImpl::CreateTimeTrack(const FLayerTimeInfo& Info)
 		MovieScene->SetEvaluationType(EMovieSceneEvaluationType::FrameLocked);
 	}
 
-	if ( Info.IsAnimated() )
+	if (Info.IsAnimated())
 	{
 		const double StartTimeCode = Info.StartTimeCode.Get(0.0);
 		const double EndTimeCode = Info.EndTimeCode.Get(0.0);
 		const double TimeCodesPerSecond = GetTimeCodesPerSecond();
 
 		FFrameRate DestTickRate = MovieScene->GetTickResolution();
-		FFrameNumber StartFrame  = UsdLevelSequenceHelperImpl::RoundAsFrameNumber(DestTickRate, StartTimeCode / TimeCodesPerSecond);
-		FFrameNumber EndFrame    = UsdLevelSequenceHelperImpl::RoundAsFrameNumber(DestTickRate, EndTimeCode / TimeCodesPerSecond);
+		FFrameNumber StartFrame = UsdLevelSequenceHelperImpl::RoundAsFrameNumber(DestTickRate, StartTimeCode / TimeCodesPerSecond);
+		FFrameNumber EndFrame = UsdLevelSequenceHelperImpl::RoundAsFrameNumber(DestTickRate, EndTimeCode / TimeCodesPerSecond);
 
-		TRange< FFrameNumber > PlaybackRange( StartFrame, EndFrame );
+		TRange< FFrameNumber > PlaybackRange(StartFrame, EndFrame);
 
 		bool bSectionAdded = false;
 
-		if ( UMovieSceneFloatSection* TimeSection = Cast<UMovieSceneFloatSection>(TimeTrack->FindOrAddSection(0, bSectionAdded)) )
+		if (UMovieSceneFloatSection* TimeSection = Cast<UMovieSceneFloatSection>(TimeTrack->FindOrAddSection(0, bSectionAdded)))
 		{
 			TimeSection->EvalOptions.CompletionMode = EMovieSceneCompletionMode::KeepState;
 			TimeSection->SetRange(TRange<FFrameNumber>::All());
 
 			TArray<FFrameNumber> FrameNumbers;
-			FrameNumbers.Add( UE::MovieScene::DiscreteInclusiveLower( PlaybackRange ) );
-			FrameNumbers.Add( UE::MovieScene::DiscreteExclusiveUpper( PlaybackRange ) );
+			FrameNumbers.Add(UE::MovieScene::DiscreteInclusiveLower(PlaybackRange));
+			FrameNumbers.Add(UE::MovieScene::DiscreteExclusiveUpper(PlaybackRange));
 
 			TArray<FMovieSceneFloatValue> FrameValues;
 			FrameValues.Add_GetRef(FMovieSceneFloatValue(StartTimeCode)).InterpMode = ERichCurveInterpMode::RCIM_Linear;
@@ -1285,44 +1291,44 @@ void FUsdLevelSequenceHelperImpl::CreateTimeTrack(const FLayerTimeInfo& Info)
 
 void FUsdLevelSequenceHelperImpl::RemoveTimeTrack(const FLayerTimeInfo* LayerTimeInfo)
 {
-	if ( !UsdStage || !LayerTimeInfo || !StageActorBinding.IsValid() )
+	if (!UsdStage || !LayerTimeInfo || !StageActorBinding.IsValid())
 	{
 		return;
 	}
 
-	ULevelSequence* Sequence = FindSequenceForIdentifier( LayerTimeInfo->Identifier );
+	ULevelSequence* Sequence = FindSequenceForIdentifier(LayerTimeInfo->Identifier);
 
-	if ( !Sequence )
+	if (!Sequence)
 	{
 		return;
 	}
 
 	UMovieScene* MovieScene = Sequence->GetMovieScene();
 
-	if ( !MovieScene )
+	if (!MovieScene)
 	{
 		return;
 	}
 
-	UMovieSceneFloatTrack* TimeTrack = MovieScene->FindTrack< UMovieSceneFloatTrack >( StageActorBinding, FName( FUsdLevelSequenceHelperImpl::TimeTrackName ) );
-	if ( TimeTrack )
+	UMovieSceneFloatTrack* TimeTrack = MovieScene->FindTrack< UMovieSceneFloatTrack >(StageActorBinding, FName(FUsdLevelSequenceHelperImpl::TimeTrackName));
+	if (TimeTrack)
 	{
-		MovieScene->RemoveTrack( *TimeTrack );
+		MovieScene->RemoveTrack(*TimeTrack);
 	}
 }
 
-void FUsdLevelSequenceHelperImpl::AddCommonTracks( const UUsdPrimTwin& PrimTwin, const UE::FUsdPrim& Prim, bool bForceVisibilityTracks )
+void FUsdLevelSequenceHelperImpl::AddCommonTracks(const UUsdPrimTwin& PrimTwin, const UE::FUsdPrim& Prim, bool bForceVisibilityTracks)
 {
 	USceneComponent* ComponentToBind = PrimTwin.GetSceneComponent();
-	if ( !ComponentToBind )
+	if (!ComponentToBind)
 	{
 		return;
 	}
 
-	UE::FSdfLayer PrimLayer = UsdUtils::FindLayerForPrim( Prim );
-	ULevelSequence* PrimSequence = FindSequenceForIdentifier( PrimLayer.GetIdentifier() );
+	UE::FSdfLayer PrimLayer = UsdUtils::FindLayerForPrim(Prim);
+	ULevelSequence* PrimSequence = FindSequenceForIdentifier(PrimLayer.GetIdentifier());
 
-	UE::FUsdGeomXformable Xformable( Prim );
+	UE::FUsdGeomXformable Xformable(Prim);
 
 	// If this xformable has an op to reset the xform stack and one of its ancestors is animated, then we need to make
 	// a transform track for it even if its transform is not animated by itself. This because that op effectively means
@@ -1331,25 +1337,25 @@ void FUsdLevelSequenceHelperImpl::AddCommonTracks( const UUsdPrimTwin& PrimTwin,
 	// being animated, we'll need to recompute this for every animation keyframe
 	TArray<double> AncestorTimeSamples;
 	bool bNeedTrackToCompensateResetXformOp = false;
-	if ( Xformable.GetResetXformStack() )
+	if (Xformable.GetResetXformStack())
 	{
 		UE::FUsdPrim AncestorPrim = Prim.GetParent();
-		while ( AncestorPrim && !AncestorPrim.IsPseudoRoot() )
+		while (AncestorPrim && !AncestorPrim.IsPseudoRoot())
 		{
-			if ( UE::FUsdGeomXformable AncestorXformable{ AncestorPrim } )
+			if (UE::FUsdGeomXformable AncestorXformable{AncestorPrim})
 			{
 				TArray<double> TimeSamples;
-				if ( AncestorXformable.GetTimeSamples( &TimeSamples ) && TimeSamples.Num() > 0 )
+				if (AncestorXformable.GetTimeSamples(&TimeSamples) && TimeSamples.Num() > 0)
 				{
 					bNeedTrackToCompensateResetXformOp = true;
 
-					AncestorTimeSamples.Append( TimeSamples );
+					AncestorTimeSamples.Append(TimeSamples);
 				}
 
 				// The exception is if our ancestor also wants to reset its xform stack (i.e. its transform is meant to be
 				// used as the world transform). In this case we don't need to care about higher up ancestors anymore, as
 				// their transforms wouldn't affect below this prim anyway
-				if ( AncestorXformable.GetResetXformStack() )
+				if (AncestorXformable.GetResetXformStack())
 				{
 					break;
 				}
@@ -1362,11 +1368,11 @@ void FUsdLevelSequenceHelperImpl::AddCommonTracks( const UUsdPrimTwin& PrimTwin,
 	// Check whether we should ignore the prim's local transform or not. We only do this for SkelRoots, in case we
 	// already appended their transform animations to the root bone as additional root motion
 	bool bIgnorePrimLocalTransform = false;
-	if ( AUsdStageActor* StageActorPtr = StageActor.Get() )
+	if (AUsdStageActor* StageActorPtr = StageActor.Get())
 	{
-		if ( StageActorPtr->RootMotionHandling == EUsdRootMotionHandling::UseMotionFromSkelRoot )
+		if (StageActorPtr->RootMotionHandling == EUsdRootMotionHandling::UseMotionFromSkelRoot)
 		{
-			if ( Prim.IsA( TEXT( "SkelRoot" ) ) )
+			if (Prim.IsA(TEXT("SkelRoot")))
 			{
 				bIgnorePrimLocalTransform = true;
 			}
@@ -1374,42 +1380,42 @@ void FUsdLevelSequenceHelperImpl::AddCommonTracks( const UUsdPrimTwin& PrimTwin,
 	}
 
 	TArray<double> TransformTimeSamples;
-	if ( ( Xformable.GetTimeSamples( &TransformTimeSamples ) && TransformTimeSamples.Num() > 0 ) || bNeedTrackToCompensateResetXformOp )
+	if ((Xformable.GetTimeSamples(&TransformTimeSamples) && TransformTimeSamples.Num() > 0) || bNeedTrackToCompensateResetXformOp)
 	{
-		TArray<UE::FUsdAttribute> Attrs = UnrealToUsd::GetAttributesForProperty( Prim, UnrealIdentifiers::TransformPropertyName );
-		if ( Attrs.Num() > 0 )
+		TArray<UE::FUsdAttribute> Attrs = UnrealToUsd::GetAttributesForProperty(Prim, UnrealIdentifiers::TransformPropertyName);
+		if (Attrs.Num() > 0)
 		{
-			if( UE::FUsdAttribute& TransformAttribute = Attrs[ 0 ] )
+			if (UE::FUsdAttribute& TransformAttribute = Attrs[0])
 			{
-				if ( ULevelSequence* AttributeSequence = FindOrAddSequenceForAttribute( TransformAttribute ) )
+				if (ULevelSequence* AttributeSequence = FindOrAddSequenceForAttribute(TransformAttribute))
 				{
-					const bool bIsMuted = UsdUtils::IsAttributeMuted( TransformAttribute, UsdStage );
+					const bool bIsMuted = UsdUtils::IsAttributeMuted(TransformAttribute, UsdStage);
 
 					FMovieSceneSequenceTransform SequenceTransform;
-					FMovieSceneSequenceID SequenceID = SequencesID.FindRef( AttributeSequence );
-					if ( FMovieSceneSubSequenceData* SubSequenceData = SequenceHierarchyCache.FindSubData( SequenceID ) )
+					FMovieSceneSequenceID SequenceID = SequencesID.FindRef(AttributeSequence);
+					if (FMovieSceneSubSequenceData* SubSequenceData = SequenceHierarchyCache.FindSubData(SequenceID))
 					{
 						SequenceTransform = SubSequenceData->RootToSequenceTransform;
 					}
 
-					if ( UMovieScene* MovieScene = AttributeSequence->GetMovieScene() )
+					if (UMovieScene* MovieScene = AttributeSequence->GetMovieScene())
 					{
 						TArray<double> TimeSamples;
-						if ( Xformable.GetTimeSamples( &TimeSamples ) )
+						if (Xformable.GetTimeSamples(&TimeSamples))
 						{
-							if ( bNeedTrackToCompensateResetXformOp )
+							if (bNeedTrackToCompensateResetXformOp)
 							{
-								TimeSamples.Append( AncestorTimeSamples );
+								TimeSamples.Append(AncestorTimeSamples);
 								TimeSamples.Sort();
 							}
 
-							if ( UMovieScene3DTransformTrack* TransformTrack = AddTrack<UMovieScene3DTransformTrack>( UnrealIdentifiers::TransformPropertyName, PrimTwin, *ComponentToBind, *AttributeSequence, bIsMuted ) )
+							if (UMovieScene3DTransformTrack* TransformTrack = AddTrack<UMovieScene3DTransformTrack>(UnrealIdentifiers::TransformPropertyName, PrimTwin, *ComponentToBind, *AttributeSequence, bIsMuted))
 							{
-								UsdToUnreal::FPropertyTrackReader Reader = UsdToUnreal::CreatePropertyTrackReader( Prim, UnrealIdentifiers::TransformPropertyName, bIgnorePrimLocalTransform );
-								UsdToUnreal::ConvertTransformTimeSamples( UsdStage, TimeSamples, Reader.TransformReader, *TransformTrack, SequenceTransform );
+								UsdToUnreal::FPropertyTrackReader Reader = UsdToUnreal::CreatePropertyTrackReader(Prim, UnrealIdentifiers::TransformPropertyName, bIgnorePrimLocalTransform);
+								UsdToUnreal::ConvertTransformTimeSamples(UsdStage, TimeSamples, Reader.TransformReader, *TransformTrack, SequenceTransform);
 							}
 
-							PrimPathByLevelSequenceName.AddUnique( AttributeSequence->GetFName(), Prim.GetPrimPath().GetString() );
+							PrimPathByLevelSequenceName.AddUnique(AttributeSequence->GetFName(), Prim.GetPrimPath().GetString());
 						}
 					}
 				}
@@ -1417,11 +1423,11 @@ void FUsdLevelSequenceHelperImpl::AddCommonTracks( const UUsdPrimTwin& PrimTwin,
 		}
 	}
 
-	TArray<UE::FUsdAttribute> Attrs = UnrealToUsd::GetAttributesForProperty( Prim, UnrealIdentifiers::HiddenInGamePropertyName );
-	if ( Attrs.Num() > 0 )
+	TArray<UE::FUsdAttribute> Attrs = UnrealToUsd::GetAttributesForProperty(Prim, UnrealIdentifiers::HiddenInGamePropertyName);
+	if (Attrs.Num() > 0)
 	{
-		UE::FUsdAttribute VisibilityAttribute = Attrs[ 0 ];
-		if ( VisibilityAttribute )
+		UE::FUsdAttribute VisibilityAttribute = Attrs[0];
+		if (VisibilityAttribute)
 		{
 			// Collect all the time samples we'll need to sample our visibility at (USD has inherited visibilities, so every time
 			// a parent has a key, we need to recompute the child visibility at that moment too)
@@ -1433,30 +1439,30 @@ void FUsdLevelSequenceHelperImpl::AddCommonTracks( const UUsdPrimTwin& PrimTwin,
 			// many parents and specs on multiple layers, but this is hopefully at least *a* reasonable answer.
 			UE::FUsdAttribute FirstAnimatedVisibilityParentAttr;
 
-			if ( ( VisibilityAttribute.GetTimeSamples( TotalVisibilityTimeSamples ) && TotalVisibilityTimeSamples.Num() > 0 )
-				|| bForceVisibilityTracks )
+			if ((VisibilityAttribute.GetTimeSamples(TotalVisibilityTimeSamples) && TotalVisibilityTimeSamples.Num() > 0)
+				|| bForceVisibilityTracks)
 			{
 				// TODO: Improve this, as this is extremely inefficient since we'll be parsing this tree for the root down and repeatedly
 				// redoing this one child at a time...
 				UE::FUsdPrim ParentPrim = Prim.GetParent();
-				while ( ParentPrim && !ParentPrim.IsPseudoRoot() )
+				while (ParentPrim && !ParentPrim.IsPseudoRoot())
 				{
-					if ( UsdUtils::HasAnimatedVisibility( ParentPrim ) )
+					if (UsdUtils::HasAnimatedVisibility(ParentPrim))
 					{
-						TArray<UE::FUsdAttribute> ParentAttrs = UnrealToUsd::GetAttributesForProperty( ParentPrim, UnrealIdentifiers::HiddenInGamePropertyName );
-						if ( ParentAttrs.Num() > 0 )
+						TArray<UE::FUsdAttribute> ParentAttrs = UnrealToUsd::GetAttributesForProperty(ParentPrim, UnrealIdentifiers::HiddenInGamePropertyName);
+						if (ParentAttrs.Num() > 0)
 						{
-							UE::FUsdAttribute ParentVisAttr = ParentAttrs[ 0 ];
+							UE::FUsdAttribute ParentVisAttr = ParentAttrs[0];
 
 							TArray<double> TimeSamples;
-							if ( ParentVisAttr && ( ParentVisAttr.GetTimeSamples( TimeSamples ) && TimeSamples.Num() ) )
+							if (ParentVisAttr && (ParentVisAttr.GetTimeSamples(TimeSamples) && TimeSamples.Num()))
 							{
-								if ( !FirstAnimatedVisibilityParentAttr )
+								if (!FirstAnimatedVisibilityParentAttr)
 								{
 									FirstAnimatedVisibilityParentAttr = ParentVisAttr;
 								}
 
-								TotalVisibilityTimeSamples.Append( TimeSamples );
+								TotalVisibilityTimeSamples.Append(TimeSamples);
 							}
 						}
 					}
@@ -1471,37 +1477,37 @@ void FUsdLevelSequenceHelperImpl::AddCommonTracks( const UUsdPrimTwin& PrimTwin,
 
 			// Pick which attribute we will use to fetch the target LevelSequence to put our baked tracks
 			UE::FUsdAttribute AttributeForSequence;
-			if ( VisibilityAttribute.GetNumTimeSamples() > 0 )
+			if (VisibilityAttribute.GetNumTimeSamples() > 0)
 			{
 				AttributeForSequence = VisibilityAttribute;
 			}
-			if ( !AttributeForSequence && bForceVisibilityTracks && FirstAnimatedVisibilityParentAttr )
+			if (!AttributeForSequence && bForceVisibilityTracks && FirstAnimatedVisibilityParentAttr)
 			{
 				AttributeForSequence = FirstAnimatedVisibilityParentAttr;
 			}
 
-			if ( AttributeForSequence && TotalVisibilityTimeSamples.Num() > 0 )
+			if (AttributeForSequence && TotalVisibilityTimeSamples.Num() > 0)
 			{
-				if ( ULevelSequence* AttributeSequence = FindOrAddSequenceForAttribute( AttributeForSequence ) )
+				if (ULevelSequence* AttributeSequence = FindOrAddSequenceForAttribute(AttributeForSequence))
 				{
-					const bool bIsMuted = UsdUtils::IsAttributeMuted( AttributeForSequence, UsdStage );
+					const bool bIsMuted = UsdUtils::IsAttributeMuted(AttributeForSequence, UsdStage);
 
 					FMovieSceneSequenceTransform SequenceTransform;
-					FMovieSceneSequenceID SequenceID = SequencesID.FindRef( AttributeSequence );
-					if ( FMovieSceneSubSequenceData* SubSequenceData = SequenceHierarchyCache.FindSubData( SequenceID ) )
+					FMovieSceneSequenceID SequenceID = SequencesID.FindRef(AttributeSequence);
+					if (FMovieSceneSubSequenceData* SubSequenceData = SequenceHierarchyCache.FindSubData(SequenceID))
 					{
 						SequenceTransform = SubSequenceData->RootToSequenceTransform;
 					}
 
-					if ( UMovieScene* MovieScene = AttributeSequence->GetMovieScene() )
+					if (UMovieScene* MovieScene = AttributeSequence->GetMovieScene())
 					{
-						if ( UMovieSceneVisibilityTrack* VisibilityTrack = AddTrack<UMovieSceneVisibilityTrack>( UnrealIdentifiers::HiddenInGamePropertyName, PrimTwin, *ComponentToBind, *AttributeSequence, bIsMuted ) )
+						if (UMovieSceneVisibilityTrack* VisibilityTrack = AddTrack<UMovieSceneVisibilityTrack>(UnrealIdentifiers::HiddenInGamePropertyName, PrimTwin, *ComponentToBind, *AttributeSequence, bIsMuted))
 						{
-							UsdToUnreal::FPropertyTrackReader Reader = UsdToUnreal::CreatePropertyTrackReader( Prim, UnrealIdentifiers::HiddenInGamePropertyName );
-							UsdToUnreal::ConvertBoolTimeSamples( UsdStage, TotalVisibilityTimeSamples, Reader.BoolReader, *VisibilityTrack, SequenceTransform );
+							UsdToUnreal::FPropertyTrackReader Reader = UsdToUnreal::CreatePropertyTrackReader(Prim, UnrealIdentifiers::HiddenInGamePropertyName);
+							UsdToUnreal::ConvertBoolTimeSamples(UsdStage, TotalVisibilityTimeSamples, Reader.BoolReader, *VisibilityTrack, SequenceTransform);
 						}
 
-						PrimPathByLevelSequenceName.AddUnique( AttributeSequence->GetFName(), Prim.GetPrimPath().GetString() );
+						PrimPathByLevelSequenceName.AddUnique(AttributeSequence->GetFName(), Prim.GetPrimPath().GetString());
 					}
 				}
 			}
@@ -1509,8 +1515,8 @@ void FUsdLevelSequenceHelperImpl::AddCommonTracks( const UUsdPrimTwin& PrimTwin,
 	}
 }
 
-void FUsdLevelSequenceHelperImpl::AddCameraTracks( const UUsdPrimTwin& PrimTwin, const UE::FUsdPrim& Prim )
- {
+void FUsdLevelSequenceHelperImpl::AddCameraTracks(const UUsdPrimTwin& PrimTwin, const UE::FUsdPrim& Prim)
+{
 	TArray<FName> TrackedProperties = {
 		UnrealIdentifiers::CurrentFocalLengthPropertyName,
 		UnrealIdentifiers::ManualFocusDistancePropertyName,
@@ -1519,77 +1525,77 @@ void FUsdLevelSequenceHelperImpl::AddCameraTracks( const UUsdPrimTwin& PrimTwin,
 		UnrealIdentifiers::SensorHeightPropertyName
 	};
 
-	UE::FSdfLayer PrimLayer = UsdUtils::FindLayerForPrim( Prim );
-	ULevelSequence* PrimSequence = FindSequenceForIdentifier( PrimLayer.GetIdentifier() );
+	UE::FSdfLayer PrimLayer = UsdUtils::FindLayerForPrim(Prim);
+	ULevelSequence* PrimSequence = FindSequenceForIdentifier(PrimLayer.GetIdentifier());
 
 	// For ACineCameraActor the camera component is not the actual root component, so we need to fetch it manually here
-	ACineCameraActor* CameraActor = Cast<ACineCameraActor>( PrimTwin.GetSceneComponent()->GetOwner() );
-	if ( !CameraActor )
+	ACineCameraActor* CameraActor = Cast<ACineCameraActor>(PrimTwin.GetSceneComponent()->GetOwner());
+	if (!CameraActor)
 	{
 		return;
 	}
 	UCineCameraComponent* ComponentToBind = CameraActor->GetCineCameraComponent();
-	if ( !ComponentToBind )
+	if (!ComponentToBind)
 	{
 		return;
 	}
 
-	for ( const FName& PropertyName : TrackedProperties )
+	for (const FName& PropertyName : TrackedProperties)
 	{
-		TArray<UE::FUsdAttribute> Attrs = UnrealToUsd::GetAttributesForProperty( Prim, PropertyName );
-		if ( Attrs.Num() < 1 )
+		TArray<UE::FUsdAttribute> Attrs = UnrealToUsd::GetAttributesForProperty(Prim, PropertyName);
+		if (Attrs.Num() < 1)
 		{
 			continue;
 		}
 
 		// Camera attributes should always match UE properties 1-to-1 here so just get the first
 		const UE::FUsdAttribute& Attr = Attrs[0];
-		if ( !Attr || Attr.GetNumTimeSamples() == 0 )
+		if (!Attr || Attr.GetNumTimeSamples() == 0)
 		{
 			continue;
 		}
 
 		// Find out the sequence where this attribute should be written to
-		if ( ULevelSequence* AttributeSequence = FindOrAddSequenceForAttribute( Attr ) )
+		if (ULevelSequence* AttributeSequence = FindOrAddSequenceForAttribute(Attr))
 		{
-			const bool bIsMuted = UsdUtils::IsAttributeMuted( Attr, UsdStage );
+			const bool bIsMuted = UsdUtils::IsAttributeMuted(Attr, UsdStage);
 
 			FMovieSceneSequenceTransform SequenceTransform;
-			FMovieSceneSequenceID SequenceID = SequencesID.FindRef( AttributeSequence );
-			if ( FMovieSceneSubSequenceData* SubSequenceData = SequenceHierarchyCache.FindSubData( SequenceID ) )
+			FMovieSceneSequenceID SequenceID = SequencesID.FindRef(AttributeSequence);
+			if (FMovieSceneSubSequenceData* SubSequenceData = SequenceHierarchyCache.FindSubData(SequenceID))
 			{
 				SequenceTransform = SubSequenceData->RootToSequenceTransform;
 			}
 
 			UMovieScene* MovieScene = AttributeSequence->GetMovieScene();
-			if ( !MovieScene )
+			if (!MovieScene)
 			{
 				continue;
 			}
 
 			TArray<double> TimeSamples;
-			if ( !Attr.GetTimeSamples( TimeSamples ) )
+			if (!Attr.GetTimeSamples(TimeSamples))
 			{
 				continue;
 			}
 
-			if ( UMovieSceneFloatTrack* FloatTrack = AddTrack<UMovieSceneFloatTrack>( PropertyName, PrimTwin, *ComponentToBind, *AttributeSequence, bIsMuted ) )
+			if (UMovieSceneFloatTrack* FloatTrack = AddTrack<UMovieSceneFloatTrack>(PropertyName, PrimTwin, *ComponentToBind, *AttributeSequence, bIsMuted))
 			{
-				UsdToUnreal::FPropertyTrackReader Reader = UsdToUnreal::CreatePropertyTrackReader( Prim, PropertyName );
-				UsdToUnreal::ConvertFloatTimeSamples( UsdStage, TimeSamples, Reader.FloatReader, *FloatTrack, SequenceTransform );
+				UsdToUnreal::FPropertyTrackReader Reader = UsdToUnreal::CreatePropertyTrackReader(Prim, PropertyName);
+				UsdToUnreal::ConvertFloatTimeSamples(UsdStage, TimeSamples, Reader.FloatReader, *FloatTrack, SequenceTransform);
 			}
 
-			PrimPathByLevelSequenceName.AddUnique( AttributeSequence->GetFName(), Prim.GetPrimPath().GetString() );
+			PrimPathByLevelSequenceName.AddUnique(AttributeSequence->GetFName(), Prim.GetPrimPath().GetString());
 		}
 	}
 }
 
-void FUsdLevelSequenceHelperImpl::AddLightTracks( const UUsdPrimTwin& PrimTwin, const UE::FUsdPrim& Prim, const TSet<FName>& PropertyPathsToRead )
+void FUsdLevelSequenceHelperImpl::AddLightTracks(const UUsdPrimTwin& PrimTwin, const UE::FUsdPrim& Prim, const TSet<FName>& PropertyPathsToRead)
 {
 	using namespace UnrealIdentifiers;
 
 	USceneComponent* ComponentToBind = PrimTwin.GetSceneComponent();
-	if ( !ComponentToBind )
+	if (!ComponentToBind)
 	{
 		return;
 	}
@@ -1602,62 +1608,62 @@ void FUsdLevelSequenceHelperImpl::AddLightTracks( const UUsdPrimTwin& PrimTwin,
 	};
 
 	TMap<FName, ETrackType> PropertyPathToTrackType;
-	PropertyPathToTrackType.Add( IntensityPropertyName, ETrackType::Float );
-	PropertyPathToTrackType.Add( LightColorPropertyName, ETrackType::Color );
+	PropertyPathToTrackType.Add(IntensityPropertyName, ETrackType::Float);
+	PropertyPathToTrackType.Add(LightColorPropertyName, ETrackType::Color);
 
-	if ( const ULightComponent* LightComponent = Cast<ULightComponent>( ComponentToBind ) )
+	if (const ULightComponent* LightComponent = Cast<ULightComponent>(ComponentToBind))
 	{
-		PropertyPathToTrackType.Add( UseTemperaturePropertyName, ETrackType::Bool );
-		PropertyPathToTrackType.Add( TemperaturePropertyName, ETrackType::Float );
+		PropertyPathToTrackType.Add(UseTemperaturePropertyName, ETrackType::Bool);
+		PropertyPathToTrackType.Add(TemperaturePropertyName, ETrackType::Float);
 
-		if ( const URectLightComponent* RectLightComponent = Cast<URectLightComponent>( ComponentToBind ) )
+		if (const URectLightComponent* RectLightComponent = Cast<URectLightComponent>(ComponentToBind))
 		{
-			PropertyPathToTrackType.Add( SourceWidthPropertyName, ETrackType::Float );
-			PropertyPathToTrackType.Add( SourceHeightPropertyName, ETrackType::Float );
+			PropertyPathToTrackType.Add(SourceWidthPropertyName, ETrackType::Float);
+			PropertyPathToTrackType.Add(SourceHeightPropertyName, ETrackType::Float);
 		}
-		else if ( const UPointLightComponent* PointLightComponent = Cast<UPointLightComponent>( ComponentToBind ) )
+		else if (const UPointLightComponent* PointLightComponent = Cast<UPointLightComponent>(ComponentToBind))
 		{
-			PropertyPathToTrackType.Add( SourceRadiusPropertyName, ETrackType::Float );
+			PropertyPathToTrackType.Add(SourceRadiusPropertyName, ETrackType::Float);
 
-			if ( const USpotLightComponent* SpotLightComponent = Cast<USpotLightComponent>( ComponentToBind ) )
+			if (const USpotLightComponent* SpotLightComponent = Cast<USpotLightComponent>(ComponentToBind))
 			{
-				PropertyPathToTrackType.Add( OuterConeAnglePropertyName, ETrackType::Float );
-				PropertyPathToTrackType.Add( InnerConeAnglePropertyName, ETrackType::Float );
+				PropertyPathToTrackType.Add(OuterConeAnglePropertyName, ETrackType::Float);
+				PropertyPathToTrackType.Add(InnerConeAnglePropertyName, ETrackType::Float);
 			}
 		}
-		else if ( const UDirectionalLightComponent* DirectionalLightComponent = Cast<UDirectionalLightComponent>( ComponentToBind ) )
+		else if (const UDirectionalLightComponent* DirectionalLightComponent = Cast<UDirectionalLightComponent>(ComponentToBind))
 		{
-			PropertyPathToTrackType.Add( LightSourceAnglePropertyName, ETrackType::Float );
+			PropertyPathToTrackType.Add(LightSourceAnglePropertyName, ETrackType::Float);
 		}
 	}
 
 	// If we were told to specifically read only some property paths, ignore the other ones
-	if ( PropertyPathsToRead.Num() > 0 )
+	if (PropertyPathsToRead.Num() > 0)
 	{
-		for ( TMap< FName, ETrackType >::TIterator Iter = PropertyPathToTrackType.CreateIterator(); Iter; ++Iter)
+		for (TMap< FName, ETrackType >::TIterator Iter = PropertyPathToTrackType.CreateIterator(); Iter; ++Iter)
 		{
 			const FName& PropertyPath = Iter->Key;
-			if ( !PropertyPathsToRead.Contains( PropertyPath ) )
+			if (!PropertyPathsToRead.Contains(PropertyPath))
 			{
 				Iter.RemoveCurrent();
 			}
 		}
 	}
 
-	UE::FSdfLayer PrimLayer = UsdUtils::FindLayerForPrim( Prim );
-	ULevelSequence* PrimSequence = FindSequenceForIdentifier( PrimLayer.GetIdentifier() );
-	if ( !PrimSequence )
+	UE::FSdfLayer PrimLayer = UsdUtils::FindLayerForPrim(Prim);
+	ULevelSequence* PrimSequence = FindSequenceForIdentifier(PrimLayer.GetIdentifier());
+	if (!PrimSequence)
 	{
 		return;
 	}
 
-	for ( const TPair< FName, ETrackType >& Pair : PropertyPathToTrackType )
+	for (const TPair< FName, ETrackType >& Pair : PropertyPathToTrackType)
 	{
 		const FName& PropertyPath = Pair.Key;
 		ETrackType TrackType = Pair.Value;
 
-		TArray<UE::FUsdAttribute> Attrs = UnrealToUsd::GetAttributesForProperty( Prim, PropertyPath );
-		if ( Attrs.Num() < 1 )
+		TArray<UE::FUsdAttribute> Attrs = UnrealToUsd::GetAttributesForProperty(Prim, PropertyPath);
+		if (Attrs.Num() < 1)
 		{
 			continue;
 		}
@@ -1665,65 +1671,65 @@ void FUsdLevelSequenceHelperImpl::AddLightTracks( const UUsdPrimTwin& PrimTwin,
 		// The main attribute is the first one, and that will dictate whether the track is muted or not
 		// This because we don't want to mute the intensity track if just our rect light width track is muted, for example
 		UE::FUsdAttribute MainAttr = Attrs[0];
-		const bool bIsMuted = MainAttr && MainAttr.GetNumTimeSamples() > 0 && UsdUtils::IsAttributeMuted( MainAttr, UsdStage );
+		const bool bIsMuted = MainAttr && MainAttr.GetNumTimeSamples() > 0 && UsdUtils::IsAttributeMuted(MainAttr, UsdStage);
 
 		// Remove attributes we failed to find on this prim (no authored data)
 		// As long as we have at least one attribute with timesamples we can carry on, because we can rely on fallback/default values for the others
-		for ( int32 AttrIndex = Attrs.Num() - 1; AttrIndex >= 0; --AttrIndex )
+		for (int32 AttrIndex = Attrs.Num() - 1; AttrIndex >= 0; --AttrIndex)
 		{
 			UE::FUsdAttribute& Attr = Attrs[AttrIndex];
 			FString AttrPath = Attr.GetPath().GetString();
 
-			if ( !Attr || Attr.GetNumTimeSamples() == 0 )
+			if (!Attr || Attr.GetNumTimeSamples() == 0)
 			{
-				Attrs.RemoveAt( AttrIndex );
+				Attrs.RemoveAt(AttrIndex);
 			}
 		}
 
 		TArray<double> UnionedTimeSamples;
-		if ( Attrs.Num() == 0 || !UE::FUsdAttribute::GetUnionedTimeSamples( Attrs, UnionedTimeSamples ) )
+		if (Attrs.Num() == 0 || !UE::FUsdAttribute::GetUnionedTimeSamples(Attrs, UnionedTimeSamples))
 		{
 			continue;
 		}
 
 		FMovieSceneSequenceTransform SequenceTransform;
-		FMovieSceneSequenceID SequenceID = SequencesID.FindRef( PrimSequence );
-		if ( FMovieSceneSubSequenceData* SubSequenceData = SequenceHierarchyCache.FindSubData( SequenceID ) )
+		FMovieSceneSequenceID SequenceID = SequencesID.FindRef(PrimSequence);
+		if (FMovieSceneSubSequenceData* SubSequenceData = SequenceHierarchyCache.FindSubData(SequenceID))
 		{
 			SequenceTransform = SubSequenceData->RootToSequenceTransform;
 		}
 
 		UMovieScene* MovieScene = PrimSequence->GetMovieScene();
-		if ( !MovieScene )
+		if (!MovieScene)
 		{
 			continue;
 		}
 
-		UsdToUnreal::FPropertyTrackReader Reader = UsdToUnreal::CreatePropertyTrackReader( Prim, PropertyPath );
+		UsdToUnreal::FPropertyTrackReader Reader = UsdToUnreal::CreatePropertyTrackReader(Prim, PropertyPath);
 
-		switch ( TrackType )
+		switch (TrackType)
 		{
 		case ETrackType::Bool:
 		{
-			if ( UMovieSceneBoolTrack* BoolTrack = AddTrack<UMovieSceneBoolTrack>( PropertyPath, PrimTwin, *ComponentToBind, *PrimSequence, bIsMuted ) )
+			if (UMovieSceneBoolTrack* BoolTrack = AddTrack<UMovieSceneBoolTrack>(PropertyPath, PrimTwin, *ComponentToBind, *PrimSequence, bIsMuted))
 			{
-				UsdToUnreal::ConvertBoolTimeSamples( UsdStage, UnionedTimeSamples, Reader.BoolReader, *BoolTrack, SequenceTransform );
+				UsdToUnreal::ConvertBoolTimeSamples(UsdStage, UnionedTimeSamples, Reader.BoolReader, *BoolTrack, SequenceTransform);
 			}
 			break;
 		}
 		case ETrackType::Float:
 		{
-			if ( UMovieSceneFloatTrack* FloatTrack = AddTrack<UMovieSceneFloatTrack>( PropertyPath, PrimTwin, *ComponentToBind, *PrimSequence, bIsMuted ) )
+			if (UMovieSceneFloatTrack* FloatTrack = AddTrack<UMovieSceneFloatTrack>(PropertyPath, PrimTwin, *ComponentToBind, *PrimSequence, bIsMuted))
 			{
-				UsdToUnreal::ConvertFloatTimeSamples( UsdStage, UnionedTimeSamples, Reader.FloatReader, *FloatTrack, SequenceTransform );
+				UsdToUnreal::ConvertFloatTimeSamples(UsdStage, UnionedTimeSamples, Reader.FloatReader, *FloatTrack, SequenceTransform);
 			}
 			break;
 		}
 		case ETrackType::Color:
 		{
-			if ( UMovieSceneColorTrack* ColorTrack = AddTrack<UMovieSceneColorTrack>( PropertyPath, PrimTwin, *ComponentToBind, *PrimSequence, bIsMuted ) )
+			if (UMovieSceneColorTrack* ColorTrack = AddTrack<UMovieSceneColorTrack>(PropertyPath, PrimTwin, *ComponentToBind, *PrimSequence, bIsMuted))
 			{
-				UsdToUnreal::ConvertColorTimeSamples( UsdStage, UnionedTimeSamples, Reader.ColorReader, *ColorTrack, SequenceTransform );
+				UsdToUnreal::ConvertColorTimeSamples(UsdStage, UnionedTimeSamples, Reader.ColorReader, *ColorTrack, SequenceTransform);
 			}
 			break;
 		}
@@ -1732,27 +1738,27 @@ void FUsdLevelSequenceHelperImpl::AddLightTracks( const UUsdPrimTwin& PrimTwin,
 			break;
 		}
 
-		PrimPathByLevelSequenceName.AddUnique( PrimSequence->GetFName(), Prim.GetPrimPath().GetString() );
+		PrimPathByLevelSequenceName.AddUnique(PrimSequence->GetFName(), Prim.GetPrimPath().GetString());
 	}
 }
 
-void FUsdLevelSequenceHelperImpl::AddSkeletalTracks( const UUsdPrimTwin& PrimTwin, const UE::FUsdPrim& Prim )
+void FUsdLevelSequenceHelperImpl::AddSkeletalTracks(const UUsdPrimTwin& PrimTwin, const UE::FUsdPrim& Prim)
 {
-	USkeletalMeshComponent* ComponentToBind = Cast<USkeletalMeshComponent>( PrimTwin.GetSceneComponent() );
-	if ( !ComponentToBind )
+	USkeletalMeshComponent* ComponentToBind = Cast<USkeletalMeshComponent>(PrimTwin.GetSceneComponent());
+	if (!ComponentToBind)
 	{
 		return;
 	}
 
-	if ( !AssetCache )
+	if (!AssetCache)
 	{
 		return;
 	}
 
 	// We'll place the skeletal animation track wherever the SkelAnimation prim is defined (not necessarily the
 	// same layer as the skel root)
-	UE::FUsdPrim SkelAnimationPrim = UsdUtils::FindFirstAnimationSource( Prim );
-	if ( !SkelAnimationPrim )
+	UE::FUsdPrim SkelAnimationPrim = UsdUtils::FindFirstAnimationSource(Prim);
+	if (!SkelAnimationPrim)
 	{
 		return;
 	}
@@ -1761,32 +1767,32 @@ void FUsdLevelSequenceHelperImpl::AddSkeletalTracks( const UUsdPrimTwin& PrimTwi
 	// but we may belong to a FUsdStageImportContext, and so there's no AUsdStageActor at all to use.
 	// At this point it doesn't matter much though, because we shouldn't need to uncollapse a SkelAnimation prim path anyway
 	FString PrimPath = SkelAnimationPrim.GetPrimPath().GetString();
-	UAnimSequence* Sequence = Cast<UAnimSequence>( AssetCache->GetAssetForPrim( PrimPath ) );
-	if ( !Sequence )
+	UAnimSequence* Sequence = Cast<UAnimSequence>(AssetCache->GetAssetForPrim(PrimPath));
+	if (!Sequence)
 	{
 		return;
 	}
 
-	UE::FUsdAttribute TranslationsAttr = SkelAnimationPrim.GetAttribute( TEXT( "translations" ) );
-	UE::FUsdAttribute RotationsAttr = SkelAnimationPrim.GetAttribute( TEXT( "rotations" ) );
-	UE::FUsdAttribute ScalesAttr = SkelAnimationPrim.GetAttribute( TEXT( "scales" ) );
-	UE::FUsdAttribute BlendShapeWeightsAttr = SkelAnimationPrim.GetAttribute( TEXT( "blendShapeWeights" ) );
+	UE::FUsdAttribute TranslationsAttr = SkelAnimationPrim.GetAttribute(TEXT("translations"));
+	UE::FUsdAttribute RotationsAttr = SkelAnimationPrim.GetAttribute(TEXT("rotations"));
+	UE::FUsdAttribute ScalesAttr = SkelAnimationPrim.GetAttribute(TEXT("scales"));
+	UE::FUsdAttribute BlendShapeWeightsAttr = SkelAnimationPrim.GetAttribute(TEXT("blendShapeWeights"));
 
 	const bool bIncludeSessionLayers = false;
-	UE::FSdfLayer SkelAnimationLayer = UsdUtils::FindLayerForAttributes( { TranslationsAttr, RotationsAttr, ScalesAttr, BlendShapeWeightsAttr }, 0.0, bIncludeSessionLayers );
-	if ( !SkelAnimationLayer )
+	UE::FSdfLayer SkelAnimationLayer = UsdUtils::FindLayerForAttributes({TranslationsAttr, RotationsAttr, ScalesAttr, BlendShapeWeightsAttr}, 0.0, bIncludeSessionLayers);
+	if (!SkelAnimationLayer)
 	{
 		return;
 	}
 
-	ULevelSequence* SkelAnimationSequence = FindOrAddSequenceForLayer( SkelAnimationLayer, SkelAnimationLayer.GetIdentifier(), SkelAnimationLayer.GetDisplayName() );
-	if ( !SkelAnimationSequence )
+	ULevelSequence* SkelAnimationSequence = FindOrAddSequenceForLayer(SkelAnimationLayer, SkelAnimationLayer.GetIdentifier(), SkelAnimationLayer.GetDisplayName());
+	if (!SkelAnimationSequence)
 	{
 		return;
 	}
 
 	UMovieScene* MovieScene = SkelAnimationSequence->GetMovieScene();
-	if ( !MovieScene )
+	if (!MovieScene)
 	{
 		return;
 	}
@@ -1795,195 +1801,262 @@ void FUsdLevelSequenceHelperImpl::AddSkeletalTracks( const UUsdPrimTwin& PrimTwi
 	// if it has all attributes muted as well.
 	// We know at least one of these attributes ones is valid and animated because we have an UAnimSequence
 	const bool bIsMuted =
-		( !TranslationsAttr || UsdUtils::IsAttributeMuted( TranslationsAttr, UsdStage ) ) &&
-		( !RotationsAttr || UsdUtils::IsAttributeMuted( RotationsAttr, UsdStage ) ) &&
-		( !ScalesAttr || UsdUtils::IsAttributeMuted( ScalesAttr, UsdStage ) ) &&
-		( !BlendShapeWeightsAttr || UsdUtils::IsAttributeMuted( BlendShapeWeightsAttr, UsdStage ) );
+		(!TranslationsAttr || UsdUtils::IsAttributeMuted(TranslationsAttr, UsdStage)) &&
+		(!RotationsAttr || UsdUtils::IsAttributeMuted(RotationsAttr, UsdStage)) &&
+		(!ScalesAttr || UsdUtils::IsAttributeMuted(ScalesAttr, UsdStage)) &&
+		(!BlendShapeWeightsAttr || UsdUtils::IsAttributeMuted(BlendShapeWeightsAttr, UsdStage));
 
-	if ( UMovieSceneSkeletalAnimationTrack* SkeletalTrack = AddTrack<UMovieSceneSkeletalAnimationTrack>( SkelAnimationPrim.GetName(), PrimTwin, *ComponentToBind, *SkelAnimationSequence, bIsMuted ) )
+	if (UMovieSceneSkeletalAnimationTrack* SkeletalTrack = AddTrack<UMovieSceneSkeletalAnimationTrack>(SkelAnimationPrim.GetName(), PrimTwin, *ComponentToBind, *SkelAnimationSequence, bIsMuted))
 	{
 		double LayerStartOffsetSeconds = 0.0f;
 #if WITH_EDITOR
-		if ( UUsdAnimSequenceAssetImportData* ImportData = Cast<UUsdAnimSequenceAssetImportData>( Sequence->AssetImportData ) )
+		if (UUsdAnimSequenceAssetImportData* ImportData = Cast<UUsdAnimSequenceAssetImportData>(Sequence->AssetImportData))
 		{
 			LayerStartOffsetSeconds = ImportData->LayerStartOffsetSeconds;
 		}
 #endif // WITH_EDITOR
-		FFrameNumber StartOffsetTick = FFrameTime::FromDecimal( LayerStartOffsetSeconds * MovieScene->GetTickResolution().AsDecimal() ).RoundToFrame();
+		FFrameNumber StartOffsetTick = FFrameTime::FromDecimal(LayerStartOffsetSeconds * MovieScene->GetTickResolution().AsDecimal()).RoundToFrame();
 
 		SkeletalTrack->RemoveAllAnimationData();
 
-		UMovieSceneSkeletalAnimationSection* NewSection = Cast< UMovieSceneSkeletalAnimationSection >( SkeletalTrack->AddNewAnimation( StartOffsetTick, Sequence ) );
+		UMovieSceneSkeletalAnimationSection* NewSection = Cast< UMovieSceneSkeletalAnimationSection >(SkeletalTrack->AddNewAnimation(StartOffsetTick, Sequence));
 		NewSection->EvalOptions.CompletionMode = EMovieSceneCompletionMode::KeepState;
 	}
 
-	PrimPathByLevelSequenceName.AddUnique( SkelAnimationSequence->GetFName(), PrimPath );
+	PrimPathByLevelSequenceName.AddUnique(SkelAnimationSequence->GetFName(), PrimPath);
 }
 
-void FUsdLevelSequenceHelperImpl::AddGroomTracks( const UUsdPrimTwin& PrimTwin, const UE::FUsdPrim& Prim )
+void FUsdLevelSequenceHelperImpl::AddGeometryCacheTracks(const UUsdPrimTwin& PrimTwin, const UE::FUsdPrim& Prim)
 {
-	UGroomComponent* ComponentToBind = Cast< UGroomComponent >( PrimTwin.GetSceneComponent() );
-	if ( !ComponentToBind )
+	UGeometryCacheComponent* ComponentToBind = Cast< UGeometryCacheComponent >(PrimTwin.GetSceneComponent());
+	if (!ComponentToBind || !AssetCache)
 	{
 		return;
 	}
 
-	if ( !AssetCache )
+	// Fetch the geometry cache asset from the asset cache. If there's none, don't actually need to create track
+	const FString PrimPath = Prim.GetPrimPath().GetString();
+	UGeometryCache* GeometryCache = Cast< UGeometryCache >(AssetCache->GetAssetForPrim(PrimPath));
+	if (!GeometryCache)
+	{
+		return;
+	}
+
+	if (ComponentToBind->GetGeometryCache() != GeometryCache)
+	{
+		return;
+	}
+
+	UE::FSdfLayer GeometryCacheLayer = UsdUtils::FindLayerForPrim(Prim);
+	if (!GeometryCacheLayer)
+	{
+		return;
+	}
+
+	ULevelSequence* GeometryCacheSequence = FindOrAddSequenceForLayer(GeometryCacheLayer, GeometryCacheLayer.GetIdentifier(), GeometryCacheLayer.GetDisplayName());
+	if (!GeometryCacheSequence)
+	{
+		return;
+	}
+
+	UMovieScene* MovieScene = GeometryCacheSequence->GetMovieScene();
+	if (!MovieScene)
+	{
+		return;
+	}
+
+	const bool bIsMuted = false;
+	if (UMovieSceneGeometryCacheTrack* GeometryCacheTrack = AddTrack< UMovieSceneGeometryCacheTrack >(Prim.GetName(), PrimTwin, *ComponentToBind, *GeometryCacheSequence, bIsMuted))
+	{
+		GeometryCacheTrack->RemoveAllAnimationData();
+
+		double LayerStartOffsetSeconds = 0;
+		if (GeometryCache->Tracks.Num())
+		{
+			// Use the geometry cache start frame index converted to time as the offset
+			UGeometryCacheTrackUsd* UsdTrack = Cast<UGeometryCacheTrackUsd>(GeometryCache->Tracks[0]);
+			LayerStartOffsetSeconds = UsdTrack->StartFrameIndex / UsdTrack->FramesPerSecond;
+		}
+
+		const FFrameNumber StartOffsetTick = FFrameTime::FromDecimal(LayerStartOffsetSeconds * MovieScene->GetTickResolution().AsDecimal()).RoundToFrame();
+		UMovieSceneGeometryCacheSection* NewSection = Cast< UMovieSceneGeometryCacheSection >(GeometryCacheTrack->AddNewAnimation(StartOffsetTick, ComponentToBind));
+		NewSection->EvalOptions.CompletionMode = EMovieSceneCompletionMode::KeepState;
+	}
+
+	PrimPathByLevelSequenceName.AddUnique(GeometryCacheSequence->GetFName(), PrimPath);
+}
+
+void FUsdLevelSequenceHelperImpl::AddGroomTracks(const UUsdPrimTwin& PrimTwin, const UE::FUsdPrim& Prim)
+{
+	UGroomComponent* ComponentToBind = Cast< UGroomComponent >(PrimTwin.GetSceneComponent());
+	if (!ComponentToBind)
+	{
+		return;
+	}
+
+	if (!AssetCache)
 	{
 		return;
 	}
 
 	// Fetch the groom cache asset from the asset cache. If there's none, don't actually need to create track
 	const FString PrimPath = Prim.GetPrimPath().GetString();
-	const FString GroomCachePath = FString::Printf( TEXT( "%s_strands_cache" ), *PrimPath );
-	UGroomCache* GroomCache = Cast< UGroomCache >( AssetCache->GetAssetForPrim( GroomCachePath ) );
-	if ( !GroomCache )
+	const FString GroomCachePath = FString::Printf(TEXT("%s_strands_cache"), *PrimPath);
+	UGroomCache* GroomCache = Cast< UGroomCache >(AssetCache->GetAssetForPrim(GroomCachePath));
+	if (!GroomCache)
 	{
 		return;
 	}
 
-	if ( ComponentToBind->GroomCache.Get() != GroomCache )
+	if (ComponentToBind->GroomCache.Get() != GroomCache)
 	{
 		return;
 	}
 
-	UE::FSdfLayer GroomLayer = UsdUtils::FindLayerForPrim( Prim );
-	if ( !GroomLayer )
+	UE::FSdfLayer GroomLayer = UsdUtils::FindLayerForPrim(Prim);
+	if (!GroomLayer)
 	{
 		return;
 	}
 
-	ULevelSequence* GroomAnimationSequence = FindOrAddSequenceForLayer( GroomLayer, GroomLayer.GetIdentifier(), GroomLayer.GetDisplayName() );
-	if ( !GroomAnimationSequence )
+	ULevelSequence* GroomAnimationSequence = FindOrAddSequenceForLayer(GroomLayer, GroomLayer.GetIdentifier(), GroomLayer.GetDisplayName());
+	if (!GroomAnimationSequence)
 	{
 		return;
 	}
 
 	UMovieScene* MovieScene = GroomAnimationSequence->GetMovieScene();
-	if ( !MovieScene )
+	if (!MovieScene)
 	{
 		return;
 	}
 
 	const bool bIsMuted = false;
-	if ( UMovieSceneGroomCacheTrack* GroomCacheTrack = AddTrack< UMovieSceneGroomCacheTrack >( Prim.GetName(), PrimTwin, *ComponentToBind, *GroomAnimationSequence, bIsMuted ) )
+	if (UMovieSceneGroomCacheTrack* GroomCacheTrack = AddTrack< UMovieSceneGroomCacheTrack >(Prim.GetName(), PrimTwin, *ComponentToBind, *GroomAnimationSequence, bIsMuted))
 	{
 		GroomCacheTrack->RemoveAllAnimationData();
 
 		const FFrameNumber StartOffset;
-		UMovieSceneGroomCacheSection* NewSection = Cast< UMovieSceneGroomCacheSection >( GroomCacheTrack->AddNewAnimation( StartOffset, ComponentToBind ) );
+		UMovieSceneGroomCacheSection* NewSection = Cast< UMovieSceneGroomCacheSection >(GroomCacheTrack->AddNewAnimation(StartOffset, ComponentToBind));
 		NewSection->EvalOptions.CompletionMode = EMovieSceneCompletionMode::KeepState;
 	}
 
-	PrimPathByLevelSequenceName.AddUnique( GroomAnimationSequence->GetFName(), PrimPath );
+	PrimPathByLevelSequenceName.AddUnique(GroomAnimationSequence->GetFName(), PrimPath);
 }
 
-void FUsdLevelSequenceHelperImpl::AddPrim( UUsdPrimTwin& PrimTwin, bool bForceVisibilityTracks )
+void FUsdLevelSequenceHelperImpl::AddPrim(UUsdPrimTwin& PrimTwin, bool bForceVisibilityTracks)
 {
-	if ( !UsdStage )
+	if (!UsdStage)
 	{
 		return;
 	}
 
-	UE::FSdfPath PrimPath( *PrimTwin.PrimPath );
-	UE::FUsdPrim UsdPrim( UsdStage.GetPrimAtPath( PrimPath ) );
+	UE::FSdfPath PrimPath(*PrimTwin.PrimPath);
+	UE::FUsdPrim UsdPrim(UsdStage.GetPrimAtPath(PrimPath));
 
-	UE::FSdfLayer PrimLayer = UsdUtils::FindLayerForPrim( UsdPrim );
-	ULevelSequence* PrimSequence = FindSequenceForIdentifier( PrimLayer.GetIdentifier() );
+	UE::FSdfLayer PrimLayer = UsdUtils::FindLayerForPrim(UsdPrim);
+	ULevelSequence* PrimSequence = FindSequenceForIdentifier(PrimLayer.GetIdentifier());
 
 	TArray< UE::FUsdAttribute > PrimAttributes = UsdPrim.GetAttributes();
 
-	for ( const UE::FUsdAttribute& PrimAttribute : PrimAttributes )
+	for (const UE::FUsdAttribute& PrimAttribute : PrimAttributes)
 	{
-		if ( PrimAttribute.GetNumTimeSamples() > 0 )
+		if (PrimAttribute.GetNumTimeSamples() > 0)
 		{
-			if ( ULevelSequence* AttributeSequence = FindOrAddSequenceForAttribute( PrimAttribute ) )
+			if (ULevelSequence* AttributeSequence = FindOrAddSequenceForAttribute(PrimAttribute))
 			{
-				PrimPathByLevelSequenceName.AddUnique( AttributeSequence->GetFName(), PrimTwin.PrimPath );
+				PrimPathByLevelSequenceName.AddUnique(AttributeSequence->GetFName(), PrimTwin.PrimPath);
 
-				if ( !SequencesID.Contains( AttributeSequence ) )
+				if (!SequencesID.Contains(AttributeSequence))
 				{
-					if ( PrimSequence )
+					if (PrimSequence)
 					{
 						// Create new subsequence section for this referencing prim
-						CreateSubSequenceSection( *PrimSequence, *AttributeSequence );
+						CreateSubSequenceSection(*PrimSequence, *AttributeSequence);
 					}
 				}
 			}
 		}
 	}
 
-	if ( UsdPrim.IsA( TEXT( "Camera" ) ) )
+	static IConsoleVariable* CVar = IConsoleManager::Get().FindConsoleVariable(TEXT("USD.DisableGeoCacheTracks"));
+	const bool bDisableGeoCacheTracks = CVar && CVar->GetBool();
+
+	if (UsdPrim.IsA(TEXT("Camera")))
+	{
+		AddCameraTracks(PrimTwin, UsdPrim);
+	}
+	else if (UsdPrim.HasAPI(TEXT("UsdLuxLightAPI")))
 	{
-		AddCameraTracks( PrimTwin, UsdPrim );
+		AddLightTracks(PrimTwin, UsdPrim);
 	}
-	else if ( UsdPrim.HasAPI( TEXT( "UsdLuxLightAPI" ) ) )
+	else if (UsdPrim.IsA(TEXT("SkelRoot")))
 	{
-		AddLightTracks( PrimTwin, UsdPrim );
+		AddSkeletalTracks(PrimTwin, UsdPrim);
 	}
-	else if ( UsdPrim.IsA( TEXT( "SkelRoot" ) ) )
+	else if (!bDisableGeoCacheTracks && UsdPrim.IsA(TEXT("Mesh")) && UsdUtils::IsAnimated(UsdPrim))
 	{
-		AddSkeletalTracks( PrimTwin, UsdPrim );
+		AddGeometryCacheTracks(PrimTwin, UsdPrim);
 	}
-	else if ( UsdUtils::PrimHasSchema( UsdPrim, UnrealIdentifiers::GroomAPI ) )
+	else if (UsdUtils::PrimHasSchema(UsdPrim, UnrealIdentifiers::GroomAPI))
 	{
-		AddGroomTracks( PrimTwin, UsdPrim );
+		AddGroomTracks(PrimTwin, UsdPrim);
 	}
 
-	AddCommonTracks( PrimTwin, UsdPrim, bForceVisibilityTracks );
+	AddCommonTracks(PrimTwin, UsdPrim, bForceVisibilityTracks);
 
 	RefreshSequencer();
 }
 
 template<typename TrackType>
-TrackType* FUsdLevelSequenceHelperImpl::AddTrack( const FName& TrackName, const UUsdPrimTwin& PrimTwin, USceneComponent& ComponentToBind, ULevelSequence& Sequence, bool bIsMuted )
+TrackType* FUsdLevelSequenceHelperImpl::AddTrack(const FName& TrackName, const UUsdPrimTwin& PrimTwin, USceneComponent& ComponentToBind, ULevelSequence& Sequence, bool bIsMuted)
 {
-	if ( !UsdStage )
+	if (!UsdStage)
 	{
 		return nullptr;
 	}
 
 	UMovieScene* MovieScene = Sequence.GetMovieScene();
-	if ( !MovieScene )
+	if (!MovieScene)
 	{
 		return nullptr;
 	}
 
-	const FGuid ComponentBinding = GetOrCreateComponentBinding( PrimTwin, ComponentToBind, Sequence );
+	const FGuid ComponentBinding = GetOrCreateComponentBinding(PrimTwin, ComponentToBind, Sequence);
 
-	TrackType* Track = MovieScene->FindTrack< TrackType >( ComponentBinding, TrackName );
-	if ( Track )
+	TrackType* Track = MovieScene->FindTrack< TrackType >(ComponentBinding, TrackName);
+	if (Track)
 	{
 		Track->RemoveAllAnimationData();
 	}
 	else
 	{
-		Track = MovieScene->AddTrack< TrackType >( ComponentBinding );
-		if ( !Track )
+		Track = MovieScene->AddTrack< TrackType >(ComponentBinding);
+		if (!Track)
 		{
 			return nullptr;
 		}
 
-		if ( UMovieScenePropertyTrack* PropertyTrack = Cast<UMovieScenePropertyTrack>( Track ) )
+		if (UMovieScenePropertyTrack* PropertyTrack = Cast<UMovieScenePropertyTrack>(Track))
 		{
-			PropertyTrack->SetPropertyNameAndPath( TrackName, TrackName.ToString() );
+			PropertyTrack->SetPropertyNameAndPath(TrackName, TrackName.ToString());
 		}
 #if WITH_EDITOR
-		else if ( UMovieSceneSkeletalAnimationTrack* SkeletalTrack = Cast<UMovieSceneSkeletalAnimationTrack>( Track ) )
+		else if (UMovieSceneSkeletalAnimationTrack* SkeletalTrack = Cast<UMovieSceneSkeletalAnimationTrack>(Track))
 		{
-			SkeletalTrack->SetDisplayName( FText::FromName( TrackName ) );
+			SkeletalTrack->SetDisplayName(FText::FromName(TrackName));
 		}
 #endif // WITH_EDITOR
 	}
 
-	UsdLevelSequenceHelperImpl::MuteTrack( Track, MovieScene, ComponentBinding.ToString(), Track->GetName(), bIsMuted );
+	UsdLevelSequenceHelperImpl::MuteTrack(Track, MovieScene, ComponentBinding.ToString(), Track->GetName(), bIsMuted);
 
 	return Track;
 }
 
-void FUsdLevelSequenceHelperImpl::RemovePrim( const UUsdPrimTwin& PrimTwin )
+void FUsdLevelSequenceHelperImpl::RemovePrim(const UUsdPrimTwin& PrimTwin)
 {
-	if ( !UsdStage )
+	if (!UsdStage)
 	{
 		return;
 	}
@@ -1992,54 +2065,54 @@ void FUsdLevelSequenceHelperImpl::RemovePrim( const UUsdPrimTwin& PrimTwin )
 
 	TSet< FName > PrimSequences;
 
-	for ( TPair< FName, FString >& PrimPathByLevelSequenceNamePair : PrimPathByLevelSequenceName )
+	for (TPair< FName, FString >& PrimPathByLevelSequenceNamePair : PrimPathByLevelSequenceName)
 	{
-		if ( PrimPathByLevelSequenceNamePair.Value == PrimTwin.PrimPath )
+		if (PrimPathByLevelSequenceNamePair.Value == PrimTwin.PrimPath)
 		{
-			PrimSequences.Add( PrimPathByLevelSequenceNamePair.Key );
+			PrimSequences.Add(PrimPathByLevelSequenceNamePair.Key);
 		}
 	}
 
 	TSet< ULevelSequence* > SequencesToRemoveForPrim;
 
-	for ( const FName& PrimSequenceName : PrimSequences )
+	for (const FName& PrimSequenceName : PrimSequences)
 	{
-		for ( const TPair< FString, ULevelSequence* >& IdentifierSequencePair : LevelSequencesByIdentifier )
+		for (const TPair< FString, ULevelSequence* >& IdentifierSequencePair : LevelSequencesByIdentifier)
 		{
-			if ( IdentifierSequencePair.Value && IdentifierSequencePair.Value->GetFName() == PrimSequenceName )
+			if (IdentifierSequencePair.Value && IdentifierSequencePair.Value->GetFName() == PrimSequenceName)
 			{
-				SequencesToRemoveForPrim.Add( IdentifierSequencePair.Value );
+				SequencesToRemoveForPrim.Add(IdentifierSequencePair.Value);
 			}
 		}
 	}
 
-	RemovePossessable( PrimTwin );
+	RemovePossessable(PrimTwin);
 
-	for ( ULevelSequence* SequenceToRemoveForPrim : SequencesToRemoveForPrim )
+	for (ULevelSequence* SequenceToRemoveForPrim : SequencesToRemoveForPrim)
 	{
-		RemoveSequenceForPrim( *SequenceToRemoveForPrim, PrimTwin );
+		RemoveSequenceForPrim(*SequenceToRemoveForPrim, PrimTwin);
 	}
 
 	RefreshSequencer();
 }
 
-void FUsdLevelSequenceHelperImpl::UpdateControlRigTracks( UUsdPrimTwin& PrimTwin )
+void FUsdLevelSequenceHelperImpl::UpdateControlRigTracks(UUsdPrimTwin& PrimTwin)
 {
 #if WITH_EDITOR
-	if ( !UsdStage )
+	if (!UsdStage)
 	{
 		return;
 	}
 
-	UE::FSdfPath PrimPath( *PrimTwin.PrimPath );
-	UE::FUsdPrim UsdPrim( UsdStage.GetPrimAtPath( PrimPath ) );
-	if ( !UsdPrim )
+	UE::FSdfPath PrimPath(*PrimTwin.PrimPath);
+	UE::FUsdPrim UsdPrim(UsdStage.GetPrimAtPath(PrimPath));
+	if (!UsdPrim)
 	{
 		return;
 	}
 
-	USkeletalMeshComponent* ComponentToBind = Cast<USkeletalMeshComponent>( PrimTwin.GetSceneComponent() );
-	if ( !ComponentToBind )
+	USkeletalMeshComponent* ComponentToBind = Cast<USkeletalMeshComponent>(PrimTwin.GetSceneComponent());
+	if (!ComponentToBind)
 	{
 		return;
 	}
@@ -2047,87 +2120,87 @@ void FUsdLevelSequenceHelperImpl::UpdateControlRigTracks( UUsdPrimTwin& PrimTwin
 	// Block here because USD needs to fire and respond to notices for the DefinePrim call to work,
 	// but we need UsdUtils::BindAnimationSource to run before we get in here again or else we'll
 	// repeatedly create Animation prims
-	FScopedBlockNoticeListening BlockNotices( StageActor.Get() );
+	FScopedBlockNoticeListening BlockNotices(StageActor.Get());
 
 	UE::FSdfLayer SkelAnimationLayer;
 
 	// We'll place the skeletal animation track wherever the SkelAnimation prim is defined (not necessarily the
 	// same layer as the skel root)
-	UE::FUsdPrim SkelAnimationPrim = UsdUtils::FindFirstAnimationSource( UsdPrim );
-	if ( SkelAnimationPrim )
+	UE::FUsdPrim SkelAnimationPrim = UsdUtils::FindFirstAnimationSource(UsdPrim);
+	if (SkelAnimationPrim)
 	{
-		SkelAnimationLayer = UsdUtils::FindLayerForPrim( SkelAnimationPrim );
+		SkelAnimationLayer = UsdUtils::FindLayerForPrim(SkelAnimationPrim);
 	}
 	else
 	{
 		// If this SkelRoot doesn't have any animation, lets create a new one on the current edit target
 		SkelAnimationLayer = UsdStage.GetEditTarget();
 
-		ensure( UsdPrim.IsA( TEXT( "SkelRoot" ) ) );
+		ensure(UsdPrim.IsA(TEXT("SkelRoot")));
 
-		FString UniqueChildName = UsdUtils::GetValidChildName( TEXT( "Animation" ), UsdPrim );
+		FString UniqueChildName = UsdUtils::GetValidChildName(TEXT("Animation"), UsdPrim);
 		SkelAnimationPrim = UsdStage.DefinePrim(
-			UsdPrim.GetPrimPath().AppendChild( *UniqueChildName ),
-			TEXT( "SkelAnimation" )
+			UsdPrim.GetPrimPath().AppendChild(*UniqueChildName),
+			TEXT("SkelAnimation")
 		);
-		if ( !SkelAnimationPrim )
+		if (!SkelAnimationPrim)
 		{
 			return;
 		}
 
-		UsdUtils::BindAnimationSource( UsdPrim, SkelAnimationPrim );
+		UsdUtils::BindAnimationSource(UsdPrim, SkelAnimationPrim);
 	}
 
 	// Fetch the UAnimSequence asset from the asset cache. Ideally we'd call AUsdStageActor::GetGeneratedAssets,
 	// but we may belong to a FUsdStageImportContext, and so there's no AUsdStageActor at all to use.
 	// At this point it doesn't matter much though, because we shouldn't need to uncollapse a SkelAnimation prim path anyway
 	FString SkelAnimationPrimPath = SkelAnimationPrim.GetPrimPath().GetString();
-	UAnimSequence* AnimSequence = Cast<UAnimSequence>( AssetCache->GetAssetForPrim( SkelAnimationPrimPath ) );
+	UAnimSequence* AnimSequence = Cast<UAnimSequence>(AssetCache->GetAssetForPrim(SkelAnimationPrimPath));
 
-	if ( !SkelAnimationLayer )
+	if (!SkelAnimationLayer)
 	{
 		return;
 	}
-	UE::FUsdEditContext EditContext{ UsdStage, SkelAnimationLayer };
+	UE::FUsdEditContext EditContext{UsdStage, SkelAnimationLayer};
 	FString Identifier = SkelAnimationLayer.GetIdentifier();
 
 	// Force-create these because these are mandatory anyway (https://graphics.pixar.com/usd/release/api/_usd_skel__schemas.html#UsdSkel_SkelAnimation)
-	UE::FUsdAttribute JointsAttr = SkelAnimationPrim.CreateAttribute( TEXT( "joints" ), TEXT( "token[]" ) );
-	UE::FUsdAttribute TranslationsAttr = SkelAnimationPrim.CreateAttribute( TEXT( "translations" ), TEXT( "float3[]" ) );
-	UE::FUsdAttribute RotationsAttr = SkelAnimationPrim.CreateAttribute( TEXT( "rotations" ), TEXT( "quatf[]" ) );
-	UE::FUsdAttribute ScalesAttr = SkelAnimationPrim.CreateAttribute( TEXT( "scales" ), TEXT( "half3[]" ) );
-	UE::FUsdAttribute BlendShapeWeightsAttr = SkelAnimationPrim.GetAttribute( TEXT( "blendShapeWeights" ) );
+	UE::FUsdAttribute JointsAttr = SkelAnimationPrim.CreateAttribute(TEXT("joints"), TEXT("token[]"));
+	UE::FUsdAttribute TranslationsAttr = SkelAnimationPrim.CreateAttribute(TEXT("translations"), TEXT("float3[]"));
+	UE::FUsdAttribute RotationsAttr = SkelAnimationPrim.CreateAttribute(TEXT("rotations"), TEXT("quatf[]"));
+	UE::FUsdAttribute ScalesAttr = SkelAnimationPrim.CreateAttribute(TEXT("scales"), TEXT("half3[]"));
+	UE::FUsdAttribute BlendShapeWeightsAttr = SkelAnimationPrim.GetAttribute(TEXT("blendShapeWeights"));
 
-	ULevelSequence* SkelAnimationSequence = FindOrAddSequenceForLayer( SkelAnimationLayer, SkelAnimationLayer.GetIdentifier(), SkelAnimationLayer.GetDisplayName() );
-	if ( !SkelAnimationSequence )
+	ULevelSequence* SkelAnimationSequence = FindOrAddSequenceForLayer(SkelAnimationLayer, SkelAnimationLayer.GetIdentifier(), SkelAnimationLayer.GetDisplayName());
+	if (!SkelAnimationSequence)
 	{
 		return;
 	}
 
 	UMovieScene* MovieScene = SkelAnimationSequence->GetMovieScene();
-	if ( !MovieScene )
+	if (!MovieScene)
 	{
 		return;
 	}
 
-	const FGuid ComponentBinding = GetOrCreateComponentBinding( PrimTwin, *ComponentToBind, *SkelAnimationSequence );
+	const FGuid ComponentBinding = GetOrCreateComponentBinding(PrimTwin, *ComponentToBind, *SkelAnimationSequence);
 
 	// NOTE: We are fetching the first skel track we find, since we can't actually use SkelAnimationPrim.GetName() here at all!
 	// The property tracks do derive GetTrackName(), but the skeletal track doesn't, so FindTrack will never find them.
 	// This likely has no effect since we only ever spawn a single skeletal track per prim anyway, but its worth to keep in mind!
-	UMovieSceneSkeletalAnimationTrack* SkelTrack = MovieScene->FindTrack< UMovieSceneSkeletalAnimationTrack >( ComponentBinding );
-	UMovieSceneControlRigParameterTrack* ControlRigTrack = MovieScene->FindTrack< UMovieSceneControlRigParameterTrack >( ComponentBinding );
+	UMovieSceneSkeletalAnimationTrack* SkelTrack = MovieScene->FindTrack< UMovieSceneSkeletalAnimationTrack >(ComponentBinding);
+	UMovieSceneControlRigParameterTrack* ControlRigTrack = MovieScene->FindTrack< UMovieSceneControlRigParameterTrack >(ComponentBinding);
 
 	// We should be in control rig track mode but don't have any tracks yet --> Setup for Control Rig
-	if ( !ControlRigTrack )
+	if (!ControlRigTrack)
 	{
 		bool bControlRigReduceKeys = false;
-		if ( UE::FUsdAttribute Attr = UsdPrim.GetAttribute( *UsdToUnreal::ConvertToken( UnrealIdentifiers::UnrealControlRigReduceKeys ) ) )
+		if (UE::FUsdAttribute Attr = UsdPrim.GetAttribute(*UsdToUnreal::ConvertToken(UnrealIdentifiers::UnrealControlRigReduceKeys)))
 		{
 			UE::FVtValue Value;
-			if ( Attr.Get( Value ) && !Value.IsEmpty() )
+			if (Attr.Get(Value) && !Value.IsEmpty())
 			{
-				if ( TOptional<bool> UnderlyingValue = UsdUtils::GetUnderlyingValue<bool>( Value ) )
+				if (TOptional<bool> UnderlyingValue = UsdUtils::GetUnderlyingValue<bool>(Value))
 				{
 					bControlRigReduceKeys = UnderlyingValue.GetValue();
 				}
@@ -2135,12 +2208,12 @@ void FUsdLevelSequenceHelperImpl::UpdateControlRigTracks( UUsdPrimTwin& PrimTwin
 		}
 
 		float ControlRigReduceTolerance = 0.001f;
-		if ( UE::FUsdAttribute Attr = UsdPrim.GetAttribute( *UsdToUnreal::ConvertToken( UnrealIdentifiers::UnrealControlRigReductionTolerance ) ) )
+		if (UE::FUsdAttribute Attr = UsdPrim.GetAttribute(*UsdToUnreal::ConvertToken(UnrealIdentifiers::UnrealControlRigReductionTolerance)))
 		{
 			UE::FVtValue Value;
-			if ( Attr.Get( Value ) && !Value.IsEmpty() )
+			if (Attr.Get(Value) && !Value.IsEmpty())
 			{
-				if ( TOptional<float> UnderlyingValue = UsdUtils::GetUnderlyingValue<float>( Value ) )
+				if (TOptional<float> UnderlyingValue = UsdUtils::GetUnderlyingValue<float>(Value))
 				{
 					ControlRigReduceTolerance = UnderlyingValue.GetValue();
 				}
@@ -2148,12 +2221,12 @@ void FUsdLevelSequenceHelperImpl::UpdateControlRigTracks( UUsdPrimTwin& PrimTwin
 		}
 
 		bool bIsFKControlRig = false;
-		if ( UE::FUsdAttribute Attr = UsdPrim.GetAttribute( *UsdToUnreal::ConvertToken( UnrealIdentifiers::UnrealUseFKControlRig ) ) )
+		if (UE::FUsdAttribute Attr = UsdPrim.GetAttribute(*UsdToUnreal::ConvertToken(UnrealIdentifiers::UnrealUseFKControlRig)))
 		{
 			UE::FVtValue Value;
-			if ( Attr.Get( Value ) )
+			if (Attr.Get(Value))
 			{
-				if (TOptional<bool> UseFKOptional = UsdUtils::GetUnderlyingValue<bool>( Value ) )
+				if (TOptional<bool> UseFKOptional = UsdUtils::GetUnderlyingValue<bool>(Value))
 				{
 					bIsFKControlRig = UseFKOptional.GetValue();
 				}
@@ -2161,29 +2234,29 @@ void FUsdLevelSequenceHelperImpl::UpdateControlRigTracks( UUsdPrimTwin& PrimTwin
 		}
 
 		UClass* ControlRigClass = nullptr;
-		if ( bIsFKControlRig )
+		if (bIsFKControlRig)
 		{
 			ControlRigClass = UFKControlRig::StaticClass();
 		}
 		else
 		{
 			FString ControlRigBPPath;
-			if ( UE::FUsdAttribute Attr = UsdPrim.GetAttribute( *UsdToUnreal::ConvertToken( UnrealIdentifiers::UnrealControlRigPath ) ) )
+			if (UE::FUsdAttribute Attr = UsdPrim.GetAttribute(*UsdToUnreal::ConvertToken(UnrealIdentifiers::UnrealControlRigPath)))
 			{
 				UE::FVtValue Value;
-				if ( Attr.Get( Value ) && !Value.IsEmpty() )
+				if (Attr.Get(Value) && !Value.IsEmpty())
 				{
-					ControlRigBPPath = UsdUtils::Stringify( Value );
+					ControlRigBPPath = UsdUtils::Stringify(Value);
 				}
 			}
 
-			if ( UControlRigBlueprint* BP = Cast<UControlRigBlueprint>( FSoftObjectPath( ControlRigBPPath ).TryLoad() ) )
+			if (UControlRigBlueprint* BP = Cast<UControlRigBlueprint>(FSoftObjectPath(ControlRigBPPath).TryLoad()))
 			{
 				ControlRigClass = BP->GetControlRigClass();
 			}
 		}
 
-		if ( ControlRigClass )
+		if (ControlRigClass)
 		{
 			UAnimSeqExportOption* NewOptions = NewObject< UAnimSeqExportOption >();
 
@@ -2203,62 +2276,62 @@ void FUsdLevelSequenceHelperImpl::UpdateControlRigTracks( UUsdPrimTwin& PrimTwin
 		}
 	}
 
-	PrimPathByLevelSequenceName.AddUnique( SkelAnimationSequence->GetFName(), PrimPath.GetString() );
+	PrimPathByLevelSequenceName.AddUnique(SkelAnimationSequence->GetFName(), PrimPath.GetString());
 #endif // WITH_EDITOR
 }
 
-void FUsdLevelSequenceHelperImpl::RemoveSequenceForPrim( ULevelSequence& Sequence, const UUsdPrimTwin& PrimTwin )
+void FUsdLevelSequenceHelperImpl::RemoveSequenceForPrim(ULevelSequence& Sequence, const UUsdPrimTwin& PrimTwin)
 {
 	TArray< FString > PrimPathsForSequence;
-	PrimPathByLevelSequenceName.MultiFind( Sequence.GetFName(), PrimPathsForSequence );
+	PrimPathByLevelSequenceName.MultiFind(Sequence.GetFName(), PrimPathsForSequence);
 
-	if ( PrimPathsForSequence.Find( PrimTwin.PrimPath ) != INDEX_NONE )
+	if (PrimPathsForSequence.Find(PrimTwin.PrimPath) != INDEX_NONE)
 	{
-		PrimPathByLevelSequenceName.Remove( Sequence.GetFName(), PrimTwin.PrimPath );
+		PrimPathByLevelSequenceName.Remove(Sequence.GetFName(), PrimTwin.PrimPath);
 
 		// If Sequence isn't used anymore, remove it and its subsection
-		if ( !PrimPathByLevelSequenceName.Contains( Sequence.GetFName() ) && !LocalLayersSequences.Contains( Sequence.GetFName() ) )
+		if (!PrimPathByLevelSequenceName.Contains(Sequence.GetFName()) && !LocalLayersSequences.Contains(Sequence.GetFName()))
 		{
 			ULevelSequence* ParentSequence = MainLevelSequence;
-			FMovieSceneSequenceID SequenceID = SequencesID.FindRef( &Sequence );
+			FMovieSceneSequenceID SequenceID = SequencesID.FindRef(&Sequence);
 
-			if ( FMovieSceneSequenceHierarchyNode* NodeData = SequenceHierarchyCache.FindNode( SequenceID ) )
+			if (FMovieSceneSequenceHierarchyNode* NodeData = SequenceHierarchyCache.FindNode(SequenceID))
 			{
 				FMovieSceneSequenceID ParentSequenceID = NodeData->ParentID;
 
-				if ( FMovieSceneSubSequenceData* ParentSubSequenceData = SequenceHierarchyCache.FindSubData( ParentSequenceID ) )
+				if (FMovieSceneSubSequenceData* ParentSubSequenceData = SequenceHierarchyCache.FindSubData(ParentSequenceID))
 				{
-					ParentSequence = Cast< ULevelSequence >( ParentSubSequenceData->GetSequence() );
+					ParentSequence = Cast< ULevelSequence >(ParentSubSequenceData->GetSequence());
 				}
 			}
 
-			if ( ParentSequence )
+			if (ParentSequence)
 			{
-				RemoveSubSequenceSection( *ParentSequence, Sequence );
+				RemoveSubSequenceSection(*ParentSequence, Sequence);
 			}
 
-			LevelSequencesByIdentifier.Remove( PrimTwin.PrimPath );
-			IdentifierByLevelSequence.Remove( &Sequence );
-			SequencesID.Remove( &Sequence );
+			LevelSequencesByIdentifier.Remove(PrimTwin.PrimPath);
+			IdentifierByLevelSequence.Remove(&Sequence);
+			SequencesID.Remove(&Sequence);
 		}
 	}
 }
 
-void FUsdLevelSequenceHelperImpl::RemovePossessable( const UUsdPrimTwin& PrimTwin )
+void FUsdLevelSequenceHelperImpl::RemovePossessable(const UUsdPrimTwin& PrimTwin)
 {
-	FPrimTwinBindings* Bindings = PrimTwinToBindings.Find( &PrimTwin );
-	if ( !Bindings )
+	FPrimTwinBindings* Bindings = PrimTwinToBindings.Find(&PrimTwin);
+	if (!Bindings)
 	{
 		return;
 	}
 
-	if( !Bindings->Sequence )
+	if (!Bindings->Sequence)
 	{
 		return;
 	}
 
 	UMovieScene* MovieScene = Bindings->Sequence->GetMovieScene();
-	if ( !MovieScene )
+	if (!MovieScene)
 	{
 		return;
 	}
@@ -2267,61 +2340,61 @@ void FUsdLevelSequenceHelperImpl::RemovePossessable( const UUsdPrimTwin& PrimTwi
 	// ones don't modify the Sequence and change properties, so we must modify them here
 	Bindings->Sequence->Modify();
 
-	for ( const TPair< const UClass*, FGuid >& Pair : Bindings->ComponentClassToBindingGuid )
+	for (const TPair< const UClass*, FGuid >& Pair : Bindings->ComponentClassToBindingGuid)
 	{
 		const FGuid& ComponentPossessableGuid = Pair.Value;
 
 		FGuid ActorPossessableGuid;
-		if ( FMovieScenePossessable* ComponentPossessable = MovieScene->FindPossessable( ComponentPossessableGuid ) )
+		if (FMovieScenePossessable* ComponentPossessable = MovieScene->FindPossessable(ComponentPossessableGuid))
 		{
 			ActorPossessableGuid = ComponentPossessable->GetParent();
 		}
 
 		// This will also remove all tracks bound to this guid
-		if ( MovieScene->RemovePossessable( ComponentPossessableGuid ) )
+		if (MovieScene->RemovePossessable(ComponentPossessableGuid))
 		{
-			Bindings->Sequence->UnbindPossessableObjects( ComponentPossessableGuid );
+			Bindings->Sequence->UnbindPossessableObjects(ComponentPossessableGuid);
 		}
 
 		// If our parent binding has nothing else in it, we should remove it too
 		bool bRemoveActorBinding = true;
-		if ( ActorPossessableGuid.IsValid() )
+		if (ActorPossessableGuid.IsValid())
 		{
-			for ( int32 PossessableIndex = 0; PossessableIndex < MovieScene->GetPossessableCount(); ++PossessableIndex )
+			for (int32 PossessableIndex = 0; PossessableIndex < MovieScene->GetPossessableCount(); ++PossessableIndex)
 			{
-				const FMovieScenePossessable& SomePossessable = MovieScene->GetPossessable( PossessableIndex );
-				if ( SomePossessable.GetParent() == ActorPossessableGuid )
+				const FMovieScenePossessable& SomePossessable = MovieScene->GetPossessable(PossessableIndex);
+				if (SomePossessable.GetParent() == ActorPossessableGuid)
 				{
 					bRemoveActorBinding = false;
 					break;
 				}
 			}
 		}
-		if ( bRemoveActorBinding )
+		if (bRemoveActorBinding)
 		{
-			MovieScene->RemovePossessable( ActorPossessableGuid );
-			Bindings->Sequence->UnbindPossessableObjects( ActorPossessableGuid );
+			MovieScene->RemovePossessable(ActorPossessableGuid);
+			Bindings->Sequence->UnbindPossessableObjects(ActorPossessableGuid);
 		}
 	}
 
-	PrimTwinToBindings.Remove( &PrimTwin );
+	PrimTwinToBindings.Remove(&PrimTwin);
 }
 
 void FUsdLevelSequenceHelperImpl::RefreshSequencer()
 {
 #if WITH_EDITOR
-	if ( !MainLevelSequence || !GIsEditor )
+	if (!MainLevelSequence || !GIsEditor)
 	{
 		return;
 	}
 
-	if ( TSharedPtr< ISequencer > Sequencer = UsdLevelSequenceHelperImpl::GetOpenedSequencerForLevelSequence( MainLevelSequence ) )
+	if (TSharedPtr< ISequencer > Sequencer = UsdLevelSequenceHelperImpl::GetOpenedSequencerForLevelSequence(MainLevelSequence))
 	{
 		// Don't try refreshing the sequencer if its displaying a stale sequence (e.g. during busy transitions like import) as it
 		// can crash
-		if ( UMovieSceneSequence* FocusedSequence = Sequencer->GetFocusedMovieSceneSequence() )
+		if (UMovieSceneSequence* FocusedSequence = Sequencer->GetFocusedMovieSceneSequence())
 		{
-			Sequencer->NotifyMovieSceneDataChanged( EMovieSceneDataChangeType::RefreshTree );
+			Sequencer->NotifyMovieSceneDataChanged(EMovieSceneDataChangeType::RefreshTree);
 		}
 	}
 #endif // WITH_EDITOR
@@ -2341,18 +2414,18 @@ void FUsdLevelSequenceHelperImpl::UpdateUsdLayerOffsetFromSection(const UMovieSc
 		return;
 	}
 
-	const FString* LayerIdentifier = LayerIdentifierByLevelSequenceName.Find( Sequence->GetFName() );
-	const FString* SubLayerIdentifier = LayerIdentifierByLevelSequenceName.Find( SubSequence->GetFName() );
+	const FString* LayerIdentifier = LayerIdentifierByLevelSequenceName.Find(Sequence->GetFName());
+	const FString* SubLayerIdentifier = LayerIdentifierByLevelSequenceName.Find(SubSequence->GetFName());
 
-	if ( !LayerIdentifier || !SubLayerIdentifier )
+	if (!LayerIdentifier || !SubLayerIdentifier)
 	{
 		return;
 	}
 
-	FLayerTimeInfo* LayerTimeInfo = LayerTimeInfosByLayerIdentifier.Find( *LayerIdentifier );
-	FLayerTimeInfo* SubLayerTimeInfo = LayerTimeInfosByLayerIdentifier.Find( *SubLayerIdentifier );
+	FLayerTimeInfo* LayerTimeInfo = LayerTimeInfosByLayerIdentifier.Find(*LayerIdentifier);
+	FLayerTimeInfo* SubLayerTimeInfo = LayerTimeInfosByLayerIdentifier.Find(*SubLayerIdentifier);
 
-	if ( !LayerTimeInfo || !SubLayerTimeInfo )
+	if (!LayerTimeInfo || !SubLayerTimeInfo)
 	{
 		return;
 	}
@@ -2365,21 +2438,21 @@ void FUsdLevelSequenceHelperImpl::UpdateUsdLayerOffsetFromSection(const UMovieSc
 
 	FFrameRate TickResolution = MovieScene->GetTickResolution();
 	FFrameNumber ModifiedStartFrame = Section->GetInclusiveStartFrame();
-	FFrameNumber ModifiedEndFrame   = Section->GetExclusiveEndFrame();
+	FFrameNumber ModifiedEndFrame = Section->GetExclusiveEndFrame();
 
 	// This will obviously be quantized to frame intervals for now
 	double SubSectionStartTimeCode = TickResolution.AsSeconds(ModifiedStartFrame) * TimeCodesPerSecond;
 
 	UE::FSdfLayerOffset NewLayerOffset;
-	NewLayerOffset.Scale = FMath::IsNearlyZero( Section->Parameters.TimeScale ) ? 0.f : 1.f / Section->Parameters.TimeScale;
+	NewLayerOffset.Scale = FMath::IsNearlyZero(Section->Parameters.TimeScale) ? 0.f : 1.f / Section->Parameters.TimeScale;
 	NewLayerOffset.Offset = SubSectionStartTimeCode - SubStartTimeCode * NewLayerOffset.Scale;
 
-	if ( FMath::IsNearlyZero( NewLayerOffset.Offset ) )
+	if (FMath::IsNearlyZero(NewLayerOffset.Offset))
 	{
 		NewLayerOffset.Offset = 0.0;
 	}
 
-	if ( FMath::IsNearlyEqual( NewLayerOffset.Scale, 1.0 ) )
+	if (FMath::IsNearlyEqual(NewLayerOffset.Scale, 1.0))
 	{
 		NewLayerOffset.Scale = 1.0;
 	}
@@ -2387,12 +2460,12 @@ void FUsdLevelSequenceHelperImpl::UpdateUsdLayerOffsetFromSection(const UMovieSc
 	// Prevent twins from being rebuilt when we update the layer offsets
 	TOptional< FScopedBlockNoticeListening > BlockNotices;
 
-	if ( StageActor.IsValid() )
+	if (StageActor.IsValid())
 	{
-		BlockNotices.Emplace( StageActor.Get() );
+		BlockNotices.Emplace(StageActor.Get());
 	}
 
-	if ( LocalLayersSequences.Contains( SubSequence->GetFName() ) )
+	if (LocalLayersSequences.Contains(SubSequence->GetFName()))
 	{
 		UE::FSdfLayer Layer = UE::FSdfLayer::FindOrOpen(*LayerTimeInfo->Identifier);
 		if (!Layer)
@@ -2402,53 +2475,53 @@ void FUsdLevelSequenceHelperImpl::UpdateUsdLayerOffsetFromSection(const UMovieSc
 		}
 
 		int32 SubLayerIndex = INDEX_NONE;
-		FLayerOffsetInfo* SubLayerOffset = Algo::FindByPredicate( LayerTimeInfo->SubLayersOffsets,
-			[ &SubLayerIndex, &SubLayerIdentifier = SubLayerTimeInfo->Identifier ]( const FLayerOffsetInfo& Other )
-			{
-				bool bFound = ( Other.LayerIdentifier == SubLayerIdentifier );
-				++SubLayerIndex;
+		FLayerOffsetInfo* SubLayerOffset = Algo::FindByPredicate(LayerTimeInfo->SubLayersOffsets,
+			[&SubLayerIndex, &SubLayerIdentifier = SubLayerTimeInfo->Identifier](const FLayerOffsetInfo& Other)
+		{
+			bool bFound = (Other.LayerIdentifier == SubLayerIdentifier);
+			++SubLayerIndex;
 
-				return bFound;
-			} );
+			return bFound;
+		});
 
-		if ( SubLayerIndex != INDEX_NONE )
+		if (SubLayerIndex != INDEX_NONE)
 		{
-			Layer.SetSubLayerOffset( NewLayerOffset, SubLayerIndex );
-			UpdateLayerTimeInfoFromLayer( *LayerTimeInfo, Layer );
+			Layer.SetSubLayerOffset(NewLayerOffset, SubLayerIndex);
+			UpdateLayerTimeInfoFromLayer(*LayerTimeInfo, Layer);
 		}
 	}
 	else
 	{
 		TArray< FString > PrimPathsForSequence;
-		PrimPathByLevelSequenceName.MultiFind( Section->GetSequence()->GetFName(), PrimPathsForSequence );
+		PrimPathByLevelSequenceName.MultiFind(Section->GetSequence()->GetFName(), PrimPathsForSequence);
 
-		for ( const FString& PrimPath : PrimPathsForSequence )
+		for (const FString& PrimPath : PrimPathsForSequence)
 		{
-			UsdUtils::SetRefOrPayloadLayerOffset( UsdStage.GetPrimAtPath( UE::FSdfPath( *PrimPath ) ), NewLayerOffset );
+			UsdUtils::SetRefOrPayloadLayerOffset(UsdStage.GetPrimAtPath(UE::FSdfPath(*PrimPath)), NewLayerOffset);
 		}
 	}
 
 	UE_LOG(LogUsd, Verbose, TEXT("\tNew OffsetScale: %f, %f"), NewLayerOffset.Offset, NewLayerOffset.Scale);
 }
 
-void FUsdLevelSequenceHelperImpl::UpdateMovieSceneTimeRanges( UMovieScene& MovieScene, const FLayerTimeInfo& LayerTimeInfo )
+void FUsdLevelSequenceHelperImpl::UpdateMovieSceneTimeRanges(UMovieScene& MovieScene, const FLayerTimeInfo& LayerTimeInfo)
 {
 	const double FramesPerSecond = GetFramesPerSecond();
 
-	if ( LayerTimeInfo.IsAnimated() )
+	if (LayerTimeInfo.IsAnimated())
 	{
 		double StartTimeCode = LayerTimeInfo.StartTimeCode.Get(0.0);
 		const double EndTimeCode = LayerTimeInfo.EndTimeCode.Get(0.0);
 
 		double TimeCodesPerSecond = 24.0;
-		if ( UE::FSdfLayer Layer = UE::FSdfLayer::FindOrOpen( *LayerTimeInfo.Identifier ) )
+		if (UE::FSdfLayer Layer = UE::FSdfLayer::FindOrOpen(*LayerTimeInfo.Identifier))
 		{
 			TimeCodesPerSecond = Layer.GetTimeCodesPerSecond();
 
 			// When composing a sublayer that has startTimeCode 10 with an offset of 25 timecodes, USD will place the sublayer's
 			// time code 0 at the 25 mark. We want to mirror that behavior when composing our subsections but still leave the root layer's
 			// playback range to [startTimeCode, endTimeCode] as that's what we'd expect to see, and it doesn't affect composition
-			if ( Layer != UsdStage.GetRootLayer() )
+			if (Layer != UsdStage.GetRootLayer())
 			{
 				StartTimeCode = 0.0;
 			}
@@ -2458,35 +2531,35 @@ void FUsdLevelSequenceHelperImpl::UpdateMovieSceneTimeRanges( UMovieScene& Movie
 			TimeCodesPerSecond = GetTimeCodesPerSecond();
 		}
 
-		const FFrameRate TickResolution  = MovieScene.GetTickResolution();
-		const FFrameNumber StartFrame    = UsdLevelSequenceHelperImpl::RoundAsFrameNumber( TickResolution, StartTimeCode / TimeCodesPerSecond );
-		const FFrameNumber EndFrame      = UsdLevelSequenceHelperImpl::RoundAsFrameNumber( TickResolution, EndTimeCode / TimeCodesPerSecond );
-		TRange< FFrameNumber > TimeRange = TRange<FFrameNumber>::Inclusive( StartFrame, EndFrame );
+		const FFrameRate TickResolution = MovieScene.GetTickResolution();
+		const FFrameNumber StartFrame = UsdLevelSequenceHelperImpl::RoundAsFrameNumber(TickResolution, StartTimeCode / TimeCodesPerSecond);
+		const FFrameNumber EndFrame = UsdLevelSequenceHelperImpl::RoundAsFrameNumber(TickResolution, EndTimeCode / TimeCodesPerSecond);
+		TRange< FFrameNumber > TimeRange = TRange<FFrameNumber>::Inclusive(StartFrame, EndFrame);
 
- 		MovieScene.SetPlaybackRange( TimeRange );
-		MovieScene.SetViewRange( StartTimeCode / TimeCodesPerSecond -1.0f, 1.0f + EndTimeCode / TimeCodesPerSecond );
-		MovieScene.SetWorkingRange( StartTimeCode / TimeCodesPerSecond -1.0f, 1.0f + EndTimeCode / TimeCodesPerSecond );
+		MovieScene.SetPlaybackRange(TimeRange);
+		MovieScene.SetViewRange(StartTimeCode / TimeCodesPerSecond - 1.0f, 1.0f + EndTimeCode / TimeCodesPerSecond);
+		MovieScene.SetWorkingRange(StartTimeCode / TimeCodesPerSecond - 1.0f, 1.0f + EndTimeCode / TimeCodesPerSecond);
 	}
 
 	// Always set these even if we're not animated because if a child layer IS animated and has a different framerate we'll get a warning
 	// from the sequencer. Realistically it makes no difference because if the root layer is not animated (i.e. has 0 for start and end timecodes)
 	// nothing will actually play, but this just prevents the warning
-	MovieScene.SetDisplayRate( FFrameRate( FramesPerSecond, 1 ) );
+	MovieScene.SetDisplayRate(FFrameRate(FramesPerSecond, 1));
 }
 
 void FUsdLevelSequenceHelperImpl::BlockMonitoringChangesForThisTransaction()
 {
-	if ( ITransaction* Trans = GUndo )
+	if (ITransaction* Trans = GUndo)
 	{
 		FTransactionContext Context = Trans->GetContext();
 
 		// We're already blocking this one, so ignore this so that we don't increment our counter too many times
-		if ( BlockedTransactionGuids.Contains( Context.TransactionId ) )
+		if (BlockedTransactionGuids.Contains(Context.TransactionId))
 		{
 			return;
 		}
 
-		BlockedTransactionGuids.Add( Context.TransactionId );
+		BlockedTransactionGuids.Add(Context.TransactionId);
 
 		StopMonitoringChanges();
 	}
@@ -2494,46 +2567,46 @@ void FUsdLevelSequenceHelperImpl::BlockMonitoringChangesForThisTransaction()
 
 void FUsdLevelSequenceHelperImpl::OnObjectTransacted(UObject* Object, const class FTransactionObjectEvent& Event)
 {
-	if ( !MainLevelSequence || !IsMonitoringChanges() || !IsValid(Object) || !UsdStage || BlockedTransactionGuids.Contains( Event.GetTransactionId() ) )
+	if (!MainLevelSequence || !IsMonitoringChanges() || !IsValid(Object) || !UsdStage || BlockedTransactionGuids.Contains(Event.GetTransactionId()))
 	{
 		return;
 	}
 
 	const ULevelSequence* LevelSequence = Object->GetTypedOuter<ULevelSequence>();
-	if ( !LevelSequence || ( LevelSequence != MainLevelSequence && !SequencesID.Contains( LevelSequence ) ) )
+	if (!LevelSequence || (LevelSequence != MainLevelSequence && !SequencesID.Contains(LevelSequence)))
 	{
 		// This is not one of our managed level sequences, so ignore changes
 		return;
 	}
 
-	if ( UMovieScene* MovieScene = Cast< UMovieScene >( Object ) )
+	if (UMovieScene* MovieScene = Cast< UMovieScene >(Object))
 	{
-		HandleMovieSceneChange( *MovieScene );
+		HandleMovieSceneChange(*MovieScene);
 	}
-	else if ( UMovieSceneSubSection* SubSection = Cast<UMovieSceneSubSection>(Object) )
+	else if (UMovieSceneSubSection* SubSection = Cast<UMovieSceneSubSection>(Object))
 	{
-		HandleSubSectionChange( *SubSection );
+		HandleSubSectionChange(*SubSection);
 	}
-	else if ( UMovieSceneTrack* Track = Cast<UMovieSceneTrack>(Object) )
+	else if (UMovieSceneTrack* Track = Cast<UMovieSceneTrack>(Object))
 	{
-		const bool bIsMuteChange = Event.GetChangedProperties().Contains( TEXT("bIsEvalDisabled") );
-		HandleTrackChange( *Track, bIsMuteChange );
+		const bool bIsMuteChange = Event.GetChangedProperties().Contains(TEXT("bIsEvalDisabled"));
+		HandleTrackChange(*Track, bIsMuteChange);
 	}
-	else if ( UMovieSceneSection* Section = Cast<UMovieSceneSection>(Object) )
+	else if (UMovieSceneSection* Section = Cast<UMovieSceneSection>(Object))
 	{
-		const bool bIsMuteChange = Event.GetChangedProperties().Contains( TEXT( "bIsActive" ) );
+		const bool bIsMuteChange = Event.GetChangedProperties().Contains(TEXT("bIsActive"));
 
-		if ( UMovieSceneTrack* ParentTrack = Section->GetTypedOuter<UMovieSceneTrack>() )
+		if (UMovieSceneTrack* ParentTrack = Section->GetTypedOuter<UMovieSceneTrack>())
 		{
-			HandleTrackChange( *ParentTrack, bIsMuteChange );
+			HandleTrackChange(*ParentTrack, bIsMuteChange);
 		}
 
 #if WITH_EDITOR
-		if ( !bIsMuteChange )
+		if (!bIsMuteChange)
 		{
-			if ( UMovieSceneControlRigParameterSection* CRSection = Cast<UMovieSceneControlRigParameterSection>( Section ) )
+			if (UMovieSceneControlRigParameterSection* CRSection = Cast<UMovieSceneControlRigParameterSection>(Section))
 			{
-				if ( GEditor )
+				if (GEditor)
 				{
 					// We have to do this on next tick because HandleControlRigSectionChange will internally bake
 					// the sequence, repeteadly updating the ControlRig hierarchy. There is no way to silence
@@ -2541,10 +2614,10 @@ void FUsdLevelSequenceHelperImpl::OnObjectTransacted(UObject* Object, const clas
 					// a brand new scoped transaction, which asserts inside UTransBuffer::CheckState when if finds
 					// out that the previous transaction wasn't fully complete (OnObjectTransacted gets called before
 					// the current transaction is fully done).
-					GEditor->GetTimerManager()->SetTimerForNextTick( [this, CRSection]( )
-					{
-						HandleControlRigSectionChange( *CRSection );
-					});
+					GEditor->GetTimerManager()->SetTimerForNextTick([this, CRSection]()
+						{
+							HandleControlRigSectionChange(*CRSection);
+						});
 				}
 			}
 		}
@@ -2552,38 +2625,38 @@ void FUsdLevelSequenceHelperImpl::OnObjectTransacted(UObject* Object, const clas
 	}
 }
 
-void FUsdLevelSequenceHelperImpl::OnUsdObjectsChanged( const UsdUtils::FObjectChangesByPath& InfoChanges, const UsdUtils::FObjectChangesByPath& ResyncChanges )
+void FUsdLevelSequenceHelperImpl::OnUsdObjectsChanged(const UsdUtils::FObjectChangesByPath& InfoChanges, const UsdUtils::FObjectChangesByPath& ResyncChanges)
 {
 	AUsdStageActor* StageActorPtr = StageActor.Get();
-	if ( !StageActorPtr || !StageActorPtr->IsListeningToUsdNotices() )
+	if (!StageActorPtr || !StageActorPtr->IsListeningToUsdNotices())
 	{
 		return;
 	}
 
-	FScopedBlockMonitoringChangesForTransaction BlockMonitoring{ *this };
+	FScopedBlockMonitoringChangesForTransaction BlockMonitoring{*this};
 
-	for ( const TPair<FString, TArray<UsdUtils::FObjectChangeNotice>>& InfoChange : InfoChanges )
+	for (const TPair<FString, TArray<UsdUtils::FObjectChangeNotice>>& InfoChange : InfoChanges)
 	{
 		const FString& PrimPath = InfoChange.Key;
-		if ( PrimPath == TEXT( "/" ) )
+		if (PrimPath == TEXT("/"))
 		{
-			for ( const UsdUtils::FObjectChangeNotice& ObjectChange : InfoChange.Value )
+			for (const UsdUtils::FObjectChangeNotice& ObjectChange : InfoChange.Value)
 			{
-				for ( const UsdUtils::FAttributeChange& AttributeChange : ObjectChange.AttributeChanges )
+				for (const UsdUtils::FAttributeChange& AttributeChange : ObjectChange.AttributeChanges)
 				{
-					if ( AttributeChange.PropertyName == TEXT( "framesPerSecond" ) && MainLevelSequence )
+					if (AttributeChange.PropertyName == TEXT("framesPerSecond") && MainLevelSequence)
 					{
 						UsdUtils::FConvertedVtValue ConvertedValue;
-						if ( UsdToUnreal::ConvertValue( AttributeChange.NewValue, ConvertedValue ) )
+						if (UsdToUnreal::ConvertValue(AttributeChange.NewValue, ConvertedValue))
 						{
-							if ( ConvertedValue.Entries.Num() == 1 && ConvertedValue.Entries[ 0 ].Num() == 1 )
+							if (ConvertedValue.Entries.Num() == 1 && ConvertedValue.Entries[0].Num() == 1)
 							{
-								if ( double* NewValue = ConvertedValue.Entries[ 0 ][ 0 ].TryGet<double>() )
+								if (double* NewValue = ConvertedValue.Entries[0][0].TryGet<double>())
 								{
-									if ( UMovieScene* MovieScene = MainLevelSequence->GetMovieScene() )
+									if (UMovieScene* MovieScene = MainLevelSequence->GetMovieScene())
 									{
 										MovieScene->Modify();
-										MovieScene->SetDisplayRate( FFrameRate( *NewValue, 1 ) );
+										MovieScene->SetDisplayRate(FFrameRate(*NewValue, 1));
 									}
 								}
 							}
@@ -2595,9 +2668,9 @@ void FUsdLevelSequenceHelperImpl::OnUsdObjectsChanged( const UsdUtils::FObjectCh
 	}
 }
 
-void FUsdLevelSequenceHelperImpl::HandleTransactionStateChanged( const FTransactionContext& InTransactionContext, const ETransactionStateEventType InTransactionState )
+void FUsdLevelSequenceHelperImpl::HandleTransactionStateChanged(const FTransactionContext& InTransactionContext, const ETransactionStateEventType InTransactionState)
 {
-	if ( InTransactionState == ETransactionStateEventType::TransactionFinalized && BlockedTransactionGuids.Contains( InTransactionContext.TransactionId ) )
+	if (InTransactionState == ETransactionStateEventType::TransactionFinalized && BlockedTransactionGuids.Contains(InTransactionContext.TransactionId))
 	{
 		StartMonitoringChanges();
 	}
@@ -2605,40 +2678,40 @@ void FUsdLevelSequenceHelperImpl::HandleTransactionStateChanged( const FTransact
 
 double FUsdLevelSequenceHelperImpl::GetFramesPerSecond() const
 {
-	if ( !UsdStage )
+	if (!UsdStage)
 	{
 		return DefaultFramerate;
 	}
 
 	const double StageFramesPerSecond = UsdStage.GetFramesPerSecond();
-	return FMath::IsNearlyZero( StageFramesPerSecond ) ? DefaultFramerate : StageFramesPerSecond;
+	return FMath::IsNearlyZero(StageFramesPerSecond) ? DefaultFramerate : StageFramesPerSecond;
 }
 
 double FUsdLevelSequenceHelperImpl::GetTimeCodesPerSecond() const
 {
-	if ( !UsdStage )
+	if (!UsdStage)
 	{
 		return DefaultFramerate;
 	}
 
 	const double StageTimeCodesPerSecond = UsdStage.GetTimeCodesPerSecond();
-	return FMath::IsNearlyZero( StageTimeCodesPerSecond ) ? DefaultFramerate : StageTimeCodesPerSecond;
+	return FMath::IsNearlyZero(StageTimeCodesPerSecond) ? DefaultFramerate : StageTimeCodesPerSecond;
 }
 
-FGuid FUsdLevelSequenceHelperImpl::GetOrCreateComponentBinding( const UUsdPrimTwin& PrimTwin, USceneComponent& ComponentToBind, ULevelSequence& Sequence )
+FGuid FUsdLevelSequenceHelperImpl::GetOrCreateComponentBinding(const UUsdPrimTwin& PrimTwin, USceneComponent& ComponentToBind, ULevelSequence& Sequence)
 {
 	UMovieScene* MovieScene = Sequence.GetMovieScene();
-	if ( !MovieScene )
+	if (!MovieScene)
 	{
 		return {};
 	}
 
-	FPrimTwinBindings& Bindings = PrimTwinToBindings.FindOrAdd( &PrimTwin );
+	FPrimTwinBindings& Bindings = PrimTwinToBindings.FindOrAdd(&PrimTwin);
 
-	ensure( Bindings.Sequence == nullptr || Bindings.Sequence == &Sequence );
+	ensure(Bindings.Sequence == nullptr || Bindings.Sequence == &Sequence);
 	Bindings.Sequence = &Sequence;
 
-	if ( FGuid* ExistingGuid = Bindings.ComponentClassToBindingGuid.Find( ComponentToBind.GetClass() ) )
+	if (FGuid* ExistingGuid = Bindings.ComponentClassToBindingGuid.Find(ComponentToBind.GetClass()))
 	{
 		return *ExistingGuid;
 	}
@@ -2647,13 +2720,13 @@ FGuid FUsdLevelSequenceHelperImpl::GetOrCreateComponentBinding( const UUsdPrimTw
 	FGuid ActorBinding;
 	UObject* ComponentContext = ComponentToBind.GetWorld();
 
-	FString PrimName = FPaths::GetBaseFilename( PrimTwin.PrimPath );
+	FString PrimName = FPaths::GetBaseFilename(PrimTwin.PrimPath);
 
 	// Make sure we always bind the parent actor too
-	if ( AActor* Actor = ComponentToBind.GetOwner() )
+	if (AActor* Actor = ComponentToBind.GetOwner())
 	{
-		ActorBinding = Sequence.FindBindingFromObject( Actor, Actor->GetWorld() );
-		if ( !ActorBinding.IsValid() )
+		ActorBinding = Sequence.FindBindingFromObject(Actor, Actor->GetWorld());
+		if (!ActorBinding.IsValid())
 		{
 			// We use the label here because that will always be named after the prim that caused the actor
 			// to be generated. If we just used our own PrimName in here we may run into situations where a child Camera prim
@@ -2667,54 +2740,54 @@ FGuid FUsdLevelSequenceHelperImpl::GetOrCreateComponentBinding( const UUsdPrimTw
 #endif // WITH_EDITOR
 				Actor->GetClass()
 			);
-			Sequence.BindPossessableObject( ActorBinding, *Actor, Actor->GetWorld() );
+			Sequence.BindPossessableObject(ActorBinding, *Actor, Actor->GetWorld());
 		}
 
 		ComponentContext = Actor;
 	}
 
-	ComponentBinding = MovieScene->AddPossessable( PrimName, ComponentToBind.GetClass() );
+	ComponentBinding = MovieScene->AddPossessable(PrimName, ComponentToBind.GetClass());
 
-	if ( ActorBinding.IsValid() && ComponentBinding.IsValid() )
+	if (ActorBinding.IsValid() && ComponentBinding.IsValid())
 	{
-		if ( FMovieScenePossessable* ComponentPossessable = MovieScene->FindPossessable( ComponentBinding ) )
+		if (FMovieScenePossessable* ComponentPossessable = MovieScene->FindPossessable(ComponentBinding))
 		{
-			ComponentPossessable->SetParent( ActorBinding, MovieScene );
+			ComponentPossessable->SetParent(ActorBinding, MovieScene);
 		}
 	}
 
 	// Bind component
-	Sequence.BindPossessableObject( ComponentBinding, ComponentToBind, ComponentContext );
-	Bindings.ComponentClassToBindingGuid.Emplace( ComponentToBind.GetClass(), ComponentBinding );
+	Sequence.BindPossessableObject(ComponentBinding, ComponentToBind, ComponentContext);
+	Bindings.ComponentClassToBindingGuid.Emplace(ComponentToBind.GetClass(), ComponentBinding);
 	return ComponentBinding;
 }
 
-void FUsdLevelSequenceHelperImpl::HandleMovieSceneChange( UMovieScene& MovieScene )
+void FUsdLevelSequenceHelperImpl::HandleMovieSceneChange(UMovieScene& MovieScene)
 {
 	// It's possible to get this called when the actor and it's level sequences are being all destroyed in one go.
 	// We need the FScopedBlockNotices in this function, but if our StageActor is already being destroyed, we can't reliably
 	// use its listener, and so then we can't do anything. We likely don't want to write back to the stage at this point anyway.
 	AUsdStageActor* StageActorPtr = StageActor.Get();
-	if ( !MainLevelSequence || !UsdStage || !StageActorPtr || StageActorPtr->IsActorBeingDestroyed() )
+	if (!MainLevelSequence || !UsdStage || !StageActorPtr || StageActorPtr->IsActorBeingDestroyed())
 	{
 		return;
 	}
 
 	ULevelSequence* Sequence = MovieScene.GetTypedOuter< ULevelSequence >();
-	if ( !Sequence )
+	if (!Sequence)
 	{
 		return;
 	}
 
-	const FString LayerIdentifier = LayerIdentifierByLevelSequenceName.FindRef( Sequence->GetFName() );
-	FLayerTimeInfo* LayerTimeInfo = LayerTimeInfosByLayerIdentifier.Find( LayerIdentifier );
-	if ( !LayerTimeInfo )
+	const FString LayerIdentifier = LayerIdentifierByLevelSequenceName.FindRef(Sequence->GetFName());
+	FLayerTimeInfo* LayerTimeInfo = LayerTimeInfosByLayerIdentifier.Find(LayerIdentifier);
+	if (!LayerTimeInfo)
 	{
 		return;
 	}
 
-	UE::FSdfLayer Layer = UE::FSdfLayer::FindOrOpen( *LayerTimeInfo->Identifier );
-	if ( !Layer )
+	UE::FSdfLayer Layer = UE::FSdfLayer::FindOrOpen(*LayerTimeInfo->Identifier);
+	if (!Layer)
 	{
 		return;
 	}
@@ -2722,103 +2795,103 @@ void FUsdLevelSequenceHelperImpl::HandleMovieSceneChange( UMovieScene& MovieScen
 	const double StageTimeCodesPerSecond = GetTimeCodesPerSecond();
 	const TRange< FFrameNumber > PlaybackRange = MovieScene.GetPlaybackRange();
 	const FFrameRate DisplayRate = MovieScene.GetDisplayRate();
-	const FFrameRate LayerTimeCodesPerSecond( Layer.GetTimeCodesPerSecond(), 1 );
-	const FFrameTime StartTime = FFrameRate::TransformTime(UE::MovieScene::DiscreteInclusiveLower( PlaybackRange ).Value, MovieScene.GetTickResolution(), LayerTimeCodesPerSecond );
-	const FFrameTime EndTime = FFrameRate::TransformTime(UE::MovieScene::DiscreteExclusiveUpper( PlaybackRange ).Value, MovieScene.GetTickResolution(), LayerTimeCodesPerSecond );
+	const FFrameRate LayerTimeCodesPerSecond(Layer.GetTimeCodesPerSecond(), 1);
+	const FFrameTime StartTime = FFrameRate::TransformTime(UE::MovieScene::DiscreteInclusiveLower(PlaybackRange).Value, MovieScene.GetTickResolution(), LayerTimeCodesPerSecond);
+	const FFrameTime EndTime = FFrameRate::TransformTime(UE::MovieScene::DiscreteExclusiveUpper(PlaybackRange).Value, MovieScene.GetTickResolution(), LayerTimeCodesPerSecond);
 
-	FScopedBlockNoticeListening BlockNotices( StageActor.Get() );
+	FScopedBlockNoticeListening BlockNotices(StageActor.Get());
 	UE::FSdfChangeBlock ChangeBlock;
-	if ( !FMath::IsNearlyEqual( DisplayRate.AsDecimal(), GetFramesPerSecond() ) )
+	if (!FMath::IsNearlyEqual(DisplayRate.AsDecimal(), GetFramesPerSecond()))
 	{
-		UsdStage.SetFramesPerSecond( DisplayRate.AsDecimal() );
+		UsdStage.SetFramesPerSecond(DisplayRate.AsDecimal());
 
 		// For whatever reason setting a stage FramesPerSecond also automatically sets its TimeCodesPerSecond to the same value, so we need to undo it.
 		// This because all the sequencer does is change display rate, which is the analogue to USD's frames per second (i.e. we are only changing how many
 		// frames we'll display between any two timecodes, not how many timecodes we'll display per second)
-		UsdStage.SetTimeCodesPerSecond( StageTimeCodesPerSecond );
+		UsdStage.SetTimeCodesPerSecond(StageTimeCodesPerSecond);
 
 		// Propagate to all movie scenes, as USD only uses the stage FramesPerSecond so the sequences should have a unified DisplayRate to reflect that
-		for ( TPair< FString, ULevelSequence* >& SequenceByIdentifier : LevelSequencesByIdentifier )
+		for (TPair< FString, ULevelSequence* >& SequenceByIdentifier : LevelSequencesByIdentifier)
 		{
-			if ( ULevelSequence* OtherSequence = SequenceByIdentifier.Value )
+			if (ULevelSequence* OtherSequence = SequenceByIdentifier.Value)
 			{
-				if ( UMovieScene* OtherMovieScene = OtherSequence->GetMovieScene() )
+				if (UMovieScene* OtherMovieScene = OtherSequence->GetMovieScene())
 				{
-					OtherMovieScene->SetDisplayRate( DisplayRate );
+					OtherMovieScene->SetDisplayRate(DisplayRate);
 				}
 			}
 		}
 	}
 
-	Layer.SetStartTimeCode( StartTime.RoundToFrame().Value );
-	Layer.SetEndTimeCode( EndTime.RoundToFrame().Value );
+	Layer.SetStartTimeCode(StartTime.RoundToFrame().Value);
+	Layer.SetEndTimeCode(EndTime.RoundToFrame().Value);
 
-	UpdateLayerTimeInfoFromLayer( *LayerTimeInfo, Layer );
+	UpdateLayerTimeInfoFromLayer(*LayerTimeInfo, Layer);
 
-	if ( Sequence == MainLevelSequence )
+	if (Sequence == MainLevelSequence)
 	{
-		CreateTimeTrack( FindOrAddLayerTimeInfo( UsdStage.GetRootLayer() ) );
+		CreateTimeTrack(FindOrAddLayerTimeInfo(UsdStage.GetRootLayer()));
 	}
 
-	auto RemoveTimeSamplesForAttr = []( const UE::FUsdAttribute& Attr )
+	auto RemoveTimeSamplesForAttr = [](const UE::FUsdAttribute& Attr)
 	{
-		if ( !Attr || Attr.GetNumTimeSamples() == 0 )
+		if (!Attr || Attr.GetNumTimeSamples() == 0)
 		{
 			return;
 		}
 
-		if ( UE::FSdfLayer AttrLayer = UsdUtils::FindLayerForAttribute( Attr, 0.0 ) )
+		if (UE::FSdfLayer AttrLayer = UsdUtils::FindLayerForAttribute(Attr, 0.0))
 		{
 			UE::FSdfPath AttrPath = Attr.GetPath();
-			for ( double TimeSample : AttrLayer.ListTimeSamplesForPath( AttrPath ) )
+			for (double TimeSample : AttrLayer.ListTimeSamplesForPath(AttrPath))
 			{
-				AttrLayer.EraseTimeSample( AttrPath, TimeSample );
+				AttrLayer.EraseTimeSample(AttrPath, TimeSample);
 			}
 		}
 	};
 
-	auto RemoveTimeSamplesForPropertyIfNeeded = [&MovieScene, &RemoveTimeSamplesForAttr]( const UE::FUsdPrim& Prim, const FGuid& Guid, const FName& PropertyPath )
+	auto RemoveTimeSamplesForPropertyIfNeeded = [&MovieScene, &RemoveTimeSamplesForAttr](const UE::FUsdPrim& Prim, const FGuid& Guid, const FName& PropertyPath)
 	{
-		if ( !UsdLevelSequenceHelperImpl::FindTrackTypeOrDerived<UMovieScenePropertyTrack>( &MovieScene, Guid, PropertyPath ) )
+		if (!UsdLevelSequenceHelperImpl::FindTrackTypeOrDerived<UMovieScenePropertyTrack>(&MovieScene, Guid, PropertyPath))
 		{
-			for ( UE::FUsdAttribute& Attr : UnrealToUsd::GetAttributesForProperty( Prim, PropertyPath ) )
+			for (UE::FUsdAttribute& Attr : UnrealToUsd::GetAttributesForProperty(Prim, PropertyPath))
 			{
-				RemoveTimeSamplesForAttr( Attr );
+				RemoveTimeSamplesForAttr(Attr);
 			}
 		}
 	};
 
 	// Check if we deleted things
-	for ( TMap< TWeakObjectPtr< const UUsdPrimTwin >, FPrimTwinBindings >::TIterator PrimTwinIt = PrimTwinToBindings.CreateIterator();
+	for (TMap< TWeakObjectPtr< const UUsdPrimTwin >, FPrimTwinBindings >::TIterator PrimTwinIt = PrimTwinToBindings.CreateIterator();
 		PrimTwinIt;
-		++PrimTwinIt )
+		++PrimTwinIt)
 	{
 		const UUsdPrimTwin* UsdPrimTwin = PrimTwinIt->Key.Get();
 		FPrimTwinBindings& Bindings = PrimTwinIt->Value;
 
-		if ( Bindings.Sequence != Sequence )
+		if (Bindings.Sequence != Sequence)
 		{
 			continue;
 		}
 
-		for ( TMap< const UClass*, FGuid >::TIterator BindingIt = Bindings.ComponentClassToBindingGuid.CreateIterator();
+		for (TMap< const UClass*, FGuid >::TIterator BindingIt = Bindings.ComponentClassToBindingGuid.CreateIterator();
 			BindingIt;
-			++BindingIt )
+			++BindingIt)
 		{
 			const FGuid& Guid = BindingIt->Value;
 
 			// Deleted the possessable
-			if ( !MovieScene.FindPossessable( Guid ) )
+			if (!MovieScene.FindPossessable(Guid))
 			{
 				BindingIt.RemoveCurrent();
 			}
 
 			// Check if we have an animated attribute and no track for it --> We may have deleted the track, so clear that attribute
 			// We could keep track of these when adding in some kind of map, but while slower this is likely more robust due to the need to support undo/redo
-			if ( UsdPrimTwin )
+			if (UsdPrimTwin)
 			{
 				USceneComponent* BoundComponent = UsdPrimTwin->GetSceneComponent();
-				if ( !BoundComponent )
+				if (!BoundComponent)
 				{
 					continue;
 				}
@@ -2827,44 +2900,44 @@ void FUsdLevelSequenceHelperImpl::HandleMovieSceneChange( UMovieScene& MovieScen
 				const bool bIsLight = BoundComponent->IsA<ULightComponentBase>();
 				const bool bIsSkeletal = BoundComponent->IsA<USkeletalMeshComponent>();
 
-				if ( UE::FUsdPrim UsdPrim = UsdStage.GetPrimAtPath( UE::FSdfPath( *UsdPrimTwin->PrimPath ) ) )
+				if (UE::FUsdPrim UsdPrim = UsdStage.GetPrimAtPath(UE::FSdfPath(*UsdPrimTwin->PrimPath)))
 				{
-					RemoveTimeSamplesForPropertyIfNeeded( UsdPrim, Guid, UnrealIdentifiers::TransformPropertyName );
-					RemoveTimeSamplesForPropertyIfNeeded( UsdPrim, Guid, UnrealIdentifiers::HiddenInGamePropertyName );
+					RemoveTimeSamplesForPropertyIfNeeded(UsdPrim, Guid, UnrealIdentifiers::TransformPropertyName);
+					RemoveTimeSamplesForPropertyIfNeeded(UsdPrim, Guid, UnrealIdentifiers::HiddenInGamePropertyName);
 
-					if ( bIsCamera )
+					if (bIsCamera)
 					{
-						RemoveTimeSamplesForPropertyIfNeeded( UsdPrim, Guid, UnrealIdentifiers::CurrentFocalLengthPropertyName );
-						RemoveTimeSamplesForPropertyIfNeeded( UsdPrim, Guid, UnrealIdentifiers::ManualFocusDistancePropertyName );
-						RemoveTimeSamplesForPropertyIfNeeded( UsdPrim, Guid, UnrealIdentifiers::CurrentAperturePropertyName );
-						RemoveTimeSamplesForPropertyIfNeeded( UsdPrim, Guid, UnrealIdentifiers::SensorWidthPropertyName );
-						RemoveTimeSamplesForPropertyIfNeeded( UsdPrim, Guid, UnrealIdentifiers::SensorHeightPropertyName );
+						RemoveTimeSamplesForPropertyIfNeeded(UsdPrim, Guid, UnrealIdentifiers::CurrentFocalLengthPropertyName);
+						RemoveTimeSamplesForPropertyIfNeeded(UsdPrim, Guid, UnrealIdentifiers::ManualFocusDistancePropertyName);
+						RemoveTimeSamplesForPropertyIfNeeded(UsdPrim, Guid, UnrealIdentifiers::CurrentAperturePropertyName);
+						RemoveTimeSamplesForPropertyIfNeeded(UsdPrim, Guid, UnrealIdentifiers::SensorWidthPropertyName);
+						RemoveTimeSamplesForPropertyIfNeeded(UsdPrim, Guid, UnrealIdentifiers::SensorHeightPropertyName);
 					}
-					else if ( bIsLight )
+					else if (bIsLight)
 					{
-						RemoveTimeSamplesForPropertyIfNeeded( UsdPrim, Guid, UnrealIdentifiers::IntensityPropertyName );
-						RemoveTimeSamplesForPropertyIfNeeded( UsdPrim, Guid, UnrealIdentifiers::LightColorPropertyName );
-						RemoveTimeSamplesForPropertyIfNeeded( UsdPrim, Guid, UnrealIdentifiers::UseTemperaturePropertyName );
-						RemoveTimeSamplesForPropertyIfNeeded( UsdPrim, Guid, UnrealIdentifiers::TemperaturePropertyName );
-						RemoveTimeSamplesForPropertyIfNeeded( UsdPrim, Guid, UnrealIdentifiers::SourceRadiusPropertyName );
-						RemoveTimeSamplesForPropertyIfNeeded( UsdPrim, Guid, UnrealIdentifiers::SourceWidthPropertyName );
-						RemoveTimeSamplesForPropertyIfNeeded( UsdPrim, Guid, UnrealIdentifiers::SourceHeightPropertyName );
-						RemoveTimeSamplesForPropertyIfNeeded( UsdPrim, Guid, UnrealIdentifiers::OuterConeAnglePropertyName );
-						RemoveTimeSamplesForPropertyIfNeeded( UsdPrim, Guid, UnrealIdentifiers::InnerConeAnglePropertyName );
-						RemoveTimeSamplesForPropertyIfNeeded( UsdPrim, Guid, UnrealIdentifiers::LightSourceAnglePropertyName );
+						RemoveTimeSamplesForPropertyIfNeeded(UsdPrim, Guid, UnrealIdentifiers::IntensityPropertyName);
+						RemoveTimeSamplesForPropertyIfNeeded(UsdPrim, Guid, UnrealIdentifiers::LightColorPropertyName);
+						RemoveTimeSamplesForPropertyIfNeeded(UsdPrim, Guid, UnrealIdentifiers::UseTemperaturePropertyName);
+						RemoveTimeSamplesForPropertyIfNeeded(UsdPrim, Guid, UnrealIdentifiers::TemperaturePropertyName);
+						RemoveTimeSamplesForPropertyIfNeeded(UsdPrim, Guid, UnrealIdentifiers::SourceRadiusPropertyName);
+						RemoveTimeSamplesForPropertyIfNeeded(UsdPrim, Guid, UnrealIdentifiers::SourceWidthPropertyName);
+						RemoveTimeSamplesForPropertyIfNeeded(UsdPrim, Guid, UnrealIdentifiers::SourceHeightPropertyName);
+						RemoveTimeSamplesForPropertyIfNeeded(UsdPrim, Guid, UnrealIdentifiers::OuterConeAnglePropertyName);
+						RemoveTimeSamplesForPropertyIfNeeded(UsdPrim, Guid, UnrealIdentifiers::InnerConeAnglePropertyName);
+						RemoveTimeSamplesForPropertyIfNeeded(UsdPrim, Guid, UnrealIdentifiers::LightSourceAnglePropertyName);
 					}
-					else if ( bIsSkeletal )
+					else if (bIsSkeletal)
 					{
-						if ( !MovieScene.FindTrack( UMovieSceneSkeletalAnimationTrack::StaticClass(), Guid ) )
+						if (!MovieScene.FindTrack(UMovieSceneSkeletalAnimationTrack::StaticClass(), Guid))
 						{
-							if ( UE::FUsdPrim SkelAnimationPrim = UsdUtils::FindFirstAnimationSource( UsdPrim ) )
+							if (UE::FUsdPrim SkelAnimationPrim = UsdUtils::FindFirstAnimationSource(UsdPrim))
 							{
-								if ( UE::FSdfLayer SkelAnimationLayer = UsdUtils::FindLayerForPrim( SkelAnimationPrim ) )
+								if (UE::FSdfLayer SkelAnimationLayer = UsdUtils::FindLayerForPrim(SkelAnimationPrim))
 								{
-									RemoveTimeSamplesForAttr( SkelAnimationPrim.GetAttribute( TEXT( "blendShapeWeights" ) ) );
-									RemoveTimeSamplesForAttr( SkelAnimationPrim.GetAttribute( TEXT( "rotations" ) ) );
-									RemoveTimeSamplesForAttr( SkelAnimationPrim.GetAttribute( TEXT( "translations" ) ) );
-									RemoveTimeSamplesForAttr( SkelAnimationPrim.GetAttribute( TEXT( "scales" ) ) );
+									RemoveTimeSamplesForAttr(SkelAnimationPrim.GetAttribute(TEXT("blendShapeWeights")));
+									RemoveTimeSamplesForAttr(SkelAnimationPrim.GetAttribute(TEXT("rotations")));
+									RemoveTimeSamplesForAttr(SkelAnimationPrim.GetAttribute(TEXT("translations")));
+									RemoveTimeSamplesForAttr(SkelAnimationPrim.GetAttribute(TEXT("scales")));
 								}
 							}
 						}
@@ -2875,7 +2948,7 @@ void FUsdLevelSequenceHelperImpl::HandleMovieSceneChange( UMovieScene& MovieScen
 	}
 }
 
-void FUsdLevelSequenceHelperImpl::HandleSubSectionChange( UMovieSceneSubSection& Section )
+void FUsdLevelSequenceHelperImpl::HandleSubSectionChange(UMovieSceneSubSection& Section)
 {
 	UMovieSceneSequence* ParentSequence = Section.GetTypedOuter<UMovieSceneSequence>();
 	if (!ParentSequence)
@@ -2886,104 +2959,104 @@ void FUsdLevelSequenceHelperImpl::HandleSubSectionChange( UMovieSceneSubSection&
 	UpdateUsdLayerOffsetFromSection(ParentSequence, &Section);
 }
 
-void FUsdLevelSequenceHelperImpl::HandleControlRigSectionChange( UMovieSceneControlRigParameterSection& Section )
+void FUsdLevelSequenceHelperImpl::HandleControlRigSectionChange(UMovieSceneControlRigParameterSection& Section)
 {
 #if WITH_EDITOR
 	AUsdStageActor* StageActorValue = StageActor.Get();
-	if ( !StageActorValue )
+	if (!StageActorValue)
 	{
 		return;
 	}
 
 	UWorld* World = StageActorValue->GetWorld();
-	if ( !World )
+	if (!World)
 	{
 		return;
 	}
 
 	ULevelSequence* LevelSequence = Section.GetTypedOuter<ULevelSequence>();
-	if ( !LevelSequence )
+	if (!LevelSequence)
 	{
 		return;
 	}
 
 	UMovieScene* MovieScene = LevelSequence->GetMovieScene();
-	if ( !MovieScene )
+	if (!MovieScene)
 	{
 		return;
 	}
 
 	UMovieSceneTrack* ParentTrack = Section.GetTypedOuter<UMovieSceneTrack>();
-	if ( !ParentTrack )
+	if (!ParentTrack)
 	{
 		return;
 	}
 
 	FGuid PossessableGuid;
-	const bool bFound = MovieScene->FindTrackBinding( *ParentTrack, PossessableGuid );
-	if ( !bFound )
+	const bool bFound = MovieScene->FindTrackBinding(*ParentTrack, PossessableGuid);
+	if (!bFound)
 	{
 		return;
 	}
 
-	FMovieScenePossessable* Possessable = MovieScene->FindPossessable( PossessableGuid );
-	if ( !Possessable )
+	FMovieScenePossessable* Possessable = MovieScene->FindPossessable(PossessableGuid);
+	if (!Possessable)
 	{
 		return;
 	}
 
 	USkeletalMeshComponent* BoundComponent = Cast< USkeletalMeshComponent >(
-		UsdLevelSequenceHelperImpl::LocateBoundObject( *LevelSequence, *Possessable )
-	);
-	if ( !BoundComponent )
+		UsdLevelSequenceHelperImpl::LocateBoundObject(*LevelSequence, *Possessable)
+		);
+	if (!BoundComponent)
 	{
 		return;
 	}
-	ensure( BoundComponent->Mobility != EComponentMobility::Static );
+	ensure(BoundComponent->Mobility != EComponentMobility::Static);
 
 	USkeleton* Skeleton = BoundComponent->GetSkeletalMeshAsset() ? BoundComponent->GetSkeletalMeshAsset()->GetSkeleton() : nullptr;
-	if ( !Skeleton )
+	if (!Skeleton)
 	{
 		return;
 	}
 
-	UUsdPrimTwin* PrimTwin = StageActorValue->RootUsdTwin->Find( BoundComponent );
-	if ( !PrimTwin )
+	UUsdPrimTwin* PrimTwin = StageActorValue->RootUsdTwin->Find(BoundComponent);
+	if (!PrimTwin)
 	{
 		return;
 	}
 
-	UE::FUsdPrim UsdPrim = UsdStage.GetPrimAtPath( UE::FSdfPath( *PrimTwin->PrimPath ) );
-	if ( !UsdPrim )
+	UE::FUsdPrim UsdPrim = UsdStage.GetPrimAtPath(UE::FSdfPath(*PrimTwin->PrimPath));
+	if (!UsdPrim)
 	{
 		return;
 	}
 
 	// We'll place the skeletal animation track wherever the SkelAnimation prim is defined (not necessarily the
 	// same layer as the skel root)
-	UE::FUsdPrim SkelAnimationPrim = UsdUtils::FindFirstAnimationSource( UsdPrim );
-	if ( !SkelAnimationPrim )
+	UE::FUsdPrim SkelAnimationPrim = UsdUtils::FindFirstAnimationSource(UsdPrim);
+	if (!SkelAnimationPrim)
 	{
 		return;
 	}
 
 	// Each sequence corresponds to a specific USD layer. If we're editing something in a sequence, then we must make
 	// that layer the edit target too
-	UE::FSdfLayer EditTargetLayer = FindEditTargetForSubsequence( LevelSequence );
-	if ( !EditTargetLayer )
+	UE::FSdfLayer EditTargetLayer = FindEditTargetForSubsequence(LevelSequence);
+	if (!EditTargetLayer)
 	{
 		return;
 	}
-	UE::FUsdEditContext EditContext{ UsdStage, EditTargetLayer };
+	UE::FUsdEditContext EditContext{UsdStage, EditTargetLayer};
 
-	TSharedPtr<ISequencer> PinnedSequencer = UsdLevelSequenceHelperImpl::GetOpenedSequencerForLevelSequence( MainLevelSequence );
+	TSharedPtr<ISequencer> PinnedSequencer = UsdLevelSequenceHelperImpl::GetOpenedSequencerForLevelSequence(MainLevelSequence);
 
 	// Fetch a sequence player we can use. We'll almost always have the sequencer opened here (we are responding to a transaction
 	// where the section was changed after all), but its possible to have a fallback too
 	IMovieScenePlayer* Player = nullptr;
 	ULevelSequencePlayer* LevelPlayer = nullptr;
 	{
-		if ( PinnedSequencer )
+		if (PinnedSequencer)
 		{
 			Player = PinnedSequencer.Get();
 		}
@@ -2991,10 +3064,10 @@ void FUsdLevelSequenceHelperImpl::HandleControlRigSectionChange( UMovieSceneCont
 		{
 			ALevelSequenceActor* OutActor = nullptr;
 			FMovieSceneSequencePlaybackSettings Settings;
-			Player = LevelPlayer = ULevelSequencePlayer::CreateLevelSequencePlayer( World, LevelSequence, Settings, OutActor);
+			Player = LevelPlayer = ULevelSequencePlayer::CreateLevelSequencePlayer(World, LevelSequence, Settings, OutActor);
 		}
 
-		if ( !Player )
+		if (!Player)
 		{
 			return;
 		}
@@ -3002,17 +3075,17 @@ void FUsdLevelSequenceHelperImpl::HandleControlRigSectionChange( UMovieSceneCont
 
 	// We obviously don't want to respond to the fact that the stage will be modified since we're the
 	// ones actually modifying it already
-	FScopedBlockNoticeListening BlockNotices( StageActorValue );
+	FScopedBlockNoticeListening BlockNotices(StageActorValue);
 
 	// Prepare for baking
 	{
-		if ( PinnedSequencer )
+		if (PinnedSequencer)
 		{
 			PinnedSequencer->EnterSilentMode();
 		}
 
-		FSpawnableRestoreState SpawnableRestoreState( MovieScene );
-		if ( LevelPlayer && SpawnableRestoreState.bWasChanged )
+		FSpawnableRestoreState SpawnableRestoreState(MovieScene);
+		if (LevelPlayer && SpawnableRestoreState.bWasChanged)
 		{
 			// Evaluate at the beginning of the subscene time to ensure that spawnables are created before export
 			// Note that we never actually generate spawnables on our LevelSequence, but its a common pattern to
@@ -3028,8 +3101,8 @@ void FUsdLevelSequenceHelperImpl::HandleControlRigSectionChange( UMovieSceneCont
 	}
 
 	FMovieSceneSequenceTransform SequenceTransform;
-	FMovieSceneSequenceID SequenceID = SequencesID.FindRef( LevelSequence );
-	if ( FMovieSceneSubSequenceData* SubSequenceData = SequenceHierarchyCache.FindSubData( SequenceID ) )
+	FMovieSceneSequenceID SequenceID = SequencesID.FindRef(LevelSequence);
+	if (FMovieSceneSubSequenceData* SubSequenceData = SequenceHierarchyCache.FindSubData(SequenceID))
 	{
 		SequenceTransform = SubSequenceData->RootToSequenceTransform;
 	}
@@ -3049,19 +3122,19 @@ void FUsdLevelSequenceHelperImpl::HandleControlRigSectionChange( UMovieSceneCont
 
 	// Cleanup after baking
 	{
-		if ( LevelPlayer )
+		if (LevelPlayer)
 		{
 			LevelPlayer->Stop();
 		}
 
-		if ( PinnedSequencer )
+		if (PinnedSequencer)
 		{
 			PinnedSequencer->ExitSilentMode();
 			PinnedSequencer->RequestEvaluate();
 		}
 	}
 
-	if ( bBaked )
+	if (bBaked)
 	{
 		// After we bake, both the sequencer and the USD stage have our updated tracks, but we still have the old
 		// AnimSequence asset on the component. If we closed the Sequencer and just animated via the Time attribute,
@@ -3070,170 +3143,170 @@ void FUsdLevelSequenceHelperImpl::HandleControlRigSectionChange( UMovieSceneCont
 		// skel root. Sadly we do need to regenerate the skeletal mesh too, since we may need to affect blend shapes
 		// for the correct bake. The user can disable this behavior (e.g. for costly skeletal meshes) by setting
 		// USD.RegenerateSkeletalAssetsOnControlRigBake to false.
-		GetOnSkelAnimationBaked().Broadcast( PrimTwin->PrimPath );
+		GetOnSkelAnimationBaked().Broadcast(PrimTwin->PrimPath);
 	}
 #endif // WITH_EDITOR
 }
 
-void FUsdLevelSequenceHelperImpl::HandleTrackChange( const UMovieSceneTrack& Track, bool bIsMuteChange )
+void FUsdLevelSequenceHelperImpl::HandleTrackChange(const UMovieSceneTrack& Track, bool bIsMuteChange)
 {
-	if ( !StageActor.IsValid() )
+	if (!StageActor.IsValid())
 	{
 		return;
 	}
 
 	ULevelSequence* Sequence = Track.GetTypedOuter< ULevelSequence >();
-	if ( !Sequence )
+	if (!Sequence)
 	{
 		return;
 	}
 
 	UMovieScene* MovieScene = Sequence->GetMovieScene();
-	if ( !MovieScene )
+	if (!MovieScene)
 	{
 		return;
 	}
 
 	FGuid PossessableGuid;
-	bool bFound = MovieScene->FindTrackBinding( Track, PossessableGuid );
-	if ( !bFound )
+	bool bFound = MovieScene->FindTrackBinding(Track, PossessableGuid);
+	if (!bFound)
 	{
 		return;
 	}
 
-	FMovieScenePossessable* Possessable = MovieScene->FindPossessable( PossessableGuid );
-	if ( !Possessable )
+	FMovieScenePossessable* Possessable = MovieScene->FindPossessable(PossessableGuid);
+	if (!Possessable)
 	{
 		return;
 	}
 
-	UObject* BoundObject = UsdLevelSequenceHelperImpl::LocateBoundObject( *Sequence, *Possessable );
-	if( !BoundObject )
+	UObject* BoundObject = UsdLevelSequenceHelperImpl::LocateBoundObject(*Sequence, *Possessable);
+	if (!BoundObject)
 	{
 		return;
 	}
 
 	// Our tracked bindings are always directly to components. If we don't have one here just abort
-	USceneComponent* BoundSceneComponent = Cast< USceneComponent >( BoundObject );
-	if ( !BoundSceneComponent )
+	USceneComponent* BoundSceneComponent = Cast< USceneComponent >(BoundObject);
+	if (!BoundSceneComponent)
 	{
 		return;
 	}
-	ensure( BoundSceneComponent->Mobility != EComponentMobility::Static );
+	ensure(BoundSceneComponent->Mobility != EComponentMobility::Static);
 
-	UUsdPrimTwin* PrimTwin = StageActor->RootUsdTwin->Find( BoundSceneComponent );
+	UUsdPrimTwin* PrimTwin = StageActor->RootUsdTwin->Find(BoundSceneComponent);
 
 	// If we exported/created this Camera prim ourselves, we'll have a decomposed parent Xform and a child Camera prim (to mirror
 	// the ACineCameraActor structure), and we should have created prim twins for both when opening this stage.
 	// If this USD layer is not authored by us, it may just be a standalone Camera prim: In this scenario the created PrimTwin
 	// will be pointing at the parent USceneComponent of the spawned ACineCameraActor, and we wouldn't find anything when searching
 	// for the camera component directly, so try again
-	if ( !PrimTwin && BoundSceneComponent->IsA<UCineCameraComponent>() )
+	if (!PrimTwin && BoundSceneComponent->IsA<UCineCameraComponent>())
 	{
-		if ( const UMovieScenePropertyTrack* PropertyTrack = Cast<const UMovieScenePropertyTrack>( &Track ) )
+		if (const UMovieScenePropertyTrack* PropertyTrack = Cast<const UMovieScenePropertyTrack>(&Track))
 		{
 			const FName& PropertyPath = PropertyTrack->GetPropertyPath();
 
 			// In the scenario where we're trying to make non-decomposed Camera prims work, we only ever want to write out
 			// actual camera properties from the CameraComponent to the Camera prim. We won't write its USceneComponent
 			// properties, as we will use the ones from the ACineCameraActor's parent USceneComponent instead
-			if ( PropertyPath == UnrealIdentifiers::CurrentFocalLengthPropertyName ||
-				 PropertyPath == UnrealIdentifiers::ManualFocusDistancePropertyName ||
-				 PropertyPath == UnrealIdentifiers::ManualFocusDistancePropertyName ||
-				 PropertyPath == UnrealIdentifiers::CurrentAperturePropertyName ||
-				 PropertyPath == UnrealIdentifiers::SensorWidthPropertyName ||
-				 PropertyPath == UnrealIdentifiers::SensorHeightPropertyName
-			)
+			if (PropertyPath == UnrealIdentifiers::CurrentFocalLengthPropertyName ||
+				PropertyPath == UnrealIdentifiers::ManualFocusDistancePropertyName ||
+				PropertyPath == UnrealIdentifiers::ManualFocusDistancePropertyName ||
+				PropertyPath == UnrealIdentifiers::CurrentAperturePropertyName ||
+				PropertyPath == UnrealIdentifiers::SensorWidthPropertyName ||
+				PropertyPath == UnrealIdentifiers::SensorHeightPropertyName
+				)
 			{
-				PrimTwin = StageActor->RootUsdTwin->Find( BoundSceneComponent->GetAttachParent() );
+				PrimTwin = StageActor->RootUsdTwin->Find(BoundSceneComponent->GetAttachParent());
 			}
 		}
 	}
 
 	// Each sequence corresponds to a specific USD layer. If we're editing something in a sequence, then we must make
 	// that layer the edit target too
-	UE::FSdfLayer EditTargetLayer = FindEditTargetForSubsequence( Sequence );
-	if ( !EditTargetLayer )
+	UE::FSdfLayer EditTargetLayer = FindEditTargetForSubsequence(Sequence);
+	if (!EditTargetLayer)
 	{
 		return;
 	}
-	UE::FUsdEditContext EditContext{ UsdStage, EditTargetLayer };
+	UE::FUsdEditContext EditContext{UsdStage, EditTargetLayer};
 
-	if ( PrimTwin )
+	if (PrimTwin)
 	{
-		FScopedBlockNoticeListening BlockNotices( StageActor.Get() );
-		UE::FUsdPrim UsdPrim = UsdStage.GetPrimAtPath( UE::FSdfPath( *PrimTwin->PrimPath ) );
+		FScopedBlockNoticeListening BlockNotices(StageActor.Get());
+		UE::FUsdPrim UsdPrim = UsdStage.GetPrimAtPath(UE::FSdfPath(*PrimTwin->PrimPath));
 
 		FPrimTwinBindings& Bindings = PrimTwinToBindings.FindOrAdd(PrimTwin);
-		ensure( Bindings.Sequence == nullptr || Bindings.Sequence == Sequence );
+		ensure(Bindings.Sequence == nullptr || Bindings.Sequence == Sequence);
 		Bindings.Sequence = Sequence;
 
 		// Make sure we track this binding
 		const UClass* ComponentClass = BoundSceneComponent->GetClass();
-		FGuid* FoundExistingGuid = Bindings.ComponentClassToBindingGuid.Find( ComponentClass );
-		ensure( !FoundExistingGuid || *FoundExistingGuid == PossessableGuid );
-		Bindings.ComponentClassToBindingGuid.Emplace( ComponentClass, PossessableGuid );
+		FGuid* FoundExistingGuid = Bindings.ComponentClassToBindingGuid.Find(ComponentClass);
+		ensure(!FoundExistingGuid || *FoundExistingGuid == PossessableGuid);
+		Bindings.ComponentClassToBindingGuid.Emplace(ComponentClass, PossessableGuid);
 
-		if ( bIsMuteChange )
+		if (bIsMuteChange)
 		{
-			if ( const UMovieScenePropertyTrack* PropertyTrack = Cast<const UMovieScenePropertyTrack>( &Track ) )
+			if (const UMovieScenePropertyTrack* PropertyTrack = Cast<const UMovieScenePropertyTrack>(&Track))
 			{
 				const FName& PropertyPath = PropertyTrack->GetPropertyPath();
 
-				TArray<UE::FUsdAttribute> Attrs = UnrealToUsd::GetAttributesForProperty( UsdPrim, PropertyPath );
-				if ( Attrs.Num() > 0 )
+				TArray<UE::FUsdAttribute> Attrs = UnrealToUsd::GetAttributesForProperty(UsdPrim, PropertyPath);
+				if (Attrs.Num() > 0)
 				{
 					// Only mute/unmute the first (i.e. main) attribute: If we mute the intensity track we don't want to also mute the
 					// rect width track if it has one
 					UE::FUsdAttribute& Attr = Attrs[0];
 
 					bool bAllSectionsMuted = true;
-					for ( const UMovieSceneSection* Section : Track.GetAllSections() ) // There's no const version of "FindSection"
+					for (const UMovieSceneSection* Section : Track.GetAllSections()) // There's no const version of "FindSection"
 					{
 						bAllSectionsMuted &= !Section->IsActive();
 					}
 
-					if ( Track.IsEvalDisabled() || bAllSectionsMuted )
+					if (Track.IsEvalDisabled() || bAllSectionsMuted)
 					{
-						UsdUtils::MuteAttribute( Attr, UsdStage );
+						UsdUtils::MuteAttribute(Attr, UsdStage);
 					}
 					else
 					{
-						UsdUtils::UnmuteAttribute( Attr, UsdStage );
+						UsdUtils::UnmuteAttribute(Attr, UsdStage);
 					}
 
 					// The attribute may have an effect on the stage, so animate it right away
 					StageActor->OnTimeChanged.Broadcast();
 				}
 			}
-			else if ( const UMovieSceneSkeletalAnimationTrack* SkeletalTrack = Cast<const UMovieSceneSkeletalAnimationTrack>( &Track ) )
+			else if (const UMovieSceneSkeletalAnimationTrack* SkeletalTrack = Cast<const UMovieSceneSkeletalAnimationTrack>(&Track))
 			{
 				bool bAllSectionsMuted = true;
-				for ( const UMovieSceneSection* Section : SkeletalTrack->GetAllSections() ) // There's no const version of "FindSection"
+				for (const UMovieSceneSection* Section : SkeletalTrack->GetAllSections()) // There's no const version of "FindSection"
 				{
 					bAllSectionsMuted &= !Section->IsActive();
 				}
 
-				if ( UE::FUsdPrim SkelAnimationPrim = UsdUtils::FindFirstAnimationSource( UsdPrim ) )
+				if (UE::FUsdPrim SkelAnimationPrim = UsdUtils::FindFirstAnimationSource(UsdPrim))
 				{
-					UE::FUsdAttribute TranslationsAttr = SkelAnimationPrim.GetAttribute( TEXT( "translations" ) );
-					UE::FUsdAttribute RotationsAttr = SkelAnimationPrim.GetAttribute( TEXT( "rotations" ) );
-					UE::FUsdAttribute ScalesAttr = SkelAnimationPrim.GetAttribute( TEXT( "scales" ) );
-					UE::FUsdAttribute BlendShapeWeightsAttr = SkelAnimationPrim.GetAttribute( TEXT( "blendShapeWeights" ) );
+					UE::FUsdAttribute TranslationsAttr = SkelAnimationPrim.GetAttribute(TEXT("translations"));
+					UE::FUsdAttribute RotationsAttr = SkelAnimationPrim.GetAttribute(TEXT("rotations"));
+					UE::FUsdAttribute ScalesAttr = SkelAnimationPrim.GetAttribute(TEXT("scales"));
+					UE::FUsdAttribute BlendShapeWeightsAttr = SkelAnimationPrim.GetAttribute(TEXT("blendShapeWeights"));
 
-					if ( Track.IsEvalDisabled() || bAllSectionsMuted )
+					if (Track.IsEvalDisabled() || bAllSectionsMuted)
 					{
-						UsdUtils::MuteAttribute( TranslationsAttr, UsdStage );
-						UsdUtils::MuteAttribute( RotationsAttr, UsdStage );
-						UsdUtils::MuteAttribute( ScalesAttr, UsdStage );
-						UsdUtils::MuteAttribute( BlendShapeWeightsAttr, UsdStage );
+						UsdUtils::MuteAttribute(TranslationsAttr, UsdStage);
+						UsdUtils::MuteAttribute(RotationsAttr, UsdStage);
+						UsdUtils::MuteAttribute(ScalesAttr, UsdStage);
+						UsdUtils::MuteAttribute(BlendShapeWeightsAttr, UsdStage);
 					}
 					else
 					{
-						UsdUtils::UnmuteAttribute( TranslationsAttr, UsdStage );
-						UsdUtils::UnmuteAttribute( RotationsAttr, UsdStage );
-						UsdUtils::UnmuteAttribute( ScalesAttr, UsdStage );
-						UsdUtils::UnmuteAttribute( BlendShapeWeightsAttr, UsdStage );
+						UsdUtils::UnmuteAttribute(TranslationsAttr, UsdStage);
+						UsdUtils::UnmuteAttribute(RotationsAttr, UsdStage);
+						UsdUtils::UnmuteAttribute(ScalesAttr, UsdStage);
+						UsdUtils::UnmuteAttribute(BlendShapeWeightsAttr, UsdStage);
 					}
 
 					// The attribute may have an effect on the stage, so animate it right away
@@ -3245,9 +3318,9 @@ void FUsdLevelSequenceHelperImpl::HandleTrackChange( const UMovieSceneTrack& Tra
 		{
 			FMovieSceneSequenceTransform SequenceTransform;
 
-			if ( const FMovieSceneSequenceID* SequenceID = SequencesID.Find( Sequence ) )
+			if (const FMovieSceneSequenceID* SequenceID = SequencesID.Find(Sequence))
 			{
-				if ( FMovieSceneSubSequenceData* SubSequenceData = SequenceHierarchyCache.FindSubData( *SequenceID ) )
+				if (FMovieSceneSubSequenceData* SubSequenceData = SequenceHierarchyCache.FindSubData(*SequenceID))
 				{
 					SequenceTransform = SubSequenceData->RootToSequenceTransform;
 				}
@@ -3256,35 +3329,35 @@ void FUsdLevelSequenceHelperImpl::HandleTrackChange( const UMovieSceneTrack& Tra
 			// Right now we don't write out changes to SkeletalAnimation tracks, and only property tracks... the UAnimSequence
 			// asset can't be modified all that much in UE anyway. Later on we may want to enable writing it out anyway though,
 			// and pick up on changes to the section offset or play rate and bake out the UAnimSequence again
-			if ( const UMovieScenePropertyTrack* PropertyTrack = Cast<const UMovieScenePropertyTrack>( &Track ) )
+			if (const UMovieScenePropertyTrack* PropertyTrack = Cast<const UMovieScenePropertyTrack>(&Track))
 			{
 				TSet<FName> PropertyPathsToRefresh;
-				UnrealToUsd::FPropertyTrackWriter Writer = UnrealToUsd::CreatePropertyTrackWriter( *BoundSceneComponent, *PropertyTrack, UsdPrim, PropertyPathsToRefresh );
+				UnrealToUsd::FPropertyTrackWriter Writer = UnrealToUsd::CreatePropertyTrackWriter(*BoundSceneComponent, *PropertyTrack, UsdPrim, PropertyPathsToRefresh);
 
-				if ( const UMovieSceneFloatTrack* FloatTrack = Cast<const UMovieSceneFloatTrack>( &Track ) )
+				if (const UMovieSceneFloatTrack* FloatTrack = Cast<const UMovieSceneFloatTrack>(&Track))
 				{
-					UnrealToUsd::ConvertFloatTrack( *FloatTrack, SequenceTransform, Writer.FloatWriter, UsdPrim );
+					UnrealToUsd::ConvertFloatTrack(*FloatTrack, SequenceTransform, Writer.FloatWriter, UsdPrim);
 				}
-				else if ( const UMovieSceneBoolTrack* BoolTrack = Cast<const UMovieSceneBoolTrack>( &Track ) )
+				else if (const UMovieSceneBoolTrack* BoolTrack = Cast<const UMovieSceneBoolTrack>(&Track))
 				{
-					UnrealToUsd::ConvertBoolTrack( *BoolTrack, SequenceTransform, Writer.BoolWriter, UsdPrim );
+					UnrealToUsd::ConvertBoolTrack(*BoolTrack, SequenceTransform, Writer.BoolWriter, UsdPrim);
 				}
-				else if ( const UMovieSceneColorTrack* ColorTrack = Cast<const UMovieSceneColorTrack>( &Track ) )
+				else if (const UMovieSceneColorTrack* ColorTrack = Cast<const UMovieSceneColorTrack>(&Track))
 				{
-					UnrealToUsd::ConvertColorTrack( *ColorTrack, SequenceTransform, Writer.ColorWriter, UsdPrim );
+					UnrealToUsd::ConvertColorTrack(*ColorTrack, SequenceTransform, Writer.ColorWriter, UsdPrim);
 				}
-				else if ( const UMovieScene3DTransformTrack* TransformTrack = Cast<const UMovieScene3DTransformTrack>( &Track ) )
+				else if (const UMovieScene3DTransformTrack* TransformTrack = Cast<const UMovieScene3DTransformTrack>(&Track))
 				{
-					UnrealToUsd::Convert3DTransformTrack( *TransformTrack, SequenceTransform, Writer.TransformWriter, UsdPrim );
+					UnrealToUsd::Convert3DTransformTrack(*TransformTrack, SequenceTransform, Writer.TransformWriter, UsdPrim);
 				}
 
 				// Refresh tracks that needed to be updated in USD (e.g. we wrote out a new keyframe to a RectLight's width -> that
 				// should become a new keyframe on our intensity track, because we use the RectLight's width for calculating intensity in UE).
-				if ( PropertyPathsToRefresh.Num() > 0 )
+				if (PropertyPathsToRefresh.Num() > 0)
 				{
 					// For now only our light tracks can request a refresh like this, so we don't even need to check what the refresh
 					// is about: Just resync the light tracks
-					AddLightTracks( *PrimTwin, UsdPrim, PropertyPathsToRefresh );
+					AddLightTracks(*PrimTwin, UsdPrim, PropertyPathsToRefresh);
 					RefreshSequencer();
 				}
 			}
@@ -3292,45 +3365,45 @@ void FUsdLevelSequenceHelperImpl::HandleTrackChange( const UMovieSceneTrack& Tra
 	}
 }
 
-FUsdLevelSequenceHelperImpl::FLayerTimeInfo& FUsdLevelSequenceHelperImpl::FindOrAddLayerTimeInfo( const UE::FSdfLayer& Layer )
+FUsdLevelSequenceHelperImpl::FLayerTimeInfo& FUsdLevelSequenceHelperImpl::FindOrAddLayerTimeInfo(const UE::FSdfLayer& Layer)
 {
-	if ( FLayerTimeInfo* LayerTimeInfo = FindLayerTimeInfo( Layer ) )
+	if (FLayerTimeInfo* LayerTimeInfo = FindLayerTimeInfo(Layer))
 	{
 		return *LayerTimeInfo;
 	}
 
 	FLayerTimeInfo LayerTimeInfo;
-	UpdateLayerTimeInfoFromLayer( LayerTimeInfo, Layer );
+	UpdateLayerTimeInfoFromLayer(LayerTimeInfo, Layer);
 
 	UE_LOG(LogUsd, Verbose, TEXT("Creating layer time info for layer '%s'. Original timecodes: ['%s', '%s']"),
 		*LayerTimeInfo.Identifier,
 		LayerTimeInfo.StartTimeCode.IsSet() ? *LexToString(LayerTimeInfo.StartTimeCode.GetValue()) : TEXT("null"),
 		LayerTimeInfo.EndTimeCode.IsSet() ? *LexToString(LayerTimeInfo.EndTimeCode.GetValue()) : TEXT("null"));
 
-	return LayerTimeInfosByLayerIdentifier.Add( Layer.GetIdentifier(), LayerTimeInfo );
+	return LayerTimeInfosByLayerIdentifier.Add(Layer.GetIdentifier(), LayerTimeInfo);
 }
 
-FUsdLevelSequenceHelperImpl::FLayerTimeInfo* FUsdLevelSequenceHelperImpl::FindLayerTimeInfo( const UE::FSdfLayer& Layer )
+FUsdLevelSequenceHelperImpl::FLayerTimeInfo* FUsdLevelSequenceHelperImpl::FindLayerTimeInfo(const UE::FSdfLayer& Layer)
 {
 	const FString Identifier = Layer.GetIdentifier();
-	return LayerTimeInfosByLayerIdentifier.Find( Identifier );
+	return LayerTimeInfosByLayerIdentifier.Find(Identifier);
 }
 
-void FUsdLevelSequenceHelperImpl::UpdateLayerTimeInfoFromLayer( FLayerTimeInfo& LayerTimeInfo, const UE::FSdfLayer& Layer )
+void FUsdLevelSequenceHelperImpl::UpdateLayerTimeInfoFromLayer(FLayerTimeInfo& LayerTimeInfo, const UE::FSdfLayer& Layer)
 {
-	if ( !Layer )
+	if (!Layer)
 	{
 		return;
 	}
 
-	LayerTimeInfo.Identifier         = Layer.GetIdentifier();
-	LayerTimeInfo.FilePath           = Layer.GetRealPath();
-	LayerTimeInfo.StartTimeCode      = Layer.HasStartTimeCode() ? Layer.GetStartTimeCode() : TOptional<double>();
-	LayerTimeInfo.EndTimeCode        = Layer.HasEndTimeCode() ? Layer.GetEndTimeCode() : TOptional<double>();
+	LayerTimeInfo.Identifier = Layer.GetIdentifier();
+	LayerTimeInfo.FilePath = Layer.GetRealPath();
+	LayerTimeInfo.StartTimeCode = Layer.HasStartTimeCode() ? Layer.GetStartTimeCode() : TOptional<double>();
+	LayerTimeInfo.EndTimeCode = Layer.HasEndTimeCode() ? Layer.GetEndTimeCode() : TOptional<double>();
 
-	if ( LayerTimeInfo.StartTimeCode.IsSet() && LayerTimeInfo.EndTimeCode.IsSet() && LayerTimeInfo.EndTimeCode.GetValue() < LayerTimeInfo.StartTimeCode.GetValue() )
+	if (LayerTimeInfo.StartTimeCode.IsSet() && LayerTimeInfo.EndTimeCode.IsSet() && LayerTimeInfo.EndTimeCode.GetValue() < LayerTimeInfo.StartTimeCode.GetValue())
 	{
-		UE_LOG( LogUsd, Warning, TEXT( "Sublayer '%s' has end time code (%f) before start time code (%f)! These values will be automatically swapped" ),
+		UE_LOG(LogUsd, Warning, TEXT("Sublayer '%s' has end time code (%f) before start time code (%f)! These values will be automatically swapped"),
 			*Layer.GetIdentifier(),
 			LayerTimeInfo.EndTimeCode.GetValue(),
 			LayerTimeInfo.StartTimeCode.GetValue()
@@ -3342,20 +3415,20 @@ void FUsdLevelSequenceHelperImpl::UpdateLayerTimeInfoFromLayer( FLayerTimeInfo&
 	}
 
 	const TArray< FString >& SubLayerPaths = Layer.GetSubLayerPaths();
-	LayerTimeInfo.SubLayersOffsets.Empty( SubLayerPaths.Num() );
+	LayerTimeInfo.SubLayersOffsets.Empty(SubLayerPaths.Num());
 
 	int32 SubLayerIndex = 0;
-	for ( const UE::FSdfLayerOffset& SubLayerOffset : Layer.GetSubLayerOffsets() )
+	for (const UE::FSdfLayerOffset& SubLayerOffset : Layer.GetSubLayerOffsets())
 	{
-		if ( SubLayerPaths.IsValidIndex( SubLayerIndex ) )
+		if (SubLayerPaths.IsValidIndex(SubLayerIndex))
 		{
-			if ( UE::FSdfLayer SubLayer = UsdUtils::FindLayerForSubLayerPath( Layer, SubLayerPaths[ SubLayerIndex ] ) )
+			if (UE::FSdfLayer SubLayer = UsdUtils::FindLayerForSubLayerPath(Layer, SubLayerPaths[SubLayerIndex]))
 			{
 				FLayerOffsetInfo SubLayerOffsetInfo;
 				SubLayerOffsetInfo.LayerIdentifier = SubLayer.GetIdentifier();
 				SubLayerOffsetInfo.LayerOffset = SubLayerOffset;
 
-				LayerTimeInfo.SubLayersOffsets.Add( MoveTemp( SubLayerOffsetInfo ) );
+				LayerTimeInfo.SubLayersOffsets.Add(MoveTemp(SubLayerOffsetInfo));
 			}
 		}
 
@@ -3363,10 +3436,10 @@ void FUsdLevelSequenceHelperImpl::UpdateLayerTimeInfoFromLayer( FLayerTimeInfo&
 	}
 }
 
-ULevelSequence* FUsdLevelSequenceHelperImpl::FindSequenceForIdentifier( const FString& SequenceIdentitifer )
+ULevelSequence* FUsdLevelSequenceHelperImpl::FindSequenceForIdentifier(const FString& SequenceIdentitifer)
 {
 	ULevelSequence* Sequence = nullptr;
-	if ( ULevelSequence** FoundSequence = LevelSequencesByIdentifier.Find( SequenceIdentitifer ) )
+	if (ULevelSequence** FoundSequence = LevelSequencesByIdentifier.Find(SequenceIdentitifer))
 	{
 		Sequence = *FoundSequence;
 	}
@@ -3381,20 +3454,20 @@ public:
 	~FUsdLevelSequenceHelperImpl() {}
 
 	ULevelSequence* Init(const UE::FUsdStage& InUsdStage) { return nullptr; }
-	void SetAssetCache( UUsdAssetCache* AssetCache ) {};
+	void SetAssetCache(UUsdAssetCache* AssetCache) {};
 	bool HasData() const { return false; };
 	void Clear() {};
 
 	void CreateLocalLayersSequences() {}
 
-	void BindToUsdStageActor( AUsdStageActor* InStageActor ) {}
+	void BindToUsdStageActor(AUsdStageActor* InStageActor) {}
 	void UnbindFromUsdStageActor() {}
 	void OnStageActorRenamed() {};
 
-	void AddPrim( UUsdPrimTwin& PrimTwin, bool bForceVisibilityTracks ) {}
+	void AddPrim(UUsdPrimTwin& PrimTwin, bool bForceVisibilityTracks) {}
 	void RemovePrim(const UUsdPrimTwin& PrimTwin) {}
 
-	void UpdateControlRigTracks( UUsdPrimTwin& PrimTwin ) {}
+	void UpdateControlRigTracks(UUsdPrimTwin& PrimTwin) {}
 
 	void StartMonitoringChanges() {}
 	void StopMonitoringChanges() {}
@@ -3412,24 +3485,24 @@ FUsdLevelSequenceHelper::FUsdLevelSequenceHelper()
 
 FUsdLevelSequenceHelper::~FUsdLevelSequenceHelper() = default;
 
-FUsdLevelSequenceHelper::FUsdLevelSequenceHelper(const FUsdLevelSequenceHelper& Other)
+FUsdLevelSequenceHelper::FUsdLevelSequenceHelper(const FUsdLevelSequenceHelper & Other)
 	: FUsdLevelSequenceHelper()
 {
 }
 
-FUsdLevelSequenceHelper& FUsdLevelSequenceHelper::operator=(const FUsdLevelSequenceHelper& Other)
+FUsdLevelSequenceHelper& FUsdLevelSequenceHelper::operator=(const FUsdLevelSequenceHelper & Other)
 {
 	// No copying, start fresh
 	UsdSequencerImpl = MakeUnique<FUsdLevelSequenceHelperImpl>();
 	return *this;
 }
 
-FUsdLevelSequenceHelper::FUsdLevelSequenceHelper(FUsdLevelSequenceHelper&& Other) = default;
-FUsdLevelSequenceHelper& FUsdLevelSequenceHelper::operator=(FUsdLevelSequenceHelper&& Other) = default;
+FUsdLevelSequenceHelper::FUsdLevelSequenceHelper(FUsdLevelSequenceHelper && Other) = default;
+FUsdLevelSequenceHelper& FUsdLevelSequenceHelper::operator=(FUsdLevelSequenceHelper && Other) = default;
 
-ULevelSequence* FUsdLevelSequenceHelper::Init(const UE::FUsdStage& UsdStage)
+ULevelSequence* FUsdLevelSequenceHelper::Init(const UE::FUsdStage & UsdStage)
 {
-	if ( UsdSequencerImpl.IsValid() )
+	if (UsdSequencerImpl.IsValid())
 	{
 		return UsdSequencerImpl->Init(UsdStage);
 	}
@@ -3441,23 +3514,23 @@ ULevelSequence* FUsdLevelSequenceHelper::Init(const UE::FUsdStage& UsdStage)
 
 void FUsdLevelSequenceHelper::OnStageActorRenamed()
 {
-	if ( UsdSequencerImpl.IsValid() )
+	if (UsdSequencerImpl.IsValid())
 	{
 		UsdSequencerImpl->OnStageActorRenamed();
 	}
 }
 
-void FUsdLevelSequenceHelper::SetAssetCache( UUsdAssetCache* AssetCache )
+void FUsdLevelSequenceHelper::SetAssetCache(UUsdAssetCache * AssetCache)
 {
-	if ( UsdSequencerImpl.IsValid() )
+	if (UsdSequencerImpl.IsValid())
 	{
-		UsdSequencerImpl->SetAssetCache( AssetCache );
+		UsdSequencerImpl->SetAssetCache(AssetCache);
 	}
 }
 
 bool FUsdLevelSequenceHelper::HasData() const
 {
-	if ( UsdSequencerImpl.IsValid() )
+	if (UsdSequencerImpl.IsValid())
 	{
 		return UsdSequencerImpl->HasData();
 	}
@@ -3467,13 +3540,13 @@ bool FUsdLevelSequenceHelper::HasData() const
 
 void FUsdLevelSequenceHelper::Clear()
 {
-	if ( UsdSequencerImpl.IsValid() )
+	if (UsdSequencerImpl.IsValid())
 	{
 		UsdSequencerImpl->Clear();
 	}
 }
 
-void FUsdLevelSequenceHelper::BindToUsdStageActor(AUsdStageActor* StageActor)
+void FUsdLevelSequenceHelper::BindToUsdStageActor(AUsdStageActor * StageActor)
 {
 	if (UsdSequencerImpl.IsValid())
 	{
@@ -3489,7 +3562,7 @@ void FUsdLevelSequenceHelper::UnbindFromUsdStageActor()
 	}
 }
 
-void FUsdLevelSequenceHelper::AddPrim(UUsdPrimTwin& PrimTwin, bool bForceVisibilityTracks)
+void FUsdLevelSequenceHelper::AddPrim(UUsdPrimTwin & PrimTwin, bool bForceVisibilityTracks)
 {
 	if (UsdSequencerImpl.IsValid())
 	{
@@ -3497,7 +3570,7 @@ void FUsdLevelSequenceHelper::AddPrim(UUsdPrimTwin& PrimTwin, bool bForceVisibil
 	}
 }
 
-void FUsdLevelSequenceHelper::RemovePrim(const UUsdPrimTwin& PrimTwin)
+void FUsdLevelSequenceHelper::RemovePrim(const UUsdPrimTwin & PrimTwin)
 {
 	if (UsdSequencerImpl.IsValid())
 	{
@@ -3505,11 +3578,11 @@ void FUsdLevelSequenceHelper::RemovePrim(const UUsdPrimTwin& PrimTwin)
 	}
 }
 
-void FUsdLevelSequenceHelper::UpdateControlRigTracks( UUsdPrimTwin& PrimTwin )
+void FUsdLevelSequenceHelper::UpdateControlRigTracks(UUsdPrimTwin & PrimTwin)
 {
-	if ( UsdSequencerImpl.IsValid() )
+	if (UsdSequencerImpl.IsValid())
 	{
-		UsdSequencerImpl->UpdateControlRigTracks( PrimTwin );
+		UsdSequencerImpl->UpdateControlRigTracks(PrimTwin);
 	}
 }
 
@@ -3531,7 +3604,7 @@ void FUsdLevelSequenceHelper::StopMonitoringChanges()
 
 void FUsdLevelSequenceHelper::BlockMonitoringChangesForThisTransaction()
 {
-	if ( UsdSequencerImpl.IsValid() )
+	if (UsdSequencerImpl.IsValid())
 	{
 		UsdSequencerImpl->BlockMonitoringChangesForThisTransaction();
 	}
@@ -3564,7 +3637,7 @@ TArray< ULevelSequence* > FUsdLevelSequenceHelper::GetSubSequences() const
 FUsdLevelSequenceHelper::FOnSkelAnimationBaked& FUsdLevelSequenceHelper::GetOnSkelAnimationBaked()
 {
 #if USE_USD_SDK
-	if ( UsdSequencerImpl.IsValid() )
+	if (UsdSequencerImpl.IsValid())
 	{
 		return UsdSequencerImpl->GetOnSkelAnimationBaked();
 	}
@@ -3576,17 +3649,17 @@ FUsdLevelSequenceHelper::FOnSkelAnimationBaked& FUsdLevelSequenceHelper::GetOnSk
 	}
 }
 
-FScopedBlockMonitoringChangesForTransaction::FScopedBlockMonitoringChangesForTransaction( FUsdLevelSequenceHelper& InHelper )
-	: FScopedBlockMonitoringChangesForTransaction( *InHelper.UsdSequencerImpl.Get() )
+FScopedBlockMonitoringChangesForTransaction::FScopedBlockMonitoringChangesForTransaction(FUsdLevelSequenceHelper & InHelper)
+	: FScopedBlockMonitoringChangesForTransaction(*InHelper.UsdSequencerImpl.Get())
 {
 }
 
-FScopedBlockMonitoringChangesForTransaction::FScopedBlockMonitoringChangesForTransaction( FUsdLevelSequenceHelperImpl& InHelperImpl )
-	: HelperImpl( InHelperImpl )
+FScopedBlockMonitoringChangesForTransaction::FScopedBlockMonitoringChangesForTransaction(FUsdLevelSequenceHelperImpl & InHelperImpl)
+	: HelperImpl(InHelperImpl)
 {
 	// If we're transacting we can just call this and the helper will unblock itself once the transaction is finished, because
 	// we need to make sure the unblocking happens after any call to OnObjectTransacted.
-	if ( GUndo )
+	if (GUndo)
 	{
 		HelperImpl.BlockMonitoringChangesForThisTransaction();
 	}
@@ -3602,7 +3675,7 @@ FScopedBlockMonitoringChangesForTransaction::FScopedBlockMonitoringChangesForTra
 
 FScopedBlockMonitoringChangesForTransaction::~FScopedBlockMonitoringChangesForTransaction()
 {
-	if ( bStoppedMonitoringChanges )
+	if (bStoppedMonitoringChanges)
 	{
 		HelperImpl.StartMonitoringChanges();
 	}
