USD: GeometryCache from USD now use tracks in LevelSequence instead of the Time property to drive their animations. Added USD.DisableGeoCacheTracks cvar to revert to the previous behavior.

diff --git a/Engine/Plugins/Importers/USDImporter/Source/USDStage/Private/USDStageActor.cpp b/Engine/Plugins/Importers/USDImporter/Source/USDStage/Private/USDStageActor.cpp
--- a/Engine/Plugins/Importers/USDImporter/Source/USDStage/Private/USDStageActor.cpp
+++ b/Engine/Plugins/Importers/USDImporter/Source/USDStage/Private/USDStageActor.cpp
@@ -54,6 +54,7 @@
 #include "Engine/World.h"
 #include "EngineAnalytics.h"
 #include "HAL/FileManager.h"
+#include "HAL/IConsoleManager.h"
 #include "LevelSequence.h"
 #include "LiveLinkComponentController.h"
 #include "Materials/Material.h"
@@ -86,9 +87,9 @@
 
 static bool GRegenerateSkeletalAssetsOnControlRigBake = true;
 static FAutoConsoleVariableRef CVarRegenerateSkeletalAssetsOnControlRigBake(
-	TEXT( "USD.RegenerateSkeletalAssetsOnControlRigBake" ),
+	TEXT("USD.RegenerateSkeletalAssetsOnControlRigBake"),
 	GRegenerateSkeletalAssetsOnControlRigBake,
-	TEXT( "Whether to regenerate the assets associated with a SkelRoot (mesh, skeleton, anim sequence, etc.) whenever we modify Control Rig tracks. The USD Stage itself is always updated however." ) );
+	TEXT("Whether to regenerate the assets associated with a SkelRoot (mesh, skeleton, anim sequence, etc.) whenever we modify Control Rig tracks. The USD Stage itself is always updated however."));
 
 static const EObjectFlags DefaultObjFlag = EObjectFlags::RF_Transactional | EObjectFlags::RF_Transient;
 
@@ -96,17 +97,17 @@ AUsdStageActor::FOnActorLoaded AUsdStageActor::OnActorLoaded;
 
 struct FUsdStageActorImpl
 {
-	static TSharedRef< FUsdSchemaTranslationContext > CreateUsdSchemaTranslationContext( AUsdStageActor* StageActor, const FString& PrimPath )
+	static TSharedRef< FUsdSchemaTranslationContext > CreateUsdSchemaTranslationContext(AUsdStageActor* StageActor, const FString& PrimPath)
 	{
 		TSharedRef< FUsdSchemaTranslationContext > TranslationContext = MakeShared< FUsdSchemaTranslationContext >(
 			StageActor->GetOrLoadUsdStage(),
 			*StageActor->AssetCache
-		);
+			);
 
 		TranslationContext->Level = StageActor->GetLevel();
 		TranslationContext->ObjectFlags = DefaultObjFlag;
 		TranslationContext->Time = StageActor->GetTime();
-		TranslationContext->PurposesToLoad = (EUsdPurpose) StageActor->PurposesToLoad;
+		TranslationContext->PurposesToLoad = (EUsdPurpose)StageActor->PurposesToLoad;
 		TranslationContext->NaniteTriangleThreshold = StageActor->NaniteTriangleThreshold;
 		TranslationContext->RenderContext = StageActor->RenderContext;
 		TranslationContext->MaterialPurpose = StageActor->MaterialPurpose;
@@ -121,20 +122,20 @@ struct FUsdStageActorImpl
 		// We parse these even when opening the stage now, as they are used in the skeletal animation tracks
 		TranslationContext->bAllowParsingSkeletalAnimations = true;
 
-		TranslationContext->KindsToCollapse = (EUsdDefaultKind) StageActor->KindsToCollapse;
+		TranslationContext->KindsToCollapse = (EUsdDefaultKind)StageActor->KindsToCollapse;
 		TranslationContext->bMergeIdenticalMaterialSlots = StageActor->bMergeIdenticalMaterialSlots;
 
-		UE::FSdfPath UsdPrimPath( *PrimPath );
-		UUsdPrimTwin* ParentUsdPrimTwin = StageActor->GetRootPrimTwin()->Find( UsdPrimPath.GetParentPath().GetString() );
+		UE::FSdfPath UsdPrimPath(*PrimPath);
+		UUsdPrimTwin* ParentUsdPrimTwin = StageActor->GetRootPrimTwin()->Find(UsdPrimPath.GetParentPath().GetString());
 
-		if ( !ParentUsdPrimTwin )
+		if (!ParentUsdPrimTwin)
 		{
 			ParentUsdPrimTwin = StageActor->RootUsdTwin;
 		}
 
 		TranslationContext->ParentComponent = ParentUsdPrimTwin ? ParentUsdPrimTwin->SceneComponent.Get() : nullptr;
 
-		if ( !TranslationContext->ParentComponent )
+		if (!TranslationContext->ParentComponent)
 		{
 			TranslationContext->ParentComponent = StageActor->RootComponent;
 		}
@@ -143,10 +144,10 @@ struct FUsdStageActorImpl
 	}
 
 	// Workaround some issues where the details panel will crash when showing a property of a component we'll force-delete
-	static void DeselectActorsAndComponents( AUsdStageActor* StageActor )
+	static void DeselectActorsAndComponents(AUsdStageActor* StageActor)
 	{
 #if WITH_EDITOR
-		if ( !StageActor )
+		if (!StageActor)
 		{
 			return;
 		}
@@ -155,55 +156,55 @@ struct FUsdStageActorImpl
 		// Don't do this during garbage collecting if we need to delay-create the root twin (can't NewObject during garbage collection).
 		// If we have no root twin we don't have any tracked spawned actors and components, so we don't need to deselect anything in the first place
 		bool bDeselected = false;
-		if ( GEditor && !IsGarbageCollecting() && StageActor->RootUsdTwin )
+		if (GEditor && !IsGarbageCollecting() && StageActor->RootUsdTwin)
 		{
 			TArray<UObject*> ActorsToDeselect;
 			TArray<UObject*> ComponentsToDeselect;
 
 			const bool bRecursive = true;
-			StageActor->GetRootPrimTwin()->Iterate( [&ActorsToDeselect, &ComponentsToDeselect]( UUsdPrimTwin& PrimTwin )
-			{
-				if ( AActor* ReferencedActor = PrimTwin.SpawnedActor.Get() )
-				{
-					ActorsToDeselect.Add( ReferencedActor );
-				}
-				if ( USceneComponent* ReferencedComponent = PrimTwin.SceneComponent.Get() )
+			StageActor->GetRootPrimTwin()->Iterate([&ActorsToDeselect, &ComponentsToDeselect](UUsdPrimTwin& PrimTwin)
 				{
-					ComponentsToDeselect.Add( ReferencedComponent );
-
-					AActor* Owner = ReferencedComponent->GetOwner();
-					if ( Owner && Owner->GetRootComponent() == ReferencedComponent )
+					if (AActor* ReferencedActor = PrimTwin.SpawnedActor.Get())
 					{
-						ActorsToDeselect.Add( Owner );
+						ActorsToDeselect.Add(ReferencedActor);
 					}
-				}
-			}, bRecursive );
+					if (USceneComponent* ReferencedComponent = PrimTwin.SceneComponent.Get())
+					{
+						ComponentsToDeselect.Add(ReferencedComponent);
 
-			if ( USelection* SelectedComponents = GEditor->GetSelectedComponents() )
+						AActor* Owner = ReferencedComponent->GetOwner();
+						if (Owner && Owner->GetRootComponent() == ReferencedComponent)
+						{
+							ActorsToDeselect.Add(Owner);
+						}
+					}
+				}, bRecursive);
+
+			if (USelection* SelectedComponents = GEditor->GetSelectedComponents())
 			{
-				for ( UObject* Component : ComponentsToDeselect )
+				for (UObject* Component : ComponentsToDeselect)
 				{
-					if ( SelectedComponents->IsSelected( Component ) )
+					if (SelectedComponents->IsSelected(Component))
 					{
-						SelectedComponents->Deselect( Component );
+						SelectedComponents->Deselect(Component);
 						bDeselected = true;
 					}
 				}
 			}
 
-			if ( USelection* SelectedActors = GEditor->GetSelectedActors() )
+			if (USelection* SelectedActors = GEditor->GetSelectedActors())
 			{
-				for ( UObject* Actor : ActorsToDeselect )
+				for (UObject* Actor : ActorsToDeselect)
 				{
-					if ( SelectedActors->IsSelected( Actor ) )
+					if (SelectedActors->IsSelected(Actor))
 					{
-						SelectedActors->Deselect( Actor );
+						SelectedActors->Deselect(Actor);
 						bDeselected = true;
 					}
 				}
 			}
 
-			if ( bDeselected && GIsEditor ) // Make sure we're not in standalone either
+			if (bDeselected && GIsEditor) // Make sure we're not in standalone either
 			{
 				GEditor->NoteSelectionChange();
 			}
@@ -212,18 +213,18 @@ struct FUsdStageActorImpl
 	}
 
 	template<typename ObjectPtr>
-	static void CloseEditorsForAssets( const TMap< FString, ObjectPtr >& AssetsCache )
+	static void CloseEditorsForAssets(const TMap< FString, ObjectPtr >& AssetsCache)
 	{
 #if WITH_EDITOR
-		if ( GIsEditor && GEditor )
+		if (GIsEditor && GEditor)
 		{
-			if ( UAssetEditorSubsystem* AssetEditorSubsysttem = GEditor->GetEditorSubsystem<UAssetEditorSubsystem>() )
+			if (UAssetEditorSubsystem* AssetEditorSubsysttem = GEditor->GetEditorSubsystem<UAssetEditorSubsystem>())
 			{
-				for ( const TPair<FString, ObjectPtr>& Pair : AssetsCache )
+				for (const TPair<FString, ObjectPtr>& Pair : AssetsCache)
 				{
-					if ( UObject* Asset = Pair.Value )
+					if (UObject* Asset = Pair.Value)
 					{
-						AssetEditorSubsysttem->CloseAllEditorsForAsset( Asset );
+						AssetEditorSubsysttem->CloseAllEditorsForAsset(Asset);
 					}
 				}
 			}
@@ -231,15 +232,15 @@ struct FUsdStageActorImpl
 #endif // WITH_EDITOR
 	}
 
-	static void DiscardStage( const UE::FUsdStage& Stage, AUsdStageActor* DiscardingActor )
+	static void DiscardStage(const UE::FUsdStage& Stage, AUsdStageActor* DiscardingActor)
 	{
-		if ( !Stage || !DiscardingActor )
+		if (!Stage || !DiscardingActor)
 		{
 			return;
 		}
 
 		UE::FSdfLayer RootLayer = Stage.GetRootLayer();
-		if ( RootLayer && RootLayer.IsAnonymous() )
+		if (RootLayer && RootLayer.IsAnonymous())
 		{
 			// Erasing an anonymous stage would fully delete it. If we later undo/redo into a path that referenced
 			// one of those anonymous layers, we wouldn't be able to load it back again.
@@ -250,19 +251,19 @@ struct FUsdStageActorImpl
 		}
 
 		TArray<UObject*> Instances;
-		AUsdStageActor::StaticClass()->GetDefaultObject()->GetArchetypeInstances( Instances );
-		for ( UObject* Instance : Instances )
+		AUsdStageActor::StaticClass()->GetDefaultObject()->GetArchetypeInstances(Instances);
+		for (UObject* Instance : Instances)
 		{
-			if ( Instance == DiscardingActor || !Instance || !IsValidChecked(Instance) || Instance->IsTemplate() )
+			if (Instance == DiscardingActor || !Instance || !IsValidChecked(Instance) || Instance->IsTemplate())
 			{
 				continue;
 			}
 
 			// Need to use the const version here or we may inadvertently load the stage
-			if ( const AUsdStageActor* Actor = Cast<const AUsdStageActor>( Instance ) )
+			if (const AUsdStageActor* Actor = Cast<const AUsdStageActor>(Instance))
 			{
 				const UE::FUsdStage& OtherStage = Actor->GetUsdStage();
-				if ( OtherStage && Stage == OtherStage )
+				if (OtherStage && Stage == OtherStage)
 				{
 					// Some other actor is still using our stage, so don't close it
 					return;
@@ -270,28 +271,28 @@ struct FUsdStageActorImpl
 			}
 		}
 
-		UnrealUSDWrapper::EraseStageFromCache( Stage );
+		UnrealUSDWrapper::EraseStageFromCache(Stage);
 	}
 
-	static bool ObjectNeedsMultiUserTag( UObject* Object, AUsdStageActor* StageActor )
+	static bool ObjectNeedsMultiUserTag(UObject* Object, AUsdStageActor* StageActor)
 	{
 		// Don't need to tag non-transient stuff
-		if ( !Object->HasAnyFlags( RF_Transient ) )
+		if (!Object->HasAnyFlags(RF_Transient))
 		{
 			return false;
 		}
 
 		// Object already has tag
-		if ( AActor* Actor = Cast<AActor>( Object ) )
+		if (AActor* Actor = Cast<AActor>(Object))
 		{
-			if ( Actor->Tags.Contains( UE::UsdTransactor::ConcertSyncEnableTag ) )
+			if (Actor->Tags.Contains(UE::UsdTransactor::ConcertSyncEnableTag))
 			{
 				return false;
 			}
 		}
-		else if ( USceneComponent* Component = Cast<USceneComponent>( Object ) )
+		else if (USceneComponent* Component = Cast<USceneComponent>(Object))
 		{
-			if ( Component->ComponentTags.Contains( UE::UsdTransactor::ConcertSyncEnableTag ) )
+			if (Component->ComponentTags.Contains(UE::UsdTransactor::ConcertSyncEnableTag))
 			{
 				return false;
 			}
@@ -299,25 +300,25 @@ struct FUsdStageActorImpl
 
 		// Only care about objects that the same actor spawned
 		bool bOwnedByStageActor = false;
-		if ( StageActor->ObjectsToWatch.Contains( Object ) )
+		if (StageActor->ObjectsToWatch.Contains(Object))
 		{
 			bOwnedByStageActor = true;
 		}
-		if ( AActor* Actor = Cast<AActor>( Object ) )
+		if (AActor* Actor = Cast<AActor>(Object))
 		{
-			if ( StageActor->ObjectsToWatch.Contains( Actor->GetRootComponent() ) )
+			if (StageActor->ObjectsToWatch.Contains(Actor->GetRootComponent()))
 			{
 				bOwnedByStageActor = true;
 			}
 		}
-		else if ( AActor* Outer = Object->GetTypedOuter<AActor>() )
+		else if (AActor* Outer = Object->GetTypedOuter<AActor>())
 		{
-			if ( StageActor->ObjectsToWatch.Contains( Outer->GetRootComponent() ) )
+			if (StageActor->ObjectsToWatch.Contains(Outer->GetRootComponent()))
 			{
 				bOwnedByStageActor = true;
 			}
 		}
-		if ( !bOwnedByStageActor )
+		if (!bOwnedByStageActor)
 		{
 			return false;
 		}
@@ -325,64 +326,64 @@ struct FUsdStageActorImpl
 		return bOwnedByStageActor;
 	}
 
-	static void AllowListComponentHierarchy( USceneComponent* Component, TSet<UObject*>& VisitedObjects )
+	static void AllowListComponentHierarchy(USceneComponent* Component, TSet<UObject*>& VisitedObjects)
 	{
-		if ( !Component || VisitedObjects.Contains( Component ) )
+		if (!Component || VisitedObjects.Contains(Component))
 		{
 			return;
 		}
 
-		VisitedObjects.Add( Component );
+		VisitedObjects.Add(Component);
 
-		if ( Component->HasAnyFlags( RF_Transient ) )
+		if (Component->HasAnyFlags(RF_Transient))
 		{
-			Component->ComponentTags.AddUnique( UE::UsdTransactor::ConcertSyncEnableTag );
+			Component->ComponentTags.AddUnique(UE::UsdTransactor::ConcertSyncEnableTag);
 		}
 
-		if ( AActor* Owner = Component->GetOwner() )
+		if (AActor* Owner = Component->GetOwner())
 		{
-			if ( !VisitedObjects.Contains( Owner ) && Owner->HasAnyFlags( RF_Transient ) )
+			if (!VisitedObjects.Contains(Owner) && Owner->HasAnyFlags(RF_Transient))
 			{
-				Owner->Tags.AddUnique( UE::UsdTransactor::ConcertSyncEnableTag );
+				Owner->Tags.AddUnique(UE::UsdTransactor::ConcertSyncEnableTag);
 			}
 
-			VisitedObjects.Add( Owner );
+			VisitedObjects.Add(Owner);
 		}
 
 		// Iterate the attachment hierarchy directly because maybe some of those actors have additional components that aren't being
 		// tracked by a prim twin
-		for ( USceneComponent* Child : Component->GetAttachChildren() )
+		for (USceneComponent* Child : Component->GetAttachChildren())
 		{
-			AllowListComponentHierarchy( Child, VisitedObjects );
+			AllowListComponentHierarchy(Child, VisitedObjects);
 		}
 	}
 
 	// Checks if a project-relative file path refers to a layer. It requires caution because anonymous layers need to be handled differently.
 	// WARNING: This will break if FilePath is a relative path relative to anything else other than the Project directory (i.e. engine binary)
-	static bool DoesPathPointToLayer( FString FilePath, const UE::FSdfLayer& Layer )
+	static bool DoesPathPointToLayer(FString FilePath, const UE::FSdfLayer& Layer)
 	{
 #if USE_USD_SDK
-		if ( !Layer )
+		if (!Layer)
 		{
 			return false;
 		}
 
-		if ( !FilePath.IsEmpty() && !FPaths::IsRelative( FilePath ) && !FilePath.StartsWith( UnrealIdentifiers::IdentifierPrefix ) )
+		if (!FilePath.IsEmpty() && !FPaths::IsRelative(FilePath) && !FilePath.StartsWith(UnrealIdentifiers::IdentifierPrefix))
 		{
-			FilePath = UsdUtils::MakePathRelativeToProjectDir( FilePath );
+			FilePath = UsdUtils::MakePathRelativeToProjectDir(FilePath);
 		}
 
 		// Special handling for anonymous layers as the RealPath is empty
-		if ( Layer.IsAnonymous() )
+		if (Layer.IsAnonymous())
 		{
 			// Something like "anon:0000022F9E194D50:tmp.usda"
 			const FString LayerIdentifier = Layer.GetIdentifier();
 
 			// Something like "@identifier:anon:0000022F9E194D50:tmp.usda" if we're also pointing at an anonymous layer
-			if ( FilePath.RemoveFromStart( UnrealIdentifiers::IdentifierPrefix ) )
+			if (FilePath.RemoveFromStart(UnrealIdentifiers::IdentifierPrefix))
 			{
 				// Same anonymous layers
-				if ( FilePath == LayerIdentifier )
+				if (FilePath == LayerIdentifier)
 				{
 					return true;
 				}
@@ -395,7 +396,7 @@ struct FUsdStageActorImpl
 		}
 		else
 		{
-			return FPaths::IsSamePath( UsdUtils::MakePathRelativeToProjectDir( Layer.GetRealPath() ), FilePath );
+			return FPaths::IsSamePath(UsdUtils::MakePathRelativeToProjectDir(Layer.GetRealPath()), FilePath);
 		}
 #endif // USE_USD_SDK
 
@@ -408,16 +409,16 @@ struct FUsdStageActorImpl
 	 * correct value of bHiddenInGame.
 	 * Note that bHiddenInGame corresponds to computed visibility, and the component tags correspond to individual prim-level visibilities
 	 */
-	static void MakeVisible( UUsdPrimTwin& UsdPrimTwin, const UE::FUsdStage& Stage )
+	static void MakeVisible(UUsdPrimTwin& UsdPrimTwin, const UE::FUsdStage& Stage)
 	{
 		// Find the highest invisible prim parent: Nothing above this can possibly change with what we're doing
 		UUsdPrimTwin* Iter = &UsdPrimTwin;
 		UUsdPrimTwin* HighestInvisibleParent = nullptr;
-		while ( Iter )
+		while (Iter)
 		{
-			if ( USceneComponent* Component = Iter->GetSceneComponent() )
+			if (USceneComponent* Component = Iter->GetSceneComponent())
 			{
-				if ( Component->ComponentTags.Contains( UnrealIdentifiers::Invisible ) )
+				if (Component->ComponentTags.Contains(UnrealIdentifiers::Invisible))
 				{
 					HighestInvisibleParent = Iter;
 				}
@@ -427,80 +428,80 @@ struct FUsdStageActorImpl
 		}
 
 		// No parent (not even UsdPrimTwin's prim directly) was invisible, so we should already be visible and there's nothing to do
-		if ( !HighestInvisibleParent )
+		if (!HighestInvisibleParent)
 		{
 			return;
 		}
 
-		UE::FUsdPrim Prim = Stage.GetPrimAtPath( UE::FSdfPath( *UsdPrimTwin.PrimPath ) );
-		if ( !Prim )
+		UE::FUsdPrim Prim = Stage.GetPrimAtPath(UE::FSdfPath(*UsdPrimTwin.PrimPath));
+		if (!Prim)
 		{
 			return;
 		}
-		UsdUtils::MakeVisible( Prim );
+		UsdUtils::MakeVisible(Prim);
 
 		TFunction<void(UUsdPrimTwin&, bool)> RecursiveResyncVisibility;
-		RecursiveResyncVisibility = [ &Stage, &RecursiveResyncVisibility ]( UUsdPrimTwin& PrimTwin, bool bPrimHasInvisibleParent )
+		RecursiveResyncVisibility = [&Stage, &RecursiveResyncVisibility](UUsdPrimTwin& PrimTwin, bool bPrimHasInvisibleParent)
 		{
 			USceneComponent* Component = PrimTwin.GetSceneComponent();
-			if ( !Component )
+			if (!Component)
 			{
 				return;
 			}
 
-			UE::FUsdPrim CurrentPrim = Stage.GetPrimAtPath( UE::FSdfPath( *PrimTwin.PrimPath ) );
-			if ( !CurrentPrim )
+			UE::FUsdPrim CurrentPrim = Stage.GetPrimAtPath(UE::FSdfPath(*PrimTwin.PrimPath));
+			if (!CurrentPrim)
 			{
 				return;
 			}
 
-			const bool bPrimHasInheritedVisibility = UsdUtils::HasInheritedVisibility( CurrentPrim );
+			const bool bPrimHasInheritedVisibility = UsdUtils::HasInheritedVisibility(CurrentPrim);
 			const bool bPrimIsVisible = bPrimHasInheritedVisibility && !bPrimHasInvisibleParent;
 
-			const bool bComponentHasInvisibleTag = Component->ComponentTags.Contains( UnrealIdentifiers::Invisible );
+			const bool bComponentHasInvisibleTag = Component->ComponentTags.Contains(UnrealIdentifiers::Invisible);
 			const bool bComponentIsVisible = !Component->bHiddenInGame;
 
 			const bool bTagIsCorrect = bComponentHasInvisibleTag == !bPrimHasInheritedVisibility;
 			const bool bComputedVisibilityIsCorrect = bPrimIsVisible == bComponentIsVisible;
 
 			// Stop recursing as this prim's or its children couldn't possibly need to be updated
-			if ( bTagIsCorrect && bComputedVisibilityIsCorrect )
+			if (bTagIsCorrect && bComputedVisibilityIsCorrect)
 			{
 				return;
 			}
 
-			if ( !bTagIsCorrect )
+			if (!bTagIsCorrect)
 			{
-				if ( bPrimHasInheritedVisibility )
+				if (bPrimHasInheritedVisibility)
 				{
-					Component->ComponentTags.Remove( UnrealIdentifiers::Invisible );
-					Component->ComponentTags.AddUnique( UnrealIdentifiers::Inherited );
+					Component->ComponentTags.Remove(UnrealIdentifiers::Invisible);
+					Component->ComponentTags.AddUnique(UnrealIdentifiers::Inherited);
 				}
 				else
 				{
-					Component->ComponentTags.AddUnique( UnrealIdentifiers::Invisible );
-					Component->ComponentTags.Remove( UnrealIdentifiers::Inherited );
+					Component->ComponentTags.AddUnique(UnrealIdentifiers::Invisible);
+					Component->ComponentTags.Remove(UnrealIdentifiers::Inherited);
 				}
 			}
 
-			if ( !bComputedVisibilityIsCorrect )
+			if (!bComputedVisibilityIsCorrect)
 			{
 				const bool bPropagateToChildren = false;
 				Component->Modify();
-				Component->SetHiddenInGame( !bPrimIsVisible, bPropagateToChildren );
+				Component->SetHiddenInGame(!bPrimIsVisible, bPropagateToChildren);
 			}
 
-			for ( const TPair<FString, TObjectPtr<UUsdPrimTwin>>& ChildPair : PrimTwin.GetChildren() )
+			for (const TPair<FString, TObjectPtr<UUsdPrimTwin>>& ChildPair : PrimTwin.GetChildren())
 			{
-				if ( UUsdPrimTwin* ChildTwin = ChildPair.Value )
+				if (UUsdPrimTwin* ChildTwin = ChildPair.Value)
 				{
-					RecursiveResyncVisibility( *ChildTwin, !bPrimIsVisible );
+					RecursiveResyncVisibility(*ChildTwin, !bPrimIsVisible);
 				}
 			}
 		};
 
 		const bool bHasInvisibleParent = false;
-		RecursiveResyncVisibility( *HighestInvisibleParent, bHasInvisibleParent );
+		RecursiveResyncVisibility(*HighestInvisibleParent, bHasInvisibleParent);
 	}
 
 	/**
@@ -508,45 +509,45 @@ struct FUsdStageActorImpl
 	 * to bHiddenInGame = false. Leave their individual prim-level visibilities intact though.
 	 * Note that bHiddenInGame corresponds to computed visibility, and the component tags correspond to individual prim-level visibilities
 	 */
-	static void MakeInvisible( UUsdPrimTwin& UsdPrimTwin )
+	static void MakeInvisible(UUsdPrimTwin& UsdPrimTwin)
 	{
 		USceneComponent* PrimSceneComponent = UsdPrimTwin.GetSceneComponent();
-		if ( !PrimSceneComponent )
+		if (!PrimSceneComponent)
 		{
 			return;
 		}
 
-		PrimSceneComponent->ComponentTags.AddUnique( UnrealIdentifiers::Invisible );
-		PrimSceneComponent->ComponentTags.Remove( UnrealIdentifiers::Inherited );
+		PrimSceneComponent->ComponentTags.AddUnique(UnrealIdentifiers::Invisible);
+		PrimSceneComponent->ComponentTags.Remove(UnrealIdentifiers::Inherited);
 
 		const bool bPropagateToChildren = true;
 		const bool bNewHidden = true;
-		PrimSceneComponent->SetHiddenInGame( bNewHidden, bPropagateToChildren );
+		PrimSceneComponent->SetHiddenInGame(bNewHidden, bPropagateToChildren);
 	}
 
-	static void SendAnalytics( AUsdStageActor* StageActor, double ElapsedSeconds, double NumberOfFrames, const FString& Extension )
+	static void SendAnalytics(AUsdStageActor* StageActor, double ElapsedSeconds, double NumberOfFrames, const FString& Extension)
 	{
-		if ( !StageActor )
+		if (!StageActor)
 		{
 			return;
 		}
 
-		if ( FEngineAnalytics::IsAvailable() )
+		if (FEngineAnalytics::IsAvailable())
 		{
 			TArray<FAnalyticsEventAttribute> EventAttributes;
 
-			EventAttributes.Emplace( TEXT( "InitialLoadSet" ), LexToString( (uint8)StageActor->InitialLoadSet ) );
-			EventAttributes.Emplace( TEXT( "InterpolationType" ), LexToString( (uint8)StageActor->InterpolationType) );
-			EventAttributes.Emplace( TEXT( "KindsToCollapse" ), LexToString( StageActor->KindsToCollapse ) );
-			EventAttributes.Emplace( TEXT( "MergeIdenticalMaterialSlots" ), LexToString( StageActor->bMergeIdenticalMaterialSlots ) );
-			EventAttributes.Emplace( TEXT( "PurposesToLoad" ), LexToString( StageActor->PurposesToLoad ) );
-			EventAttributes.Emplace( TEXT( "NaniteTriangleThreshold" ), LexToString( StageActor->NaniteTriangleThreshold ) );
-			EventAttributes.Emplace( TEXT( "RenderContext" ), StageActor->RenderContext.ToString() );
-			EventAttributes.Emplace( TEXT( "MaterialPurpose" ), StageActor->MaterialPurpose.ToString() );
-			EventAttributes.Emplace( TEXT( "RootMotionHandling" ), LexToString( ( uint8 ) StageActor->RootMotionHandling ) );
+			EventAttributes.Emplace(TEXT("InitialLoadSet"), LexToString((uint8)StageActor->InitialLoadSet));
+			EventAttributes.Emplace(TEXT("InterpolationType"), LexToString((uint8)StageActor->InterpolationType));
+			EventAttributes.Emplace(TEXT("KindsToCollapse"), LexToString(StageActor->KindsToCollapse));
+			EventAttributes.Emplace(TEXT("MergeIdenticalMaterialSlots"), LexToString(StageActor->bMergeIdenticalMaterialSlots));
+			EventAttributes.Emplace(TEXT("PurposesToLoad"), LexToString(StageActor->PurposesToLoad));
+			EventAttributes.Emplace(TEXT("NaniteTriangleThreshold"), LexToString(StageActor->NaniteTriangleThreshold));
+			EventAttributes.Emplace(TEXT("RenderContext"), StageActor->RenderContext.ToString());
+			EventAttributes.Emplace(TEXT("MaterialPurpose"), StageActor->MaterialPurpose.ToString());
+			EventAttributes.Emplace(TEXT("RootMotionHandling"), LexToString((uint8)StageActor->RootMotionHandling));
 
 			const bool bAutomated = false;
-			IUsdClassesModule::SendAnalytics( MoveTemp( EventAttributes ), TEXT( "Open" ), bAutomated, ElapsedSeconds, NumberOfFrames, Extension );
+			IUsdClassesModule::SendAnalytics(MoveTemp(EventAttributes), TEXT("Open"), bAutomated, ElapsedSeconds, NumberOfFrames, Extension);
 		}
 	}
 };
@@ -568,47 +569,47 @@ struct FRecompilationTracker
 {
 	static void SetupEvents()
 	{
-		if ( bEventIsSetup || !GIsEditor || !GEditor )
+		if (bEventIsSetup || !GIsEditor || !GEditor)
 		{
 			return;
 		}
 
-		GEditor->OnBlueprintPreCompile().AddStatic( &FRecompilationTracker::OnCompilationStarted );
+		GEditor->OnBlueprintPreCompile().AddStatic(&FRecompilationTracker::OnCompilationStarted);
 		bEventIsSetup = true;
 	}
 
-	static bool IsBeingCompiled( UBlueprint* BP )
+	static bool IsBeingCompiled(UBlueprint* BP)
 	{
-		return FRecompilationTracker::RecompilingBlueprints.Contains( BP );
+		return FRecompilationTracker::RecompilingBlueprints.Contains(BP);
 	}
 
-	static void OnCompilationStarted( UBlueprint* BP )
+	static void OnCompilationStarted(UBlueprint* BP)
 	{
 		// We don't care if a BP is compiling on first load: It only matters to use if we're compiling one that already has loaded instances on the level
-		if ( !BP ||
-			 BP->bIsRegeneratingOnLoad ||
-			 !BP->GeneratedClass ||
-			 !BP->GeneratedClass->IsChildOf( AUsdStageActor::StaticClass() ) ||
-			 RecompilingBlueprints.Contains( BP ) )
+		if (!BP ||
+			BP->bIsRegeneratingOnLoad ||
+			!BP->GeneratedClass ||
+			!BP->GeneratedClass->IsChildOf(AUsdStageActor::StaticClass()) ||
+			RecompilingBlueprints.Contains(BP))
 		{
 			return;
 		}
 
-		FDelegateHandle Handle = BP->OnCompiled().AddStatic( &FRecompilationTracker::OnCompilationEnded );
-		FRecompilationTracker::RecompilingBlueprints.Add( BP, Handle );
+		FDelegateHandle Handle = BP->OnCompiled().AddStatic(&FRecompilationTracker::OnCompilationEnded);
+		FRecompilationTracker::RecompilingBlueprints.Add(BP, Handle);
 	}
 
-	static void OnCompilationEnded( UBlueprint* BP )
+	static void OnCompilationEnded(UBlueprint* BP)
 	{
-		if ( !BP )
+		if (!BP)
 		{
 			return;
 		}
 
 		FDelegateHandle RemovedHandle;
-		if ( FRecompilationTracker::RecompilingBlueprints.RemoveAndCopyValue( BP, RemovedHandle ) )
+		if (FRecompilationTracker::RecompilingBlueprints.RemoveAndCopyValue(BP, RemovedHandle))
 		{
-			BP->OnCompiled().Remove( RemovedHandle );
+			BP->OnCompiled().Remove(RemovedHandle);
 		}
 	}
 
@@ -621,25 +622,25 @@ TMap<UBlueprint*, FDelegateHandle> FRecompilationTracker::RecompilingBlueprints;
 #endif // WITH_EDITOR
 
 AUsdStageActor::AUsdStageActor()
-	: InitialLoadSet( EUsdInitialLoadSet::LoadAll )
-	, InterpolationType( EUsdInterpolationType::Linear )
-	, KindsToCollapse( ( int32 ) ( EUsdDefaultKind::Component | EUsdDefaultKind::Subcomponent ) )
-	, bMergeIdenticalMaterialSlots( true )
-	, PurposesToLoad( (int32) EUsdPurpose::Proxy )
-	, NaniteTriangleThreshold( (uint64) 1000000 )
-	, MaterialPurpose( *UnrealIdentifiers::MaterialAllPurpose )
-	, RootMotionHandling( EUsdRootMotionHandling::NoAdditionalRootMotion )
-	, Time( 0.0f )
-	, bIsTransitioningIntoPIE( false )
-	, bIsModifyingAProperty( false )
-	, bIsUndoRedoing( false )
-{
-	SceneComponent = CreateDefaultSubobject< USceneComponent >( TEXT("SceneComponent0") );
+	: InitialLoadSet(EUsdInitialLoadSet::LoadAll)
+	, InterpolationType(EUsdInterpolationType::Linear)
+	, KindsToCollapse((int32)(EUsdDefaultKind::Component | EUsdDefaultKind::Subcomponent))
+	, bMergeIdenticalMaterialSlots(true)
+	, PurposesToLoad((int32)EUsdPurpose::Proxy)
+	, NaniteTriangleThreshold((uint64)1000000)
+	, MaterialPurpose(*UnrealIdentifiers::MaterialAllPurpose)
+	, RootMotionHandling(EUsdRootMotionHandling::NoAdditionalRootMotion)
+	, Time(0.0f)
+	, bIsTransitioningIntoPIE(false)
+	, bIsModifyingAProperty(false)
+	, bIsUndoRedoing(false)
+{
+	SceneComponent = CreateDefaultSubobject< USceneComponent >(TEXT("SceneComponent0"));
 	SceneComponent->Mobility = EComponentMobility::Static;
 
 	RootComponent = SceneComponent;
 
-	AssetCache = CreateDefaultSubobject< UUsdAssetCache >( TEXT("AssetCache") );
+	AssetCache = CreateDefaultSubobject< UUsdAssetCache >(TEXT("AssetCache"));
 
 	// Note: We can't construct our RootUsdTwin as a default subobject here, it needs to be built on-demand.
 	// Even if we NewObject'd one it will work as a subobject in some contexts (maybe because the CDO will have a dedicated root twin?).
@@ -648,26 +649,26 @@ AUsdStageActor::AUsdStageActor()
 	// some transitions like reinstantiation.
 	// c.f. doc comment on FRecompilationTracker for more info.
 
-	IUsdSchemasModule& UsdSchemasModule = FModuleManager::Get().LoadModuleChecked< IUsdSchemasModule >( TEXT("USDSchemas") );
+	IUsdSchemasModule& UsdSchemasModule = FModuleManager::Get().LoadModuleChecked< IUsdSchemasModule >(TEXT("USDSchemas"));
 	RenderContext = UsdSchemasModule.GetRenderContextRegistry().GetUnrealRenderContext();
 
-	Transactor = NewObject<UUsdTransactor>( this, TEXT( "Transactor" ), EObjectFlags::RF_Transactional );
-	Transactor->Initialize( this );
+	Transactor = NewObject<UUsdTransactor>(this, TEXT("Transactor"), EObjectFlags::RF_Transactional);
+	Transactor->Initialize(this);
 
-	if ( HasAuthorityOverStage() )
+	if (HasAuthorityOverStage())
 	{
 #if WITH_EDITOR
 		// Update the supported filetypes in our RootPath property
-		for ( TFieldIterator<FProperty> PropertyIterator( AUsdStageActor::StaticClass() ); PropertyIterator; ++PropertyIterator )
+		for (TFieldIterator<FProperty> PropertyIterator(AUsdStageActor::StaticClass()); PropertyIterator; ++PropertyIterator)
 		{
 			FProperty* Property = *PropertyIterator;
-			if ( Property && Property->GetFName() == GET_MEMBER_NAME_CHECKED( AUsdStageActor, RootLayer ) )
+			if (Property && Property->GetFName() == GET_MEMBER_NAME_CHECKED(AUsdStageActor, RootLayer))
 			{
 				TArray< FString > SupportedExtensions = UnrealUSDWrapper::GetAllSupportedFileFormats();
-				if ( SupportedExtensions.Num() > 0 )
+				if (SupportedExtensions.Num() > 0)
 				{
-					FString JoinedExtensions = FString::Join( SupportedExtensions, TEXT( "; *." ) ); // Combine "usd" and "usda" into "usd; *.usda"
-					Property->SetMetaData( TEXT("FilePathFilter"), FString::Printf( TEXT( "Universal Scene Description files|*.%s" ), *JoinedExtensions, *JoinedExtensions ) );
+					FString JoinedExtensions = FString::Join(SupportedExtensions, TEXT("; *.")); // Combine "usd" and "usda" into "usd; *.usda"
+					Property->SetMetaData(TEXT("FilePathFilter"), FString::Printf(TEXT("Universal Scene Description files|*.%s"), *JoinedExtensions, *JoinedExtensions));
 				}
 				break;
 			}
@@ -676,85 +677,85 @@ AUsdStageActor::AUsdStageActor()
 		FEditorDelegates::BeginPIE.AddUObject(this, &AUsdStageActor::OnBeginPIE);
 		FEditorDelegates::PostPIEStarted.AddUObject(this, &AUsdStageActor::OnPostPIEStarted);
 
-		FUsdDelegates::OnPostUsdImport.AddUObject( this, &AUsdStageActor::OnPostUsdImport );
-		FUsdDelegates::OnPreUsdImport.AddUObject( this, &AUsdStageActor::OnPreUsdImport );
+		FUsdDelegates::OnPostUsdImport.AddUObject(this, &AUsdStageActor::OnPostUsdImport);
+		FUsdDelegates::OnPreUsdImport.AddUObject(this, &AUsdStageActor::OnPreUsdImport);
 
-		GEngine->OnLevelActorDeleted().AddUObject( this, &AUsdStageActor::OnLevelActorDeleted );
+		GEngine->OnLevelActorDeleted().AddUObject(this, &AUsdStageActor::OnLevelActorDeleted);
 
 		// When another client of a multi-user session modifies their version of this actor, the transaction will be replicated here.
 		// The multi-user system uses "redo" to apply those transactions, so this is our best chance to respond to events as e.g. neither
 		// PostTransacted nor Destroyed get called when the other user deletes the actor
-		if ( UTransBuffer* TransBuffer = GUnrealEd ? Cast<UTransBuffer>( GUnrealEd->Trans ) : nullptr )
+		if (UTransBuffer* TransBuffer = GUnrealEd ? Cast<UTransBuffer>(GUnrealEd->Trans) : nullptr)
 		{
-			TransBuffer->OnTransactionStateChanged().AddUObject( this, &AUsdStageActor::HandleTransactionStateChanged );
+			TransBuffer->OnTransactionStateChanged().AddUObject(this, &AUsdStageActor::HandleTransactionStateChanged);
 
 			// We can't use AddUObject here as we may specifically want to respond *after* we're marked as pending kill
 			OnRedoHandle = TransBuffer->OnRedo().AddLambda(
-				[ this ]( const FTransactionContext& TransactionContext, bool bSucceeded )
+				[this](const FTransactionContext& TransactionContext, bool bSucceeded)
 				{
 					// This text should match the one in ConcertClientTransactionBridge.cpp
-					if ( this &&
-						 HasAuthorityOverStage() &&
-						 TransactionContext.Title.EqualTo( LOCTEXT( "ConcertTransactionEvent", "Concert Transaction Event" ) ) &&
-						 !RootLayer.FilePath.IsEmpty() )
+					if (this &&
+						HasAuthorityOverStage() &&
+						TransactionContext.Title.EqualTo(LOCTEXT("ConcertTransactionEvent", "Concert Transaction Event")) &&
+						!RootLayer.FilePath.IsEmpty())
 					{
 						// Other user deleted us
-						if ( !IsValid(this) )
+						if (!IsValid(this))
 						{
 							Reset();
 						}
 						// We have a valid filepath but no objects/assets spawned, so it's likely we were just spawned on the
 						// other client, and were replicated here with our RootLayer path already filled out, meaning we should just load that stage
 						// Note that now our UUsdTransactor may have already caused the stage itself to be loaded, but we may still need to call LoadUsdStage on our end.
-						else if ( ObjectsToWatch.Num() == 0 && ( !AssetCache || AssetCache->GetNumAssets() == 0 ) )
+						else if (ObjectsToWatch.Num() == 0 && (!AssetCache || AssetCache->GetNumAssets() == 0))
 						{
 							this->LoadUsdStage();
-							AUsdStageActor::OnActorLoaded.Broadcast( this );
+							AUsdStageActor::OnActorLoaded.Broadcast(this);
 						}
 					}
 				}
 			);
 		}
 
-		FCoreUObjectDelegates::OnObjectPropertyChanged.AddUObject( this, &AUsdStageActor::OnObjectPropertyChanged );
+		FCoreUObjectDelegates::OnObjectPropertyChanged.AddUObject(this, &AUsdStageActor::OnObjectPropertyChanged);
 
 		// Also prevent standalone from doing this
-		if ( GIsEditor && GEditor )
+		if (GIsEditor && GEditor)
 		{
-			if ( UBlueprintGeneratedClass* BPClass = Cast<UBlueprintGeneratedClass>( GetClass() ) )
+			if (UBlueprintGeneratedClass* BPClass = Cast<UBlueprintGeneratedClass>(GetClass()))
 			{
 				FRecompilationTracker::SetupEvents();
-				FCoreUObjectDelegates::OnObjectsReplaced.AddUObject( this, &AUsdStageActor::OnObjectsReplaced );
+				FCoreUObjectDelegates::OnObjectsReplaced.AddUObject(this, &AUsdStageActor::OnObjectsReplaced);
 			}
 		}
 
-		LevelSequenceHelper.GetOnSkelAnimationBaked().AddUObject( this, &AUsdStageActor::OnSkelAnimationBaked );
+		LevelSequenceHelper.GetOnSkelAnimationBaked().AddUObject(this, &AUsdStageActor::OnSkelAnimationBaked);
 
 #endif // WITH_EDITOR
 
-		OnTimeChanged.AddUObject( this, &AUsdStageActor::AnimatePrims );
+		OnTimeChanged.AddUObject(this, &AUsdStageActor::AnimatePrims);
 
-		UsdListener.GetOnObjectsChanged().AddUObject( this, &AUsdStageActor::OnUsdObjectsChanged );
+		UsdListener.GetOnObjectsChanged().AddUObject(this, &AUsdStageActor::OnUsdObjectsChanged);
 
 		UsdListener.GetOnLayersChanged().AddLambda(
-			[&]( const TArray< FString >& ChangeVec )
+			[&](const TArray< FString >& ChangeVec)
 			{
-				if ( !IsListeningToUsdNotices() )
+				if (!IsListeningToUsdNotices())
 				{
 					return;
 				}
 
 				TOptional< TGuardValue< ITransaction* > > SuppressTransaction;
-				if ( this->GetOutermost()->HasAnyPackageFlags( PKG_PlayInEditor ) )
+				if (this->GetOutermost()->HasAnyPackageFlags(PKG_PlayInEditor))
 				{
-					SuppressTransaction.Emplace( GUndo, nullptr );
+					SuppressTransaction.Emplace(GUndo, nullptr);
 				}
 
 				// Check to see if any layer reloaded. If so, rebuild all of our animations as a single layer changing
 				// might propagate timecodes through all level sequences
-				for ( const FString& ChangeVecItem : ChangeVec )
+				for (const FString& ChangeVecItem : ChangeVec)
 				{
-					UE_LOG( LogUsd, Verbose, TEXT("Reloading animations because layer '%s' was added/removed/reloaded"), *ChangeVecItem );
+					UE_LOG(LogUsd, Verbose, TEXT("Reloading animations because layer '%s' was added/removed/reloaded"), *ChangeVecItem);
 					ReloadAnimations();
 
 					// Make sure our PrimsToAnimate and the LevelSequenceHelper are kept in sync, because we'll use PrimsToAnimate to
@@ -773,7 +774,7 @@ void AUsdStageActor::NewStage()
 {
 #if USE_USD_SDK
 	UE::FUsdStage NewStage = UnrealUSDWrapper::NewStage();
-	if ( !NewStage )
+	if (!NewStage)
 	{
 		return;
 	}
@@ -781,40 +782,40 @@ void AUsdStageActor::NewStage()
 	// We'll create an in-memory stage, and so the "RootLayer" path we'll use will be a
 	// magic path that is guaranteed to never exist in a filesystem due to invalid characters.
 	UE::FSdfLayer Layer = NewStage.GetRootLayer();
-	if ( !Layer )
+	if (!Layer)
 	{
 		return;
 	}
-	FString	StagePath = FString( UnrealIdentifiers::IdentifierPrefix ) + Layer.GetIdentifier();
+	FString	StagePath = FString(UnrealIdentifiers::IdentifierPrefix) + Layer.GetIdentifier();
 
-	UE::FUsdPrim RootPrim = NewStage.DefinePrim( UE::FSdfPath{ TEXT( "/Root" ) }, TEXT( "Xform" ) );
-	ensure( UsdUtils::SetDefaultKind( RootPrim, EUsdDefaultKind::Assembly ) );
+	UE::FUsdPrim RootPrim = NewStage.DefinePrim(UE::FSdfPath{TEXT("/Root")}, TEXT("Xform"));
+	ensure(UsdUtils::SetDefaultKind(RootPrim, EUsdDefaultKind::Assembly));
 
-	NewStage.SetDefaultPrim( RootPrim );
+	NewStage.SetDefaultPrim(RootPrim);
 
-	SetRootLayer( StagePath );
+	SetRootLayer(StagePath);
 #endif // USE_USD_SDK
 }
 
-void AUsdStageActor::IsolateLayer( const UE::FSdfLayer& Layer )
+void AUsdStageActor::IsolateLayer(const UE::FSdfLayer& Layer)
 {
-	if ( IsolatedStage && IsolatedStage.GetRootLayer() == Layer )
+	if (IsolatedStage && IsolatedStage.GetRootLayer() == Layer)
 	{
 		return;
 	}
 
 	// Stop isolating
-	if ( !Layer || Layer == UsdStage.GetRootLayer() )
+	if (!Layer || Layer == UsdStage.GetRootLayer())
 	{
 		IsolatedStage = UE::FUsdStage{};
 
-		UsdListener.Register( UsdStage );
+		UsdListener.Register(UsdStage);
 
 		LoadUsdStage();
 		return;
 	}
 
-	if ( UsdStage )
+	if (UsdStage)
 	{
 		// Ideally we'd check to see if Layer belonged to the stage first, but there isn't any foolproof way to check
 		// for that because if the layer is muted it won't show up on the layer stack at all. Its also valid to mute
@@ -825,20 +826,20 @@ void AUsdStageActor::IsolateLayer( const UE::FSdfLayer& Layer )
 		// We really want our own stage for this and not something from the stage cache.
 		// Plus, this means its easier to cleanup: Just drop our IsolatedStage
 		const bool bUseStageCache = false;
-		IsolatedStage = UnrealUSDWrapper::OpenStage( Layer, {}, EUsdInitialLoadSet::LoadAll, bUseStageCache );
-		IsolatedStage.SetEditTarget( IsolatedStage.GetRootLayer() );
-		IsolatedStage.SetInterpolationType( InterpolationType );
+		IsolatedStage = UnrealUSDWrapper::OpenStage(Layer, {}, EUsdInitialLoadSet::LoadAll, bUseStageCache);
+		IsolatedStage.SetEditTarget(IsolatedStage.GetRootLayer());
+		IsolatedStage.SetInterpolationType(InterpolationType);
 
-		UsdListener.Register( IsolatedStage );
+		UsdListener.Register(IsolatedStage);
 
 		LoadUsdStage();
 	}
 }
 
-void AUsdStageActor::OnUsdObjectsChanged( const UsdUtils::FObjectChangesByPath& InfoChanges, const UsdUtils::FObjectChangesByPath& ResyncChanges )
+void AUsdStageActor::OnUsdObjectsChanged(const UsdUtils::FObjectChangesByPath& InfoChanges, const UsdUtils::FObjectChangesByPath& ResyncChanges)
 {
 #if USE_USD_SDK
-	if ( !IsListeningToUsdNotices() )
+	if (!IsListeningToUsdNotices())
 	{
 		return;
 	}
@@ -847,13 +848,13 @@ void AUsdStageActor::OnUsdObjectsChanged( const UsdUtils::FObjectChangesByPath&
 	// from our component changes, and this will also make sure we're not duplicating the events we store and replicate via multi-user: If a modification
 	// can be described purely via UObject changes, then those changes will be responsible for the whole modification and we won't record the corresponding
 	// stage changes. The intent is that when undo/redo/replicating that UObject change, it will automatically generate the corresponding stage changes
-	if ( Transactor )
+	if (Transactor)
 	{
-		Transactor->Update( InfoChanges, ResyncChanges );
+		Transactor->Update(InfoChanges, ResyncChanges);
 	}
 
 	const UE::FUsdStage& Stage = GetOrLoadUsdStage();
-	if ( !Stage )
+	if (!Stage)
 	{
 		return;
 	}
@@ -864,12 +865,12 @@ void AUsdStageActor::OnUsdObjectsChanged( const UsdUtils::FObjectChangesByPath&
 	{
 		const UE::FUsdStage& BaseStage = GetBaseUsdStage();
 		const UE::FSdfLayer& StageRoot = BaseStage.GetRootLayer();
-		if ( !StageRoot )
+		if (!StageRoot)
 		{
 			return;
 		}
 
-		if ( !FUsdStageActorImpl::DoesPathPointToLayer( RootLayer.FilePath, StageRoot ) )
+		if (!FUsdStageActorImpl::DoesPathPointToLayer(RootLayer.FilePath, StageRoot))
 		{
 			return;
 		}
@@ -878,43 +879,43 @@ void AUsdStageActor::OnUsdObjectsChanged( const UsdUtils::FObjectChangesByPath&
 	// Mark the level as dirty since we received a notice about our stage having changed in some way.
 	// The main goal of this is to trigger the "save layers" dialog if we then save the UE level
 	const bool bAlwaysMarkDirty = true;
-	Modify( bAlwaysMarkDirty );
+	Modify(bAlwaysMarkDirty);
 
 	// We may update our levelsequence objects (tracks, moviescene, sections, etc.) due to these changes. We definitely don't want to write anything
 	// back to USD when these objects change though.
-	FScopedBlockMonitoringChangesForTransaction BlockMonitoring{ LevelSequenceHelper };
+	FScopedBlockMonitoringChangesForTransaction BlockMonitoring{LevelSequenceHelper};
 
 	bool bHasResync = ResyncChanges.Num() > 0;
 
 	// The most important thing here is to iterate in parent to child order, so build SortedPrimsChangedList
 	TMap< FString, bool > SortedPrimsChangedList;
-	for ( const TPair<FString, TArray<UsdUtils::FObjectChangeNotice>>& InfoChange : InfoChanges )
+	for (const TPair<FString, TArray<UsdUtils::FObjectChangeNotice>>& InfoChange : InfoChanges)
 	{
-		UE::FSdfPath PrimPath = UE::FSdfPath( *InfoChange.Key ).StripAllVariantSelections();
+		UE::FSdfPath PrimPath = UE::FSdfPath(*InfoChange.Key).StripAllVariantSelections();
 
 		// Upgrade some info changes into resync changes
 		bool bIsResync = false;
-		for ( const UsdUtils::FObjectChangeNotice& ObjectChange : InfoChange.Value )
+		for (const UsdUtils::FObjectChangeNotice& ObjectChange : InfoChange.Value)
 		{
-			for ( const UsdUtils::FAttributeChange& AttributeChange : ObjectChange.AttributeChanges )
+			for (const UsdUtils::FAttributeChange& AttributeChange : ObjectChange.AttributeChanges)
 			{
 				static const TSet<FString> StageResyncProperties = {
-					TEXT( "metersPerUnit" ),
-					TEXT( "upAxis" )
+					TEXT("metersPerUnit"),
+					TEXT("upAxis")
 				};
 
 				// Upgrade these to resync so that the prim twins are regenerated, which clears all the existing
 				// animation tracks and adds new ones, automatically re-baking to control rig
 				static const TSet<FString> PrimResyncProperties = {
-					*UsdToUnreal::ConvertToken( UnrealIdentifiers::UnrealControlRigPath ),
-					*UsdToUnreal::ConvertToken( UnrealIdentifiers::UnrealUseFKControlRig ),
-					*UsdToUnreal::ConvertToken( UnrealIdentifiers::UnrealControlRigReduceKeys ),
-					*UsdToUnreal::ConvertToken( UnrealIdentifiers::UnrealControlRigReductionTolerance )
+					*UsdToUnreal::ConvertToken(UnrealIdentifiers::UnrealControlRigPath),
+					*UsdToUnreal::ConvertToken(UnrealIdentifiers::UnrealUseFKControlRig),
+					*UsdToUnreal::ConvertToken(UnrealIdentifiers::UnrealControlRigReduceKeys),
+					*UsdToUnreal::ConvertToken(UnrealIdentifiers::UnrealControlRigReductionTolerance)
 				};
 
 				// Some stage info should trigger some resyncs because they should trigger reparsing of geometry
-				if ( ( PrimPath.IsAbsoluteRootPath() && StageResyncProperties.Contains( AttributeChange.PropertyName ) )
-					|| PrimResyncProperties.Contains( AttributeChange.PropertyName ) )
+				if ((PrimPath.IsAbsoluteRootPath() && StageResyncProperties.Contains(AttributeChange.PropertyName))
+					|| PrimResyncProperties.Contains(AttributeChange.PropertyName))
 				{
 					bIsResync = true;
 					bHasResync = true;
@@ -926,15 +927,15 @@ void AUsdStageActor::OnUsdObjectsChanged( const UsdUtils::FObjectChangesByPath&
 		// We may need the full spec path with variant selections later, but for traversal and retrieving prims from the stage we always need
 		// the prim path without any variant selections in it (i.e. GetPrimAtPath("/Root{Varset=Var}Child") doesn't work, we need GetPrimAtPath("/Root/Child")),
 		// and USD sometimes emits changes with the variant selection path (like during renames).
-		SortedPrimsChangedList.Add( UE::FSdfPath( *InfoChange.Key ).StripAllVariantSelections().GetString(), bIsResync );
+		SortedPrimsChangedList.Add(UE::FSdfPath(*InfoChange.Key).StripAllVariantSelections().GetString(), bIsResync);
 	}
 	// Do Resyncs after so that they overwrite pure info changes if we have any
-	for ( const TPair<FString, TArray<UsdUtils::FObjectChangeNotice>>& ResyncChange : ResyncChanges )
+	for (const TPair<FString, TArray<UsdUtils::FObjectChangeNotice>>& ResyncChange : ResyncChanges)
 	{
-		UE::FSdfPath PrimPath = UE::FSdfPath( *ResyncChange.Key ).StripAllVariantSelections();
+		UE::FSdfPath PrimPath = UE::FSdfPath(*ResyncChange.Key).StripAllVariantSelections();
 
 		const bool bIsResync = true;
-		SortedPrimsChangedList.Add( PrimPath.GetString(), bIsResync );
+		SortedPrimsChangedList.Add(PrimPath.GetString(), bIsResync);
 	}
 
 	// If we have just resynced a child prim that was collapsed, we actually need to resync from its collapsed root
@@ -942,33 +943,33 @@ void AUsdStageActor::OnUsdObjectsChanged( const UsdUtils::FObjectChangesByPath&
 	// UpdateComponents/ReloadAssets below will already do this of course, but the point of this is that we unwind
 	// the collapsed prim paths *before* we rebuild the cache: We're really interested in whether this prim was
 	// collapsed or not before this change, as the change itself may have made the prim uncollapsible now
-	if ( InfoCache.IsValid() )
+	if (InfoCache.IsValid())
 	{
 		TMap<FString, bool> NewEntries;
-		NewEntries.Reserve( SortedPrimsChangedList.Num() );
+		NewEntries.Reserve(SortedPrimsChangedList.Num());
 
-		for ( TMap< FString, bool >::TIterator ChangeIt = SortedPrimsChangedList.CreateIterator(); ChangeIt; ++ChangeIt )
+		for (TMap< FString, bool >::TIterator ChangeIt = SortedPrimsChangedList.CreateIterator(); ChangeIt; ++ChangeIt)
 		{
-			if ( ChangeIt.Value() )
+			if (ChangeIt.Value())
 			{
-				UE::FSdfPath ChildPath{ *ChangeIt.Key() };
+				UE::FSdfPath ChildPath{*ChangeIt.Key()};
 
 				// Usually the cache should contain info about any prim on the stage, but if this is a
 				// notice about a prim being *created*, it may not have anything about this prim yet
-				if ( InfoCache->ContainsInfoAboutPrim( ChildPath ) )
+				if (InfoCache->ContainsInfoAboutPrim(ChildPath))
 				{
-					NewEntries.Add( InfoCache->UnwindToNonCollapsedPath( ChildPath, ECollapsingType::Assets ).GetString(), true );
-					NewEntries.Add( InfoCache->UnwindToNonCollapsedPath( ChildPath, ECollapsingType::Components ).GetString(), true );
+					NewEntries.Add(InfoCache->UnwindToNonCollapsedPath(ChildPath, ECollapsingType::Assets).GetString(), true);
+					NewEntries.Add(InfoCache->UnwindToNonCollapsedPath(ChildPath, ECollapsingType::Components).GetString(), true);
 
 					ChangeIt.RemoveCurrent();
 				}
 			}
 		}
 
-		SortedPrimsChangedList.Append( NewEntries );
+		SortedPrimsChangedList.Append(NewEntries);
 	}
 
-	SortedPrimsChangedList.KeySort([]( const FString& A, const FString& B ) -> bool { return A.Len() < B.Len(); } );
+	SortedPrimsChangedList.KeySort([](const FString& A, const FString& B) -> bool { return A.Len() < B.Len(); });
 
 	// During PIE, the PIE and the editor world will respond to notices. We have to prevent any PIE
 	// objects from being added to the transaction however, or else it will be discarded when finalized.
@@ -976,12 +977,12 @@ void AUsdStageActor::OnUsdObjectsChanged( const UsdUtils::FObjectChangesByPath&
 	// system that want to use assets that will be destroyed by it on an undo.
 	// Note that we can't just make the spawned components/assets nontransactional because the PIE world will transact too
 	TOptional<TGuardValue<ITransaction*>> SuppressTransaction;
-	if ( this->GetOutermost()->HasAnyPackageFlags( PKG_PlayInEditor ) )
+	if (this->GetOutermost()->HasAnyPackageFlags(PKG_PlayInEditor))
 	{
 		SuppressTransaction.Emplace(GUndo, nullptr);
 	}
 
-	FScopedSlowTask RefreshStageTask( SortedPrimsChangedList.Num(), LOCTEXT( "RefreshingUSDStage", "Refreshing USD Stage" ) );
+	FScopedSlowTask RefreshStageTask(SortedPrimsChangedList.Num(), LOCTEXT("RefreshingUSDStage", "Refreshing USD Stage"));
 	RefreshStageTask.MakeDialog();
 
 	FScopedUsdMessageLog ScopedMessageLog;
@@ -993,44 +994,44 @@ void AUsdStageActor::OnUsdObjectsChanged( const UsdUtils::FObjectChangesByPath&
 
 	bool bDeselected = false;
 
-	if ( bHasResync && InfoCache.IsValid() )
+	if (bHasResync && InfoCache.IsValid())
 	{
 		// TODO: Selective rebuild of only the required parts of the cache.
 		// If a prim changes from CanBeCollapsed to not (or vice-versa), that means its parent may change, and its grandparent, etc. so we'd need to check a large
 		// part of the tree (here we're just rebuilding the whole thing for now).
 		// However, if we know that only Prim resynced, we can traverse the tree root down and if we reach Prim and its collapsing state
 		// hasn't updated from before, we don't have to update its subtree at all, or sibling subtrees.
-		TSharedRef< FUsdSchemaTranslationContext > TranslationContext = FUsdStageActorImpl::CreateUsdSchemaTranslationContext( this, TEXT( "/" ) );
-		InfoCache->RebuildCacheForSubtree( Stage.GetPseudoRoot(), TranslationContext.Get() );
+		TSharedRef< FUsdSchemaTranslationContext > TranslationContext = FUsdStageActorImpl::CreateUsdSchemaTranslationContext(this, TEXT("/"));
+		InfoCache->RebuildCacheForSubtree(Stage.GetPseudoRoot(), TranslationContext.Get());
 	}
 
-	for ( const TPair< FString, bool >& PrimChangedInfo : SortedPrimsChangedList )
+	for (const TPair< FString, bool >& PrimChangedInfo : SortedPrimsChangedList)
 	{
 		RefreshStageTask.EnterProgressFrame();
 
-		const UE::FSdfPath PrimPath = UE::FSdfPath{ *PrimChangedInfo.Key };
+		const UE::FSdfPath PrimPath = UE::FSdfPath{*PrimChangedInfo.Key};
 		const bool bIsResync = PrimChangedInfo.Value;
 
-		if ( bIsResync && !bDeselected )
+		if (bIsResync && !bDeselected)
 		{
-			FUsdStageActorImpl::DeselectActorsAndComponents( this );
+			FUsdStageActorImpl::DeselectActorsAndComponents(this);
 			bDeselected = true;
 		}
 
 		// Return if the path or any of its higher level paths are already processed
-		auto IsPathAlreadyProcessed = []( TSet< UE::FSdfPath >& PathsProcessed, UE::FSdfPath PathToProcess ) -> bool
+		auto IsPathAlreadyProcessed = [](TSet< UE::FSdfPath >& PathsProcessed, UE::FSdfPath PathToProcess) -> bool
 		{
 			FString SubPath;
 			FString ParentPath;
 
-			if ( PathsProcessed.Contains( UE::FSdfPath::AbsoluteRootPath() ) )
+			if (PathsProcessed.Contains(UE::FSdfPath::AbsoluteRootPath()))
 			{
 				return true;
 			}
 
-			while ( !PathToProcess.IsEmpty() && !PathsProcessed.Contains( PathToProcess ) )
+			while (!PathToProcess.IsEmpty() && !PathsProcessed.Contains(PathToProcess))
 			{
-				if ( !PathToProcess.IsAbsoluteRootPath() )
+				if (!PathToProcess.IsAbsoluteRootPath())
 				{
 					PathToProcess = PathToProcess.GetParentPath();
 				}
@@ -1040,90 +1041,90 @@ void AUsdStageActor::OnUsdObjectsChanged( const UsdUtils::FObjectChangesByPath&
 				}
 			}
 
-			return !PathToProcess.IsEmpty() && PathsProcessed.Contains( PathToProcess );
+			return !PathToProcess.IsEmpty() && PathsProcessed.Contains(PathToProcess);
 		};
 
-		auto UpdateComponents = [&]( const UE::FUsdPrim& Prim, const UE::FSdfPath& InPrimPath, const bool bInResync )
+		auto UpdateComponents = [&](const UE::FUsdPrim& Prim, const UE::FSdfPath& InPrimPath, const bool bInResync)
 		{
 			// Don't query the InfoCache about prims that don't exist since that will ensure.
 			// We need to carry on though, because we must update components in case a prim is renamed or deleted,
 			// which will emit notices for prims that don't actually exist on the stage
-			UE::FSdfPath ComponentsPrimPath = ( Prim && InfoCache.IsValid() )
-				? InfoCache->UnwindToNonCollapsedPath( InPrimPath, ECollapsingType::Components )
+			UE::FSdfPath ComponentsPrimPath = (Prim && InfoCache.IsValid())
+				? InfoCache->UnwindToNonCollapsedPath(InPrimPath, ECollapsingType::Components)
 				: InPrimPath;
 
 			TSet< UE::FSdfPath >& RefreshedComponents = bInResync ? ResyncedComponents : UpdatedComponents;
 
-			if ( !IsPathAlreadyProcessed( RefreshedComponents, ComponentsPrimPath ) )
+			if (!IsPathAlreadyProcessed(RefreshedComponents, ComponentsPrimPath))
 			{
-				TSharedRef< FUsdSchemaTranslationContext > TranslationContext = FUsdStageActorImpl::CreateUsdSchemaTranslationContext( this, ComponentsPrimPath.GetString() );
-				UpdatePrim( ComponentsPrimPath, bInResync, *TranslationContext );
+				TSharedRef< FUsdSchemaTranslationContext > TranslationContext = FUsdStageActorImpl::CreateUsdSchemaTranslationContext(this, ComponentsPrimPath.GetString());
+				UpdatePrim(ComponentsPrimPath, bInResync, *TranslationContext);
 				TranslationContext->CompleteTasks();
 
-				RefreshedComponents.Add( ComponentsPrimPath );
+				RefreshedComponents.Add(ComponentsPrimPath);
 
-				if ( bInResync )
+				if (bInResync)
 				{
 					// Consider that the path has been updated in the case of a resync
-					UpdatedComponents.Add( ComponentsPrimPath );
+					UpdatedComponents.Add(ComponentsPrimPath);
 				}
 			}
 		};
 
-		auto ReloadAssets = [&]( const UE::FUsdPrim& Prim, const UE::FSdfPath& InPrimPath, const bool bInResync )
+		auto ReloadAssets = [&](const UE::FUsdPrim& Prim, const UE::FSdfPath& InPrimPath, const bool bInResync)
 		{
-			UE::FSdfPath AssetsPrimPath = ( Prim && InfoCache.IsValid() )
-				? InfoCache->UnwindToNonCollapsedPath( InPrimPath, ECollapsingType::Assets )
+			UE::FSdfPath AssetsPrimPath = (Prim && InfoCache.IsValid())
+				? InfoCache->UnwindToNonCollapsedPath(InPrimPath, ECollapsingType::Assets)
 				: InPrimPath;
 
 			TSet< UE::FSdfPath >& RefreshedAssets = bInResync ? ResyncedAssets : UpdatedAssets;
 
-			if ( !IsPathAlreadyProcessed( RefreshedAssets, AssetsPrimPath ) )
+			if (!IsPathAlreadyProcessed(RefreshedAssets, AssetsPrimPath))
 			{
-				TSharedRef< FUsdSchemaTranslationContext > TranslationContext = FUsdStageActorImpl::CreateUsdSchemaTranslationContext( this, AssetsPrimPath.GetString() );
+				TSharedRef< FUsdSchemaTranslationContext > TranslationContext = FUsdStageActorImpl::CreateUsdSchemaTranslationContext(this, AssetsPrimPath.GetString());
 
-				if ( bInResync )
+				if (bInResync)
 				{
-					UMaterialInterface* ExistingMaterial = Cast<UMaterialInterface>( AssetCache->GetAssetForPrim( AssetsPrimPath.GetString() ) );
+					UMaterialInterface* ExistingMaterial = Cast<UMaterialInterface>(AssetCache->GetAssetForPrim(AssetsPrimPath.GetString()));
 
-					UE::FUsdPrim PrimToResync = GetOrLoadUsdStage().GetPrimAtPath( AssetsPrimPath );
-					LoadAssets( *TranslationContext, PrimToResync );
+					UE::FUsdPrim PrimToResync = GetOrLoadUsdStage().GetPrimAtPath(AssetsPrimPath);
+					LoadAssets(*TranslationContext, PrimToResync);
 
-					UMaterialInterface* NewMaterial = Cast<UMaterialInterface>( AssetCache->GetAssetForPrim( AssetsPrimPath.GetString() ) );
+					UMaterialInterface* NewMaterial = Cast<UMaterialInterface>(AssetCache->GetAssetForPrim(AssetsPrimPath.GetString()));
 
 					// For UE-120185: If we recreated a material for a prim path we also need to update all components that were using it.
 					// This could be fleshed out further if other asset types require this refresh of "dependent components" but materials
 					// seem to be the only ones that do at the moment
-					if ( ExistingMaterial && ( ExistingMaterial != NewMaterial ) )
+					if (ExistingMaterial && (ExistingMaterial != NewMaterial))
 					{
-						for ( const FString& MaterialUserPrim : UsdUtils::GetMaterialUsers( PrimToResync, MaterialPurpose ) )
+						for (const FString& MaterialUserPrim : UsdUtils::GetMaterialUsers(PrimToResync, MaterialPurpose))
 						{
 							const bool bResyncComponent = true; // We need to force resync to reassign materials
-							UE::FSdfPath MaterialPrimPath{ *MaterialUserPrim };
-							UE::FUsdPrim MaterialPrim = GetOrLoadUsdStage().GetPrimAtPath( MaterialPrimPath );
-							UpdateComponents( MaterialPrim, MaterialPrimPath, bResyncComponent );
+							UE::FSdfPath MaterialPrimPath{*MaterialUserPrim};
+							UE::FUsdPrim MaterialPrim = GetOrLoadUsdStage().GetPrimAtPath(MaterialPrimPath);
+							UpdateComponents(MaterialPrim, MaterialPrimPath, bResyncComponent);
 						}
 					}
 
 					// Resyncing also includes "updating" the prim
-					UpdatedAssets.Add( AssetsPrimPath );
+					UpdatedAssets.Add(AssetsPrimPath);
 				}
 				else
 				{
-					LoadAsset( *TranslationContext, GetOrLoadUsdStage().GetPrimAtPath( AssetsPrimPath ) );
+					LoadAsset(*TranslationContext, GetOrLoadUsdStage().GetPrimAtPath(AssetsPrimPath));
 				}
 
-				RefreshedAssets.Add( AssetsPrimPath );
+				RefreshedAssets.Add(AssetsPrimPath);
 			}
 		};
 
-		UE::FUsdPrim Prim = GetOrLoadUsdStage().GetPrimAtPath( PrimPath );
-		ReloadAssets( Prim, PrimPath, bIsResync );
-		UpdateComponents( Prim, PrimPath, bIsResync );
+		UE::FUsdPrim Prim = GetOrLoadUsdStage().GetPrimAtPath(PrimPath);
+		ReloadAssets(Prim, PrimPath, bIsResync);
+		UpdateComponents(Prim, PrimPath, bIsResync);
 
-		if ( HasAuthorityOverStage() )
+		if (HasAuthorityOverStage())
 		{
-			OnPrimChanged.Broadcast( PrimChangedInfo.Key, PrimChangedInfo.Value );
+			OnPrimChanged.Broadcast(PrimChangedInfo.Key, PrimChangedInfo.Value);
 		}
 	}
 #endif // USE_USD_SDK
@@ -1169,115 +1170,115 @@ void AUsdStageActor::BlockMonitoringLevelSequenceForThisTransaction()
 	LevelSequenceHelper.BlockMonitoringChangesForThisTransaction();
 }
 
-UUsdPrimTwin* AUsdStageActor::GetOrCreatePrimTwin( const UE::FSdfPath& UsdPrimPath )
+UUsdPrimTwin* AUsdStageActor::GetOrCreatePrimTwin(const UE::FSdfPath& UsdPrimPath)
 {
 	const FString PrimPath = UsdPrimPath.GetString();
 	const FString ParentPrimPath = UsdPrimPath.GetParentPath().GetString();
 
 	UUsdPrimTwin* RootTwin = GetRootPrimTwin();
-	UUsdPrimTwin* UsdPrimTwin = RootTwin->Find( PrimPath );
-	UUsdPrimTwin* ParentUsdPrimTwin = RootTwin->Find( ParentPrimPath );
+	UUsdPrimTwin* UsdPrimTwin = RootTwin->Find(PrimPath);
+	UUsdPrimTwin* ParentUsdPrimTwin = RootTwin->Find(ParentPrimPath);
 
-	const UE::FUsdPrim Prim = GetOrLoadUsdStage().GetPrimAtPath( UsdPrimPath );
+	const UE::FUsdPrim Prim = GetOrLoadUsdStage().GetPrimAtPath(UsdPrimPath);
 
-	if ( !Prim )
+	if (!Prim)
 	{
 		return nullptr;
 	}
 
-	if ( !ParentUsdPrimTwin )
+	if (!ParentUsdPrimTwin)
 	{
 		ParentUsdPrimTwin = RootUsdTwin;
 	}
 
-	if ( !UsdPrimTwin )
+	if (!UsdPrimTwin)
 	{
-		UsdPrimTwin = &ParentUsdPrimTwin->AddChild( *PrimPath );
+		UsdPrimTwin = &ParentUsdPrimTwin->AddChild(*PrimPath);
 
-		UsdPrimTwin->OnDestroyed.AddUObject( this, &AUsdStageActor::OnUsdPrimTwinDestroyed );
+		UsdPrimTwin->OnDestroyed.AddUObject(this, &AUsdStageActor::OnUsdPrimTwinDestroyed);
 	}
 
 	return UsdPrimTwin;
 }
 
-UUsdPrimTwin* AUsdStageActor::ExpandPrim( const UE::FUsdPrim& Prim, FUsdSchemaTranslationContext& TranslationContext )
+UUsdPrimTwin* AUsdStageActor::ExpandPrim(const UE::FUsdPrim& Prim, FUsdSchemaTranslationContext& TranslationContext)
 {
 	UUsdPrimTwin* UsdPrimTwin = nullptr;
 #if USE_USD_SDK
 	// "Active" is the non-destructive deletion used in USD. Sometimes when we rename/remove a prim in a complex stage it may remain in
 	// an inactive state, but its otherwise effectively deleted
-	if ( !Prim || !Prim.IsActive() )
+	if (!Prim || !Prim.IsActive())
 	{
 		return nullptr;
 	}
 
-	TRACE_CPUPROFILER_EVENT_SCOPE( AUsdStageActor::ExpandPrim );
+	TRACE_CPUPROFILER_EVENT_SCOPE(AUsdStageActor::ExpandPrim);
 
-	UsdPrimTwin = GetOrCreatePrimTwin( Prim.GetPrimPath() );
+	UsdPrimTwin = GetOrCreatePrimTwin(Prim.GetPrimPath());
 
-	if ( !UsdPrimTwin )
+	if (!UsdPrimTwin)
 	{
 		return nullptr;
 	}
 
 	bool bExpandChildren = true;
 
-	IUsdSchemasModule& UsdSchemasModule = FModuleManager::Get().LoadModuleChecked< IUsdSchemasModule >( TEXT("USDSchemas") );
+	IUsdSchemasModule& UsdSchemasModule = FModuleManager::Get().LoadModuleChecked< IUsdSchemasModule >(TEXT("USDSchemas"));
 
-	if ( TSharedPtr< FUsdSchemaTranslator > SchemaTranslator = UsdSchemasModule.GetTranslatorRegistry().CreateTranslatorForSchema( TranslationContext.AsShared(), UE::FUsdTyped( Prim ) ) )
+	if (TSharedPtr< FUsdSchemaTranslator > SchemaTranslator = UsdSchemasModule.GetTranslatorRegistry().CreateTranslatorForSchema(TranslationContext.AsShared(), UE::FUsdTyped(Prim)))
 	{
-		if ( !UsdPrimTwin->SceneComponent.IsValid() )
+		if (!UsdPrimTwin->SceneComponent.IsValid())
 		{
 			UsdPrimTwin->SceneComponent = SchemaTranslator->CreateComponents();
 		}
 		else
 		{
-			ObjectsToWatch.Remove( UsdPrimTwin->SceneComponent.Get() );
-			SchemaTranslator->UpdateComponents( UsdPrimTwin->SceneComponent.Get() );
+			ObjectsToWatch.Remove(UsdPrimTwin->SceneComponent.Get());
+			SchemaTranslator->UpdateComponents(UsdPrimTwin->SceneComponent.Get());
 		}
 
-		bExpandChildren = !SchemaTranslator->CollapsesChildren( ECollapsingType::Components );
+		bExpandChildren = !SchemaTranslator->CollapsesChildren(ECollapsingType::Components);
 	}
 
-	if ( bExpandChildren )
+	if (bExpandChildren)
 	{
 		USceneComponent* ContextParentComponent = TranslationContext.ParentComponent;
 
-		if ( UsdPrimTwin->SceneComponent.IsValid() )
+		if (UsdPrimTwin->SceneComponent.IsValid())
 		{
 			ContextParentComponent = UsdPrimTwin->SceneComponent.Get();
 		}
 
-		TGuardValue< USceneComponent* > ParentComponentGuard( TranslationContext.ParentComponent, ContextParentComponent );
+		TGuardValue< USceneComponent* > ParentComponentGuard(TranslationContext.ParentComponent, ContextParentComponent);
 
 		const bool bTraverseInstanceProxies = true;
-		const TArray< UE::FUsdPrim > PrimChildren = Prim.GetFilteredChildren( bTraverseInstanceProxies );
+		const TArray< UE::FUsdPrim > PrimChildren = Prim.GetFilteredChildren(bTraverseInstanceProxies);
 
-		for ( const UE::FUsdPrim& ChildPrim : PrimChildren )
+		for (const UE::FUsdPrim& ChildPrim : PrimChildren)
 		{
-			ExpandPrim( ChildPrim, TranslationContext );
+			ExpandPrim(ChildPrim, TranslationContext);
 		}
 	}
 
-	if ( UsdPrimTwin->SceneComponent.IsValid() )
+	if (UsdPrimTwin->SceneComponent.IsValid())
 	{
 #if WITH_EDITOR
 		UsdPrimTwin->SceneComponent->PostEditChange();
 #endif // WITH_EDITOR
 
-		if ( !UsdPrimTwin->SceneComponent->IsRegistered() )
+		if (!UsdPrimTwin->SceneComponent->IsRegistered())
 		{
 			UsdPrimTwin->SceneComponent->RegisterComponent();
 		}
 
-		ObjectsToWatch.Add( UsdPrimTwin->SceneComponent.Get(), UsdPrimTwin->PrimPath );
+		ObjectsToWatch.Add(UsdPrimTwin->SceneComponent.Get(), UsdPrimTwin->PrimPath);
 	}
 
 	// Update the prim animated status
-	const bool bIsAnimated = UsdUtils::IsAnimated( Prim );
-	if ( bIsAnimated )
+	const bool bIsAnimated = UsdUtils::IsAnimated(Prim);
+	if (bIsAnimated)
 	{
-		if ( !PrimsToAnimate.Contains( UsdPrimTwin->PrimPath ) )
+		if (!PrimsToAnimate.Contains(UsdPrimTwin->PrimPath))
 		{
 			// Unfortunately if we have animated visibility we need to be ready to update the visibility
 			// of all components that we spawned for child prims whenever this prim's visibility updates.
@@ -1289,55 +1290,66 @@ UUsdPrimTwin* AUsdStageActor::ExpandPrim( const UE::FUsdPrim& Prim, FUsdSchemaTr
 			// to be e.g. a breadth first traversal on the prim tree or anything like this, so these updates
 			// need to be order-independent, which means we really should add the entire subtree to the list
 			// and have UpdateComponents called on all components.
-			if ( UsdUtils::HasAnimatedVisibility( Prim ) )
+			if (UsdUtils::HasAnimatedVisibility(Prim))
 			{
 				const bool bRecursive = true;
 				UsdPrimTwin->Iterate(
-					[this]( UUsdPrimTwin& Twin )
+					[this](UUsdPrimTwin& Twin)
 					{
-						if ( !PrimsToAnimate.Contains( Twin.PrimPath ) )
+						if (!PrimsToAnimate.Contains(Twin.PrimPath))
 						{
-							PrimsToAnimate.Add( Twin.PrimPath );
-							LevelSequenceHelper.AddPrim( Twin );
+							PrimsToAnimate.Add(Twin.PrimPath);
+							LevelSequenceHelper.AddPrim(Twin);
 						}
 					},
 					bRecursive
-				);
+					);
 			}
 
-			PrimsToAnimate.Add( UsdPrimTwin->PrimPath );
-			LevelSequenceHelper.AddPrim( *UsdPrimTwin );
+			UE::FSdfPath PrimPath(*UsdPrimTwin->PrimPath);
+			UE::FUsdPrim UsdPrim(UsdStage.GetPrimAtPath(PrimPath));
+
+			static IConsoleVariable* CVar = IConsoleManager::Get().FindConsoleVariable(TEXT("USD.DisableGeoCacheTracks"));
+			const bool bDisableGeoCacheTracks = CVar && CVar->GetBool();
+
+			// Optimization: manually animate mesh prims only when the USD GeometryCaches tracks are disabled
+			const bool bAnimatePrim = bDisableGeoCacheTracks || !UsdPrim.IsA(TEXT("Mesh"));
+			if (bAnimatePrim)
+			{
+				PrimsToAnimate.Add(UsdPrimTwin->PrimPath);
+			}
+			LevelSequenceHelper.AddPrim(*UsdPrimTwin);
 		}
 	}
 	else
 	{
-		if ( PrimsToAnimate.Contains( UsdPrimTwin->PrimPath ) )
+		if (PrimsToAnimate.Contains(UsdPrimTwin->PrimPath))
 		{
-			PrimsToAnimate.Remove( UsdPrimTwin->PrimPath );
-			LevelSequenceHelper.RemovePrim( *UsdPrimTwin );
+			PrimsToAnimate.Remove(UsdPrimTwin->PrimPath);
+			LevelSequenceHelper.RemovePrim(*UsdPrimTwin);
 		}
 	}
 
 	// Setup Control Rig tracks if we need to. This must be done after adding regular skeletal animation tracks
 	// if we have any as if will properly deactivate them like the usual "Bake to Control Rig" workflow.
-	if ( Prim.IsA( TEXT( "SkelRoot" ) ) )
+	if (Prim.IsA(TEXT("SkelRoot")))
 	{
-		if ( UsdUtils::PrimHasSchema( Prim, UnrealIdentifiers::ControlRigAPI ) )
+		if (UsdUtils::PrimHasSchema(Prim, UnrealIdentifiers::ControlRigAPI))
 		{
-			LevelSequenceHelper.UpdateControlRigTracks( *UsdPrimTwin );
+			LevelSequenceHelper.UpdateControlRigTracks(*UsdPrimTwin);
 
 			// If our prim wasn't originally considered animated and we just added a new track, it should be
 			// considered animated too, so lets add it to the proper locations. This will also ensure that
 			// we can close the sequencer after creating a new animation in this way and see it animate on
 			// the level
-			if ( !bIsAnimated )
+			if (!bIsAnimated)
 			{
-				PrimsToAnimate.Add( UsdPrimTwin->PrimPath );
-				LevelSequenceHelper.AddPrim( *UsdPrimTwin );
+				PrimsToAnimate.Add(UsdPrimTwin->PrimPath);
+				LevelSequenceHelper.AddPrim(*UsdPrimTwin);
 
 				// Prevent register/unregister spam when calling FUsdGeomXformableTranslator::UpdateComponents later
 				// during sequencer animation (which can cause the Sequencer UI to glitch out a bit)
-				UsdPrimTwin->SceneComponent->SetMobility( EComponentMobility::Movable );
+				UsdPrimTwin->SceneComponent->SetMobility(EComponentMobility::Movable);
 			}
 		}
 	}
@@ -1346,37 +1358,37 @@ UUsdPrimTwin* AUsdStageActor::ExpandPrim( const UE::FUsdPrim& Prim, FUsdSchemaTr
 	return UsdPrimTwin;
 }
 
-void AUsdStageActor::UpdatePrim( const UE::FSdfPath& InUsdPrimPath, bool bResync, FUsdSchemaTranslationContext& TranslationContext )
+void AUsdStageActor::UpdatePrim(const UE::FSdfPath& InUsdPrimPath, bool bResync, FUsdSchemaTranslationContext& TranslationContext)
 {
-	TRACE_CPUPROFILER_EVENT_SCOPE( AUsdStageActor::UpdatePrim );
+	TRACE_CPUPROFILER_EVENT_SCOPE(AUsdStageActor::UpdatePrim);
 
-	FScopedSlowTask SlowTask( 1.f, LOCTEXT( "UpdatingUSDPrim", "Updating USD Prim" ) );
+	FScopedSlowTask SlowTask(1.f, LOCTEXT("UpdatingUSDPrim", "Updating USD Prim"));
 	SlowTask.MakeDialog();
 	SlowTask.EnterProgressFrame();
 
 	UE::FSdfPath UsdPrimPath = InUsdPrimPath;
 
-	if ( !UsdPrimPath.IsAbsoluteRootOrPrimPath() )
+	if (!UsdPrimPath.IsAbsoluteRootOrPrimPath())
 	{
 		UsdPrimPath = UsdPrimPath.GetAbsoluteRootOrPrimPath();
 	}
 
-	if ( UsdPrimPath.IsAbsoluteRootOrPrimPath() )
+	if (UsdPrimPath.IsAbsoluteRootOrPrimPath())
 	{
-		if ( bResync )
+		if (bResync)
 		{
 			FString PrimPath = UsdPrimPath.GetString();
-			if ( UUsdPrimTwin* UsdPrimTwin = GetRootPrimTwin()->Find( PrimPath ) )
+			if (UUsdPrimTwin* UsdPrimTwin = GetRootPrimTwin()->Find(PrimPath))
 			{
 				UsdPrimTwin->Clear();
 			}
 		}
 
-		UE::FUsdPrim PrimToExpand = GetOrLoadUsdStage().GetPrimAtPath( UsdPrimPath );
-		UUsdPrimTwin* UsdPrimTwin = ExpandPrim( PrimToExpand, TranslationContext );
+		UE::FUsdPrim PrimToExpand = GetOrLoadUsdStage().GetPrimAtPath(UsdPrimPath);
+		UUsdPrimTwin* UsdPrimTwin = ExpandPrim(PrimToExpand, TranslationContext);
 
 #if WITH_EDITOR
-		if ( GIsEditor && GEditor && !IsGarbageCollecting() ) // Make sure we're not in standalone either
+		if (GIsEditor && GEditor && !IsGarbageCollecting()) // Make sure we're not in standalone either
 		{
 			GEditor->BroadcastLevelActorListChanged();
 			GEditor->RedrawLevelEditingViewports();
@@ -1407,26 +1419,26 @@ UE::FUsdStage& AUsdStageActor::GetOrLoadUsdStage()
 	return IsolatedStage ? IsolatedStage : UsdStage;
 }
 
-void AUsdStageActor::SetRootLayer( const FString& RootFilePath )
+void AUsdStageActor::SetRootLayer(const FString& RootFilePath)
 {
 	const bool bMarkDirty = false;
 	Modify(bMarkDirty);
 
 	FString RelativeFilePath = RootFilePath;
 #if USE_USD_SDK
-	if ( !RelativeFilePath.IsEmpty() && !FPaths::IsRelative( RelativeFilePath ) && !RelativeFilePath.StartsWith( UnrealIdentifiers::IdentifierPrefix ) )
+	if (!RelativeFilePath.IsEmpty() && !FPaths::IsRelative(RelativeFilePath) && !RelativeFilePath.StartsWith(UnrealIdentifiers::IdentifierPrefix))
 	{
-		RelativeFilePath = UsdUtils::MakePathRelativeToProjectDir( RootFilePath );
+		RelativeFilePath = UsdUtils::MakePathRelativeToProjectDir(RootFilePath);
 	}
 #endif // USE_USD_SDK
 
 	// See if we're talking about the stage that is already loaded
-	if ( UsdStage )
+	if (UsdStage)
 	{
 		const UE::FSdfLayer& StageRootLayer = UsdStage.GetRootLayer();
-		if ( StageRootLayer )
+		if (StageRootLayer)
 		{
-			if ( FUsdStageActorImpl::DoesPathPointToLayer( RelativeFilePath, StageRootLayer ) )
+			if (FUsdStageActorImpl::DoesPathPointToLayer(RelativeFilePath, StageRootLayer))
 			{
 				return;
 			}
@@ -1438,7 +1450,7 @@ void AUsdStageActor::SetRootLayer( const FString& RootFilePath )
 	LoadUsdStage();
 }
 
-void AUsdStageActor::SetInitialLoadSet( EUsdInitialLoadSet NewLoadSet )
+void AUsdStageActor::SetInitialLoadSet(EUsdInitialLoadSet NewLoadSet)
 {
 	const bool bMarkDirty = false;
 	Modify(bMarkDirty);
@@ -1447,7 +1459,7 @@ void AUsdStageActor::SetInitialLoadSet( EUsdInitialLoadSet NewLoadSet )
 	LoadUsdStage();
 }
 
-void AUsdStageActor::SetInterpolationType( EUsdInterpolationType NewType )
+void AUsdStageActor::SetInterpolationType(EUsdInterpolationType NewType)
 {
 	const bool bMarkDirty = false;
 	Modify(bMarkDirty);
@@ -1456,31 +1468,31 @@ void AUsdStageActor::SetInterpolationType( EUsdInterpolationType NewType )
 	LoadUsdStage();
 }
 
-void AUsdStageActor::SetKindsToCollapse( int32 NewKindsToCollapse )
+void AUsdStageActor::SetKindsToCollapse(int32 NewKindsToCollapse)
 {
 	const bool bMarkDirty = false;
 	Modify(bMarkDirty);
 
-	const EUsdDefaultKind NewEnum = ( EUsdDefaultKind ) NewKindsToCollapse;
+	const EUsdDefaultKind NewEnum = (EUsdDefaultKind)NewKindsToCollapse;
 	EUsdDefaultKind Result = NewEnum;
 
 	// If we're collapsing all 'model's, then we must collapse all of its derived kinds
-	if ( EnumHasAnyFlags( NewEnum, EUsdDefaultKind::Model ) )
+	if (EnumHasAnyFlags(NewEnum, EUsdDefaultKind::Model))
 	{
-		Result |= ( EUsdDefaultKind::Component | EUsdDefaultKind::Group | EUsdDefaultKind::Assembly );
+		Result |= (EUsdDefaultKind::Component | EUsdDefaultKind::Group | EUsdDefaultKind::Assembly);
 	}
 
 	// If we're collapsing all 'group's, then we must collapse all of its derived kinds
-	if ( EnumHasAnyFlags( NewEnum, EUsdDefaultKind::Group ) )
+	if (EnumHasAnyFlags(NewEnum, EUsdDefaultKind::Group))
 	{
-		Result |= ( EUsdDefaultKind::Assembly );
+		Result |= (EUsdDefaultKind::Assembly);
 	}
 
-	KindsToCollapse = ( int32 ) Result;
+	KindsToCollapse = (int32)Result;
 	LoadUsdStage();
 }
 
-void AUsdStageActor::SetMergeIdenticalMaterialSlots( bool bMerge )
+void AUsdStageActor::SetMergeIdenticalMaterialSlots(bool bMerge)
 {
 	Modify();
 
@@ -1488,11 +1500,11 @@ void AUsdStageActor::SetMergeIdenticalMaterialSlots( bool bMerge )
 	LoadUsdStage();
 }
 
-void AUsdStageActor::SetCollapseTopLevelPointInstancers( bool bCollapse )
+void AUsdStageActor::SetCollapseTopLevelPointInstancers(bool bCollapse)
 {
 }
 
-void AUsdStageActor::SetPurposesToLoad( int32 NewPurposesToLoad )
+void AUsdStageActor::SetPurposesToLoad(int32 NewPurposesToLoad)
 {
 	const bool bMarkDirty = false;
 	Modify(bMarkDirty);
@@ -1501,7 +1513,7 @@ void AUsdStageActor::SetPurposesToLoad( int32 NewPurposesToLoad )
 	LoadUsdStage();
 }
 
-void AUsdStageActor::SetNaniteTriangleThreshold( int32 NewNaniteTriangleThreshold )
+void AUsdStageActor::SetNaniteTriangleThreshold(int32 NewNaniteTriangleThreshold)
 {
 	const bool bMarkDirty = false;
 	Modify(bMarkDirty);
@@ -1510,7 +1522,7 @@ void AUsdStageActor::SetNaniteTriangleThreshold( int32 NewNaniteTriangleThreshol
 	LoadUsdStage();
 }
 
-void AUsdStageActor::SetRenderContext( const FName& NewRenderContext )
+void AUsdStageActor::SetRenderContext(const FName& NewRenderContext)
 {
 	const bool bMarkDirty = false;
 	Modify(bMarkDirty);
@@ -1519,19 +1531,19 @@ void AUsdStageActor::SetRenderContext( const FName& NewRenderContext )
 	LoadUsdStage();
 }
 
-void AUsdStageActor::SetMaterialPurpose( const FName& NewMaterialPurpose )
+void AUsdStageActor::SetMaterialPurpose(const FName& NewMaterialPurpose)
 {
 	const bool bMarkDirty = false;
-	Modify( bMarkDirty );
+	Modify(bMarkDirty);
 
 	MaterialPurpose = NewMaterialPurpose;
 	LoadUsdStage();
 }
 
-USDSTAGE_API void AUsdStageActor::SetRootMotionHandling( EUsdRootMotionHandling NewHandlingStrategy )
+USDSTAGE_API void AUsdStageActor::SetRootMotionHandling(EUsdRootMotionHandling NewHandlingStrategy)
 {
 	const bool bMarkDirty = false;
-	Modify( bMarkDirty );
+	Modify(bMarkDirty);
 
 	RootMotionHandling = NewHandlingStrategy;
 	LoadUsdStage();
@@ -1546,28 +1558,28 @@ void AUsdStageActor::SetTime(float InTime)
 	Refresh();
 }
 
-USceneComponent* AUsdStageActor::GetGeneratedComponent( const FString& PrimPath )
+USceneComponent* AUsdStageActor::GetGeneratedComponent(const FString& PrimPath)
 {
-	const UE::FUsdStage& CurrentStage = static_cast< const AUsdStageActor* >( this )->GetUsdStage();
-	if ( !CurrentStage )
+	const UE::FUsdStage& CurrentStage = static_cast<const AUsdStageActor*>(this)->GetUsdStage();
+	if (!CurrentStage)
 	{
 		return nullptr;
 	}
 
 	// We can't query our InfoCache with invalid paths, as we're using ensures to track when we miss the cache (which shouldn't ever happen)
-	UE::FSdfPath UsdPath{ *PrimPath };
-	if ( !CurrentStage.GetPrimAtPath( UsdPath ) )
+	UE::FSdfPath UsdPath{*PrimPath};
+	if (!CurrentStage.GetPrimAtPath(UsdPath))
 	{
 		return nullptr;
 	}
 
 	FString UncollapsedPath = PrimPath;
-	if ( InfoCache.IsValid() )
+	if (InfoCache.IsValid())
 	{
-		UncollapsedPath = InfoCache->UnwindToNonCollapsedPath( UsdPath, ECollapsingType::Components ).GetString();
+		UncollapsedPath = InfoCache->UnwindToNonCollapsedPath(UsdPath, ECollapsingType::Components).GetString();
 	}
 
-	if ( UUsdPrimTwin* UsdPrimTwin = GetRootPrimTwin()->Find( UncollapsedPath ) )
+	if (UUsdPrimTwin* UsdPrimTwin = GetRootPrimTwin()->Find(UncollapsedPath))
 	{
 		return UsdPrimTwin->GetSceneComponent();
 	}
@@ -1575,22 +1587,22 @@ USceneComponent* AUsdStageActor::GetGeneratedComponent( const FString& PrimPath
 	return nullptr;
 }
 
-TArray<UObject*> AUsdStageActor::GetGeneratedAssets( const FString& PrimPath )
+TArray<UObject*> AUsdStageActor::GetGeneratedAssets(const FString& PrimPath)
 {
 	const UE::FUsdStage& CurrentStage = static_cast<const AUsdStageActor*>(this)->GetUsdStage();
-	if ( !CurrentStage )
+	if (!CurrentStage)
 	{
 		return {};
 	}
 
 	// We can't query our InfoCache with invalid paths, as we're using ensures to track when we miss the cache (which shouldn't ever happen)
-	UE::FSdfPath UsdPath{ *PrimPath };
-	if ( !CurrentStage.GetPrimAtPath( UsdPath ) )
+	UE::FSdfPath UsdPath{*PrimPath};
+	if (!CurrentStage.GetPrimAtPath(UsdPath))
 	{
 		return {};
 	}
 
-	if ( !AssetCache )
+	if (!AssetCache)
 	{
 		return {};
 	}
@@ -1600,29 +1612,29 @@ TArray<UObject*> AUsdStageActor::GetGeneratedAssets( const FString& PrimPath )
 
 	// If we found an asset generated specifically for a prim path, use that. If this prim has been collapsed,
 	// the asset generated for it will be assigned to its collapsing root, so we won't find anything here
-	if ( UObject* FoundAsset = AssetCache->GetAssetForPrim( PrimPath ) )
+	if (UObject* FoundAsset = AssetCache->GetAssetForPrim(PrimPath))
 	{
-		Result.Add( FoundAsset );
+		Result.Add(FoundAsset);
 	}
 	// If we haven't try seeing if we collapsed into another asset
-	else if ( InfoCache.IsValid() )
+	else if (InfoCache.IsValid())
 	{
-		PathToUse = InfoCache->UnwindToNonCollapsedPath( UsdPath, ECollapsingType::Assets ).GetString();
-		if ( UObject* FoundCollapsedAsset = AssetCache->GetAssetForPrim( PathToUse ) )
+		PathToUse = InfoCache->UnwindToNonCollapsedPath(UsdPath, ECollapsingType::Assets).GetString();
+		if (UObject* FoundCollapsedAsset = AssetCache->GetAssetForPrim(PathToUse))
 		{
-			Result.Add( FoundCollapsedAsset );
+			Result.Add(FoundCollapsedAsset);
 		}
 	}
 
 	// Collect any other asset that claims they came from the same prim (e.g. also return the skeleton if we query
 	// the SkelRoot, or return textures if we query a material prim that used them, etc.)
-	for ( const TPair<FString, TObjectPtr<UObject>>& HashToAsset : AssetCache->GetCachedAssets() )
+	for (const TPair<FString, TObjectPtr<UObject>>& HashToAsset : AssetCache->GetCachedAssets())
 	{
-		if ( UUsdAssetImportData* ImportData = UsdUtils::GetAssetImportData( HashToAsset.Value ) )
+		if (UUsdAssetImportData* ImportData = UsdUtils::GetAssetImportData(HashToAsset.Value))
 		{
-			if ( ImportData->PrimPath == PathToUse )
+			if (ImportData->PrimPath == PathToUse)
 			{
-				Result.Add( HashToAsset.Value );
+				Result.Add(HashToAsset.Value);
 			}
 		}
 	}
@@ -1630,31 +1642,31 @@ TArray<UObject*> AUsdStageActor::GetGeneratedAssets( const FString& PrimPath )
 	return Result.Array();
 }
 
-FString AUsdStageActor::GetSourcePrimPath( UObject* Object )
+FString AUsdStageActor::GetSourcePrimPath(UObject* Object)
 {
-	if ( USceneComponent* Component = Cast<USceneComponent>( Object ) )
+	if (USceneComponent* Component = Cast<USceneComponent>(Object))
 	{
-		if ( UUsdPrimTwin* UsdPrimTwin = GetRootPrimTwin()->Find( Component ) )
+		if (UUsdPrimTwin* UsdPrimTwin = GetRootPrimTwin()->Find(Component))
 		{
 			return UsdPrimTwin->PrimPath;
 		}
 	}
 
-	if ( AssetCache )
+	if (AssetCache)
 	{
-		for ( const TPair<FString, TWeakObjectPtr<UObject>>& PrimPathToAsset : AssetCache->GetAssetPrimLinks() )
+		for (const TPair<FString, TWeakObjectPtr<UObject>>& PrimPathToAsset : AssetCache->GetAssetPrimLinks())
 		{
-			if ( PrimPathToAsset.Value.Get() == Object )
+			if (PrimPathToAsset.Value.Get() == Object)
 			{
 				return PrimPathToAsset.Key;
 			}
 		}
 
-		for ( const TPair<FString, TObjectPtr<UObject>>& HashToAsset : AssetCache->GetCachedAssets() )
+		for (const TPair<FString, TObjectPtr<UObject>>& HashToAsset : AssetCache->GetCachedAssets())
 		{
-			if ( HashToAsset.Value == Object )
+			if (HashToAsset.Value == Object)
 			{
-				if ( UUsdAssetImportData* ImportData = UsdUtils::GetAssetImportData( HashToAsset.Value ) )
+				if (UUsdAssetImportData* ImportData = UsdUtils::GetAssetImportData(HashToAsset.Value))
 				{
 					return ImportData->PrimPath;
 				}
@@ -1668,21 +1680,21 @@ FString AUsdStageActor::GetSourcePrimPath( UObject* Object )
 void AUsdStageActor::OpenUsdStage()
 {
 	// Early exit if stage is already opened
-	if ( UsdStage || RootLayer.FilePath.IsEmpty() )
+	if (UsdStage || RootLayer.FilePath.IsEmpty())
 	{
 		return;
 	}
 
-	TRACE_CPUPROFILER_EVENT_SCOPE( AUsdStageActor::OpenUsdStage );
+	TRACE_CPUPROFILER_EVENT_SCOPE(AUsdStageActor::OpenUsdStage);
 
 	UsdUtils::StartMonitoringErrors();
 
 	FString AbsPath;
-	if ( !RootLayer.FilePath.StartsWith( UnrealIdentifiers::IdentifierPrefix ) && FPaths::IsRelative( RootLayer.FilePath ) )
+	if (!RootLayer.FilePath.StartsWith(UnrealIdentifiers::IdentifierPrefix) && FPaths::IsRelative(RootLayer.FilePath))
 	{
 		// The RootLayer property is marked as RelativeToGameDir, and FUsdStageViewModel::OpenStage will also give us paths relative to the project's directory
-		FString ProjectDir = FPaths::ConvertRelativePathToFull( FPaths::ProjectDir() );
-		AbsPath = FPaths::ConvertRelativePathToFull( FPaths::Combine( ProjectDir, RootLayer.FilePath ) );
+		FString ProjectDir = FPaths::ConvertRelativePathToFull(FPaths::ProjectDir());
+		AbsPath = FPaths::ConvertRelativePathToFull(FPaths::Combine(ProjectDir, RootLayer.FilePath));
 	}
 	else
 	{
@@ -1691,87 +1703,87 @@ void AUsdStageActor::OpenUsdStage()
 
 	OnPreStageChanged.Broadcast();
 
-	UsdStage = UnrealUSDWrapper::OpenStage( *AbsPath, InitialLoadSet );
+	UsdStage = UnrealUSDWrapper::OpenStage(*AbsPath, InitialLoadSet);
 	IsolatedStage = UE::FUsdStage{};
 
-	if ( UsdStage )
+	if (UsdStage)
 	{
-		UsdStage.SetEditTarget( UsdStage.GetRootLayer() );
+		UsdStage.SetEditTarget(UsdStage.GetRootLayer());
 
-		UsdStage.SetInterpolationType( InterpolationType );
+		UsdStage.SetInterpolationType(InterpolationType);
 
-		UsdListener.Register( UsdStage );
+		UsdListener.Register(UsdStage);
 
 #if USE_USD_SDK
 		// Try loading a UE-state session layer if we can find one
 		const bool bCreateIfNeeded = false;
-		UsdUtils::GetUEPersistentStateSublayer( UsdStage, bCreateIfNeeded );
+		UsdUtils::GetUEPersistentStateSublayer(UsdStage, bCreateIfNeeded);
 #endif // #if USE_USD_SDK
 
 		OnStageChanged.Broadcast();
 	}
 
-	UsdUtils::ShowErrorsAndStopMonitoring( FText::Format( LOCTEXT("USDOpenError", "Encountered some errors opening USD file at path '{0}!\nCheck the Output Log for details."), FText::FromString( RootLayer.FilePath ) ) );
+	UsdUtils::ShowErrorsAndStopMonitoring(FText::Format(LOCTEXT("USDOpenError", "Encountered some errors opening USD file at path '{0}!\nCheck the Output Log for details."), FText::FromString(RootLayer.FilePath)));
 }
 
 void AUsdStageActor::CloseUsdStage()
 {
 	OnPreStageChanged.Broadcast();
 
-	FUsdStageActorImpl::DiscardStage( UsdStage, this );
+	FUsdStageActorImpl::DiscardStage(UsdStage, this);
 	UsdStage = UE::FUsdStage();
 	IsolatedStage = UE::FUsdStage();  // We don't keep our isolated stages on the stage cache
-	LevelSequenceHelper.Init( UE::FUsdStage() ); // Drop the helper's reference to the stage
+	LevelSequenceHelper.Init(UE::FUsdStage()); // Drop the helper's reference to the stage
 	OnStageChanged.Broadcast();
 }
 
 #if WITH_EDITOR
-void AUsdStageActor::OnBeginPIE( bool bIsSimulating )
+void AUsdStageActor::OnBeginPIE(bool bIsSimulating)
 {
 	// Remove transient flag from our spawned actors and components so they can be duplicated for PIE
 	const bool bTransient = false;
-	UpdateSpawnedObjectsTransientFlag( bTransient );
+	UpdateSpawnedObjectsTransientFlag(bTransient);
 
 	bIsTransitioningIntoPIE = true;
 
 	// Take ownership of our RootTwin and pretend our entire prim tree is a subobject so that it's duplicated over with us into PIE
-	if ( UUsdPrimTwin* RootTwin = GetRootPrimTwin() )
+	if (UUsdPrimTwin* RootTwin = GetRootPrimTwin())
 	{
-		RootTwin->Rename( nullptr, this );
+		RootTwin->Rename(nullptr, this);
 
-		if ( FProperty* Prop = GetClass()->FindPropertyByName( GET_MEMBER_NAME_CHECKED( AUsdStageActor, RootUsdTwin ) ) )
+		if (FProperty* Prop = GetClass()->FindPropertyByName(GET_MEMBER_NAME_CHECKED(AUsdStageActor, RootUsdTwin)))
 		{
-			Prop->ClearPropertyFlags( CPF_Transient );
+			Prop->ClearPropertyFlags(CPF_Transient);
 		}
 
-		if ( FProperty* Prop = UUsdPrimTwin::StaticClass()->FindPropertyByName( UUsdPrimTwin::GetChildrenPropertyName() ) )
+		if (FProperty* Prop = UUsdPrimTwin::StaticClass()->FindPropertyByName(UUsdPrimTwin::GetChildrenPropertyName()))
 		{
-			Prop->ClearPropertyFlags( CPF_Transient );
+			Prop->ClearPropertyFlags(CPF_Transient);
 		}
 	}
 }
 
-void AUsdStageActor::OnPostPIEStarted( bool bIsSimulating )
+void AUsdStageActor::OnPostPIEStarted(bool bIsSimulating)
 {
 	// Restore transient flags to our spawned actors and components so they aren't saved otherwise
 	const bool bTransient = true;
-	UpdateSpawnedObjectsTransientFlag( bTransient );
+	UpdateSpawnedObjectsTransientFlag(bTransient);
 
 	bIsTransitioningIntoPIE = false;
 
 	// Put our RootTwin back on the transient package so that if our blueprint is compiled it doesn't get reconstructed with us
-	if ( UUsdPrimTwin* RootTwin = GetRootPrimTwin() )
+	if (UUsdPrimTwin* RootTwin = GetRootPrimTwin())
 	{
-		RootTwin->Rename( nullptr, GetTransientPackage() );
+		RootTwin->Rename(nullptr, GetTransientPackage());
 
-		if ( FProperty* Prop = GetClass()->FindPropertyByName( GET_MEMBER_NAME_CHECKED( AUsdStageActor, RootUsdTwin ) ) )
+		if (FProperty* Prop = GetClass()->FindPropertyByName(GET_MEMBER_NAME_CHECKED(AUsdStageActor, RootUsdTwin)))
 		{
-			Prop->SetPropertyFlags( CPF_Transient );
+			Prop->SetPropertyFlags(CPF_Transient);
 		}
 
-		if ( FProperty* Prop = UUsdPrimTwin::StaticClass()->FindPropertyByName( UUsdPrimTwin::GetChildrenPropertyName() ) )
+		if (FProperty* Prop = UUsdPrimTwin::StaticClass()->FindPropertyByName(UUsdPrimTwin::GetChildrenPropertyName()))
 		{
-			Prop->SetPropertyFlags( CPF_Transient );
+			Prop->SetPropertyFlags(CPF_Transient);
 		}
 	}
 
@@ -1782,23 +1794,23 @@ void AUsdStageActor::OnPostPIEStarted( bool bIsSimulating )
 	AnimatePrims();
 }
 
-void AUsdStageActor::OnObjectsReplaced( const TMap<UObject*, UObject*>& ObjectReplacementMap )
+void AUsdStageActor::OnObjectsReplaced(const TMap<UObject*, UObject*>& ObjectReplacementMap)
 {
-	UBlueprintGeneratedClass* BPClass = Cast<UBlueprintGeneratedClass>( GetClass() );
-	if ( !BPClass )
+	UBlueprintGeneratedClass* BPClass = Cast<UBlueprintGeneratedClass>(GetClass());
+	if (!BPClass)
 	{
 		return;
 	}
 
-	UBlueprint* BP = Cast<UBlueprint>( BPClass->ClassGeneratedBy );
-	if ( !BP )
+	UBlueprint* BP = Cast<UBlueprint>(BPClass->ClassGeneratedBy);
+	if (!BP)
 	{
 		return;
 	}
 
 	// We are a replacement actor: Anything that is a property was already copied over,
 	// and the spawned actors and components are still alive. We just need to move over any remaining non-property data
-	if ( AUsdStageActor* NewActor = Cast<AUsdStageActor>( ObjectReplacementMap.FindRef( this ) ) )
+	if (AUsdStageActor* NewActor = Cast<AUsdStageActor>(ObjectReplacementMap.FindRef(this)))
 	{
 		// If our BP has changes and we're going into PIE, we'll get automatically recompiled. Sadly OnBeginPIE will trigger
 		// before we're duplicated for the reinstantiation process, which is a problem because our prim twins will then be
@@ -1806,24 +1818,24 @@ void AUsdStageActor::OnObjectsReplaced( const TMap<UObject*, UObject*>& ObjectRe
 		// twins from the old actor, which is what the reinstantiation process should have done instead anyway. Note that only
 		// later will the components and actors being pointed to by this duplicated prim twin be moved to the PIE world, so those
 		// references would be updated correctly.
-		if ( RootUsdTwin && RootUsdTwin->GetOuter() == this )
+		if (RootUsdTwin && RootUsdTwin->GetOuter() == this)
 		{
-			NewActor->RootUsdTwin = DuplicateObject( RootUsdTwin, NewActor );
+			NewActor->RootUsdTwin = DuplicateObject(RootUsdTwin, NewActor);
 		}
 
-		if ( FRecompilationTracker::IsBeingCompiled( BP ) )
+		if (FRecompilationTracker::IsBeingCompiled(BP))
 		{
 			// Can't just move out of this one as TUsdStore expects its TOptional to always contain a value, and we may
 			// still need to use the bool operator on it to test for validity
 			NewActor->UsdStage = UsdStage;
 			NewActor->IsolatedStage = IsolatedStage;
 
-			NewActor->LevelSequenceHelper = MoveTemp( LevelSequenceHelper );
+			NewActor->LevelSequenceHelper = MoveTemp(LevelSequenceHelper);
 			NewActor->LevelSequence = LevelSequence;
-			NewActor->BlendShapesByPath = MoveTemp( BlendShapesByPath );
-			NewActor->MaterialToPrimvarToUVIndex = MoveTemp( MaterialToPrimvarToUVIndex );
+			NewActor->BlendShapesByPath = MoveTemp(BlendShapesByPath);
+			NewActor->MaterialToPrimvarToUVIndex = MoveTemp(MaterialToPrimvarToUVIndex);
 
-			NewActor->UsdListener.Register( NewActor->UsdStage );
+			NewActor->UsdListener.Register(NewActor->UsdStage);
 
 			// This does not look super safe...
 			NewActor->OnActorDestroyed = OnActorDestroyed;
@@ -1834,7 +1846,7 @@ void AUsdStageActor::OnObjectsReplaced( const TMap<UObject*, UObject*>& ObjectRe
 
 			// UEngine::CopyPropertiesForUnrelatedObjects won't copy over the cache's transient assets, but we still
 			// need to ensure their lifetime here, so just take the previous asset cache instead, which still that has the transient assets
-			AssetCache->Rename( nullptr, NewActor );
+			AssetCache->Rename(nullptr, NewActor);
 			NewActor->AssetCache = AssetCache;
 
 			NewActor->InfoCache = InfoCache;
@@ -1846,7 +1858,7 @@ void AUsdStageActor::OnObjectsReplaced( const TMap<UObject*, UObject*>& ObjectRe
 			NewActor->bIsModifyingAProperty = bIsModifyingAProperty;
 			NewActor->bIsUndoRedoing = bIsUndoRedoing;
 
-			NewActor->IsBlockedFromUsdNotices.Set( IsBlockedFromUsdNotices.GetValue() );
+			NewActor->IsBlockedFromUsdNotices.Set(IsBlockedFromUsdNotices.GetValue());
 			NewActor->OldRootLayer = OldRootLayer;
 
 			// Close our stage or else it will remain open forever. NewActor has a a reference to it now so it won't actually close
@@ -1855,15 +1867,15 @@ void AUsdStageActor::OnObjectsReplaced( const TMap<UObject*, UObject*>& ObjectRe
 	}
 }
 
-void AUsdStageActor::OnLevelActorDeleted( AActor* DeletedActor )
+void AUsdStageActor::OnLevelActorDeleted(AActor* DeletedActor)
 {
 	// Check for this here because it could be that we tried to delete this actor before changing any of its
 	// properties, in which case our similar check within OnObjectPropertyChange hasn't had the chance to tag this actor
-	if ( RootLayer.FilePath == OldRootLayer.FilePath && FUsdStageActorImpl::ObjectNeedsMultiUserTag( DeletedActor, this ) )
+	if (RootLayer.FilePath == OldRootLayer.FilePath && FUsdStageActorImpl::ObjectNeedsMultiUserTag(DeletedActor, this))
 	{
 		// DeletedActor is already detached from our hierarchy, so we must tag it directly
 		TSet<UObject*> VisitedObjects;
-		FUsdStageActorImpl::AllowListComponentHierarchy( DeletedActor->GetRootComponent(), VisitedObjects );
+		FUsdStageActorImpl::AllowListComponentHierarchy(DeletedActor->GetRootComponent(), VisitedObjects);
 	}
 }
 
@@ -1871,18 +1883,18 @@ void AUsdStageActor::OnLevelActorDeleted( AActor* DeletedActor )
 
 void AUsdStageActor::LoadUsdStage()
 {
-	TRACE_CPUPROFILER_EVENT_SCOPE( AUsdStageActor::LoadUsdStage );
+	TRACE_CPUPROFILER_EVENT_SCOPE(AUsdStageActor::LoadUsdStage);
 
 	double StartTime = FPlatformTime::Cycles64();
 
-	FScopedSlowTask SlowTask( 1.f, LOCTEXT( "LoadingUDStage", "Loading USD Stage") );
+	FScopedSlowTask SlowTask(1.f, LOCTEXT("LoadingUDStage", "Loading USD Stage"));
 	SlowTask.MakeDialog();
 
 	OnPreStageChanged.Broadcast();
 
-	if ( !AssetCache )
+	if (!AssetCache)
 	{
-		AssetCache = NewObject< UUsdAssetCache >( this, TEXT("AssetCache"), GetMaskedFlags( RF_PropagateToSubObjects ) );
+		AssetCache = NewObject< UUsdAssetCache >(this, TEXT("AssetCache"), GetMaskedFlags(RF_PropagateToSubObjects));
 	}
 
 	// Block writing level sequence changes back to the USD stage until we finished this transaction, because once we do
@@ -1892,17 +1904,17 @@ void AUsdStageActor::LoadUsdStage()
 
 	ObjectsToWatch.Reset();
 
-	FUsdStageActorImpl::DeselectActorsAndComponents( this );
+	FUsdStageActorImpl::DeselectActorsAndComponents(this);
 
 	UUsdPrimTwin* RootTwin = GetRootPrimTwin();
 	RootTwin->Clear();
-	RootTwin->PrimPath = TEXT( "/" );
+	RootTwin->PrimPath = TEXT("/");
 
 	FScopedUsdMessageLog ScopedMessageLog;
 
 	// If we have an isolated stage we may have been asked to refresh to display it: We should keep our UsdStage opened
 	UE::FUsdStage StageToLoad;
-	if ( IsolatedStage )
+	if (IsolatedStage)
 	{
 		StageToLoad = IsolatedStage;
 	}
@@ -1913,7 +1925,7 @@ void AUsdStageActor::LoadUsdStage()
 		UsdStage = UE::FUsdStage();
 
 		OpenUsdStage();
-		if ( !UsdStage )
+		if (!UsdStage)
 		{
 			OnStageChanged.Broadcast();
 			return;
@@ -1922,7 +1934,7 @@ void AUsdStageActor::LoadUsdStage()
 		StageToLoad = UsdStage;
 	}
 
-	StageToLoad.SetInterpolationType( InterpolationType );
+	StageToLoad.SetInterpolationType(InterpolationType);
 
 	ReloadAnimations();
 
@@ -1932,27 +1944,27 @@ void AUsdStageActor::LoadUsdStage()
 	// the tracks
 	PrimsToAnimate.Reset();
 
-	TSharedRef< FUsdSchemaTranslationContext > TranslationContext = FUsdStageActorImpl::CreateUsdSchemaTranslationContext( this, RootTwin->PrimPath );
+	TSharedRef< FUsdSchemaTranslationContext > TranslationContext = FUsdStageActorImpl::CreateUsdSchemaTranslationContext(this, RootTwin->PrimPath);
 
-	SlowTask.EnterProgressFrame( 0.1f );
-	if ( !InfoCache.IsValid() )
+	SlowTask.EnterProgressFrame(0.1f);
+	if (!InfoCache.IsValid())
 	{
 		InfoCache = MakeShared<FUsdInfoCache>();
 	}
-	InfoCache->RebuildCacheForSubtree( StageToLoad.GetPseudoRoot(), TranslationContext.Get() );
+	InfoCache->RebuildCacheForSubtree(StageToLoad.GetPseudoRoot(), TranslationContext.Get());
 
-	SlowTask.EnterProgressFrame( 0.7f );
-	LoadAssets( *TranslationContext, StageToLoad.GetPseudoRoot() );
+	SlowTask.EnterProgressFrame(0.7f);
+	LoadAssets(*TranslationContext, StageToLoad.GetPseudoRoot());
 
-	SlowTask.EnterProgressFrame( 0.2f );
-	UpdatePrim( StageToLoad.GetPseudoRoot().GetPrimPath(), true, *TranslationContext );
+	SlowTask.EnterProgressFrame(0.2f);
+	UpdatePrim(StageToLoad.GetPseudoRoot().GetPrimPath(), true, *TranslationContext);
 
 	TranslationContext->CompleteTasks();
 
 	// Keep our old Time value if we're loading the stage during initialization, so that we can save/load Time values
-	if ( StageToLoad.GetRootLayer() && IsActorInitialized() )
+	if (StageToLoad.GetRootLayer() && IsActorInitialized())
 	{
-		SetTime( StageToLoad.GetRootLayer().GetStartTimeCode() );
+		SetTime(StageToLoad.GetRootLayer().GetStartTimeCode());
 
 		// If we're an instance of a blueprint that derives the stage actor and we're in the editor preview world, it means we're the
 		// blueprint preview actor. We (the instance) will load the stage and update our Time to StartTimeCode, but our CDO will not.
@@ -1960,22 +1972,22 @@ void AUsdStageActor::LoadUsdStage()
 		// value in the blueprint editor won't be propagated to the instance, so we wouldn't be able to animate the preview actor at all.
 		// Here we fix that by updating our CDO to our new Time value. Note how we just do this if we're the preview instance though, we don't
 		// want other instances driving the CDO like this
-		if ( UBlueprintGeneratedClass* BPClass = Cast<UBlueprintGeneratedClass>( GetClass() ) )
+		if (UBlueprintGeneratedClass* BPClass = Cast<UBlueprintGeneratedClass>(GetClass()))
 		{
 			UWorld* World = GetWorld();
-			if ( World && World->WorldType == EWorldType::EditorPreview )
+			if (World && World->WorldType == EWorldType::EditorPreview)
 			{
 				// Note: CDO is an instance of a BlueprintGeneratedClass here and this is just a base class pointer. We're not changing the actual AUsdStageActor's CDO
-				if ( AUsdStageActor* CDO = Cast<AUsdStageActor>( GetClass()->GetDefaultObject() ) )
+				if (AUsdStageActor* CDO = Cast<AUsdStageActor>(GetClass()->GetDefaultObject()))
 				{
-					CDO->SetTime( GetTime() );
+					CDO->SetTime(GetTime());
 				}
 			}
 		}
 	}
 
 #if WITH_EDITOR
-	if ( GIsEditor && GEditor && !IsGarbageCollecting() )
+	if (GIsEditor && GEditor && !IsGarbageCollecting())
 	{
 		GEditor->BroadcastLevelActorListChanged();
 	}
@@ -1987,14 +1999,14 @@ void AUsdStageActor::LoadUsdStage()
 	int ElapsedMin = int(ElapsedSeconds / 60.0);
 	ElapsedSeconds -= 60.0 * (double)ElapsedMin;
 
-	UE_LOG( LogUsd, Log, TEXT("%s %s in [%d min %.3f s]"), TEXT("Stage loaded"), *FPaths::GetBaseFilename( RootLayer.FilePath ), ElapsedMin, ElapsedSeconds );
+	UE_LOG(LogUsd, Log, TEXT("%s %s in [%d min %.3f s]"), TEXT("Stage loaded"), *FPaths::GetBaseFilename(RootLayer.FilePath), ElapsedMin, ElapsedSeconds);
 
 #if USE_USD_SDK
 	FUsdStageActorImpl::SendAnalytics(
 		this,
 		ElapsedSeconds,
-		UsdUtils::GetUsdStageNumFrames( StageToLoad ),
-		FPaths::GetExtension( RootLayer.FilePath )
+		UsdUtils::GetUsdStageNumFrames(StageToLoad),
+		FPaths::GetExtension(RootLayer.FilePath)
 	);
 #endif // #if USE_USD_SDK
 }
@@ -2004,19 +2016,19 @@ void AUsdStageActor::UnloadUsdStage()
 	const bool bMarkDirty = false;
 	Modify(bMarkDirty);
 
-	FUsdStageActorImpl::DeselectActorsAndComponents( this );
+	FUsdStageActorImpl::DeselectActorsAndComponents(this);
 
 	// Stop listening because we'll discard LevelSequence assets, which may trigger transactions
 	// and could lead to stage changes
 	BlockMonitoringLevelSequenceForThisTransaction();
 
-	if ( AssetCache )
+	if (AssetCache)
 	{
-		FUsdStageActorImpl::CloseEditorsForAssets( AssetCache->GetCachedAssets() );
+		FUsdStageActorImpl::CloseEditorsForAssets(AssetCache->GetCachedAssets());
 		AssetCache->Reset();
 	}
 
-	if ( InfoCache )
+	if (InfoCache)
 	{
 		InfoCache->Clear();
 	}
@@ -2025,28 +2037,28 @@ void AUsdStageActor::UnloadUsdStage()
 	BlendShapesByPath.Reset();
 	MaterialToPrimvarToUVIndex.Reset();
 
-	if ( LevelSequence )
+	if (LevelSequence)
 	{
 #if WITH_EDITOR
-		if ( GEditor )
+		if (GEditor)
 		{
-			GEditor->GetEditorSubsystem<UAssetEditorSubsystem>()->CloseAllEditorsForAsset( LevelSequence );
+			GEditor->GetEditorSubsystem<UAssetEditorSubsystem>()->CloseAllEditorsForAsset(LevelSequence);
 		}
 #endif // WITH_EDITOR
 		LevelSequence = nullptr;
 	}
 	LevelSequenceHelper.Clear();
 
-	if ( RootUsdTwin )
+	if (RootUsdTwin)
 	{
 		RootUsdTwin->Clear();
-		RootUsdTwin->PrimPath = TEXT( "/" );
+		RootUsdTwin->PrimPath = TEXT("/");
 	}
 
 #if WITH_EDITOR
 	// We can't emit this when garbage collecting as it may lead to objects being created
 	// (we may unload stage when going into PIE or other sensitive transitions)
-	if ( GEditor && !IsGarbageCollecting() )
+	if (GEditor && !IsGarbageCollecting())
 	{
 		GEditor->BroadcastLevelActorListChanged();
 	}
@@ -2059,14 +2071,14 @@ void AUsdStageActor::UnloadUsdStage()
 
 UUsdPrimTwin* AUsdStageActor::GetRootPrimTwin()
 {
-	if ( !RootUsdTwin )
+	if (!RootUsdTwin)
 	{
 		FScopedUnrealAllocs Allocs;
 
 		// Be careful not to give it a name, as there could be multiple of these on the transient package.
 		// It needs to be public or else FArchiveReplaceOrClearExternalReferences will reset our property
 		// whenever it is used from UEngine::CopyPropertiesForUnrelatedObjects for blueprint recompilation (if we're a blueprint class)
-		RootUsdTwin = NewObject<UUsdPrimTwin>( GetTransientPackage(), NAME_None, DefaultObjFlag | RF_Public );
+		RootUsdTwin = NewObject<UUsdPrimTwin>(GetTransientPackage(), NAME_None, DefaultObjFlag | RF_Public);
 	}
 
 	return RootUsdTwin;
@@ -2079,7 +2091,7 @@ void AUsdStageActor::Refresh() const
 
 void AUsdStageActor::ReloadAnimations()
 {
-	TRACE_CPUPROFILER_EVENT_SCOPE( AUsdStageActor::ReloadAnimations );
+	TRACE_CPUPROFILER_EVENT_SCOPE(AUsdStageActor::ReloadAnimations);
 
 	// If we're using some property editor that can trigger a stage reload (like the Nanite threshold spinbox),
 	// applying a value may trigger ReloadAnimations -> Can trigger asset editors to open/close/change focus ->
@@ -2090,28 +2102,28 @@ void AUsdStageActor::ReloadAnimations()
 	// hasn't finished being destroyed.
 	// In that case we only want the outer call to change the level sequence, so a reentrant guard does what we need
 	static bool bIsReentrant = false;
-	if ( bIsReentrant )
+	if (bIsReentrant)
 	{
 		return;
 	}
-	TGuardValue<bool> ReentrantGuard( bIsReentrant, true );
+	TGuardValue<bool> ReentrantGuard(bIsReentrant, true);
 
 	const UE::FUsdStage& CurrentStage = GetOrLoadUsdStage();
-	if ( !CurrentStage )
+	if (!CurrentStage)
 	{
 		return;
 	}
 
 	// Don't check for full authority here because even if we can't write back to the stage (i.e. during PIE) we still
 	// want to listen to it and have valid level sequences
-	if ( !IsTemplate() )
+	if (!IsTemplate())
 	{
 		bool bLevelSequenceEditorWasOpened = false;
 		if (LevelSequence)
 		{
 			// The sequencer won't update on its own, so let's at least force it closed
 #if WITH_EDITOR
-			if ( GIsEditor && GEditor )
+			if (GIsEditor && GEditor)
 			{
 				bLevelSequenceEditorWasOpened = GEditor->GetEditorSubsystem<UAssetEditorSubsystem>()->CloseAllEditorsForAsset(LevelSequence) > 0;
 			}
@@ -2122,8 +2134,8 @@ void AUsdStageActor::ReloadAnimations()
 		const bool bMarkDirty = false;
 		Modify(bMarkDirty);
 
-		LevelSequence = LevelSequenceHelper.Init( CurrentStage );
-		LevelSequenceHelper.BindToUsdStageActor( this );
+		LevelSequence = LevelSequenceHelper.Init(CurrentStage);
+		LevelSequenceHelper.BindToUsdStageActor(this);
 
 #if WITH_EDITOR
 		if (GIsEditor && GEditor && LevelSequence && bLevelSequenceEditorWasOpened)
@@ -2136,7 +2148,7 @@ void AUsdStageActor::ReloadAnimations()
 
 #if WITH_EDITOR
 
-void AUsdStageActor::PostEditChangeProperty( FPropertyChangedEvent& PropertyChangedEvent )
+void AUsdStageActor::PostEditChangeProperty(FPropertyChangedEvent& PropertyChangedEvent)
 {
 	// For handling root layer changes via direct changes to properties we want to go through OnObjectPropertyChanged -> HandlePropertyChangedEvent ->
 	// -> SetRootLayer (which checks whether this stage is already opened or not) -> PostRegisterAllComponents.
@@ -2146,17 +2158,17 @@ void AUsdStageActor::PostEditChangeProperty( FPropertyChangedEvent& PropertyChan
 	// the same stage again if we don't need to.
 
 	bIsModifyingAProperty = true;
-	Super::PostEditChangeProperty( PropertyChangedEvent );
+	Super::PostEditChangeProperty(PropertyChangedEvent);
 }
 
 void AUsdStageActor::PostTransacted(const FTransactionObjectEvent& TransactionEvent)
 {
 	const TArray<FName>& ChangedProperties = TransactionEvent.GetChangedProperties();
 
-	if ( TransactionEvent.HasPendingKillChange() )
+	if (TransactionEvent.HasPendingKillChange())
 	{
 		// Fires when being deleted in editor, redo delete
-		if ( !IsValidChecked(this) )
+		if (!IsValidChecked(this))
 		{
 			CloseUsdStage();
 		}
@@ -2170,15 +2182,15 @@ void AUsdStageActor::PostTransacted(const FTransactionObjectEvent& TransactionEv
 	// If we're in the persistent level don't do anything, because hiding/showing the persistent level doesn't
 	// cause actors to load/unload like it does if they're in sublevels
 	ULevel* CurrentLevel = GetLevel();
-	if ( CurrentLevel && !CurrentLevel->IsPersistentLevel() )
+	if (CurrentLevel && !CurrentLevel->IsPersistentLevel())
 	{
 		// If we're in a sublevel that is hidden, we'll respond to the generated PostUnregisterAllComponent call
 		// and unload our spawned actors/assets, so let's close/open the stage too
-		if ( ChangedProperties.Contains( GET_MEMBER_NAME_CHECKED( AActor, bHiddenEdLevel ) ) ||
-			 ChangedProperties.Contains( GET_MEMBER_NAME_CHECKED( AActor, bHiddenEdLayer ) ) ||
-			 ChangedProperties.Contains( GET_MEMBER_NAME_CHECKED( AActor, bHiddenEd ) ) )
+		if (ChangedProperties.Contains(GET_MEMBER_NAME_CHECKED(AActor, bHiddenEdLevel)) ||
+			ChangedProperties.Contains(GET_MEMBER_NAME_CHECKED(AActor, bHiddenEdLayer)) ||
+			ChangedProperties.Contains(GET_MEMBER_NAME_CHECKED(AActor, bHiddenEd)))
 		{
-			if ( IsHiddenEd() )
+			if (IsHiddenEd())
 			{
 				CloseUsdStage();
 			}
@@ -2213,40 +2225,40 @@ void AUsdStageActor::PostTransacted(const FTransactionObjectEvent& TransactionEv
 			// very garbled meshes. Here we workaround that by recreating all those render states manually
 			const bool bRecurive = true;
 			GetRootPrimTwin()->Iterate([](UUsdPrimTwin& PrimTwin)
-			{
-				if ( USkinnedMeshComponent* Component = Cast<USkinnedMeshComponent>( PrimTwin.GetSceneComponent() ) )
 				{
-					FRenderStateRecreator RecreateRenderState{ Component };
-				}
-			}, bRecurive);
+					if (USkinnedMeshComponent* Component = Cast<USkinnedMeshComponent>(PrimTwin.GetSceneComponent()))
+					{
+						FRenderStateRecreator RecreateRenderState{Component};
+					}
+				}, bRecurive);
 		}
 	}
 
 	// Fire OnObjectTransacted so that multi-user can track our transactions
-	Super::PostTransacted( TransactionEvent );
+	Super::PostTransacted(TransactionEvent);
 }
 
-void AUsdStageActor::PreEditChange( FProperty* PropertyThatWillChange )
+void AUsdStageActor::PreEditChange(FProperty* PropertyThatWillChange)
 {
 	// If we're just editing some other actor property like Time or anything else, we will get
 	// PostRegister/Unregister calls in the editor due to AActor::PostEditChangeProperty *and* AActor::PreEditChange.
 	// Here we determine in which cases we should ignore those PostRegister/Unregister calls by using the
 	// bIsModifyingAProperty flag
-	if ( !IsActorBeingDestroyed() )
+	if (!IsActorBeingDestroyed())
 	{
-		if ( ( GEditor && GEditor->bIsSimulatingInEditor && GetWorld() != nullptr ) || ReregisterComponentsWhenModified() )
+		if ((GEditor && GEditor->bIsSimulatingInEditor && GetWorld() != nullptr) || ReregisterComponentsWhenModified())
 		{
 			// PreEditChange gets called for actor lifecycle functions too (like if the actor transacts on undo/redo).
 			// In those cases we will have nullptr PropertyThatWillChange, and we don't want to block our PostRegister/Unregister
 			// functions. We only care about blocking the calls triggered by AActor::PostEditChangeProperty and AActor::PreEditChange
-			if ( PropertyThatWillChange )
+			if (PropertyThatWillChange)
 			{
 				bIsModifyingAProperty = true;
 			}
 		}
 	}
 
-	Super::PreEditChange( PropertyThatWillChange );
+	Super::PreEditChange(PropertyThatWillChange);
 }
 
 void AUsdStageActor::PreEditUndo()
@@ -2256,44 +2268,44 @@ void AUsdStageActor::PreEditUndo()
 	Super::PreEditUndo();
 }
 
-void AUsdStageActor::HandleTransactionStateChanged( const FTransactionContext& InTransactionContext, const ETransactionStateEventType InTransactionState )
+void AUsdStageActor::HandleTransactionStateChanged(const FTransactionContext& InTransactionContext, const ETransactionStateEventType InTransactionState)
 {
 	// Hack for solving UE-127253
 	// When we Reload (or open a new stage), we call ReloadAnimations which will close the Sequencer (if opened), recreate our LevelSequence, and get the Sequencer
 	// to show that one instead. If we undo the Reload, that new LevelSequence will be deleted and the Sequencer will be left open trying to display it,
 	// which leads to crashes. Here we try detecting for that case and close/reopen the sequencer to show the correct one.
 #if WITH_EDITOR
-	if ( GIsEditor && GEditor && ( InTransactionState == ETransactionStateEventType::UndoRedoStarted || InTransactionState == ETransactionStateEventType::UndoRedoFinalized ) )
+	if (GIsEditor && GEditor && (InTransactionState == ETransactionStateEventType::UndoRedoStarted || InTransactionState == ETransactionStateEventType::UndoRedoFinalized))
 	{
-		if ( UTransactor* Trans = GEditor->Trans )
+		if (UTransactor* Trans = GEditor->Trans)
 		{
 			static TSet<AUsdStageActor*> ActorsThatClosedTheSequencer;
 
-			int32 CurrentTransactionIndex = Trans->FindTransactionIndex( InTransactionContext.TransactionId );
-			const FTransaction* Transaction = Trans->GetTransaction( CurrentTransactionIndex );
+			int32 CurrentTransactionIndex = Trans->FindTransactionIndex(InTransactionContext.TransactionId);
+			const FTransaction* Transaction = Trans->GetTransaction(CurrentTransactionIndex);
 
-			if ( Transaction )
+			if (Transaction)
 			{
 				TArray<UObject*> TransactionObjects;
-				Transaction->GetTransactionObjects( TransactionObjects );
+				Transaction->GetTransactionObjects(TransactionObjects);
 
-				if ( TransactionObjects.Contains( this ) )
+				if (TransactionObjects.Contains(this))
 				{
-					if ( InTransactionState == ETransactionStateEventType::UndoRedoStarted )
+					if (InTransactionState == ETransactionStateEventType::UndoRedoStarted)
 					{
-						const bool bLevelSequenceEditorWasOpened = GEditor->GetEditorSubsystem<UAssetEditorSubsystem>()->CloseAllEditorsForAsset( LevelSequence ) > 0;
-						if ( bLevelSequenceEditorWasOpened )
+						const bool bLevelSequenceEditorWasOpened = GEditor->GetEditorSubsystem<UAssetEditorSubsystem>()->CloseAllEditorsForAsset(LevelSequence) > 0;
+						if (bLevelSequenceEditorWasOpened)
 						{
-							ActorsThatClosedTheSequencer.Add( this );
+							ActorsThatClosedTheSequencer.Add(this);
 						}
 					}
 
-					if ( InTransactionState == ETransactionStateEventType::UndoRedoFinalized )
+					if (InTransactionState == ETransactionStateEventType::UndoRedoFinalized)
 					{
-						if ( ActorsThatClosedTheSequencer.Contains( this ) )
+						if (ActorsThatClosedTheSequencer.Contains(this))
 						{
-							GEditor->GetEditorSubsystem<UAssetEditorSubsystem>()->OpenEditorForAsset( LevelSequence );
-							ActorsThatClosedTheSequencer.Remove( this );
+							GEditor->GetEditorSubsystem<UAssetEditorSubsystem>()->OpenEditorForAsset(LevelSequence);
+							ActorsThatClosedTheSequencer.Remove(this);
 						}
 					}
 				}
@@ -2302,9 +2314,9 @@ void AUsdStageActor::HandleTransactionStateChanged( const FTransactionContext& I
 	}
 #endif // WITH_EDITOR
 
-	if ( InTransactionState == ETransactionStateEventType::TransactionFinalized ||
-		 InTransactionState == ETransactionStateEventType::UndoRedoFinalized ||
-		 InTransactionState == ETransactionStateEventType::TransactionCanceled )
+	if (InTransactionState == ETransactionStateEventType::TransactionFinalized ||
+		InTransactionState == ETransactionStateEventType::UndoRedoFinalized ||
+		InTransactionState == ETransactionStateEventType::TransactionCanceled)
 	{
 		OldRootLayer = RootLayer;
 	}
@@ -2312,11 +2324,11 @@ void AUsdStageActor::HandleTransactionStateChanged( const FTransactionContext& I
 
 #endif // WITH_EDITOR
 
-void AUsdStageActor::PostDuplicate( bool bDuplicateForPIE )
+void AUsdStageActor::PostDuplicate(bool bDuplicateForPIE)
 {
-	Super::PostDuplicate( bDuplicateForPIE );
+	Super::PostDuplicate(bDuplicateForPIE);
 
-	if ( bDuplicateForPIE )
+	if (bDuplicateForPIE)
 	{
 		OpenUsdStage();
 	}
@@ -2342,14 +2354,14 @@ void AUsdStageActor::Serialize(FArchive& Ar)
 		Ar << bIsTransitioningIntoPIE;
 	}
 
-	if ( ( Ar.GetPortFlags() & PPF_DuplicateForPIE ) || Ar.IsTransacting() )
+	if ((Ar.GetPortFlags() & PPF_DuplicateForPIE) || Ar.IsTransacting())
 	{
-		if ( !InfoCache.IsValid() )
+		if (!InfoCache.IsValid())
 		{
 			InfoCache = MakeShared<FUsdInfoCache>();
 		}
 
-		InfoCache->Serialize( Ar );
+		InfoCache->Serialize(Ar);
 	}
 }
 
@@ -2362,9 +2374,9 @@ void AUsdStageActor::Destroyed()
 	// they will be restored as detached, which we don't want, so here we make sure they are first recorded as attached.
 
 	TArray<AActor*> ChildActors;
-	GetAttachedActors( ChildActors );
+	GetAttachedActors(ChildActors);
 
-	for ( AActor* Child : ChildActors )
+	for (AActor* Child : ChildActors)
 	{
 		Child->Modify();
 	}
@@ -2377,9 +2389,9 @@ void AUsdStageActor::PostActorCreated()
 	Super::PostActorCreated();
 }
 
-void AUsdStageActor::PostRename( UObject* OldOuter, const FName OldName )
+void AUsdStageActor::PostRename(UObject* OldOuter, const FName OldName)
 {
-	Super::PostRename( OldOuter, OldName );
+	Super::PostRename(OldOuter, OldName);
 
 	// Update the binding to this actor on the level sequence. This happens consistently when placing a BP-derived
 	// stage actor with a set root layer onto the stage: We will call ReloadAnimations() before something else calls SetActorLabel()
@@ -2390,20 +2402,20 @@ void AUsdStageActor::PostRename( UObject* OldOuter, const FName OldName )
 void AUsdStageActor::BeginDestroy()
 {
 #if WITH_EDITOR
-	if ( !IsEngineExitRequested() && HasAuthorityOverStage() )
+	if (!IsEngineExitRequested() && HasAuthorityOverStage())
 	{
-		FEditorDelegates::BeginPIE.RemoveAll( this );
-		FEditorDelegates::PostPIEStarted.RemoveAll( this );
-		FUsdDelegates::OnPostUsdImport.RemoveAll( this );
-		FUsdDelegates::OnPreUsdImport.RemoveAll( this );
-		if ( UTransBuffer* TransBuffer = GUnrealEd ? Cast<UTransBuffer>( GUnrealEd->Trans ) : nullptr )
+		FEditorDelegates::BeginPIE.RemoveAll(this);
+		FEditorDelegates::PostPIEStarted.RemoveAll(this);
+		FUsdDelegates::OnPostUsdImport.RemoveAll(this);
+		FUsdDelegates::OnPreUsdImport.RemoveAll(this);
+		if (UTransBuffer* TransBuffer = GUnrealEd ? Cast<UTransBuffer>(GUnrealEd->Trans) : nullptr)
 		{
-			TransBuffer->OnTransactionStateChanged().RemoveAll( this );
-			TransBuffer->OnRedo().Remove( OnRedoHandle );
+			TransBuffer->OnTransactionStateChanged().RemoveAll(this);
+			TransBuffer->OnRedo().Remove(OnRedoHandle);
 		}
 
-		GEngine->OnLevelActorDeleted().RemoveAll( this );
-		FCoreUObjectDelegates::OnObjectsReplaced.RemoveAll( this );
+		GEngine->OnLevelActorDeleted().RemoveAll(this);
+		FCoreUObjectDelegates::OnObjectsReplaced.RemoveAll(this);
 	}
 
 	// This clears the SUSDStage window whenever the level we're currently in gets destroyed.
@@ -2412,7 +2424,7 @@ void AUsdStageActor::BeginDestroy()
 	CloseUsdStage();
 
 	// If our prims are already destroyed then likely the entire map has been destroyed anyway, so don't need to clear it
-	if ( RootUsdTwin && !RootUsdTwin->HasAnyFlags( RF_BeginDestroyed ) )
+	if (RootUsdTwin && !RootUsdTwin->HasAnyFlags(RF_BeginDestroyed))
 	{
 		RootUsdTwin->Clear();
 	}
@@ -2428,45 +2440,45 @@ void AUsdStageActor::PostRegisterAllComponents()
 	// This may fail if our stage happened to not spawn any components, actors or assets, but by that
 	// point "being loaded" doesn't really mean anything anyway
 	const bool bStageIsLoaded = GetBaseUsdStage()
-		&& ( ( RootUsdTwin && RootUsdTwin->GetSceneComponent() != nullptr ) || ( AssetCache && AssetCache->GetNumAssets() > 0 ) );
+		&& ((RootUsdTwin && RootUsdTwin->GetSceneComponent() != nullptr) || (AssetCache && AssetCache->GetNumAssets() > 0));
 
 	// Blocks loading stage when going into PIE, if we already have something loaded (we'll want to duplicate stuff instead).
 	// We need to allow loading when going into PIE when we have nothing loaded yet because the MovieRenderQueue (or other callers)
 	// may directly trigger PIE sessions providing an override world. Without this exception a map saved with a loaded stage
 	// wouldn't load it at all when opening the level in that way
 	UWorld* World = GetWorld();
-	if ( bIsTransitioningIntoPIE && bStageIsLoaded && ( !World || World->WorldType == EWorldType::PIE ) )
+	if (bIsTransitioningIntoPIE && bStageIsLoaded && (!World || World->WorldType == EWorldType::PIE))
 	{
 		return;
 	}
 
 	// We get an inactive world when dragging a ULevel asset
 	// This is just hiding though, so we shouldn't actively load/unload anything
-	if ( !World || World->WorldType == EWorldType::Inactive )
+	if (!World || World->WorldType == EWorldType::Inactive)
 	{
 		return;
 	}
 
 #if WITH_EDITOR
 	// Prevent loading on bHiddenEdLevel because PostRegisterAllComponents gets called in the process of hiding our level, if we're in the persistent level.
-	if ( bIsEditorPreviewActor || bHiddenEdLevel )
+	if (bIsEditorPreviewActor || bHiddenEdLevel)
 	{
 		return;
 	}
 
-	if ( UBlueprintGeneratedClass* BPClass = Cast<UBlueprintGeneratedClass>( GetClass() ) )
+	if (UBlueprintGeneratedClass* BPClass = Cast<UBlueprintGeneratedClass>(GetClass()))
 	{
 		// We can't load stage when recompiling our blueprint because blueprint recompilation is not a transaction. We're forced
 		// to reuse the existing spawned components, actors and prim twins instead ( which we move over on OnObjectsReplaced ), or
 		// we'd get tons of undo/redo bugs.
-		if ( FRecompilationTracker::IsBeingCompiled( Cast<UBlueprint>( BPClass->ClassGeneratedBy ) ) )
+		if (FRecompilationTracker::IsBeingCompiled(Cast<UBlueprint>(BPClass->ClassGeneratedBy)))
 		{
 			return;
 		}
 
 		// For blueprints that derive from the stage actor, any property change on the blueprint preview window will trigger a full
 		// PostRegisterAllComponents. We don't want to reload the stage when e.g. changing the Time property, so we have to return here
-		if ( World && World->WorldType == EWorldType::EditorPreview && bStageIsLoaded )
+		if (World && World->WorldType == EWorldType::EditorPreview && bStageIsLoaded)
 		{
 			return;
 		}
@@ -2477,43 +2489,43 @@ void AUsdStageActor::PostRegisterAllComponents()
 	// default it will become visible shortly after this call. On subsequent postregisters, if our level is invisible there is no point to loading our stage,
 	// as our spawned actors/components should be invisible too
 	ULevel* Level = GetLevel();
-	const bool bIsLevelHidden = !Level || ( !Level->bIsVisible && !Level->bIsAssociatingLevel );
-	if ( bIsLevelHidden )
+	const bool bIsLevelHidden = !Level || (!Level->bIsVisible && !Level->bIsAssociatingLevel);
+	if (bIsLevelHidden)
 	{
 		return;
 	}
 
-	if ( IsTemplate() || bIsModifyingAProperty || bIsUndoRedoing )
+	if (IsTemplate() || bIsModifyingAProperty || bIsUndoRedoing)
 	{
 		return;
 	}
 
 	// Send this before we load the stage so that we know SUSDStage is synced to a potential OnStageChanged broadcast
-	OnActorLoaded.Broadcast( this );
+	OnActorLoaded.Broadcast(this);
 
 	LoadUsdStage();
 }
 
-void AUsdStageActor::UnregisterAllComponents( bool bForReregister )
+void AUsdStageActor::UnregisterAllComponents(bool bForReregister)
 {
-	Super::UnregisterAllComponents( bForReregister );
+	Super::UnregisterAllComponents(bForReregister);
 
-	if ( bForReregister || bIsModifyingAProperty || bIsUndoRedoing )
+	if (bForReregister || bIsModifyingAProperty || bIsUndoRedoing)
 	{
 		return;
 	}
 
 #if WITH_EDITOR
-	if ( bIsEditorPreviewActor )
+	if (bIsEditorPreviewActor)
 	{
 		return;
 	}
 
 	// We can't unload stage when recompiling our blueprint because blueprint recompilation is not a transaction.
 	// After recompiling we will reuse these already spawned actors and assets.
-	if ( UBlueprintGeneratedClass* BPClass = Cast<UBlueprintGeneratedClass>( GetClass() ) )
+	if (UBlueprintGeneratedClass* BPClass = Cast<UBlueprintGeneratedClass>(GetClass()))
 	{
-		if ( FRecompilationTracker::IsBeingCompiled( Cast<UBlueprint>( BPClass->ClassGeneratedBy ) ) )
+		if (FRecompilationTracker::IsBeingCompiled(Cast<UBlueprint>(BPClass->ClassGeneratedBy)))
 		{
 			return;
 		}
@@ -2521,10 +2533,10 @@ void AUsdStageActor::UnregisterAllComponents( bool bForReregister )
 #endif // WITH_EDITOR
 
 	const bool bStageIsLoaded = GetBaseUsdStage()
-		&& ( ( RootUsdTwin && RootUsdTwin->GetSceneComponent() != nullptr ) || ( AssetCache && AssetCache->GetNumAssets() > 0 ) );
+		&& ((RootUsdTwin && RootUsdTwin->GetSceneComponent() != nullptr) || (AssetCache && AssetCache->GetNumAssets() > 0));
 
 	UWorld* World = GetWorld();
-	if ( bIsTransitioningIntoPIE && bStageIsLoaded && ( !World || World->WorldType == EWorldType::PIE ) )
+	if (bIsTransitioningIntoPIE && bStageIsLoaded && (!World || World->WorldType == EWorldType::PIE))
 	{
 		return;
 	}
@@ -2532,12 +2544,12 @@ void AUsdStageActor::UnregisterAllComponents( bool bForReregister )
 	// We get an inactive world when dragging a ULevel asset
 	// Unlike on PostRegister, we still want to unload our stage if our world is nullptr, as that likely means we were in
 	// a sublevel that got unloaded
-	if ( World && World->WorldType == EWorldType::Inactive )
+	if (World && World->WorldType == EWorldType::Inactive)
 	{
 		return;
 	}
 
-	if ( IsTemplate() || IsEngineExitRequested() )
+	if (IsTemplate() || IsEngineExitRequested())
 	{
 		return;
 	}
@@ -2550,10 +2562,10 @@ void AUsdStageActor::PostUnregisterAllComponents()
 	Super::PostUnregisterAllComponents();
 }
 
-void AUsdStageActor::OnPreUsdImport( FString FilePath )
+void AUsdStageActor::OnPreUsdImport(FString FilePath)
 {
-	const UE::FUsdStage& CurrentStage = static_cast< const AUsdStageActor* >( this )->GetUsdStage();
-	if ( !CurrentStage || !HasAuthorityOverStage() )
+	const UE::FUsdStage& CurrentStage = static_cast<const AUsdStageActor*>(this)->GetUsdStage();
+	if (!CurrentStage || !HasAuthorityOverStage())
 	{
 		return;
 	}
@@ -2561,25 +2573,25 @@ void AUsdStageActor::OnPreUsdImport( FString FilePath )
 	// Stop listening to events because a USD import may temporarily modify the stage (e.g. when importing with
 	// a different MetersPerUnit value), and we don't want to respond to the notices in the meantime
 	FString RootPath = CurrentStage.GetRootLayer().GetRealPath();
-	FPaths::NormalizeFilename( RootPath );
-	if ( RootPath == FilePath )
+	FPaths::NormalizeFilename(RootPath);
+	if (RootPath == FilePath)
 	{
 		StopListeningToUsdNotices();
 	}
 }
 
-void AUsdStageActor::OnPostUsdImport( FString FilePath )
+void AUsdStageActor::OnPostUsdImport(FString FilePath)
 {
-	const UE::FUsdStage& CurrentStage = static_cast< const AUsdStageActor* >( this )->GetUsdStage();
-	if ( !CurrentStage || !HasAuthorityOverStage() )
+	const UE::FUsdStage& CurrentStage = static_cast<const AUsdStageActor*>(this)->GetUsdStage();
+	if (!CurrentStage || !HasAuthorityOverStage())
 	{
 		return;
 	}
 
 	// Resume listening to events
 	FString RootPath = CurrentStage.GetRootLayer().GetRealPath();
-	FPaths::NormalizeFilename( RootPath );
-	if ( RootPath == FilePath )
+	FPaths::NormalizeFilename(RootPath);
+	if (RootPath == FilePath)
 	{
 		ResumeListeningToUsdNotices();
 	}
@@ -2618,26 +2630,26 @@ void AUsdStageActor::UpdateSpawnedObjectsTransientFlag(bool bTransient)
 	GetRootPrimTwin()->Iterate(UpdateTransient, bRecursive);
 }
 
-void AUsdStageActor::OnUsdPrimTwinDestroyed( const UUsdPrimTwin& UsdPrimTwin )
+void AUsdStageActor::OnUsdPrimTwinDestroyed(const UUsdPrimTwin& UsdPrimTwin)
 {
-	PrimsToAnimate.Remove( UsdPrimTwin.PrimPath );
+	PrimsToAnimate.Remove(UsdPrimTwin.PrimPath);
 
 	UObject* WatchedObject = UsdPrimTwin.SpawnedActor.IsValid() ? (UObject*)UsdPrimTwin.SpawnedActor.Get() : (UObject*)UsdPrimTwin.SceneComponent.Get();
-	ObjectsToWatch.Remove( WatchedObject );
+	ObjectsToWatch.Remove(WatchedObject);
 
-	LevelSequenceHelper.RemovePrim( UsdPrimTwin );
+	LevelSequenceHelper.RemovePrim(UsdPrimTwin);
 }
 
-void AUsdStageActor::OnObjectPropertyChanged( UObject* ObjectBeingModified, FPropertyChangedEvent& PropertyChangedEvent )
+void AUsdStageActor::OnObjectPropertyChanged(UObject* ObjectBeingModified, FPropertyChangedEvent& PropertyChangedEvent)
 {
-	if ( ObjectBeingModified == this )
+	if (ObjectBeingModified == this)
 	{
-		HandlePropertyChangedEvent( PropertyChangedEvent );
+		HandlePropertyChangedEvent(PropertyChangedEvent);
 		return;
 	}
 
 	// Don't modify the stage if we're in PIE
-	if ( !HasAuthorityOverStage() )
+	if (!HasAuthorityOverStage())
 	{
 		return;
 	}
@@ -2646,16 +2658,16 @@ void AUsdStageActor::OnObjectPropertyChanged( UObject* ObjectBeingModified, FPro
 	// where we're spawning objects and components, traverse our hierarchy and tag everything that needs it.
 	// We avoid the RootLayer change transaction because if we tagged our spawns then the actual spawning would be
 	// replicated, and we want other clients to spawn their own actors and components instead
-	if ( RootLayer.FilePath == OldRootLayer.FilePath && FUsdStageActorImpl::ObjectNeedsMultiUserTag( ObjectBeingModified, this ) )
+	if (RootLayer.FilePath == OldRootLayer.FilePath && FUsdStageActorImpl::ObjectNeedsMultiUserTag(ObjectBeingModified, this))
 	{
 		TSet<UObject*> VisitedObjects;
-		FUsdStageActorImpl::AllowListComponentHierarchy( GetRootComponent(), VisitedObjects );
+		FUsdStageActorImpl::AllowListComponentHierarchy(GetRootComponent(), VisitedObjects);
 	}
 
 	// So that we can detect when the user enables/disables live link properties on a ULiveLinkComponentController that may
 	// be controlling a component that we *do* care about
-	ULiveLinkComponentController* Controller = Cast< ULiveLinkComponentController >( ObjectBeingModified );
-	if ( Controller )
+	ULiveLinkComponentController* Controller = Cast< ULiveLinkComponentController >(ObjectBeingModified);
+	if (Controller)
 	{
 		if (UActorComponent* ControlledComponent = Controller->GetControlledComponent(ULiveLinkTransformRole::StaticClass()))
 		{
@@ -2665,11 +2677,11 @@ void AUsdStageActor::OnObjectPropertyChanged( UObject* ObjectBeingModified, FPro
 
 	UObject* PrimObject = ObjectBeingModified;
 
-	if ( !ObjectsToWatch.Contains( ObjectBeingModified ) )
+	if (!ObjectsToWatch.Contains(ObjectBeingModified))
 	{
-		if ( AActor* ActorBeingModified = Cast< AActor >( ObjectBeingModified ) )
+		if (AActor* ActorBeingModified = Cast< AActor >(ObjectBeingModified))
 		{
-			if ( !ObjectsToWatch.Contains( ActorBeingModified->GetRootComponent() ) )
+			if (!ObjectsToWatch.Contains(ActorBeingModified->GetRootComponent()))
 			{
 				return;
 			}
@@ -2684,65 +2696,65 @@ void AUsdStageActor::OnObjectPropertyChanged( UObject* ObjectBeingModified, FPro
 		}
 	}
 
-	const UE::FUsdStage& CurrentStage = static_cast< const AUsdStageActor* >( this )->GetUsdStage();
+	const UE::FUsdStage& CurrentStage = static_cast<const AUsdStageActor*>(this)->GetUsdStage();
 
-	FString PrimPath = ObjectsToWatch[ PrimObject ];
+	FString PrimPath = ObjectsToWatch[PrimObject];
 
-	if ( UUsdPrimTwin* UsdPrimTwin = GetRootPrimTwin()->Find( PrimPath ) )
+	if (UUsdPrimTwin* UsdPrimTwin = GetRootPrimTwin()->Find(PrimPath))
 	{
 		// Update prim from UE
 		USceneComponent* PrimSceneComponent = UsdPrimTwin->SceneComponent.Get();
 
-		if ( !PrimSceneComponent && UsdPrimTwin->SpawnedActor.IsValid() )
+		if (!PrimSceneComponent && UsdPrimTwin->SpawnedActor.IsValid())
 		{
 			PrimSceneComponent = UsdPrimTwin->SpawnedActor->GetRootComponent();
 		}
 
-		if ( PrimSceneComponent )
+		if (PrimSceneComponent)
 		{
-			if ( CurrentStage )
+			if (CurrentStage)
 			{
 				// This block is important, as it not only prevents us from getting into infinite loops with the USD notices,
 				// but it also guarantees that if we have an object property change, the corresponding stage notice is not also
 				// independently saved to the transaction via the UUsdTransactor, which would be duplication
-				FScopedBlockNoticeListening BlockNotices( this );
+				FScopedBlockNoticeListening BlockNotices(this);
 
-				UE::FUsdPrim UsdPrim = CurrentStage.GetPrimAtPath( UE::FSdfPath( *PrimPath ) );
+				UE::FUsdPrim UsdPrim = CurrentStage.GetPrimAtPath(UE::FSdfPath(*PrimPath));
 
 				// We want to keep component visibilities in sync with USD, which uses inherited visibilities
 				// To accomplish that while blocking notices we must always propagate component visibility changes manually.
 				// This part is effectively the same as calling pxr::UsdGeomImageable::MakeVisible/Invisible.
-				if ( PropertyChangedEvent.GetPropertyName() == TEXT( "bHiddenInGame" ) )
+				if (PropertyChangedEvent.GetPropertyName() == TEXT("bHiddenInGame"))
 				{
 					PrimSceneComponent->Modify();
 
-					if ( PrimSceneComponent->bHiddenInGame )
+					if (PrimSceneComponent->bHiddenInGame)
 					{
-						FUsdStageActorImpl::MakeInvisible( *UsdPrimTwin );
+						FUsdStageActorImpl::MakeInvisible(*UsdPrimTwin);
 					}
 					else
 					{
-						FUsdStageActorImpl::MakeVisible( *UsdPrimTwin, CurrentStage );
+						FUsdStageActorImpl::MakeVisible(*UsdPrimTwin, CurrentStage);
 					}
 				}
 
 #if USE_USD_SDK
 
-				UnrealToUsd::ConvertLiveLinkProperties( Controller ? Cast<UActorComponent>( Controller ) : PrimSceneComponent, UsdPrim );
+				UnrealToUsd::ConvertLiveLinkProperties(Controller ? Cast<UActorComponent>(Controller) : PrimSceneComponent, UsdPrim);
 
-				UnrealToUsd::ConvertSceneComponent( CurrentStage, PrimSceneComponent, UsdPrim );
+				UnrealToUsd::ConvertSceneComponent(CurrentStage, PrimSceneComponent, UsdPrim);
 
-				if ( UMeshComponent* MeshComponent = Cast< UMeshComponent >( PrimSceneComponent ) )
+				if (UMeshComponent* MeshComponent = Cast< UMeshComponent >(PrimSceneComponent))
 				{
-					UnrealToUsd::ConvertMeshComponent( CurrentStage, MeshComponent, UsdPrim );
+					UnrealToUsd::ConvertMeshComponent(CurrentStage, MeshComponent, UsdPrim);
 				}
-				else if ( UsdPrim && UsdPrim.IsA( TEXT( "Camera" ) ) )
+				else if (UsdPrim && UsdPrim.IsA(TEXT("Camera")))
 				{
 					// Our component may be pointing directly at a camera component in case we recreated an exported
 					// ACineCameraActor (see UE-120826)
-					if ( UCineCameraComponent* RecreatedCameraComponent = Cast<UCineCameraComponent>( PrimSceneComponent ) )
+					if (UCineCameraComponent* RecreatedCameraComponent = Cast<UCineCameraComponent>(PrimSceneComponent))
 					{
-						UnrealToUsd::ConvertCameraComponent( *RecreatedCameraComponent, UsdPrim );
+						UnrealToUsd::ConvertCameraComponent(*RecreatedCameraComponent, UsdPrim);
 					}
 					// Or it could have been just a generic Camera prim, at which case we'll have spawned an entire new
 					// ACineCameraActor for it. In this scenario our prim twin is pointing at the root component, so we need
@@ -2750,55 +2762,55 @@ void AUsdStageActor::OnObjectPropertyChanged( UObject* ObjectBeingModified, FPro
 					// We should only do this when the Prim actually corresponds to the Camera though, or else we'll also catch
 					// the prim/component pair that corresponds to the root scene component in case we recreated an exported
 					// ACineCameraActor.
-					else if ( ACineCameraActor* CameraActor = Cast<ACineCameraActor>( PrimSceneComponent->GetOwner() ) )
+					else if (ACineCameraActor* CameraActor = Cast<ACineCameraActor>(PrimSceneComponent->GetOwner()))
 					{
-						if ( UCineCameraComponent* CameraComponent = CameraActor->GetCineCameraComponent() )
+						if (UCineCameraComponent* CameraComponent = CameraActor->GetCineCameraComponent())
 						{
-							UnrealToUsd::ConvertCameraComponent( *CameraComponent, UsdPrim );
+							UnrealToUsd::ConvertCameraComponent(*CameraComponent, UsdPrim);
 						}
 					}
 				}
-				else if ( ALight* LightActor = Cast<ALight>( PrimSceneComponent->GetOwner() ) )
+				else if (ALight* LightActor = Cast<ALight>(PrimSceneComponent->GetOwner()))
 				{
-					if ( ULightComponent* LightComponent = LightActor->GetLightComponent() )
+					if (ULightComponent* LightComponent = LightActor->GetLightComponent())
 					{
-						UnrealToUsd::ConvertLightComponent( *LightComponent, UsdPrim, UsdUtils::GetDefaultTimeCode() );
+						UnrealToUsd::ConvertLightComponent(*LightComponent, UsdPrim, UsdUtils::GetDefaultTimeCode());
 
-						if ( UDirectionalLightComponent* DirectionalLight = Cast<UDirectionalLightComponent>( LightComponent ) )
+						if (UDirectionalLightComponent* DirectionalLight = Cast<UDirectionalLightComponent>(LightComponent))
 						{
-							UnrealToUsd::ConvertDirectionalLightComponent( *DirectionalLight, UsdPrim, UsdUtils::GetDefaultTimeCode() );
+							UnrealToUsd::ConvertDirectionalLightComponent(*DirectionalLight, UsdPrim, UsdUtils::GetDefaultTimeCode());
 						}
-						else if ( URectLightComponent* RectLight = Cast<URectLightComponent>( LightComponent ) )
+						else if (URectLightComponent* RectLight = Cast<URectLightComponent>(LightComponent))
 						{
-							UnrealToUsd::ConvertRectLightComponent( *RectLight, UsdPrim, UsdUtils::GetDefaultTimeCode() );
+							UnrealToUsd::ConvertRectLightComponent(*RectLight, UsdPrim, UsdUtils::GetDefaultTimeCode());
 						}
-						else if ( UPointLightComponent* PointLight = Cast<UPointLightComponent>( LightComponent ) )
+						else if (UPointLightComponent* PointLight = Cast<UPointLightComponent>(LightComponent))
 						{
-							UnrealToUsd::ConvertPointLightComponent( *PointLight, UsdPrim, UsdUtils::GetDefaultTimeCode() );
+							UnrealToUsd::ConvertPointLightComponent(*PointLight, UsdPrim, UsdUtils::GetDefaultTimeCode());
 
-							if ( USpotLightComponent* SpotLight = Cast<USpotLightComponent>( LightComponent ) )
+							if (USpotLightComponent* SpotLight = Cast<USpotLightComponent>(LightComponent))
 							{
-								UnrealToUsd::ConvertSpotLightComponent( *SpotLight, UsdPrim, UsdUtils::GetDefaultTimeCode() );
+								UnrealToUsd::ConvertSpotLightComponent(*SpotLight, UsdPrim, UsdUtils::GetDefaultTimeCode());
 							}
 						}
 					}
 				}
 				// In contrast to the other light types, the USkyLightComponent is the root component of the ASkyLight
-				else if ( USkyLightComponent* SkyLightComponent = Cast<USkyLightComponent>( PrimSceneComponent ) )
+				else if (USkyLightComponent* SkyLightComponent = Cast<USkyLightComponent>(PrimSceneComponent))
 				{
-					UnrealToUsd::ConvertLightComponent( *SkyLightComponent, UsdPrim, UsdUtils::GetDefaultTimeCode() );
-					UnrealToUsd::ConvertSkyLightComponent( *SkyLightComponent, UsdPrim, UsdUtils::GetDefaultTimeCode() );
+					UnrealToUsd::ConvertLightComponent(*SkyLightComponent, UsdPrim, UsdUtils::GetDefaultTimeCode());
+					UnrealToUsd::ConvertSkyLightComponent(*SkyLightComponent, UsdPrim, UsdUtils::GetDefaultTimeCode());
 				}
 #endif // #if USE_USD_SDK
 
 				// Update stage window in case any of our component changes trigger USD stage changes
-				this->OnPrimChanged.Broadcast( PrimPath, false );
+				this->OnPrimChanged.Broadcast(PrimPath, false);
 			}
 		}
 	}
 }
 
-void AUsdStageActor::HandlePropertyChangedEvent( FPropertyChangedEvent& PropertyChangedEvent )
+void AUsdStageActor::HandlePropertyChangedEvent(FPropertyChangedEvent& PropertyChangedEvent)
 {
 	// Handle property changed events with this function (called from our OnObjectPropertyChanged delegate) instead of overriding PostEditChangeProperty because replicated
 	// multi-user transactions directly broadcast OnObjectPropertyChanged on the properties that were changed, instead of making PostEditChangeProperty events.
@@ -2808,49 +2820,49 @@ void AUsdStageActor::HandlePropertyChangedEvent( FPropertyChangedEvent& Property
 	FProperty* PropertyThatChanged = PropertyChangedEvent.MemberProperty;
 	const FName PropertyName = PropertyThatChanged ? PropertyThatChanged->GetFName() : NAME_None;
 
-	if ( PropertyName == GET_MEMBER_NAME_CHECKED( AUsdStageActor, RootLayer ) )
+	if (PropertyName == GET_MEMBER_NAME_CHECKED(AUsdStageActor, RootLayer))
 	{
-		SetRootLayer( RootLayer.FilePath );
+		SetRootLayer(RootLayer.FilePath);
 	}
-	else if ( PropertyName == GET_MEMBER_NAME_CHECKED( AUsdStageActor, Time ) )
+	else if (PropertyName == GET_MEMBER_NAME_CHECKED(AUsdStageActor, Time))
 	{
-		SetTime( Time );
+		SetTime(Time);
 	}
-	else if ( PropertyName == GET_MEMBER_NAME_CHECKED( AUsdStageActor, InitialLoadSet ) )
+	else if (PropertyName == GET_MEMBER_NAME_CHECKED(AUsdStageActor, InitialLoadSet))
 	{
-		SetInitialLoadSet( InitialLoadSet );
+		SetInitialLoadSet(InitialLoadSet);
 	}
-	else if ( PropertyName == GET_MEMBER_NAME_CHECKED( AUsdStageActor, InterpolationType ) )
+	else if (PropertyName == GET_MEMBER_NAME_CHECKED(AUsdStageActor, InterpolationType))
 	{
-		SetInterpolationType( InterpolationType );
+		SetInterpolationType(InterpolationType);
 	}
-	else if ( PropertyName == GET_MEMBER_NAME_CHECKED( AUsdStageActor, KindsToCollapse ) )
+	else if (PropertyName == GET_MEMBER_NAME_CHECKED(AUsdStageActor, KindsToCollapse))
 	{
-		SetKindsToCollapse( KindsToCollapse );
+		SetKindsToCollapse(KindsToCollapse);
 	}
-	else if ( PropertyName == GET_MEMBER_NAME_CHECKED( AUsdStageActor, bMergeIdenticalMaterialSlots ) )
+	else if (PropertyName == GET_MEMBER_NAME_CHECKED(AUsdStageActor, bMergeIdenticalMaterialSlots))
 	{
-		SetMergeIdenticalMaterialSlots( bMergeIdenticalMaterialSlots );
+		SetMergeIdenticalMaterialSlots(bMergeIdenticalMaterialSlots);
 	}
-	else if ( PropertyName == GET_MEMBER_NAME_CHECKED( AUsdStageActor, PurposesToLoad ) )
+	else if (PropertyName == GET_MEMBER_NAME_CHECKED(AUsdStageActor, PurposesToLoad))
 	{
-		SetPurposesToLoad( PurposesToLoad );
+		SetPurposesToLoad(PurposesToLoad);
 	}
-	else if ( PropertyName == GET_MEMBER_NAME_CHECKED( AUsdStageActor, NaniteTriangleThreshold ) )
+	else if (PropertyName == GET_MEMBER_NAME_CHECKED(AUsdStageActor, NaniteTriangleThreshold))
 	{
-		SetNaniteTriangleThreshold( NaniteTriangleThreshold );
+		SetNaniteTriangleThreshold(NaniteTriangleThreshold);
 	}
-	else if ( PropertyName == GET_MEMBER_NAME_CHECKED( AUsdStageActor, RenderContext ) )
+	else if (PropertyName == GET_MEMBER_NAME_CHECKED(AUsdStageActor, RenderContext))
 	{
-		SetRenderContext( RenderContext );
+		SetRenderContext(RenderContext);
 	}
-	else if ( PropertyName == GET_MEMBER_NAME_CHECKED( AUsdStageActor, MaterialPurpose ) )
+	else if (PropertyName == GET_MEMBER_NAME_CHECKED(AUsdStageActor, MaterialPurpose))
 	{
-		SetMaterialPurpose( MaterialPurpose );
+		SetMaterialPurpose(MaterialPurpose);
 	}
-	else if ( PropertyName == GET_MEMBER_NAME_CHECKED( AUsdStageActor, RootMotionHandling ) )
+	else if (PropertyName == GET_MEMBER_NAME_CHECKED(AUsdStageActor, RootMotionHandling))
 	{
-		SetRootMotionHandling( RootMotionHandling );
+		SetRootMotionHandling(RootMotionHandling);
 	}
 
 	bIsModifyingAProperty = false;
@@ -2859,57 +2871,57 @@ void AUsdStageActor::HandlePropertyChangedEvent( FPropertyChangedEvent& Property
 bool AUsdStageActor::HasAuthorityOverStage() const
 {
 #if WITH_EDITOR
-	if ( GIsEditor ) // Don't check for world in Standalone: The game world is the only one there, so it's OK if we have authority while in it
+	if (GIsEditor) // Don't check for world in Standalone: The game world is the only one there, so it's OK if we have authority while in it
 	{
 		// In the editor we have to prevent actors in PIE worlds from having authority
-		return !IsTemplate() && ( !GetWorld() || !GetWorld()->IsGameWorld() );
+		return !IsTemplate() && (!GetWorld() || !GetWorld()->IsGameWorld());
 	}
 #endif // WITH_EDITOR
 
 	return !IsTemplate();
 }
 
-void AUsdStageActor::OnSkelAnimationBaked( const FString& SkelRootPrimPath )
+void AUsdStageActor::OnSkelAnimationBaked(const FString& SkelRootPrimPath)
 {
 #if USE_USD_SDK
-	const UE::FUsdStage& CurrentStage = static_cast< const AUsdStageActor* >( this )->GetUsdStage();
-	if ( !CurrentStage || !GRegenerateSkeletalAssetsOnControlRigBake )
+	const UE::FUsdStage& CurrentStage = static_cast<const AUsdStageActor*>(this)->GetUsdStage();
+	if (!CurrentStage || !GRegenerateSkeletalAssetsOnControlRigBake)
 	{
 		return;
 	}
 
-	UE::FUsdPrim SkelRootPrim = CurrentStage.GetPrimAtPath( UE::FSdfPath{ *SkelRootPrimPath } );
-	if ( !SkelRootPrim || !SkelRootPrim.IsA( TEXT( "SkelRoot" ) ) )
+	UE::FUsdPrim SkelRootPrim = CurrentStage.GetPrimAtPath(UE::FSdfPath{*SkelRootPrimPath});
+	if (!SkelRootPrim || !SkelRootPrim.IsA(TEXT("SkelRoot")))
 	{
 		return;
 	}
 
 	UUsdPrimTwin* RootTwin = GetRootPrimTwin();
-	if ( !RootTwin )
+	if (!RootTwin)
 	{
 		return;
 	}
 
 	UUsdPrimTwin* Twin = RootTwin->Find(SkelRootPrimPath);
-	if ( !Twin )
+	if (!Twin)
 	{
 		return;
 	}
 
-	USkeletalMeshComponent* SkeletalMeshComponent = Cast<USkeletalMeshComponent>( Twin->GetSceneComponent() );
-	if ( !SkeletalMeshComponent )
+	USkeletalMeshComponent* SkeletalMeshComponent = Cast<USkeletalMeshComponent>(Twin->GetSceneComponent());
+	if (!SkeletalMeshComponent)
 	{
 		return;
 	}
 
-	TSharedRef< FUsdSchemaTranslationContext > TranslationContext = FUsdStageActorImpl::CreateUsdSchemaTranslationContext( this, SkelRootPrimPath );
+	TSharedRef< FUsdSchemaTranslationContext > TranslationContext = FUsdStageActorImpl::CreateUsdSchemaTranslationContext(this, SkelRootPrimPath);
 	// The only way we could have baked a skel animation is via the sequencer, so we know its playing
 	TranslationContext->bSequencerIsAnimating = true;
 
-	IUsdSchemasModule& UsdSchemasModule = FModuleManager::Get().LoadModuleChecked< IUsdSchemasModule >( TEXT( "USDSchemas" ) );
-	if ( TSharedPtr< FUsdSchemaTranslator > SchemaTranslator = UsdSchemasModule.GetTranslatorRegistry().CreateTranslatorForSchema( TranslationContext, UE::FUsdTyped( SkelRootPrim ) ) )
+	IUsdSchemasModule& UsdSchemasModule = FModuleManager::Get().LoadModuleChecked< IUsdSchemasModule >(TEXT("USDSchemas"));
+	if (TSharedPtr< FUsdSchemaTranslator > SchemaTranslator = UsdSchemasModule.GetTranslatorRegistry().CreateTranslatorForSchema(TranslationContext, UE::FUsdTyped(SkelRootPrim)))
 	{
-		if ( TSharedPtr< FUsdSkelRootTranslator > SkelRootTranslator = StaticCastSharedPtr<FUsdSkelRootTranslator>( SchemaTranslator ) )
+		if (TSharedPtr< FUsdSkelRootTranslator > SkelRootTranslator = StaticCastSharedPtr<FUsdSkelRootTranslator>(SchemaTranslator))
 		{
 			// For now we're regenerating all asset types (including skeletal meshes) but we could
 			// eventually just split off the anim sequence generation and call exclusively that from
@@ -2918,78 +2930,78 @@ void AUsdStageActor::OnSkelAnimationBaked( const FString& SkelRootPrimPath )
 			TranslationContext->CompleteTasks();
 
 			// Have to update the components to assign the new assets
-			SkelRootTranslator->UpdateComponents( SkeletalMeshComponent );
+			SkelRootTranslator->UpdateComponents(SkeletalMeshComponent);
 		}
 	}
 #endif // #if USE_USD_SDK
 }
 
-void AUsdStageActor::LoadAsset( FUsdSchemaTranslationContext& TranslationContext, const UE::FUsdPrim& Prim )
+void AUsdStageActor::LoadAsset(FUsdSchemaTranslationContext& TranslationContext, const UE::FUsdPrim& Prim)
 {
-	TRACE_CPUPROFILER_EVENT_SCOPE( AUsdStageActor::LoadAsset );
+	TRACE_CPUPROFILER_EVENT_SCOPE(AUsdStageActor::LoadAsset);
 
 	// Mark the assets as non transactional so that they don't get serialized in the transaction buffer
-	TGuardValue< EObjectFlags > ContextFlagsGuard( TranslationContext.ObjectFlags, TranslationContext.ObjectFlags & ~RF_Transactional );
+	TGuardValue< EObjectFlags > ContextFlagsGuard(TranslationContext.ObjectFlags, TranslationContext.ObjectFlags & ~RF_Transactional);
 
 	FString PrimPath;
 #if USE_USD_SDK
-	PrimPath = UsdToUnreal::ConvertPath( Prim.GetPrimPath() );
+	PrimPath = UsdToUnreal::ConvertPath(Prim.GetPrimPath());
 #endif // #if USE_USD_SDK
 
-	if ( AssetCache )
+	if (AssetCache)
 	{
-		AssetCache->RemoveAssetPrimLink( PrimPath );
+		AssetCache->RemoveAssetPrimLink(PrimPath);
 	}
 
-	IUsdSchemasModule& UsdSchemasModule = FModuleManager::Get().LoadModuleChecked< IUsdSchemasModule >( TEXT("USDSchemas") );
-	if ( TSharedPtr< FUsdSchemaTranslator > SchemaTranslator = UsdSchemasModule.GetTranslatorRegistry().CreateTranslatorForSchema( TranslationContext.AsShared(), UE::FUsdTyped( Prim ) ) )
+	IUsdSchemasModule& UsdSchemasModule = FModuleManager::Get().LoadModuleChecked< IUsdSchemasModule >(TEXT("USDSchemas"));
+	if (TSharedPtr< FUsdSchemaTranslator > SchemaTranslator = UsdSchemasModule.GetTranslatorRegistry().CreateTranslatorForSchema(TranslationContext.AsShared(), UE::FUsdTyped(Prim)))
 	{
-		TRACE_CPUPROFILER_EVENT_SCOPE( AUsdStageActor::CreateAssetsForPrim );
+		TRACE_CPUPROFILER_EVENT_SCOPE(AUsdStageActor::CreateAssetsForPrim);
 		SchemaTranslator->CreateAssets();
 	}
 
 	TranslationContext.CompleteTasks(); // Finish the asset tasks before moving on
 }
 
-void AUsdStageActor::LoadAssets( FUsdSchemaTranslationContext& TranslationContext, const UE::FUsdPrim& StartPrim )
+void AUsdStageActor::LoadAssets(FUsdSchemaTranslationContext& TranslationContext, const UE::FUsdPrim& StartPrim)
 {
-	TRACE_CPUPROFILER_EVENT_SCOPE( AUsdStageActor::LoadAssets );
+	TRACE_CPUPROFILER_EVENT_SCOPE(AUsdStageActor::LoadAssets);
 
 	// Mark the assets as non transactional so that they don't get serialized in the transaction buffer
-	TGuardValue< EObjectFlags > ContextFlagsGuard( TranslationContext.ObjectFlags, TranslationContext.ObjectFlags & ~RF_Transactional );
+	TGuardValue< EObjectFlags > ContextFlagsGuard(TranslationContext.ObjectFlags, TranslationContext.ObjectFlags & ~RF_Transactional);
 
 	// Clear existing prim/asset association
-	if ( AssetCache )
+	if (AssetCache)
 	{
 		FString StartPrimPath = StartPrim.GetPrimPath().GetString();
 		TSet<FString> PrimPathsToRemove;
-		for ( const TPair< FString, TWeakObjectPtr<UObject> >& PrimPathToAssetIt : AssetCache->GetAssetPrimLinks() )
+		for (const TPair< FString, TWeakObjectPtr<UObject> >& PrimPathToAssetIt : AssetCache->GetAssetPrimLinks())
 		{
 			const FString& PrimPath = PrimPathToAssetIt.Key;
-			if ( PrimPath.StartsWith( StartPrimPath ) || PrimPath == StartPrimPath )
+			if (PrimPath.StartsWith(StartPrimPath) || PrimPath == StartPrimPath)
 			{
-				PrimPathsToRemove.Add( PrimPath );
+				PrimPathsToRemove.Add(PrimPath);
 			}
 		}
-		for ( const FString& PrimPathToRemove : PrimPathsToRemove )
+		for (const FString& PrimPathToRemove : PrimPathsToRemove)
 		{
-			AssetCache->RemoveAssetPrimLink( PrimPathToRemove );
+			AssetCache->RemoveAssetPrimLink(PrimPathToRemove);
 		}
 	}
 
-	IUsdSchemasModule& UsdSchemasModule = FModuleManager::Get().LoadModuleChecked< IUsdSchemasModule >( TEXT("USDSchemas") );
+	IUsdSchemasModule& UsdSchemasModule = FModuleManager::Get().LoadModuleChecked< IUsdSchemasModule >(TEXT("USDSchemas"));
 
-	auto CreateAssetsForPrims = [ &UsdSchemasModule, &TranslationContext ]( const TArray< UE::FUsdPrim >& AllPrimAssets, FSlowTask& Progress )
+	auto CreateAssetsForPrims = [&UsdSchemasModule, &TranslationContext](const TArray< UE::FUsdPrim >& AllPrimAssets, FSlowTask& Progress)
 	{
-		TRACE_CPUPROFILER_EVENT_SCOPE( AUsdStageActor::CreateAssetsForPrims );
+		TRACE_CPUPROFILER_EVENT_SCOPE(AUsdStageActor::CreateAssetsForPrims);
 
-		for ( const UE::FUsdPrim& UsdPrim : AllPrimAssets )
+		for (const UE::FUsdPrim& UsdPrim : AllPrimAssets)
 		{
 			Progress.EnterProgressFrame(1.f);
 
-			if ( TSharedPtr< FUsdSchemaTranslator > SchemaTranslator = UsdSchemasModule.GetTranslatorRegistry().CreateTranslatorForSchema( TranslationContext.AsShared(), UE::FUsdTyped( UsdPrim ) ) )
+			if (TSharedPtr< FUsdSchemaTranslator > SchemaTranslator = UsdSchemasModule.GetTranslatorRegistry().CreateTranslatorForSchema(TranslationContext.AsShared(), UE::FUsdTyped(UsdPrim)))
 			{
-				TRACE_CPUPROFILER_EVENT_SCOPE( AUsdStageActor::CreateAssetsForPrim );
+				TRACE_CPUPROFILER_EVENT_SCOPE(AUsdStageActor::CreateAssetsForPrim);
 				SchemaTranslator->CreateAssets();
 			}
 		}
@@ -2997,65 +3009,67 @@ void AUsdStageActor::LoadAssets( FUsdSchemaTranslationContext& TranslationContex
 		TranslationContext.CompleteTasks(); // Finish the assets tasks before moving on
 	};
 
-	auto PruneChildren = [ &UsdSchemasModule, &TranslationContext ]( const UE::FUsdPrim& UsdPrim ) -> bool
+	auto PruneChildren = [&UsdSchemasModule, &TranslationContext](const UE::FUsdPrim& UsdPrim) -> bool
 	{
-		if ( TSharedPtr< FUsdSchemaTranslator > SchemaTranslator = UsdSchemasModule.GetTranslatorRegistry().CreateTranslatorForSchema( TranslationContext.AsShared(), UE::FUsdTyped( UsdPrim ) ) )
+		if (TSharedPtr< FUsdSchemaTranslator > SchemaTranslator = UsdSchemasModule.GetTranslatorRegistry().CreateTranslatorForSchema(TranslationContext.AsShared(), UE::FUsdTyped(UsdPrim)))
 		{
-			return SchemaTranslator->CollapsesChildren( ECollapsingType::Assets );
+			return SchemaTranslator->CollapsesChildren(ECollapsingType::Assets);
 		}
 
 		return false;
 	};
 
 	// Load materials first since meshes are referencing them
-	TArray< UE::FUsdPrim > AllPrimAssets = UsdUtils::GetAllPrimsOfType( StartPrim, TEXT( "UsdShadeMaterial" ) );
+	TArray< UE::FUsdPrim > AllPrimAssets = UsdUtils::GetAllPrimsOfType(StartPrim, TEXT("UsdShadeMaterial"));
 	{
-		FScopedSlowTask MaterialsProgress( AllPrimAssets.Num(), LOCTEXT("CreateMaterials", "Creating materials"));
-		CreateAssetsForPrims( AllPrimAssets, MaterialsProgress );
+		FScopedSlowTask MaterialsProgress(AllPrimAssets.Num(), LOCTEXT("CreateMaterials", "Creating materials"));
+		CreateAssetsForPrims(AllPrimAssets, MaterialsProgress);
 	}
 
 	// Load everything else (including meshes)
-	AllPrimAssets = UsdUtils::GetAllPrimsOfType( StartPrim, TEXT( "UsdSchemaBase" ), PruneChildren, { TEXT( "UsdShadeMaterial" ) } );
+	AllPrimAssets = UsdUtils::GetAllPrimsOfType(StartPrim, TEXT("UsdSchemaBase"), PruneChildren, {TEXT("UsdShadeMaterial")});
 	{
-		FScopedSlowTask AssetsProgress( AllPrimAssets.Num(), LOCTEXT("CreateAssets", "Creating assets"));
-		CreateAssetsForPrims( AllPrimAssets, AssetsProgress );
+		FScopedSlowTask AssetsProgress(AllPrimAssets.Num(), LOCTEXT("CreateAssets", "Creating assets"));
+		CreateAssetsForPrims(AllPrimAssets, AssetsProgress);
 	}
 }
 
 void AUsdStageActor::AnimatePrims()
 {
+	TRACE_CPUPROFILER_EVENT_SCOPE(AUsdStageActor::AnimatePrims);
+
 	// Don't try to animate if we don't have a stage opened
-	const UE::FUsdStage& CurrentStage = static_cast< const AUsdStageActor* >( this )->GetUsdStage();
-	if ( !CurrentStage )
+	const UE::FUsdStage& CurrentStage = static_cast<const AUsdStageActor*>(this)->GetUsdStage();
+	if (!CurrentStage)
 	{
 		return;
 	}
 
-	TSharedRef< FUsdSchemaTranslationContext > TranslationContext = FUsdStageActorImpl::CreateUsdSchemaTranslationContext( this, GetRootPrimTwin()->PrimPath );
+	TSharedRef< FUsdSchemaTranslationContext > TranslationContext = FUsdStageActorImpl::CreateUsdSchemaTranslationContext(this, GetRootPrimTwin()->PrimPath);
 
 	// c.f. comment on bSequencerIsAnimating's declaration
 #if WITH_EDITOR
-    if ( GEditor )
-    {
-        const bool bFocusIfOpen = false;
-        IAssetEditorInstance* AssetEditor = GEditor->GetEditorSubsystem<UAssetEditorSubsystem>()->FindEditorForAsset( LevelSequence, bFocusIfOpen );
-        if ( ILevelSequenceEditorToolkit* LevelSequenceEditor = static_cast< ILevelSequenceEditorToolkit* >( AssetEditor ) )
-        {
-            TranslationContext->bSequencerIsAnimating = true;
-        }
-    }
+	if (GEditor)
+	{
+		const bool bFocusIfOpen = false;
+		IAssetEditorInstance* AssetEditor = GEditor->GetEditorSubsystem<UAssetEditorSubsystem>()->FindEditorForAsset(LevelSequence, bFocusIfOpen);
+		if (ILevelSequenceEditorToolkit* LevelSequenceEditor = static_cast<ILevelSequenceEditorToolkit*>(AssetEditor))
+		{
+			TranslationContext->bSequencerIsAnimating = true;
+		}
+	}
 #endif // WITH_EDITOR
 
-	for ( const FString& PrimToAnimate : PrimsToAnimate )
+	for (const FString& PrimToAnimate : PrimsToAnimate)
 	{
-		UE::FSdfPath PrimPath( *PrimToAnimate );
+		UE::FSdfPath PrimPath(*PrimToAnimate);
 
-		IUsdSchemasModule& SchemasModule = FModuleManager::Get().LoadModuleChecked< IUsdSchemasModule >( "USDSchemas" );
-		if ( TSharedPtr< FUsdSchemaTranslator > SchemaTranslator = SchemasModule.GetTranslatorRegistry().CreateTranslatorForSchema( TranslationContext, UE::FUsdTyped( CurrentStage.GetPrimAtPath( PrimPath ) ) ) )
+		IUsdSchemasModule& SchemasModule = FModuleManager::Get().LoadModuleChecked< IUsdSchemasModule >("USDSchemas");
+		if (TSharedPtr< FUsdSchemaTranslator > SchemaTranslator = SchemasModule.GetTranslatorRegistry().CreateTranslatorForSchema(TranslationContext, UE::FUsdTyped(CurrentStage.GetPrimAtPath(PrimPath))))
 		{
-			if ( UUsdPrimTwin* UsdPrimTwin = GetRootPrimTwin()->Find( PrimToAnimate ) )
+			if (UUsdPrimTwin* UsdPrimTwin = GetRootPrimTwin()->Find(PrimToAnimate))
 			{
-				SchemaTranslator->UpdateComponents( UsdPrimTwin->SceneComponent.Get() );
+				SchemaTranslator->UpdateComponents(UsdPrimTwin->SceneComponent.Get());
 			}
 		}
 	}
@@ -3063,7 +3077,7 @@ void AUsdStageActor::AnimatePrims()
 	TranslationContext->CompleteTasks();
 
 #if WITH_EDITOR
-	if ( GIsEditor && GEditor && !IsGarbageCollecting() )
+	if (GIsEditor && GEditor && !IsGarbageCollecting())
 	{
 		GEditor->BroadcastLevelActorListChanged();
 		GEditor->RedrawLevelEditingViewports();
@@ -3071,10 +3085,10 @@ void AUsdStageActor::AnimatePrims()
 #endif // WITH_EDITOR
 }
 
-FScopedBlockNoticeListening::FScopedBlockNoticeListening( AUsdStageActor* InStageActor )
+FScopedBlockNoticeListening::FScopedBlockNoticeListening(AUsdStageActor* InStageActor)
 {
 	StageActor = InStageActor;
-	if ( InStageActor )
+	if (InStageActor)
 	{
 		StageActor->StopListeningToUsdNotices();
 	}
@@ -3082,7 +3096,7 @@ FScopedBlockNoticeListening::FScopedBlockNoticeListening( AUsdStageActor* InStag
 
 FScopedBlockNoticeListening::~FScopedBlockNoticeListening()
 {
-	if ( AUsdStageActor* StageActorPtr = StageActor.Get() )
+	if (AUsdStageActor* StageActorPtr = StageActor.Get())
 	{
 		StageActorPtr->ResumeListeningToUsdNotices();
 	}
