USD: Fixed crash when trying to reference in the stage an Unreal material that doesn't exist.

diff --git a/Engine/Plugins/Importers/USDImporter/Source/USDSchemas/Private/MeshTranslationImpl.cpp b/Engine/Plugins/Importers/USDImporter/Source/USDSchemas/Private/MeshTranslationImpl.cpp
--- a/Engine/Plugins/Importers/USDImporter/Source/USDSchemas/Private/MeshTranslationImpl.cpp
+++ b/Engine/Plugins/Importers/USDImporter/Source/USDSchemas/Private/MeshTranslationImpl.cpp
@@ -21,12 +21,12 @@
 #if USE_USD_SDK
 
 #include "USDIncludesStart.h"
-	#include "pxr/usd/usd/editContext.h"
-	#include "pxr/usd/usd/prim.h"
-	#include "pxr/usd/usd/stage.h"
-	#include "pxr/usd/usdGeom/mesh.h"
-	#include "pxr/usd/usdGeom/xformable.h"
-	#include "pxr/usd/usdShade/material.h"
+#include "pxr/usd/usd/editContext.h"
+#include "pxr/usd/usd/prim.h"
+#include "pxr/usd/usd/stage.h"
+#include "pxr/usd/usdGeom/mesh.h"
+#include "pxr/usd/usdGeom/xformable.h"
+#include "pxr/usd/usdShade/material.h"
 #include "USDIncludesEnd.h"
 
 TMap<const UsdUtils::FUsdPrimMaterialSlot*, UMaterialInterface*> MeshTranslationImpl::ResolveMaterialAssignmentInfo(
@@ -41,47 +41,47 @@ TMap<const UsdUtils::FUsdPrimMaterialSlot*, UMaterialInterface*> MeshTranslation
 	TMap<const UsdUtils::FUsdPrimMaterialSlot*, UMaterialInterface*> ResolvedMaterials;
 
 	uint32 GlobalResolvedMaterialIndex = 0;
-	for ( int32 InfoIndex = 0; InfoIndex < AssignmentInfo.Num(); ++InfoIndex )
+	for (int32 InfoIndex = 0; InfoIndex < AssignmentInfo.Num(); ++InfoIndex)
 	{
-		const TArray< UsdUtils::FUsdPrimMaterialSlot >& Slots = AssignmentInfo[ InfoIndex ].Slots;
+		const TArray< UsdUtils::FUsdPrimMaterialSlot >& Slots = AssignmentInfo[InfoIndex].Slots;
 
-		for ( int32 SlotIndex = 0; SlotIndex < Slots.Num(); ++SlotIndex, ++GlobalResolvedMaterialIndex )
+		for (int32 SlotIndex = 0; SlotIndex < Slots.Num(); ++SlotIndex, ++GlobalResolvedMaterialIndex)
 		{
-			const UsdUtils::FUsdPrimMaterialSlot& Slot = Slots[ SlotIndex ];
+			const UsdUtils::FUsdPrimMaterialSlot& Slot = Slots[SlotIndex];
 			UMaterialInterface* Material = nullptr;
 
-			switch ( Slot.AssignmentType )
+			switch (Slot.AssignmentType)
 			{
 			case UsdUtils::EPrimAssignmentType::DisplayColor:
 			{
 				// Try reusing an already created DisplayColor material
-				if ( UMaterialInterface* ExistingMaterial = Cast<UMaterialInterface>( AssetCache.GetCachedAsset( Slot.MaterialSource ) ) )
+				if (UMaterialInterface* ExistingMaterial = Cast<UMaterialInterface>(AssetCache.GetCachedAsset(Slot.MaterialSource)))
 				{
 					Material = ExistingMaterial;
 				}
 
 				// Need to create a new DisplayColor material
-				if ( Material == nullptr )
+				if (Material == nullptr)
 				{
-					if ( TOptional< UsdUtils::FDisplayColorMaterial > DisplayColorDesc = UsdUtils::FDisplayColorMaterial::FromString( Slot.MaterialSource ) )
+					if (TOptional< UsdUtils::FDisplayColorMaterial > DisplayColorDesc = UsdUtils::FDisplayColorMaterial::FromString(Slot.MaterialSource))
 					{
 						UMaterialInstance* MaterialInstance = nullptr;
 
-						if ( GIsEditor )  // Editor, PIE => true; Standlone, packaged => false
+						if (GIsEditor)  // Editor, PIE => true; Standlone, packaged => false
 						{
-							MaterialInstance = UsdUtils::CreateDisplayColorMaterialInstanceConstant( DisplayColorDesc.GetValue() );
+							MaterialInstance = UsdUtils::CreateDisplayColorMaterialInstanceConstant(DisplayColorDesc.GetValue());
 #if WITH_EDITOR
 							// Leave PrimPath as empty as it likely will be reused by many prims
-							UUsdAssetImportData* ImportData = NewObject< UUsdAssetImportData >( MaterialInstance, TEXT( "USDAssetImportData" ) );
+							UUsdAssetImportData* ImportData = NewObject< UUsdAssetImportData >(MaterialInstance, TEXT("USDAssetImportData"));
 							MaterialInstance->AssetImportData = ImportData;
 #endif // WITH_EDITOR
 						}
 						else
 						{
-							MaterialInstance = UsdUtils::CreateDisplayColorMaterialInstanceDynamic( DisplayColorDesc.GetValue() );
+							MaterialInstance = UsdUtils::CreateDisplayColorMaterialInstanceDynamic(DisplayColorDesc.GetValue());
 						}
 
-						AssetCache.CacheAsset( Slot.MaterialSource, MaterialInstance );
+						AssetCache.CacheAsset(Slot.MaterialSource, MaterialInstance);
 						Material = MaterialInstance;
 					}
 				}
@@ -91,22 +91,22 @@ TMap<const UsdUtils::FUsdPrimMaterialSlot*, UMaterialInterface*> MeshTranslation
 			case UsdUtils::EPrimAssignmentType::MaterialPrim:
 			{
 				FString PrimPath = Slot.MaterialSource;
-				const static FString TwoSidedToken = TEXT( "!twosided" );
-				if ( Slot.bMeshIsDoubleSided )
+				const static FString TwoSidedToken = TEXT("!twosided");
+				if (Slot.bMeshIsDoubleSided)
 				{
 					PrimPath += TwoSidedToken;
 				}
 
-				Material = Cast< UMaterialInterface >( AssetCache.GetAssetForPrim( PrimPath ) );
+				Material = Cast< UMaterialInterface >(AssetCache.GetAssetForPrim(PrimPath));
 
 				// Need to create a two-sided material on-demand
-				if ( !Material && Slot.bMeshIsDoubleSided )
+				if (!Material && Slot.bMeshIsDoubleSided)
 				{
 					// By now we parsed all materials so we must have the single-sided version of this material
 					UMaterialInstance* OneSidedMat = Cast< UMaterialInstance >(
-						AssetCache.GetAssetForPrim( Slot.MaterialSource )
-					);
-					if ( !ensure( OneSidedMat ) )
+						AssetCache.GetAssetForPrim(Slot.MaterialSource)
+						);
+					if (!ensure(OneSidedMat))
 					{
 						continue;
 					}
@@ -115,8 +115,8 @@ TMap<const UsdUtils::FUsdPrimMaterialSlot*, UMaterialInterface*> MeshTranslation
 					// get the reference UsdPreviewSurface instead, as that is also *its* reference
 					UMaterialInterface* ReferenceMaterial = OneSidedMat->Parent.Get();
 					UMaterialInterface* ReferenceMaterialTwoSided =
-						MeshTranslationImpl::GetTwoSidedVersionOfReferencePreviewSurfaceMaterial( ReferenceMaterial );
-					if ( !ensure( ReferenceMaterialTwoSided && ReferenceMaterialTwoSided != ReferenceMaterial ) )
+						MeshTranslationImpl::GetTwoSidedVersionOfReferencePreviewSurfaceMaterial(ReferenceMaterial);
+					if (!ensure(ReferenceMaterialTwoSided && ReferenceMaterialTwoSided != ReferenceMaterial))
 					{
 						continue;
 					}
@@ -124,67 +124,73 @@ TMap<const UsdUtils::FUsdPrimMaterialSlot*, UMaterialInterface*> MeshTranslation
 					const FName NewInstanceName = MakeUniqueObjectName(
 						GetTransientPackage(),
 						UMaterialInstance::StaticClass(),
-						*( FPaths::GetBaseFilename( Slot.MaterialSource ) + UnrealIdentifiers::TwoSidedMaterialSuffix )
+						*(FPaths::GetBaseFilename(Slot.MaterialSource) + UnrealIdentifiers::TwoSidedMaterialSuffix)
 					);
 
 #if WITH_EDITOR
-					UMaterialInstanceConstant* MIC = Cast<UMaterialInstanceConstant>( OneSidedMat );
-					if ( GIsEditor && MIC )
+					UMaterialInstanceConstant* MIC = Cast<UMaterialInstanceConstant>(OneSidedMat);
+					if (GIsEditor && MIC)
 					{
 						UMaterialInstanceConstant* TwoSidedMat = NewObject<UMaterialInstanceConstant>(
 							GetTransientPackage(),
 							NewInstanceName,
 							Flags
-						);
-						if ( TwoSidedMat )
+							);
+						if (TwoSidedMat)
 						{
 							UUsdAssetImportData* ImportData = NewObject< UUsdAssetImportData >(
 								TwoSidedMat,
-								TEXT( "USDAssetImportData" )
-							);
+								TEXT("USDAssetImportData")
+								);
 							ImportData->PrimPath = Slot.MaterialSource;
 							TwoSidedMat->AssetImportData = ImportData;
 						}
 
-						TwoSidedMat->SetParentEditorOnly( ReferenceMaterialTwoSided );
-						TwoSidedMat->CopyMaterialUniformParametersEditorOnly( OneSidedMat );
+						TwoSidedMat->SetParentEditorOnly(ReferenceMaterialTwoSided);
+						TwoSidedMat->CopyMaterialUniformParametersEditorOnly(OneSidedMat);
 
-						AssetCache.CacheAsset( PrimPath, TwoSidedMat, PrimPath );
+						AssetCache.CacheAsset(PrimPath, TwoSidedMat, PrimPath);
 						Material = TwoSidedMat;
 					}
 					else
 #endif // WITH_EDITOR
-					if ( UMaterialInstanceDynamic* MID = Cast<UMaterialInstanceDynamic>( OneSidedMat ) )
-					{
-						UMaterialInstanceDynamic* TwoSidedMat = UMaterialInstanceDynamic::Create(
-							ReferenceMaterialTwoSided,
-							GetTransientPackage(),
-							NewInstanceName
-						);
-						if ( !ensure( TwoSidedMat ) )
+						if (UMaterialInstanceDynamic* MID = Cast<UMaterialInstanceDynamic>(OneSidedMat))
 						{
-							continue;
-						}
+							UMaterialInstanceDynamic* TwoSidedMat = UMaterialInstanceDynamic::Create(
+								ReferenceMaterialTwoSided,
+								GetTransientPackage(),
+								NewInstanceName
+							);
+							if (!ensure(TwoSidedMat))
+							{
+								continue;
+							}
 
-						TwoSidedMat->CopyParameterOverrides( MID );
+							TwoSidedMat->CopyParameterOverrides(MID);
 
-						AssetCache.CacheAsset( PrimPath, TwoSidedMat, PrimPath );
-						Material = TwoSidedMat;
-					}
+							AssetCache.CacheAsset(PrimPath, TwoSidedMat, PrimPath);
+							Material = TwoSidedMat;
+						}
 				}
 
 				break;
 			}
 			case UsdUtils::EPrimAssignmentType::UnrealMaterial:
 			{
-				Material = Cast< UMaterialInterface >( FSoftObjectPath( Slot.MaterialSource ).TryLoad() );
-				if ( !Material->IsTwoSided() && Slot.bMeshIsDoubleSided )
+				Material = Cast< UMaterialInterface >(FSoftObjectPath(Slot.MaterialSource).TryLoad());
+				if (!Material)
+				{
+					UE_LOG(LogUsd, Warning, TEXT("UE material '%s' for prim '%s' could not be loaded or was not found."),
+						*Slot.MaterialSource,
+						*UsdToUnreal::ConvertPath(UsdPrim.GetPrimPath()));
+				}
+				else if (!Material->IsTwoSided() && Slot.bMeshIsDoubleSided)
 				{
 					// TODO: Update this message with the proper source prim paths when UE-138122 is submitted,
 					// as 'UsdPrim' may just be e.g. a SkelRoot
-					UE_LOG( LogUsd, Warning, TEXT( "Using one-sided UE material '%s' for doubleSided prim '%s'" ),
+					UE_LOG(LogUsd, Warning, TEXT("Using one-sided UE material '%s' for doubleSided prim '%s'"),
 						*Slot.MaterialSource,
-						*UsdToUnreal::ConvertPath( UsdPrim.GetPrimPath() )
+						*UsdToUnreal::ConvertPath(UsdPrim.GetPrimPath())
 					);
 				}
 
@@ -193,12 +199,12 @@ TMap<const UsdUtils::FUsdPrimMaterialSlot*, UMaterialInterface*> MeshTranslation
 			case UsdUtils::EPrimAssignmentType::None:
 			default:
 			{
-				ensure( false );
+				ensure(false);
 				break;
 			}
 			}
 
-			ResolvedMaterials.Add( &Slot, Material );
+			ResolvedMaterials.Add(&Slot, Material);
 		}
 	}
 
@@ -219,8 +225,8 @@ void MeshTranslationImpl::SetMaterialOverrides(
 {
 	FScopedUsdAllocs Allocs;
 
-	pxr::UsdGeomMesh Mesh{ Prim };
-	if ( !Mesh )
+	pxr::UsdGeomMesh Mesh{Prim};
+	if (!Mesh)
 	{
 		return;
 	}
@@ -228,15 +234,15 @@ void MeshTranslationImpl::SetMaterialOverrides(
 	pxr::UsdStageRefPtr Stage = Prim.GetStage();
 
 	pxr::TfToken RenderContextToken = pxr::UsdShadeTokens->universalRenderContext;
-	if ( !RenderContext.IsNone() )
+	if (!RenderContext.IsNone())
 	{
-		RenderContextToken = UnrealToUsd::ConvertToken( *RenderContext.ToString() ).Get();
+		RenderContextToken = UnrealToUsd::ConvertToken(*RenderContext.ToString()).Get();
 	}
 
 	pxr::TfToken MaterialPurposeToken = pxr::UsdShadeTokens->allPurpose;
-	if ( !MaterialPurpose.IsNone() )
+	if (!MaterialPurpose.IsNone())
 	{
-		MaterialPurposeToken = UnrealToUsd::ConvertToken( *MaterialPurpose.ToString() ).Get();
+		MaterialPurposeToken = UnrealToUsd::ConvertToken(*MaterialPurpose.ToString()).Get();
 	}
 
 	TArray<UsdUtils::FUsdPrimMaterialAssignmentInfo> LODIndexToAssignments;
@@ -244,45 +250,45 @@ void MeshTranslationImpl::SetMaterialOverrides(
 
 	// Extract material assignment info from prim if it is a LOD mesh
 	bool bInterpretedLODs = false;
-	if ( bInterpretLODs && UsdUtils::IsGeomMeshALOD( Prim ) )
+	if (bInterpretLODs && UsdUtils::IsGeomMeshALOD(Prim))
 	{
 		TMap<int32, UsdUtils::FUsdPrimMaterialAssignmentInfo> LODIndexToAssignmentsMap;
-		TFunction<bool( const pxr::UsdGeomMesh&, int32 )> IterateLODs = [ & ]( const pxr::UsdGeomMesh& LODMesh, int32 LODIndex )
+		TFunction<bool(const pxr::UsdGeomMesh&, int32)> IterateLODs = [&](const pxr::UsdGeomMesh& LODMesh, int32 LODIndex)
 		{
 			UsdUtils::FUsdPrimMaterialAssignmentInfo LODInfo = UsdUtils::GetPrimMaterialAssignments(
 				LODMesh.GetPrim(),
-				pxr::UsdTimeCode( Time ),
+				pxr::UsdTimeCode(Time),
 				bProvideMaterialIndices,
 				RenderContextToken,
 				MaterialPurposeToken
 			);
-			LODIndexToAssignmentsMap.Add( LODIndex, LODInfo );
+			LODIndexToAssignmentsMap.Add(LODIndex, LODInfo);
 			return true;
 		};
 
 		pxr::UsdPrim ParentPrim = Prim.GetParent();
-		bInterpretedLODs = UsdUtils::IterateLODMeshes( ParentPrim, IterateLODs );
+		bInterpretedLODs = UsdUtils::IterateLODMeshes(ParentPrim, IterateLODs);
 
-		if ( bInterpretedLODs )
+		if (bInterpretedLODs)
 		{
-			LODIndexToAssignmentsMap.KeySort( TLess<int32>() );
-			for ( TPair<int32, UsdUtils::FUsdPrimMaterialAssignmentInfo>& Entry : LODIndexToAssignmentsMap )
+			LODIndexToAssignmentsMap.KeySort(TLess<int32>());
+			for (TPair<int32, UsdUtils::FUsdPrimMaterialAssignmentInfo>& Entry : LODIndexToAssignmentsMap)
 			{
-				LODIndexToAssignments.Add( MoveTemp( Entry.Value ) );
+				LODIndexToAssignments.Add(MoveTemp(Entry.Value));
 			}
 		}
 	}
 
 	// Refresh reference to Prim because variant switching potentially invalidated it
-	pxr::UsdPrim ValidPrim = Stage->GetPrimAtPath( PrimPath );
+	pxr::UsdPrim ValidPrim = Stage->GetPrimAtPath(PrimPath);
 
 	// Extract material assignment info from prim if its *not* a LOD mesh, or if we failed to parse LODs
-	if ( !bInterpretedLODs )
+	if (!bInterpretedLODs)
 	{
 		LODIndexToAssignments = {
 			UsdUtils::GetPrimMaterialAssignments(
 				ValidPrim,
-				pxr::UsdTimeCode( Time ),
+				pxr::UsdTimeCode(Time),
 				bProvideMaterialIndices,
 				RenderContextToken,
 				MaterialPurposeToken
@@ -299,55 +305,55 @@ void MeshTranslationImpl::SetMaterialOverrides(
 
 	// Compare resolved materials with existing assignments, and create overrides if we need to
 	uint32 StaticMeshSlotIndex = 0;
-	for ( int32 LODIndex = 0; LODIndex < LODIndexToAssignments.Num(); ++LODIndex )
+	for (int32 LODIndex = 0; LODIndex < LODIndexToAssignments.Num(); ++LODIndex)
 	{
-		const TArray< UsdUtils::FUsdPrimMaterialSlot >& LODSlots = LODIndexToAssignments[ LODIndex ].Slots;
-		for ( int32 LODSlotIndex = 0; LODSlotIndex < LODSlots.Num(); ++LODSlotIndex, ++StaticMeshSlotIndex )
+		const TArray< UsdUtils::FUsdPrimMaterialSlot >& LODSlots = LODIndexToAssignments[LODIndex].Slots;
+		for (int32 LODSlotIndex = 0; LODSlotIndex < LODSlots.Num(); ++LODSlotIndex, ++StaticMeshSlotIndex)
 		{
-			const UsdUtils::FUsdPrimMaterialSlot& Slot = LODSlots[ LODSlotIndex ];
+			const UsdUtils::FUsdPrimMaterialSlot& Slot = LODSlots[LODSlotIndex];
 
 			UMaterialInterface* Material = nullptr;
-			if ( UMaterialInterface** FoundMaterial = ResolvedMaterials.Find( &Slot ) )
+			if (UMaterialInterface** FoundMaterial = ResolvedMaterials.Find(&Slot))
 			{
 				Material = *FoundMaterial;
 			}
 			else
 			{
-				UE_LOG( LogUsd, Error, TEXT( "Lost track of resolved material for slot '%d' of LOD '%d' for mesh '%s'" ), LODSlotIndex, LODIndex, *UsdToUnreal::ConvertPath( Prim.GetPath() ) );
+				UE_LOG(LogUsd, Error, TEXT("Lost track of resolved material for slot '%d' of LOD '%d' for mesh '%s'"), LODSlotIndex, LODIndex, *UsdToUnreal::ConvertPath(Prim.GetPath()));
 				continue;
 			}
 
-			UMaterialInterface* ExistingMaterial = ExistingAssignments[ StaticMeshSlotIndex ];
-			if ( ExistingMaterial == Material )
+			UMaterialInterface* ExistingMaterial = ExistingAssignments[StaticMeshSlotIndex];
+			if (ExistingMaterial == Material)
 			{
 				continue;
 			}
 			else
 			{
-				MeshComponent.SetMaterial( StaticMeshSlotIndex, Material );
+				MeshComponent.SetMaterial(StaticMeshSlotIndex, Material);
 			}
 		}
 	}
 }
 
-UMaterialInterface* MeshTranslationImpl::GetReferencePreviewSurfaceMaterial( EUsdReferenceMaterialProperties ReferenceMaterialProperties )
+UMaterialInterface* MeshTranslationImpl::GetReferencePreviewSurfaceMaterial(EUsdReferenceMaterialProperties ReferenceMaterialProperties)
 {
 	const UUsdProjectSettings* Settings = GetDefault<UUsdProjectSettings>();
-	if ( !Settings )
+	if (!Settings)
 	{
 		return nullptr;
 	}
 
-	const bool bIsTranslucent = EnumHasAnyFlags( ReferenceMaterialProperties, EUsdReferenceMaterialProperties::Translucent );
-	const bool bIsVT = EnumHasAnyFlags( ReferenceMaterialProperties, EUsdReferenceMaterialProperties::VT );
-	const bool bIsTwoSided = EnumHasAnyFlags( ReferenceMaterialProperties, EUsdReferenceMaterialProperties::TwoSided );
+	const bool bIsTranslucent = EnumHasAnyFlags(ReferenceMaterialProperties, EUsdReferenceMaterialProperties::Translucent);
+	const bool bIsVT = EnumHasAnyFlags(ReferenceMaterialProperties, EUsdReferenceMaterialProperties::VT);
+	const bool bIsTwoSided = EnumHasAnyFlags(ReferenceMaterialProperties, EUsdReferenceMaterialProperties::TwoSided);
 
 	const FSoftObjectPath* TargetMaterialPath = nullptr;
-	if ( bIsTranslucent )
+	if (bIsTranslucent)
 	{
-		if ( bIsVT )
+		if (bIsVT)
 		{
-			if ( bIsTwoSided )
+			if (bIsTwoSided)
 			{
 				TargetMaterialPath = &Settings->ReferencePreviewSurfaceTranslucentTwoSidedVTMaterial;
 			}
@@ -358,7 +364,7 @@ UMaterialInterface* MeshTranslationImpl::GetReferencePreviewSurfaceMaterial( EUs
 		}
 		else
 		{
-			if ( bIsTwoSided )
+			if (bIsTwoSided)
 			{
 				TargetMaterialPath = &Settings->ReferencePreviewSurfaceTranslucentTwoSidedMaterial;
 			}
@@ -370,9 +376,9 @@ UMaterialInterface* MeshTranslationImpl::GetReferencePreviewSurfaceMaterial( EUs
 	}
 	else
 	{
-		if ( bIsVT )
+		if (bIsVT)
 		{
-			if ( bIsTwoSided )
+			if (bIsTwoSided)
 			{
 				TargetMaterialPath = &Settings->ReferencePreviewSurfaceTwoSidedVTMaterial;
 			}
@@ -383,7 +389,7 @@ UMaterialInterface* MeshTranslationImpl::GetReferencePreviewSurfaceMaterial( EUs
 		}
 		else
 		{
-			if ( bIsTwoSided )
+			if (bIsTwoSided)
 			{
 				TargetMaterialPath = &Settings->ReferencePreviewSurfaceTwoSidedMaterial;
 			}
@@ -394,91 +400,91 @@ UMaterialInterface* MeshTranslationImpl::GetReferencePreviewSurfaceMaterial( EUs
 		}
 	}
 
-	if ( !TargetMaterialPath )
+	if (!TargetMaterialPath)
 	{
 		return nullptr;
 	}
 
-	return Cast< UMaterialInterface >( TargetMaterialPath->TryLoad() );
+	return Cast< UMaterialInterface >(TargetMaterialPath->TryLoad());
 }
 
-UMaterialInterface* MeshTranslationImpl::GetVTVersionOfReferencePreviewSurfaceMaterial( UMaterialInterface* ReferenceMaterial )
+UMaterialInterface* MeshTranslationImpl::GetVTVersionOfReferencePreviewSurfaceMaterial(UMaterialInterface* ReferenceMaterial)
 {
-	if ( !ReferenceMaterial )
+	if (!ReferenceMaterial)
 	{
 		return nullptr;
 	}
 
 	const UUsdProjectSettings* Settings = GetDefault<UUsdProjectSettings>();
-	if ( !Settings )
+	if (!Settings)
 	{
 		return nullptr;
 	}
 
 	const FSoftObjectPath PathName = ReferenceMaterial->GetPathName();
-	if ( PathName.ToString().Contains(TEXT("VT"), ESearchCase::CaseSensitive, ESearchDir::FromEnd) )
+	if (PathName.ToString().Contains(TEXT("VT"), ESearchCase::CaseSensitive, ESearchDir::FromEnd))
 	{
 		return ReferenceMaterial;
 	}
-	else if ( PathName == Settings->ReferencePreviewSurfaceMaterial )
+	else if (PathName == Settings->ReferencePreviewSurfaceMaterial)
 	{
-		return Cast< UMaterialInterface >( Settings->ReferencePreviewSurfaceVTMaterial.TryLoad() );
+		return Cast< UMaterialInterface >(Settings->ReferencePreviewSurfaceVTMaterial.TryLoad());
 	}
-	else if ( PathName == Settings->ReferencePreviewSurfaceTwoSidedMaterial )
+	else if (PathName == Settings->ReferencePreviewSurfaceTwoSidedMaterial)
 	{
-		return Cast< UMaterialInterface >( Settings->ReferencePreviewSurfaceTwoSidedVTMaterial.TryLoad() );
+		return Cast< UMaterialInterface >(Settings->ReferencePreviewSurfaceTwoSidedVTMaterial.TryLoad());
 	}
-	else if ( PathName == Settings->ReferencePreviewSurfaceTranslucentMaterial )
+	else if (PathName == Settings->ReferencePreviewSurfaceTranslucentMaterial)
 	{
-		return Cast< UMaterialInterface >( Settings->ReferencePreviewSurfaceTranslucentVTMaterial.TryLoad() );
+		return Cast< UMaterialInterface >(Settings->ReferencePreviewSurfaceTranslucentVTMaterial.TryLoad());
 	}
-	else if ( PathName == Settings->ReferencePreviewSurfaceTranslucentTwoSidedMaterial )
+	else if (PathName == Settings->ReferencePreviewSurfaceTranslucentTwoSidedMaterial)
 	{
-		return Cast< UMaterialInterface >( Settings->ReferencePreviewSurfaceTranslucentTwoSidedVTMaterial.TryLoad() );
+		return Cast< UMaterialInterface >(Settings->ReferencePreviewSurfaceTranslucentTwoSidedVTMaterial.TryLoad());
 	}
 
 	// We should only ever call this function with a ReferenceMaterial that matches one of the above paths
-	ensure( false );
+	ensure(false);
 	return nullptr;
 }
 
-UMaterialInterface* MeshTranslationImpl::GetTwoSidedVersionOfReferencePreviewSurfaceMaterial( UMaterialInterface* ReferenceMaterial )
+UMaterialInterface* MeshTranslationImpl::GetTwoSidedVersionOfReferencePreviewSurfaceMaterial(UMaterialInterface* ReferenceMaterial)
 {
-	if ( !ReferenceMaterial )
+	if (!ReferenceMaterial)
 	{
 		return nullptr;
 	}
 
 	const UUsdProjectSettings* Settings = GetDefault<UUsdProjectSettings>();
-	if ( !Settings )
+	if (!Settings)
 	{
 		return nullptr;
 	}
 
 	const FSoftObjectPath PathName = ReferenceMaterial->GetPathName();
-	if ( PathName.ToString().Contains(TEXT("TwoSided"), ESearchCase::CaseSensitive, ESearchDir::FromEnd) )
+	if (PathName.ToString().Contains(TEXT("TwoSided"), ESearchCase::CaseSensitive, ESearchDir::FromEnd))
 	{
 		return ReferenceMaterial;
 	}
-	else if ( PathName == Settings->ReferencePreviewSurfaceMaterial )
+	else if (PathName == Settings->ReferencePreviewSurfaceMaterial)
 	{
-		return Cast< UMaterialInterface >( Settings->ReferencePreviewSurfaceTwoSidedMaterial.TryLoad() );
+		return Cast< UMaterialInterface >(Settings->ReferencePreviewSurfaceTwoSidedMaterial.TryLoad());
 	}
-	else if ( PathName == Settings->ReferencePreviewSurfaceTranslucentMaterial )
+	else if (PathName == Settings->ReferencePreviewSurfaceTranslucentMaterial)
 	{
-		return Cast< UMaterialInterface >( Settings->ReferencePreviewSurfaceTranslucentTwoSidedMaterial.TryLoad() );
+		return Cast< UMaterialInterface >(Settings->ReferencePreviewSurfaceTranslucentTwoSidedMaterial.TryLoad());
 	}
-	else if ( PathName == Settings->ReferencePreviewSurfaceVTMaterial )
+	else if (PathName == Settings->ReferencePreviewSurfaceVTMaterial)
 	{
-		return Cast< UMaterialInterface >( Settings->ReferencePreviewSurfaceTwoSidedVTMaterial.TryLoad() );
+		return Cast< UMaterialInterface >(Settings->ReferencePreviewSurfaceTwoSidedVTMaterial.TryLoad());
 	}
-	else if ( PathName == Settings->ReferencePreviewSurfaceTranslucentVTMaterial )
+	else if (PathName == Settings->ReferencePreviewSurfaceTranslucentVTMaterial)
 	{
-		return Cast< UMaterialInterface >( Settings->ReferencePreviewSurfaceTranslucentTwoSidedVTMaterial.TryLoad() );
+		return Cast< UMaterialInterface >(Settings->ReferencePreviewSurfaceTranslucentTwoSidedVTMaterial.TryLoad());
 	}
 
 	// We should only ever call this function with a ReferenceMaterial that matches one of the above paths
-	ensure( false );
+	ensure(false);
 	return nullptr;
 }
 
