USD: Add button to stage actor details customization to open the stage editor attached to that actor.

diff --git a/Engine/Plugins/Importers/USDImporter/Source/USDStage/Private/USDStageActor.cpp b/Engine/Plugins/Importers/USDImporter/Source/USDStage/Private/USDStageActor.cpp
--- a/Engine/Plugins/Importers/USDImporter/Source/USDStage/Private/USDStageActor.cpp
+++ b/Engine/Plugins/Importers/USDImporter/Source/USDStage/Private/USDStageActor.cpp
@@ -94,6 +94,7 @@ static FAutoConsoleVariableRef CVarRegenerateSkeletalAssetsOnControlRigBake(
 static const EObjectFlags DefaultObjFlag = EObjectFlags::RF_Transactional | EObjectFlags::RF_Transient;
 
 AUsdStageActor::FOnActorLoaded AUsdStageActor::OnActorLoaded;
+AUsdStageActor::FOnOpenStageEditorClicked AUsdStageActor::OnOpenStageEditorClicked;
 
 struct FUsdStageActorImpl
 {
diff --git a/Engine/Plugins/Importers/USDImporter/Source/USDStage/Private/USDStageActorCustomization.cpp b/Engine/Plugins/Importers/USDImporter/Source/USDStage/Private/USDStageActorCustomization.cpp
--- a/Engine/Plugins/Importers/USDImporter/Source/USDStage/Private/USDStageActorCustomization.cpp
+++ b/Engine/Plugins/Importers/USDImporter/Source/USDStage/Private/USDStageActorCustomization.cpp
@@ -17,6 +17,7 @@
 #include "Styling/AppStyle.h"
 #include "UObject/Object.h"
 #include "Widgets/DeclarativeSyntaxSupport.h"
+#include "Widgets/Input/SButton.h"
 #include "Widgets/Input/SEditableTextBox.h"
 
 #define LOCTEXT_NAMESPACE "UsdStageActorCustomization"
@@ -78,6 +79,27 @@ void FUsdStageActorCustomization::CustomizeDetails( IDetailLayoutBuilder& Detail
 
 	IDetailCategoryBuilder& CatBuilder = DetailLayoutBuilder.EditCategory( TEXT( "USD" ) );
 
+	TWeakObjectPtr<AUsdStageActor> WeakActor = CurrentActor;
+	CatBuilder.AddCustomRow(FText())
+	.WholeRowContent()
+	[
+		SNew(SHorizontalBox)
+		+ SHorizontalBox::Slot()
+		[
+			SNew(SButton)
+			.Text(LOCTEXT("OpenInStageEditorText", "Open in USD Stage Editor"))
+			.OnClicked_Lambda([WeakActor]()
+			{
+				if( AUsdStageActor* StageActor = WeakActor.Get() )
+				{
+					AUsdStageActor::OnOpenStageEditorClicked.Broadcast( StageActor );
+				}
+
+				return FReply::Handled();
+			})
+		]
+	];
+
 	if ( TSharedPtr<IPropertyHandle> RenderContextProperty = DetailLayoutBuilder.GetProperty( GET_MEMBER_NAME_CHECKED( AUsdStageActor, RenderContext ) ) )
 	{
 		DetailLayoutBuilder.HideProperty( RenderContextProperty );
diff --git a/Engine/Plugins/Importers/USDImporter/Source/USDStage/Public/USDStageActor.h b/Engine/Plugins/Importers/USDImporter/Source/USDStage/Public/USDStageActor.h
--- a/Engine/Plugins/Importers/USDImporter/Source/USDStage/Public/USDStageActor.h
+++ b/Engine/Plugins/Importers/USDImporter/Source/USDStage/Public/USDStageActor.h
@@ -180,6 +180,9 @@ public:
 	DECLARE_EVENT_OneParam( AUsdStageActor, FOnActorLoaded, AUsdStageActor* );
 	USDSTAGE_API static FOnActorLoaded OnActorLoaded;
 
+	DECLARE_EVENT_OneParam( AUsdStageActor, FOnOpenStageEditorClicked, AUsdStageActor* );
+	USDSTAGE_API static FOnOpenStageEditorClicked OnOpenStageEditorClicked;
+
 	DECLARE_EVENT( AUsdStageActor, FOnStageActorEvent );
 	FOnStageActorEvent OnPreStageChanged;
 	FOnStageActorEvent OnStageChanged;
diff --git a/Engine/Plugins/Importers/USDImporter/Source/USDStageEditor/Private/SUSDStage.cpp b/Engine/Plugins/Importers/USDImporter/Source/USDStageEditor/Private/SUSDStage.cpp
--- a/Engine/Plugins/Importers/USDImporter/Source/USDStageEditor/Private/SUSDStage.cpp
+++ b/Engine/Plugins/Importers/USDImporter/Source/USDStageEditor/Private/SUSDStage.cpp
@@ -270,7 +270,7 @@ namespace SUSDStageImpl
 
 void SUsdStage::Construct( const FArguments& InArgs )
 {
-	OnActorLoadedHandle = AUsdStageActor::OnActorLoaded.AddSP( SharedThis( this ), &SUsdStage::OnStageActorLoaded );
+	OnActorLoadedHandle = AUsdStageActor::OnActorLoaded.AddSP( SharedThis( this ), &SUsdStage::AttachToStageActor );
 
 	OnViewportSelectionChangedHandle = USelection::SelectionChangedEvent.AddRaw( this, &SUsdStage::OnViewportSelectionChanged );
 
@@ -389,6 +389,22 @@ void SUsdStage::Construct( const FArguments& InArgs )
 	}
 }
 
+void SUsdStage::AttachToStageActor( AUsdStageActor* InUsdStageActor )
+{
+	if ( ViewModel.UsdStageActor == InUsdStageActor )
+	{
+		return;
+	}
+
+	ClearStageActorDelegates();
+	ViewModel.UsdStageActor = InUsdStageActor;
+	SetupStageActorDelegates();
+
+	// Refresh here because we may be receiving an actor that has a stage already loaded,
+	// like during undo/redo
+	Refresh();
+}
+
 void SUsdStage::SetupStageActorDelegates()
 {
 	ClearStageActorDelegates();
@@ -2107,22 +2123,6 @@ void SUsdStage::Refresh()
 	}
 }
 
-void SUsdStage::OnStageActorLoaded( AUsdStageActor* InUsdStageActor )
-{
-	if ( ViewModel.UsdStageActor == InUsdStageActor )
-	{
-		return;
-	}
-
-	ClearStageActorDelegates();
-	ViewModel.UsdStageActor = InUsdStageActor;
-	SetupStageActorDelegates();
-
-	// Refresh here because we may be receiving an actor that has a stage already loaded,
-	// like during undo/redo
-	Refresh();
-}
-
 void SUsdStage::OnViewportSelectionChanged( UObject* NewSelection )
 {
 	// This may be called when first opening a project, before the our widgets are fully initialized
diff --git a/Engine/Plugins/Importers/USDImporter/Source/USDStageEditor/Private/SUSDStage.h b/Engine/Plugins/Importers/USDImporter/Source/USDStageEditor/Private/SUSDStage.h
--- a/Engine/Plugins/Importers/USDImporter/Source/USDStageEditor/Private/SUSDStage.h
+++ b/Engine/Plugins/Importers/USDImporter/Source/USDStageEditor/Private/SUSDStage.h
@@ -33,6 +33,8 @@ class SUsdStage : public SCompoundWidget
 
 	virtual ~SUsdStage();
 
+	void AttachToStageActor( AUsdStageActor* InUsdStageActor );
+
 protected:
 	void SetupStageActorDelegates();
 	void ClearStageActorDelegates();
@@ -76,8 +78,6 @@ protected:
 
 	void Refresh();
 
-	void OnStageActorLoaded( AUsdStageActor* InUsdStageActor );
-
 	void OnViewportSelectionChanged( UObject* NewSelection );
 
 	int32 GetNaniteTriangleThresholdValue() const;
diff --git a/Engine/Plugins/Importers/USDImporter/Source/USDStageEditor/Private/USDStageEditorModule.cpp b/Engine/Plugins/Importers/USDImporter/Source/USDStageEditor/Private/USDStageEditorModule.cpp
--- a/Engine/Plugins/Importers/USDImporter/Source/USDStageEditor/Private/USDStageEditorModule.cpp
+++ b/Engine/Plugins/Importers/USDImporter/Source/USDStageEditor/Private/USDStageEditorModule.cpp
@@ -340,10 +340,28 @@ public:
 			const bool bForClosing = false;
 			UE::UsdStageEditorModule::Private::SaveStageActorLayersForWorld( World, bForClosing );
 		});
+
+		OpenStageEditorClickedHandle = AUsdStageActor::OnOpenStageEditorClicked.AddLambda( []( AUsdStageActor* StageActor )
+		{
+			FLevelEditorModule& LevelEditorModule = FModuleManager::GetModuleChecked< FLevelEditorModule >( "LevelEditor" );
+			TSharedPtr< FTabManager > LevelEditorTabManager = LevelEditorModule.GetLevelEditorTabManager();
+			if ( TSharedPtr< SDockTab > Tab = LevelEditorTabManager->TryInvokeTab( FTabId{ TEXT( "USDStage" ) } ) )
+			{
+				if ( TSharedPtr<SBorder> ContentBorder = StaticCastSharedRef<SBorder>( Tab->GetContent() ) )
+				{
+					if ( TSharedPtr<SUsdStage> UsdStageEditor = StaticCastSharedRef<SUsdStage>( ContentBorder->GetContent() ) )
+					{
+						UsdStageEditor->AttachToStageActor( StageActor );
+					}
+				}
+			}
+		});
 	}
 
 	virtual void ShutdownModule() override
 	{
+		AUsdStageActor::OnOpenStageEditorClicked.Remove( OpenStageEditorClickedHandle );
+
 		FEditorDelegates::PreSaveWorldWithContext.Remove( PreSaveWorldEditorDelegateHandle );
 
 		AUsdStageActor::OnActorLoaded.Remove( StageActorLoadedHandle );
@@ -367,6 +385,7 @@ private:
 	FDelegateHandle PreSaveWorldEditorDelegateHandle;
 	FDelegateHandle EditorCanCloseDelegate;
 	FDelegateHandle StageActorLoadedHandle;
+	FDelegateHandle OpenStageEditorClickedHandle;
 	FDelegateHandle OnTransactionStateChangedHandle;
 #endif // #if USE_USD_SDK
 };
