USD: Improve automatic attachment of stage actors to USD Stage Editor;

Attach to PIE actor when going into PIE (and vice-versa);
Spawn actor on proper world and attach to it when opening a stage during PIE;
Attach to a new stage actor when deleting the currently attached actor;
Attach to a new stage actor when selection sync is enabled and we select another actor's spawned component;
Flash actor picker combobutton whenever an actor is switched programmatically;
Remove unnecessary usages of GWorld;
Remove duplicate functions SetActor and OnStageActorLoaded.

diff --git a/Engine/Plugins/Importers/USDImporter/Source/USDClasses/Private/USDClassesModule.cpp b/Engine/Plugins/Importers/USDImporter/Source/USDClasses/Private/USDClassesModule.cpp
--- a/Engine/Plugins/Importers/USDImporter/Source/USDClasses/Private/USDClassesModule.cpp
+++ b/Engine/Plugins/Importers/USDImporter/Source/USDClasses/Private/USDClassesModule.cpp
@@ -5,6 +5,7 @@
 #include "USDLog.h"
 
 #include "AnalyticsEventAttribute.h"
+#include "Engine/Engine.h"
 #include "EngineAnalytics.h"
 #include "HAL/FileManager.h"
 #include "Misc/EngineVersion.h"
@@ -244,6 +245,42 @@ bool IUsdClassesModule::HashObjectPackage( const UObject* Object, FSHA1& HashToU
 #endif // WITH_EDITOR
 }
 
+UWorld* IUsdClassesModule::GetCurrentWorld( bool bEditorWorldsOnly )
+{
+	UWorld* EditorWorld = nullptr;
+	UWorld* LowestPIEWorld = nullptr;
+	int32 LowestPIEInstance = TNumericLimits<int32>::Max();
+
+	for ( const FWorldContext& Context : GEngine->GetWorldContexts() )
+	{
+		UWorld* World = Context.World();
+		if ( !World )
+		{
+			continue;
+		}
+
+		if ( !bEditorWorldsOnly && Context.WorldType == EWorldType::PIE )
+		{
+			if ( !LowestPIEInstance || Context.PIEInstance < LowestPIEInstance )
+			{
+				LowestPIEWorld = World;
+				LowestPIEInstance = Context.PIEInstance;
+			}
+		}
+		else if ( Context.WorldType == EWorldType::Editor )
+		{
+			EditorWorld = World;
+		}
+	}
+
+	if ( LowestPIEWorld )
+	{
+		return LowestPIEWorld;
+	}
+
+	return EditorWorld;
+}
+
 class FUsdClassesModule : public IUsdClassesModule
 {
 
diff --git a/Engine/Plugins/Importers/USDImporter/Source/USDClasses/Public/USDClassesModule.h b/Engine/Plugins/Importers/USDImporter/Source/USDClasses/Public/USDClassesModule.h
--- a/Engine/Plugins/Importers/USDImporter/Source/USDClasses/Public/USDClassesModule.h
+++ b/Engine/Plugins/Importers/USDImporter/Source/USDClasses/Public/USDClassesModule.h
@@ -37,4 +37,7 @@ public:
 	 * This can be used to track the version of exported assets and levels, to prevent unnecessary re-exports.
 	 */
 	USDCLASSES_API static bool HashObjectPackage( const UObject* Object, FSHA1& HashToUpdate );
+
+	/** Returns a world that could be suitably described as "the current world". (e.g. when in PIE, the PIE world) */
+	USDCLASSES_API static UWorld* GetCurrentWorld( bool bEditorWorldsOnly = false );
 };
\ No newline at end of file
diff --git a/Engine/Plugins/Importers/USDImporter/Source/USDExporter/Private/LevelExporterUSDOptionsCustomization.cpp b/Engine/Plugins/Importers/USDImporter/Source/USDExporter/Private/LevelExporterUSDOptionsCustomization.cpp
--- a/Engine/Plugins/Importers/USDImporter/Source/USDExporter/Private/LevelExporterUSDOptionsCustomization.cpp
+++ b/Engine/Plugins/Importers/USDImporter/Source/USDExporter/Private/LevelExporterUSDOptionsCustomization.cpp
@@ -2,8 +2,9 @@
 
 #include "LevelExporterUSDOptionsCustomization.h"
 
-#include "LevelExporterUSDOptions.h"
+#include "USDClassesModule.h"
 #include "LevelSequenceExporterUSDOptions.h"
+#include "LevelExporterUSDOptions.h"
 
 #include "Brushes/SlateColorBrush.h"
 #include "CoreMinimal.h"
@@ -199,7 +200,8 @@ void FLevelExporterUSDOptionsCustomization::CustomizeDetails(IDetailLayoutBuilde
 		// by looking at the soft object paths, but even those aren't exposed, so here we just default to using the current level as the export level.
 		if ( LevelSequenceOptions->Level == nullptr )
 		{
-			LevelSequenceOptions->Level = GWorld;
+			const bool bEditorWorldsOnly = true;
+			LevelSequenceOptions->Level = IUsdClassesModule::GetCurrentWorld( bEditorWorldsOnly );
 		}
 
 		PickerTree = SNew( LevelExporterUSDImpl::SLevelPickerList, &LevelSequenceOptions->LevelExportOptions, LevelSequenceOptions->Level->GetWorld() );
diff --git a/Engine/Plugins/Importers/USDImporter/Source/USDExporter/Private/LevelSequenceExporterUSD.cpp b/Engine/Plugins/Importers/USDImporter/Source/USDExporter/Private/LevelSequenceExporterUSD.cpp
--- a/Engine/Plugins/Importers/USDImporter/Source/USDExporter/Private/LevelSequenceExporterUSD.cpp
+++ b/Engine/Plugins/Importers/USDImporter/Source/USDExporter/Private/LevelSequenceExporterUSD.cpp
@@ -1198,7 +1198,8 @@ bool ULevelSequenceExporterUsd::ExportBinary( UObject* Object, const TCHAR* Type
 	UWorld* World = Options->Level.Get();
 	if ( !World )
 	{
-		World = GWorld;
+		const bool bEditorWorldsOnly = true;
+		World = IUsdClassesModule::GetCurrentWorld( bEditorWorldsOnly );
 	}
 	Params.PlaybackContext = TAttribute<UObject*>::Create( TAttribute<UObject*>::FGetter::CreateLambda( [World]()
 	{
diff --git a/Engine/Plugins/Importers/USDImporter/Source/USDSchemas/Private/USDGeomMeshTranslator.cpp b/Engine/Plugins/Importers/USDImporter/Source/USDSchemas/Private/USDGeomMeshTranslator.cpp
--- a/Engine/Plugins/Importers/USDImporter/Source/USDSchemas/Private/USDGeomMeshTranslator.cpp
+++ b/Engine/Plugins/Importers/USDImporter/Source/USDSchemas/Private/USDGeomMeshTranslator.cpp
@@ -11,6 +11,7 @@
 #include "UnrealUSDWrapper.h"
 #include "USDAssetCache.h"
 #include "USDAssetImportData.h"
+#include "USDClassesModule.h"
 #include "USDConversionUtils.h"
 #include "USDGeomMeshConversion.h"
 #include "USDGroomTranslatorUtils.h"
@@ -1296,7 +1297,7 @@ void FBuildStaticMeshTaskChain::SetupTasks()
 			UWorld* World = Context->Level ? Context->Level->GetWorld() : nullptr;
 			if (!World)
 			{
-				World = GWorld;
+				World = IUsdClassesModule::GetCurrentWorld();
 			}
 			if (World)
 			{
diff --git a/Engine/Plugins/Importers/USDImporter/Source/USDStageEditor/Private/SUSDStage.cpp b/Engine/Plugins/Importers/USDImporter/Source/USDStageEditor/Private/SUSDStage.cpp
--- a/Engine/Plugins/Importers/USDImporter/Source/USDStageEditor/Private/SUSDStage.cpp
+++ b/Engine/Plugins/Importers/USDImporter/Source/USDStageEditor/Private/SUSDStage.cpp
@@ -4,6 +4,7 @@
 
 #include "SUSDLayersTreeView.h"
 #include "SUSDPrimInfo.h"
+#include "SUSDStageEditorStyle.h"
 #include "SUSDStageTreeView.h"
 #include "UnrealUSDWrapper.h"
 #include "USDClassesModule.h"
@@ -29,6 +30,7 @@
 #include "Async/Async.h"
 #include "DesktopPlatformModule.h"
 #include "Dialogs/DlgPickPath.h"
+#include "Editor.h"
 #include "Engine/Selection.h"
 #include "Engine/World.h"
 #include "EngineAnalytics.h"
@@ -43,8 +45,8 @@
 #include "UObject/StrongObjectPtr.h"
 #include "Widgets/Input/SComboBox.h"
 #include "Widgets/Input/SComboButton.h"
-#include "Widgets/Input/SEditableTextBox.h"
 #include "Widgets/Input/SButton.h"
+#include "Widgets/Input/SEditableTextBox.h"
 #include "Widgets/Input/SSpinBox.h"
 #include "Widgets/Layout/SSplitter.h"
 
@@ -270,12 +272,16 @@ namespace SUSDStageImpl
 
 void SUsdStage::Construct( const FArguments& InArgs )
 {
-	OnActorLoadedHandle = AUsdStageActor::OnActorLoaded.AddSP( SharedThis( this ), &SUsdStage::AttachToStageActor );
+	const bool bFlashButton = true;
+	OnActorLoadedHandle = AUsdStageActor::OnActorLoaded.AddSP( SharedThis( this ), &SUsdStage::AttachToStageActor, bFlashButton );
 
-	OnViewportSelectionChangedHandle = USelection::SelectionChangedEvent.AddRaw( this, &SUsdStage::OnViewportSelectionChanged );
+	OnViewportSelectionChangedHandle = USelection::SelectionChangedEvent.AddSP( SharedThis(this), &SUsdStage::OnViewportSelectionChanged );
+
+	PostPIEStartedHandle = FEditorDelegates::PostPIEStarted.AddSP( SharedThis(this), &SUsdStage::OnPostPIEStarted );
+	EndPIEHandle = FEditorDelegates::EndPIE.AddSP( SharedThis(this), &SUsdStage::OnEndPIE );
 
 	IUsdStageModule& UsdStageModule = FModuleManager::Get().LoadModuleChecked< IUsdStageModule >( "UsdStage" );
-	ViewModel.UsdStageActor = UsdStageModule.FindUsdStageActor( GWorld );
+	ViewModel.UsdStageActor = UsdStageModule.FindUsdStageActor( IUsdClassesModule::GetCurrentWorld() );
 
 	bUpdatingViewportSelection = false;
 	bUpdatingPrimSelection = false;
@@ -389,17 +395,37 @@ void SUsdStage::Construct( const FArguments& InArgs )
 	}
 }
 
-void SUsdStage::AttachToStageActor( AUsdStageActor* InUsdStageActor )
+void SUsdStage::AttachToStageActor( AUsdStageActor* InUsdStageActor, bool bFlashButton )
 {
 	if ( ViewModel.UsdStageActor == InUsdStageActor )
 	{
 		return;
 	}
 
+	if ( bFlashButton )
+	{
+		const float StartTimeSeconds = 0.0f;
+		const float FlashDuration = 0.5f;
+		const ECurveEaseFunction EaseFunction = ECurveEaseFunction::Linear;
+		FlashActorPickerCurve = FCurveSequence{ StartTimeSeconds, FlashDuration, EaseFunction };
+		FlashActorPickerCurve.Play( this->AsShared() );
+	}
+
 	ClearStageActorDelegates();
 	ViewModel.UsdStageActor = InUsdStageActor;
 	SetupStageActorDelegates();
 
+	if ( InUsdStageActor )
+	{
+		UE::FUsdStage UsdStage = static_cast< const AUsdStageActor* >( InUsdStageActor )->GetUsdStage();
+
+		if ( this->UsdPrimInfoWidget && InUsdStageActor )
+		{
+			// Just reset to the pseudoroot for now
+			this->UsdPrimInfoWidget->SetPrimPath( UsdStage, TEXT( "/" ) );
+		}
+	}
+
 	// Refresh here because we may be receiving an actor that has a stage already loaded,
 	// like during undo/redo
 	Refresh();
@@ -575,6 +601,9 @@ void SUsdStage::ClearStageActorDelegates()
 
 SUsdStage::~SUsdStage()
 {
+	FEditorDelegates::PostPIEStarted.Remove( PostPIEStartedHandle );
+	FEditorDelegates::EndPIE.Remove( EndPIEHandle );
+
 	FCoreUObjectDelegates::OnObjectPropertyChanged.Remove( OnStageActorPropertyChangedHandle );
 	AUsdStageActor::OnActorLoaded.Remove( OnActorLoadedHandle );
 	USelection::SelectionChangedEvent.Remove( OnViewportSelectionChangedHandle );
@@ -616,22 +645,50 @@ TSharedRef< SWidget > SUsdStage::MakeMainMenu()
 
 TSharedRef< SWidget > SUsdStage::MakeActorPickerMenu()
 {
-	return SNew( SComboButton )
-		.ComboButtonStyle( FAppStyle::Get(), "SimpleComboButton" )
-		.OnGetMenuContent( this, &SUsdStage::MakeActorPickerMenuContent )
-		.MenuPlacement( EMenuPlacement::MenuPlacement_BelowRightAnchor )
-		.ToolTipText( LOCTEXT( "ActorPicker_ToolTip", "Switch the active stage actor" ) )
-		.ButtonContent()
+	return
+		SNew(SOverlay)
+		+ SOverlay::Slot()
 		[
-			SNew( STextBlock )
-			.Text_Lambda( [this]()
+			SNew( SComboButton )
+			.ComboButtonStyle( FAppStyle::Get(), "SimpleComboButton" )
+			.OnGetMenuContent( this, &SUsdStage::MakeActorPickerMenuContent )
+			.MenuPlacement( EMenuPlacement::MenuPlacement_BelowRightAnchor )
+			.ToolTipText( LOCTEXT( "ActorPicker_ToolTip", "Switch the active stage actor" ) )
+			.ButtonContent()
+			[
+				SNew( STextBlock )
+				.Text_Lambda( [this]()
+				{
+					if ( AUsdStageActor* StageActor = ViewModel.UsdStageActor.Get() )
+					{
+						return FText::FromString( StageActor->GetActorLabel() );
+					}
+
+					return FText::FromString( "None" );
+				})
+			]
+		]
+
+		// Overlay for flashing for attention
+		+ SOverlay::Slot()
+		[
+			SNew(SBorder)
+			.Visibility( EVisibility::HitTestInvisible )
+			.BorderImage( FUsdStageEditorStyle::Get()->GetBrush( "UsdStageEditor.ActorPickerFlashBrush" ) )
+			.BorderBackgroundColor_Lambda( [this]() -> FSlateColor
 			{
-				if ( AUsdStageActor* StageActor = ViewModel.UsdStageActor.Get() )
+				FLinearColor Color = FLinearColor::White;
+
+				if ( FlashActorPickerCurve.IsPlaying() )
+				{
+					Color.A = 0.6f * (1.0f - FlashActorPickerCurve.GetLerp());
+				}
+				else
 				{
-					return FText::FromString( StageActor->GetActorLabel() );
+					Color.A = 0.0f;
 				}
 
-				return FText::FromString( "None" );
+				return FSlateColor( Color );
 			})
 		];
 }
@@ -692,7 +749,9 @@ TSharedRef< SWidget > SUsdStage::MakeActorPickerMenuContent()
 				{
 					if ( Actor && Actor->IsA<AUsdStageActor>() )
 					{
-						this->SetActor( Cast<AUsdStageActor>( Actor ) );
+						// Don't need to flash here as it was an intentional act to attach to a new actor
+						const bool bFlashButton = false;
+						AttachToStageActor( Cast<AUsdStageActor>( Actor ), bFlashButton );
 
 						FSlateApplication::Get().DismissAllMenus();
 					}
@@ -2055,38 +2114,9 @@ void SUsdStage::OpenStage( const TCHAR* FilePath )
 		FText::FromString( FilePath )
 	) );
 
-	if ( !ViewModel.UsdStageActor.IsValid() )
-	{
-		IUsdStageModule& UsdStageModule = FModuleManager::Get().LoadModuleChecked< IUsdStageModule >( "UsdStage" );
-		SetActor( &UsdStageModule.GetUsdStageActor( GWorld ) );
-	}
-
 	ViewModel.OpenStage( FilePath );
 }
 
-void SUsdStage::SetActor( AUsdStageActor* InUsdStageActor )
-{
-	// Call this first so that we clear all of our delegates from the previous actor before we switch actors
-	ClearStageActorDelegates();
-
-	ViewModel.UsdStageActor = InUsdStageActor;
-
-	SetupStageActorDelegates();
-
-	if( InUsdStageActor )
-	{
-		UE::FUsdStage UsdStage = static_cast< const AUsdStageActor* >( InUsdStageActor )->GetUsdStage();
-
-		if ( this->UsdPrimInfoWidget && InUsdStageActor )
-		{
-			// Just reset to the pseudoroot for now
-			this->UsdPrimInfoWidget->SetPrimPath( UsdStage, TEXT( "/" ) );
-		}
-	}
-
-	Refresh();
-}
-
 void SUsdStage::Refresh()
 {
 	// May be nullptr, but that is ok. Its how the widgets are reset
@@ -2121,6 +2151,16 @@ void SUsdStage::Refresh()
 	{
 		UsdPrimInfoWidget->SetPrimPath( GetCurrentStage(), *SelectedPrimPath );
 	}
+
+	// Automatically attach to another stage actor if we can find one
+	if ( !StageActor )
+	{
+		IUsdStageModule& UsdStageModule = FModuleManager::Get().LoadModuleChecked< IUsdStageModule >( "UsdStage" );
+		if ( AUsdStageActor* OtherStageActor = UsdStageModule.FindUsdStageActor( IUsdClassesModule::GetCurrentWorld() ) )
+		{
+			AttachToStageActor( OtherStageActor );
+		}
+	}
 }
 
 void SUsdStage::OnViewportSelectionChanged( UObject* NewSelection )
@@ -2137,19 +2177,17 @@ void SUsdStage::OnViewportSelectionChanged( UObject* NewSelection )
 		return;
 	}
 
-	AUsdStageActor* StageActor = ViewModel.UsdStageActor.Get();
+	AUsdStageActor* AttachedStageActor = ViewModel.UsdStageActor.Get();
 	USelection* Selection = Cast<USelection>( NewSelection );
-	if ( !Selection || !StageActor )
+	if ( !Selection || !AttachedStageActor )
 	{
 		return;
 	}
 
 	TGuardValue<bool> SelectionLoopGuard( bUpdatingPrimSelection, true );
 
-	TArray<USceneComponent*> SelectedComponents;
+	TArray<USceneComponent*> AllSelectedComponents;
 	{
-		Selection->GetSelectedObjects<USceneComponent>( SelectedComponents );
-
 		TArray<AActor*> SelectedActors;
 		Selection->GetSelectedObjects<AActor>( SelectedActors );
 
@@ -2157,24 +2195,69 @@ void SUsdStage::OnViewportSelectionChanged( UObject* NewSelection )
 		{
 			if ( SelectedActor )
 			{
-				SelectedComponents.Add( SelectedActor->GetRootComponent() );
+				AllSelectedComponents.Add( SelectedActor->GetRootComponent() );
+			}
+		}
+
+		// Push the components last as we'll show info for the last selected component,
+		// and we must first select the actor before we can select its individual components
+		TArray<USceneComponent*> SelectedComponents;
+		Selection->GetSelectedObjects<USceneComponent>( SelectedComponents );
+
+		AllSelectedComponents.Append( SelectedComponents );
+	}
+	if ( AllSelectedComponents.Num() == 0 )
+	{
+		return;
+	}
+
+	// Find the exact stage actor that owns these components, if any
+	AUsdStageActor* OwnerActor = nullptr;
+	if ( AllSelectedComponents.Num() > 0 )
+	{
+		USceneComponent* AttachParent = AllSelectedComponents[ AllSelectedComponents.Num() - 1 ]->GetAttachParent();
+		while ( AttachParent )
+		{
+			USceneComponent* NewAttachParent = AttachParent->GetAttachParent();
+			if ( !NewAttachParent )
+			{
+				break;
 			}
+			AttachParent = NewAttachParent;
+		}
+
+		if ( AttachParent )
+		{
+			OwnerActor = Cast<AUsdStageActor>( AttachParent->GetOwner() );
 		}
 	}
+	if ( !OwnerActor )
+	{
+		// Selected component is not something we spawned
+		return;
+	}
 
 	TArray<FString> PrimPaths;
-	for ( USceneComponent* Component : SelectedComponents )
+	for ( USceneComponent* Component : AllSelectedComponents )
 	{
-		FString FoundPrimPath = StageActor->GetSourcePrimPath( Component );
+		FString FoundPrimPath = OwnerActor->GetSourcePrimPath( Component );
 		if ( !FoundPrimPath.IsEmpty() )
 		{
 			PrimPaths.Add( FoundPrimPath );
 		}
 	}
 
+	if ( OwnerActor && OwnerActor != AttachedStageActor )
+	{
+		AttachToStageActor( OwnerActor );
+	}
+
 	if ( PrimPaths.Num() > 0 )
 	{
-		UsdStageTreeView->SelectPrims( PrimPaths );
+		if ( UsdStageTreeView )
+		{
+			UsdStageTreeView->SelectPrims( PrimPaths );
+		}
 
 		if ( UsdPrimInfoWidget )
 		{
@@ -2183,6 +2266,40 @@ void SUsdStage::OnViewportSelectionChanged( UObject* NewSelection )
 	}
 }
 
+void SUsdStage::OnPostPIEStarted( bool bIsSimulating )
+{
+	AUsdStageActor* CurrentStageActor = ViewModel.UsdStageActor.Get();
+	if ( !CurrentStageActor )
+	{
+		return;
+	}
+
+	AUsdStageActor* PIEStageActor = Cast<AUsdStageActor>( EditorUtilities::GetSimWorldCounterpartActor( CurrentStageActor ) );
+	if ( !PIEStageActor )
+	{
+		return;
+	}
+
+	AttachToStageActor( PIEStageActor );
+}
+
+void SUsdStage::OnEndPIE( bool bIsSimulating )
+{
+	AUsdStageActor* CurrentStageActor = ViewModel.UsdStageActor.Get();
+	if ( !CurrentStageActor )
+	{
+		return;
+	}
+
+	AUsdStageActor* EditorStageActor = Cast<AUsdStageActor>( EditorUtilities::GetEditorWorldCounterpartActor( CurrentStageActor ) );
+	if ( !EditorStageActor )
+	{
+		return;
+	}
+
+	AttachToStageActor( EditorStageActor );
+}
+
 int32 SUsdStage::GetNaniteTriangleThresholdValue() const
 {
 	return CurrentNaniteThreshold;
diff --git a/Engine/Plugins/Importers/USDImporter/Source/USDStageEditor/Private/SUSDStage.h b/Engine/Plugins/Importers/USDImporter/Source/USDStageEditor/Private/SUSDStage.h
--- a/Engine/Plugins/Importers/USDImporter/Source/USDStageEditor/Private/SUSDStage.h
+++ b/Engine/Plugins/Importers/USDImporter/Source/USDStageEditor/Private/SUSDStage.h
@@ -5,14 +5,15 @@
 #include "USDStageViewModel.h"
 #include "UsdWrappers/UsdStage.h"
 
+#include "Animation/CurveSequence.h"
 #include "CoreMinimal.h"
-#include "Layout/Visibility.h"
 #include "Input/Reply.h"
+#include "Layout/Visibility.h"
 #include "Misc/Optional.h"
 #include "Templates/SharedPointer.h"
-#include "Widgets/SWidget.h"
 #include "Widgets/DeclarativeSyntaxSupport.h"
 #include "Widgets/SCompoundWidget.h"
+#include "Widgets/SWidget.h"
 
 class AUsdStageActor;
 class FLevelCollectionModel;
@@ -33,7 +34,10 @@ class SUsdStage : public SCompoundWidget
 
 	virtual ~SUsdStage();
 
-	void AttachToStageActor( AUsdStageActor* InUsdStageActor );
+	// Attaches to a new stage actor.
+	// When bFlashButton is true we will also blink the button on the UI to draw attention to the fact
+	// that the actor changed
+	void AttachToStageActor( AUsdStageActor* InUsdStageActor, bool bFlashButton=true );
 
 protected:
 	void SetupStageActorDelegates();
@@ -74,12 +78,13 @@ protected:
 
 	void OpenStage( const TCHAR* FilePath );
 
-	void SetActor( AUsdStageActor* InUsdStageActor );
-
 	void Refresh();
 
 	void OnViewportSelectionChanged( UObject* NewSelection );
 
+	void OnPostPIEStarted( bool bIsSimulating );
+	void OnEndPIE( bool bIsSimulating );
+
 	int32 GetNaniteTriangleThresholdValue() const;
 	void OnNaniteTriangleThresholdValueChanged( int32 InValue );
 	void OnNaniteTriangleThresholdValueCommitted( int32 InValue, ETextCommit::Type InCommitType );
@@ -100,13 +105,16 @@ protected:
 	FDelegateHandle OnStageEditTargetChangedHandle;
 	FDelegateHandle OnPrimChangedHandle;
 	FDelegateHandle OnLayersChangedHandle;
-
 	FDelegateHandle OnViewportSelectionChangedHandle;
+	FDelegateHandle PostPIEStartedHandle;
+	FDelegateHandle EndPIEHandle;
 
 	FString SelectedPrimPath;
 
 	FUsdStageViewModel ViewModel;
 
+	FCurveSequence FlashActorPickerCurve;
+
 	// True while we're in the middle of setting the viewport selection from the prim selection
 	bool bUpdatingViewportSelection;
 
diff --git a/Engine/Plugins/Importers/USDImporter/Source/USDStageEditor/Private/SUSDStageEditorStyle.cpp b/Engine/Plugins/Importers/USDImporter/Source/USDStageEditor/Private/SUSDStageEditorStyle.cpp
--- a/Engine/Plugins/Importers/USDImporter/Source/USDStageEditor/Private/SUSDStageEditorStyle.cpp
+++ b/Engine/Plugins/Importers/USDImporter/Source/USDStageEditor/Private/SUSDStageEditorStyle.cpp
@@ -67,6 +67,8 @@ void FUsdStageEditorStyle::Initialize()
 		 */
 		StyleSet->Set( "UsdStageEditor.CheckBoxImage", new IMAGE_BRUSH( "Common/Check", Icon14x14, FLinearColor::White ) );
 
+		StyleSet->Set( "UsdStageEditor.ActorPickerFlashBrush", new FSlateRoundedBoxBrush( FLinearColor::White, 4.0f ) );
+
 		/**
 		 * Button without a background, border or foreground color that doesn't move when pressed.
 		 * No foreground color is important to let the automatic slate foreground color propagate down.
diff --git a/Engine/Plugins/Importers/USDImporter/Source/USDStageEditorViewModels/Private/USDStageViewModel.cpp b/Engine/Plugins/Importers/USDImporter/Source/USDStageEditorViewModels/Private/USDStageViewModel.cpp
--- a/Engine/Plugins/Importers/USDImporter/Source/USDStageEditorViewModels/Private/USDStageViewModel.cpp
+++ b/Engine/Plugins/Importers/USDImporter/Source/USDStageEditorViewModels/Private/USDStageViewModel.cpp
@@ -3,6 +3,7 @@
 #include "USDStageViewModel.h"
 
 #include "UnrealUSDWrapper.h"
+#include "USDClassesModule.h"
 #include "USDConversionUtils.h"
 #include "USDErrorUtils.h"
 #include "USDLayerUtils.h"
@@ -64,7 +65,7 @@ void FUsdStageViewModel::NewStage()
 	if ( !UsdStageActor.IsValid() )
 	{
 		IUsdStageModule& UsdStageModule = FModuleManager::GetModuleChecked< IUsdStageModule >( TEXT( "USDStage" ) );
-		UsdStageActor = &UsdStageModule.GetUsdStageActor( GWorld );
+		UsdStageActor = &UsdStageModule.GetUsdStageActor( IUsdClassesModule::GetCurrentWorld() );
 	}
 
 	if ( AUsdStageActor* StageActor = UsdStageActor.Get() )
@@ -83,7 +84,7 @@ void FUsdStageViewModel::OpenStage( const TCHAR* FilePath )
 	if ( !UsdStageActor.IsValid() )
 	{
 		IUsdStageModule& UsdStageModule = FModuleManager::GetModuleChecked< IUsdStageModule >( TEXT("USDStage") );
-		UsdStageActor = &UsdStageModule.GetUsdStageActor( GWorld );
+		UsdStageActor = &UsdStageModule.GetUsdStageActor( IUsdClassesModule::GetCurrentWorld() );
 	}
 
 	if ( AUsdStageActor* StageActor = UsdStageActor.Get() )
diff --git a/Engine/Plugins/Importers/USDImporter/Source/USDUtilities/Private/USDSkeletalDataConversion.cpp b/Engine/Plugins/Importers/USDImporter/Source/USDUtilities/Private/USDSkeletalDataConversion.cpp
--- a/Engine/Plugins/Importers/USDImporter/Source/USDUtilities/Private/USDSkeletalDataConversion.cpp
+++ b/Engine/Plugins/Importers/USDImporter/Source/USDUtilities/Private/USDSkeletalDataConversion.cpp
@@ -7,6 +7,7 @@
 #include "UObject/Package.h"
 #include "USDAssetImportData.h"
 #include "USDAttributeUtils.h"
+#include "USDClassesModule.h"
 #include "USDConversionUtils.h"
 #include "USDErrorUtils.h"
 #include "USDGeomMeshConversion.h"
@@ -3303,7 +3304,7 @@ bool UnrealToUsd::ConvertAnimSequence( UAnimSequence* AnimSequence, pxr::UsdPrim
 		pxr::UsdAttribute ScalesAttr = UsdSkelAnim.CreateScalesAttr();
 
 		UDebugSkelMeshComponent* DebugSkelMeshComponent = NewObject< UDebugSkelMeshComponent >();
-		DebugSkelMeshComponent->RegisterComponentWithWorld( GWorld );
+		DebugSkelMeshComponent->RegisterComponentWithWorld( IUsdClassesModule::GetCurrentWorld() );
 		DebugSkelMeshComponent->EmptyOverrideMaterials();
 		DebugSkelMeshComponent->SetSkeletalMesh( SkeletalMesh );
 
