USD: Remove unused UUsdPrimTwin::SpawnedActor.

diff --git a/Engine/Plugins/Importers/USDImporter/Source/USDStage/Private/USDPrimTwin.cpp b/Engine/Plugins/Importers/USDImporter/Source/USDStage/Private/USDPrimTwin.cpp
--- a/Engine/Plugins/Importers/USDImporter/Source/USDStage/Private/USDPrimTwin.cpp
+++ b/Engine/Plugins/Importers/USDImporter/Source/USDStage/Private/USDPrimTwin.cpp
@@ -69,14 +69,10 @@ void UUsdPrimTwin::Clear()
 		OnDestroyed.Broadcast( *this );
 	}
 
-	AActor* ActorToDestroy = SpawnedActor.Get();
-
-	if ( !ActorToDestroy )
+	AActor* ActorToDestroy = nullptr;
+	if ( SceneComponent.IsValid() && SceneComponent->GetOwner()->GetRootComponent() == SceneComponent.Get() )
 	{
-		if ( SceneComponent.IsValid() && SceneComponent->GetOwner()->GetRootComponent() == SceneComponent.Get() )
-		{
-			ActorToDestroy = SceneComponent->GetOwner();
-		}
+		ActorToDestroy = SceneComponent->GetOwner();
 	}
 
 	if ( ActorToDestroy && !ActorToDestroy->IsA< AUsdStageActor >() && !ActorToDestroy->IsActorBeingDestroyed() && ActorToDestroy->GetWorld() )
@@ -93,7 +89,6 @@ void UUsdPrimTwin::Clear()
 
 		ActorToDestroy->Modify();
 		ActorToDestroy->GetWorld()->DestroyActor( ActorToDestroy );
-		SpawnedActor = nullptr;
 	}
 	else if ( SceneComponent.IsValid() && !SceneComponent->IsBeingDestroyed() )
 	{
@@ -179,10 +174,5 @@ USceneComponent* UUsdPrimTwin::GetSceneComponent() const
 		return SceneComponent.Get();
 	}
 
-	if ( SpawnedActor.IsValid() )
-	{
-		return SpawnedActor->GetRootComponent();
-	}
-
 	return nullptr;
 }
diff --git a/Engine/Plugins/Importers/USDImporter/Source/USDStage/Private/USDStageActor.cpp b/Engine/Plugins/Importers/USDImporter/Source/USDStage/Private/USDStageActor.cpp
--- a/Engine/Plugins/Importers/USDImporter/Source/USDStage/Private/USDStageActor.cpp
+++ b/Engine/Plugins/Importers/USDImporter/Source/USDStage/Private/USDStageActor.cpp
@@ -164,22 +164,18 @@ struct FUsdStageActorImpl
 
 			const bool bRecursive = true;
 			StageActor->GetRootPrimTwin()->Iterate([&ActorsToDeselect, &ComponentsToDeselect](UUsdPrimTwin& PrimTwin)
+			{
+				if (USceneComponent* ReferencedComponent = PrimTwin.SceneComponent.Get())
 				{
-					if (AActor* ReferencedActor = PrimTwin.SpawnedActor.Get())
-					{
-						ActorsToDeselect.Add(ReferencedActor);
-					}
-					if (USceneComponent* ReferencedComponent = PrimTwin.SceneComponent.Get())
-					{
-						ComponentsToDeselect.Add(ReferencedComponent);
+					ComponentsToDeselect.Add(ReferencedComponent);
 
-						AActor* Owner = ReferencedComponent->GetOwner();
-						if (Owner && Owner->GetRootComponent() == ReferencedComponent)
-						{
-							ActorsToDeselect.Add(Owner);
-						}
+					AActor* Owner = ReferencedComponent->GetOwner();
+					if (Owner && Owner->GetRootComponent() == ReferencedComponent)
+					{
+						ActorsToDeselect.Add(Owner);
 					}
-				}, bRecursive);
+				}
+			}, bRecursive);
 
 			if (USelection* SelectedComponents = GEditor->GetSelectedComponents())
 			{
@@ -2608,12 +2604,6 @@ void AUsdStageActor::UpdateSpawnedObjectsTransientFlag(bool bTransient)
 	EObjectFlags Flag = bTransient ? EObjectFlags::RF_Transient : EObjectFlags::RF_NoFlags;
 	TFunction<void(UUsdPrimTwin&)> UpdateTransient = [=](UUsdPrimTwin& PrimTwin)
 	{
-		if (AActor* SpawnedActor = PrimTwin.SpawnedActor.Get())
-		{
-			SpawnedActor->ClearFlags(EObjectFlags::RF_Transient);
-			SpawnedActor->SetFlags(Flag);
-		}
-
 		if (USceneComponent* Component = PrimTwin.SceneComponent.Get())
 		{
 			Component->ClearFlags(EObjectFlags::RF_Transient);
@@ -2634,10 +2624,7 @@ void AUsdStageActor::UpdateSpawnedObjectsTransientFlag(bool bTransient)
 void AUsdStageActor::OnUsdPrimTwinDestroyed(const UUsdPrimTwin& UsdPrimTwin)
 {
 	PrimsToAnimate.Remove(UsdPrimTwin.PrimPath);
-
-	UObject* WatchedObject = UsdPrimTwin.SpawnedActor.IsValid() ? (UObject*)UsdPrimTwin.SpawnedActor.Get() : (UObject*)UsdPrimTwin.SceneComponent.Get();
-	ObjectsToWatch.Remove(WatchedObject);
-
+	ObjectsToWatch.Remove( UsdPrimTwin.SceneComponent.Get() );
 	LevelSequenceHelper.RemovePrim(UsdPrimTwin);
 }
 
@@ -2705,12 +2692,6 @@ void AUsdStageActor::OnObjectPropertyChanged(UObject* ObjectBeingModified, FProp
 	{
 		// Update prim from UE
 		USceneComponent* PrimSceneComponent = UsdPrimTwin->SceneComponent.Get();
-
-		if (!PrimSceneComponent && UsdPrimTwin->SpawnedActor.IsValid())
-		{
-			PrimSceneComponent = UsdPrimTwin->SpawnedActor->GetRootComponent();
-		}
-
 		if (PrimSceneComponent)
 		{
 			if (CurrentStage)
diff --git a/Engine/Plugins/Importers/USDImporter/Source/USDStage/Public/USDPrimTwin.h b/Engine/Plugins/Importers/USDImporter/Source/USDStage/Public/USDPrimTwin.h
--- a/Engine/Plugins/Importers/USDImporter/Source/USDStage/Public/USDPrimTwin.h
+++ b/Engine/Plugins/Importers/USDImporter/Source/USDStage/Public/USDPrimTwin.h
@@ -60,9 +60,6 @@ public:
 	UPROPERTY()
 	FString PrimPath;
 
-	UPROPERTY()
-	TSoftObjectPtr< class AActor > SpawnedActor;
-
 	UPROPERTY()
 	TWeakObjectPtr< class USceneComponent > SceneComponent;
 
