USD: Downgrade ensure into a warning as sometimes it's expected for material generation to fail.

diff --git a/Engine/Plugins/Importers/USDImporter/Source/USDSchemas/Private/MeshTranslationImpl.cpp b/Engine/Plugins/Importers/USDImporter/Source/USDSchemas/Private/MeshTranslationImpl.cpp
--- a/Engine/Plugins/Importers/USDImporter/Source/USDSchemas/Private/MeshTranslationImpl.cpp
+++ b/Engine/Plugins/Importers/USDImporter/Source/USDSchemas/Private/MeshTranslationImpl.cpp
@@ -103,11 +103,12 @@ TMap<const UsdUtils::FUsdPrimMaterialSlot*, UMaterialInterface*> MeshTranslation
 				if (!Material && Slot.bMeshIsDoubleSided)
 				{
 					// By now we parsed all materials so we must have the single-sided version of this material
-					UMaterialInstance* OneSidedMat = Cast< UMaterialInstance >(
+					UMaterialInstance* OneSidedMat = Cast<UMaterialInstance>(
 						AssetCache.GetAssetForPrim(Slot.MaterialSource)
-						);
-					if (!ensure(OneSidedMat))
+					);
+					if (!OneSidedMat)
 					{
+						UE_LOG(LogUsd, Warning, TEXT("Failed to generate a two-sided material from the material prim at path '%s' as no single-sided material was generated for it."), *PrimPath);
 						continue;
 					}
 
