USD: Allow disabling the shrinking of baked textures to a single value.

diff --git a/Engine/Plugins/Importers/USDImporter/Source/USDExporter/Private/MaterialExporterUSD.cpp b/Engine/Plugins/Importers/USDImporter/Source/USDExporter/Private/MaterialExporterUSD.cpp
--- a/Engine/Plugins/Importers/USDImporter/Source/USDExporter/Private/MaterialExporterUSD.cpp
+++ b/Engine/Plugins/Importers/USDImporter/Source/USDExporter/Private/MaterialExporterUSD.cpp
@@ -179,6 +179,11 @@ namespace UE::MaterialExporterUSD::Private
 			sizeof( Options.DefaultTextureSize )
 		);
 
+		VersionHash.Update(
+			reinterpret_cast<const uint8*>(&Options.bConstantColorAsSingleValue),
+			sizeof(Options.bConstantColorAsSingleValue)
+		);
+
 		VersionHash.Update(
 			reinterpret_cast< const uint8* >( Options.Properties.GetData() ),
 			Options.Properties.Num() * Options.Properties.GetTypeSize()
@@ -362,7 +367,8 @@ bool UMaterialExporterUsd::ExportMaterial(
 		Options.Properties,
 		Options.DefaultTextureSize,
 		Options.TexturesDir,
-		RootPrim
+		RootPrim,
+		Options.bConstantColorAsSingleValue
 	);
 
 	// Write asset info now that we finished exporting
diff --git a/Engine/Plugins/Importers/USDImporter/Source/USDExporter/Private/USDAssetOptions.cpp b/Engine/Plugins/Importers/USDImporter/Source/USDExporter/Private/USDAssetOptions.cpp
--- a/Engine/Plugins/Importers/USDImporter/Source/USDExporter/Private/USDAssetOptions.cpp
+++ b/Engine/Plugins/Importers/USDImporter/Source/USDExporter/Private/USDAssetOptions.cpp
@@ -24,6 +24,7 @@ void UsdUtils::AddAnalyticsAttributes(
 
 	InOutAttributes.Emplace( TEXT( "BakedProperties" ), BakedPropertiesString );
 	InOutAttributes.Emplace( TEXT( "DefaultTextureSize" ), Options.DefaultTextureSize.ToString() );
+	InOutAttributes.Emplace( TEXT( "ConstantColorAsSingleValue" ), Options.bConstantColorAsSingleValue);
 }
 
 void UsdUtils::AddAnalyticsAttributes(
@@ -57,6 +58,11 @@ void UsdUtils::HashForMaterialExport( const FUsdMaterialBakingOptions& Options,
 		sizeof( Options.DefaultTextureSize )
 	);
 
+	HashToUpdate.Update(
+		reinterpret_cast<const uint8*>(&Options.bConstantColorAsSingleValue),
+		sizeof(Options.bConstantColorAsSingleValue)
+	);
+
 	// If we changed where we want the textures exported we need to re-export them and update the texture paths on the
 	// material
 	HashToUpdate.UpdateWithString( *Options.TexturesDir.Path, Options.TexturesDir.Path.Len() );
diff --git a/Engine/Plugins/Importers/USDImporter/Source/USDExporter/Public/USDAssetOptions.h b/Engine/Plugins/Importers/USDImporter/Source/USDExporter/Public/USDAssetOptions.h
--- a/Engine/Plugins/Importers/USDImporter/Source/USDExporter/Public/USDAssetOptions.h
+++ b/Engine/Plugins/Importers/USDImporter/Source/USDExporter/Public/USDAssetOptions.h
@@ -23,6 +23,10 @@ struct USDEXPORTER_API FUsdMaterialBakingOptions
 	UPROPERTY( EditAnywhere, config, BlueprintReadWrite, Category = "Material baking options", meta = ( ClampMin = "1", UIMin = "1" ) )
 	FIntPoint DefaultTextureSize = FIntPoint{ 128, 128 };
 
+	/** When this is true and a baked texture contains a single flat color we will write out that color value directly on the USD layer and skip generating a texture */
+	UPROPERTY( EditAnywhere, BlueprintReadWrite, Category = "Material baking options" )
+	bool bConstantColorAsSingleValue = true;
+
 	/** Where baked textures are placed. Intentionally not a config as it's heavily dependent on where the stage is */
 	UPROPERTY( EditAnywhere, BlueprintReadWrite, Category = "Material baking options" )
 	FDirectoryPath TexturesDir;
diff --git a/Engine/Plugins/Importers/USDImporter/Source/USDUtilities/Private/USDShadeConversion.cpp b/Engine/Plugins/Importers/USDImporter/Source/USDUtilities/Private/USDShadeConversion.cpp
--- a/Engine/Plugins/Importers/USDImporter/Source/USDUtilities/Private/USDShadeConversion.cpp
+++ b/Engine/Plugins/Importers/USDImporter/Source/USDUtilities/Private/USDShadeConversion.cpp
@@ -1333,7 +1333,13 @@ namespace UE
 #if WITH_EDITOR
 			// Note that we will bake things that aren't supported on the default USD surface shader schema. These could be useful in case the user has a custom
 			// renderer though, and they can pick which properties they want anyway
-			bool BakeMaterial( const UMaterialInterface& Material, const TArray<FPropertyEntry>& InMaterialProperties, const FIntPoint& InDefaultTextureSize, FBakeOutput& OutBakedData )
+			bool BakeMaterial(
+				const UMaterialInterface& Material,
+				const TArray<FPropertyEntry>& InMaterialProperties,
+				const FIntPoint& InDefaultTextureSize,
+				FBakeOutput& OutBakedData,
+				bool bInDecayTexturesToSinglePixel
+			)
 			{
 				const bool bAllQualityLevels = true;
 				const bool bAllFeatureLevels = true;
@@ -1352,6 +1358,7 @@ namespace UE
 
 				FMaterialData MatSet;
 				MatSet.Material = const_cast<UMaterialInterface*>(&Material); // We don't modify it and neither does the material baking module, it's just a bad signature
+				MatSet.bPerformShrinking = bInDecayTexturesToSinglePixel;
 
 				for ( const FPropertyEntry& Entry : InMaterialProperties )
 				{
@@ -2375,7 +2382,14 @@ bool UsdToUnreal::ConvertShadeInputsToParameters( const pxr::UsdShadeMaterial& U
 }
 
 #if WITH_EDITOR
-bool UnrealToUsd::ConvertMaterialToBakedSurface( const UMaterialInterface& InMaterial, const TArray<FPropertyEntry>& InMaterialProperties, const FIntPoint& InDefaultTextureSize, const FDirectoryPath& InTexturesDir, pxr::UsdPrim& OutUsdShadeMaterialPrim )
+bool UnrealToUsd::ConvertMaterialToBakedSurface(
+	const UMaterialInterface& InMaterial,
+	const TArray<FPropertyEntry>& InMaterialProperties,
+	const FIntPoint& InDefaultTextureSize,
+	const FDirectoryPath& InTexturesDir,
+	pxr::UsdPrim& OutUsdShadeMaterialPrim,
+	bool bInDecayTexturesToSinglePixel
+)
 {
 	pxr::UsdShadeMaterial OutUsdShadeMaterial{ OutUsdShadeMaterialPrim };
 	if ( !OutUsdShadeMaterial )
@@ -2384,7 +2398,7 @@ bool UnrealToUsd::ConvertMaterialToBakedSurface( const UMaterialInterface& InMat
 	}
 
 	FBakeOutput BakedData;
-	if ( !UsdShadeConversionImpl::BakeMaterial( InMaterial, InMaterialProperties, InDefaultTextureSize, BakedData ) )
+	if ( !UsdShadeConversionImpl::BakeMaterial( InMaterial, InMaterialProperties, InDefaultTextureSize, BakedData, bInDecayTexturesToSinglePixel ) )
 	{
 		return false;
 	}
diff --git a/Engine/Plugins/Importers/USDImporter/Source/USDUtilities/Public/USDShadeConversion.h b/Engine/Plugins/Importers/USDImporter/Source/USDUtilities/Public/USDShadeConversion.h
--- a/Engine/Plugins/Importers/USDImporter/Source/USDUtilities/Public/USDShadeConversion.h
+++ b/Engine/Plugins/Importers/USDImporter/Source/USDUtilities/Public/USDShadeConversion.h
@@ -77,9 +77,18 @@ namespace UnrealToUsd
 	 * @param InDefaultTextureSize - Size of the baked texture to use for any material property that does not have a custom size set
 	 * @param InTexturesDir - Directory where the baked textures will be placed
 	 * @param OutUsdShadeMaterialPrim - UsdPrim with the UsdShadeMaterial schema that will be configured to use the baked textures and constants
+	 * @param bInDecayTexturesToSinglePixel - Whether to use a single value directly on the material instead of writing out textures with that
+											  have a single uniform color value for all pixels
 	 * @return Whether the conversion was successful or not.
 	 */
-	USDUTILITIES_API bool ConvertMaterialToBakedSurface( const UMaterialInterface& InMaterial, const TArray<FPropertyEntry>& InMaterialProperties, const FIntPoint& InDefaultTextureSize, const FDirectoryPath& InTexturesDir, pxr::UsdPrim& OutUsdShadeMaterialPrim );
+	USDUTILITIES_API bool ConvertMaterialToBakedSurface(
+		const UMaterialInterface& InMaterial,
+		const TArray<FPropertyEntry>& InMaterialProperties,
+		const FIntPoint& InDefaultTextureSize,
+		const FDirectoryPath& InTexturesDir,
+		pxr::UsdPrim& OutUsdShadeMaterialPrim,
+		bool bInDecayTexturesToSinglePixel = true
+	);
 
 	/**
 	 * Converts a flattened material's data into textures placed at InTexturesDir, and configures OutUsdShadeMaterial to use the baked textures.
diff --git a/Engine/Source/Developer/MaterialBaking/Private/MaterialBakingHelpers.cpp b/Engine/Source/Developer/MaterialBaking/Private/MaterialBakingHelpers.cpp
--- a/Engine/Source/Developer/MaterialBaking/Private/MaterialBakingHelpers.cpp
+++ b/Engine/Source/Developer/MaterialBaking/Private/MaterialBakingHelpers.cpp
@@ -330,6 +330,11 @@ void FMaterialBakingHelpers::PerformUVBorderSmear(TArray<FColor>& InOutPixels, i
 	FMaterialBakingHelpersImpl::PerformUVBorderSmear(InOutPixels, ImageWidth, ImageHeight, MaxIterations);
 }
 
+void FMaterialBakingHelpers::PerformShrinking(TArray<FColor>& InOutPixels, int32& InOutImageWidth, int32& InOutImageHeight)
+{
+	FMaterialBakingHelpersImpl::PerformShrinking(InOutPixels, InOutImageWidth, InOutImageHeight);
+}
+
 void FMaterialBakingHelpers::PerformUVBorderSmearAndShrink(TArray<FColor>& InOutPixels, int32& InOutImageWidth, int32& InOutImageHeight)
 {
 	FMaterialBakingHelpersImpl::PerformShrinking(InOutPixels, InOutImageWidth, InOutImageHeight);
diff --git a/Engine/Source/Developer/MaterialBaking/Private/MaterialBakingModule.cpp b/Engine/Source/Developer/MaterialBaking/Private/MaterialBakingModule.cpp
--- a/Engine/Source/Developer/MaterialBaking/Private/MaterialBakingModule.cpp
+++ b/Engine/Source/Developer/MaterialBaking/Private/MaterialBakingModule.cpp
@@ -327,6 +327,7 @@ void FMaterialBakingModule::BakeMaterials(const TArray<FMaterialData*>& Material
 		FMaterialDataEx& MaterialDataEx = MaterialDataExs.AddDefaulted_GetRef();
 		MaterialDataEx.Material = MaterialData->Material;
 		MaterialDataEx.bPerformBorderSmear = MaterialData->bPerformBorderSmear;
+		MaterialDataEx.bPerformShrinking = MaterialData->bPerformShrinking;
 		MaterialDataEx.bTangentSpaceNormal = MaterialData->bTangentSpaceNormal;
 
 		for (const TPair<EMaterialProperty, FIntPoint>& PropertySizePair : MaterialData->PropertySizes)
@@ -676,10 +677,14 @@ void FMaterialBakingModule::BakeMaterials(const TArray<FMaterialDataEx*>& Materi
 									// We can't unmap ourself since we're not on the render thread
 									StagingBufferPool.ReleaseStagingBufferForUnmap_AnyThread(StagingBuffer);
 
+									if (CurrentMaterialSettings->bPerformShrinking)
+									{
+										FMaterialBakingHelpers::PerformShrinking(OutputColor, OutputSize.X, OutputSize.Y);
+									}
+
 									if (CurrentMaterialSettings->bPerformBorderSmear)
 									{
-										// This will resize the output to a single pixel if the result is monochrome.
-										FMaterialBakingHelpers::PerformUVBorderSmearAndShrink(OutputColor, OutputSize.X, OutputSize.Y);
+										FMaterialBakingHelpers::PerformUVBorderSmear(OutputColor, OutputSize.X, OutputSize.Y);
 									}
 #if WITH_EDITOR
 									// If saving intermediates is turned on
diff --git a/Engine/Source/Developer/MaterialBaking/Public/MaterialBakingHelpers.h b/Engine/Source/Developer/MaterialBaking/Public/MaterialBakingHelpers.h
--- a/Engine/Source/Developer/MaterialBaking/Public/MaterialBakingHelpers.h
+++ b/Engine/Source/Developer/MaterialBaking/Public/MaterialBakingHelpers.h
@@ -10,7 +10,7 @@ struct FColor;
 class MATERIALBAKING_API FMaterialBakingHelpers
 {
 public:
-	/** 
+	/**
 	 * Applies a box blur to magenta pixels found in given texture represented by InOutPixels using non-magenta pixels, this creates a smear across the magenta/filled pixels.
 	 * @param  InOutPixels     The image to apply the smear to
 	 * @param  ImageWidth      The width of the image in pixels
@@ -19,7 +19,17 @@ public:
 	 */
 	static void PerformUVBorderSmear(TArray<FColor>& InOutPixels, int32 ImageWidth, int32 ImageHeight, int32 MaxIterations = -1);
 
-	/** 
+	/**
+	 * Shrinks the array to a single pixel if the entire output is monochrome. Params will be adjusted accordingly when this occurs.
+	 * It's useful to perform this before the border smear as a single non-magenta pixel value would have led to a monochrome image post-smear anyway,
+	 * so the actual smearing process can be skipped.
+	 * @param  InOutPixels     The image to shrink
+	 * @param  ImageWidth      The width of the image in pixels
+	 * @param  ImageHeight     The height of the image in pixels
+	 */
+	static void PerformShrinking(TArray<FColor>& InOutPixels, int32& OnOutImageWidth, int32& InOutImageHeight);
+
+	/**
 	 * Applies a box blur to magenta pixels found in given texture represented by InOutPixels using non-magenta pixels, this creates a smear across the magenta/filled pixels.
 	 * @param  InOutPixels      The image to apply the smear to
 	 * @param  InOutImageWidth  The width of the image in pixels
diff --git a/Engine/Source/Developer/MaterialBaking/Public/MaterialBakingStructures.h b/Engine/Source/Developer/MaterialBaking/Public/MaterialBakingStructures.h
--- a/Engine/Source/Developer/MaterialBaking/Public/MaterialBakingStructures.h
+++ b/Engine/Source/Developer/MaterialBaking/Public/MaterialBakingStructures.h
@@ -22,6 +22,9 @@ struct FMaterialData
 	/** Whether to smear borders after baking */
 	bool bPerformBorderSmear = true;
 
+	/** Whether to try to shrink the color values to a single non-magenta color value after baking */
+	bool bPerformShrinking = true;
+
 	/** Whether to transform normals from world-space to tangent-space (does nothing if material already uses tangent-space normals) */
 	bool bTangentSpaceNormal = false;
 };
@@ -38,6 +41,9 @@ struct FMaterialDataEx
 	/** Whether to smear borders after baking */
 	bool bPerformBorderSmear = true;
 
+	/** Whether to try to shrink the color values to a single non-magenta color value after baking */
+	bool bPerformShrinking = true;
+
 	/** Whether to transform normals from world-space to tangent-space (does nothing if material already uses tangent-space normals) */
 	bool bTangentSpaceNormal = false;
 };
