USD: Fix how new tracks on actor bindings were being ignored and not written to the USD Stage.

Fix how binding a transform track on a camera component wiped the camera property timeSamples in USD;
Add a warning when user tries binding a transform track directly to the camera component without a decomposed camera setup.

diff --git a/Engine/Plugins/Importers/USDImporter/Source/USDClasses/Public/USDProjectSettings.h b/Engine/Plugins/Importers/USDImporter/Source/USDClasses/Public/USDProjectSettings.h
--- a/Engine/Plugins/Importers/USDImporter/Source/USDClasses/Public/USDProjectSettings.h
+++ b/Engine/Plugins/Importers/USDImporter/Source/USDClasses/Public/USDProjectSettings.h
@@ -47,6 +47,10 @@ public:
 	UPROPERTY( config, EditAnywhere, Category = "USD|Dialogs" )
 	bool bShowWarningOnIncompleteDuplication = true;
 
+	// Whether to show the warning dialog when authoring a transform track directly to a camera component
+	UPROPERTY( config, EditAnywhere, Category = "USD|Dialogs" )
+	bool bShowTransformTrackOnCameraComponentWarning = true;
+
 	// Whether to display the pop up dialog asking what to do about dirty USD layers when saving the UE level
 	UPROPERTY(config, EditAnywhere, Category = "USD|Dialogs" )
 	EUsdSaveDialogBehavior ShowSaveLayersDialogWhenSaving = EUsdSaveDialogBehavior::ShowPrompt;
diff --git a/Engine/Plugins/Importers/USDImporter/Source/USDStage/Private/USDLevelSequenceHelper.cpp b/Engine/Plugins/Importers/USDImporter/Source/USDStage/Private/USDLevelSequenceHelper.cpp
--- a/Engine/Plugins/Importers/USDImporter/Source/USDStage/Private/USDLevelSequenceHelper.cpp
+++ b/Engine/Plugins/Importers/USDImporter/Source/USDStage/Private/USDLevelSequenceHelper.cpp
@@ -12,6 +12,7 @@
 #include "USDLog.h"
 #include "USDMemory.h"
 #include "USDPrimConversion.h"
+#include "USDProjectSettings.h"
 #include "USDStageActor.h"
 #include "USDTypesConversion.h"
 #include "USDValueConversion.h"
@@ -42,6 +43,7 @@
 #include "CoreMinimal.h"
 #include "GeometryCacheComponent.h"
 #include "GeometryCacheTrackUSD.h"
+#include "Framework/Notifications/NotificationManager.h"
 #include "GroomCache.h"
 #include "GroomComponent.h"
 #include "HAL/IConsoleManager.h"
@@ -72,6 +74,7 @@
 #include "Tracks/MovieSceneSubTrack.h"
 #include "Tracks/MovieSceneVisibilityTrack.h"
 #include "UObject/UObjectGlobals.h"
+#include "Widgets/Notifications/SNotificationList.h"
 
 #if WITH_EDITOR
 #include "ControlRigBlueprint.h"
@@ -86,6 +89,8 @@
 #include "Subsystems/AssetEditorSubsystem.h"
 #endif // WITH_EDITOR
 
+#define LOCTEXT_NAMESPACE "USDLevelSequenceHelper"
+
 #if USE_USD_SDK
 
 namespace UsdLevelSequenceHelperImpl
@@ -412,6 +417,88 @@ namespace UsdLevelSequenceHelperImpl
 
 		return bResult;
 	}
+
+	void ShowTransformTrackOnCameraComponentWarning( const USceneComponent* Component )
+	{
+		const UCineCameraComponent* CameraComponent = Cast<const UCineCameraComponent>( Component );
+		if( !CameraComponent )
+		{
+			return;
+		}
+		const AActor* OwnerActor = CameraComponent->GetOwner();
+		if ( !OwnerActor )
+		{
+			return;
+		}
+
+		FObjectKey NewComponentKey{Component};
+		static TSet<FObjectKey> WarnedComponents;
+		if (WarnedComponents.Contains(NewComponentKey))
+		{
+			return;
+		}
+		WarnedComponents.Add(NewComponentKey);
+
+		const FText Text = LOCTEXT( "TransformTrackOnCameraComponentText", "USD: Transform track on camera component" );
+
+		const FText SubText = FText::Format(
+			LOCTEXT( "TransformTrackOnCameraComponentSubText", "The LevelSequence binding to the camera component '{0}' has a transform track, which is not supported.\n\nFor animating camera transforms, please bind a transform track to the '{1}' CameraActor directly, or to its root scene component instead." ),
+			FText::FromString( Component->GetName() ),
+			FText::FromString( OwnerActor->GetActorLabel() )
+		);
+
+		UE_LOG( LogUsd, Warning, TEXT( "%s" ), *SubText.ToString().Replace( TEXT( "\n\n" ), TEXT( " " ) ) );
+
+		const UUsdProjectSettings* Settings = GetDefault<UUsdProjectSettings>();
+		if ( Settings && Settings->bShowTransformTrackOnCameraComponentWarning )
+		{
+			static TWeakPtr<SNotificationItem> Notification;
+
+			FNotificationInfo Toast( Text );
+			Toast.SubText = SubText;
+			Toast.Image = FCoreStyle::Get().GetBrush( TEXT( "MessageLog.Warning" ) );
+			Toast.CheckBoxText = LOCTEXT( "DontAskAgain", "Don't prompt again" );
+			Toast.bUseLargeFont = false;
+			Toast.bFireAndForget = false;
+			Toast.FadeOutDuration = 0.0f;
+			Toast.ExpireDuration = 0.0f;
+			Toast.bUseThrobber = false;
+			Toast.bUseSuccessFailIcons = false;
+			Toast.ButtonDetails.Emplace(
+				LOCTEXT( "OverridenOpinionMessageOk", "Ok" ),
+				FText::GetEmpty(),
+				FSimpleDelegate::CreateLambda( []() {
+				if ( TSharedPtr<SNotificationItem> PinnedNotification = Notification.Pin() )
+				{
+					PinnedNotification->SetCompletionState( SNotificationItem::CS_Success );
+					PinnedNotification->ExpireAndFadeout();
+				}
+			} )
+			);
+			// This is flipped because the default checkbox message is "Don't prompt again"
+			Toast.CheckBoxState = Settings->bShowTransformTrackOnCameraComponentWarning ? ECheckBoxState::Unchecked : ECheckBoxState::Checked;
+			Toast.CheckBoxStateChanged = FOnCheckStateChanged::CreateStatic( []( ECheckBoxState NewState )
+			{
+				if ( UUsdProjectSettings* Settings = GetMutableDefault<UUsdProjectSettings>() )
+				{
+					// This is flipped because the default checkbox message is "Don't prompt again"
+					Settings->bShowTransformTrackOnCameraComponentWarning = NewState == ECheckBoxState::Unchecked;
+					Settings->SaveConfig();
+				}
+			} );
+
+			// Only show one at a time
+			if ( !Notification.IsValid() )
+			{
+				Notification = FSlateNotificationManager::Get().AddNotification( Toast );
+			}
+
+			if ( TSharedPtr<SNotificationItem> PinnedNotification = Notification.Pin() )
+			{
+				PinnedNotification->SetCompletionState( SNotificationItem::CS_Pending );
+			}
+		}
+	}
 #endif // WITH_EDITOR
 }
 
@@ -2896,7 +2983,7 @@ void FUsdLevelSequenceHelperImpl::HandleMovieSceneChange(UMovieScene& MovieScene
 					continue;
 				}
 
-				const bool bIsCamera = BoundComponent->GetOwner()->IsA<ACineCameraActor>();
+				const bool bIsCamera = BoundComponent->IsA<UCineCameraComponent>();
 				const bool bIsLight = BoundComponent->IsA<ULightComponentBase>();
 				const bool bIsSkeletal = BoundComponent->IsA<USkeletalMeshComponent>();
 
@@ -3186,13 +3273,23 @@ void FUsdLevelSequenceHelperImpl::HandleTrackChange(const UMovieSceneTrack& Trac
 		return;
 	}
 
-	// Our tracked bindings are always directly to components. If we don't have one here just abort
+	// Our tracked bindings are always directly to components
 	USceneComponent* BoundSceneComponent = Cast< USceneComponent >(BoundObject);
 	if (!BoundSceneComponent)
+	{
+		// The sequencer allows binding actor Transform directly, which gets deferred to its root component
+		// transform. Let's also allow that here
+		if ( AActor* BoundActor = Cast< AActor >( BoundObject ) )
+		{
+			BoundSceneComponent = BoundActor->GetRootComponent();
+		}
+	}
+	if ( !BoundSceneComponent )
 	{
 		return;
 	}
-	ensure(BoundSceneComponent->Mobility != EComponentMobility::Static);
+
+	ensure( BoundSceneComponent->Mobility != EComponentMobility::Static );
 
 	UUsdPrimTwin* PrimTwin = StageActor->RootUsdTwin->Find(BoundSceneComponent);
 
@@ -3220,6 +3317,11 @@ void FUsdLevelSequenceHelperImpl::HandleTrackChange(const UMovieSceneTrack& Trac
 			{
 				PrimTwin = StageActor->RootUsdTwin->Find(BoundSceneComponent->GetAttachParent());
 			}
+			else if ( PropertyPath == UnrealIdentifiers::TransformPropertyName )
+			{
+				// Let the user know that we currently don't support transform tracks directly on camera components
+				UsdLevelSequenceHelperImpl::ShowTransformTrackOnCameraComponentWarning( BoundSceneComponent );
+			}
 		}
 	}
 
@@ -3679,4 +3781,6 @@ FScopedBlockMonitoringChangesForTransaction::~FScopedBlockMonitoringChangesForTr
 	{
 		HelperImpl.StartMonitoringChanges();
 	}
-}
\ No newline at end of file
+}
+
+#undef LOCTEXT_NAMESPACE
\ No newline at end of file
