USD: Change default value for GDisableGeoCacheTracks as it prevents Time animation for all Mesh prims.

diff --git a/Engine/Plugins/Importers/USDImporter/Source/GeometryCacheUSD/Private/GeometryCacheTrackUSD.cpp b/Engine/Plugins/Importers/USDImporter/Source/GeometryCacheUSD/Private/GeometryCacheTrackUSD.cpp
--- a/Engine/Plugins/Importers/USDImporter/Source/GeometryCacheUSD/Private/GeometryCacheTrackUSD.cpp
+++ b/Engine/Plugins/Importers/USDImporter/Source/GeometryCacheUSD/Private/GeometryCacheTrackUSD.cpp
@@ -11,7 +11,10 @@
 #include "HAL/IConsoleManager.h"
 #include "IGeometryCacheStreamer.h"
 
-static bool GDisableGeoCacheTracks = false;
+// Currently disabled by default as it prevents AUsdStageActor::ExpandPrim from animating Mesh prim transforms,
+// and fixing it is not trivial: UGeometryCacheTrackUsd currently either works for animating via the Time property
+// *OR* via the geometry cache Sequencer tracks, while we still need the Time property animation to work in both cases
+static bool GDisableGeoCacheTracks = true;
 static FAutoConsoleVariableRef CVarDisableGeoCacheTracks(
 	TEXT("USD.DisableGeoCacheTracks"),
 	GDisableGeoCacheTracks,
