USD: Fixed non-editor build.

diff --git a/Engine/Plugins/Importers/USDImporter/Source/USDStage/Private/USDLevelSequenceHelper.cpp b/Engine/Plugins/Importers/USDImporter/Source/USDStage/Private/USDLevelSequenceHelper.cpp
--- a/Engine/Plugins/Importers/USDImporter/Source/USDStage/Private/USDLevelSequenceHelper.cpp
+++ b/Engine/Plugins/Importers/USDImporter/Source/USDStage/Private/USDLevelSequenceHelper.cpp
@@ -3317,11 +3317,13 @@ void FUsdLevelSequenceHelperImpl::HandleTrackChange(const UMovieSceneTrack& Trac
 			{
 				PrimTwin = StageActor->RootUsdTwin->Find(BoundSceneComponent->GetAttachParent());
 			}
+#if WITH_EDITOR
 			else if ( PropertyPath == UnrealIdentifiers::TransformPropertyName )
 			{
 				// Let the user know that we currently don't support transform tracks directly on camera components
 				UsdLevelSequenceHelperImpl::ShowTransformTrackOnCameraComponentWarning( BoundSceneComponent );
 			}
+#endif // WITH_EDITOR
 		}
 	}
 
