USD: Added debug path to FSdfPath and FUsdPrim for easier debugging. Enabled by setting ENABLE_USD_DEBUG_PATH to 1 in the build.cs.

diff --git a/Engine/Plugins/Importers/USDImporter/Source/UnrealUSDWrapper/Private/UsdWrappers/SdfPath.cpp b/Engine/Plugins/Importers/USDImporter/Source/UnrealUSDWrapper/Private/UsdWrappers/SdfPath.cpp
--- a/Engine/Plugins/Importers/USDImporter/Source/UnrealUSDWrapper/Private/UsdWrappers/SdfPath.cpp
+++ b/Engine/Plugins/Importers/USDImporter/Source/UnrealUSDWrapper/Private/UsdWrappers/SdfPath.cpp
@@ -22,14 +22,24 @@ namespace UE
 			FSdfPathImpl() = default;
 
 #if USE_USD_SDK
+
+#if ENABLE_USD_DEBUG_PATH
+			FString DebugPath;
+#endif
 			explicit FSdfPathImpl( const pxr::SdfPath& InSdfPath )
 				: PxrSdfPath( InSdfPath )
 			{
+#if ENABLE_USD_DEBUG_PATH
+				DebugPath = FString(ANSI_TO_TCHAR(PxrSdfPath.Get().GetString().c_str()));
+#endif
 			}
 
 			explicit FSdfPathImpl( pxr::SdfPath&& InSdfPath )
 				: PxrSdfPath( MoveTemp( InSdfPath ) )
 			{
+#if ENABLE_USD_DEBUG_PATH
+				DebugPath = FString(ANSI_TO_TCHAR(PxrSdfPath.Get().GetString().c_str()));
+#endif
 			}
 
 			TUsdStore< pxr::SdfPath > PxrSdfPath;
diff --git a/Engine/Plugins/Importers/USDImporter/Source/UnrealUSDWrapper/Private/UsdWrappers/UsdPrim.cpp b/Engine/Plugins/Importers/USDImporter/Source/UnrealUSDWrapper/Private/UsdWrappers/UsdPrim.cpp
--- a/Engine/Plugins/Importers/USDImporter/Source/UnrealUSDWrapper/Private/UsdWrappers/UsdPrim.cpp
+++ b/Engine/Plugins/Importers/USDImporter/Source/UnrealUSDWrapper/Private/UsdWrappers/UsdPrim.cpp
@@ -30,14 +30,24 @@ namespace UE
 			FUsdPrimImpl() = default;
 
 #if USE_USD_SDK
+
+#if ENABLE_USD_DEBUG_PATH
+			FSdfPath DebugPath;
+#endif
 			explicit FUsdPrimImpl( const pxr::UsdPrim& InUsdPrim )
 				: PxrUsdPrim( InUsdPrim )
 			{
+#if ENABLE_USD_DEBUG_PATH
+				DebugPath = PxrUsdPrim.Get().GetPrimPath();
+#endif
 			}
 
 			explicit FUsdPrimImpl( pxr::UsdPrim&& InUsdPrim )
 				: PxrUsdPrim( MoveTemp( InUsdPrim ) )
 			{
+#if ENABLE_USD_DEBUG_PATH
+				DebugPath = PxrUsdPrim.Get().GetPrimPath();
+#endif
 			}
 
 			TUsdStore< pxr::UsdPrim > PxrUsdPrim;
diff --git a/Engine/Plugins/Importers/USDImporter/Source/UnrealUSDWrapper/UnrealUSDWrapper.Build.cs b/Engine/Plugins/Importers/USDImporter/Source/UnrealUSDWrapper/UnrealUSDWrapper.Build.cs
--- a/Engine/Plugins/Importers/USDImporter/Source/UnrealUSDWrapper/UnrealUSDWrapper.Build.cs
+++ b/Engine/Plugins/Importers/USDImporter/Source/UnrealUSDWrapper/UnrealUSDWrapper.Build.cs
@@ -32,6 +32,7 @@ namespace UnrealBuildTool.Rules
 				PublicDependencyModuleNames.Add("Python3");
 
 				PublicDefinitions.Add("USE_USD_SDK=1");
+				PublicDefinitions.Add("ENABLE_USD_DEBUG_PATH=0");
 
 				var EngineDir = Path.GetFullPath(Target.RelativeEnginePath);
 				var PythonSourceTPSDir = Path.Combine(EngineDir, "Source", "ThirdParty", "Python3", Target.Platform.ToString());
