USD: Cleanup USD Stage Actor interface.

Group up similar members and functions together;
Remove unnecessary includes from header file;
Remove unnecessary nonconst GetTime() function;
Export symbols for all public functions;
Make NewStage info a UFUNCTION;
Move all implementations to cpp file;
Use protected instead of private overall so that users can derive the class how they wish.

diff --git a/Engine/Plugins/Importers/USDImporter/Source/USDStage/Private/USDLevelSequenceHelper.cpp b/Engine/Plugins/Importers/USDImporter/Source/USDStage/Private/USDLevelSequenceHelper.cpp
--- a/Engine/Plugins/Importers/USDImporter/Source/USDStage/Private/USDLevelSequenceHelper.cpp
+++ b/Engine/Plugins/Importers/USDImporter/Source/USDStage/Private/USDLevelSequenceHelper.cpp
@@ -13,10 +13,10 @@
 #include "USDMemory.h"
 #include "USDPrimConversion.h"
 #include "USDProjectSettings.h"
+#include "USDPrimTwin.h"
 #include "USDStageActor.h"
 #include "USDTypesConversion.h"
 #include "USDValueConversion.h"
-
 #include "UsdWrappers/SdfChangeBlock.h"
 #include "UsdWrappers/SdfLayer.h"
 #include "UsdWrappers/UsdAttribute.h"
diff --git a/Engine/Plugins/Importers/USDImporter/Source/USDStage/Private/USDStageActor.cpp b/Engine/Plugins/Importers/USDImporter/Source/USDStage/Private/USDStageActor.cpp
--- a/Engine/Plugins/Importers/USDImporter/Source/USDStage/Private/USDStageActor.cpp
+++ b/Engine/Plugins/Importers/USDImporter/Source/USDStage/Private/USDStageActor.cpp
@@ -2,9 +2,6 @@
 
 #include "USDStageActor.h"
 
-#include "Engine/Blueprint.h"
-#include "Engine/BlueprintGeneratedClass.h"
-#include "Engine/Level.h"
 #include "UnrealUSDWrapper.h"
 #include "UObject/Package.h"
 #include "USDAssetCache.h"
@@ -22,6 +19,7 @@
 #include "USDListener.h"
 #include "USDLog.h"
 #include "USDPrimConversion.h"
+#include "USDPrimTwin.h"
 #include "USDSchemasModule.h"
 #include "USDSchemaTranslator.h"
 #include "USDSkelRootTranslator.h"
@@ -46,6 +44,8 @@
 #include "Components/SkyLightComponent.h"
 #include "Components/SpotLightComponent.h"
 #include "Components/StaticMeshComponent.h"
+#include "Engine/Blueprint.h"
+#include "Engine/BlueprintGeneratedClass.h"
 #include "Engine/Engine.h"
 #include "Engine/Level.h"
 #include "Engine/Light.h"
@@ -1537,7 +1537,7 @@ void AUsdStageActor::SetMaterialPurpose(const FName& NewMaterialPurpose)
 	LoadUsdStage();
 }
 
-USDSTAGE_API void AUsdStageActor::SetRootMotionHandling(EUsdRootMotionHandling NewHandlingStrategy)
+void AUsdStageActor::SetRootMotionHandling( EUsdRootMotionHandling NewHandlingStrategy )
 {
 	const bool bMarkDirty = false;
 	Modify(bMarkDirty);
@@ -1546,6 +1546,11 @@ USDSTAGE_API void AUsdStageActor::SetRootMotionHandling(EUsdRootMotionHandling N
 	LoadUsdStage();
 }
 
+float AUsdStageActor::GetTime() const
+{
+	return Time;
+}
+
 void AUsdStageActor::SetTime(float InTime)
 {
 	const bool bMarkDirty = false;
@@ -1555,7 +1560,12 @@ void AUsdStageActor::SetTime(float InTime)
 	Refresh();
 }
 
-USceneComponent* AUsdStageActor::GetGeneratedComponent(const FString& PrimPath)
+ULevelSequence* AUsdStageActor::GetLevelSequence()
+{
+	return LevelSequence;
+}
+
+USceneComponent* AUsdStageActor::GetGeneratedComponent( const FString& PrimPath )
 {
 	const UE::FUsdStage& CurrentStage = static_cast<const AUsdStageActor*>(this)->GetUsdStage();
 	if (!CurrentStage)
@@ -2143,6 +2153,36 @@ void AUsdStageActor::ReloadAnimations()
 	}
 }
 
+UUsdAssetCache* AUsdStageActor::GetAssetCache()
+{
+	return AssetCache;
+}
+
+TSharedPtr<FUsdInfoCache> AUsdStageActor::GetInfoCache()
+{
+	return InfoCache;
+}
+
+TMap<FString, TMap<FString, int32>> AUsdStageActor::GetMaterialToPrimvarToUVIndex()
+{
+	return MaterialToPrimvarToUVIndex;
+}
+
+const UsdUtils::FBlendShapeMap& AUsdStageActor::GetBlendShapeMap()
+{
+	return BlendShapesByPath;
+}
+
+FUsdListener& AUsdStageActor::GetUsdListener()
+{
+	return UsdListener;
+}
+
+const FUsdListener& AUsdStageActor::GetUsdListener() const
+{
+	return UsdListener;
+}
+
 #if WITH_EDITOR
 
 void AUsdStageActor::PostEditChangeProperty(FPropertyChangedEvent& PropertyChangedEvent)
diff --git a/Engine/Plugins/Importers/USDImporter/Source/USDStage/Public/USDStageActor.h b/Engine/Plugins/Importers/USDImporter/Source/USDStage/Public/USDStageActor.h
--- a/Engine/Plugins/Importers/USDImporter/Source/USDStage/Public/USDStageActor.h
+++ b/Engine/Plugins/Importers/USDImporter/Source/USDStage/Public/USDStageActor.h
@@ -2,43 +2,33 @@
 
 #pragma once
 
-#include "GameFramework/Actor.h"
-#include "Misc/SecureHash.h"
-#include "Misc/ITransaction.h"
-#include "Templates/Function.h"
-
-#include "USDInfoCache.h"
 #include "USDLevelSequenceHelper.h"
 #include "USDListener.h"
-#include "USDPrimTwin.h"
 #include "USDSkeletalDataConversion.h"
-
-#include "UsdWrappers/UsdStage.h"
-#include "UsdWrappers/UsdPrim.h"
 #include "UsdWrappers/SdfPath.h"
+#include "UsdWrappers/UsdStage.h"
 
+#include "GameFramework/Actor.h"
+#include "Misc/ITransaction.h"
+
 #include "USDStageActor.generated.h"
 
-class ALevelSequenceActor;
-class IMeshBuilderModule;
 class ULevelSequence;
-class UMaterial;
-class UUsdAsset;
 class UUsdAssetCache;
+class UUsdPrimTwin;
 class UUsdTransactor;
-enum class EMapChangeType : uint8;
-enum class EUsdPurpose : int32;
-struct FMeshDescription;
+struct FUsdInfoCache;
 struct FUsdSchemaTranslationContext;
+namespace UE
+{
+	class FUsdPrim;
+}
 
 UCLASS( MinimalAPI, config = USDImporter )
 class AUsdStageActor : public AActor
 {
 	GENERATED_BODY()
 
-	friend struct FUsdStageActorImpl;
-	friend class FUsdLevelSequenceHelperImpl;
-
 public:
 	UPROPERTY(EditAnywhere, BlueprintReadOnly, Category = "USD", meta = (RelativeToGameDir))
 	FFilePath RootLayer;
@@ -97,5 +87,23 @@
 	EUsdRootMotionHandling RootMotionHandling = EUsdRootMotionHandling::NoAdditionalRootMotion;
 
 public:
+	DECLARE_EVENT_OneParam(AUsdStageActor, FOnActorLoaded, AUsdStageActor*);
+	USDSTAGE_API static FOnActorLoaded OnActorLoaded;
+
+	DECLARE_EVENT(AUsdStageActor, FOnStageActorEvent);
+	FOnStageActorEvent OnPreStageChanged;
+	FOnStageActorEvent OnStageChanged;
+	FOnStageActorEvent OnActorDestroyed;
+
+	DECLARE_EVENT_TwoParams(AUsdStageActor, FOnPrimChanged, const FString&, bool);
+	FOnPrimChanged OnPrimChanged;
+
+	DECLARE_MULTICAST_DELEGATE(FOnUsdStageTimeChanged);
+	FOnUsdStageTimeChanged OnTimeChanged;
+
+	DECLARE_EVENT_OneParam(AUsdStageActor, FOnOpenStageEditorClicked, AUsdStageActor*);
+	USDSTAGE_API static FOnOpenStageEditorClicked OnOpenStageEditorClicked;
+
+public:
 	UFUNCTION(BlueprintCallable, Category = "USD", meta = (CallInEditor = "true"))
 	USDSTAGE_API void SetRootLayer(const FString& RootFilePath );
@@ -132,13 +140,13 @@
 	USDSTAGE_API void SetRootMotionHandling( EUsdRootMotionHandling NewHandlingStrategy );
 
 	UFUNCTION(BlueprintCallable, Category = "USD", meta = (CallInEditor = "true"))
-	USDSTAGE_API float GetTime() const { return Time; }
+	USDSTAGE_API float GetTime() const;
 
 	UFUNCTION(BlueprintCallable, Category = "USD", meta = (CallInEditor = "true"))
 	USDSTAGE_API void SetTime(float InTime);
 
 	UFUNCTION( BlueprintCallable, Category = "USD", meta = ( CallInEditor = "true" ) )
-	USDSTAGE_API ULevelSequence* GetLevelSequence() { return LevelSequence; }
+	USDSTAGE_API ULevelSequence* GetLevelSequence();
 
 	/**
 	 * Gets the transient component that was generated for a prim with a given prim path.
@@ -165,101 +173,127 @@
 	UFUNCTION(BlueprintCallable, Category = "USD", meta = (CallInEditor = "true"))
 	USDSTAGE_API FString GetSourcePrimPath( UObject* Object );
 
-private:
-	UPROPERTY(Category = UsdStageActor, VisibleAnywhere, BlueprintReadOnly, meta = (ExposeFunctionCategories = "Mesh,Rendering,Physics,Components|StaticMesh", AllowPrivateAccess = "true"))
-	TObjectPtr<class USceneComponent> SceneComponent;
-
-	/* TimeCode to evaluate the USD stage at */
-	UPROPERTY(EditAnywhere, Category = "USD")
-	float Time;
-
-	UPROPERTY(VisibleAnywhere, Category = "USD", Transient)
-	TObjectPtr<ULevelSequence> LevelSequence;
-
-public:
-	DECLARE_EVENT_OneParam( AUsdStageActor, FOnActorLoaded, AUsdStageActor* );
-	USDSTAGE_API static FOnActorLoaded OnActorLoaded;
-
-	DECLARE_EVENT_OneParam( AUsdStageActor, FOnOpenStageEditorClicked, AUsdStageActor* );
-	USDSTAGE_API static FOnOpenStageEditorClicked OnOpenStageEditorClicked;
-
-	DECLARE_EVENT( AUsdStageActor, FOnStageActorEvent );
-	FOnStageActorEvent OnPreStageChanged;
-	FOnStageActorEvent OnStageChanged;
-	FOnStageActorEvent OnActorDestroyed;
-
-	DECLARE_EVENT_TwoParams( AUsdStageActor, FOnPrimChanged, const FString&, bool );
-	FOnPrimChanged OnPrimChanged;
-
-	DECLARE_MULTICAST_DELEGATE(FOnUsdStageTimeChanged);
-	FOnUsdStageTimeChanged OnTimeChanged;
-
-public:
-	AUsdStageActor();
-
 	// Creates a brand new, memory-only USD stage and opens it
+	UFUNCTION(BlueprintCallable, Category = "USD", meta = (CallInEditor = "true"))
 	USDSTAGE_API void NewStage();
 
+public:
+	// Loads the stage with RootLayer if its not loaded already, and returns the either the isolated stage (if any) or
+	// the base stage
+	USDSTAGE_API UE::FUsdStage& GetOrLoadUsdStage();
+
+	// Returns either the isolated stage (if any) or the base stage
+	USDSTAGE_API const UE::FUsdStage& GetUsdStage() const;
+
+	// Always returns the base stage, regardless of whether we have an isolated stage or not
+	USDSTAGE_API const UE::FUsdStage& GetBaseUsdStage() const;
+
+	// Always returns the isolated stage, being an invalid stage in case we're not isolating anything
+	USDSTAGE_API const UE::FUsdStage& GetIsolatedUsdStage() const;
+
+	// Enters isolated mode by creating a new USD Stage using the provided layer as its root
 	USDSTAGE_API void IsolateLayer( const UE::FSdfLayer& Layer );
 
-	USDSTAGE_API void Reset() override;
+	// Regenerates our LevelSequence from the USD Stage
+	USDSTAGE_API void ReloadAnimations();
+
+	USDSTAGE_API UUsdAssetCache* GetAssetCache();
+	USDSTAGE_API TSharedPtr<FUsdInfoCache> GetInfoCache();
+	USDSTAGE_API TMap< FString, TMap< FString, int32 > > GetMaterialToPrimvarToUVIndex();
+	USDSTAGE_API const UsdUtils::FBlendShapeMap& GetBlendShapeMap();
+	USDSTAGE_API FUsdListener& GetUsdListener();
+	USDSTAGE_API const FUsdListener& GetUsdListener() const;
+
+	/** Control whether we respond to USD notices or not. Mostly used to prevent us from responding to them when we're writing data to the stage */
+	USDSTAGE_API void StopListeningToUsdNotices();
+	USDSTAGE_API void ResumeListeningToUsdNotices();
+	USDSTAGE_API bool IsListeningToUsdNotices() const;
+
+	/** Prevents writing back data to the USD stage whenever our LevelSequences are modified */
+	USDSTAGE_API void StopMonitoringLevelSequence();
+	USDSTAGE_API void ResumeMonitoringLevelSequence();
+	USDSTAGE_API void BlockMonitoringLevelSequenceForThisTransaction();
+
+	// Begin UObject interface
+	USDSTAGE_API virtual void PostDuplicate( bool bDuplicateForPIE ) override;
+	USDSTAGE_API virtual void Serialize(FArchive& Ar) override;
+	USDSTAGE_API virtual void Destroyed() override;
+	USDSTAGE_API virtual void PostActorCreated() override;
+	USDSTAGE_API virtual void PostRename( UObject* OldOuter, const FName OldName ) override;
+	USDSTAGE_API virtual void BeginDestroy() override;
+#if WITH_EDITOR
+	USDSTAGE_API virtual void PostEditChangeProperty( FPropertyChangedEvent& PropertyChangedEvent ) override;
+	USDSTAGE_API virtual void PostTransacted(const FTransactionObjectEvent& TransactionEvent) override;
+	USDSTAGE_API virtual void PreEditChange( FProperty* PropertyThatWillChange ) override;
+	USDSTAGE_API virtual void PreEditUndo() override;
+#endif // WITH_EDITOR
+	// End UObject interface
+
+	// Begin AActor interface
+	USDSTAGE_API virtual void Reset() override;
+	USDSTAGE_API virtual void PostRegisterAllComponents() override;
+	USDSTAGE_API virtual void UnregisterAllComponents(bool bForReregister = false) override;
+	USDSTAGE_API virtual void PostUnregisterAllComponents() override;
+	// End AActor interface
+
+	AUsdStageActor();
+
+protected:
+	/** Loads the asset for a single prim */
+	void LoadAsset(FUsdSchemaTranslationContext& TranslationContext, const UE::FUsdPrim& Prim);
+
+	/** Loads the assets for all prims from StartPrim and its children */
+	void LoadAssets(FUsdSchemaTranslationContext& TranslationContext, const UE::FUsdPrim& StartPrim);
+
 	void Refresh() const;
-	void ReloadAnimations();
-	float GetTime() { return Time; }
-	UUsdAssetCache* GetAssetCache() { return AssetCache; }
-	TSharedPtr<FUsdInfoCache> GetInfoCache() { return InfoCache; }
-	TMap< FString, TMap< FString, int32 > > GetMaterialToPrimvarToUVIndex() { return MaterialToPrimvarToUVIndex; }
-	const UsdUtils::FBlendShapeMap& GetBlendShapeMap() { return BlendShapesByPath; }
+	void AnimatePrims();
 
-public:
-#if WITH_EDITOR
-	virtual void PostEditChangeProperty( FPropertyChangedEvent& PropertyChangedEvent ) override;
-	virtual void PostTransacted(const FTransactionObjectEvent& TransactionEvent) override;
-	virtual void PreEditChange( FProperty* PropertyThatWillChange ) override;
-	virtual void PreEditUndo() override;
-	void HandleTransactionStateChanged( const FTransactionContext& InTransactionContext, const ETransactionStateEventType InTransactionState );
-#endif // WITH_EDITOR
-	virtual void PostDuplicate( bool bDuplicateForPIE ) override;
-	virtual void Serialize(FArchive& Ar) override;
-	virtual void Destroyed() override;
-	virtual void PostActorCreated() override;
-	virtual void PostRename( UObject* OldOuter, const FName OldName ) override;
-	virtual void BeginDestroy() override;
+	UUsdPrimTwin* GetRootPrimTwin();
+	UUsdPrimTwin* GetOrCreatePrimTwin(const UE::FSdfPath& UsdPrimPath);
 
-	virtual void PostRegisterAllComponents() override;
-	virtual void UnregisterAllComponents( bool bForReregister = false ) override;
-	virtual void PostUnregisterAllComponents() override;
+	UUsdPrimTwin* ExpandPrim(const UE::FUsdPrim& Prim, FUsdSchemaTranslationContext& TranslationContext);
+	void UpdatePrim(const UE::FSdfPath& UsdPrimPath, bool bResync, FUsdSchemaTranslationContext& TranslationContext);
 
-	void OnPreUsdImport( FString FilePath );
-	void OnPostUsdImport( FString FilePath );
-
-private:
 	void OpenUsdStage();
 	void CloseUsdStage();
 
 	void LoadUsdStage();
 	void UnloadUsdStage();
 
-	UUsdPrimTwin* GetRootPrimTwin();
+	bool HasAuthorityOverStage() const;
+
+	void UpdateSpawnedObjectsTransientFlag(bool bTransient);
 
 #if WITH_EDITOR
 	void OnBeginPIE(bool bIsSimulating);
 	void OnPostPIEStarted(bool bIsSimulating);
 	void OnObjectsReplaced( const TMap<UObject*, UObject*>& ObjectReplacementMap );
 	void OnLevelActorDeleted(AActor* DeletedActor);
+	void HandleTransactionStateChanged(const FTransactionContext& InTransactionContext, const ETransactionStateEventType InTransactionState);
 #endif // WITH_EDITOR
 
-	void UpdateSpawnedObjectsTransientFlag( bool bTransient );
-
+	void OnPreUsdImport(FString FilePath);
+	void OnPostUsdImport(FString FilePath);
 	void OnUsdObjectsChanged( const UsdUtils::FObjectChangesByPath& InfoChanges, const UsdUtils::FObjectChangesByPath& ResyncChanges );
 	void OnUsdPrimTwinDestroyed( const UUsdPrimTwin& UsdPrimTwin );
-
 	void OnObjectPropertyChanged( UObject* ObjectBeingModified, FPropertyChangedEvent& PropertyChangedEvent );
 	void HandlePropertyChangedEvent( FPropertyChangedEvent& PropertyChangedEvent );
-	bool HasAuthorityOverStage() const;
 	void OnSkelAnimationBaked( const FString& SkelRootPrimPath );
 
-private:
+protected:
+	friend struct FUsdStageActorImpl;
+	friend class FUsdLevelSequenceHelperImpl;
+
+	UPROPERTY(Category = UsdStageActor, VisibleAnywhere, BlueprintReadOnly, meta = (ExposeFunctionCategories = "Mesh,Rendering,Physics,Components|StaticMesh", AllowPrivateAccess = "true"))
+	TObjectPtr<class USceneComponent> SceneComponent;
+
+	/* TimeCode to evaluate the USD stage at */
+	UPROPERTY(EditAnywhere, Category = "USD")
+	float Time;
+
+	UPROPERTY(VisibleAnywhere, Category = "USD", Transient)
+	TObjectPtr<ULevelSequence> LevelSequence;
+
 	// Note: This cannot be instanced: Read the comment in the constructor
 	UPROPERTY( Transient )
 	TObjectPtr<UUsdPrimTwin> RootUsdTwin;
@@ -276,53 +310,22 @@
 	UPROPERTY( Transient )
 	TObjectPtr<UUsdTransactor> Transactor;
 
-public:
-	// Loads the stage with RootLayer if its not loaded already, and returns the either the isolated stage (if any) or
-	// the base stage
-	USDSTAGE_API UE::FUsdStage& GetOrLoadUsdStage();
-
-	// Returns either the isolated stage (if any) or the base stage
-	USDSTAGE_API const UE::FUsdStage& GetUsdStage() const;
-
-	// Always returns the base stage, regardless of whether we have an isolated stage or not
-	USDSTAGE_API const UE::FUsdStage& GetBaseUsdStage() const;
-
-	// Always returns the isolated stage, being an invalid stage in case we're not isolating anything
-	USDSTAGE_API const UE::FUsdStage& GetIsolatedUsdStage() const;
-
-	FUsdListener& GetUsdListener() { return UsdListener; }
-	const FUsdListener& GetUsdListener() const { return UsdListener; }
-
-	/** Control whether we respond to USD notices or not. Mostly used to prevent us from responding to them when we're writing data to the stage */
-	USDSTAGE_API void StopListeningToUsdNotices();
-	USDSTAGE_API void ResumeListeningToUsdNotices();
-	USDSTAGE_API bool IsListeningToUsdNotices() const;
-
-	/** Prevents writing back data to the USD stage whenever our LevelSequences are modified */
-	USDSTAGE_API void StopMonitoringLevelSequence();
-	USDSTAGE_API void ResumeMonitoringLevelSequence();
-	USDSTAGE_API void BlockMonitoringLevelSequenceForThisTransaction();
-
-	UUsdPrimTwin* GetOrCreatePrimTwin( const UE::FSdfPath& UsdPrimPath );
-	UUsdPrimTwin* ExpandPrim( const UE::FUsdPrim& Prim, FUsdSchemaTranslationContext& TranslationContext );
-	void UpdatePrim( const UE::FSdfPath& UsdPrimPath, bool bResync, FUsdSchemaTranslationContext& TranslationContext );
-
-protected:
-	/** Loads the asset for a single prim */
-	void LoadAsset( FUsdSchemaTranslationContext& TranslationContext, const UE::FUsdPrim& Prim );
-
-	/** Loads the assets for all prims from StartPrim and its children */
-	void LoadAssets( FUsdSchemaTranslationContext& TranslationContext, const UE::FUsdPrim& StartPrim );
-
-	void AnimatePrims();
-
-private:
 	/** Caches various information about prims that are expensive to query */
 	TSharedPtr<FUsdInfoCache> InfoCache;
 
+	FUsdListener UsdListener;
+
+	FUsdLevelSequenceHelper LevelSequenceHelper;
+
+	// The main UsdStage that is currently opened
+	UE::FUsdStage UsdStage;
+
+	// Another stage that has as root layer one of the non-root local layers of UsdStage. This is the stage we'll
+	// be displaying if it is valid, otherwise we'll be displaying UsdStage directly.
+	UE::FUsdStage IsolatedStage;
+
 	/** Keep track of blend shapes so that we can map 'inbetween shapes' to their separate morph targets when animating */
 	UsdUtils::FBlendShapeMap BlendShapesByPath;
-
 	/**
 	 * When parsing materials, we keep track of which primvar we mapped to which UV channel.
 	 * When parsing meshes later, we use this data to place the correct primvar values in each UV channel.
@@ -330,13 +333,6 @@
 	 */
 	TMap< FString, TMap< FString, int32 > > MaterialToPrimvarToUVIndex;
 
-	// The main UsdStage that is currently opened
-	UE::FUsdStage UsdStage;
-
-	// Another stage that has as root layer one of the non-root local layers of UsdStage. This is the stage we'll
-	// be displaying if it is valid, otherwise we'll be displaying UsdStage directly.
-	UE::FUsdStage IsolatedStage;
-
 	/**
 	 * We use PostRegisterAllComponents and PostUnregisterAllComponents as main entry points to decide when to load/unload
 	 * the USD stage. These are the three exceptions we must avoid though:
@@ -349,10 +345,8 @@
 	bool bIsModifyingAProperty;
 	bool bIsUndoRedoing;
 
-	FUsdListener UsdListener;
-	FUsdLevelSequenceHelper LevelSequenceHelper;
-
 	FDelegateHandle OnRedoHandle;
+
 	FThreadSafeCounter IsBlockedFromUsdNotices;
 
 	/**
diff --git a/Engine/Plugins/Importers/USDImporter/Source/USDStageImporter/Private/USDStageImporter.cpp b/Engine/Plugins/Importers/USDImporter/Source/USDStageImporter/Private/USDStageImporter.cpp
--- a/Engine/Plugins/Importers/USDImporter/Source/USDStageImporter/Private/USDStageImporter.cpp
+++ b/Engine/Plugins/Importers/USDImporter/Source/USDStageImporter/Private/USDStageImporter.cpp
@@ -2,10 +2,6 @@
 
 #include "USDStageImporter.h"
 
-#include "Components/StaticMeshComponent.h"
-#include "Engine/SkinnedAssetCommon.h"
-#include "PhysicsEngine/PhysicsAsset.h"
-#include "UObject/UObjectIterator.h"
 #include "USDAssetCache.h"
 #include "USDAssetImportData.h"
 #include "USDClassesModule.h"
@@ -13,12 +9,12 @@
 #include "USDErrorUtils.h"
 #include "USDGeomMeshConversion.h"
 #include "USDLog.h"
+#include "USDPrimTwin.h"
 #include "USDSchemasModule.h"
 #include "USDSchemaTranslator.h"
 #include "USDStageImportContext.h"
 #include "USDStageImportOptions.h"
 #include "USDTypesConversion.h"
-
 #include "UsdWrappers/SdfLayer.h"
 #include "UsdWrappers/UsdStage.h"
 #include "UsdWrappers/UsdTyped.h"
@@ -30,9 +26,11 @@
 #include "AssetToolsModule.h"
 #include "ComponentRecreateRenderStateContext.h"
 #include "Components/SkeletalMeshComponent.h"
+#include "Components/StaticMeshComponent.h"
 #include "Dialogs/DlgPickPath.h"
 #include "Editor.h"
 #include "Engine/SkeletalMesh.h"
+#include "Engine/SkinnedAssetCommon.h"
 #include "Engine/StaticMesh.h"
 #include "Engine/Texture.h"
 #include "Engine/World.h"
@@ -52,9 +50,11 @@
 #include "Modules/ModuleManager.h"
 #include "ObjectTools.h"
 #include "PackageTools.h"
+#include "PhysicsEngine/PhysicsAsset.h"
 #include "Serialization/ArchiveReplaceObjectRef.h"
 #include "Subsystems/AssetEditorSubsystem.h"
 #include "UObject/StrongObjectPtr.h"
+#include "UObject/UObjectIterator.h"
 
 #define LOCTEXT_NAMESPACE "USDStageImporter"
 
diff --git a/Engine/Plugins/Importers/USDImporter/Source/USDTests/Private/USDTestsBlueprintLibrary.cpp b/Engine/Plugins/Importers/USDImporter/Source/USDTests/Private/USDTestsBlueprintLibrary.cpp
--- a/Engine/Plugins/Importers/USDImporter/Source/USDTests/Private/USDTestsBlueprintLibrary.cpp
+++ b/Engine/Plugins/Importers/USDImporter/Source/USDTests/Private/USDTestsBlueprintLibrary.cpp
@@ -2,6 +2,7 @@
 
 #include "USDTestsBlueprintLibrary.h"
 
+#include "USDInfoCache.h"
 #include "USDLog.h"
 #include "USDStageActor.h"
 
