USD: Fix issues encountered when packaging projects with the USDImporter plugin.

Explicitly stage the default USD config (needed to supress some dialogs on EngineTest);
Resave assets with non-zero build version;
Fix missing includes when packaging likely due to header optimization / include what you use changes.

diff --git a/Engine/Plugins/Importers/USDImporter/Source/USDStage/Private/USDLevelSequenceHelper.cpp b/Engine/Plugins/Importers/USDImporter/Source/USDStage/Private/USDLevelSequenceHelper.cpp
--- a/Engine/Plugins/Importers/USDImporter/Source/USDStage/Private/USDLevelSequenceHelper.cpp
+++ b/Engine/Plugins/Importers/USDImporter/Source/USDStage/Private/USDLevelSequenceHelper.cpp
@@ -41,9 +41,9 @@
 #include "Components/SpotLightComponent.h"
 #include "ControlRigObjectBinding.h"
 #include "CoreMinimal.h"
+#include "Framework/Notifications/NotificationManager.h"
 #include "GeometryCacheComponent.h"
 #include "GeometryCacheTrackUSD.h"
-#include "Framework/Notifications/NotificationManager.h"
 #include "GroomCache.h"
 #include "GroomComponent.h"
 #include "HAL/IConsoleManager.h"
@@ -73,6 +73,7 @@
 #include "Tracks/MovieSceneSkeletalAnimationTrack.h"
 #include "Tracks/MovieSceneSubTrack.h"
 #include "Tracks/MovieSceneVisibilityTrack.h"
+#include "UObject/Package.h"
 #include "UObject/UObjectGlobals.h"
 #include "Widgets/Notifications/SNotificationList.h"
 
diff --git a/Engine/Plugins/Importers/USDImporter/Source/USDUtilities/Private/USDConversionUtils.cpp b/Engine/Plugins/Importers/USDImporter/Source/USDUtilities/Private/USDConversionUtils.cpp
--- a/Engine/Plugins/Importers/USDImporter/Source/USDUtilities/Private/USDConversionUtils.cpp
+++ b/Engine/Plugins/Importers/USDImporter/Source/USDUtilities/Private/USDConversionUtils.cpp
@@ -45,6 +45,7 @@
 #include "GroomCache.h"
 #include "InstancedFoliageActor.h"
 #include "LandscapeProxy.h"
+#include "Misc/PackageName.h"
 #include "Widgets/Notifications/SNotificationList.h"
 
 #if WITH_EDITOR
