Fixed a possible CIS issue.

diff --git a/Engine/Plugins/Importers/USDImporter/Source/USDUtilities/Public/USDShadeConversion.h b/Engine/Plugins/Importers/USDImporter/Source/USDUtilities/Public/USDShadeConversion.h
--- a/Engine/Plugins/Importers/USDImporter/Source/USDUtilities/Public/USDShadeConversion.h
+++ b/Engine/Plugins/Importers/USDImporter/Source/USDUtilities/Public/USDShadeConversion.h
@@ -33,7 +33,7 @@ class UTexture;
 struct FFlattenMaterial;
 struct FPropertyEntry;
 enum class EFlattenMaterialProperties : uint8;
-enum EMaterialProperty;
+enum EMaterialProperty : int;
 namespace UE
 {
 	class FUsdPrim;
