USD: Fix handling of published assets when importing with asset replace policy

This finally fixes that annoying ensure on that line "ensureMsgf(false, TEXT("Texture not valid! UE-23902! Parameter (%s)"), *Parameter.ParameterInfo.Name.ToString());". You can see on the Jira that I was just working on 5.2 for this, but noticed this is happening on 5.1 as well. Not sure why I got a Jira about the 5.2 issue and nothing about the 5.1 issue?

You can see the same description on the comment, but basically there were two problems here:

The first problem is that when we import assets with EReplaceAssetPolicy::Replace, *and* we run into already existing UAssets during PublishAsset, we'll end up using DuplicateObject to "stomp" the newly imported asset on top of the existing asset (that is, we'll do the "Replace"). The actual issue here is that the asset that will remain on the "UsedAssetsAndDependencies" set (which we continue reusing throughout the import) is the original, transient asset! On the other EReplaceAssetPolicy options (or even on EReplaceAssetPolicy::Replace if we don't find an existing asset to stomp), we'll actually move the same imported asset onto its published package, so while the same asset remains on UsedAssetsAndDependencies, it goes from the transient package to the published package. When we end up using DuplicateObject though, the "new" duplicated object doesn't go into the UsedAssetsAndDependencies, but the original asset remains there, even though it remained on the transient package.

This first problem is fixed by using GetPublishedAssetsAndDependencies after doing all the replacements, and using the set of published assets from that point onwards (which is what the code was expecting in the first place).

The second problem is that apparently material instance proxies have direct pointers to the UTexture assets that they're referencing, *separate* from the UMaterialInstance's pointers. When we call our ReplaceReferences and ReflaceSoftReferences after publishing stuff, we remap the UMaterialInstance's references to the published textures, but we can't reach into the material proxies to do the same. You can see how this becomes an issue when it is combined with the first problem: It will then be possible to get these proxies pointing at transient textures, that can eventually get fully collected. Once the engine allocates a new UObject onto wherever address the proxies are pointing, some code realizes its texture parameters are not pointing at textures anymore and freaks out (you can see on Horde error message above the callstack how it prints that it's texture pointer was pointing at ShinyBlue_327 instead, which is a material).

The second problem is fixed by calling PostEditChange on the imported materials. Luckily that is enough to get it to regenerate the material proxy, which recreates its texture pointers from the material instance's, which do get remapped successfully.

diff --git a/Engine/Plugins/Importers/USDImporter/Source/USDStageImporter/Private/USDStageImporter.cpp b/Engine/Plugins/Importers/USDImporter/Source/USDStageImporter/Private/USDStageImporter.cpp
--- a/Engine/Plugins/Importers/USDImporter/Source/USDStageImporter/Private/USDStageImporter.cpp
+++ b/Engine/Plugins/Importers/USDImporter/Source/USDStageImporter/Private/USDStageImporter.cpp
@@ -1141,7 +1141,26 @@ namespace UsdStageImporterImpl
 			}
 
 			constexpr EArchiveReplaceObjectFlags ReplaceFlags = (EArchiveReplaceObjectFlags::IgnoreOuterRef | EArchiveReplaceObjectFlags::IgnoreArchetypeRef);
-			FArchiveReplaceObjectRef< UObject > ArchiveReplaceObjectRefInner(Referencer,ObjectsToRemap, ReplaceFlags);
+			FArchiveReplaceObjectRef< UObject > ArchiveReplaceObjectRefInner(Referencer, ObjectsToRemap, ReplaceFlags);
+		}
+
+
+	}
+
+	void CallAssetsPostEditChange(const TSet<UObject*>& PublishedObjects)
+	{
+		for (UObject* PublishedObject : PublishedObjects)
+		{
+			if (UMaterialInstance* Material = Cast<UMaterialInstance>(PublishedObject))
+			{
+				// After we remapped all references to the persistent assets we need to call PostEditChange on
+				// all materials. This because material instance proxies hold direct references to the UTextures that
+				// are used by the material, and those references will be left pointing at the transient assets in
+				// case we had to use DuplicateObject to "publish" them to their final locations. Calling
+				// PostEditChange rebuilds those proxies from the UMaterialInstance's parameters, which *have* been
+				// remapped by our RemapReferences and RemapSoftReferences functions.
+				Material->PostEditChange();
+			}
 		}
 	}
 
@@ -1373,6 +1392,43 @@ namespace UsdStageImporterImpl
 		AssetTools.RenameReferencingSoftObjectPaths( Packages.Array(), SoftObjectsToRemap );
 	}
 
+	void GetPublishedAssetsAndDependencies(
+		const TSet<UObject*>& UsedAssetsAndDependencies,
+		const TMap< UObject*, UObject* >& ObjectsToRemap,
+		const TMap<FSoftObjectPath, FSoftObjectPath>& SoftObjectsToRemap,
+		TSet<UObject*>& OutPublishedAssetsAndDependencies
+	)
+	{
+		OutPublishedAssetsAndDependencies.Empty(UsedAssetsAndDependencies.Num());
+
+		for (UObject* Asset : UsedAssetsAndDependencies)
+		{
+			if (!Asset)
+			{
+				continue;
+			}
+
+			UObject* RemappedAsset = ObjectsToRemap.FindRef(Asset);
+
+			// The asset itself was renamed onto its final package, no remapping needed
+			if (!RemappedAsset && Asset->GetOutermost() != GetTransientPackage())
+			{
+				OutPublishedAssetsAndDependencies.Add(Asset);
+				continue;
+			}
+
+			// If we have remapped it, ObjectsToRemap and SoftObjectsToRemap should have agreed on where to
+			const FSoftObjectPath* RemappedSoftAsset = SoftObjectsToRemap.Find(FSoftObjectPath{Asset});
+			if (!RemappedAsset || !RemappedSoftAsset || RemappedAsset != RemappedSoftAsset->TryLoad())
+			{
+				UE_LOG(LogUsd, Warning, TEXT("Failed to publish or remap asset '%s'!"), *Asset->GetPathName());
+				continue;
+			}
+
+			OutPublishedAssetsAndDependencies.Add(RemappedAsset);
+		}
+	}
+
 	/** After we remapped everything, notify the AssetRegistry that we created some new assets */
 	void NotifyAssetRegistry( const TSet<UObject*>& UsedAssetsAndDependencies )
 	{
@@ -1644,6 +1700,7 @@ void UUsdStageImporter::ImportFromFile(FUsdStageImportContext& ImportContext)
 	TMap<FSoftObjectPath, FSoftObjectPath> SoftObjectsToRemap;
 	TMap<UObject*, UObject*> ObjectsToRemap;
 	TSet<UObject*> UsedAssetsAndDependencies;
+	TSet<UObject*> PublishedAssetsAndDependencies;
 	UsdUtils::FBlendShapeMap BlendShapesByPath;
 
 	// Ensure a valid asset cache
@@ -1694,8 +1751,10 @@ void UUsdStageImporter::ImportFromFile(FUsdStageImportContext& ImportContext)
 	UsdStageImporterImpl::ResolveActorConflicts( ImportContext, ExistingSceneActor, ObjectsToRemap, SoftObjectsToRemap );
 	UsdStageImporterImpl::RemapReferences( ImportContext, UsedAssetsAndDependencies, ObjectsToRemap );
 	UsdStageImporterImpl::RemapSoftReferences( ImportContext, UsedAssetsAndDependencies, SoftObjectsToRemap );
+	UsdStageImporterImpl::GetPublishedAssetsAndDependencies(UsedAssetsAndDependencies, ObjectsToRemap, SoftObjectsToRemap, PublishedAssetsAndDependencies);
+	UsdStageImporterImpl::CallAssetsPostEditChange(PublishedAssetsAndDependencies);
 	UsdStageImporterImpl::Cleanup( ImportContext.SceneActor, ExistingSceneActor, ImportContext.ImportOptions->ExistingActorPolicy );
-	UsdStageImporterImpl::NotifyAssetRegistry( UsedAssetsAndDependencies );
+	UsdStageImporterImpl::NotifyAssetRegistry(PublishedAssetsAndDependencies);
 	UsdStageImporterImpl::RefreshComponents( ImportContext.SceneActor, ImportContext.ImportOptions->bImportAtSpecificTimeCode );
 
 	FUsdDelegates::OnPostUsdImport.Broadcast( ImportContext.FilePath );
@@ -1703,7 +1762,7 @@ void UUsdStageImporter::ImportFromFile(FUsdStageImportContext& ImportContext)
 	// Analytics
 	{
 		double ElapsedSeconds = FPlatformTime::ToSeconds64( FPlatformTime::Cycles64() - StartTime );
-		UsdStageImporterImpl::SendAnalytics( ImportContext, nullptr, TEXT("Import"), UsedAssetsAndDependencies, ElapsedSeconds);
+		UsdStageImporterImpl::SendAnalytics( ImportContext, nullptr, TEXT("Import"), PublishedAssetsAndDependencies, ElapsedSeconds);
 	}
 
 	UsdStageImporterImpl::CloseStageIfNeeded( ImportContext );
@@ -1817,8 +1876,13 @@ bool UUsdStageImporter::ReimportSingleAsset(FUsdStageImportContext& ImportContex
 
 		// Just publish the one asset we wanted to reimport. Note that we may have other assets here too, but we'll ignore those e.g. a displayColor material or a skeleton
 		OutReimportedAsset = UsdStageImporterImpl::PublishAsset(ImportContext, ReimportedObject, OriginalAsset->GetOutermost()->GetPathName(), ObjectsToRemap, SoftObjectsToRemap);
-		UsdStageImporterImpl::RemapReferences( ImportContext, ImportContext.AssetCache->GetActiveAssets(), ObjectsToRemap );
-		UsdStageImporterImpl::RemapSoftReferences( ImportContext, ImportContext.AssetCache->GetActiveAssets(), SoftObjectsToRemap );
+
+		TSet<UObject*> UsedAssetsAndDependencies = ImportContext.AssetCache->GetActiveAssets();
+		TSet<UObject*> PublishedAssetsAndDependencies;
+		UsdStageImporterImpl::RemapReferences(ImportContext, UsedAssetsAndDependencies, ObjectsToRemap);
+		UsdStageImporterImpl::RemapSoftReferences(ImportContext, UsedAssetsAndDependencies, SoftObjectsToRemap);
+		UsdStageImporterImpl::GetPublishedAssetsAndDependencies(UsedAssetsAndDependencies, ObjectsToRemap, SoftObjectsToRemap, PublishedAssetsAndDependencies);
+		UsdStageImporterImpl::CallAssetsPostEditChange(PublishedAssetsAndDependencies);
 
 		bSuccess = OutReimportedAsset != nullptr && ImportContext.AssetCache->GetActiveAssets().Contains( ReimportedObject );
 	}
