USD: Improve USD Stage actor interface.

Add new function to set the stage directly via C++;
Allow setting/getting the currently isolated via C++/Python/Blueprint.

diff --git a/Engine/Plugins/Importers/USDImporter/Source/USDStage/Private/USDStageActor.cpp b/Engine/Plugins/Importers/USDImporter/Source/USDStage/Private/USDStageActor.cpp
--- a/Engine/Plugins/Importers/USDImporter/Source/USDStage/Private/USDStageActor.cpp
+++ b/Engine/Plugins/Importers/USDImporter/Source/USDStage/Private/USDStageActor.cpp
@@ -794,7 +794,32 @@ void AUsdStageActor::NewStage()
 #endif // USE_USD_SDK
 }
 
-void AUsdStageActor::IsolateLayer(const UE::FSdfLayer& Layer)
+void AUsdStageActor::SetIsolatedRootLayer(const FString& IsolatedStageRootLayer)
+{
+	// Only clear the isolated layer if we intentionally pass an empty path
+	if (IsolatedStageRootLayer.IsEmpty())
+	{
+		IsolateLayer(UE::FSdfLayer{});
+	}
+	else
+	{
+		if (UE::FSdfLayer LayerToIsolate = UE::FSdfLayer::FindOrOpen(*IsolatedStageRootLayer))
+		{
+			IsolateLayer(LayerToIsolate);
+		}
+		else
+		{
+			UE_LOG(LogUsd, Warning, TEXT("Failed to isolate layer '%s': File does not exist or is not a valid USD layer"), *IsolatedStageRootLayer);
+		}
+	}
+}
+
+FString AUsdStageActor::GetIsolatedRootLayer() const
+{
+	return IsolatedStage ? IsolatedStage.GetRootLayer().GetIdentifier() : FString{};
+}
+
+void AUsdStageActor::IsolateLayer( const UE::FSdfLayer& Layer )
 {
 	if (IsolatedStage && IsolatedStage.GetRootLayer() == Layer)
 	{
@@ -814,11 +839,27 @@ void AUsdStageActor::IsolateLayer(const UE::FSdfLayer& Layer)
 
 	if (UsdStage)
 	{
-		// Ideally we'd check to see if Layer belonged to the stage first, but there isn't any foolproof way to check
-		// for that because if the layer is muted it won't show up on the layer stack at all. Its also valid to mute
-		// any layer anyway (even one that doesn't belong to the stage yet), so checking the list of muted layers also
-		// wouldn't improve anything: It would always be possible to trick that code into isolating an external layer
-		// by just muting the layer first...
+		// Check if the layer we're trying to isolate is part of the stage. To do that we need to reopen a fresh copy
+		// of the current stage in case we have any layers muted, as muted layers don't show up on the layer stack.
+		// Note that we can't just check the list of muted layers either, as it's possible to mute *any* layer for
+		// a given stage, not only the layers that are currently used by it.
+		// We'll use an empty population mask though (which should prevent prim composition) and just use the layers
+		// that are already opened on the current stage anyway, so this should be cheap
+		UE::FUsdStage FreshCurrentStage = UnrealUSDWrapper::OpenMaskedStage(
+			*UsdStage.GetRootLayer().GetIdentifier(),
+			EUsdInitialLoadSet::LoadNone,
+			{}
+		);
+		ensure(FreshCurrentStage);
+		TArray<UE::FSdfLayer> CurrentLayerStack = FreshCurrentStage.GetLayerStack();
+		if (!CurrentLayerStack.Contains(Layer))
+		{
+			UE_LOG(LogTemp, Warning, TEXT("Failed to isolate layer '%s' as it is not part of the currently opened USD Stage's local layer stack"),
+				*Layer.GetIdentifier(),
+				*UsdStage.GetRootLayer().GetIdentifier()
+			);
+			return;
+		}
 
 		// We really want our own stage for this and not something from the stage cache.
 		// Plus, this means its easier to cleanup: Just drop our IsolatedStage
@@ -830,6 +871,10 @@ void AUsdStageActor::IsolateLayer(const UE::FSdfLayer& Layer)
 		UsdListener.Register(IsolatedStage);
 
 		LoadUsdStage();
+
+		// Fire this so that the USD Stage Editor knows to refresh.
+		// Plus we kind of changed the active stage too
+		OnStageChanged.Broadcast();
 	}
 }
 
@@ -1409,6 +1454,44 @@ const UE::FUsdStage& AUsdStageActor::GetIsolatedUsdStage() const
 	return IsolatedStage;
 }
 
+void AUsdStageActor::SetUsdStage(const UE::FUsdStage& NewStage)
+{
+	if (UsdStage == NewStage)
+	{
+		return;
+	}
+
+	const bool bMarkDirty = false;
+	Modify(bMarkDirty);
+
+	OnPreStageChanged.Broadcast();
+	UnloadUsdStage();
+
+	RootLayer.FilePath = NewStage.GetRootLayer().GetIdentifier();
+
+	UsdStage = NewStage;
+	IsolatedStage = UE::FUsdStage{};
+
+	if (UsdStage)
+	{
+		UsdStage.SetEditTarget(UsdStage.GetRootLayer());
+
+		UsdStage.SetInterpolationType(InterpolationType);
+
+		UsdListener.Register(UsdStage);
+
+#if USE_USD_SDK
+		// Try loading a UE-state session layer if we can find one
+		const bool bCreateIfNeeded = false;
+		UsdUtils::GetUEPersistentStateSublayer(UsdStage, bCreateIfNeeded);
+#endif // #if USE_USD_SDK
+
+	}
+
+	LoadUsdStage();
+	OnStageChanged.Broadcast();
+}
+
 UE::FUsdStage& AUsdStageActor::GetOrLoadUsdStage()
 {
 	OpenUsdStage();
diff --git a/Engine/Plugins/Importers/USDImporter/Source/USDStage/Public/USDStageActor.h b/Engine/Plugins/Importers/USDImporter/Source/USDStage/Public/USDStageActor.h
--- a/Engine/Plugins/Importers/USDImporter/Source/USDStage/Public/USDStageActor.h
+++ b/Engine/Plugins/Importers/USDImporter/Source/USDStage/Public/USDStageActor.h
@@ -177,6 +177,21 @@ public:
 	UFUNCTION(BlueprintCallable, Category = "USD", meta = (CallInEditor = "true"))
 	USDSTAGE_API void NewStage();
 
+	/**
+	 * If IsolatedStageRootLayer is the identifier of one of the sublayers of the currently opened stage, this will
+	 * enter isolated mode by creating a new stage with IsolatedStageRootLayer as its root, and displaying that.
+	 * Provide an empty string to leave isolated mode.
+	 */
+	UFUNCTION(BlueprintCallable, Category = "USD", meta = (CallInEditor = "true"))
+	USDSTAGE_API void SetIsolatedRootLayer(const FString& IsolatedStageRootLayer);
+
+	/**
+	 * Returns the root layer identifier of the currently isolated stage if we're in isolated mode, and the empty
+	 * string otherwise.
+	 */
+	UFUNCTION(BlueprintCallable, Category = "USD", meta = (CallInEditor = "true"))
+	USDSTAGE_API FString GetIsolatedRootLayer() const;
+
 public:
 	// Loads the stage with RootLayer if its not loaded already, and returns the either the isolated stage (if any) or
 	// the base stage
@@ -191,6 +206,8 @@ public:
 	// Always returns the isolated stage, being an invalid stage in case we're not isolating anything
 	USDSTAGE_API const UE::FUsdStage& GetIsolatedUsdStage() const;
 
+	USDSTAGE_API void SetUsdStage(const UE::FUsdStage& NewStage);
+
 	// Enters isolated mode by creating a new USD Stage using the provided layer as its root
 	USDSTAGE_API void IsolateLayer( const UE::FSdfLayer& Layer );
 
diff --git a/Engine/Plugins/Importers/USDImporter/Source/USDTests/Private/USDTestsBlueprintLibrary.cpp b/Engine/Plugins/Importers/USDImporter/Source/USDTests/Private/USDTestsBlueprintLibrary.cpp
--- a/Engine/Plugins/Importers/USDImporter/Source/USDTests/Private/USDTestsBlueprintLibrary.cpp
+++ b/Engine/Plugins/Importers/USDImporter/Source/USDTests/Private/USDTestsBlueprintLibrary.cpp
@@ -134,3 +134,14 @@ int64 USDTestsBlueprintLibrary::GetSubtreeMaterialSlotCount( AUsdStageActor* Sta
 
 	return -1;
 }
+
+void USDTestsBlueprintLibrary::SetUsdStageCpp(AUsdStageActor* StageActor, const FString& NewStageRootLayer)
+{
+	if (!StageActor)
+	{
+		return;
+	}
+
+	UE::FUsdStage NewStage = UnrealUSDWrapper::OpenStage(*NewStageRootLayer, EUsdInitialLoadSet::LoadAll);
+	StageActor->SetUsdStage(NewStage);
+}
diff --git a/Engine/Plugins/Importers/USDImporter/Source/USDTests/Public/USDTestsBlueprintLibrary.h b/Engine/Plugins/Importers/USDImporter/Source/USDTests/Public/USDTestsBlueprintLibrary.h
--- a/Engine/Plugins/Importers/USDImporter/Source/USDTests/Public/USDTestsBlueprintLibrary.h
+++ b/Engine/Plugins/Importers/USDImporter/Source/USDTests/Public/USDTestsBlueprintLibrary.h
@@ -46,4 +46,10 @@ public:
 	 */
 	UFUNCTION( BlueprintCallable, Category = "Counts" )
 	static int64 GetSubtreeMaterialSlotCount( AUsdStageActor* StageActor, const FString& PrimPath );
+
+	/**
+	 * Internally opens the provided stage and sets it on the stage actor via C++
+	 */
+	UFUNCTION(BlueprintCallable, Category = "Counts")
+	static void SetUsdStageCpp(AUsdStageActor* StageActor, const FString& NewStageRootLayer);
 };
