USD: Expose USD Stage Editor functionality to C++/Python/Blueprints.

diff --git a/Engine/Plugins/Importers/USDImporter/Resources/PythonExamples/stage_actor_and_stage_editor.py b/Engine/Plugins/Importers/USDImporter/Resources/PythonExamples/stage_actor_and_stage_editor.py
new file mode 100644
--- /dev/null
+++ b/Engine/Plugins/Importers/USDImporter/Resources/PythonExamples/stage_actor_and_stage_editor.py
@@ -0,0 +1,65 @@
+"""
+This script describes the basic usage of the USD Stage Actor and USD Stage Editor.
+"""
+
+import os
+import unreal
+from pxr import Usd, UsdUtils, UsdGeom, Gf
+
+ROOT_LAYER_FILENAME = r"C:\MyFolder\my_scene.usda"
+DESTINATION_CONTENT_PATH = r"/Game/Imported/"
+
+# Stage Actors can be spawned on the UE level like this.
+# If the Stage Editor is open when this happens, it will automatically attach to the last actor spawned.
+editor_actor_subsystem = unreal.get_editor_subsystem(unreal.EditorActorSubsystem)
+stage_actor = editor_actor_subsystem.spawn_actor_from_class(unreal.UsdStageActor, unreal.Vector())
+
+# You can open a stage and interact with it by manipulating the Stage Actor directly:
+stage_actor.set_editor_property('root_layer', unreal.FilePath(ROOT_LAYER_FILENAME))
+
+# Any stage opened by a Stage Actor will be opened within the UsdUtils' stage cache.
+# This means you can retrieve the same stage via Python from it:
+stage = None
+for s in UsdUtils.StageCache.Get().GetAllStages():
+    if s.GetRootLayer().GetDisplayName() == os.path.basename(ROOT_LAYER_FILENAME):
+         stage = s
+         break
+
+# Interacting with the stage via Python will update the UE level automatically:
+prim = stage.GetPrimAtPath("/Cup")
+xform_api = UsdGeom.XformCommonAPI(prim)
+trans, _, _, _, _ = xform_api.GetXformVectors(Usd.TimeCode.Default())
+delta = Gf.Vec3d(100, 100, 100)
+xform_api.SetTranslate(trans + delta)
+
+# You can also toggle variants in this way, and both the USD Stage Editor and the UE level will update automatically:
+prim = stage.GetPrimAtPath("/PrimWithALODVariantSet")
+variant_sets = prim.GetVariantSets()
+if variant_sets.HasVariantSet("LOD"):
+    var_set = variant_sets.GetVariantSet("LOD")
+    for var in var_set.GetVariantNames():
+        print(f"Setting variant set '{var_set.GetName()}' with variant '{var}'")
+        var_set.SetVariantSelection(var)
+
+# Alternatively it is possible to interact with the USD Stage Editor via Python.
+# Here are some examples of operations you can do. Please consult the documentation for the full list
+unreal.UsdStageEditorLibrary.open_stage_editor()
+unreal.UsdStageEditorLibrary.file_open(ROOT_LAYER_FILENAME)
+other_stage_actor = unreal.UsdStageEditorLibrary.get_attached_stage_actor()
+unreal.UsdStageEditorLibrary.set_attached_stage_actor(stage_actor)
+selected_prims = unreal.UsdStageEditorLibrary.get_selected_prim_paths()
+
+# You can even trigger the Actions->Import button from here, as a way of importing
+# the current state of the opened stage:
+options = unreal.UsdStageImportOptions()
+options.import_actors = True
+options.import_geometry = True
+options.import_skeletal_animations = True
+options.import_level_sequences = True
+options.import_materials = True
+options.prim_path_folder_structure = False
+options.prims_to_import = [
+    "/ParentA",
+    "/ParentB/ChildC"
+]
+unreal.UsdStageEditorLibrary.actions_import(DESTINATION_CONTENT_PATH, options)
\ No newline at end of file
diff --git a/Engine/Plugins/Importers/USDImporter/Source/USDStageEditor/Private/SUSDLayersTreeView.cpp b/Engine/Plugins/Importers/USDImporter/Source/USDStageEditor/Private/SUSDLayersTreeView.cpp
--- a/Engine/Plugins/Importers/USDImporter/Source/USDStageEditor/Private/SUSDLayersTreeView.cpp
+++ b/Engine/Plugins/Importers/USDImporter/Source/USDStageEditor/Private/SUSDLayersTreeView.cpp
@@ -986,6 +986,53 @@ FReply SUsdLayersTreeView::OnRowAcceptDrop( const FDragDropEvent& Event, EItemDr
 	return FReply::Unhandled();
 }
 
+TArray<UE::FSdfLayer> SUsdLayersTreeView::GetSelectedLayers() const
+{
+	TArray<UE::FSdfLayer> SelectedLayers;
+	TArray<FUsdLayerViewModelRef> SelectedViewModels = GetSelectedItems();
+	SelectedLayers.Reserve( SelectedViewModels.Num() );
+
+	for ( const FUsdLayerViewModelRef& SelectedItem : SelectedViewModels )
+	{
+		SelectedLayers.AddUnique( SelectedItem->GetLayer() );
+	}
+
+	return SelectedLayers;
+}
+
+void SUsdLayersTreeView::SetSelectedLayers( const TArray<UE::FSdfLayer>& NewSelection )
+{
+	TSet<FString> IdentifiersToSelect;
+	IdentifiersToSelect.Reserve( NewSelection.Num() );
+	for ( const UE::FSdfLayer& SelectedLayer : NewSelection )
+	{
+		IdentifiersToSelect.Add( SelectedLayer.GetIdentifier() );
+	}
+
+	TArray<FUsdLayerViewModelRef> ItemsToSelect;
+	ItemsToSelect.Reserve( NewSelection.Num() );
+
+	TFunction<void( const TArray< FUsdLayerViewModelRef >& )> Traverse;
+	Traverse = [this, &IdentifiersToSelect, &ItemsToSelect, &Traverse]( const TArray< FUsdLayerViewModelRef >& Items )
+	{
+		for ( const FUsdLayerViewModelRef& Item : Items )
+		{
+			if ( IdentifiersToSelect.Contains( Item->LayerIdentifier ) )
+			{
+				ItemsToSelect.Add( Item );
+			}
+
+			Traverse( Item->Children );
+		}
+	};
+	Traverse( RootItems );
+
+	Private_ClearSelection();
+
+	const bool bSelected = true;
+	SetItemSelection( ItemsToSelect, bSelected );
+}
+
 TSharedRef< ITableRow > SUsdLayersTreeView::OnGenerateRow( FUsdLayerViewModelRef InDisplayNode, const TSharedRef< STableViewBase >& OwnerTable )
 {
 	return SNew( SUsdTreeRow< FUsdLayerViewModelRef >, InDisplayNode, OwnerTable, SharedData )
@@ -1130,7 +1177,10 @@ TSharedPtr< SWidget > SUsdLayersTreeView::ConstructLayerContextMenu()
 			LOCTEXT( "Export_ToolTip", "Export the selected layers, having the exported layers reference the original stage's layers" ),
 			FSlateIcon(),
 			FUIAction(
-				FExecuteAction::CreateSP( this, &SUsdLayersTreeView::OnExportSelectedLayers ),
+				FExecuteAction::CreateLambda([this]()
+				{
+					ExportSelectedLayers();
+				}),
 				FCanExecuteAction()
 			),
 			NAME_None,
@@ -1430,7 +1480,7 @@ bool SUsdLayersTreeView::CanSaveSelectedLayers() const
 	return false;
 }
 
-void SUsdLayersTreeView::OnExportSelectedLayers() const
+void SUsdLayersTreeView::ExportSelectedLayers(const FString& OutputLayerOrDirectory) const
 {
 	TArray< FUsdLayerViewModelRef > MySelectedItems = GetSelectedItems();
 
@@ -1450,20 +1500,26 @@ void SUsdLayersTreeView::OnExportSelectedLayers() const
 
 	double StartTime = FPlatformTime::Cycles64();
 	FString Extension;
+	FString OutputLayerOrDirectoryCopy = OutputLayerOrDirectory;
 
 	// Single layer -> Allow picking the target layer filename
 	if ( LayersToExport.Num() == 1 )
 	{
-		TOptional< FString > UsdFilePath = UsdUtils::BrowseUsdFile( UsdUtils::EBrowseFileMode::Save );
-		if ( !UsdFilePath.IsSet() )
+		if (OutputLayerOrDirectoryCopy.IsEmpty() )
 		{
-			return;
+			TOptional< FString > UsdFilePath = UsdUtils::BrowseUsdFile( UsdUtils::EBrowseFileMode::Save );
+			if ( !UsdFilePath.IsSet() )
+			{
+				return;
+			}
+
+			OutputLayerOrDirectoryCopy = UsdFilePath.GetValue();
 		}
 
 		StartTime = FPlatformTime::Cycles64();
-		Extension = FPaths::GetExtension( UsdFilePath.GetValue() );
+		Extension = FPaths::GetExtension(OutputLayerOrDirectoryCopy);
 
-		UE::USDLayersTreeViewImpl::Private::ExportLayerToPath( LayersToExport[ 0 ], UsdFilePath.GetValue() );
+		UE::USDLayersTreeViewImpl::Private::ExportLayerToPath( LayersToExport[ 0 ], OutputLayerOrDirectoryCopy);
 	}
 	// Multiple layers -> Pick folder and export them with the same name
 	else if ( LayersToExport.Num() > 1 )
@@ -1479,32 +1535,36 @@ void SUsdLayersTreeView::OnExportSelectedLayers() const
 			? ParentWindow->GetNativeWindow()->GetOSWindowHandle()
 			: nullptr;
 
-		FString TargetFolderPath;
-		if ( !DesktopPlatform->OpenDirectoryDialog( ParentWindowHandle, LOCTEXT( "ChooseFolder", "Choose output folder" ).ToString(), TEXT( "" ), TargetFolderPath ) )
+		if (OutputLayerOrDirectoryCopy.IsEmpty() )
 		{
-			return;
+			if ( !DesktopPlatform->OpenDirectoryDialog(
+				ParentWindowHandle,
+				LOCTEXT( "ChooseFolder", "Choose output folder" ).ToString(),
+				TEXT( "" ),
+				OutputLayerOrDirectoryCopy
+			))
+			{
+				return;
+			}
+			OutputLayerOrDirectoryCopy = FPaths::ConvertRelativePathToFull(OutputLayerOrDirectoryCopy);
 		}
-		TargetFolderPath = FPaths::ConvertRelativePathToFull( TargetFolderPath );
 
 		StartTime = FPlatformTime::Cycles64();
 		Extension = FPaths::GetExtension( LayersToExport[ 0 ].GetDisplayName() );
 
-		if ( FPaths::DirectoryExists( TargetFolderPath ) )
+		for ( const UE::FSdfLayer& LayerToExport : LayersToExport )
 		{
-			for ( const UE::FSdfLayer& LayerToExport : LayersToExport )
-			{
-				FString TargetFileName = FPaths::GetCleanFilename( LayerToExport.GetRealPath() );
-				FString FullPath = FPaths::Combine( TargetFolderPath, TargetFileName );
-				FString FinalFullPath = FullPath;
+			FString TargetFileName = FPaths::GetCleanFilename( LayerToExport.GetRealPath() );
+			FString FullPath = FPaths::Combine(OutputLayerOrDirectoryCopy, TargetFileName );
+			FString FinalFullPath = FullPath;
 
-				uint32 Suffix = 0;
-				while ( FPaths::FileExists( FinalFullPath ) )
-				{
-					FinalFullPath = FString::Printf( TEXT( "%s_%u" ), *FullPath, Suffix++ );
-				}
-
-				UE::USDLayersTreeViewImpl::Private::ExportLayerToPath( LayerToExport, FinalFullPath );
+			uint32 Suffix = 0;
+			while ( FPaths::FileExists( FinalFullPath ) )
+			{
+				FinalFullPath = FString::Printf( TEXT( "%s_%u" ), *FullPath, Suffix++ );
 			}
+
+			UE::USDLayersTreeViewImpl::Private::ExportLayerToPath( LayerToExport, FinalFullPath );
 		}
 	}
 
diff --git a/Engine/Plugins/Importers/USDImporter/Source/USDStageEditor/Private/SUSDLayersTreeView.h b/Engine/Plugins/Importers/USDImporter/Source/USDStageEditor/Private/SUSDLayersTreeView.h
--- a/Engine/Plugins/Importers/USDImporter/Source/USDStageEditor/Private/SUSDLayersTreeView.h
+++ b/Engine/Plugins/Importers/USDImporter/Source/USDStageEditor/Private/SUSDLayersTreeView.h
@@ -33,6 +33,11 @@ public:
 	FReply OnRowAcceptDrop( const FDragDropEvent& Event, EItemDropZone Zone, FUsdLayerViewModelRef Item );
 	// End drag and drop interface
 
+	TArray<UE::FSdfLayer> GetSelectedLayers() const;
+	void SetSelectedLayers( const TArray<UE::FSdfLayer>& NewSelection );
+
+	void ExportSelectedLayers(const FString& OutputDirectory = {}) const;
+
 	const UE::FUsdStageWeak& GetStage() const { return UsdStage; }
 
 private:
@@ -57,8 +62,6 @@ private:
 	void OnSaveSelectedLayers();
 	bool CanSaveSelectedLayers() const;
 
-	void OnExportSelectedLayers() const;
-
 	bool CanInsertSubLayer() const;
 	void OnAddSubLayer();
 	void OnNewSubLayer();
diff --git a/Engine/Plugins/Importers/USDImporter/Source/USDStageEditor/Private/SUSDPrimInfo.h b/Engine/Plugins/Importers/USDImporter/Source/USDStageEditor/Private/SUSDPrimInfo.h
--- a/Engine/Plugins/Importers/USDImporter/Source/USDStageEditor/Private/SUSDPrimInfo.h
+++ b/Engine/Plugins/Importers/USDImporter/Source/USDStageEditor/Private/SUSDPrimInfo.h
@@ -19,6 +19,7 @@ public:
 	void SetPrimPath( const UE::FUsdStageWeak& UsdStage, const TCHAR* PrimPath );
 
 private:
+	friend class SUsdStage;
 	TSharedPtr< class SUsdPrimPropertiesList > PropertiesList;
 	TSharedPtr< class SUsdIntegrationsPanel > IntegrationsPanel;
 	TSharedPtr< class SVariantsList > VariantsList;
diff --git a/Engine/Plugins/Importers/USDImporter/Source/USDStageEditor/Private/SUSDStage.cpp b/Engine/Plugins/Importers/USDImporter/Source/USDStageEditor/Private/SUSDStage.cpp
--- a/Engine/Plugins/Importers/USDImporter/Source/USDStageEditor/Private/SUSDStage.cpp
+++ b/Engine/Plugins/Importers/USDImporter/Source/USDStageEditor/Private/SUSDStage.cpp
@@ -22,9 +22,10 @@
 #include "USDStageImportOptions.h"
 #include "USDStageModule.h"
 #include "USDTypesConversion.h"
-
 #include "UsdWrappers/SdfLayer.h"
+#include "UsdWrappers/UsdAttribute.h"
 #include "UsdWrappers/UsdStage.h"
+#include "Widgets/SUSDPrimPropertiesList.h"
 
 #include "ActorTreeItem.h"
 #include "Async/Async.h"
@@ -43,6 +44,7 @@
 #include "ScopedTransaction.h"
 #include "Styling/AppStyle.h"
 #include "UObject/StrongObjectPtr.h"
+#include "Widgets/Input/SButton.h"
 #include "Widgets/Input/SComboBox.h"
 #include "Widgets/Input/SComboButton.h"
 #include "Widgets/Input/SButton.h"
@@ -395,6 +397,65 @@ void SUsdStage::Construct( const FArguments& InArgs )
 	}
 }
 
+AUsdStageActor* SUsdStage::GetAttachedStageActor() const
+{
+	return ViewModel.UsdStageActor.Get();
+}
+
+TArray<UE::FSdfLayer> SUsdStage::GetSelectedLayers() const
+{
+	if ( UsdLayersTreeView )
+	{
+		return UsdLayersTreeView->GetSelectedLayers();
+	}
+
+	return {};
+}
+
+void SUsdStage::SetSelectedLayers( const TArray<UE::FSdfLayer>& NewSelection ) const
+{
+	if ( UsdLayersTreeView )
+	{
+		UsdLayersTreeView->SetSelectedLayers( NewSelection );
+	}
+}
+
+TArray<UE::FUsdPrim> SUsdStage::GetSelectedPrims() const
+{
+	if ( UsdStageTreeView )
+	{
+		return UsdStageTreeView->GetSelectedPrims();
+	}
+
+	return {};
+}
+
+void SUsdStage::SetSelectedPrims( const TArray<UE::FUsdPrim>& NewSelection ) const
+{
+	if (UsdStageTreeView )
+	{
+		UsdStageTreeView->SetSelectedPrims( NewSelection );
+	}
+}
+
+TArray<FString> SUsdStage::GetSelectedPropertyNames() const
+{
+	if( UsdPrimInfoWidget && UsdPrimInfoWidget->PropertiesList )
+	{
+		return UsdPrimInfoWidget->PropertiesList->GetSelectedPropertyNames();
+	}
+
+	return {};
+}
+
+void SUsdStage::SetSelectedPropertyNames( const TArray<FString>& NewSelection )
+{
+	if ( UsdPrimInfoWidget && UsdPrimInfoWidget->PropertiesList )
+	{
+		UsdPrimInfoWidget->PropertiesList->SetSelectedPropertyNames( NewSelection );
+	}
+}
+
 void SUsdStage::AttachToStageActor( AUsdStageActor* InUsdStageActor, bool bFlashButton )
 {
 	if ( ViewModel.UsdStageActor == InUsdStageActor )
@@ -829,7 +890,7 @@ void SUsdStage::FillFileMenu( FMenuBuilder& MenuBuilder )
 			LOCTEXT("New_ToolTip", "Creates a new layer and opens the stage with it at its root"),
 			FSlateIcon(),
 			FUIAction(
-				FExecuteAction::CreateSP( this, &SUsdStage::OnNew ),
+				FExecuteAction::CreateSP( this, &SUsdStage::FileNew ),
 				FCanExecuteAction()
 			),
 			NAME_None,
@@ -841,7 +902,10 @@ void SUsdStage::FillFileMenu( FMenuBuilder& MenuBuilder )
 			LOCTEXT("Open_ToolTip", "Opens a USD file"),
 			FSlateIcon(),
 			FUIAction(
-				FExecuteAction::CreateSP( this, &SUsdStage::OnOpen ),
+				FExecuteAction::CreateLambda([this]()
+				{
+					FileOpen();
+				}),
 				FCanExecuteAction()
 			),
 			NAME_None,
@@ -853,7 +917,10 @@ void SUsdStage::FillFileMenu( FMenuBuilder& MenuBuilder )
 			LOCTEXT("Save_ToolTip", "Saves the stage"),
 			FSlateIcon(),
 			FUIAction(
-				FExecuteAction::CreateSP( this, &SUsdStage::OnSave ),
+				FExecuteAction::CreateLambda([this]()
+				{
+					FileSave();
+				}),
 				FCanExecuteAction::CreateLambda( [this]()
 				{
 					if ( const AUsdStageActor* StageActor = ViewModel.UsdStageActor.Get() )
@@ -886,7 +953,7 @@ void SUsdStage::FillFileMenu( FMenuBuilder& MenuBuilder )
 			LOCTEXT("Reload_ToolTip", "Reloads the stage from disk, keeping aspects of the session intact"),
 			FSlateIcon(),
 			FUIAction(
-				FExecuteAction::CreateSP( this, &SUsdStage::OnReloadStage ),
+				FExecuteAction::CreateSP( this, &SUsdStage::FileReload ),
 				FCanExecuteAction::CreateLambda( [this]()
 				{
 					if ( const AUsdStageActor* StageActor = ViewModel.UsdStageActor.Get() )
@@ -912,7 +979,7 @@ void SUsdStage::FillFileMenu( FMenuBuilder& MenuBuilder )
 			LOCTEXT( "ResetState_ToolTip", "Resets the session layer and other options like edit target and muted layers" ),
 			FSlateIcon(),
 			FUIAction(
-				FExecuteAction::CreateSP( this, &SUsdStage::OnResetStage ),
+				FExecuteAction::CreateSP( this, &SUsdStage::FileReset ),
 				FCanExecuteAction::CreateLambda( [this]()
 				{
 					if ( const AUsdStageActor* StageActor = ViewModel.UsdStageActor.Get() )
@@ -939,7 +1006,7 @@ void SUsdStage::FillFileMenu( FMenuBuilder& MenuBuilder )
 			LOCTEXT("Close_ToolTip", "Closes the opened stage"),
 			FSlateIcon(),
 			FUIAction(
-				FExecuteAction::CreateSP( this, &SUsdStage::OnClose ),
+				FExecuteAction::CreateSP( this, &SUsdStage::FileClose ),
 				FCanExecuteAction::CreateLambda( [this]()
 				{
 					if ( const AUsdStageActor* StageActor = ViewModel.UsdStageActor.Get() )
@@ -969,7 +1036,7 @@ void SUsdStage::FillActionsMenu( FMenuBuilder& MenuBuilder )
 			LOCTEXT("Import_ToolTip", "Imports the stage as Unreal assets"),
 			FSlateIcon(),
 			FUIAction(
-				FExecuteAction::CreateSP( this, &SUsdStage::OnImport ),
+				FExecuteAction::CreateSP( this, &SUsdStage::ActionsImportWithDialog ),
 				FCanExecuteAction::CreateLambda( [this]()
 				{
 					if ( const AUsdStageActor* StageActor = ViewModel.UsdStageActor.Get() )
@@ -1051,12 +1118,17 @@ void SUsdStage::FillOptionsMenu(FMenuBuilder& MenuBuilder)
 
 void SUsdStage::FillExportSubMenu( FMenuBuilder& MenuBuilder )
 {
+	const FString FilePath;
+
 	MenuBuilder.AddMenuEntry(
 		LOCTEXT( "ExportAll", "Export all layers..." ),
 		LOCTEXT( "ExportAll_ToolTip", "Exports copies of all file-based layers in the stage's layer stack to a new folder" ),
 		FSlateIcon(),
 		FUIAction(
-			FExecuteAction::CreateSP( this, &SUsdStage::OnExportAll ),
+			FExecuteAction::CreateLambda([this]()
+			{
+				FileExportAllLayers();
+			}),
 			FCanExecuteAction::CreateLambda( [this]()
 			{
 				if ( const AUsdStageActor* StageActor = ViewModel.UsdStageActor.Get() )
@@ -1079,7 +1151,10 @@ void SUsdStage::FillExportSubMenu( FMenuBuilder& MenuBuilder )
 		LOCTEXT( "ExportFlattened_ToolTip", "Flattens the current stage to a single USD layer and exports it as a new USD file" ),
 		FSlateIcon(),
 		FUIAction(
-			FExecuteAction::CreateSP( this, &SUsdStage::OnExportFlattened ),
+			FExecuteAction::CreateLambda([this]()
+			{
+				FileExportFlattenedStage();
+			}),
 			FCanExecuteAction::CreateLambda( [this]()
 			{
 				if ( const AUsdStageActor* StageActor = ViewModel.UsdStageActor.Get() )
@@ -1742,7 +1817,7 @@ void SUsdStage::FillSelectionSubMenu( FMenuBuilder& MenuBuilder )
 					{
 						TGuardValue<bool> SelectionLoopGuard( bUpdatingViewportSelection, true );
 
-						SUSDStageImpl::SelectGeneratedComponentsAndActors( ViewModel.UsdStageActor.Get(), UsdStageTreeView->GetSelectedPrims() );
+						SUSDStageImpl::SelectGeneratedComponentsAndActors( ViewModel.UsdStageActor.Get(), UsdStageTreeView->GetSelectedPrimPaths() );
 					}
 				}
 			}),
@@ -1780,22 +1855,29 @@ void SUsdStage::FillNaniteThresholdSubMenu( FMenuBuilder& MenuBuilder )
 	MenuBuilder.AddWidget( Slider, FText::FromString( TEXT( "Triangle threshold: " ) ), bNoIndent );
 }
 
-void SUsdStage::OnNew()
+void SUsdStage::FileNew()
 {
 	ViewModel.NewStage();
 }
 
-void SUsdStage::OnOpen()
+void SUsdStage::FileOpen( const FString& FilePath )
 {
-	TOptional< FString > UsdFilePath = UsdUtils::BrowseUsdFile( UsdUtils::EBrowseFileMode::Open );
-
-	if ( UsdFilePath )
+	FString FilePathCopy = FilePath;
+	if (FilePathCopy.IsEmpty())
 	{
-		OpenStage( *UsdFilePath.GetValue() );
+		TOptional< FString > UsdFilePath = UsdUtils::BrowseUsdFile( UsdUtils::EBrowseFileMode::Open );
+		if( !UsdFilePath.IsSet() )
+		{
+			return;
+		}
+
+		FilePathCopy = UsdFilePath.GetValue();
 	}
+
+	OpenStage(*FilePathCopy);
 }
 
-void SUsdStage::OnSave()
+void SUsdStage::FileSave( const FString& FilePath )
 {
 	const AUsdStageActor* StageActor = ViewModel.UsdStageActor.Get();
 	if ( !StageActor )
@@ -1818,16 +1900,24 @@ void SUsdStage::OnSave()
 		}
 		else
 		{
-			TOptional< FString > UsdFilePath = UsdUtils::BrowseUsdFile( UsdUtils::EBrowseFileMode::Save );
-			if ( UsdFilePath )
+			FString FilePathCopy = FilePath;
+			if (FilePathCopy.IsEmpty())
 			{
-				ViewModel.SaveStageAs( *UsdFilePath.GetValue() );
+				TOptional< FString > UsdFilePath = UsdUtils::BrowseUsdFile( UsdUtils::EBrowseFileMode::Save );
+				if( !UsdFilePath.IsSet() )
+				{
+					return;
+				}
+
+				FilePathCopy = UsdFilePath.GetValue();
 			}
+
+			ViewModel.SaveStageAs(*FilePathCopy);
 		}
 	}
 }
 
-void SUsdStage::OnExportAll()
+void SUsdStage::FileExportAllLayers( const FString& OutputDirectory )
 {
 	const AUsdStageActor* StageActor = ViewModel.UsdStageActor.Get();
 	if ( !StageActor )
@@ -1859,12 +1949,21 @@ void SUsdStage::OnExportAll()
 		? ParentWindow->GetNativeWindow()->GetOSWindowHandle()
 		: nullptr;
 
-	FString TargetFolderPath;
-	if ( !DesktopPlatform->OpenDirectoryDialog( ParentWindowHandle, LOCTEXT( "ChooseFolder", "Choose output folder" ).ToString(), TEXT( "" ), TargetFolderPath ) )
+	FString OutputDirectoryCopy = OutputDirectory;
+	if (OutputDirectoryCopy.IsEmpty())
 	{
-		return;
+		if ( !DesktopPlatform->OpenDirectoryDialog(
+			ParentWindowHandle,
+			LOCTEXT( "ChooseFolder", "Choose output folder" ).ToString(),
+			TEXT( "" ),
+			OutputDirectoryCopy
+		) )
+		{
+			return;
+		}
+
+		OutputDirectoryCopy = FPaths::ConvertRelativePathToFull(OutputDirectoryCopy);
 	}
-	TargetFolderPath = FPaths::ConvertRelativePathToFull( TargetFolderPath );
 
 	double StartTime = FPlatformTime::Cycles64();
 
@@ -1898,7 +1997,7 @@ void SUsdStage::OnExportAll()
 		FString LayerPath = Layer.GetRealPath();
 		FPaths::NormalizeFilename( LayerPath );
 
-		FString TargetPath = FPaths::Combine( TargetFolderPath, Layer.GetDisplayName() );
+		FString TargetPath = FPaths::Combine( OutputDirectory, Layer.GetDisplayName() );
 		FPaths::NormalizeFilename( TargetPath );
 
 		// Filename collision (should be rare, but possible given that we're discarding the folder structure)
@@ -1972,7 +2071,7 @@ void SUsdStage::OnExportAll()
 	}
 }
 
-void SUsdStage::OnExportFlattened()
+void SUsdStage::FileExportFlattenedStage(const FString& OutputLayer)
 {
 	const AUsdStageActor* StageActor = ViewModel.UsdStageActor.Get();
 	if ( !StageActor )
@@ -1986,22 +2085,28 @@ void SUsdStage::OnExportFlattened()
 		return;
 	}
 
-	TOptional< FString > UsdFilePath = UsdUtils::BrowseUsdFile( UsdUtils::EBrowseFileMode::Save );
-	if ( !UsdFilePath.IsSet() )
+	FString OutputLayerCopy = OutputLayer;
+	if (OutputLayerCopy.IsEmpty())
 	{
-		return;
+		TOptional< FString > UsdFilePath = UsdUtils::BrowseUsdFile( UsdUtils::EBrowseFileMode::Save );
+		if ( !UsdFilePath.IsSet() )
+		{
+			return;
+		}
+
+		OutputLayerCopy = UsdFilePath.GetValue();
 	}
 
 	double StartTime = FPlatformTime::Cycles64();
 
-	UsdStage.Export( *( UsdFilePath.GetValue() ) );
+	UsdStage.Export(*OutputLayerCopy);
 
 	// Send analytics
 	if ( FEngineAnalytics::IsAvailable() )
 	{
 		bool bAutomated = false;
 		double ElapsedSeconds = FPlatformTime::ToSeconds64( FPlatformTime::Cycles64() - StartTime );
-		FString Extension = FPaths::GetExtension( UsdFilePath.GetValue() );
+		FString Extension = FPaths::GetExtension(OutputLayerCopy);
 		IUsdClassesModule::SendAnalytics(
 			{},
 			TEXT( "ExportStageFlattened" ),
@@ -2013,7 +2118,7 @@ void SUsdStage::OnExportFlattened()
 	}
 }
 
-void SUsdStage::OnReloadStage()
+void SUsdStage::FileReload()
 {
 	FScopedTransaction Transaction( LOCTEXT( "ReloadTransaction", "Reload USD stage" ) );
 
@@ -2028,7 +2133,7 @@ void SUsdStage::OnReloadStage()
 	}
 }
 
-void SUsdStage::OnResetStage()
+void SUsdStage::FileReset()
 {
 	ViewModel.ResetStage();
 
@@ -2041,7 +2146,7 @@ void SUsdStage::OnResetStage()
 	}
 }
 
-void SUsdStage::OnClose()
+void SUsdStage::FileClose()
 {
 	FScopedTransaction Transaction(LOCTEXT("CloseTransaction", "Close USD stage"));
 
@@ -2059,11 +2164,24 @@ void SUsdStage::OnLayerIsolated( const UE::FSdfLayer& IsolatedLayer )
 	}
 }
 
-void SUsdStage::OnImport()
+void SUsdStage::ActionsImportWithDialog()
 {
 	ViewModel.ImportStage();
 }
 
+void SUsdStage::ActionsImport( const FString& OutputContentFolder, UUsdStageImportOptions* Options )
+{
+	ViewModel.ImportStage( *OutputContentFolder, Options );
+}
+
+void SUsdStage::ExportSelectedLayers(const FString& OutputLayerOrDirectory)
+{
+	if ( UsdLayersTreeView )
+	{
+		UsdLayersTreeView->ExportSelectedLayers(OutputLayerOrDirectory);
+	}
+}
+
 void SUsdStage::OnPrimSelectionChanged( const TArray<FString>& PrimPaths )
 {
 	if ( bUpdatingPrimSelection )
@@ -2114,6 +2232,12 @@ void SUsdStage::OpenStage( const TCHAR* FilePath )
 		FText::FromString( FilePath )
 	) );
 
+	if ( !ViewModel.UsdStageActor.IsValid() )
+	{
+		IUsdStageModule& UsdStageModule = FModuleManager::Get().LoadModuleChecked< IUsdStageModule >( "UsdStage" );
+		AttachToStageActor( &UsdStageModule.GetUsdStageActor( GWorld ) );
+	}
+
 	ViewModel.OpenStage( FilePath );
 }
 
@@ -2142,7 +2266,7 @@ void SUsdStage::Refresh()
 
 		// Refresh will generate brand new RootItems, so let's immediately select the new item
 		// that corresponds to the prim we're supposed to be selecting
-		UsdStageTreeView->SelectPrims( { SelectedPrimPath } );
+		UsdStageTreeView->SetSelectedPrimPaths( { SelectedPrimPath } );
 
 		UsdStageTreeView->RequestTreeRefresh();
 	}
@@ -2151,16 +2275,6 @@ void SUsdStage::Refresh()
 	{
 		UsdPrimInfoWidget->SetPrimPath( GetCurrentStage(), *SelectedPrimPath );
 	}
-
-	// Automatically attach to another stage actor if we can find one
-	if ( !StageActor )
-	{
-		IUsdStageModule& UsdStageModule = FModuleManager::Get().LoadModuleChecked< IUsdStageModule >( "UsdStage" );
-		if ( AUsdStageActor* OtherStageActor = UsdStageModule.FindUsdStageActor( IUsdClassesModule::GetCurrentWorld() ) )
-		{
-			AttachToStageActor( OtherStageActor );
-		}
-	}
 }
 
 void SUsdStage::OnViewportSelectionChanged( UObject* NewSelection )
@@ -2254,10 +2368,7 @@ void SUsdStage::OnViewportSelectionChanged( UObject* NewSelection )
 
 	if ( PrimPaths.Num() > 0 )
 	{
-		if ( UsdStageTreeView )
-		{
-			UsdStageTreeView->SelectPrims( PrimPaths );
-		}
+		UsdStageTreeView->SetSelectedPrimPaths( PrimPaths );
 
 		if ( UsdPrimInfoWidget )
 		{
diff --git a/Engine/Plugins/Importers/USDImporter/Source/USDStageEditor/Private/SUSDStage.h b/Engine/Plugins/Importers/USDImporter/Source/USDStageEditor/Private/SUSDStage.h
--- a/Engine/Plugins/Importers/USDImporter/Source/USDStageEditor/Private/SUSDStage.h
+++ b/Engine/Plugins/Importers/USDImporter/Source/USDStageEditor/Private/SUSDStage.h
@@ -3,7 +3,7 @@
 #pragma once
 
 #include "USDStageViewModel.h"
-#include "UsdWrappers/UsdStage.h"
+#include "UsdWrappers/ForwardDeclarations.h"
 
 #include "Animation/CurveSequence.h"
 #include "CoreMinimal.h"
@@ -22,6 +22,11 @@ class ISceneOutliner;
 enum class EMapChangeType : uint8;
 enum class EUsdInitialLoadSet : uint8;
 struct FSlateBrush;
+namespace UE
+{
+	class FUsdPrim;
+	class FUsdAttribute;
+}
 
 #if USE_USD_SDK
 
@@ -30,6 +35,7 @@ class SUsdStage : public SCompoundWidget
 	SLATE_BEGIN_ARGS( SUsdStage ) {}
 	SLATE_END_ARGS()
 
+public:
 	void Construct( const FArguments& InArgs );
 
 	virtual ~SUsdStage();
@@ -39,6 +45,32 @@ class SUsdStage : public SCompoundWidget
 	// that the actor changed
 	void AttachToStageActor( AUsdStageActor* InUsdStageActor, bool bFlashButton=true );
 
+	AUsdStageActor* GetAttachedStageActor() const;
+
+	TArray<UE::FSdfLayer> GetSelectedLayers() const;
+	void SetSelectedLayers( const TArray<UE::FSdfLayer>& NewSelection ) const;
+
+	TArray<UE::FUsdPrim> GetSelectedPrims() const;
+	void SetSelectedPrims( const TArray<UE::FUsdPrim>& NewSelection ) const;
+
+	TArray<FString> GetSelectedPropertyNames() const;
+	void SetSelectedPropertyNames( const TArray<FString>& NewSelection );
+
+	// Main menu actions.
+	// For all of these, providing an empty path will cause us to pop open a dialog to let the user pick the path
+	// instead.
+	void FileNew();
+	void FileOpen(const FString& FilePath = {});
+	void FileSave(const FString& OutputFileIfUnsaved = {});
+	void FileExportAllLayers(const FString& OutputDirectory = {});
+	void FileExportFlattenedStage(const FString& OutputLayer = {});
+	void FileReload();
+	void FileReset();
+	void FileClose();
+	void ActionsImportWithDialog();
+	void ActionsImport( const FString& OutputContentFolder, UUsdStageImportOptions* Options );
+	void ExportSelectedLayers(const FString& OutputLayerOrDirectory = {});
+
 protected:
 	void SetupStageActorDelegates();
 	void ClearStageActorDelegates();
@@ -61,19 +93,8 @@ protected:
 	void FillSelectionSubMenu( FMenuBuilder& MenuBuilder );
 	void FillNaniteThresholdSubMenu( FMenuBuilder& MenuBuilder );
 
-	void OnNew();
-	void OnOpen();
-	void OnSave();
-	void OnExportAll();
-	void OnExportFlattened();
-	void OnReloadStage();
-	void OnResetStage();
-	void OnClose();
-
 	void OnLayerIsolated( const UE::FSdfLayer& IsolatedLayer );
 
-	void OnImport();
-
 	void OnPrimSelectionChanged( const TArray<FString>& PrimPath );
 
 	void OpenStage( const TCHAR* FilePath );
diff --git a/Engine/Plugins/Importers/USDImporter/Source/USDStageEditor/Private/SUSDStageTreeView.cpp b/Engine/Plugins/Importers/USDImporter/Source/USDStageEditor/Private/SUSDStageTreeView.cpp
--- a/Engine/Plugins/Importers/USDImporter/Source/USDStageEditor/Private/SUSDStageTreeView.cpp
+++ b/Engine/Plugins/Importers/USDImporter/Source/USDStageEditor/Private/SUSDStageTreeView.cpp
@@ -330,7 +330,7 @@ void SUsdStageTreeView::Construct( const FArguments& InArgs )
 			SelectedPrimPath = UsdToUnreal::ConvertPath( UsdStageTreeItem->UsdPrim.GetPrimPath() );
 		}
 
-		TArray<FString> SelectedPrimPaths = GetSelectedPrims();
+		TArray<FString> SelectedPrimPaths = GetSelectedPrimPaths();
 		this->OnPrimSelectionChanged.ExecuteIfBound( SelectedPrimPaths );
 	} );
 
@@ -516,19 +516,8 @@ FUsdPrimViewModelPtr SUsdStageTreeView::GetItemFromPrimPath( const FString& Prim
 	return FoundItem;
 }
 
-void SUsdStageTreeView::SelectPrims( const TArray<FString>& PrimPaths )
+void SUsdStageTreeView::SelectItemsInternal( const TArray<FUsdPrimViewModelRef>& ItemsToSelect )
 {
-	TArray< FUsdPrimViewModelRef > ItemsToSelect;
-	ItemsToSelect.Reserve( PrimPaths.Num() );
-
-	for ( const FString& PrimPath : PrimPaths )
-	{
-		if ( FUsdPrimViewModelPtr FoundItem = GetItemFromPrimPath( PrimPath ) )
-		{
-			ItemsToSelect.Add( FoundItem.ToSharedRef() );
-		}
-	}
-
 	if ( ItemsToSelect.Num() > 0 )
 	{
 		// Clear selection without emitting events, as we'll emit new events with SetItemSelection
@@ -546,11 +535,56 @@ void SUsdStageTreeView::SelectPrims( const TArray<FString>& PrimPaths )
 
 		// ClearSelection is not going to fire the OnSelectionChanged event in case we have nothing selected, but we
 		// need to do that to refresh the prim properties panel to display the stage properties instead
-		OnPrimSelectionChanged.ExecuteIfBound({});
+		OnPrimSelectionChanged.ExecuteIfBound( {} );
 	}
 }
 
-TArray<FString> SUsdStageTreeView::GetSelectedPrims()
+void SUsdStageTreeView::SetSelectedPrimPaths( const TArray<FString>& PrimPaths )
+{
+	TArray< FUsdPrimViewModelRef > ItemsToSelect;
+	ItemsToSelect.Reserve( PrimPaths.Num() );
+
+	for ( const FString& PrimPath : PrimPaths )
+	{
+		if ( FUsdPrimViewModelPtr FoundItem = GetItemFromPrimPath( PrimPath ) )
+		{
+			ItemsToSelect.Add( FoundItem.ToSharedRef() );
+		}
+	}
+
+	SelectItemsInternal( ItemsToSelect );
+}
+
+void SUsdStageTreeView::SetSelectedPrims( const TArray<UE::FUsdPrim>& Prims )
+{
+	TSet<UE::FSdfPath> PrimPaths;
+	for ( const UE::FUsdPrim& Prim : Prims )
+	{
+		PrimPaths.Add( Prim.GetPrimPath() );
+	}
+
+	TArray< FUsdPrimViewModelRef > ItemsToSelect;
+	ItemsToSelect.Reserve( Prims.Num() );
+
+	TFunction<void( const TArray< FUsdPrimViewModelRef >& )> Traverse;
+	Traverse = [this, &PrimPaths, &ItemsToSelect, &Traverse]( const TArray< FUsdPrimViewModelRef >& Items )
+	{
+		for ( const FUsdPrimViewModelRef& Item : Items )
+		{
+			if ( PrimPaths.Contains( Item->UsdPrim.GetPrimPath() ) )
+			{
+				ItemsToSelect.Add( Item );
+			}
+
+			Traverse( Item->Children );
+		}
+	};
+	Traverse( RootItems );
+
+	SelectItemsInternal( ItemsToSelect );
+}
+
+TArray<FString> SUsdStageTreeView::GetSelectedPrimPaths()
 {
 	TArray<FString> SelectedPrimPaths;
 	SelectedPrimPaths.Reserve( GetNumItemsSelected() );
@@ -563,6 +597,19 @@ TArray<FString> SUsdStageTreeView::GetSelectedPrims()
 	return SelectedPrimPaths;
 }
 
+TArray<UE::FUsdPrim> SUsdStageTreeView::GetSelectedPrims()
+{
+	TArray<UE::FUsdPrim> SelectedPrims;
+	SelectedPrims.Reserve( GetNumItemsSelected() );
+
+	for ( FUsdPrimViewModelRef SelectedItem : GetSelectedItems() )
+	{
+		SelectedPrims.Add( SelectedItem->UsdPrim );
+	}
+
+	return SelectedPrims;
+}
+
 void SUsdStageTreeView::SetupColumns()
 {
 	HeaderRowWidget->ClearColumns();
diff --git a/Engine/Plugins/Importers/USDImporter/Source/USDStageEditor/Private/SUSDStageTreeView.h b/Engine/Plugins/Importers/USDImporter/Source/USDStageEditor/Private/SUSDStageTreeView.h
--- a/Engine/Plugins/Importers/USDImporter/Source/USDStageEditor/Private/SUSDStageTreeView.h
+++ b/Engine/Plugins/Importers/USDImporter/Source/USDStageEditor/Private/SUSDStageTreeView.h
@@ -13,6 +13,10 @@ class AUsdStageActor;
 class FUICommandList;
 enum class EPayloadsTrigger;
 enum class EUsdDuplicateType : uint8;
+namespace UE
+{
+	class FUsdPrim;
+}
 
 #if USE_USD_SDK
 
@@ -33,8 +37,10 @@ public:
 
 	FUsdPrimViewModelPtr GetItemFromPrimPath( const FString& PrimPath );
 
-	void SelectPrims( const TArray<FString>& PrimPaths );
-	TArray<FString> GetSelectedPrims();
+	void SetSelectedPrimPaths( const TArray<FString>& PrimPaths );
+	void SetSelectedPrims( const TArray<UE::FUsdPrim>& Prims );
+	TArray<FString> GetSelectedPrimPaths();
+	TArray<UE::FUsdPrim> GetSelectedPrims();
 
 private:
 	virtual TSharedRef< ITableRow > OnGenerateRow( FUsdPrimViewModelRef InDisplayNode, const TSharedRef< STableViewBase >& OwnerTable ) override;
@@ -86,6 +92,8 @@ private:
 	void RestoreExpansionStates();
 	virtual void RequestListRefresh() override;
 
+	void SelectItemsInternal( const TArray< FUsdPrimViewModelRef >& ItemsToSelect );
+
 private:
 	// Should always be valid, we keep the one we're given on Refresh()
 	UE::FUsdStageWeak UsdStage;
diff --git a/Engine/Plugins/Importers/USDImporter/Source/USDStageEditor/Private/USDStageEditorBlueprintLibrary.cpp b/Engine/Plugins/Importers/USDImporter/Source/USDStageEditor/Private/USDStageEditorBlueprintLibrary.cpp
new file mode 100644
--- /dev/null
+++ b/Engine/Plugins/Importers/USDImporter/Source/USDStageEditor/Private/USDStageEditorBlueprintLibrary.cpp
@@ -0,0 +1,182 @@
+// Copyright Epic Games, Inc. All Rights Reserved.
+
+#include "USDStageEditorBlueprintLibrary.h"
+
+#include "USDStageActor.h"
+#include "USDStageEditorModule.h"
+#include "UsdWrappers/SdfLayer.h"
+#include "UsdWrappers/UsdPrim.h"
+
+#include "CoreMinimal.h"
+#include "Modules/ModuleManager.h"
+
+bool UUsdStageEditorBlueprintLibrary::OpenStageEditor()
+{
+	IUsdStageEditorModule& StageEditorModule = FModuleManager::GetModuleChecked< IUsdStageEditorModule >( "USDStageEditor" );
+	return StageEditorModule.OpenStageEditor();
+}
+
+bool UUsdStageEditorBlueprintLibrary::CloseStageEditor()
+{
+	IUsdStageEditorModule& StageEditorModule = FModuleManager::GetModuleChecked< IUsdStageEditorModule >( "USDStageEditor" );
+	return StageEditorModule.CloseStageEditor();
+}
+
+AUsdStageActor* UUsdStageEditorBlueprintLibrary::GetAttachedStageActor()
+{
+	IUsdStageEditorModule& StageEditorModule = FModuleManager::GetModuleChecked< IUsdStageEditorModule >( "USDStageEditor" );
+	return StageEditorModule.GetAttachedStageActor();
+}
+
+bool UUsdStageEditorBlueprintLibrary::SetAttachedStageActor( AUsdStageActor* NewActor )
+{
+	IUsdStageEditorModule& StageEditorModule = FModuleManager::GetModuleChecked< IUsdStageEditorModule >( "USDStageEditor" );
+	return StageEditorModule.SetAttachedStageActor( NewActor );
+}
+
+TArray<FString> UUsdStageEditorBlueprintLibrary::GetSelectedLayerIdentifiers()
+{
+	IUsdStageEditorModule& StageEditorModule = FModuleManager::GetModuleChecked< IUsdStageEditorModule >( "USDStageEditor" );
+	TArray<UE::FSdfLayer> Layers = StageEditorModule.GetSelectedLayers();
+
+	TArray<FString> LayerIdentifiers;
+	LayerIdentifiers.Reserve( Layers.Num() );
+
+	for ( const UE::FSdfLayer& Layer : Layers )
+	{
+		LayerIdentifiers.Add( Layer.GetIdentifier() );
+	}
+
+	return LayerIdentifiers;
+}
+
+void UUsdStageEditorBlueprintLibrary::SetSelectedLayerIdentifiers( const TArray<FString>& NewSelection )
+{
+	TArray<UE::FSdfLayer> Layers;
+	for ( const FString& Identifier : NewSelection )
+	{
+		Layers.Add( UE::FSdfLayer::FindOrOpen( *Identifier ) );
+	}
+
+	IUsdStageEditorModule& StageEditorModule = FModuleManager::GetModuleChecked< IUsdStageEditorModule >( "USDStageEditor" );
+	StageEditorModule.SetSelectedLayers( Layers );
+}
+
+TArray<FString> UUsdStageEditorBlueprintLibrary::GetSelectedPrimPaths()
+{
+	IUsdStageEditorModule& StageEditorModule = FModuleManager::GetModuleChecked< IUsdStageEditorModule >( "USDStageEditor" );
+	TArray<UE::FUsdPrim> Prims = StageEditorModule.GetSelectedPrims();
+
+	TArray<FString> PrimPaths;
+	PrimPaths.Reserve( Prims.Num() );
+
+	for ( const UE::FUsdPrim& Prim : Prims )
+	{
+		PrimPaths.Add( Prim.GetPrimPath().GetString() );
+	}
+
+	return PrimPaths;
+}
+
+void UUsdStageEditorBlueprintLibrary::SetSelectedPrimPaths( const TArray<FString>& NewSelection )
+{
+	IUsdStageEditorModule& StageEditorModule = FModuleManager::GetModuleChecked< IUsdStageEditorModule >( "USDStageEditor" );
+	const AUsdStageActor* StageActor = StageEditorModule.GetAttachedStageActor();
+	if ( !StageActor )
+	{
+		return;
+	}
+
+	UE::FUsdStage Stage = StageActor->GetUsdStage();
+	if ( !Stage )
+	{
+		return;
+	}
+
+	TArray<UE::FUsdPrim> Prims;
+	Prims.Reserve( NewSelection.Num() );
+
+	for ( const FString& PrimPath : NewSelection )
+	{
+		if ( UE::FUsdPrim SelectedPrim = Stage.GetPrimAtPath( UE::FSdfPath{ *PrimPath } ) )
+		{
+			Prims.Add( SelectedPrim );
+		}
+	}
+
+	StageEditorModule.SetSelectedPrims( Prims );
+}
+
+TArray<FString> UUsdStageEditorBlueprintLibrary::GetSelectedPropertyNames()
+{
+	IUsdStageEditorModule& StageEditorModule = FModuleManager::GetModuleChecked< IUsdStageEditorModule >( "USDStageEditor" );
+	return StageEditorModule.GetSelectedPropertyNames();
+}
+
+void UUsdStageEditorBlueprintLibrary::SetSelectedPropertyNames( const TArray<FString>& NewSelection )
+{
+	IUsdStageEditorModule& StageEditorModule = FModuleManager::GetModuleChecked< IUsdStageEditorModule >( "USDStageEditor" );
+	StageEditorModule.SetSelectedPropertyNames( NewSelection );
+}
+
+void UUsdStageEditorBlueprintLibrary::FileNew()
+{
+	IUsdStageEditorModule& StageEditorModule = FModuleManager::GetModuleChecked< IUsdStageEditorModule >( "USDStageEditor" );
+	StageEditorModule.FileNew();
+}
+
+void UUsdStageEditorBlueprintLibrary::FileOpen( const FString& FilePath )
+{
+	IUsdStageEditorModule& StageEditorModule = FModuleManager::GetModuleChecked< IUsdStageEditorModule >( "USDStageEditor" );
+	StageEditorModule.FileOpen( FilePath );
+}
+
+void UUsdStageEditorBlueprintLibrary::FileSave( const FString& OutputFilePathIfUnsaved )
+{
+	IUsdStageEditorModule& StageEditorModule = FModuleManager::GetModuleChecked< IUsdStageEditorModule >( "USDStageEditor" );
+	StageEditorModule.FileSave( OutputFilePathIfUnsaved );
+}
+
+void UUsdStageEditorBlueprintLibrary::FileExportAllLayers( const FString& OutputDirectory )
+{
+	IUsdStageEditorModule& StageEditorModule = FModuleManager::GetModuleChecked< IUsdStageEditorModule >( "USDStageEditor" );
+	StageEditorModule.FileExportAllLayers( OutputDirectory );
+}
+
+void UUsdStageEditorBlueprintLibrary::FileExportFlattenedStage( const FString& OutputLayer )
+{
+	IUsdStageEditorModule& StageEditorModule = FModuleManager::GetModuleChecked< IUsdStageEditorModule >( "USDStageEditor" );
+	StageEditorModule.FileExportFlattenedStage( OutputLayer );
+}
+
+void UUsdStageEditorBlueprintLibrary::FileReload()
+{
+	IUsdStageEditorModule& StageEditorModule = FModuleManager::GetModuleChecked< IUsdStageEditorModule >( "USDStageEditor" );
+	StageEditorModule.FileReload();
+}
+
+void UUsdStageEditorBlueprintLibrary::FileReset()
+{
+	IUsdStageEditorModule& StageEditorModule = FModuleManager::GetModuleChecked< IUsdStageEditorModule >( "USDStageEditor" );
+	StageEditorModule.FileReset();
+}
+
+void UUsdStageEditorBlueprintLibrary::FileClose()
+{
+	IUsdStageEditorModule& StageEditorModule = FModuleManager::GetModuleChecked< IUsdStageEditorModule >( "USDStageEditor" );
+	StageEditorModule.FileClose();
+}
+
+void UUsdStageEditorBlueprintLibrary::ActionsImport( const FString& OutputContentFolder, UUsdStageImportOptions* Options )
+{
+	IUsdStageEditorModule& StageEditorModule = FModuleManager::GetModuleChecked< IUsdStageEditorModule >( "USDStageEditor" );
+	StageEditorModule.ActionsImport( OutputContentFolder, Options );
+}
+
+void UUsdStageEditorBlueprintLibrary::ExportSelectedLayers( const FString& OutputDirectory )
+{
+	IUsdStageEditorModule& StageEditorModule = FModuleManager::GetModuleChecked< IUsdStageEditorModule >( "USDStageEditor" );
+	StageEditorModule.ExportSelectedLayers( OutputDirectory );
+}
+
+
diff --git a/Engine/Plugins/Importers/USDImporter/Source/USDStageEditor/Private/USDStageEditorModule.cpp b/Engine/Plugins/Importers/USDImporter/Source/USDStageEditor/Private/USDStageEditorModule.cpp
--- a/Engine/Plugins/Importers/USDImporter/Source/USDStageEditor/Private/USDStageEditorModule.cpp
+++ b/Engine/Plugins/Importers/USDImporter/Source/USDStageEditor/Private/USDStageEditorModule.cpp
@@ -10,6 +10,9 @@
 #include "USDMemory.h"
 #include "USDProjectSettings.h"
 #include "USDStageActor.h"
+#include "UsdWrappers/SdfLayer.h"
+#include "UsdWrappers/UsdAttribute.h"
+#include "UsdWrappers/UsdPrim.h"
 
 #include "Editor/TransBuffer.h"
 #include "Editor/UnrealEdEngine.h"
@@ -30,6 +33,8 @@
 
 namespace UE::UsdStageEditorModule::Private
 {
+	const static FTabId UsdStageEditorTabID{ TEXT( "USDStage" ) };
+
 	void SaveStageActorLayersForWorld( UWorld* World, bool bForClosing, AUsdStageActor* TargetStageActor = nullptr )
 	{
 #if USE_USD_SDK
@@ -220,6 +225,29 @@ namespace UE::UsdStageEditorModule::Private
 		}
 #endif // USE_USD_SDK
 	}
+
+	TSharedPtr<SUsdStage> GetUsdStageEditor( bool bOpenIfNeeded = false )
+	{
+		FLevelEditorModule& LevelEditorModule = FModuleManager::GetModuleChecked< FLevelEditorModule >( "LevelEditor" );
+		TSharedPtr< FTabManager > LevelEditorTabManager = LevelEditorModule.GetLevelEditorTabManager();
+
+		TSharedPtr< SDockTab > Tab = bOpenIfNeeded
+			? LevelEditorTabManager->TryInvokeTab( UsdStageEditorTabID )
+			: LevelEditorTabManager->FindExistingLiveTab( UsdStageEditorTabID );
+
+		if ( Tab )
+		{
+			if ( TSharedPtr<SBorder> ContentBorder = StaticCastSharedRef<SBorder>( Tab->GetContent() ) )
+			{
+				if ( TSharedPtr<SUsdStage> UsdStageEditor = StaticCastSharedRef<SUsdStage>( ContentBorder->GetContent() ) )
+				{
+					return UsdStageEditor;
+				}
+			}
+		}
+
+		return nullptr;
+	}
 }
 
 class FUsdStageEditorModule : public IUsdStageEditorModule
@@ -258,11 +286,14 @@ public:
 
 				const FSlateIcon LayersIcon(FAppStyle::GetAppStyleSetName(), "LevelEditor.Tabs.USDStage");
 
-				LevelEditorTabManager->RegisterTabSpawner( TEXT("USDStage"), FOnSpawnTab::CreateStatic( &FUsdStageEditorModule::SpawnUsdStageTab ) )
-					.SetDisplayName( LOCTEXT( "USDStage", "USD Stage" ) )
-					.SetTooltipText( LOCTEXT( "USDStageTab", "Open USD Stage tab" ) )
-					.SetGroup(WorkspaceMenu::GetMenuStructure().GetLevelEditorVirtualProductionCategory())
-					.SetIcon( LayersIcon );
+				LevelEditorTabManager->RegisterTabSpawner(
+					UE::UsdStageEditorModule::Private::UsdStageEditorTabID.TabType,
+					FOnSpawnTab::CreateStatic( &FUsdStageEditorModule::SpawnUsdStageTab )
+				)
+				.SetDisplayName( LOCTEXT( "USDStage", "USD Stage" ) )
+				.SetTooltipText( LOCTEXT( "USDStageTab", "Open USD Stage tab" ) )
+				.SetGroup(WorkspaceMenu::GetMenuStructure().GetLevelEditorVirtualProductionCategory())
+				.SetIcon( LayersIcon );
 			});
 
 		// Prompt to save modified USD layers when closing the editor
@@ -390,6 +421,186 @@ private:
 #endif // #if USE_USD_SDK
 };
 
+bool IUsdStageEditorModule::OpenStageEditor() const
+{
+	const bool bOpenIfNeeded = true;
+	return UE::UsdStageEditorModule::Private::GetUsdStageEditor( bOpenIfNeeded ).IsValid();
+}
+
+bool IUsdStageEditorModule::CloseStageEditor() const
+{
+	FLevelEditorModule& LevelEditorModule = FModuleManager::GetModuleChecked< FLevelEditorModule >( "LevelEditor" );
+	TSharedPtr< FTabManager > LevelEditorTabManager = LevelEditorModule.GetLevelEditorTabManager();
+	if ( TSharedPtr< SDockTab > Tab = LevelEditorTabManager->FindExistingLiveTab( UE::UsdStageEditorModule::Private::UsdStageEditorTabID ) )
+	{
+		return Tab->RequestCloseTab();
+	}
+
+	return false;
+}
+
+AUsdStageActor* IUsdStageEditorModule::GetAttachedStageActor() const
+{
+	if ( TSharedPtr<SUsdStage> UsdStageEditor = UE::UsdStageEditorModule::Private::GetUsdStageEditor() )
+	{
+		return UsdStageEditor->GetAttachedStageActor();
+	}
+
+	return nullptr;
+}
+
+bool IUsdStageEditorModule::SetAttachedStageActor( AUsdStageActor* NewActor ) const
+{
+	if ( TSharedPtr<SUsdStage> UsdStageEditor = UE::UsdStageEditorModule::Private::GetUsdStageEditor() )
+	{
+		UsdStageEditor->AttachToStageActor( NewActor );
+		return true;
+	}
+
+	return false;
+}
+
+TArray<UE::FSdfLayer> IUsdStageEditorModule::GetSelectedLayers() const
+{
+	if ( TSharedPtr<SUsdStage> UsdStageEditor = UE::UsdStageEditorModule::Private::GetUsdStageEditor() )
+	{
+		return UsdStageEditor->GetSelectedLayers();
+	}
+
+	return {};
+}
+
+void IUsdStageEditorModule::SetSelectedLayers( const TArray<UE::FSdfLayer>& NewSelection ) const
+{
+	if ( TSharedPtr<SUsdStage> UsdStageEditor = UE::UsdStageEditorModule::Private::GetUsdStageEditor() )
+	{
+		return UsdStageEditor->SetSelectedLayers( NewSelection );
+	}
+}
+
+TArray<UE::FUsdPrim> IUsdStageEditorModule::GetSelectedPrims() const
+{
+	if ( TSharedPtr<SUsdStage> UsdStageEditor = UE::UsdStageEditorModule::Private::GetUsdStageEditor() )
+	{
+		return UsdStageEditor->GetSelectedPrims();
+	}
+
+	return {};
+}
+
+void IUsdStageEditorModule::SetSelectedPrims( const TArray<UE::FUsdPrim>& NewSelection ) const
+{
+	if ( TSharedPtr<SUsdStage> UsdStageEditor = UE::UsdStageEditorModule::Private::GetUsdStageEditor() )
+	{
+		UsdStageEditor->SetSelectedPrims( NewSelection );
+	}
+}
+
+TArray<FString> IUsdStageEditorModule::GetSelectedPropertyNames() const
+{
+	if ( TSharedPtr<SUsdStage> UsdStageEditor = UE::UsdStageEditorModule::Private::GetUsdStageEditor() )
+	{
+		return UsdStageEditor->GetSelectedPropertyNames();
+	}
+
+	return {};
+}
+
+void IUsdStageEditorModule::SetSelectedPropertyNames( const TArray<FString>& NewSelection ) const
+{
+	if ( TSharedPtr<SUsdStage> UsdStageEditor = UE::UsdStageEditorModule::Private::GetUsdStageEditor() )
+	{
+		UsdStageEditor->SetSelectedPropertyNames( NewSelection );
+	}
+}
+
+void IUsdStageEditorModule::FileNew() const
+{
+	if ( TSharedPtr<SUsdStage> UsdStageEditor = UE::UsdStageEditorModule::Private::GetUsdStageEditor() )
+	{
+		UsdStageEditor->FileNew();
+	}
+}
+
+void IUsdStageEditorModule::FileOpen( const FString& FilePath ) const
+{
+	if ( TSharedPtr<SUsdStage> UsdStageEditor = UE::UsdStageEditorModule::Private::GetUsdStageEditor() )
+	{
+		UsdStageEditor->FileOpen( FilePath );
+	}
+}
+
+void IUsdStageEditorModule::FileSave( const FString& OutputFilePathIfUnsaved ) const
+{
+	if ( TSharedPtr<SUsdStage> UsdStageEditor = UE::UsdStageEditorModule::Private::GetUsdStageEditor() )
+	{
+		UsdStageEditor->FileSave( OutputFilePathIfUnsaved );
+	}
+}
+
+void IUsdStageEditorModule::FileExportAllLayers( const FString& OutputDirectory ) const
+{
+	if ( TSharedPtr<SUsdStage> UsdStageEditor = UE::UsdStageEditorModule::Private::GetUsdStageEditor() )
+	{
+		UsdStageEditor->FileExportAllLayers( OutputDirectory  );
+	}
+}
+
+void IUsdStageEditorModule::FileExportFlattenedStage( const FString& OutputLayer ) const
+{
+	if ( TSharedPtr<SUsdStage> UsdStageEditor = UE::UsdStageEditorModule::Private::GetUsdStageEditor() )
+	{
+		UsdStageEditor->FileExportFlattenedStage( OutputLayer  );
+	}
+}
+
+void IUsdStageEditorModule::FileReload() const
+{
+	if ( TSharedPtr<SUsdStage> UsdStageEditor = UE::UsdStageEditorModule::Private::GetUsdStageEditor() )
+	{
+		UsdStageEditor->FileReload();
+	}
+}
+
+void IUsdStageEditorModule::FileReset() const
+{
+	if ( TSharedPtr<SUsdStage> UsdStageEditor = UE::UsdStageEditorModule::Private::GetUsdStageEditor() )
+	{
+		UsdStageEditor->FileReset();
+	}
+}
+
+void IUsdStageEditorModule::FileClose() const
+{
+	if ( TSharedPtr<SUsdStage> UsdStageEditor = UE::UsdStageEditorModule::Private::GetUsdStageEditor() )
+	{
+		UsdStageEditor->FileClose();
+	}
+}
+
+void IUsdStageEditorModule::ActionsImport( const FString& OutputContentFolder, UUsdStageImportOptions* Options ) const
+{
+	if ( TSharedPtr<SUsdStage> UsdStageEditor = UE::UsdStageEditorModule::Private::GetUsdStageEditor() )
+	{
+		if ( OutputContentFolder.IsEmpty() )
+		{
+			UsdStageEditor->ActionsImportWithDialog();
+		}
+		else
+		{
+			UsdStageEditor->ActionsImport( OutputContentFolder, Options );
+		}
+	}
+}
+
+void IUsdStageEditorModule::ExportSelectedLayers( const FString& OutputLayerOrDirectory ) const
+{
+	if ( TSharedPtr<SUsdStage> UsdStageEditor = UE::UsdStageEditorModule::Private::GetUsdStageEditor() )
+	{
+		UsdStageEditor->ExportSelectedLayers( OutputLayerOrDirectory );
+	}
+}
+
 IMPLEMENT_MODULE_USD( FUsdStageEditorModule, USDStageEditor );
 
 #undef LOCTEXT_NAMESPACE
diff --git a/Engine/Plugins/Importers/USDImporter/Source/USDStageEditor/Private/Widgets/SUSDPrimPropertiesList.cpp b/Engine/Plugins/Importers/USDImporter/Source/USDStageEditor/Private/Widgets/SUSDPrimPropertiesList.cpp
--- a/Engine/Plugins/Importers/USDImporter/Source/USDStageEditor/Private/Widgets/SUSDPrimPropertiesList.cpp
+++ b/Engine/Plugins/Importers/USDImporter/Source/USDStageEditor/Private/Widgets/SUSDPrimPropertiesList.cpp
@@ -692,6 +692,45 @@ void SUsdPrimPropertiesList::SetPrimPath( const UE::FUsdStageWeak& UsdStage, con
 	RequestListRefresh();
 }
 
+TArray<FString> SUsdPrimPropertiesList::GetSelectedPropertyNames() const
+{
+	TArray<FString> SelectedProperties;
+
+	TArray< TSharedPtr< FUsdPrimAttributeViewModel > > SelectedViewModels = GetSelectedItems();
+	SelectedProperties.Reserve( SelectedViewModels.Num() );
+
+	for ( const TSharedPtr< FUsdPrimAttributeViewModel >& SelectedItem : SelectedViewModels )
+	{
+		if ( SelectedItem )
+		{
+			SelectedProperties.Add( SelectedItem->Label );
+		}
+	}
+
+	return SelectedProperties;
+}
+
+void SUsdPrimPropertiesList::SetSelectedPropertyNames( const TArray<FString>& NewSelection )
+{
+	TSet<FString> NewSelectionSet{ NewSelection };
+
+	Private_ClearSelection();
+
+	TArray< TSharedPtr< FUsdPrimAttributeViewModel > > NewItemSelection;
+	NewItemSelection.Reserve( NewSelection.Num() );
+
+	for ( const TSharedPtr< FUsdPrimAttributeViewModel >& Item : ViewModel.PrimAttributes )
+	{
+		if ( Item && NewSelectionSet.Contains( Item->Label ) )
+		{
+			NewItemSelection.Add( Item );
+		}
+	}
+
+	const bool bSelected = true;
+	SetItemSelection( NewItemSelection, bSelected );
+}
+
 #undef LOCTEXT_NAMESPACE
 
 #endif // #if USE_USD_SDK
diff --git a/Engine/Plugins/Importers/USDImporter/Source/USDStageEditor/Private/Widgets/SUSDPrimPropertiesList.h b/Engine/Plugins/Importers/USDImporter/Source/USDStageEditor/Private/Widgets/SUSDPrimPropertiesList.h
--- a/Engine/Plugins/Importers/USDImporter/Source/USDStageEditor/Private/Widgets/SUSDPrimPropertiesList.h
+++ b/Engine/Plugins/Importers/USDImporter/Source/USDStageEditor/Private/Widgets/SUSDPrimPropertiesList.h
@@ -16,6 +16,9 @@ public:
 	void Construct( const FArguments& InArgs );
 	void SetPrimPath( const UE::FUsdStageWeak& UsdStage, const TCHAR* InPrimPath );
 
+	TArray<FString> GetSelectedPropertyNames() const;
+	void SetSelectedPropertyNames( const TArray<FString>& NewSelection );
+
 protected:
 	TSharedRef< ITableRow > OnGenerateRow( TSharedPtr< FUsdPrimAttributeViewModel > InDisplayNode, const TSharedRef< STableViewBase >& OwnerTable );
 	void GeneratePropertiesList( const UE::FUsdStageWeak& UsdStage, const TCHAR* InPrimPath );
diff --git a/Engine/Plugins/Importers/USDImporter/Source/USDStageEditor/Public/USDStageEditorBlueprintLibrary.h b/Engine/Plugins/Importers/USDImporter/Source/USDStageEditor/Public/USDStageEditorBlueprintLibrary.h
new file mode 100644
--- /dev/null
+++ b/Engine/Plugins/Importers/USDImporter/Source/USDStageEditor/Public/USDStageEditorBlueprintLibrary.h
@@ -0,0 +1,184 @@
+// Copyright Epic Games, Inc. All Rights Reserved.
+
+#pragma once
+
+#include "Engine/EngineTypes.h"
+#include "Kismet/BlueprintFunctionLibrary.h"
+
+#include "USDStageEditorBlueprintLibrary.generated.h"
+
+class AUsdStageActor;
+class UUsdStageImportOptions;
+
+/** Library of functions that can be used from scripting to interact with the USD Stage Editor UI */
+UCLASS(meta=(ScriptName="UsdStageEditorLibrary"))
+class USDSTAGEEDITOR_API UUsdStageEditorBlueprintLibrary : public UBlueprintFunctionLibrary
+{
+	GENERATED_BODY()
+
+public:
+	/**
+	 * Opens the the USD Stage Editor window, or focus it in case it is already open.
+	 * @return True if a stage is now opened and available.
+	 */
+	UFUNCTION( BlueprintCallable, Category = "USD|Stage Editor" )
+	static bool OpenStageEditor();
+
+	/**
+	 * Closes the USD Stage Editor window if it is opened. Does nothing in case it is already closed.
+	 * @return True if a stage was closed by this action.
+	 */
+	UFUNCTION( BlueprintCallable, Category = "USD|Stage Editor" )
+	static bool CloseStageEditor();
+
+	/**
+	 * Gets which actor is currently attached to the USD Stage Editor, if any.
+	 * May return nullptr in case no actor is attached.
+	 * @return The attached stage actor.
+	 */
+	UFUNCTION( BlueprintCallable, Category = "USD|Stage Actor" )
+	static AUsdStageActor* GetAttachedStageActor();
+
+	/**
+	 * Sets which actor is currently attached to the USD Stage Editor.
+	 * Provide None/nullptr to clear the attached stage actor.
+	 * @param NewActor - The actor to select
+	 * @return True in case the new actor was attached to the USD Stage Editor.
+	 */
+	UFUNCTION( BlueprintCallable, Category = "USD|Stage Actor" )
+	static bool SetAttachedStageActor( AUsdStageActor* NewActor );
+
+	/**
+	 * Returns the full identifiers of all layers that are currently selected on the USD Stage Editor.
+	 * @return The list of identifiers (e.g. ["c:/MyFolder/root.usda", "c:/MyFolder/sublayer.usda"])
+	 */
+	UFUNCTION( BlueprintCallable, Category = "USD|Selection" )
+	static TArray<FString> GetSelectedLayerIdentifiers();
+
+	/**
+	 * Sets the USD Stage Editor layer selection to all occurences of the layers with identifiers
+	 * in NewSelection. Provide an empty array to clear the selection.
+	 * @param NewSelection - The list of identifiers to select
+	 *                       (e.g. ["c:/MyFolder/root.usda", "c:/MyFolder/sublayer.usda"])
+	 */
+	UFUNCTION( BlueprintCallable, Category = "USD|Selection" )
+	static void SetSelectedLayerIdentifiers( const TArray<FString>& NewSelection );
+
+	/**
+	 * Returns the full paths to all prims currently selected on the USD Stage Editor.
+	 * @return The paths of selected prims (e.g. ["/ParentPrim/Mesh", "/Root"])
+	 */
+	UFUNCTION( BlueprintCallable, Category = "USD|Selection" )
+	static TArray<FString> GetSelectedPrimPaths();
+
+	/**
+	 * Sets the USD Stage Editor prim selection to the prims with paths contained in NewSelection.
+	 * Provide an empty array to clear the selection.
+	 * @param NewSelection - The list of prim paths to select (e.g. ["/ParentPrim/Mesh", "/Root"])
+	 */
+	UFUNCTION( BlueprintCallable, Category = "USD|Selection" )
+	static void SetSelectedPrimPaths( const TArray<FString>& NewSelection );
+
+	/**
+	 * Returns the names of the currently selected properties on the right panel of the USD Stage Editor.
+	 * @return The names of selected properties (e.g. ["points", "displayColor"])
+	 */
+	UFUNCTION( BlueprintCallable, Category = "USD|Selection" )
+	static TArray<FString> GetSelectedPropertyNames();
+
+	/**
+	 * Sets the USD Stage Editor property selection to the properties with names contained in NewSelection.
+	 * Provide an empty array to clear the selection.
+	 * @param NewSelection - The list of property names to select (e.g. ["points", "displayColor"])
+	 */
+	UFUNCTION( BlueprintCallable, Category = "USD|Selection" )
+	static void SetSelectedPropertyNames( const TArray<FString>& NewSelection );
+
+    /**
+	 * Creates a new memory-only layer and opens an USD Stage with that layer as its root.
+	 * Corresponds to the "File -> New" action on the USD Stage Editor menu bar.
+	 */
+	UFUNCTION( BlueprintCallable, Category = "USD|Menu Actions" )
+	static void FileNew();
+
+    /**
+     * Opens an USD Stage from a file on disk.
+	 * Corresponds to the "File -> Open" action on the USD Stage Editor menu bar.
+	 * @param FilePath - File path to the USD layer to open (e.g. "C:/Folder/MyFile.usda").
+	 *                   If this path is the empty string a dialog will be shown to let the user pick the file.
+     */
+	UFUNCTION( BlueprintCallable, Category = "USD|Menu Actions" )
+	static void FileOpen( const FString& FilePath );
+
+    /**
+     * Saves the currently opened USD Stage back to disk, or to a new file in case it hasn't been saved yet.
+	 * Corresponds to the "File -> Save" action on the USD Stage Editor menu bar.
+	 * @param OutputFilePathIfUnsaved - File path (e.g. "C:/Folder/MyFile.usda") to use when the currently opened
+	 *									stage hasn't been saved yet.
+	 *                                  If this path is the empty string a dialog will be shown to let the user pick
+	 *                                  the file.
+     */
+	UFUNCTION( BlueprintCallable, Category = "USD|Menu Actions" )
+	static void FileSave( const FString& OutputFilePathIfUnsaved );
+
+    /**
+     * Exports all layers of the currently opened USD Stage to brand new files in a new location.
+	 * Corresponds to the "File -> Export -> All Layers" action on the USD Stage Editor menu bar.
+     * @param OutputDirectory - Directory path (e.g. "C:/ExportFolder/") to receive the exported files.
+	 *                          If this path is the empty string a dialog will be shown to let the user pick the
+	 *                          directory.
+     */
+	UFUNCTION( BlueprintCallable, Category = "USD|Menu Actions" )
+	static void FileExportAllLayers( const FString& OutputDirectory );
+
+	/**
+     * Exports the currently opened USD Stage to a single "flattened" USD layer.
+	 * Corresponds to the "File -> Export -> All Layers" action on the USD Stage Editor menu bar.
+	 * @param OutputLayer - File path (e.g. "C:/ExportFolder/out.usda") to export the flattened layer to.
+	 *                      If this path is the empty string a dialog will be shown to let the user pick the file.
+	 */
+	UFUNCTION( BlueprintCallable, Category = "USD|Menu Actions" )
+	static void FileExportFlattenedStage( const FString& OutputLayer );
+
+    /**
+	 * Reloads all layers of the current USD Stage.
+	 * Corresponds to the "File -> Reload" action on the USD Stage Editor menu bar.
+	 */
+	UFUNCTION( BlueprintCallable, Category = "USD|Menu Actions" )
+	static void FileReload();
+
+    /**
+	 * Resets the state of the current USD Stage (which layers are muted, the current edit target, etc.).
+	 * Corresponds to the "File -> Reset state" action on the USD Stage Editor menu bar.
+	 */
+	UFUNCTION( BlueprintCallable, Category = "USD|Menu Actions" )
+	static void FileReset();
+
+    /**
+	 * Closes the currently opened USD Stage (by clearing the attached Stage Actor's RootLayer property).
+	 * Corresponds to the "File -> Close" action on the USD Stage Editor menu bar.
+	 */
+	UFUNCTION( BlueprintCallable, Category = "USD|Menu Actions" )
+	static void FileClose();
+
+	/**
+	 * Imports the currently opened USD Stage into persistent UE assets, actors and components on the level.
+	 * Corresponds to the "Actions -> Import" action on the USD Stage Editor menu bar.
+     * @param OutputContentFolder - Content path (e.g. "/Game/Imports") to receive the imported assets.
+	 *								If this path is the empty string a dialog will be shown to let the user pick
+	 *								the folder.
+     * @papam Options - Options object to use for the import.
+	 */
+	UFUNCTION( BlueprintCallable, Category = "USD|Menu Actions" )
+	static void ActionsImport( const FString& OutputContentFolder, UUsdStageImportOptions* Options );
+
+	/**
+	 * Exports the currently selected layers on the USD Stage Editor to brand new files in a new location.
+	 * Corresponds to right-clicking selected layers on the USD Stage Editor and picking "Export".
+	 * @param OutputDirectory - Directory path (e.g. "C:/ExportFolder/") to receive the exported files.
+	 *                          If this path is the empty string a dialog will be shown to let the user pick the
+	 *                          directory.
+	 */
+	UFUNCTION( BlueprintCallable, Category = "USD|Menu Actions" )
+	static void ExportSelectedLayers( const FString& OutputDirectory );
+};
diff --git a/Engine/Plugins/Importers/USDImporter/Source/USDStageEditor/Public/USDStageEditorModule.h b/Engine/Plugins/Importers/USDImporter/Source/USDStageEditor/Public/USDStageEditorModule.h
--- a/Engine/Plugins/Importers/USDImporter/Source/USDStageEditor/Public/USDStageEditorModule.h
+++ b/Engine/Plugins/Importers/USDImporter/Source/USDStageEditor/Public/USDStageEditorModule.h
@@ -2,9 +2,51 @@
 
 #pragma once
 
+#include "UsdWrappers/ForwardDeclarations.h"
+
 #include "Modules/ModuleInterface.h"
 
+class AUsdStageActor;
+class UUsdStageImportOptions;
+namespace UE
+{
+	class FUsdPrim;
+	class FUsdAttribute;
+}
+
 class IUsdStageEditorModule : public IModuleInterface
 {
+public:
+
+	/**
+	 * Check out USDStageEditorBlueprintLibrary.h for documentation regarding these functions.
+	 */
+
+	bool OpenStageEditor() const;
+	bool CloseStageEditor() const;
+
+	AUsdStageActor* GetAttachedStageActor() const;
+	bool SetAttachedStageActor( AUsdStageActor* NewActor ) const;
+
+	TArray<UE::FSdfLayer> GetSelectedLayers() const;
+	void SetSelectedLayers( const TArray<UE::FSdfLayer>& NewSelection ) const;
+
+	TArray<UE::FUsdPrim> GetSelectedPrims() const;
+	void SetSelectedPrims( const TArray<UE::FUsdPrim>& NewSelection ) const;
+
+	TArray<FString> GetSelectedPropertyNames() const;
+	void SetSelectedPropertyNames( const TArray<FString>& NewSelection ) const;
 
+	// For all of these, providing an empty path will cause us to pop open a dialog to let the user pick the path
+	// instead.
+	void FileNew() const;
+	void FileOpen( const FString& FilePath ) const;
+	void FileSave( const FString& OutputFilePathIfUnsaved ) const;
+	void FileExportAllLayers( const FString& OutputDirectory ) const;
+	void FileExportFlattenedStage( const FString& OutputLayer ) const;
+	void FileReload() const;
+	void FileReset() const;
+	void FileClose() const;
+	void ActionsImport( const FString& OutputContentFolder, UUsdStageImportOptions* Options ) const;
+	void ExportSelectedLayers( const FString& OutputLayerOrDirectory ) const;
 };
diff --git a/Engine/Plugins/Importers/USDImporter/Source/USDStageEditorViewModels/Private/USDStageViewModel.cpp b/Engine/Plugins/Importers/USDImporter/Source/USDStageEditorViewModels/Private/USDStageViewModel.cpp
--- a/Engine/Plugins/Importers/USDImporter/Source/USDStageEditorViewModels/Private/USDStageViewModel.cpp
+++ b/Engine/Plugins/Importers/USDImporter/Source/USDStageEditorViewModels/Private/USDStageViewModel.cpp
@@ -264,7 +264,7 @@ void FUsdStageViewModel::SaveStageAs( const TCHAR* FilePath )
 #endif // #if USE_USD_SDK
 }
 
-void FUsdStageViewModel::ImportStage()
+void FUsdStageViewModel::ImportStage( const TCHAR* TargetContentFolder, UUsdStageImportOptions* Options )
 {
 #if USE_USD_SDK
 	AUsdStageActor* StageActor = UsdStageActor.Get();
@@ -304,11 +304,18 @@ void FUsdStageViewModel::ImportStage()
 		// Pass the stage directly too in case we're importing a transient stage with no filepath
 		ImportContext.Stage = UsdStage;
 
-		const bool bIsAutomated = false;
+		const bool bIsAutomated = TargetContentFolder && Options;
 		if ( ImportContext.Init( StageName, RootPath, TEXT("/Game/"), RF_Public | RF_Transactional, bIsAutomated ) )
 		{
 			FScopedTransaction Transaction( FText::Format(LOCTEXT("ImportTransaction", "Import USD stage '{0}'"), FText::FromString(StageName)));
 
+			if ( bIsAutomated )
+			{
+				// Apply same conversion that FUsdStageImportContext::Init does on our received path
+				ImportContext.PackagePath = FString::Printf( TEXT( "%s/%s/" ), TargetContentFolder, *StageName );
+				ImportContext.ImportOptions = Options;
+			}
+
 			// Let the importer reuse our assets, but force it to spawn new actors and components always
 			// This allows a different setting for asset/component collapsing, and doesn't require modifying the PrimTwins
 			ImportContext.AssetCache = StageActor->GetAssetCache();
diff --git a/Engine/Plugins/Importers/USDImporter/Source/USDStageEditorViewModels/Public/USDStageViewModel.h b/Engine/Plugins/Importers/USDImporter/Source/USDStageEditorViewModels/Public/USDStageViewModel.h
--- a/Engine/Plugins/Importers/USDImporter/Source/USDStageEditorViewModels/Public/USDStageViewModel.h
+++ b/Engine/Plugins/Importers/USDImporter/Source/USDStageEditorViewModels/Public/USDStageViewModel.h
@@ -8,6 +8,7 @@
 
 class AUsdStageActor;
 class UPackage;
+class UUsdStageImportOptions;
 
 class USDSTAGEEDITORVIEWMODELS_API FUsdStageViewModel
 {
@@ -20,7 +21,7 @@ public:
 	void SaveStage();
 	/** Temporary until SaveAs feature is properly implemented, may be removed in a future release */
 	void SaveStageAs( const TCHAR* FilePath );
-	void ImportStage();
+	void ImportStage( const TCHAR* TargetContentFolder = nullptr, UUsdStageImportOptions* Options = nullptr );
 
 public:
 	TWeakObjectPtr< AUsdStageActor > UsdStageActor;
