USD: Fix build when USD SDK is unavailable.

diff --git a/Engine/Plugins/Importers/USDImporter/Source/USDStageEditor/Private/USDStageEditorModule.cpp b/Engine/Plugins/Importers/USDImporter/Source/USDStageEditor/Private/USDStageEditorModule.cpp
--- a/Engine/Plugins/Importers/USDImporter/Source/USDStageEditor/Private/USDStageEditorModule.cpp
+++ b/Engine/Plugins/Importers/USDImporter/Source/USDStageEditor/Private/USDStageEditorModule.cpp
@@ -226,6 +226,7 @@ namespace UE::UsdStageEditorModule::Private
 #endif // USE_USD_SDK
 	}
 
+#if USE_USD_SDK
 	TSharedPtr<SUsdStage> GetUsdStageEditor( bool bOpenIfNeeded = false )
 	{
 		FLevelEditorModule& LevelEditorModule = FModuleManager::GetModuleChecked< FLevelEditorModule >( "LevelEditor" );
@@ -248,6 +249,7 @@ namespace UE::UsdStageEditorModule::Private
 
 		return nullptr;
 	}
+#endif // USE_USD_SDK
 }
 
 class FUsdStageEditorModule : public IUsdStageEditorModule
@@ -418,168 +420,207 @@ private:
 	FDelegateHandle StageActorLoadedHandle;
 	FDelegateHandle OpenStageEditorClickedHandle;
 	FDelegateHandle OnTransactionStateChangedHandle;
-#endif // #if USE_USD_SDK
+#endif // USE_USD_SDK
 };
 
 bool IUsdStageEditorModule::OpenStageEditor() const
 {
+#if USE_USD_SDK
 	const bool bOpenIfNeeded = true;
 	return UE::UsdStageEditorModule::Private::GetUsdStageEditor( bOpenIfNeeded ).IsValid();
+#else
+	return false;
+#endif // USE_USD_SDK
 }
 
 bool IUsdStageEditorModule::CloseStageEditor() const
 {
+#if USE_USD_SDK
 	FLevelEditorModule& LevelEditorModule = FModuleManager::GetModuleChecked< FLevelEditorModule >( "LevelEditor" );
 	TSharedPtr< FTabManager > LevelEditorTabManager = LevelEditorModule.GetLevelEditorTabManager();
 	if ( TSharedPtr< SDockTab > Tab = LevelEditorTabManager->FindExistingLiveTab( UE::UsdStageEditorModule::Private::UsdStageEditorTabID ) )
 	{
 		return Tab->RequestCloseTab();
 	}
+#endif // USE_USD_SDK
 
 	return false;
 }
 
 AUsdStageActor* IUsdStageEditorModule::GetAttachedStageActor() const
 {
+#if USE_USD_SDK
 	if ( TSharedPtr<SUsdStage> UsdStageEditor = UE::UsdStageEditorModule::Private::GetUsdStageEditor() )
 	{
 		return UsdStageEditor->GetAttachedStageActor();
 	}
+#endif // USE_USD_SDK
 
 	return nullptr;
 }
 
 bool IUsdStageEditorModule::SetAttachedStageActor( AUsdStageActor* NewActor ) const
 {
+#if USE_USD_SDK
 	if ( TSharedPtr<SUsdStage> UsdStageEditor = UE::UsdStageEditorModule::Private::GetUsdStageEditor() )
 	{
 		UsdStageEditor->AttachToStageActor( NewActor );
 		return true;
 	}
+#endif // USE_USD_SDK
 
 	return false;
 }
 
 TArray<UE::FSdfLayer> IUsdStageEditorModule::GetSelectedLayers() const
 {
+#if USE_USD_SDK
 	if ( TSharedPtr<SUsdStage> UsdStageEditor = UE::UsdStageEditorModule::Private::GetUsdStageEditor() )
 	{
 		return UsdStageEditor->GetSelectedLayers();
 	}
+#endif // USE_USD_SDK
 
 	return {};
 }
 
 void IUsdStageEditorModule::SetSelectedLayers( const TArray<UE::FSdfLayer>& NewSelection ) const
 {
+#if USE_USD_SDK
 	if ( TSharedPtr<SUsdStage> UsdStageEditor = UE::UsdStageEditorModule::Private::GetUsdStageEditor() )
 	{
 		return UsdStageEditor->SetSelectedLayers( NewSelection );
 	}
+#endif // USE_USD_SDK
 }
 
 TArray<UE::FUsdPrim> IUsdStageEditorModule::GetSelectedPrims() const
 {
+#if USE_USD_SDK
 	if ( TSharedPtr<SUsdStage> UsdStageEditor = UE::UsdStageEditorModule::Private::GetUsdStageEditor() )
 	{
 		return UsdStageEditor->GetSelectedPrims();
 	}
+#endif // USE_USD_SDK
 
 	return {};
 }
 
 void IUsdStageEditorModule::SetSelectedPrims( const TArray<UE::FUsdPrim>& NewSelection ) const
 {
+#if USE_USD_SDK
 	if ( TSharedPtr<SUsdStage> UsdStageEditor = UE::UsdStageEditorModule::Private::GetUsdStageEditor() )
 	{
 		UsdStageEditor->SetSelectedPrims( NewSelection );
 	}
+#endif // USE_USD_SDK
 }
 
 TArray<FString> IUsdStageEditorModule::GetSelectedPropertyNames() const
 {
+#if USE_USD_SDK
 	if ( TSharedPtr<SUsdStage> UsdStageEditor = UE::UsdStageEditorModule::Private::GetUsdStageEditor() )
 	{
 		return UsdStageEditor->GetSelectedPropertyNames();
 	}
+#endif // USE_USD_SDK
 
 	return {};
 }
 
 void IUsdStageEditorModule::SetSelectedPropertyNames( const TArray<FString>& NewSelection ) const
 {
+#if USE_USD_SDK
 	if ( TSharedPtr<SUsdStage> UsdStageEditor = UE::UsdStageEditorModule::Private::GetUsdStageEditor() )
 	{
 		UsdStageEditor->SetSelectedPropertyNames( NewSelection );
 	}
+#endif // USE_USD_SDK
 }
 
 void IUsdStageEditorModule::FileNew() const
 {
+#if USE_USD_SDK
 	if ( TSharedPtr<SUsdStage> UsdStageEditor = UE::UsdStageEditorModule::Private::GetUsdStageEditor() )
 	{
 		UsdStageEditor->FileNew();
 	}
+#endif // USE_USD_SDK
 }
 
 void IUsdStageEditorModule::FileOpen( const FString& FilePath ) const
 {
+#if USE_USD_SDK
 	if ( TSharedPtr<SUsdStage> UsdStageEditor = UE::UsdStageEditorModule::Private::GetUsdStageEditor() )
 	{
 		UsdStageEditor->FileOpen( FilePath );
 	}
+#endif // USE_USD_SDK
 }
 
 void IUsdStageEditorModule::FileSave( const FString& OutputFilePathIfUnsaved ) const
 {
+#if USE_USD_SDK
 	if ( TSharedPtr<SUsdStage> UsdStageEditor = UE::UsdStageEditorModule::Private::GetUsdStageEditor() )
 	{
 		UsdStageEditor->FileSave( OutputFilePathIfUnsaved );
 	}
+#endif // USE_USD_SDK
 }
 
 void IUsdStageEditorModule::FileExportAllLayers( const FString& OutputDirectory ) const
 {
+#if USE_USD_SDK
 	if ( TSharedPtr<SUsdStage> UsdStageEditor = UE::UsdStageEditorModule::Private::GetUsdStageEditor() )
 	{
 		UsdStageEditor->FileExportAllLayers( OutputDirectory  );
 	}
+#endif // USE_USD_SDK
 }
 
 void IUsdStageEditorModule::FileExportFlattenedStage( const FString& OutputLayer ) const
 {
+#if USE_USD_SDK
 	if ( TSharedPtr<SUsdStage> UsdStageEditor = UE::UsdStageEditorModule::Private::GetUsdStageEditor() )
 	{
 		UsdStageEditor->FileExportFlattenedStage( OutputLayer  );
 	}
+#endif // USE_USD_SDK
 }
 
 void IUsdStageEditorModule::FileReload() const
 {
+#if USE_USD_SDK
 	if ( TSharedPtr<SUsdStage> UsdStageEditor = UE::UsdStageEditorModule::Private::GetUsdStageEditor() )
 	{
 		UsdStageEditor->FileReload();
 	}
+#endif // USE_USD_SDK
 }
 
 void IUsdStageEditorModule::FileReset() const
 {
+#if USE_USD_SDK
 	if ( TSharedPtr<SUsdStage> UsdStageEditor = UE::UsdStageEditorModule::Private::GetUsdStageEditor() )
 	{
 		UsdStageEditor->FileReset();
 	}
+#endif // USE_USD_SDK
 }
 
 void IUsdStageEditorModule::FileClose() const
 {
+#if USE_USD_SDK
 	if ( TSharedPtr<SUsdStage> UsdStageEditor = UE::UsdStageEditorModule::Private::GetUsdStageEditor() )
 	{
 		UsdStageEditor->FileClose();
 	}
+#endif // USE_USD_SDK
 }
 
 void IUsdStageEditorModule::ActionsImport( const FString& OutputContentFolder, UUsdStageImportOptions* Options ) const
 {
+#if USE_USD_SDK
 	if ( TSharedPtr<SUsdStage> UsdStageEditor = UE::UsdStageEditorModule::Private::GetUsdStageEditor() )
 	{
 		if ( OutputContentFolder.IsEmpty() )
@@ -591,14 +632,17 @@ void IUsdStageEditorModule::ActionsImport( const FString& OutputContentFolder, U
 			UsdStageEditor->ActionsImport( OutputContentFolder, Options );
 		}
 	}
+#endif // USE_USD_SDK
 }
 
 void IUsdStageEditorModule::ExportSelectedLayers( const FString& OutputLayerOrDirectory ) const
 {
+#if USE_USD_SDK
 	if ( TSharedPtr<SUsdStage> UsdStageEditor = UE::UsdStageEditorModule::Private::GetUsdStageEditor() )
 	{
 		UsdStageEditor->ExportSelectedLayers( OutputLayerOrDirectory );
 	}
+#endif // USE_USD_SDK
 }
 
 IMPLEMENT_MODULE_USD( FUsdStageEditorModule, USDStageEditor );
diff --git a/Engine/Plugins/Importers/USDImporter/Source/USDStageEditor/Public/USDStageEditorModule.h b/Engine/Plugins/Importers/USDImporter/Source/USDStageEditor/Public/USDStageEditorModule.h
--- a/Engine/Plugins/Importers/USDImporter/Source/USDStageEditor/Public/USDStageEditorModule.h
+++ b/Engine/Plugins/Importers/USDImporter/Source/USDStageEditor/Public/USDStageEditorModule.h
@@ -4,6 +4,7 @@
 
 #include "UsdWrappers/ForwardDeclarations.h"
 
+#include "Containers/Array.h"
 #include "Modules/ModuleInterface.h"
 
 class AUsdStageActor;
