USD: Fix failing geomcache import tests and re-enable them.

diff --git a/Engine/Plugins/Importers/USDImporter/Source/USDUtilities/Private/USDGeomMeshConversion.cpp b/Engine/Plugins/Importers/USDImporter/Source/USDUtilities/Private/USDGeomMeshConversion.cpp
--- a/Engine/Plugins/Importers/USDImporter/Source/USDUtilities/Private/USDGeomMeshConversion.cpp
+++ b/Engine/Plugins/Importers/USDImporter/Source/USDUtilities/Private/USDGeomMeshConversion.cpp
@@ -752,6 +752,10 @@ bool UsdToUnreal::ConvertGeomMesh(
 		{
 			pxr::VtArray< pxr::GfVec3f > PointsArray;
 			Points.Get( &PointsArray, TimeCodeValue );
+			if (PointsArray.size() < 3)
+			{
+				return false;
+			}
 
 			OutMeshDescription.ReserveNewVertices( PointsArray.size() );
 
@@ -821,6 +825,10 @@ bool UsdToUnreal::ConvertGeomMesh(
 		{
 			FaceCountsAttribute.Get( &FaceCounts, TimeCodeValue );
 			NumPolygons = FaceCounts.size();
+			if (NumPolygons < 1)
+			{
+				return false;
+			}
 		}
 
 		// Face indices
@@ -830,6 +838,10 @@ bool UsdToUnreal::ConvertGeomMesh(
 		if ( FaceIndicesAttribute )
 		{
 			FaceIndicesAttribute.Get( &FaceIndices, TimeCodeValue );
+			if (FaceIndices.size() < 1)
+			{
+				return false;
+			}
 		}
 
 		// Normals
@@ -958,7 +970,12 @@ bool UsdToUnreal::ConvertGeomMesh(
 				const FVertexInstanceID VertexInstanceID( VertexInstanceIndex );
 				const int32 ControlPointIndex = FaceIndices[ CurrentVertexInstanceIndex ];
 				const FVertexID VertexID( VertexOffset + ControlPointIndex );
-				const FVector VertexPosition = ( FVector ) MeshDescriptionVertexPositions[ VertexID ];
+
+				// This data is read straight from USD so there's nothing guaranteeing we have as many positions as we need
+				if (VertexID.GetValue() >= MeshDescriptionVertexPositions.GetNumElements() || VertexID.GetValue() < 0)
+				{
+					continue;
+				}
 
 				// Make sure a face doesn't use the same vertex twice as MeshDescription doesn't like that
 				if ( CornerVerticesIDs.Contains( VertexID ) )
@@ -966,6 +983,8 @@ bool UsdToUnreal::ConvertGeomMesh(
 					continue;
 				}
 
+				const FVector VertexPosition = (FVector)MeshDescriptionVertexPositions[VertexID];
+
 				CornerVerticesIDs.Add( VertexID );
 
 				FVertexInstanceID AddedVertexInstanceId = OutMeshDescription.CreateVertexInstance( VertexID );
