USD: Add cvar to allow discarding the undo buffer when opening or closing stages.

diff --git a/Engine/Plugins/Importers/USDImporter/Source/USDStage/Private/USDStageActor.cpp b/Engine/Plugins/Importers/USDImporter/Source/USDStage/Private/USDStageActor.cpp
--- a/Engine/Plugins/Importers/USDImporter/Source/USDStage/Private/USDStageActor.cpp
+++ b/Engine/Plugins/Importers/USDImporter/Source/USDStage/Private/USDStageActor.cpp
@@ -2927,7 +2927,23 @@ void AUsdStageActor::HandlePropertyChangedEvent(FPropertyChangedEvent& PropertyC
 
 	if (PropertyName == GET_MEMBER_NAME_CHECKED(AUsdStageActor, RootLayer))
 	{
+#if WITH_EDITOR
+		static IConsoleVariable* CVar = IConsoleManager::Get().FindConsoleVariable(TEXT("USD.DiscardUndoBufferOnStageOpenClose"));
+		const bool bDiscardUndo = CVar && CVar->GetBool();
+#endif // WITH_EDITOR
+
 		SetRootLayer(RootLayer.FilePath);
+
+#if WITH_EDITOR
+		if (bDiscardUndo && GEditor)
+		{
+			if (UTransactor* EditorTransactor = GEditor->Trans)
+			{
+				const FText Reason = LOCTEXT("DiscardTransactionReason", "Resetting because USD.DiscardUndoBufferOnStageOpenClose is enabled");
+				EditorTransactor->Reset(Reason);
+			}
+		}
+#endif // WITH_EDITOR
 	}
 	else if (PropertyName == GET_MEMBER_NAME_CHECKED(AUsdStageActor, Time))
 	{
diff --git a/Engine/Plugins/Importers/USDImporter/Source/USDStageEditor/Private/SUSDStage.cpp b/Engine/Plugins/Importers/USDImporter/Source/USDStageEditor/Private/SUSDStage.cpp
--- a/Engine/Plugins/Importers/USDImporter/Source/USDStageEditor/Private/SUSDStage.cpp
+++ b/Engine/Plugins/Importers/USDImporter/Source/USDStageEditor/Private/SUSDStage.cpp
@@ -32,6 +32,7 @@
 #include "DesktopPlatformModule.h"
 #include "Dialogs/DlgPickPath.h"
 #include "Editor.h"
+#include "Editor/Transactor.h"
 #include "Engine/Selection.h"
 #include "Engine/World.h"
 #include "EngineAnalytics.h"
@@ -54,6 +55,12 @@
 
 #define LOCTEXT_NAMESPACE "SUsdStage"
 
+static bool GDiscardUndoBufferOnStageOpenClose = false;
+static FAutoConsoleVariableRef CVarDiscardUndoBufferOnStageOpenClose(
+	TEXT("USD.DiscardUndoBufferOnStageOpenClose"),
+	GDiscardUndoBufferOnStageOpenClose,
+	TEXT("Enabling this will prevent the recording of open/close stage transactions, but also discard the undo buffer after they happen. Use this when memory-constrained, as sometimes recording all created assets and actors in the undo buffer can be expensive."));
+
 #if USE_USD_SDK
 
 namespace SUSDStageConstants
@@ -2148,10 +2155,23 @@ void SUsdStage::FileReset()
 
 void SUsdStage::FileClose()
 {
-	FScopedTransaction Transaction(LOCTEXT("CloseTransaction", "Close USD stage"));
+	TOptional<FScopedTransaction> Transaction;
+	if (!GDiscardUndoBufferOnStageOpenClose)
+	{
+		Transaction.Emplace(LOCTEXT("CloseTransaction", "Close USD stage"));
+	}
 
 	ViewModel.CloseStage();
 	Refresh();
+
+	if (GDiscardUndoBufferOnStageOpenClose && GEditor)
+	{
+		if (UTransactor* Transactor = GEditor->Trans)
+		{
+			const FText Reason = LOCTEXT("DiscardCloseTransactionReason", "Resetting because the cvar USD.DiscardUndoBufferOnStageOpenClose is true");
+			Transactor->Reset(Reason);
+		}
+	}
 }
 
 void SUsdStage::OnLayerIsolated( const UE::FSdfLayer& IsolatedLayer )
@@ -2225,12 +2245,16 @@ void SUsdStage::OnPrimSelectionChanged( const TArray<FString>& PrimPaths )
 
 void SUsdStage::OpenStage( const TCHAR* FilePath )
 {
-	// Create the transaction before calling UsdStageModule.GetUsdStageActor as that may create the actor, and we want
-	// the actor spawning to be part of the transaction
-	FScopedTransaction Transaction( FText::Format(
-		LOCTEXT( "OpenStageTransaction", "Open USD stage '{0}'" ),
-		FText::FromString( FilePath )
-	) );
+	TOptional<FScopedTransaction> Transaction;
+	if (!GDiscardUndoBufferOnStageOpenClose)
+	{
+		// Create the transaction before calling UsdStageModule.GetUsdStageActor as that may create the actor, and we want
+		// the actor spawning to be part of the transaction
+		Transaction.Emplace(FText::Format(
+			LOCTEXT("OpenStageTransaction", "Open USD stage '{0}'"),
+			FText::FromString(FilePath)
+		));
+	}
 
 	if ( !ViewModel.UsdStageActor.IsValid() )
 	{
@@ -2239,6 +2263,15 @@ void SUsdStage::OpenStage( const TCHAR* FilePath )
 	}
 
 	ViewModel.OpenStage( FilePath );
+
+	if (GDiscardUndoBufferOnStageOpenClose && GEditor)
+	{
+		if (UTransactor* Transactor = GEditor->Trans)
+		{
+			const FText Reason = LOCTEXT("DiscardOpenTransactionReason", "Resetting because the cvar USD.DiscardUndoBufferOnStageOpenClose is true");
+			Transactor->Reset(Reason);
+		}
+	}
 }
 
 void SUsdStage::Refresh()
