PR #9986: Add missing debug undef for boost in the USDImporter plugin (Contributed by [REDACTED_PERSON])

diff --git a/Engine/Plugins/Importers/USDImporter/Source/UnrealUSDWrapper/Public/USDIncludesEnd.h b/Engine/Plugins/Importers/USDImporter/Source/UnrealUSDWrapper/Public/USDIncludesEnd.h
--- a/Engine/Plugins/Importers/USDImporter/Source/UnrealUSDWrapper/Public/USDIncludesEnd.h
+++ b/Engine/Plugins/Importers/USDImporter/Source/UnrealUSDWrapper/Public/USDIncludesEnd.h
@@ -14,6 +14,13 @@
 	#pragma pop_macro("catch")
 #endif
 
+// Boost needed _DEBUG defined when /RTCs build flag is enabled (Run Time Checks)
+#if PLATFORM_WINDOWS && UE_BUILD_DEBUG
+	#ifdef _DEBUG
+		#undef _DEBUG
+	#endif
+#endif
+
 THIRD_PARTY_INCLUDES_END
 
 #endif // #if USE_USD_SDK
