USD: Global USD Asset Cache.

diff --git a/Engine/Plugins/Importers/USDImporter/Source/USDClasses/Private/USDAssetCache.cpp b/Engine/Plugins/Importers/USDImporter/Source/USDClasses/Private/USDAssetCache.cpp
--- a/Engine/Plugins/Importers/USDImporter/Source/USDClasses/Private/USDAssetCache.cpp
+++ b/Engine/Plugins/Importers/USDImporter/Source/USDClasses/Private/USDAssetCache.cpp
@@ -106,6 +106,7 @@ void UUsdAssetCache::DiscardAsset( const FString& Hash )
 	}
 }
 
+PRAGMA_DISABLE_DEPRECATION_WARNINGS
 UObject* UUsdAssetCache::GetCachedAsset( const FString& Hash ) const
 {
 	FScopeLock Lock( &CriticalSection );
@@ -166,6 +167,7 @@ UObject* UUsdAssetCache::GetAssetForPrim( const FString& PrimPath ) const
 
 	return nullptr;
 }
+PRAGMA_ENABLE_DEPRECATION_WARNINGS
 
 FString UUsdAssetCache::GetPrimForAsset( UObject* Asset ) const
 {
diff --git a/Engine/Plugins/Importers/USDImporter/Source/USDClasses/Private/USDAssetCache2.cpp b/Engine/Plugins/Importers/USDImporter/Source/USDClasses/Private/USDAssetCache2.cpp
new file mode 100644
--- /dev/null
+++ b/Engine/Plugins/Importers/USDImporter/Source/USDClasses/Private/USDAssetCache2.cpp
@@ -0,0 +1,1507 @@
+// Copyright Epic Games, Inc. All Rights Reserved.
+
+#include "USDAssetCache2.h"
+
+#include "USDClassesModule.h"
+#include "USDLog.h"
+
+#include "Engine/StaticMesh.h"
+#include "Engine/Texture2D.h"
+#include "Misc/ScopeRWLock.h"
+#include "Serialization/ObjectReader.h"
+#include "Serialization/ObjectWriter.h"
+#include "Serialization/ObjectAndNameAsStringProxyArchive.h"
+#include "UObject/ObjectRedirector.h"
+#include "UObject/Package.h"
+
+static bool GOnDemandCachedAssetLoading = true;
+static FAutoConsoleVariableRef COnDemandCachedAssetLoading(
+	TEXT("USD.OnDemandCachedAssetLoading"),
+	GOnDemandCachedAssetLoading,
+	TEXT("When true will cause the USD Asset Cache to only load cached UObjects from disk once they are needed, regardless of when the actual Asset Cache asset itself is loaded."));
+
+static int32 GCurrentCacheVersion = 1;
+static int32 GCurrentAssetInfoVersion = 1;
+
+namespace UE::AssetCache::Private
+{
+	UClass* FindClass(const TCHAR* ClassName)
+	{
+		check(ClassName);
+
+		if (UClass* Result = FindFirstObject<UClass>(ClassName, EFindFirstObjectOptions::None, ELogVerbosity::Warning, TEXT("USDAssetCache2")))
+		{
+			return Result;
+		}
+
+		if (UObjectRedirector* RenamedClassRedirector = FindFirstObject<UObjectRedirector>(ClassName, EFindFirstObjectOptions::None, ELogVerbosity::Warning, TEXT("USDAssetCache2")))
+		{
+			return CastChecked<UClass>(RenamedClassRedirector->DestinationObject);
+		}
+
+		return nullptr;
+	}
+
+	bool AssetCanBePersisted(const TCHAR* ClassName)
+	{
+		if (!ClassName)
+		{
+			return false;
+		}
+		UClass* AssetClass = FindClass(ClassName);
+
+		// Just textures for now
+		static TArray<const UClass*> AllowedPersistentClasses = { UTexture2D::StaticClass() };
+
+		for (const UClass* AllowedClass : AllowedPersistentClasses)
+		{
+			if (AssetClass->IsChildOf(AllowedClass))
+			{
+				return true;
+			}
+		}
+
+		return false;
+	}
+
+	uint64 GetAssetSizeBytes(UObject* Asset, bool bOnDisk)
+	{
+		if (!Asset)
+		{
+			return 0;
+		}
+
+		if (UTexture2D* Texture = Cast<UTexture2D>(Asset))
+		{
+#if WITH_EDITOR
+			if (bOnDisk)
+			{
+				return static_cast<uint64>(Texture->Source.GetSizeOnDisk());
+			}
+			else
+#endif // WITH_EDITOR
+			{
+				// TODO: Try doing this without calling GetPlatformData() if this is ever a speed bottleneck,
+				// as apparently that can stall.
+				// We can't use the texture Source here as it's not available at runtime.
+				if (FTexturePlatformData* PlatformData = Texture->GetPlatformData())
+				{
+					const uint32 MipCount = 1;
+					return GPixelFormats[PlatformData->PixelFormat].Get2DTextureSizeInBytes(
+						PlatformData->SizeX,
+						PlatformData->SizeY,
+						MipCount
+					);
+				}
+			}
+		}
+
+		return Asset->GetResourceSizeBytes(EResourceSizeMode::EstimatedTotal);
+	}
+
+	// Extends FObjectAndNameAsStringProxyArchive to support FLazyObjectPtr
+	// Copied from FSnapshotCustomArchive
+	struct FArchiveObjectPtrAsStringWrapper: public FObjectAndNameAsStringProxyArchive
+	{
+		FArchiveObjectPtrAsStringWrapper(FArchive& InInnerArchive)
+			: FObjectAndNameAsStringProxyArchive(InInnerArchive, false)
+		{
+		}
+
+		virtual FArchive& operator<<(FLazyObjectPtr& Obj) override
+		{
+			if (IsLoading())
+			{
+				// Reset before serializing to clear the internal weak pointer.
+				Obj.Reset();
+			}
+			InnerArchive << Obj.GetUniqueID();
+
+			return *this;
+		}
+	};
+
+	// The main archiver used to write UObjects to byte buffers
+	class FAssetCacheObjectWriter final : public FObjectWriter
+	{
+	public:
+		FAssetCacheObjectWriter(UObject& Object, TArray< uint8 >& Bytes)
+			: FObjectWriter(Bytes)
+		{
+			SetIsLoading(false);
+			SetIsSaving(true);
+			SetIsPersistent(true);
+		}
+	};
+
+	// The main archiver used to read UObjects from byte buffers
+	class FAssetCacheObjectReader final : public FObjectReader
+	{
+	public:
+		FAssetCacheObjectReader(UObject& Object, const TArray< uint8 >& Bytes)
+			: FObjectReader(Bytes)
+		{
+			SetIsLoading(true);
+			SetIsSaving(false);
+			SetIsPersistent(true);
+		}
+	};
+
+	// Copied from DataprepSnapshot.cpp, the "WriteSnapshotData" function
+	void SerializeObjectAndSubObjects(UObject* Object, TArray<uint8>& Buffer, FArchive& DestArchive)
+	{
+		if (!Object)
+		{
+			return;
+		}
+
+		// Helper struct to identify dependency of a UObject on other UObject(s) except given one (its outer)
+		struct FObjectDependencyAnalyzer : public FArchiveUObject
+		{
+			FObjectDependencyAnalyzer(UObject* InSourceObject, const TSet<UObject*>& InValidObjects)
+				: SourceObject(InSourceObject)
+				, ValidObjects(InValidObjects)
+			{ }
+
+			virtual FArchive& operator<<(UObject*& Obj) override
+			{
+				if (Obj != nullptr)
+				{
+					// Limit serialization to sub-object of source object
+					if (Obj == SourceObject->GetOuter() || Obj->IsA<UPackage>() || (Obj->HasAnyFlags(RF_Public) && Obj->GetOuter()->IsA<UPackage>()))
+					{
+						return FArchiveUObject::operator<<(Obj);
+					}
+					// Stop serialization when a dependency is found or has been found
+					else if (Obj != SourceObject && !DependentObjects.Contains(Obj) && ValidObjects.Contains(Obj))
+					{
+						DependentObjects.Add(Obj);
+					}
+				}
+
+				return *this;
+			}
+
+			UObject* SourceObject;
+			const TSet<UObject*>& ValidObjects;
+			TSet<UObject*> DependentObjects;
+		};
+
+		FAssetCacheObjectWriter Writer{ *Object, Buffer };
+		FArchiveObjectPtrAsStringWrapper Ar{ Writer };
+
+		// Transfer over whether we're cooking or not, but otherwise leave at the default settings, as that is
+		// how we'll try to deserialize on the other end, which is potentially at runtime
+		Ar.SetCookData(DestArchive.GetCookData());
+		Ar.SetFilterEditorOnly(DestArchive.IsFilterEditorOnly());
+		Ar.SetUEVer(DestArchive.UEVer());
+		Ar.SetLicenseeUEVer(DestArchive.LicenseeUEVer());
+		Ar.SetCustomVersions(DestArchive.GetCustomVersions());
+
+		// Collect sub-objects depending on input object including nested objects
+		TArray< UObject* > SubObjectsArray;
+		GetObjectsWithOuter(Object, SubObjectsArray, /*bIncludeNestedObjects = */ true);
+		TSet<UObject*> SubObjectsSet(SubObjectsArray);
+
+		// Keep track of which subobjects are actually referenced by our UObject tree.
+		// Other objects may just be hiding inside of them if they're referenced by an external UObject, and so
+		// shouldn't be serialized to disk (this happens with UAssetImportData: Some UObjects are created with them
+		// already, and those importdatas remain there even when we swap in our own UUsdAssetImportData, as they're
+		// referenced by the transaction buffer);
+		TSet<UObject*> ActualDependencies;
+		{
+			FObjectDependencyAnalyzer Analyzer(Object, SubObjectsSet);
+			Object->Serialize(Analyzer);
+			ActualDependencies.Append(Analyzer.DependentObjects);
+		}
+
+		// Sort array of sub-objects based on their inter-dependency
+		{
+			// Create and initialize graph of dependency between sub-objects
+			TMap< UObject*, TSet<UObject*> > SubObjectDependencyGraph;
+			SubObjectDependencyGraph.Reserve(SubObjectsArray.Num());
+
+			for (UObject* SubObject : SubObjectsArray)
+			{
+				SubObjectDependencyGraph.Add(SubObject);
+			}
+
+			// Build graph of dependency: each entry contains the set of sub-objects to create before itself
+			for (UObject* SubObject : SubObjectsArray)
+			{
+				FObjectDependencyAnalyzer SubAnalyzer(SubObject, SubObjectsSet);
+				SubObject->Serialize(SubAnalyzer);
+
+				SubObjectDependencyGraph[SubObject].Append(SubAnalyzer.DependentObjects);
+				ActualDependencies.Append(SubAnalyzer.DependentObjects);
+			}
+
+			// Only keep the UObjects that we're actually referencing
+			TArray<UObject*> UsedSubObjectsArray;
+			UsedSubObjectsArray.Reserve(SubObjectsArray.Num());
+			for (UObject* SubObject : SubObjectsArray)
+			{
+				if (ActualDependencies.Contains(SubObject))
+				{
+					UsedSubObjectsArray.Add(SubObject);
+				}
+			}
+			Swap(UsedSubObjectsArray, SubObjectsArray);
+
+			// Sort array of sub-objects: first objects do not depend on ones below
+			int32 Count = SubObjectsArray.Num();
+			SubObjectsArray.Empty(Count);
+
+			while (Count != SubObjectsArray.Num())
+			{
+				for (auto& Entry : SubObjectDependencyGraph)
+				{
+					if (Entry.Value.Num() == 0)
+					{
+						UObject* SubObject = Entry.Key;
+
+						SubObjectDependencyGraph.Remove(SubObject);
+
+						SubObjectsArray.Add(SubObject);
+
+						for (auto& SubEntry : SubObjectDependencyGraph)
+						{
+							if (SubEntry.Value.Num() > 0)
+							{
+								SubEntry.Value.Remove(SubObject);
+							}
+						}
+
+						break;
+					}
+				}
+			}
+		}
+
+		// Serialize size of array
+		int32 SubObjectsCount = SubObjectsArray.Num();
+		Ar << SubObjectsCount;
+
+		// Serialize class and path of each sub-object
+		for (int32 Index = SubObjectsArray.Num() - 1; Index >= 0; --Index)
+		{
+			const UObject* SubObject = SubObjectsArray[Index];
+
+			UClass* SubObjectClass = SubObject->GetClass();
+
+			FString ClassName = SubObjectClass->GetPathName();
+			Ar << ClassName;
+
+			int32 ObjectFlags = SubObject->GetFlags();
+			Ar << ObjectFlags;
+		}
+
+		// Serialize sub-objects' outer path and name
+		// Done in reverse order since a sub-object can be the outer of another sub-object
+		// it depends on. Not the opposite
+		for (int32 Index = SubObjectsArray.Num() - 1; Index >= 0; --Index)
+		{
+			const UObject* SubObject = SubObjectsArray[Index];
+
+			FSoftObjectPath SoftPath(SubObject->GetOuter());
+
+			FString SoftPathString = SoftPath.ToString();
+			Ar << SoftPathString;
+
+			FString SubObjectName = SubObject->GetName();
+			Ar << SubObjectName;
+		}
+
+		for (UObject* SubObject : SubObjectsArray)
+		{
+			if (SubObject && !SubObject->HasAnyFlags(RF_DefaultSubObject))
+			{
+				// Ensure these objects are serialized with flags as if they were regular
+				// persistent assets, in case they have branching on their serialization functions that watch
+				// out for this
+				SubObject->ClearFlags(RF_Transient);
+				SubObject->SetFlags(RF_Public);
+				{
+					SubObject->Serialize(Ar);
+				}
+				SubObject->SetFlags(RF_Transient);
+			}
+		}
+
+		Object->ClearFlags(RF_Transient);
+		Object->SetFlags(RF_Public);
+		{
+			Object->Serialize(Ar);
+		}
+		Object->SetFlags(RF_Transient);
+	}
+
+	// Copied from DataprepSnapshot.cpp, the "ReadSnapshotData" function
+	void DeserializeObjectAndSubObjects(UObject* Object, const TArray<uint8>& InBuffer)
+	{
+		// Remove all objects created by default that InObject is dependent on
+		// This method must obviously be called just after the InObject is created
+		auto RemoveDefaultDependencies = [](UObject* InObject)
+		{
+			TArray< UObject* > ObjectsWithOuter;
+			GetObjectsWithOuter(InObject, ObjectsWithOuter, /*bIncludeNestedObjects = */ true);
+
+			for (UObject* ObjectWithOuter : ObjectsWithOuter)
+			{
+				// Do not delete default sub-objects
+				if (!ObjectWithOuter->HasAnyFlags(RF_DefaultSubObject))
+				{
+					if (ObjectWithOuter != nullptr)
+					{
+						ObjectWithOuter->Rename(nullptr, GetTransientPackage(), REN_NonTransactional | REN_DontCreateRedirectors);
+					}
+					ObjectWithOuter->MarkAsGarbage();
+				}
+			}
+		};
+
+		RemoveDefaultDependencies(Object);
+
+		FAssetCacheObjectReader Reader{ *Object, InBuffer };
+		FArchiveObjectPtrAsStringWrapper Ar{ Reader };
+
+		// Deserialize count of sub-objects
+		int32 SubObjectsCount = 0;
+		Ar << SubObjectsCount;
+
+		// Create empty sub-objects based on class and patch
+		TArray< UObject* > SubObjectsArray;
+		SubObjectsArray.SetNumZeroed(SubObjectsCount);
+
+		// Create root name to avoid name collision
+		FString RootName = FGuid::NewGuid().ToString();
+
+		for (int32 Index = SubObjectsCount - 1; Index >= 0; --Index)
+		{
+			FString ClassName;
+			Ar << ClassName;
+
+			UClass* SubObjectClass = FindClass(*ClassName);
+			check(SubObjectClass);
+
+			int32 ObjectFlags;
+			Ar << ObjectFlags;
+
+			FString SubObjectName = RootName + FString::FromInt(Index);
+
+			UObject* SubObject = NewObject<UObject>(Object, SubObjectClass, *SubObjectName, EObjectFlags(ObjectFlags));
+			SubObjectsArray[Index] = SubObject;
+
+			RemoveDefaultDependencies(SubObject);
+		}
+
+		// Restore sub-objects' outer if original outer differs from Object
+		// Restoration is done in the order the serialization was done: reverse order
+		for (int32 Index = SubObjectsArray.Num() - 1; Index >= 0; --Index)
+		{
+			FString SoftPathString;
+			Ar << SoftPathString;
+
+			FString SubObjectName;
+			Ar << SubObjectName;
+
+			const FSoftObjectPath SoftPath(SoftPathString);
+
+			UObject* NewOuter = SoftPath.ResolveObject();
+			ensure(NewOuter);
+
+			if (UObject* SubObject = SubObjectsArray[Index])
+			{
+				if (NewOuter != SubObject->GetOuter())
+				{
+					const TCHAR* NewName = nullptr;
+					SubObject->Rename(NewName, NewOuter, REN_NonTransactional | REN_DontCreateRedirectors);
+				}
+
+				UObject* NoNewOuter = nullptr;
+				if (SubObjectName != SubObject->GetName() && SubObject->Rename(*SubObjectName, NoNewOuter, REN_Test))
+				{
+					SubObject->Rename(*SubObjectName, NoNewOuter, REN_NonTransactional | REN_DontCreateRedirectors);
+				}
+			}
+		}
+
+		for (UObject* SubObject : SubObjectsArray)
+		{
+			if (SubObject && !SubObject->HasAnyFlags(RF_DefaultSubObject))
+			{
+				SubObject->Serialize(Ar);
+				SubObject->SetFlags(RF_Transient | RF_Public);
+			}
+		}
+
+		Object->Serialize(Ar);
+		Object->SetFlags(RF_Transient | RF_Public);
+
+		if (UTexture* Texture = Cast<UTexture>(Object))
+		{
+			Texture->UpdateResource();
+		}
+	}
+}
+
+FArchive& operator<<(FArchive& Ar, UUsdAssetCache2::ECacheStorageType& Type)
+{
+	Ar << (uint8&)Type;
+	return Ar;
+}
+
+void UUsdAssetCache2::FCachedAssetInfo::Serialize(FArchive& Ar, UObject* Owner)
+{
+	if (Ar.IsSaving())
+	{
+		Ar << GCurrentAssetInfoVersion;
+	}
+	else if (Ar.IsLoading())
+	{
+		int32 SavedAssetInfoVersion = INDEX_NONE;
+		Ar << SavedAssetInfoVersion;
+		ensure(SavedAssetInfoVersion <= GCurrentAssetInfoVersion);
+	}
+
+	Ar << Hash;
+	Ar << AssetClassName;
+	Ar << AssetName;
+	Ar << (uint32&)AssetFlags;
+
+	Ar << SizeOnDiskInBytes;
+	Ar << SizeOnMemoryInBytes;
+	Ar << Dependencies;
+	Ar << Dependents;
+
+	// Very important: This flag ensures that when we call Serialize below we won't actually
+	// load the full bulkdata right away, so that we can do it on-demand
+	BulkData.SetBulkDataFlags(BULKDATA_Force_NOT_InlinePayload);
+	BulkData.Serialize(Ar, Owner);
+}
+
+UUsdAssetCache2::UUsdAssetCache2()
+	: LRUCache(-1)  // Negative MaxElements as we don't want the cache to drop elements by itself
+{
+}
+
+#if WITH_EDITOR
+void UUsdAssetCache2::PostEditChangeProperty(FPropertyChangedEvent& PropertyChangedEvent)
+{
+	if (PropertyChangedEvent.GetPropertyName() == GET_MEMBER_NAME_CHECKED(UUsdAssetCache2, PersistentAssetStorageSizeMB))
+	{
+		// We won't be persisting anything if we're transient anyway
+		if (GetOutermost() != GetTransientPackage())
+		{
+			RefreshStorage();
+		}
+	}
+	else if (PropertyChangedEvent.GetPropertyName() == GET_MEMBER_NAME_CHECKED(UUsdAssetCache2, UnreferencedAssetStorageSizeMB))
+	{
+		RefreshStorage();
+	}
+}
+#endif // #if WITH_EDITOR
+
+void UUsdAssetCache2::CacheAsset(const FString& Hash, UObject* Asset, const UObject* Referencer)
+{
+	if (!Asset)
+	{
+		UE_LOG(LogUsd, Warning, TEXT("Attempted to add a null asset to USD Asset Cache with hash '%s'!"), *Hash);
+		return;
+	}
+
+	UE_LOG(LogUsd, Verbose, TEXT("Caching asset '%s' with for hash '%s' into the USD Asset Cache '%s'"),
+		*Asset->GetPathName(),
+		*Hash,
+		*GetPathName()
+	);
+
+	FCachedAssetInfo* FoundInfo = nullptr;
+	{
+		FReadScopeLock Lock(RWLock);
+
+		FoundInfo = LRUCache.Find(Asset);
+		if (FoundInfo)
+		{
+			if (FoundInfo->Hash != Hash)
+			{
+				UE_LOG(LogUsd, Warning, TEXT("Attempted to store asset '%s' more than once into the USD Asset Cache '%s'! (old hash: '%s', new hash: '%s')"),
+					*Asset->GetPathName(),
+					*GetPathName(),
+					*FoundInfo->Hash,
+					*Hash
+				);
+				return;
+			}
+		}
+	}
+
+	const UObject* ReferencerToUse = Referencer ? Referencer : CurrentScopedReferencer;
+
+	if (UObject* ExistingAsset = AssetStorage.FindRef(Hash))
+	{
+		FWriteScopeLock Lock(RWLock);
+
+		if (Asset == ExistingAsset)
+		{
+			// We're trying to cache an asset that's already in the cache with this same hash.
+			// Just record that we "touched" this asset and return
+			Modify();
+			if (ReferencerToUse)
+			{
+				FoundInfo->Referencers.Add(FObjectKey{ ReferencerToUse });
+			}
+			LRUCache.FindAndTouch(Asset);
+			ActiveAssets.Add(Asset);
+			return;
+		}
+		else
+		{
+			bool bRemoved = false;
+			if (CanRemoveAsset(Hash))
+			{
+				if (UObject* OldAsset = RemoveAsset(Hash))
+				{
+					bRemoved = true;
+					UE_LOG(LogUsd, Log, TEXT("Overwriting asset '%s' with '%s' (for hash '%s') in the USD Asset Cache '%s'"),
+						*OldAsset->GetPathName(),
+						*Asset->GetPathName(),
+						*Hash,
+						*GetPathName()
+					);
+				}
+			}
+
+			if(!bRemoved)
+			{
+				UE_LOG(LogUsd, Error, TEXT("Irrecoverable hash collision! Asset '%s' cannot be cached into USD Asset Cache '%s' as the existing asset '%s' with the same hash '%s' could not be discarded!"),
+					*Asset->GetPathName(),
+					*GetPathName(),
+					*ExistingAsset->GetPathName(),
+					*Hash
+				);
+				return;
+			}
+		}
+	}
+
+	Modify();
+
+	FWriteScopeLock Lock(RWLock);
+
+	// We only cache assets that are originally transient, and calling code should know this... we shouldn't try
+	// caching some random persistent asset
+	ensure(Asset->HasAnyFlags(RF_Transient));
+	ensure(Asset->GetOutermost() == GetTransientPackage() || Asset->GetOuter() == this);
+
+	// Rename before adding to our maps as we track some stuff with SoftObjectPaths
+	if (Asset->GetOuter() != this)
+	{
+		UObject* NewOuter = this;
+		const FName NewName = MakeUniqueObjectName(NewOuter, Asset->GetClass(), Asset->GetFName());
+		Asset->Rename(*NewName.ToString(), NewOuter);
+	}
+
+	FCachedAssetInfo NewInfo;
+	NewInfo.Hash = Hash;
+	NewInfo.AssetClassName = Asset->GetClass()->GetPathName();
+	NewInfo.AssetName = Asset->GetName();
+	NewInfo.AssetFlags = Asset->GetFlags();
+	if (ReferencerToUse)
+	{
+		NewInfo.Referencers.Add(FObjectKey{ ReferencerToUse });
+	}
+
+	AssetStorage.Add(Hash, Asset);
+	LRUCache.Add(Asset, MoveTemp(NewInfo));
+	ActiveAssets.Add(Asset);
+}
+
+bool UUsdAssetCache2::CanRemoveAsset(const FString& Hash)
+{
+	FReadScopeLock Lock(RWLock);
+
+	FCachedAssetInfo* Info = nullptr;
+
+	if (UObject* FoundAsset = AssetStorage.FindRef(Hash))
+	{
+		Info = LRUCache.Find(FoundAsset);
+	}
+	else if (FSoftObjectPath* FoundPendingPath = PendingPersistentStorage.Find(Hash))
+	{
+		Info = LRUCache.Find(*FoundPendingPath);
+	}
+
+	if (!Info)
+	{
+		return true;
+	}
+
+	if (Info->Referencers.Num() > 0 || Info->Dependents.Num() > 0)
+	{
+		return false;
+	}
+
+	return true;
+}
+
+UObject* UUsdAssetCache2::RemoveAsset(const FString& Hash)
+{
+	FSoftObjectPath FoundObjectPath;
+	TObjectPtr<UObject> FoundObject;
+	{
+		FWriteScopeLock Lock(RWLock);
+
+		if (AssetStorage.RemoveAndCopyValue(Hash, FoundObject))
+		{
+			FoundObjectPath = FoundObject;
+			ActiveAssets.Remove(FoundObject);
+		}
+		else
+		{
+			PendingPersistentStorage.RemoveAndCopyValue(Hash, FoundObjectPath);
+		}
+	}
+
+	if (!FoundObjectPath.IsNull())
+	{
+		Modify();
+
+		// Modify will call Serialize, so we need to release and get a new lock
+		FWriteScopeLock Lock(RWLock);
+
+		if (FCachedAssetInfo* Info = LRUCache.Find(FoundObjectPath))
+		{
+			ensure(Info->Referencers.Num() == 0);
+			ensure(Info->Dependents.Num() == 0);
+
+			// Make sure all of our dependencies know we're no longer dependent on them
+			for (const FSoftObjectPath& Dependency : Info->Dependencies)
+			{
+				if (FCachedAssetInfo* DependencyInfo = LRUCache.Find(Dependency))
+				{
+					DependencyInfo->Dependents.Remove(FoundObjectPath);
+				}
+			}
+			Info->Dependencies.Empty();
+
+			TryUnloadAsset(*Info);
+			Info->BulkData.RemoveBulkData();
+			LRUCache.Remove(FoundObjectPath);
+		}
+	}
+
+	return FoundObject;
+}
+
+UObject* UUsdAssetCache2::GetCachedAsset(const FString& Hash)
+{
+	// Write here as this will affect LRUCache order and PendingPersistentStorage
+	TOptional<FRWScopeLock> Lock;
+	Lock.Emplace(RWLock, SLT_Write);
+
+	UObject* FoundObject = AssetStorage.FindRef(Hash);
+
+	UE_LOG(LogUsd, Verbose, TEXT("Fetching cached asset with hash '%s' for cache '%s'. Loaded in AssetStorage? %d. Size of LRU cache: %d, Pending persistent assets: %d"),
+		*Hash,
+		*GetPathName(),
+		FoundObject != nullptr,
+		LRUCache.Num(),
+		PendingPersistentStorage.Num()
+	);
+
+	// Maybe this is a persistent asset we just haven't loaded yet
+	if (!FoundObject)
+	{
+		if (FSoftObjectPath* FoundSoftPath = PendingPersistentStorage.Find(Hash))
+		{
+			if (FCachedAssetInfo* FoundCachedInfo = LRUCache.Find(*FoundSoftPath))
+			{
+				const FString ClassNameString = FoundCachedInfo->AssetClassName.ToString();
+
+				// Actually load the UObject and subobjects from disk
+				if (UClass* FoundClass = UE::AssetCache::Private::FindClass(*ClassNameString))
+				{
+					// Deserialize the referenced asset from the Info's bulkdata
+					// TODO: If performance is an issue here we can probably prevent this copy by having our own
+					// FObjectReader that uses a pointer and offset and reads it off the bulkdata directly
+					TArray<uint8> BulkDataCopy;
+					const void* BulkDataBytes = FoundCachedInfo->BulkData.LockReadOnly();
+					{
+						BulkDataCopy.SetNumUninitialized(FoundCachedInfo->BulkData.GetBulkDataSize());
+						FMemory::Memcpy(BulkDataCopy.GetData(), BulkDataBytes, BulkDataCopy.Num());
+					}
+					FoundCachedInfo->BulkData.Unlock();
+
+					FName AssetName = FoundCachedInfo->AssetName.IsEmpty() ? NAME_None : FName{ *FoundCachedInfo->AssetName };
+
+					// They're always "transient", since we'll persist them ourselves
+					UObject* NewAsset = NewObject<UObject>(this, FoundClass, AssetName, FoundCachedInfo->AssetFlags | RF_Transient);
+					UE::AssetCache::Private::DeserializeObjectAndSubObjects(NewAsset, BulkDataCopy);
+					ensure(NewAsset);
+
+					UE_LOG(LogUsd, Verbose, TEXT("Deserialized asset at '%s' (%s, hash '%s', storage %d) from bulkdata (%.3f MB)"),
+						*FoundSoftPath->ToString(),
+						NewAsset ? *NewAsset->GetPathName() : TEXT("nullptr"),
+						*Hash,
+						FoundCachedInfo->CurrentStorageType,
+						BulkDataCopy.Num() / 1000000.0
+					);
+
+					// Explicitly clearing this here for safety, as it will be invalidated when we remove Hash from
+					// PendingPersistentStorage just below
+					FoundSoftPath = nullptr;
+
+					FoundObject = NewAsset;
+					AssetStorage.Add(Hash, NewAsset);
+					PendingPersistentStorage.Remove(Hash);
+				}
+				else
+				{
+					UE_LOG(LogUsd, Warning, TEXT("Failed to find object class '%s' when deserializing asset '%s' with hash '%s'"),
+						*ClassNameString,
+						*FoundSoftPath->ToString(),
+						*Hash
+					);
+				}
+			}
+			else
+			{
+				UE_LOG(LogUsd, Warning, TEXT("Failed to find info about asset '%s' with hash '%s' when deserializing on-demand"),
+					*FoundSoftPath->ToString(),
+					*Hash
+				);
+			}
+		}
+	}
+
+	if (FoundObject)
+	{
+		ActiveAssets.Add(FoundObject);
+		ensure(LRUCache.FindAndTouch(FoundObject));
+
+		if (CurrentScopedReferencer)
+		{
+			Lock.Reset(); // Release our lock as AddAssetReference will want to acquire it
+			AddAssetReference(FoundObject, CurrentScopedReferencer);
+		}
+	}
+
+	return FoundObject;
+}
+
+bool UUsdAssetCache2::AddAssetReference(const UObject* Asset, const UObject* Referencer)
+{
+	if (!Asset || !Referencer)
+	{
+		return false;
+	}
+
+	FWriteScopeLock Lock(RWLock);
+
+	if (FCachedAssetInfo* Info = LRUCache.Find(Asset))
+	{
+		Info->Referencers.Add(FObjectKey{Referencer});
+
+		UE_LOG(LogUsd, Verbose, TEXT("Added referencer '%s' for asset '%s'"),
+			*Referencer->GetPathName(),
+			*Asset->GetPathName()
+		);
+		return true;
+	}
+
+	return false;
+}
+
+bool UUsdAssetCache2::RemoveAssetReference(const UObject* Asset, const UObject* Referencer)
+{
+	if (!Asset || !Referencer)
+	{
+		return false;
+	}
+
+	FWriteScopeLock Lock(RWLock);
+
+	bool bRemovedSomething = false;
+	if (FCachedAssetInfo* Info = LRUCache.Find(Asset))
+	{
+		const int32 NumRemoved = Info->Referencers.Remove(FObjectKey{Referencer});
+		if (NumRemoved > 0)
+		{
+			UE_LOG(LogUsd, Verbose, TEXT("Removed referencer '%s' from asset '%s'"),
+				*Referencer->GetPathName(),
+				*Asset->GetPathName()
+			);
+
+			bRemovedSomething = true;
+		}
+	}
+
+	return bRemovedSomething;
+}
+
+bool UUsdAssetCache2::RemoveAllAssetReferences(const UObject* Referencer)
+{
+	if (!Referencer)
+	{
+		return false;
+	}
+
+	FWriteScopeLock Lock(RWLock);
+
+	FObjectKey Key{Referencer};
+
+	bool bRemovedSomething = false;
+	for (TLruCache<FSoftObjectPath, FCachedAssetInfo>::TIterator Iter{LRUCache}; Iter; ++Iter)
+	{
+		FCachedAssetInfo& Info = Iter.Value();
+		int32 NumRemoved = Info.Referencers.Remove(Key);
+		if (NumRemoved > 0)
+		{
+			bRemovedSomething = true;
+		}
+	}
+
+	return bRemovedSomething;
+}
+
+FString UUsdAssetCache2::GetHashForAsset(const UObject* Asset) const
+{
+	if (!Asset)
+	{
+		return {};
+	}
+
+	FReadScopeLock Lock(RWLock);
+
+	if (const FCachedAssetInfo* Info = LRUCache.Find(Asset))
+	{
+		return Info->Hash;
+	}
+
+	return {};
+}
+
+bool UUsdAssetCache2::IsAssetOwnedByCache(const FString& AssetPath) const
+{
+	FReadScopeLock Lock(RWLock);
+
+	return IsAssetOwnedByCacheInternal(AssetPath);
+}
+
+int32 UUsdAssetCache2::GetNumAssets() const
+{
+	return LRUCache.Num();
+}
+
+TArray<FString> UUsdAssetCache2::GetAllAssetHashes() const
+{
+	FReadScopeLock Lock(RWLock);
+
+	TArray<FString> Hashes;
+	Hashes.Reserve(LRUCache.Num());
+
+	for (TLruCache<FSoftObjectPath, FCachedAssetInfo>::TConstIterator Iter{ LRUCache }; Iter; ++Iter)
+	{
+		Hashes.Add(Iter.Value().Hash);
+	}
+
+	return Hashes;
+}
+
+TArray<UObject*> UUsdAssetCache2::GetAllLoadedAssets() const
+{
+	FReadScopeLock Lock(RWLock);
+
+	TArray<UObject*> Objects;
+	Objects.Reserve(AssetStorage.Num());
+
+	for (const TPair<FString, TObjectPtr<UObject>>& Pair : AssetStorage)
+	{
+		Objects.Add(Pair.Value.Get());
+	}
+
+	return Objects;
+}
+
+TArray<FString> UUsdAssetCache2::GetAllCachedAssetPaths() const
+{
+	FReadScopeLock Lock(RWLock);
+
+	TArray<FString> Paths;
+	Paths.Reserve(LRUCache.Num());
+
+	for (TLruCache<FSoftObjectPath, FCachedAssetInfo>::TConstIterator Iter{ LRUCache }; Iter; ++Iter)
+	{
+		Paths.Add(Iter.Key().ToString());
+	}
+
+	return Paths;
+}
+
+void UUsdAssetCache2::Reset()
+{
+	FWriteScopeLock Lock(RWLock);
+
+	Modify();
+
+	// We never want the LRUCache to drop items by itself
+	const int32 MaxElements = -1;
+	AssetStorage.Reset();
+	ActiveAssets.Reset();
+	for (TLruCache<FSoftObjectPath, FCachedAssetInfo>::TIterator Iter{ LRUCache }; Iter; ++Iter)
+	{
+		FCachedAssetInfo& Info = Iter.Value();
+		Info.BulkData.RemoveBulkData();
+	}
+	LRUCache.Empty(MaxElements);
+	PendingPersistentStorage.Reset();
+}
+
+void UUsdAssetCache2::RefreshStorage()
+{
+	using namespace UE::AssetCache::Private;
+
+	const bool bIsTransientAssetCache = GetOutermost() == GetTransientPackage();
+
+	const uint64 MaxUnreferencedTransientBytes = UnreferencedAssetStorageSizeMB < 0.0
+		? TNumericLimits<uint64>::Max()
+		: static_cast<uint64>(FMath::Max(UnreferencedAssetStorageSizeMB, 0.0) + 0.5) * 1000000;
+
+	const uint64 MaxPersistentBytes = PersistentAssetStorageSizeMB < 0.0
+		? TNumericLimits<uint64>::Max()
+		: static_cast<uint64>(FMath::Max(PersistentAssetStorageSizeMB, 0.0) + 0.5) * 1000000;
+
+	uint64 UnreferencedAssetsSumBytes = 0;
+	uint64 PersistentAssetSumBytes = 0;
+
+	TMap<FString, TObjectPtr<UObject>> NewAssetStorage;
+	TMap<FString, FSoftObjectPath> NewPendingPersistentStorage;
+
+	NewAssetStorage.Reserve(AssetStorage.Num());
+	NewPendingPersistentStorage.Reserve(PendingPersistentStorage.Num());
+
+	TSet<FSoftObjectPath> VisitedAssets;
+
+	Modify();
+
+	FWriteScopeLock Lock(RWLock);
+
+	double StartTime = FPlatformTime::Cycles64();
+
+	// Do an initial pass to reset our current storage types and fetch some remaining data
+	for (TLruCache<FSoftObjectPath, FCachedAssetInfo>::TIterator Iter{ LRUCache }; Iter; ++Iter)
+	{
+		FCachedAssetInfo& Info = Iter.Value();
+
+		Info.CurrentStorageType = ECacheStorageType::None;
+
+		if (Info.SizeOnDiskInBytes == 0)
+		{
+			const FSoftObjectPath& AssetPath = Iter.Key();
+			UObject* Asset = AssetStorage.FindRef(Info.Hash);
+
+			// If we have no size yet, then we must have just cached the asset, and so it must still be loaded
+			ensure(Asset);
+
+			// We defer fetching asset sizes to here as some assets are not "complete" by the time we originally cache
+			// them, but should be fine now (e.g. static meshes that still need to be built).
+			// We should always do one RefreshStorage after we fully finish any operation that adds assets
+			// (e.g. loading a stage, reloading, etc.), so we'll always have at least have one pass here
+			Info.SizeOnDiskInBytes = UE::AssetCache::Private::GetAssetSizeBytes(Asset, /*bOnDisk*/ true);
+			Info.SizeOnMemoryInBytes = UE::AssetCache::Private::GetAssetSizeBytes(Asset, /*bOnDisk*/ false);
+
+			// We also wait to fetch dependencies only here so that we know we have cached everything we want,
+			// and can ensure that IsAssetOwnedByCache wouldn't reject some asset that would have eventually be cached
+			// (e.g. if we cached a SkeletalMesh and only later cached the Skeleton that it uses)
+			TSet<UObject*> Dependencies = IUsdClassesModule::GetAssetDependencies(Asset);
+			Info.Dependencies.Empty(Dependencies.Num());
+			for (UObject* Dependency : Dependencies)
+			{
+				// We need to check this here because sometimes our assets refer to other persistent assets
+				// (e.g. unreal materials)
+				if (IsAssetOwnedByCacheInternal(Dependency->GetPathName()))
+				{
+					Info.Dependencies.Add(Dependency);
+					if (FCachedAssetInfo* DependencyInfo = LRUCache.Find(Dependency))
+					{
+						DependencyInfo->Dependents.Add(AssetPath);
+					}
+				}
+			}
+		}
+	}
+
+	TFunction<bool(const FCachedAssetInfo&)> DependentTreeHasReferencedAsset;
+	DependentTreeHasReferencedAsset = [this, &DependentTreeHasReferencedAsset](const FCachedAssetInfo& Info) -> bool
+	{
+		for (const FSoftObjectPath& Dependent : Info.Dependents)
+		{
+			FCachedAssetInfo* DependentInfo = LRUCache.Find(Dependent);
+			if (ensure(DependentInfo))
+			{
+				if (DependentInfo->Referencers.Num() > 0)
+				{
+					return true;
+				}
+
+				if (DependentTreeHasReferencedAsset(*DependentInfo))
+				{
+					return true;
+				}
+			}
+		}
+
+		return false;
+	};
+
+	TFunction<ECacheStorageType(FCachedAssetInfo&, FSoftObjectPath&)> VisitEntryRecursively;
+	VisitEntryRecursively =
+		[
+			this,
+			&DependentTreeHasReferencedAsset,
+			&VisitEntryRecursively,
+			&VisitedAssets,
+			bIsTransientAssetCache,
+			MaxUnreferencedTransientBytes,
+			MaxPersistentBytes,
+			&UnreferencedAssetsSumBytes,
+			&PersistentAssetSumBytes,
+			&NewAssetStorage,
+			&NewPendingPersistentStorage
+		]
+		(FCachedAssetInfo& Info, FSoftObjectPath& AssetPath) -> ECacheStorageType
+		{
+			if (AssetPath.IsNull())
+			{
+				ensure(false);
+				return ECacheStorageType::None;
+			}
+
+			if (VisitedAssets.Contains(AssetPath))
+			{
+				return Info.CurrentStorageType;
+			}
+			VisitedAssets.Add(AssetPath);
+
+			// Always visit dependencies first: We can only ever persist an asset if its dependencies are also
+			// persisted. Here we'll also track the "worst" storage any of our dependencies got (e.g. if one of our
+			// textures managed Persistent, but the other got Referenced, we'll track Referenced)
+			ECacheStorageType WorstDependencyStorage = ECacheStorageType::Persistent;
+			for (FSoftObjectPath& Dependency : Info.Dependencies)
+			{
+				FCachedAssetInfo* DependencyInfo = LRUCache.Find(Dependency);
+				if (!ensure(DependencyInfo))
+				{
+					continue;
+				}
+
+				ECacheStorageType DependencyStorage = VisitEntryRecursively(*DependencyInfo, Dependency);
+				WorstDependencyStorage = static_cast<ECacheStorageType>(FMath::Min(
+					static_cast<uint8>(WorstDependencyStorage),
+					static_cast<uint8>(DependencyStorage)
+				));
+			}
+
+			const bool bAssetFitsInPersistent = PersistentAssetSumBytes + Info.SizeOnDiskInBytes <= MaxPersistentBytes;
+			const bool bAssetFitsInUnreferenced = UnreferencedAssetsSumBytes + Info.SizeOnMemoryInBytes <= MaxUnreferencedTransientBytes;
+			const bool bDependenciesArePersistent = WorstDependencyStorage == ECacheStorageType::Persistent;
+			const bool bAllDependenciesInStorage = WorstDependencyStorage != ECacheStorageType::None;
+			const bool bDependentIsReferenced = DependentTreeHasReferencedAsset(Info);
+
+			const FString AssetPathStr = AssetPath.ToString();
+			const FString AssetClassNameStr = Info.AssetClassName.ToString();
+			UObject* Asset = AssetStorage.FindRef(Info.Hash);
+
+			// Priority 1: Asset can be moved to persistent storage
+			if (!bIsTransientAssetCache && bAssetFitsInPersistent && bDependenciesArePersistent && AssetCanBePersisted(*AssetClassNameStr))
+			{
+				Info.CurrentStorageType = ECacheStorageType::Persistent;
+			}
+
+			// Priority 2: Asset is kept on transient storage because it (or a dependent) is referenced
+			else if (Info.Referencers.Num() > 0 || bDependentIsReferenced)
+			{
+				Info.CurrentStorageType = ECacheStorageType::Referenced;
+			}
+
+			// Priority 3: Asset fits on transient storage even while unreferenced, as long as all dependencies also
+			// made it somewhere and weren't evicted
+			else if (bAssetFitsInUnreferenced && bAllDependenciesInStorage)
+			{
+				Info.CurrentStorageType = ECacheStorageType::Unreferenced;
+			}
+
+			UE_LOG(LogUsd, Verbose, TEXT("Asset '%s' (class %s, hash '%s', size %.3f MB on memory, %.3f MB on disk) fits in persistent? %d, fits in unreferenced? %d, referencer count %d, dependency count %d. Target storage: %d"),
+				*AssetPathStr,
+				*AssetClassNameStr,
+				*Info.Hash,
+				Info.SizeOnMemoryInBytes / 1000000.0,
+				Info.SizeOnDiskInBytes / 1000000.0,
+				bAssetFitsInPersistent,
+				bAssetFitsInUnreferenced,
+				Info.Referencers.Num(),
+				Info.Dependencies.Num(),
+				Info.CurrentStorageType
+			);
+
+			switch (Info.CurrentStorageType)
+			{
+				case ECacheStorageType::None:
+				{
+					Info.BulkData.UnloadBulkData();
+					break;
+				}
+				case ECacheStorageType::Referenced:
+				{
+					// We'll at least unload the bulkdata here, but not actually remove it from disk yet, in case we
+					// put this asset back in persistent later and can reuse the same bulkdata.
+					// We'll get rid of the bulkdata for sure on UUsdAssetCache2::Serialize(), if the asset is not
+					// persistent by then
+					Info.BulkData.UnloadBulkData();
+
+					ensure(Asset);
+					NewAssetStorage.Add(Info.Hash, Asset);
+					break;
+				}
+				case ECacheStorageType::Unreferenced:
+				{
+					Info.BulkData.UnloadBulkData();
+
+					ensure(Asset);
+					NewAssetStorage.Add(Info.Hash, Asset);
+
+					UnreferencedAssetsSumBytes += Info.SizeOnMemoryInBytes;
+					UE_LOG(LogUsd, Verbose, TEXT("Updated unreferenced sum to %.3f out of max %u MB"),
+						UnreferencedAssetsSumBytes / 1000000.0,
+						MaxUnreferencedTransientBytes / 1000000
+					);
+					break;
+				}
+				case ECacheStorageType::Persistent:
+				{
+					if (Asset)
+					{
+						ensure(AssetStorage.Contains(Info.Hash));
+						NewAssetStorage.Add(Info.Hash, Asset);
+					}
+					else
+					{
+						ensure(!AssetStorage.Contains(Info.Hash));
+						NewPendingPersistentStorage.Add(Info.Hash, AssetPath);
+					}
+
+					PersistentAssetSumBytes += Info.SizeOnDiskInBytes;
+					UE_LOG(LogUsd, Verbose, TEXT("Updated persistent sum to %.3f out of max %u MB"),
+						PersistentAssetSumBytes / 1000000.0,
+						MaxPersistentBytes / 1000000
+					);
+					break;
+				}
+				default:
+				{
+					ensure(false);
+					break;
+				}
+			}
+
+			return Info.CurrentStorageType;
+		};
+
+	for (TLruCache<FSoftObjectPath, FCachedAssetInfo>::TIterator Iter{ LRUCache }; Iter; ++Iter)
+	{
+		FSoftObjectPath& AssetPath = Iter.Key();
+		FCachedAssetInfo& Info = Iter.Value();
+
+		VisitEntryRecursively(Info, AssetPath);
+	}
+
+	Swap(AssetStorage, NewAssetStorage);
+	Swap(PendingPersistentStorage, NewPendingPersistentStorage);
+
+	// Evict from the LRUCache assets that never made it into any storage: These assets are unreferenced and don't fit
+	// anymore
+	{
+		TLruCache<FSoftObjectPath, FCachedAssetInfo>::TIterator Iter{ LRUCache };
+		while (Iter)
+		{
+			FSoftObjectPath& AssetPath = Iter.Key();
+			FCachedAssetInfo& Info = Iter.Value();
+
+			if (Info.CurrentStorageType == ECacheStorageType::None)
+			{
+				UE_LOG(LogUsd, Verbose, TEXT("Evicting asset '%s' (hash '%s') from the cache '%s'"), *AssetPath.ToString(), *Info.Hash, *GetPathName());
+
+				// Discard data on disk
+				TryUnloadAsset(Info);
+				Info.BulkData.RemoveBulkData();
+
+				// Make sure all of our dependencies know we're no longer a dependent
+				for (const FSoftObjectPath& Dependency : Info.Dependencies)
+				{
+					if (FCachedAssetInfo* DepInfo = LRUCache.Find(Dependency))
+					{
+						DepInfo->Dependents.Remove(AssetPath);
+					}
+				}
+
+				Iter.RemoveCurrentAndIncrement();
+			}
+			else
+			{
+				++Iter;
+			}
+		}
+	}
+
+	// Now that we evicted all potential dependencies we don't care about anymore,
+	// unload all persistent assets that aren't being currently used
+	if (GOnDemandCachedAssetLoading)
+	{
+		for (TLruCache<FSoftObjectPath, FCachedAssetInfo>::TIterator Iter{ LRUCache }; Iter; ++Iter)
+		{
+			FCachedAssetInfo& Info = Iter.Value();
+			TryUnloadAsset(Info);
+		}
+	}
+
+	double ElapsedSeconds = FPlatformTime::ToSeconds64(FPlatformTime::Cycles64() - StartTime);
+
+	UE_LOG(LogUsd, Log, TEXT("Finished asset cache '%s' storage refresh in %.3f seconds: %.3f out of max %u MB of unreferenced assets, %.3f out of max %u MB of persistent assets"),
+		*GetName(),
+		ElapsedSeconds,
+		UnreferencedAssetsSumBytes / 1000000.0,
+		MaxUnreferencedTransientBytes / 1000000,
+		PersistentAssetSumBytes / 1000000.0,
+		MaxPersistentBytes / 1000000
+	);
+}
+
+void UUsdAssetCache2::MarkAssetsAsStale()
+{
+	FWriteScopeLock Lock(RWLock);
+
+	ActiveAssets.Reset();
+}
+
+TSet<UObject*> UUsdAssetCache2::GetActiveAssets() const
+{
+	return ActiveAssets;
+}
+
+void UUsdAssetCache2::PostLoad()
+{
+	Super::PostLoad();
+
+	// Finish loading all of our persistent assets right away if we should.
+	// We do this here and not inside ::Serialize() because this involves renaming UObjects, which can't be done
+	// during load
+	if (!GOnDemandCachedAssetLoading)
+	{
+		TMap<FString, FSoftObjectPath> PendingStorageCopy = PendingPersistentStorage;
+		for (const TPair<FString, FSoftObjectPath>& PendingAsset : PendingStorageCopy)
+		{
+			UObject* AssetFromBulkData = GetCachedAsset(PendingAsset.Key);
+			ensure(AssetFromBulkData);
+		}
+	}
+}
+
+bool UUsdAssetCache2::TryUnloadAsset(FCachedAssetInfo& InOutInfo)
+{
+	if (GOnDemandCachedAssetLoading && InOutInfo.BulkData.CanLoadFromDisk() && InOutInfo.Referencers.Num() == 0)
+	{
+		// Asset is currently loaded
+		if (TObjectPtr<UObject> Asset = AssetStorage.FindRef(InOutInfo.Hash))
+		{
+			// Check to see if all of its dependents can be unloaded too
+			for (const FSoftObjectPath& Dependent : InOutInfo.Dependents)
+			{
+				if (FCachedAssetInfo* DependentInfo = LRUCache.Find(Dependent))
+				{
+					ensure(DependentInfo->Dependencies.Contains(Asset));
+
+					bool bDependentIsUnloaded = TryUnloadAsset(*DependentInfo);
+					if (!bDependentIsUnloaded)
+					{
+						// Our dependent must stay loaded -> So must we
+						return false;
+					}
+				}
+				// We don't know about this dependent anymore (maybe it was evicted?)
+				else
+				{
+					InOutInfo.Dependents.Remove(Dependent);
+				}
+			}
+
+			PendingPersistentStorage.Add(InOutInfo.Hash, Asset);
+			AssetStorage.Remove(InOutInfo.Hash);
+
+			const TCHAR* NewName = nullptr;
+			UObject* NewOuter = GetTransientPackage();
+			Asset->Rename(NewName, NewOuter);
+		}
+
+		return InOutInfo.BulkData.UnloadBulkData();
+	}
+
+	return false;
+}
+
+bool UUsdAssetCache2::IsAssetOwnedByCacheInternal(const FString& AssetPath) const
+{
+	return !AssetPath.IsEmpty() && LRUCache.Contains(AssetPath);
+}
+
+void UUsdAssetCache2::Serialize(FArchive& Ar)
+{
+	if (Ar.IsSaving())
+	{
+		Ar << GCurrentCacheVersion;
+	}
+	else if (Ar.IsLoading())
+	{
+		int32 SavedCacheVersion = INDEX_NONE;
+		Ar << SavedCacheVersion;
+		ensure(SavedCacheVersion <= GCurrentCacheVersion);
+	}
+
+	FWriteScopeLock Lock(RWLock);
+
+	Super::Serialize(Ar);
+
+	if (Ar.IsPersistent() || Ar.IsTransacting() || (Ar.GetPortFlags() & (PPF_DuplicateForPIE | PPF_Duplicate)))
+	{
+		if (Ar.IsSaving())
+		{
+			TSet<FSoftObjectPath> PersistentAssets;
+			for (TLruCache<FSoftObjectPath, FCachedAssetInfo>::TIterator Iter{ LRUCache }; Iter; ++Iter)
+			{
+				const FCachedAssetInfo& Info = Iter.Value();
+				if (Info.CurrentStorageType == ECacheStorageType::Persistent)
+				{
+					PersistentAssets.Add(Iter.Key());
+				}
+			}
+			int32 NumPersistent = PersistentAssets.Num();
+			Ar << NumPersistent;
+
+			// When we cook we must generate a new cooked bulkdata as it may need to filter editor-only stuff during
+			// the serialization process, so make sure none of these assets are pending
+			if (Ar.IsCooking())
+			{
+				// Take a copy here because GetCachedAsset will modify PendingPersistentStorage
+				TMap<FString, FSoftObjectPath> PendingStorageCopy = PendingPersistentStorage;
+				for (const TPair<FString, FSoftObjectPath>& PendingAsset : PendingStorageCopy)
+				{
+					UObject* AssetFromBulkData = GetCachedAsset(PendingAsset.Key);
+					ensure(AssetFromBulkData);
+				}
+			}
+
+			// By using the iterators we'll traverse the most recently used object first
+			for (TLruCache<FSoftObjectPath, FCachedAssetInfo>::TIterator Iter{ LRUCache }; Iter; ++Iter)
+			{
+				FCachedAssetInfo& Info = Iter.Value();
+				FSoftObjectPath& AssetPath = Iter.Key();
+				ensure(!AssetPath.IsNull());
+
+				if (Info.CurrentStorageType != ECacheStorageType::Persistent)
+				{
+					continue;
+				}
+
+				// This is the main way in which assets get serialized into the bulkdata
+				if (UObject* Asset = AssetStorage.FindRef(Info.Hash))
+				{
+					// Serialize the referenced asset into a byte buffer
+					TArray<uint8> Buffer;
+					UE::AssetCache::Private::SerializeObjectAndSubObjects(Asset, Buffer, Ar);
+					SIZE_T NumBytes = Buffer.Num();
+
+					// Copy the byte buffer into the Info's bulkdata in memory
+					// TODO: Maybe it's possible to avoid this copy by having a custom object writer
+					// that can write directly to the realloc'd bulkdata bytes
+					Info.BulkData.Lock(LOCK_READ_WRITE);
+					{
+						void* BulkDataBytes = Info.BulkData.Realloc(Buffer.Num());
+						FMemory::Memcpy(BulkDataBytes, Buffer.GetData(), Buffer.Num());
+					}
+					Info.BulkData.Unlock();
+
+					// Now we know *exactly* how big the asset is on disk now
+					Info.SizeOnDiskInBytes = Buffer.Num();
+				}
+
+				Ar << AssetPath;
+
+				// We only want to serialize dependencies on persistent assets, or else when we deserialize this
+				// we'll have a bunch of dependencies on assets that don't exist
+				TSet<FSoftObjectPath> PersistentDependencies = Info.Dependencies.Intersect(PersistentAssets);
+				TSet<FSoftObjectPath> PersistentDependents = Info.Dependents.Intersect(PersistentAssets);
+				Swap(Info.Dependencies, PersistentDependencies);
+				Swap(Info.Dependents, PersistentDependents);
+				{
+					Info.CurrentStorageType = ECacheStorageType::Persistent;
+					Info.Serialize(Ar, this);
+				}
+				Swap(Info.Dependencies, PersistentDependencies);
+				Swap(Info.Dependents, PersistentDependents);
+			}
+		}
+		else if (Ar.IsLoading())
+		{
+			int32 NumPersistent = INDEX_NONE;
+			Ar << NumPersistent;
+			ensure(NumPersistent >= 0);
+
+			TArray<FCachedAssetInfo> MostToLeastRecent;
+			MostToLeastRecent.Reserve(NumPersistent);
+
+			TArray<FSoftObjectPath> AssetPaths;
+			AssetPaths.Reserve(NumPersistent);
+
+			for (int32 Index = 0; Index < NumPersistent; ++Index)
+			{
+				FSoftObjectPath& AssetPath = AssetPaths.Emplace_GetRef();
+				Ar << AssetPath;
+
+				FCachedAssetInfo& NewEl = MostToLeastRecent.Emplace_GetRef();
+				NewEl.Serialize(Ar, this);
+				NewEl.CurrentStorageType = ECacheStorageType::Persistent;
+
+				PendingPersistentStorage.Add(NewEl.Hash, AssetPath);
+			}
+
+			// We never want the cache to drop items by itself
+			const int32 MaxElements = -1;
+			LRUCache.Empty(MaxElements);
+
+			// Push them back-to-front into the LRU cache to make sure the most recent one is added last
+			for (int32 Index = NumPersistent - 1; Index >= 0; --Index)
+			{
+				LRUCache.Add(MoveTemp(AssetPaths[Index]), MoveTemp(MostToLeastRecent[Index]));
+			}
+		}
+	}
+}
+
+FUsdScopedAssetCacheReferencer::FUsdScopedAssetCacheReferencer(UUsdAssetCache2* InAssetCache, const UObject* Referencer)
+{
+	// For now we're assuming you can't nest these objects
+	ensure(InAssetCache && Referencer && !InAssetCache->CurrentScopedReferencer);
+	InAssetCache->CurrentScopedReferencer = Referencer;
+
+	AssetCache = InAssetCache;
+}
+
+FUsdScopedAssetCacheReferencer::~FUsdScopedAssetCacheReferencer()
+{
+	UUsdAssetCache2* ValidCache = AssetCache.Get();
+	if (ensure(ValidCache))
+	{
+		ValidCache->CurrentScopedReferencer = nullptr;
+	}
+}
diff --git a/Engine/Plugins/Importers/USDImporter/Source/USDClasses/Private/USDClassesModule.cpp b/Engine/Plugins/Importers/USDImporter/Source/USDClasses/Private/USDClassesModule.cpp
--- a/Engine/Plugins/Importers/USDImporter/Source/USDClasses/Private/USDClassesModule.cpp
+++ b/Engine/Plugins/Importers/USDImporter/Source/USDClasses/Private/USDClassesModule.cpp
@@ -5,15 +5,25 @@
 #include "USDLog.h"
 
 #include "AnalyticsEventAttribute.h"
+#include "Animation/AnimBlueprint.h"
+#include "Animation/AnimSequence.h"
+#include "Animation/Skeleton.h"
 #include "Engine/Engine.h"
+#include "Engine/SkeletalMesh.h"
+#include "Engine/SkinnedAssetCommon.h"
+#include "Engine/StaticMesh.h"
+#include "Engine/Texture2D.h"
 #include "EngineAnalytics.h"
+#include "GeometryCache.h"
 #include "HAL/FileManager.h"
+#include "Materials/Material.h"
+#include "Materials/MaterialInstance.h"
 #include "Misc/EngineVersion.h"
 #include "Misc/FileHelper.h"
 #include "Misc/PackageName.h"
 #include "Misc/Paths.h"
 #include "Misc/SecureHash.h"
-#include "Modules/ModuleManager.h"
+#include "PhysicsEngine/PhysicsAsset.h"
 #include "Serialization/JsonSerializer.h"
 #include "UObject/ObjectSaveContext.h"
 #include "UObject/Package.h"
@@ -281,6 +291,112 @@ UWorld* IUsdClassesModule::GetCurrentWorld( bool bEditorWorldsOnly )
 	return EditorWorld;
 }
 
+TSet<UObject*> IUsdClassesModule::GetAssetDependencies(UObject* Asset)
+{
+	TSet<UObject*> Result;
+
+	if (UMaterial* Material = Cast<UMaterial>(Asset))
+	{
+		TArray<UTexture*> UsedTextures;
+		const bool bAllQualityLevels = true;
+		const bool bAllFeatureLevels = true;
+		Material->GetUsedTextures(UsedTextures, EMaterialQualityLevel::High, bAllQualityLevels, ERHIFeatureLevel::SM5, bAllFeatureLevels);
+
+		Result.Reserve(Result.Num() + UsedTextures.Num());
+		for (UTexture* UsedTexture : UsedTextures)
+		{
+			Result.Add(UsedTexture);
+		}
+	}
+	else if (UMaterialInstance* MaterialInstance = Cast<UMaterialInstance>(Asset))
+	{
+		Result.Reserve(Result.Num() + MaterialInstance->TextureParameterValues.Num());
+		for (const FTextureParameterValue& TextureValue : MaterialInstance->TextureParameterValues)
+		{
+			if (UTexture* Texture = TextureValue.ParameterValue)
+			{
+				Result.Add(Texture);
+			}
+		}
+	}
+	else if (USkeletalMesh* SkeletalMesh = Cast<USkeletalMesh>(Asset))
+	{
+		if (USkeleton* Skeleton = SkeletalMesh->GetSkeleton())
+		{
+			Result.Add(Skeleton);
+		}
+
+		if (UPhysicsAsset* PhysicsAsset = SkeletalMesh->GetPhysicsAsset())
+		{
+			Result.Add(PhysicsAsset);
+		}
+
+		Result.Reserve(Result.Num() + SkeletalMesh->GetMaterials().Num());
+		for (const FSkeletalMaterial& SkeletalMaterial : SkeletalMesh->GetMaterials())
+		{
+			if (UMaterialInterface* UsedMaterial = SkeletalMaterial.MaterialInterface)
+			{
+				Result.Add(UsedMaterial);
+			}
+		}
+	}
+	else if (UStaticMesh* StaticMesh = Cast<UStaticMesh>(Asset))
+	{
+		Result.Reserve(Result.Num() + StaticMesh->GetStaticMaterials().Num());
+		for (const FStaticMaterial& StaticMaterial : StaticMesh->GetStaticMaterials())
+		{
+			if (UMaterialInterface* UsedMaterial = StaticMaterial.MaterialInterface)
+			{
+				Result.Add(UsedMaterial);
+			}
+		}
+	}
+	else if (UGeometryCache* GeometryCache = Cast<UGeometryCache>(Asset))
+	{
+		for (UMaterialInterface* UsedMaterial : GeometryCache->Materials)
+		{
+			if (UsedMaterial)
+			{
+				Result.Add(UsedMaterial);
+			}
+		}
+	}
+	else if (UAnimSequence* AnimSequence = Cast<UAnimSequence>(Asset))
+	{
+		if (USkeletalMesh* Mesh = AnimSequence->GetPreviewMesh())
+		{
+			Result.Add(Mesh);
+		}
+
+		if (USkeleton* Skeleton = AnimSequence->GetSkeleton())
+		{
+			Result.Add(Skeleton);
+		}
+	}
+	else if (UTexture* Texture = Cast<UTexture>(Asset))
+	{
+		// Do nothing. Textures have no additional dependencies
+	}
+	else if (USkeleton* Skeleton = Cast<USkeleton>(Asset))
+	{
+		// Do nothing. Skeletons have no additional dependencies
+	}
+	else if (UPhysicsAsset* PhysicsAsset = Cast<UPhysicsAsset>(Asset))
+	{
+		// Do nothing. PhysicsAssets have no additional dependencies
+	}
+	else if (UAnimBlueprint* AnimBP = Cast<UAnimBlueprint>(Asset))
+	{
+		// Do nothing. AnimBlueprints have no additional dependencies
+	}
+	else
+	{
+		UE_LOG(LogUsd, Warning, TEXT("Unknown asset '%s' encountered when collecting dependencies."), Asset ? *Asset->GetName() : TEXT("nullptr"));
+	}
+
+	return Result;
+}
+
 class FUsdClassesModule : public IUsdClassesModule
 {
 
diff --git a/Engine/Plugins/Importers/USDImporter/Source/USDClasses/Public/USDAssetCache.h b/Engine/Plugins/Importers/USDImporter/Source/USDClasses/Public/USDAssetCache.h
--- a/Engine/Plugins/Importers/USDImporter/Source/USDClasses/Public/USDAssetCache.h
+++ b/Engine/Plugins/Importers/USDImporter/Source/USDClasses/Public/USDAssetCache.h
@@ -7,8 +7,10 @@
 
 #include "USDAssetCache.generated.h"
 
+class UE_DEPRECATED(5.2, "Use the UUsdAssetCache2 class instead, which is an actual asset you can create on the Content Browser") UUsdAssetCache;
+
 /** Owns the assets generated and reused by the USD stage, allowing thread-safe retrieval/storage */
-UCLASS()
+UCLASS(meta=(ScriptName="UsdAssetCache_Deprecated"))
 class USDCLASSES_API UUsdAssetCache : public UObject
 {
 	GENERATED_BODY()
diff --git a/Engine/Plugins/Importers/USDImporter/Source/USDClasses/Public/USDAssetCache2.h b/Engine/Plugins/Importers/USDImporter/Source/USDClasses/Public/USDAssetCache2.h
new file mode 100644
--- /dev/null
+++ b/Engine/Plugins/Importers/USDImporter/Source/USDClasses/Public/USDAssetCache2.h
@@ -0,0 +1,269 @@
+// Copyright Epic Games, Inc. All Rights Reserved.
+
+#pragma once
+
+#include "Containers/LruCache.h"
+#include "Containers/Map.h"
+#include "CoreMinimal.h"
+#include "Serialization/BulkData.h"
+#include "UObject/ObjectKey.h"
+
+#include "USDAssetCache2.generated.h"
+
+/**
+ * Owns the assets generated and reused by USD Stages, allowing thread-safe retrieval/storage.
+ */
+UCLASS(BlueprintType, Blueprintable, meta=(DisplayName="USD Asset Cache", ScriptName="UsdAssetCache"))
+class USDCLASSES_API UUsdAssetCache2 : public UObject
+{
+	GENERATED_BODY()
+
+	class FCachedAssetInfo;
+
+public:
+	/**
+	 * The asset cache will always retain all currently used assets.
+	 * In addition to that, this limit specifies how much size is allocated to storing assets that remain only for the
+	 * current session and that aren't being used by any stage.
+	 * Set this to 0 to disable storing unreferenced assets.
+	 */
+	UPROPERTY(EditAnywhere, BlueprintReadWrite, Category = "Settings")
+	double UnreferencedAssetStorageSizeMB = 500;
+
+	/**
+	 * This limit specifies how much size is allocated to storing all persistent assets (i.e. assets that will be
+	 * saved to disk, even if unused by any stage).
+	 * Set this to 0 to disable persistent asset storage.
+	 * This has no effect on temporary asset caches (e.g. the ones automatically generated when opening a stage
+	 * without an asset cache asset assigned)
+	 */
+	UPROPERTY(EditAnywhere, BlueprintReadWrite, Category = "Settings")
+	double PersistentAssetStorageSizeMB = 500;
+
+public:
+	/** Adds an asset to the cache attached to a particular hash, and optionally registering a referencer */
+	UFUNCTION(BlueprintCallable, Category = "Caching", meta = (CallInEditor = "true"))
+	void CacheAsset(const FString& Hash, UObject* Asset, const UObject* Referencer = nullptr);
+
+	/**
+	 * Returns true if the asset with the given hash can be removed from the cache. It will return false in case the
+	 * asset is still being used, either by another dependent asset or directly by some referencer.
+	 */
+	UFUNCTION(BlueprintCallable, Category = "Caching", meta = (CallInEditor = "true"))
+	bool CanRemoveAsset(const FString& Hash);
+
+	/**
+	 * If an asset is associated with `Hash`, it will be returned and the asset cache will stop tracking this asset
+	 * entirely. Returns nullptr otherwise. See CanRemoveAsset.
+	 *
+	 * WARNING: The asset will still be outer'd to the asset cache, however. The caller is in charge of properly
+	 * placing the asset at a new Outer object.
+	 */
+	UFUNCTION(BlueprintCallable, Category = "Caching", meta = (CallInEditor = "true"))
+	UObject* RemoveAsset(const FString& Hash);
+
+	/**
+	 * Returns an asset associated with a particular `Hash`.
+	 * If the asset is persistent, unloaded and the "USD.OnDemandCachedAssetLoading" cvar is true, this may cause the
+	 * asset to be read from the asset cache's file on disk.
+	 */
+	UFUNCTION(BlueprintCallable, Category = "Caching", meta = (CallInEditor = "true"))
+	UObject* GetCachedAsset(const FString& Hash);
+
+	/**
+	 * Adds a new UObject referencer to a particular asset, returning true if the operation succeeded.
+	 *
+	 * Assets will not be evicted or removed from the cache while the referencer is registered.
+	 * Note that internally the cache keeps FObjectKey structs constructed from the referencers, instead of direct
+	 * pointers to them.
+	 */
+	UFUNCTION(BlueprintCallable, Category = "Caching", meta = (CallInEditor = "true"))
+	bool AddAssetReference(const UObject* Asset, const UObject* Referencer);
+
+	/**
+	 * Removes an UObject referencer from a particular asset, returning true if the operation succeeded.
+	 */
+	UFUNCTION(BlueprintCallable, Category = "Caching", meta = (CallInEditor = "true"))
+	bool RemoveAssetReference(const UObject* Asset, const UObject* Referencer);
+
+	/**
+	 * Removes the particular referencer to all assets tracked by the cache, if it was a referencer to any of them.
+	 * Returns true if at least one asset referencer count was altered by this.
+	 */
+	UFUNCTION(BlueprintCallable, Category = "Caching", meta = (CallInEditor = "true"))
+	bool RemoveAllAssetReferences(const UObject* Referencer);
+
+	/**
+	 * Returns the hash associated with an asset, in case we own it. Returns the empty string otherwise.
+	 * Note: This has O(1) time complexity.
+	 */
+	UFUNCTION(BlueprintCallable, Category = "USD", meta = (CallInEditor = "true"))
+	FString GetHashForAsset(const UObject* Asset) const;
+
+	/**
+	 * Returns true in case the asset at `AssetPath` is tracked by the cache in any way (persistent asset,
+	 * unreferenced or referenced). An example AssetPath would be "/Game/MyTextures/RedBrick.RedBrick".
+	 */
+	UFUNCTION(BlueprintCallable, Category = "USD", meta = (CallInEditor = "true"))
+	bool IsAssetOwnedByCache(const FString& AssetPath) const;
+
+	/**
+	 * Returns how many assets are tracked by the asset cache in total (summing up persistent, referenced and
+	 * unreferenced storage)
+	 */
+	UFUNCTION(BlueprintCallable, Category = "USD", meta = (CallInEditor = "true"))
+	int32 GetNumAssets() const;
+
+	/**
+	 * Returns all asset hashes tracked by the asset cache, for all storage types. This includes unloaded
+	 * persistent assets
+	 */
+	UFUNCTION(BlueprintCallable, Category = "USD", meta = (CallInEditor = "true"))
+	TArray<FString> GetAllAssetHashes() const;
+
+	/**
+	 * Returns all assets that are currently loaded in the asset cache.
+	 * This will not include persistent assets that haven't been loaded yet.
+	 */
+	UFUNCTION(BlueprintCallable, Category = "USD", meta = (CallInEditor = "true"))
+	TArray<UObject*> GetAllLoadedAssets() const;
+
+	/**
+	 * Returns all asset paths tracked by the asset cache, for all storage types. (e.g.
+	 * "/Game/MyTextures/RedBrick.RedBrick"). This includes unloaded persistent assets
+	 */
+	UFUNCTION(BlueprintCallable, Category = "USD", meta = (CallInEditor = "true"))
+	TArray<FString> GetAllCachedAssetPaths() const;
+
+	/** Discards all tracked assets across all storage types */
+	UFUNCTION(BlueprintCallable, Category = "Settings", meta = (CallInEditor = "true"))
+	void Reset();
+
+	/**
+	 * Updates which assets belong to each storage type. You must call this in case you perform direct operations on
+	 * the asset cache, after those operations are fully complete.
+	 */
+	UFUNCTION(BlueprintCallable, Category = "Settings", meta = (CallInEditor = "true"))
+	void RefreshStorage();
+
+public:
+	UUsdAssetCache2();
+
+	/**
+	 * Every time an asset is retrieved/inserted we add it to ActiveAssets.
+	 * With these functions that container can be reset/returned.
+	 */
+	void MarkAssetsAsStale();
+
+	/** Returns assets that aren't marked as stale */
+	TSet<UObject*> GetActiveAssets() const;
+
+	// Begin UObject interfacte
+#if WITH_EDITOR
+	virtual void PostEditChangeProperty(struct FPropertyChangedEvent& PropertyChangedEvent) override;
+#endif
+	virtual void Serialize(FArchive& Ar) override;
+	virtual void PostLoad() override;
+	// End UObject interfacte
+
+private:
+	// Tries unloading unreferenced, persistent assets back to bulkdata on disk if we can
+	bool TryUnloadAsset(FCachedAssetInfo& InOutInfo);
+	bool IsAssetOwnedByCacheInternal(const FString& AssetPath) const;
+
+private:
+	friend struct FUsdScopedAssetCacheReferencer;
+
+	// Describes the different "types of storage" we allow on the asset cache.
+	// Note: Order is very important here. The rule of thumb is that it should always be OK for an asset in a storage
+	// of lower value to have dependencies on assets in storage of higher value (e.g. it's OK for a material in
+	// referenced storage to depend on textures in persistent storage, but *not* the other way around!)
+	enum class ECacheStorageType : uint8
+	{
+		None = 0,
+		Unreferenced,
+		Referenced,
+		Persistent
+	};
+	friend FArchive& operator<<(FArchive& Ar, UUsdAssetCache2::ECacheStorageType& Type);
+
+	// Describes all information we know about an asset in storage.
+	// Have to put definition here because of how TLruCache is header-only
+	class FCachedAssetInfo
+	{
+	public:
+		// We also store the asset hash in here as it is handy and allows O(1) asset to hash mapping via the LRUCache
+		FString Hash;
+
+		// Where the cached asset is actually serialized into
+		FByteBulkData BulkData;
+
+		// Information we need separately in order to deserialize the UObject from bulkdata
+		FSoftObjectPath AssetClassName;
+		FString AssetName;
+		EObjectFlags AssetFlags;
+
+		// Information fetched on-demand during RefreshStorage, but that persists to disk.
+		// Note that "on disk" here really means "on the byte buffer", although the idea is the same
+		uint64 SizeOnDiskInBytes = 0;
+		uint64 SizeOnMemoryInBytes = 0;
+		TSet<FSoftObjectPath> Dependencies;
+		TSet<FSoftObjectPath> Dependents;
+
+		// Transient stuff
+		ECacheStorageType CurrentStorageType = ECacheStorageType::None;  // Only used internally during RefreshStorage()
+		TSet<FObjectKey> Referencers;
+
+	public:
+		void Serialize(FArchive& Ar, UObject* Owner);
+	};
+
+	// Main hash to asset storage for all assets that we're currently using and shouldn't be GC'd
+	UPROPERTY(Transient, VisibleAnywhere, Category = "Owned Assets")
+	TMap<FString, TObjectPtr<UObject>> AssetStorage;
+
+	// Storage of paths to additional assets that we know are persisted on disk (and tracked on LRUCache),
+	// but that we haven't loaded from bulkdata in this session yet
+	TMap<FString, FSoftObjectPath> PendingPersistentStorage;
+
+	// Main storage of assets in most recent to least recently used.
+	// This also doubles as an inverse map from object to hash, and as a quick way of checking if we own an asset
+	TLruCache<FSoftObjectPath, FCachedAssetInfo> LRUCache;
+
+	// Assets that were added/retrieved since the last call to MarkAssetsAsSlate();
+	mutable TSet<UObject*> ActiveAssets;
+
+	// Used to ensure thread safety
+	mutable FRWLock RWLock;
+
+	// When this is set to something, we will track that it is referencing any new asset that we cache.
+	// See FUsdScopedAssetCacheReferencer just below
+	const UObject* CurrentScopedReferencer = nullptr;
+};
+
+/**
+ * The UUsdAssetCache2 can track all the UObjects that are referencing assets, so that it knows when to discard an
+ * unreferenced asset.
+ *
+ * This struct lets you specify a single UObject that will be automatically used as the referencer for the given asset
+ * cache for the duration of the scope.
+ *
+ * Usage:
+ *     FUsdScopedAssetCacheReferencer ScopedReferencer{ StageActor->UsdAssetCache, StageActor };
+ *     StageActor->SetRootLayer(MyRootLayer);  // Adds assets to the cache
+ */
+struct USDCLASSES_API FUsdScopedAssetCacheReferencer
+{
+public:
+	explicit FUsdScopedAssetCacheReferencer(UUsdAssetCache2* InAssetCache, const UObject* Referencer);
+	~FUsdScopedAssetCacheReferencer();
+
+	FUsdScopedAssetCacheReferencer() = delete;
+	FUsdScopedAssetCacheReferencer(const FUsdScopedAssetCacheReferencer&) = delete;
+	FUsdScopedAssetCacheReferencer(FUsdScopedAssetCacheReferencer&&) = delete;
+	FUsdScopedAssetCacheReferencer& operator=(const FUsdScopedAssetCacheReferencer&) = delete;
+	FUsdScopedAssetCacheReferencer& operator=(FUsdScopedAssetCacheReferencer&&) = delete;
+
+private:
+	TWeakObjectPtr<UUsdAssetCache2> AssetCache;
+};
\ No newline at end of file
diff --git a/Engine/Plugins/Importers/USDImporter/Source/USDClasses/Public/USDClassesModule.h b/Engine/Plugins/Importers/USDImporter/Source/USDClasses/Public/USDClassesModule.h
--- a/Engine/Plugins/Importers/USDImporter/Source/USDClasses/Public/USDClassesModule.h
+++ b/Engine/Plugins/Importers/USDImporter/Source/USDClasses/Public/USDClassesModule.h
@@ -40,4 +40,10 @@ public:
 
 	/** Returns a world that could be suitably described as "the current world". (e.g. when in PIE, the PIE world) */
 	USDCLASSES_API static UWorld* GetCurrentWorld( bool bEditorWorldsOnly = false );
+
+	/**
+	 * Returns the set of assets that this object depends on (e.g. when given a material, will return its textures.
+	 * When given a mesh, will return materials, etc.)
+	 */
+	USDCLASSES_API static TSet<UObject*> GetAssetDependencies(UObject* Asset);
 };
\ No newline at end of file
diff --git a/Engine/Plugins/Importers/USDImporter/Source/USDClasses/Public/USDProjectSettings.h b/Engine/Plugins/Importers/USDImporter/Source/USDClasses/Public/USDProjectSettings.h
--- a/Engine/Plugins/Importers/USDImporter/Source/USDClasses/Public/USDProjectSettings.h
+++ b/Engine/Plugins/Importers/USDImporter/Source/USDClasses/Public/USDProjectSettings.h
@@ -32,6 +32,16 @@ public:
 	UPROPERTY( config, EditAnywhere, Category = USD )
 	TArray<FName> AdditionalMaterialPurposes;
 
+	/**
+	 * USD Asset Cache to use for USD Stage Actors that don't have any asset cache specified.
+	 * Leave this empty to have each stage actor generate it's on transient cache instead.
+	 */
+	UPROPERTY(config, EditAnywhere, Category = USD, meta = (AllowedClasses = "/Script/USDClasses.UsdAssetCache2"))
+	FSoftObjectPath DefaultAssetCache;
+
+	UPROPERTY(config, EditAnywhere, Category = "USD|Dialogs")
+	bool bShowCreateDefaultAssetCacheDialog = true;
+
 	UPROPERTY( config, EditAnywhere, Category = "USD|Dialogs" )
 	bool bShowConfirmationWhenClearingLayers = true;
 
diff --git a/Engine/Plugins/Importers/USDImporter/Source/USDClasses/USDClasses.Build.cs b/Engine/Plugins/Importers/USDImporter/Source/USDClasses/USDClasses.Build.cs
--- a/Engine/Plugins/Importers/USDImporter/Source/USDClasses/USDClasses.Build.cs
+++ b/Engine/Plugins/Importers/USDImporter/Source/USDClasses/USDClasses.Build.cs
@@ -8,7 +8,7 @@ namespace UnrealBuildTool.Rules
 	public class USDClasses : ModuleRules
 	{
 		public USDClasses(ReadOnlyTargetRules Target) : base(Target)
-        {
+		{
 			PublicDependencyModuleNames.AddRange(
 				new string[]
 				{
@@ -25,7 +25,7 @@ namespace UnrealBuildTool.Rules
 				new string[]
 				{
 					"DeveloperSettings",
-					"InputCore",
+					"GeometryCache",
 					"Json", // To read/write plugInfo.json files from UnrealUSDWrapper.cpp
 				}
 			);
diff --git a/Engine/Plugins/Importers/USDImporter/Source/USDClassesEditor/Private/USDAssetCacheAssetActions.cpp b/Engine/Plugins/Importers/USDImporter/Source/USDClassesEditor/Private/USDAssetCacheAssetActions.cpp
new file mode 100644
--- /dev/null
+++ b/Engine/Plugins/Importers/USDImporter/Source/USDClassesEditor/Private/USDAssetCacheAssetActions.cpp
@@ -0,0 +1,46 @@
+// Copyright Epic Games, Inc. All Rights Reserved.
+
+#include "USDAssetCacheAssetActions.h"
+
+#include "USDAssetCache2.h"
+#include "USDAssetCacheAssetEditorToolkit.h"
+
+#define LOCTEXT_NAMESPACE "AssetCacheAssetActions"
+
+uint32 FUsdAssetCacheAssetActions::GetCategories()
+{
+	return EAssetTypeCategories::Misc;
+}
+
+FText FUsdAssetCacheAssetActions::GetName() const
+{
+	return NSLOCTEXT("AssetTypeActions", "AssetTypeActions_AssetCache", "USD Asset Cache");
+}
+
+UClass* FUsdAssetCacheAssetActions::GetSupportedClass() const
+{
+	return UUsdAssetCache2::StaticClass();
+}
+
+FColor FUsdAssetCacheAssetActions::GetTypeColor() const
+{
+	return FColor(32, 145, 208);
+}
+
+void FUsdAssetCacheAssetActions::OpenAssetEditor(const TArray<UObject*>& InObjects, TSharedPtr<class IToolkitHost> EditWithinLevelEditor)
+{
+	EToolkitMode::Type Mode = EditWithinLevelEditor.IsValid()
+		? EToolkitMode::WorldCentric
+		: EToolkitMode::Standalone;
+
+	for (auto ObjIt = InObjects.CreateConstIterator(); ObjIt; ++ObjIt)
+	{
+		if (UUsdAssetCache2* Asset = Cast<UUsdAssetCache2>(*ObjIt))
+		{
+			TSharedRef<FUsdAssetCacheAssetEditorToolkit> NewCustomAssetEditor = MakeShared<FUsdAssetCacheAssetEditorToolkit>();
+			NewCustomAssetEditor->Initialize(Mode, EditWithinLevelEditor, Asset);
+		}
+	}
+}
+
+#undef LOCTEXT_NAMESPACE
diff --git a/Engine/Plugins/Importers/USDImporter/Source/USDClassesEditor/Private/USDAssetCacheAssetActions.h b/Engine/Plugins/Importers/USDImporter/Source/USDClassesEditor/Private/USDAssetCacheAssetActions.h
new file mode 100644
--- /dev/null
+++ b/Engine/Plugins/Importers/USDImporter/Source/USDClassesEditor/Private/USDAssetCacheAssetActions.h
@@ -0,0 +1,19 @@
+// Copyright Epic Games, Inc. All Rights Reserved.
+
+#pragma once
+
+#include "AssetTypeActions_Base.h"
+#include "CoreMinimal.h"
+#include "Toolkits/IToolkitHost.h"
+
+class FUsdAssetCacheAssetActions : public FAssetTypeActions_Base
+{
+public:
+	// IAssetTypeActions interface
+	virtual uint32 GetCategories() override;
+	virtual FText GetName() const override;
+	virtual UClass* GetSupportedClass() const override;
+	virtual FColor GetTypeColor() const override;
+	virtual void OpenAssetEditor(const TArray<UObject*>& InObjects, TSharedPtr<class IToolkitHost> EditWithinLevelEditor = TSharedPtr<IToolkitHost>()) override;
+	//~ End IAssetTypeActions interface
+};
\ No newline at end of file
diff --git a/Engine/Plugins/Importers/USDImporter/Source/USDClassesEditor/Private/USDAssetCacheAssetEditorToolkit.cpp b/Engine/Plugins/Importers/USDImporter/Source/USDClassesEditor/Private/USDAssetCacheAssetEditorToolkit.cpp
new file mode 100644
--- /dev/null
+++ b/Engine/Plugins/Importers/USDImporter/Source/USDClassesEditor/Private/USDAssetCacheAssetEditorToolkit.cpp
@@ -0,0 +1,125 @@
+// Copyright Epic Games, Inc. All Rights Reserved.
+
+#include "USDAssetCacheAssetEditorToolkit.h"
+
+#include "USDAssetCache2.h"
+
+#include "Modules/ModuleManager.h"
+#include "PropertyEditorModule.h"
+#include "Widgets/Docking/SDockTab.h"
+
+#define LOCTEXT_NAMESPACE "FUsdAssetCacheAssetEditorToolkit"
+
+const FName FUsdAssetCacheAssetEditorToolkit::TabId("AssetCacheEditor");
+
+void FUsdAssetCacheAssetEditorToolkit::Initialize(const EToolkitMode::Type Mode, const TSharedPtr<IToolkitHost>& InitToolkitHost, UUsdAssetCache2* InAssetCache)
+{
+	const TSharedRef<FTabManager::FLayout> StandaloneDefaultLayout = FTabManager::NewLayout("Standalone_AssetCacheEditor")
+		->AddArea
+		(
+			FTabManager::NewPrimaryArea()
+			->SetOrientation(Orient_Vertical)
+			->Split(FTabManager::NewSplitter()
+				->SetOrientation(Orient_Horizontal)
+				->Split
+				(
+					FTabManager::NewStack()
+					->SetHideTabWell(true)
+					->AddTab(TabId, ETabState::OpenedTab)
+				)
+			)
+		);
+
+	AssetCache = InAssetCache;
+
+	// Just use a property details view for now
+	FPropertyEditorModule& PropertyEditorModule = FModuleManager::GetModuleChecked<FPropertyEditorModule>("PropertyEditor");
+	FDetailsViewArgs DetailsViewArgs;
+	DetailsViewArgs.bAllowSearch = false;
+	DetailsViewArgs.NameAreaSettings = FDetailsViewArgs::HideNameArea;
+	AssetCacheEditorWidget = PropertyEditorModule.CreateDetailView(DetailsViewArgs);
+	AssetCacheEditorWidget->SetObject(InAssetCache);
+
+	const bool bUseSmallIcons = true;
+	const bool bToolbarFocusable = false;
+	const bool bCreateDefaultToolbar = true;
+	const bool bCreateDefaultStandaloneMenu = true;
+	FAssetEditorToolkit::InitAssetEditor(
+		Mode,
+		InitToolkitHost,
+		TabId,
+		StandaloneDefaultLayout,
+		bCreateDefaultStandaloneMenu,
+		bCreateDefaultToolbar,
+		AssetCache,
+		bToolbarFocusable,
+		bUseSmallIcons
+	);
+}
+
+void FUsdAssetCacheAssetEditorToolkit::RegisterTabSpawners(const TSharedRef<class FTabManager>& InTabManager)
+{
+	FAssetEditorToolkit::RegisterTabSpawners(InTabManager);
+
+	WorkspaceMenuCategory = InTabManager->AddLocalWorkspaceMenuCategory(GetBaseToolkitName());
+
+	const FSlateIcon LayersIcon(FAppStyle::GetAppStyleSetName(), "LevelEditor.Tabs.USDStage");
+
+	InTabManager->RegisterTabSpawner(TabId, FOnSpawnTab::CreateSP(this, &FUsdAssetCacheAssetEditorToolkit::SpawnTab))
+		.SetDisplayName(GetBaseToolkitName())
+		.SetGroup(WorkspaceMenuCategory.ToSharedRef())
+		.SetIcon(LayersIcon);
+}
+
+void FUsdAssetCacheAssetEditorToolkit::UnregisterTabSpawners(const TSharedRef<FTabManager>& InTabManager)
+{
+	FAssetEditorToolkit::UnregisterTabSpawners(InTabManager);
+	InTabManager->UnregisterTabSpawner(TabId);
+}
+
+TSharedRef<SDockTab> FUsdAssetCacheAssetEditorToolkit::SpawnTab(const FSpawnTabArgs& Args)
+{
+	check(Args.GetTabId().TabType == TabId);
+
+	TSharedRef<SDockTab> NewDockTab = SNew(SDockTab)
+		.Label(GetBaseToolkitName())
+		.TabColorScale(GetWorldCentricTabColorScale())
+		[
+			AssetCacheEditorWidget.ToSharedRef()
+		];
+
+	return NewDockTab;
+}
+
+void FUsdAssetCacheAssetEditorToolkit::AddReferencedObjects(FReferenceCollector& Collector)
+{
+	Collector.AddReferencedObject(AssetCache);
+}
+
+FString FUsdAssetCacheAssetEditorToolkit::GetReferencerName() const
+{
+	return TEXT("FUsdAssetCacheAssetEditorToolkit");
+}
+
+FText FUsdAssetCacheAssetEditorToolkit::GetBaseToolkitName() const
+{
+	return LOCTEXT("AppLabel", "USD Asset Cache Editor");
+}
+
+FName FUsdAssetCacheAssetEditorToolkit::GetToolkitFName() const
+{
+	static FName Name("USD Asset Cache Editor");
+	return Name;
+}
+
+FLinearColor FUsdAssetCacheAssetEditorToolkit::GetWorldCentricTabColorScale() const
+{
+	return FLinearColor{FColor(32, 145, 208)};
+}
+
+FString FUsdAssetCacheAssetEditorToolkit::GetWorldCentricTabPrefix() const
+{
+	return LOCTEXT("WorldCentricTabPrefix", "USD Asset Cache ").ToString();
+}
+
+#undef LOCTEXT_NAMESPACE
diff --git a/Engine/Plugins/Importers/USDImporter/Source/USDClassesEditor/Private/USDAssetCacheAssetEditorToolkit.h b/Engine/Plugins/Importers/USDImporter/Source/USDClassesEditor/Private/USDAssetCacheAssetEditorToolkit.h
new file mode 100644
--- /dev/null
+++ b/Engine/Plugins/Importers/USDImporter/Source/USDClassesEditor/Private/USDAssetCacheAssetEditorToolkit.h
@@ -0,0 +1,39 @@
+// Copyright Epic Games, Inc. All Rights Reserved.
+
+#pragma once
+
+#include "Templates/SharedPointer.h"
+#include "Toolkits/AssetEditorToolkit.h"
+
+class UUsdAssetCache2;
+
+class FUsdAssetCacheAssetEditorToolkit : public FAssetEditorToolkit, public FGCObject
+{
+public:
+	void Initialize(
+		const EToolkitMode::Type Mode,
+		const TSharedPtr<IToolkitHost>& InitToolkitHost,
+		UUsdAssetCache2* InAssetCache
+	);
+
+private:
+	virtual FText GetBaseToolkitName() const override;
+	virtual FName GetToolkitFName() const override;
+	virtual FLinearColor GetWorldCentricTabColorScale() const override;
+	virtual FString GetWorldCentricTabPrefix() const override;
+
+	virtual void RegisterTabSpawners(const TSharedRef<FTabManager>& InTabManager) override;
+	virtual void UnregisterTabSpawners(const TSharedRef<FTabManager>& InTabManager) override;
+
+	TSharedRef<SDockTab> SpawnTab(const FSpawnTabArgs& Args);
+
+	virtual void AddReferencedObjects(FReferenceCollector& Collector) override;
+	virtual FString GetReferencerName() const override;
+
+private:
+	UUsdAssetCache2* AssetCache;
+
+	TSharedPtr<class IDetailsView> AssetCacheEditorWidget;
+
+	static const FName TabId;
+};
diff --git a/Engine/Plugins/Importers/USDImporter/Source/USDClassesEditor/Private/USDAssetCacheFactory.cpp b/Engine/Plugins/Importers/USDImporter/Source/USDClassesEditor/Private/USDAssetCacheFactory.cpp
new file mode 100644
--- /dev/null
+++ b/Engine/Plugins/Importers/USDImporter/Source/USDClassesEditor/Private/USDAssetCacheFactory.cpp
@@ -0,0 +1,23 @@
+// Copyright Epic Games, Inc. All Rights Reserved.
+
+#include "USDAssetCacheFactory.h"
+
+#include "USDAssetCache2.h"
+
+UUsdAssetCacheFactory::UUsdAssetCacheFactory(const FObjectInitializer& ObjectInitializer)
+	: Super(ObjectInitializer)
+{
+	bCreateNew = true;
+	bEditAfterNew = true;
+	SupportedClass = UUsdAssetCache2::StaticClass();
+}
+
+UObject* UUsdAssetCacheFactory::FactoryCreateNew(UClass* Class, UObject* InParent, FName Name, EObjectFlags Flags, UObject* Context, FFeedbackContext* Warn)
+{
+	return NewObject<UUsdAssetCache2>(InParent, Name, Flags | RF_Transactional | RF_Public | RF_Standalone);
+}
+
+bool UUsdAssetCacheFactory::ShouldShowInNewMenu() const
+{
+	return true;
+}
\ No newline at end of file
diff --git a/Engine/Plugins/Importers/USDImporter/Source/USDClassesEditor/Private/USDAssetCacheFactory.h b/Engine/Plugins/Importers/USDImporter/Source/USDClassesEditor/Private/USDAssetCacheFactory.h
new file mode 100644
--- /dev/null
+++ b/Engine/Plugins/Importers/USDImporter/Source/USDClassesEditor/Private/USDAssetCacheFactory.h
@@ -0,0 +1,21 @@
+// Copyright Epic Games, Inc. All Rights Reserved.
+
+#pragma once
+
+#include "CoreMinimal.h"
+#include "Factories/Factory.h"
+#include "UObject/ObjectMacros.h"
+
+#include "USDAssetCacheFactory.generated.h"
+
+UCLASS(hidecategories=Object)
+class UUsdAssetCacheFactory : public UFactory
+{
+	GENERATED_UCLASS_BODY()
+
+public:
+	// UFactory Interface
+	virtual UObject* FactoryCreateNew(UClass* Class, UObject* InParent, FName Name, EObjectFlags Flags, UObject* Context, FFeedbackContext* Warn) override;
+	virtual bool ShouldShowInNewMenu() const override;
+    //~ End UFactory Interface
+};
\ No newline at end of file
diff --git a/Engine/Plugins/Importers/USDImporter/Source/USDClassesEditor/Private/USDClassesEditorModule.cpp b/Engine/Plugins/Importers/USDImporter/Source/USDClassesEditor/Private/USDClassesEditorModule.cpp
new file mode 100644
--- /dev/null
+++ b/Engine/Plugins/Importers/USDImporter/Source/USDClassesEditor/Private/USDClassesEditorModule.cpp
@@ -0,0 +1,80 @@
+// Copyright Epic Games, Inc. All Rights Reserved.
+
+#include "USDClassesEditorModule.h"
+
+#include "USDAssetCache2.h"
+#include "USDAssetCacheAssetActions.h"
+#include "USDDefaultAssetCacheDialog.h"
+
+#include "AssetToolsModule.h"
+#include "Modules/ModuleManager.h"
+#include "Framework/Application/SlateApplication.h"
+#include "Interfaces/IMainFrameModule.h"
+
+#define LOCTEXT_NAMESPACE "USDClassesEditorModule"
+
+UUsdAssetCache2* IUsdClassesEditorModule::ShowMissingDefaultAssetCacheDialog()
+{
+	TSharedPtr<SWindow> ParentWindow;
+
+	if (FModuleManager::Get().IsModuleLoaded("MainFrame"))
+	{
+		IMainFrameModule& MainFrame = FModuleManager::LoadModuleChecked<IMainFrameModule>("MainFrame");
+		ParentWindow = MainFrame.GetParentWindow();
+	}
+
+	FText WindowTitle = LOCTEXT("WindowTitle", "Set the default USD Asset Cache");
+
+	TSharedRef<SWindow> Window = SNew(SWindow)
+		.Title(WindowTitle)
+		.SizingRule(ESizingRule::Autosized)
+		.AdjustInitialSizeAndPositionForDPIScale(false);
+
+	TSharedPtr<SUsdDefaultAssetCacheDialog> OptionsWindow;
+	Window->SetContent
+	(
+		SAssignNew(OptionsWindow, SUsdDefaultAssetCacheDialog)
+		.WidgetWindow(Window)
+	);
+
+	const bool bSlowTaskWindow = false;
+	FSlateApplication::Get().AddModalWindow(Window, ParentWindow, bSlowTaskWindow);
+
+	UUsdAssetCache2* NewCache = nullptr;
+	if (OptionsWindow->UserAccepted())
+	{
+		NewCache = OptionsWindow->GetCreatedCache();
+	}
+
+	return NewCache;
+}
+
+class FUsdClassesEditorModule : public IUsdClassesEditorModule
+{
+public:
+	virtual void StartupModule() override
+	{
+		// Register asset actions for the AssetCache asset
+		IAssetTools& AssetTools = FModuleManager::LoadModuleChecked<FAssetToolsModule>(TEXT("AssetTools")).Get();
+		AssetCacheAssetActions = MakeShared<FUsdAssetCacheAssetActions>();
+		AssetTools.RegisterAssetTypeActions(AssetCacheAssetActions.ToSharedRef());
+	}
+
+	virtual void ShutdownModule() override
+	{
+		// Unregister asset actions for the AssetCache asset
+		if (FAssetToolsModule* AssetToolsModule = FModuleManager::GetModulePtr<FAssetToolsModule>(TEXT("AssetTools")))
+		{
+			IAssetTools& AssetTools = AssetToolsModule->Get();
+			AssetTools.UnregisterAssetTypeActions(AssetCacheAssetActions.ToSharedRef());
+		}
+	}
+
+private:
+	TSharedPtr<IAssetTypeActions> AssetCacheAssetActions;
+};
+
+IMPLEMENT_MODULE( FUsdClassesEditorModule, USDClassesEditor );
+
+#undef LOCTEXT_NAMESPACE
+
diff --git a/Engine/Plugins/Importers/USDImporter/Source/USDClassesEditor/Private/USDDefaultAssetCacheDialog.cpp b/Engine/Plugins/Importers/USDImporter/Source/USDClassesEditor/Private/USDDefaultAssetCacheDialog.cpp
new file mode 100644
--- /dev/null
+++ b/Engine/Plugins/Importers/USDImporter/Source/USDClassesEditor/Private/USDDefaultAssetCacheDialog.cpp
@@ -0,0 +1,287 @@
+// Copyright Epic Games, Inc. All Rights Reserved.
+
+#include "USDDefaultAssetCacheDialog.h"
+
+#include "USDAssetCache2.h"
+#include "USDAssetCacheFactory.h"
+#include "USDProjectSettings.h"
+
+#include "AssetToolsModule.h"
+#include "ContentBrowserModule.h"
+#include "CoreGlobals.h"
+#include "CoreMinimal.h"
+#include "Editor.h"
+#include "EditorDirectories.h"
+#include "HAL/PlatformApplicationMisc.h"
+#include "IContentBrowserSingleton.h"
+#include "Misc/ConfigCacheIni.h"
+#include "Misc/ScopedSlowTask.h"
+#include "Modules/ModuleManager.h"
+#include "PropertyEditorModule.h"
+#include "UObject/GCObjectScopeGuard.h"
+#include "Widgets/Images/SImage.h"
+#include "Widgets/Input/SButton.h"
+#include "Widgets/Input/SCheckBox.h"
+#include "Widgets/Input/SSearchBox.h"
+#include "Widgets/Layout/SBox.h"
+#include "Widgets/Layout/SSplitter.h"
+#include "Widgets/Layout/SUniformGridPanel.h"
+#include "Widgets/SBoxPanel.h"
+#include "Widgets/SWindow.h"
+#include "Widgets/Text/STextBlock.h"
+
+#define LOCTEXT_NAMESPACE "USDDefaultAssetCacheDialog"
+
+namespace UE::DefaultCacheDialog::Private
+{
+	UUsdAssetCache2* CreateNewAssetCacheWithDialog()
+	{
+		UUsdAssetCache2* Result = nullptr;
+
+		IAssetTools& AssetTools = FModuleManager::LoadModuleChecked<FAssetToolsModule>(TEXT("AssetTools")).Get();
+		UUsdAssetCacheFactory* Factory = NewObject<UUsdAssetCacheFactory>();
+
+		// Copied from UAssetToolsImpl::CreateAssetWithDialog, but modified so that we can set our own dialog title
+		if (Factory)
+		{
+			// Determine the starting path. Try to use the most recently used directory
+			FString AssetPath;
+
+			const FString DefaultFilesystemDirectory = FEditorDirectories::Get().GetLastDirectory(ELastDirectory::NEW_ASSET);
+			if (DefaultFilesystemDirectory.IsEmpty() || !FPackageName::TryConvertFilenameToLongPackageName(DefaultFilesystemDirectory, AssetPath))
+			{
+				// No saved path, just use the game content root
+				AssetPath = TEXT("/Game");
+			}
+
+			FString PackageName;
+			FString AssetName;
+			AssetTools.CreateUniqueAssetName(AssetPath / Factory->GetDefaultNewAssetName(), TEXT(""), PackageName, AssetName);
+
+			FGCObjectScopeGuard DontGCFactory(Factory);
+
+			FSaveAssetDialogConfig SaveAssetDialogConfig;
+			SaveAssetDialogConfig.DialogTitleOverride = LOCTEXT("CreateAssetCacheDialogTitle", "Save a new USD Asset Cache");
+			SaveAssetDialogConfig.DefaultPath = AssetPath;
+			SaveAssetDialogConfig.DefaultAssetName = AssetName;
+			SaveAssetDialogConfig.ExistingAssetPolicy = ESaveAssetDialogExistingAssetPolicy::AllowButWarn;
+
+			FContentBrowserModule& ContentBrowserModule = FModuleManager::LoadModuleChecked<FContentBrowserModule>("ContentBrowser");
+			FString SaveObjectPath = ContentBrowserModule.Get().CreateModalSaveAssetDialog(SaveAssetDialogConfig);
+			if (!SaveObjectPath.IsEmpty())
+			{
+				FEditorDelegates::OnConfigureNewAssetProperties.Broadcast(Factory);
+				bool bCreateAsset = Factory->ConfigureProperties();
+
+				if (bCreateAsset)
+				{
+					const FString SavePackageName = FPackageName::ObjectPathToPackageName(SaveObjectPath);
+					const FString SavePackagePath = FPaths::GetPath(SavePackageName);
+					const FString SaveAssetName = FPaths::GetBaseFilename(SavePackageName);
+					FEditorDirectories::Get().SetLastDirectory(ELastDirectory::NEW_ASSET, SavePackagePath);
+
+					const FName CallingContext = NAME_None;
+					Result = Cast<UUsdAssetCache2>(AssetTools.CreateAsset(
+						SaveAssetName,
+						SavePackagePath,
+						UUsdAssetCache2::StaticClass(),
+						Factory,
+						CallingContext
+					));
+				}
+			}
+		}
+
+		return Result;
+	}
+}
+
+UUsdAssetCache2* SUsdDefaultAssetCacheDialog::GetCreatedCache()
+{
+	return ChosenCache;
+}
+
+void SUsdDefaultAssetCacheDialog::Construct( const FArguments& InArgs )
+{
+	ChosenCache = nullptr;
+	Window = InArgs._WidgetWindow;
+	bAccepted = false;
+
+	FSlateFontInfo MessageFont(FAppStyle::GetFontStyle("StandardDialog.LargeFont"));
+
+	ChildSlot
+	[
+		SNew(SBorder)
+		.BorderImage(FAppStyle::GetBrush("ToolPanel.GroupBorder"))
+		.Padding(16.f)
+		[
+			SNew(SVerticalBox)
+
+			+ SVerticalBox::Slot()
+			.HAlign(HAlign_Fill)
+			.VAlign(VAlign_Fill)
+			.FillHeight(1.0f)
+			.MaxHeight(550)
+			[
+				SNew(SHorizontalBox)
+				+SHorizontalBox::Slot()
+				.AutoWidth()
+				.VAlign(VAlign_Top)
+				.HAlign(HAlign_Left)
+				[
+					SNew(SImage)
+					.DesiredSizeOverride(FVector2D(24.f, 24.f))
+					.Image(FAppStyle::Get().GetBrush("Icons.WarningWithColor.Large"))
+				]
+
+				+SHorizontalBox::Slot()
+				.Padding(16.f, 0.f, 0.f, 0.f)
+				[
+					SNew(STextBlock)
+					.Text(LOCTEXT("DefaultAssetCacheDialogText", "This project has no default USD Asset Cache set.\n\nThe cache stores the assets generated when opening USD Stages, and shares them across all USD Stage Actors.\n\nDo you wish to set the default asset cache?"))
+					.Font(MessageFont)
+					.WrapTextAt(512.0f)
+				]
+			]
+
+			+SVerticalBox::Slot()
+			.AutoHeight()
+			[
+				SNew(SHorizontalBox)
+
+				+ SHorizontalBox::Slot()
+				.FillWidth(1.0f)
+				.HAlign(HAlign_Left)
+				.VAlign(VAlign_Center)
+				.Padding(FMargin(16.f, 0.f, 0.f, 0.f))
+				[
+					SNew(SCheckBox)
+					.IsChecked(ECheckBoxState::Unchecked)
+					.OnCheckStateChanged_Lambda([](ECheckBoxState NewState)
+					{
+						if(UUsdProjectSettings* ProjectSettings = GetMutableDefault<UUsdProjectSettings>())
+						{
+							ProjectSettings->bShowCreateDefaultAssetCacheDialog = NewState != ECheckBoxState::Checked;
+							ProjectSettings->SaveConfig();
+						}
+					})
+					[
+						SNew(STextBlock)
+						.WrapTextAt(615.0f)
+						.Text(LOCTEXT("DontPromptAgainText", "Don't prompt again"))
+					]
+				]
+
+				+ SHorizontalBox::Slot()
+				.AutoWidth()
+				.HAlign(HAlign_Right)
+				.VAlign(VAlign_Center)
+				[
+					SNew( SUniformGridPanel )
+					.SlotPadding( 2 )
+					.MinDesiredSlotWidth(FAppStyle::Get().GetFloat("StandardDialog.MinDesiredSlotWidth"))
+					.MinDesiredSlotHeight(FAppStyle::Get().GetFloat("StandardDialog.MinDesiredSlotHeight"))
+
+					+ SUniformGridPanel::Slot( 0, 0 )
+					[
+						SNew(SButton)
+						.HAlign(HAlign_Center)
+						.Text(LOCTEXT("UseExistingButtonText", "Pick an existing cache"))
+						.ContentPadding(FAppStyle::GetMargin("StandardDialog.ContentPadding"))
+						.ToolTipText(LOCTEXT("UseExistingButtonToolTip", "Opens a Content Browser window to pick an asset cache to use as the default.\nThe default cache can be later changed on the project settings."))
+						.OnClicked(this, &SUsdDefaultAssetCacheDialog::OnUseExisting)
+						.ButtonStyle(&FAppStyle::Get(), "PrimaryButton")
+					]
+
+					+ SUniformGridPanel::Slot( 1, 0 )
+					[
+						SNew(SButton)
+						.HAlign(HAlign_Center)
+						.Text(LOCTEXT("CreateNewButtonText", "Create new cache"))
+						.ContentPadding(FAppStyle::GetMargin("StandardDialog.ContentPadding"))
+						.ToolTipText(LOCTEXT("CreateNewButtonToolTip", "Creates a brand new asset cache to be used as the default right now, and pick where to save it.\nThe default cache can be changed on the project settings."))
+						.OnClicked(this, &SUsdDefaultAssetCacheDialog::OnCreateNew)
+						.ButtonStyle(&FAppStyle::Get(), "Button")
+					]
+
+					+ SUniformGridPanel::Slot( 2, 0 )
+					[
+						SNew(SButton)
+						.HAlign(HAlign_Center)
+						.Text(LOCTEXT("DontCreateButtonText", "Don't use a default cache"))
+						.ContentPadding(FAppStyle::GetMargin("StandardDialog.ContentPadding"))
+						.ToolTipText(LOCTEXT("DontCreateButtonToolTip", "Will use a temporary cache individual to this USD Stage Actor."))
+						.OnClicked(this, &SUsdDefaultAssetCacheDialog::OnDontCreate)
+						.ButtonStyle(&FAppStyle::Get(), "Button")
+					]
+				]
+			]
+		]
+	];
+}
+
+bool SUsdDefaultAssetCacheDialog::SupportsKeyboardFocus() const
+{
+	return true;
+}
+
+FReply SUsdDefaultAssetCacheDialog::OnUseExisting()
+{
+	FContentBrowserModule& ContentBrowserModule = FModuleManager::LoadModuleChecked<FContentBrowserModule>("ContentBrowser");
+
+	FOpenAssetDialogConfig Config;
+	Config.DialogTitleOverride = LOCTEXT("ExistingCachePickerDialogTitle", "Pick USD Asset Cache to use as default");
+	Config.bAllowMultipleSelection = false;
+	Config.AssetClassNames = {UUsdAssetCache2::StaticClass()->GetClassPathName()};
+
+	TArray<FAssetData> PickedAssets = ContentBrowserModule.Get().CreateModalOpenAssetDialog(Config);
+	if (PickedAssets.Num() == 1)
+	{
+		ChosenCache = Cast<UUsdAssetCache2>(PickedAssets[0].GetAsset());
+	}
+
+	bAccepted = true;
+	if (Window.IsValid())
+	{
+		Window.Pin()->RequestDestroyWindow();
+	}
+	return FReply::Handled();
+}
+
+FReply SUsdDefaultAssetCacheDialog::OnCreateNew()
+{
+	ChosenCache = UE::DefaultCacheDialog::Private::CreateNewAssetCacheWithDialog();
+
+	bAccepted = true;
+	if (Window.IsValid())
+	{
+		Window.Pin()->RequestDestroyWindow();
+	}
+	return FReply::Handled();
+}
+
+FReply SUsdDefaultAssetCacheDialog::OnDontCreate()
+{
+	bAccepted = false;
+	if (Window.IsValid())
+	{
+		Window.Pin()->RequestDestroyWindow();
+	}
+	return FReply::Handled();
+}
+
+FReply SUsdDefaultAssetCacheDialog::OnKeyDown( const FGeometry& MyGeometry, const FKeyEvent& InKeyEvent )
+{
+	if ( InKeyEvent.GetKey() == EKeys::Escape )
+	{
+		return OnDontCreate();
+	}
+	return FReply::Unhandled();
+}
+
+bool SUsdDefaultAssetCacheDialog::UserAccepted() const
+{
+	return bAccepted;
+}
+
+#undef LOCTEXT_NAMESPACE
\ No newline at end of file
diff --git a/Engine/Plugins/Importers/USDImporter/Source/USDClassesEditor/Private/USDDefaultAssetCacheDialog.h b/Engine/Plugins/Importers/USDImporter/Source/USDClassesEditor/Private/USDDefaultAssetCacheDialog.h
new file mode 100644
--- /dev/null
+++ b/Engine/Plugins/Importers/USDImporter/Source/USDClassesEditor/Private/USDDefaultAssetCacheDialog.h
@@ -0,0 +1,42 @@
+// Copyright Epic Games, Inc. All Rights Reserved.
+
+#pragma once
+
+#include "Widgets/SCompoundWidget.h"
+#include "Widgets/SWindow.h"
+
+class UUsdAssetCache2;
+
+/**
+ * Modal window used to prompt user whether they want to create a new default USD Asset Cache on-demand
+ */
+class SUsdDefaultAssetCacheDialog : public SCompoundWidget
+{
+public:
+	SLATE_BEGIN_ARGS( SUsdDefaultAssetCacheDialog ) {}
+		SLATE_ARGUMENT( TSharedPtr<SWindow>, WidgetWindow )
+	SLATE_END_ARGS()
+
+public:
+	UUsdAssetCache2* GetCreatedCache();
+	void Construct( const FArguments& InArgs );
+
+	bool UserAccepted() const;
+
+private:
+	virtual bool SupportsKeyboardFocus() const override;
+
+	FReply OnUseExisting();
+	FReply OnCreateNew();
+	FReply OnDontCreate();
+
+	virtual FReply OnKeyDown( const FGeometry& MyGeometry, const FKeyEvent& InKeyEvent ) override;
+
+private:
+	UUsdAssetCache2* ChosenCache;
+
+	TWeakPtr< SWindow > Window;
+
+	FText AcceptText;
+	bool bAccepted;
+};
\ No newline at end of file
diff --git a/Engine/Plugins/Importers/USDImporter/Source/USDClassesEditor/Public/USDClassesEditorModule.h b/Engine/Plugins/Importers/USDImporter/Source/USDClassesEditor/Public/USDClassesEditorModule.h
new file mode 100644
--- /dev/null
+++ b/Engine/Plugins/Importers/USDImporter/Source/USDClassesEditor/Public/USDClassesEditorModule.h
@@ -0,0 +1,14 @@
+// Copyright Epic Games, Inc. All Rights Reserved.
+
+#pragma once
+
+#include "Modules/ModuleInterface.h"
+
+class UUsdAssetCache2;
+
+class IUsdClassesEditorModule : public IModuleInterface
+{
+public:
+    /** Shows the dialog that asks the user to set a default asset cache for the project */
+	USDCLASSESEDITOR_API static UUsdAssetCache2* ShowMissingDefaultAssetCacheDialog();
+};
\ No newline at end of file
diff --git a/Engine/Plugins/Importers/USDImporter/Source/USDClassesEditor/USDClassesEditor.Build.cs b/Engine/Plugins/Importers/USDImporter/Source/USDClassesEditor/USDClassesEditor.Build.cs
new file mode 100644
--- /dev/null
+++ b/Engine/Plugins/Importers/USDImporter/Source/USDClassesEditor/USDClassesEditor.Build.cs
@@ -0,0 +1,34 @@
+// Copyright Epic Games, Inc. All Rights Reserved.
+
+using System;
+using System.IO;
+
+namespace UnrealBuildTool.Rules
+{
+	public class USDClassesEditor : ModuleRules
+	{
+		public USDClassesEditor(ReadOnlyTargetRules Target) : base(Target)
+		{
+			PublicDependencyModuleNames.AddRange(
+				new string[]
+				{
+					"Core",
+					"CoreUObject",
+					"Engine",
+					"Slate",
+					"SlateCore"
+				}
+			);
+
+			PrivateDependencyModuleNames.AddRange(
+				new string[]
+				{
+					"ContentBrowser",
+					"InputCore",
+					"UnrealEd",
+					"USDClasses",
+				}
+			);
+		}
+	}
+}
diff --git a/Engine/Plugins/Importers/USDImporter/Source/USDSchemas/Private/Custom/MDLUSDShadeMaterialTranslator.cpp b/Engine/Plugins/Importers/USDImporter/Source/USDSchemas/Private/Custom/MDLUSDShadeMaterialTranslator.cpp
--- a/Engine/Plugins/Importers/USDImporter/Source/USDSchemas/Private/Custom/MDLUSDShadeMaterialTranslator.cpp
+++ b/Engine/Plugins/Importers/USDImporter/Source/USDSchemas/Private/Custom/MDLUSDShadeMaterialTranslator.cpp
@@ -5,6 +5,7 @@
 
 #if USE_USD_SDK && WITH_EDITOR
 
+#include "USDAssetCache.h"
 #include "USDAssetImportData.h"
 #include "USDConversionUtils.h"
 #include "USDErrorUtils.h"
@@ -128,17 +129,17 @@ void FMdlUsdShadeMaterialTranslator::CreateAssets()
 		return ModuleName;
 	}();
 
-	if ( !MdlModuleName.IsEmpty() )
+	if (!MdlModuleName.IsEmpty() && Context->AssetCache)
 	{
 		pxr::TfToken MdlDefinitionToken;
 		SurfaceShader.GetSourceAssetSubIdentifier( &MdlDefinitionToken, MdlToken );
 
 		const FString MdlDefinitionName = UsdToUnreal::ConvertToken( MdlDefinitionToken );
 
-		const FString MdlFullname = MdlModuleName + TEXT("::") + MdlDefinitionName;
-
-		UMaterialInterface* MdlMaterial = Cast< UMaterialInterface >( Context->AssetCache->GetCachedAsset( MdlFullname ) );
+		const FString MdlFullName = MdlModuleName + TEXT("::") + MdlDefinitionName;
+		const FString MdlFullInstanceName = MdlFullName + TEXT("_Instance");
 
+		UMaterialInterface* MdlMaterial = Cast< UMaterialInterface >( Context->AssetCache->GetCachedAsset( MdlFullName ) );
 		if ( !MdlMaterial )
 		{
 			FScopedUnrealAllocs UnrealAllocs;
@@ -156,7 +157,7 @@ void FMdlUsdShadeMaterialTranslator::CreateAssets()
 				ImportData->PrimPath = PrimPath.GetString();
 				MdlMaterial->AssetImportData = ImportData;
 
-				Context->AssetCache->CacheAsset( MdlFullname, MdlMaterial );
+				Context->AssetCache->CacheAsset( MdlFullName, MdlMaterial );
 
 				UE::MDLShadeTranslatorImpl::Private::NotifyIfMaterialNeedsVirtualTextures( MdlMaterial );
 			}
@@ -168,42 +169,42 @@ void FMdlUsdShadeMaterialTranslator::CreateAssets()
 			}
 		}
 
-		UMaterialInstanceConstant* MdlMaterialInstance = nullptr;
-		if ( MdlMaterial )
+		UMaterialInstanceConstant* MdlMaterialInstance = Cast< UMaterialInstanceConstant >(Context->AssetCache->GetCachedAsset(MdlFullInstanceName));
+		if ( !MdlMaterialInstance && MdlMaterial )
 		{
-			MdlMaterialInstance = NewObject< UMaterialInstanceConstant >( GetTransientPackage() );
+			MdlMaterialInstance = NewObject< UMaterialInstanceConstant >(GetTransientPackage());
 
-			if ( MdlMaterialInstance )
-			{
-				UUsdAssetImportData* ImportData = NewObject< UUsdAssetImportData >( MdlMaterialInstance, TEXT( "USDAssetImportData" ) );
-				ImportData->PrimPath = PrimPath.GetString();
-				MdlMaterialInstance->AssetImportData = ImportData;
+			UUsdAssetImportData* ImportData = NewObject< UUsdAssetImportData >( MdlMaterialInstance, TEXT( "USDAssetImportData" ) );
+			ImportData->PrimPath = PrimPath.GetString();
+			MdlMaterialInstance->AssetImportData = ImportData;
 
-				MdlMaterialInstance->SetParentEditorOnly( MdlMaterial );
+			MdlMaterialInstance->SetParentEditorOnly( MdlMaterial );
 
-				UsdToUnreal::ConvertShadeInputsToParameters( ShadeMaterial, *MdlMaterialInstance, Context->AssetCache.Get(), *Context->RenderContext.ToString() );
+			UsdToUnreal::ConvertShadeInputsToParameters( ShadeMaterial, *MdlMaterialInstance, Context->AssetCache.Get(), *Context->RenderContext.ToString() );
 
-				// We can't blindly recreate all component render states when a level is being added, because we may end up first creating
-				// render states for some components, and UWorld::AddToWorld calls FScene::AddPrimitive which expects the component to not have
-				// primitives yet
-				FMaterialUpdateContext::EOptions::Type Options = FMaterialUpdateContext::EOptions::Default;
-				if ( Context->Level->bIsAssociatingLevel )
-				{
-					Options = ( FMaterialUpdateContext::EOptions::Type ) ( Options & ~FMaterialUpdateContext::EOptions::RecreateRenderStates );
-				}
+			// We can't blindly recreate all component render states when a level is being added, because we may end up first creating
+			// render states for some components, and UWorld::AddToWorld calls FScene::AddPrimitive which expects the component to not have
+			// primitives yet
+			FMaterialUpdateContext::EOptions::Type Options = FMaterialUpdateContext::EOptions::Default;
+			if ( Context->Level->bIsAssociatingLevel )
+			{
+				Options = ( FMaterialUpdateContext::EOptions::Type ) ( Options & ~FMaterialUpdateContext::EOptions::RecreateRenderStates );
+			}
 
-				FMaterialUpdateContext UpdateContext( Options, GMaxRHIShaderPlatform );
-				UpdateContext.AddMaterialInstance( MdlMaterialInstance );
-				MdlMaterialInstance->PreEditChange( nullptr );
-				MdlMaterialInstance->PostEditChange();
+			FMaterialUpdateContext UpdateContext( Options, GMaxRHIShaderPlatform );
+			UpdateContext.AddMaterialInstance( MdlMaterialInstance );
+			MdlMaterialInstance->PreEditChange( nullptr );
+			MdlMaterialInstance->PostEditChange();
 
-				UE::MDLShadeTranslatorImpl::Private::NotifyIfMaterialNeedsVirtualTextures( MdlMaterialInstance );
-			}
+			Context->AssetCache->CacheAsset(MdlFullInstanceName, MdlMaterialInstance);
+
+			UE::MDLShadeTranslatorImpl::Private::NotifyIfMaterialNeedsVirtualTextures( MdlMaterialInstance );
 		}
 
-		const FString PrimPathString = PrimPath.GetString();
-		Context->AssetCache->CacheAsset( PrimPathString, MdlMaterialInstance );
-		Context->AssetCache->LinkAssetToPrim( PrimPathString, MdlMaterialInstance );
+		if (Context->InfoCache)
+		{
+			Context->InfoCache->LinkAssetToPrim(PrimPath, MdlMaterialInstance);
+		}
 	}
 	else
 	{
diff --git a/Engine/Plugins/Importers/USDImporter/Source/USDSchemas/Private/MeshTranslationImpl.cpp b/Engine/Plugins/Importers/USDImporter/Source/USDSchemas/Private/MeshTranslationImpl.cpp
--- a/Engine/Plugins/Importers/USDImporter/Source/USDSchemas/Private/MeshTranslationImpl.cpp
+++ b/Engine/Plugins/Importers/USDImporter/Source/USDSchemas/Private/MeshTranslationImpl.cpp
@@ -2,10 +2,10 @@
 
 #include "MeshTranslationImpl.h"
 
-#include "UObject/Package.h"
-#include "USDAssetCache.h"
+#include "USDAssetCache2.h"
 #include "USDAssetImportData.h"
 #include "USDGeomMeshConversion.h"
+#include "USDInfoCache.h"
 #include "USDLog.h"
 #include "USDMemory.h"
 #include "USDProjectSettings.h"
@@ -17,6 +17,7 @@
 #include "Materials/MaterialInstanceDynamic.h"
 #include "Materials/MaterialInterface.h"
 #include "Misc/Paths.h"
+#include "UObject/Package.h"
 
 #if USE_USD_SDK
 
@@ -32,7 +33,8 @@
 TMap<const UsdUtils::FUsdPrimMaterialSlot*, UMaterialInterface*> MeshTranslationImpl::ResolveMaterialAssignmentInfo(
 	const pxr::UsdPrim& UsdPrim,
 	const TArray<UsdUtils::FUsdPrimMaterialAssignmentInfo>& AssignmentInfo,
-	UUsdAssetCache& AssetCache,
+	UUsdAssetCache2& AssetCache,
+	FUsdInfoCache& InfoCache,
 	EObjectFlags Flags
 )
 {
@@ -81,6 +83,9 @@ TMap<const UsdUtils::FUsdPrimMaterialSlot*, UMaterialInterface*> MeshTranslation
 							MaterialInstance = UsdUtils::CreateDisplayColorMaterialInstanceDynamic(DisplayColorDesc.GetValue());
 						}
 
+						// We can only cache transient assets
+						MaterialInstance->SetFlags(RF_Transient);
+
 						AssetCache.CacheAsset(Slot.MaterialSource, MaterialInstance);
 						Material = MaterialInstance;
 					}
@@ -90,21 +95,25 @@ TMap<const UsdUtils::FUsdPrimMaterialSlot*, UMaterialInterface*> MeshTranslation
 			}
 			case UsdUtils::EPrimAssignmentType::MaterialPrim:
 			{
+				// TODO: Instead of using a custom prim suffix, add the material to the same prim and then pick
+				// between all assets on a prim when needed.
+				// Previously the suffix was fine because PrimPathToAssets dealt with strings, now that it has an
+				// SdfPath API we can't use that. It's sort of a hack anyway
 				FString PrimPath = Slot.MaterialSource;
-				const static FString TwoSidedToken = TEXT("!twosided");
+				const static FString TwoSidedToken = TEXT("_twosided");
 				if (Slot.bMeshIsDoubleSided)
 				{
 					PrimPath += TwoSidedToken;
 				}
 
-				Material = Cast< UMaterialInterface >(AssetCache.GetAssetForPrim(PrimPath));
+				Material = Cast<UMaterialInterface>(InfoCache.GetSingleAssetForPrim(UE::FSdfPath{*PrimPath}, UMaterialInterface::StaticClass()));
 
 				// Need to create a two-sided material on-demand
 				if (!Material && Slot.bMeshIsDoubleSided)
 				{
 					// By now we parsed all materials so we must have the single-sided version of this material
 					UMaterialInstance* OneSidedMat = Cast<UMaterialInstance>(
-						AssetCache.GetAssetForPrim(Slot.MaterialSource)
+						InfoCache.GetSingleAssetForPrim(UE::FSdfPath{*Slot.MaterialSource}, UMaterialInterface::StaticClass())
 					);
 					if (!OneSidedMat)
 					{
@@ -112,66 +121,79 @@ TMap<const UsdUtils::FUsdPrimMaterialSlot*, UMaterialInterface*> MeshTranslation
 						continue;
 					}
 
-					// Important to not use GetBaseMaterial() here because if our parent is the translucent we'll
-					// get the reference UsdPreviewSurface instead, as that is also *its* reference
-					UMaterialInterface* ReferenceMaterial = OneSidedMat->Parent.Get();
-					UMaterialInterface* ReferenceMaterialTwoSided =
-						MeshTranslationImpl::GetTwoSidedVersionOfReferencePreviewSurfaceMaterial(ReferenceMaterial);
-					if (!ensure(ReferenceMaterialTwoSided && ReferenceMaterialTwoSided != ReferenceMaterial))
+					// Check if for some reason we already have a two-sided material ready due to a complex scenario
+					// related to the global cache
+					const FString OneSidedHash = AssetCache.GetHashForAsset(OneSidedMat);
+					const FString TwoSidedHash = OneSidedHash + TEXT("_TwoSided");
+					UMaterialInstance* TwoSidedMat = Cast<UMaterialInstance>(
+						AssetCache.GetCachedAsset(TwoSidedHash)
+					);
+					if (!TwoSidedMat)
 					{
-						continue;
-					}
+						// Important to not use GetBaseMaterial() here because if our parent is the translucent we'll
+						// get the reference UsdPreviewSurface instead, as that is also *its* reference
+						UMaterialInterface* ReferenceMaterial = OneSidedMat->Parent.Get();
+						UMaterialInterface* ReferenceMaterialTwoSided =
+							MeshTranslationImpl::GetTwoSidedVersionOfReferencePreviewSurfaceMaterial(ReferenceMaterial);
+						if (!ensure(ReferenceMaterialTwoSided && ReferenceMaterialTwoSided != ReferenceMaterial))
+						{
+							continue;
+						}
 
-					const FName NewInstanceName = MakeUniqueObjectName(
-						GetTransientPackage(),
-						UMaterialInstance::StaticClass(),
-						*(FPaths::GetBaseFilename(Slot.MaterialSource) + UnrealIdentifiers::TwoSidedMaterialSuffix)
-					);
+						const FName NewInstanceName = MakeUniqueObjectName(
+							GetTransientPackage(),
+							UMaterialInstance::StaticClass(),
+							*(FPaths::GetBaseFilename(Slot.MaterialSource) + UnrealIdentifiers::TwoSidedMaterialSuffix)
+						);
 
 #if WITH_EDITOR
-					UMaterialInstanceConstant* MIC = Cast<UMaterialInstanceConstant>(OneSidedMat);
-					if (GIsEditor && MIC)
-					{
-						UMaterialInstanceConstant* TwoSidedMat = NewObject<UMaterialInstanceConstant>(
-							GetTransientPackage(),
-							NewInstanceName,
-							Flags
-							);
-						if (TwoSidedMat)
+						UMaterialInstanceConstant* MIC = Cast<UMaterialInstanceConstant>(OneSidedMat);
+						if (GIsEditor && MIC)
 						{
-							UUsdAssetImportData* ImportData = NewObject< UUsdAssetImportData >(
-								TwoSidedMat,
-								TEXT("USDAssetImportData")
+							UMaterialInstanceConstant* TwoSidedMIC = NewObject<UMaterialInstanceConstant>(
+								GetTransientPackage(),
+								NewInstanceName,
+								Flags
 								);
-							ImportData->PrimPath = Slot.MaterialSource;
-							TwoSidedMat->AssetImportData = ImportData;
-						}
+							if (TwoSidedMIC)
+							{
+								UUsdAssetImportData* ImportData = NewObject< UUsdAssetImportData >(
+									TwoSidedMIC,
+									TEXT("USDAssetImportData")
+									);
+								ImportData->PrimPath = Slot.MaterialSource;
+								TwoSidedMIC->AssetImportData = ImportData;
+							}
 
-						TwoSidedMat->SetParentEditorOnly(ReferenceMaterialTwoSided);
-						TwoSidedMat->CopyMaterialUniformParametersEditorOnly(OneSidedMat);
+							TwoSidedMIC->SetParentEditorOnly(ReferenceMaterialTwoSided);
+							TwoSidedMIC->CopyMaterialUniformParametersEditorOnly(OneSidedMat);
 
-						AssetCache.CacheAsset(PrimPath, TwoSidedMat, PrimPath);
-						Material = TwoSidedMat;
-					}
-					else
+							TwoSidedMat = TwoSidedMIC;
+						}
+						else
 #endif // WITH_EDITOR
 						if (UMaterialInstanceDynamic* MID = Cast<UMaterialInstanceDynamic>(OneSidedMat))
 						{
-							UMaterialInstanceDynamic* TwoSidedMat = UMaterialInstanceDynamic::Create(
+							UMaterialInstanceDynamic* TwoSidedMID = UMaterialInstanceDynamic::Create(
 								ReferenceMaterialTwoSided,
 								GetTransientPackage(),
 								NewInstanceName
 							);
-							if (!ensure(TwoSidedMat))
+							if (!ensure(TwoSidedMID))
 							{
 								continue;
 							}
 
-							TwoSidedMat->CopyParameterOverrides(MID);
+							TwoSidedMID->CopyParameterOverrides(MID);
 
-							AssetCache.CacheAsset(PrimPath, TwoSidedMat, PrimPath);
-							Material = TwoSidedMat;
+							TwoSidedMat = TwoSidedMID;
 						}
+
+					}
+
+					AssetCache.CacheAsset(TwoSidedHash, TwoSidedMat);
+					InfoCache.LinkAssetToPrim(UE::FSdfPath{*PrimPath}, TwoSidedMat);
+					Material = TwoSidedMat;
 				}
 
 				break;
@@ -216,7 +238,8 @@ void MeshTranslationImpl::SetMaterialOverrides(
 	const pxr::UsdPrim& Prim,
 	const TArray<UMaterialInterface*>& ExistingAssignments,
 	UMeshComponent& MeshComponent,
-	UUsdAssetCache& AssetCache,
+	UUsdAssetCache2& AssetCache,
+	FUsdInfoCache& InfoCache,
 	float Time,
 	EObjectFlags Flags,
 	bool bInterpretLODs,
@@ -301,6 +324,7 @@ void MeshTranslationImpl::SetMaterialOverrides(
 		ValidPrim,
 		LODIndexToAssignments,
 		AssetCache,
+		InfoCache,
 		Flags
 	);
 
@@ -324,6 +348,17 @@ void MeshTranslationImpl::SetMaterialOverrides(
 				continue;
 			}
 
+			// If we don't even have as many existing assignments as we have overrides just stop here.
+			// This should happen often now because we'll always at least attempt at setting overrides on every
+			// component (but only ever set anything if we really need to).
+			// Previously we only attempted setting overrides in case the component didn't "own" the mesh prim,
+			// but now it is not feasible to do that given the global asset cache and how assets may have come
+			// from an entirely new stage/session.
+			if (!ExistingAssignments.IsValidIndex(StaticMeshSlotIndex))
+			{
+				continue;
+			}
+
 			UMaterialInterface* ExistingMaterial = ExistingAssignments[StaticMeshSlotIndex];
 			if (ExistingMaterial == Material)
 			{
diff --git a/Engine/Plugins/Importers/USDImporter/Source/USDSchemas/Private/MeshTranslationImpl.h b/Engine/Plugins/Importers/USDImporter/Source/USDSchemas/Private/MeshTranslationImpl.h
--- a/Engine/Plugins/Importers/USDImporter/Source/USDSchemas/Private/MeshTranslationImpl.h
+++ b/Engine/Plugins/Importers/USDImporter/Source/USDSchemas/Private/MeshTranslationImpl.h
@@ -15,7 +15,8 @@ PXR_NAMESPACE_OPEN_SCOPE
 	class UsdPrim;
 PXR_NAMESPACE_CLOSE_SCOPE
 
-class UUsdAssetCache;
+class FUsdInfoCache;
+class UUsdAssetCache2;
 class UMaterialInterface;
 class UMeshComponent;
 namespace UsdUtils
@@ -31,7 +32,8 @@ namespace MeshTranslationImpl
 	TMap<const UsdUtils::FUsdPrimMaterialSlot*, UMaterialInterface*> ResolveMaterialAssignmentInfo(
 		const pxr::UsdPrim& UsdPrim,
 		const TArray<UsdUtils::FUsdPrimMaterialAssignmentInfo>& AssignmentInfo,
-		UUsdAssetCache& AssetCache,
+		UUsdAssetCache2& AssetCache,
+		FUsdInfoCache& InfoCache,
 		EObjectFlags Flags
 	);
 
@@ -43,7 +45,8 @@ namespace MeshTranslationImpl
 		const pxr::UsdPrim& Prim,
 		const TArray<UMaterialInterface*>& ExistingAssignments,
 		UMeshComponent& MeshComponent,
-		UUsdAssetCache& AssetCache,
+		UUsdAssetCache2& AssetCache,
+		FUsdInfoCache& InfoCache,
 		float Time,
 		EObjectFlags Flags,
 		bool bInterpretLODs,
diff --git a/Engine/Plugins/Importers/USDImporter/Source/USDSchemas/Private/USDGeomMeshTranslator.cpp b/Engine/Plugins/Importers/USDImporter/Source/USDSchemas/Private/USDGeomMeshTranslator.cpp
--- a/Engine/Plugins/Importers/USDImporter/Source/USDSchemas/Private/USDGeomMeshTranslator.cpp
+++ b/Engine/Plugins/Importers/USDImporter/Source/USDSchemas/Private/USDGeomMeshTranslator.cpp
@@ -15,6 +15,7 @@
 #include "USDConversionUtils.h"
 #include "USDGeomMeshConversion.h"
 #include "USDGroomTranslatorUtils.h"
+#include "USDInfoCache.h"
 #include "USDIntegrationUtils.h"
 #include "USDLog.h"
 #include "USDPrimConversion.h"
@@ -224,8 +225,21 @@ namespace UsdGeomMeshTranslatorImpl
 	}
 
 	/** Returns true if material infos have changed on the StaticMesh */
-	bool ProcessMaterials(const pxr::UsdPrim& UsdPrim, const TArray<UsdUtils::FUsdPrimMaterialAssignmentInfo>& LODIndexToMaterialInfo, UStaticMesh& StaticMesh, UUsdAssetCache& AssetCache, float Time, EObjectFlags Flags)
+	bool ProcessMaterials(
+		const pxr::UsdPrim& UsdPrim,
+		const TArray<UsdUtils::FUsdPrimMaterialAssignmentInfo>& LODIndexToMaterialInfo,
+		UStaticMesh& StaticMesh,
+		UUsdAssetCache2& AssetCache,
+		FUsdInfoCache* InfoCache,
+		float Time,
+		EObjectFlags Flags
+	)
 	{
+		if (!InfoCache)
+		{
+			return false;
+		}
+
 		bool bMaterialAssignementsHaveChanged = false;
 
 		TArray<UMaterialInterface*> ExistingAssignments;
@@ -242,6 +256,7 @@ namespace UsdGeomMeshTranslatorImpl
 			UsdPrim,
 			LODIndexToMaterialInfo,
 			AssetCache,
+			*InfoCache,
 			Flags
 		);
 
@@ -320,7 +335,15 @@ namespace UsdGeomMeshTranslatorImpl
 
 #if WITH_EDITOR
 	// #ueent_todo: Merge the code with ProcessMaterials
-	bool ProcessGeometryCacheMaterials(const pxr::UsdPrim& UsdPrim, const TArray< UsdUtils::FUsdPrimMaterialAssignmentInfo >& LODIndexToMaterialInfo, UGeometryCache& GeometryCache, UUsdAssetCache& AssetCache, float Time, EObjectFlags Flags)
+	bool ProcessGeometryCacheMaterials(
+		const pxr::UsdPrim& UsdPrim,
+		const TArray< UsdUtils::FUsdPrimMaterialAssignmentInfo >& LODIndexToMaterialInfo,
+		UGeometryCache& GeometryCache,
+		UUsdAssetCache2& AssetCache,
+		FUsdInfoCache* InfoCache,
+		float Time,
+		EObjectFlags Flags
+	)
 	{
 		bool bMaterialAssignementsHaveChanged = false;
 
@@ -370,8 +393,12 @@ namespace UsdGeomMeshTranslatorImpl
 								MaterialInstance = UsdUtils::CreateDisplayColorMaterialInstanceDynamic(DisplayColorDesc.GetValue());
 							}
 
-							AssetCache.CacheAsset(Slot.MaterialSource, MaterialInstance);
-							Material = MaterialInstance;
+							if (MaterialInstance)
+							{
+								MaterialInstance->SetFlags(RF_Transient);
+								AssetCache.CacheAsset(Slot.MaterialSource, MaterialInstance);
+								Material = MaterialInstance;
+							}
 						}
 					}
 
@@ -388,12 +415,16 @@ namespace UsdGeomMeshTranslatorImpl
 
 						// TODO: This may break if MaterialPrimPath targets a prim inside a LOD variant that is disabled...
 						TUsdStore< pxr::UsdPrim > MaterialPrim = UsdPrim.GetStage()->GetPrimAtPath(MaterialPrimPath);
-						if (MaterialPrim.Get())
+						if (MaterialPrim.Get() && InfoCache)
 						{
-							Material = Cast< UMaterialInterface >(AssetCache.GetAssetForPrim(UsdToUnreal::ConvertPath(MaterialPrim.Get().GetPrimPath())));
+							Material = Cast<UMaterialInterface>(InfoCache->GetSingleAssetForPrim(
+								UE::FSdfPath{MaterialPrim.Get().GetPrimPath()}
+							));
 						}
 					}
 
+					// TODO: This stuff is missing the two-sided material upgrade code...
+
 					break;
 				}
 				case UsdUtils::EPrimAssignmentType::UnrealMaterial:
@@ -410,6 +441,7 @@ namespace UsdGeomMeshTranslatorImpl
 				}
 
 				// Fallback to this UsdGeomMesh DisplayColor material if present
+				// TODO: This code is almost an exact copy-paste of the DisplayColor case above...
 				if (Material == nullptr)
 				{
 					FScopedUsdAllocs Allocs;
@@ -440,8 +472,12 @@ namespace UsdGeomMeshTranslatorImpl
 								MaterialInstance = UsdUtils::CreateDisplayColorMaterialInstanceDynamic(DisplayColorDesc.GetValue());
 							}
 
-							AssetCache.CacheAsset(Slot.MaterialSource, MaterialInstance);
-							Material = MaterialInstance;
+							if (MaterialInstance)
+							{
+								MaterialInstance->SetFlags(RF_Transient);
+								AssetCache.CacheAsset(Slot.MaterialSource, MaterialInstance);
+								Material = MaterialInstance;
+							}
 						}
 					}
 				}
@@ -697,7 +733,8 @@ namespace UsdGeomMeshTranslatorImpl
 
 			for (const FMeshDescription& MeshDescription : LODIndexToMeshDescription)
 			{
-				FSHAHash LODHash = FStaticMeshOperations::ComputeSHAHash(MeshDescription);
+				const bool bSkipTransientAttributes = true;
+				FSHAHash LODHash = FStaticMeshOperations::ComputeSHAHash(MeshDescription, bSkipTransientAttributes);
 				SHA1.Update(&LODHash.Hash[0], sizeof(LODHash.Hash));
 
 				if (!MeshDescription.IsEmpty())
@@ -729,14 +766,17 @@ namespace UsdGeomMeshTranslatorImpl
 			SHA1.GetHash(&AllLODHash.Hash[0]);
 		}
 
-		StaticMesh = Cast< UStaticMesh >(Context.AssetCache->GetCachedAsset(AllLODHash.ToString()));
+		if (Context.AssetCache)
+		{
+			StaticMesh = Cast< UStaticMesh >(Context.AssetCache->GetCachedAsset(AllLODHash.ToString()));
+		}
 
 		if (!StaticMesh && bHasValidMeshDescription)
 		{
 			bOutIsNew = true;
 
 			FName AssetName = MakeUniqueObjectName(GetTransientPackage(), UStaticMesh::StaticClass(), *FPaths::GetBaseFilename(MeshName));
-			StaticMesh = NewObject< UStaticMesh >(GetTransientPackage(), AssetName, Context.ObjectFlags | EObjectFlags::RF_Public);
+			StaticMesh = NewObject< UStaticMesh >(GetTransientPackage(), AssetName, Context.ObjectFlags | EObjectFlags::RF_Public | EObjectFlags::RF_Transient);
 
 #if WITH_EDITOR
 			for (int32 LODIndex = 0; LODIndex < LODIndexToMeshDescription.Num(); ++LODIndex)
@@ -762,7 +802,10 @@ namespace UsdGeomMeshTranslatorImpl
 
 			StaticMesh->SetLightingGuid();
 
-			Context.AssetCache->CacheAsset(AllLODHash.ToString(), StaticMesh);
+			if (Context.AssetCache)
+			{
+				Context.AssetCache->CacheAsset(AllLODHash.ToString(), StaticMesh);
+			}
 		}
 		else
 		{
@@ -1207,7 +1250,10 @@ void FBuildStaticMeshTaskChain::SetupTasks()
 
 			if (StaticMesh)
 			{
-				Context->AssetCache->LinkAssetToPrim(PrimPathString, StaticMesh);
+				if (Context->InfoCache)
+				{
+					Context->InfoCache->LinkAssetToPrim(PrimPath, StaticMesh);
+				}
 
 #if WITH_EDITOR
 				StaticMesh->NaniteSettings.bEnabled = bShouldEnableNanite;
@@ -1220,42 +1266,31 @@ void FBuildStaticMeshTaskChain::SetupTasks()
 				}
 #endif // WITH_EDITOR
 
-
-				// Only process the materials if we own the mesh. If it's new we know we do
-#if WITH_EDITOR
-				UUsdAssetImportData* ImportData = Cast<UUsdAssetImportData>(StaticMesh->AssetImportData);
-				if (ImportData && ImportData->PrimPath == PrimPathString)
-#endif // WITH_EDITOR
+				// Only the original creator of the prim at creation time gets to set the material assignments
+				// directly on the mesh, all others prims ensure their materials via material overrides on the
+				// components
+				if (bIsNew)
 				{
-					if (!bIsNew)
-					{
-						// We may change material assignments
-						StaticMesh->Modify();
-					}
-
-					const bool bMaterialsHaveChanged = UsdGeomMeshTranslatorImpl::ProcessMaterials(
+					UsdGeomMeshTranslatorImpl::ProcessMaterials(
 						GetPrim(),
 						LODIndexToMaterialInfo,
 						*StaticMesh,
 						*Context->AssetCache.Get(),
+						Context->InfoCache.Get(),
 						Context->Time,
 						Context->ObjectFlags
 					);
 
-					if (bMaterialsHaveChanged)
-					{
-						const bool bRebuildAll = true;
-
 #if WITH_EDITOR
-						StaticMesh->UpdateUVChannelData(bRebuildAll);
+					const bool bRebuildAll = true;
+					StaticMesh->UpdateUVChannelData(bRebuildAll);
 #else
-						// UpdateUVChannelData doesn't do anything without the editor
-						for (FStaticMaterial& Material : StaticMesh->GetStaticMaterials())
-						{
-							Material.UVChannelData.bInitialized = true;
-						}
-#endif // WITH_EDITOR
+					// UpdateUVChannelData doesn't do anything without the editor
+					for (FStaticMaterial& Material : StaticMesh->GetStaticMaterials())
+					{
+						Material.UVChannelData.bInitialized = true;
 					}
+#endif // WITH_EDITOR
 				}
 			}
 
@@ -1467,23 +1502,28 @@ void FGeometryCacheCreateAssetsTaskChain::SetupTasks()
 				UUsdMeshAssetImportData* ImportData = NewObject< UUsdMeshAssetImportData >(GeometryCache, TEXT("UUSDAssetImportData"));
 				ImportData->PrimPath = PrimPathString;
 				GeometryCache->AssetImportData = ImportData;
+
+				// Only the original creator of the prim at creation time gets to set the material assignments
+				// directly on the geomcache, all others prims ensure their materials via material overrides on the
+				// components
+				UsdGeomMeshTranslatorImpl::ProcessGeometryCacheMaterials(
+					GetPrim(),
+					LODIndexToMaterialInfo,
+					*GeometryCache,
+					*Context->AssetCache.Get(),
+					Context->InfoCache.Get(),
+					Context->Time,
+					Context->ObjectFlags
+				);
 			}
 
-			bool bMaterialsHaveChanged = false;
-			if (GeometryCache)
+			if (GeometryCache && Context->InfoCache)
 			{
-				Context->AssetCache->LinkAssetToPrim(PrimPathString, GeometryCache);
-
-				// Only process the materials if we own the GeometryCache. If it's new we know we do
-				UUsdMeshAssetImportData* ImportData = Cast< UUsdMeshAssetImportData >(GeometryCache->AssetImportData);
-				if (ImportData && ImportData->PrimPath == PrimPathString)
-				{
-					bMaterialsHaveChanged = UsdGeomMeshTranslatorImpl::ProcessGeometryCacheMaterials(GetPrim(), LODIndexToMaterialInfo, *GeometryCache, *Context->AssetCache.Get(), Context->Time, Context->ObjectFlags);
-				}
+				Context->InfoCache->LinkAssetToPrim(PrimPath, GeometryCache);
 			}
 
-			const bool bContinueTaskChain = (bIsNew || bMaterialsHaveChanged);
-			return bContinueTaskChain;
+			// We don't have any additional step
+			return false;
 		});
 }
 #endif // WITH_EDITOR
@@ -1524,31 +1564,36 @@ USceneComponent* FUsdGeomMeshTranslator::CreateComponents()
 	USceneComponent* SceneComponent = CreateComponentsEx(ComponentType, {});
 	UpdateComponents(SceneComponent);
 
-	const FString PrimPathString = PrimPath.GetString();
-
 	// Handle material overrides
 	// Note: This can be here and not in USDGeomXformableTranslator because there is no way that a collapsed mesh prim could end up with a material override
 	if (UStaticMeshComponent* StaticMeshComponent = Cast<UStaticMeshComponent>(SceneComponent))
 	{
-		if (UStaticMesh* StaticMesh = Cast< UStaticMesh >(Context->AssetCache->GetAssetForPrim(PrimPathString)))
+		if (Context->InfoCache)
 		{
-			TArray<UMaterialInterface*> ExistingAssignments;
-			for (FStaticMaterial& StaticMaterial : StaticMesh->GetStaticMaterials())
+			if (UStaticMesh* StaticMesh = Cast<UStaticMesh>(Context->InfoCache->GetSingleAssetForPrim(
+				PrimPath,
+				UStaticMesh::StaticClass()
+			)))
 			{
-				ExistingAssignments.Add(StaticMaterial.MaterialInterface);
-			}
+				TArray<UMaterialInterface*> ExistingAssignments;
+				for (FStaticMaterial& StaticMaterial : StaticMesh->GetStaticMaterials())
+				{
+					ExistingAssignments.Add(StaticMaterial.MaterialInterface);
+				}
 
-			MeshTranslationImpl::SetMaterialOverrides(
-				GetPrim(),
-				ExistingAssignments,
-				*StaticMeshComponent,
-				*Context->AssetCache.Get(),
-				Context->Time,
-				Context->ObjectFlags,
-				Context->bAllowInterpretingLODs,
-				Context->RenderContext,
-				Context->MaterialPurpose
-			);
+				MeshTranslationImpl::SetMaterialOverrides(
+					GetPrim(),
+					ExistingAssignments,
+					*StaticMeshComponent,
+					*Context->AssetCache.Get(),
+					*Context->InfoCache.Get(),
+					Context->Time,
+					Context->ObjectFlags,
+					Context->bAllowInterpretingLODs,
+					Context->RenderContext,
+					Context->MaterialPurpose
+				);
+			}
 		}
 	}
 	else if (UGeometryCacheComponent* Component = Cast<UGeometryCacheComponent>(SceneComponent))
@@ -1556,40 +1601,52 @@ USceneComponent* FUsdGeomMeshTranslator::CreateComponents()
 		// We don't support importing geometry cache stuff yet so we should not be in here
 		ensure(!Context->bIsImporting);
 
-		if (UGeometryCache* GeometryCache = Cast< UGeometryCache >(Context->AssetCache->GetAssetForPrim(PrimPathString)))
+		if (Context->InfoCache && Context->AssetCache)
 		{
-			// Geometry caches don't support LODs
-			const bool bAllowInterpretingLODs = false;
-
-			MeshTranslationImpl::SetMaterialOverrides(
-				GetPrim(),
-				GeometryCache->Materials,
-				*Component,
-				*Context->AssetCache.Get(),
-				Context->Time,
-				Context->ObjectFlags,
-				bAllowInterpretingLODs,
-				Context->RenderContext,
-				Context->MaterialPurpose
-			);
+			if (UGeometryCache* GeometryCache = Cast< UGeometryCache >(Context->InfoCache->GetSingleAssetForPrim(
+				PrimPath,
+				UGeometryCache::StaticClass()
+			)))
+			{
+				// Geometry caches don't support LODs
+				const bool bAllowInterpretingLODs = false;
+
+				MeshTranslationImpl::SetMaterialOverrides(
+					GetPrim(),
+					GeometryCache->Materials,
+					*Component,
+					*Context->AssetCache,
+					*Context->InfoCache,
+					Context->Time,
+					Context->ObjectFlags,
+					bAllowInterpretingLODs,
+					Context->RenderContext,
+					Context->MaterialPurpose
+				);
 
 #if WITH_EDITOR
-			// Check if the prim has the GroomBinding schema and setup the component and assets necessary to bind the groom to the GeometryCache
-			if (UsdUtils::PrimHasSchema(GetPrim(), UnrealIdentifiers::GroomBindingAPI))
-			{
-				UsdGroomTranslatorUtils::CreateGroomBindingAsset(GetPrim(), *(Context->AssetCache), Context->ObjectFlags);
-
-				// For the groom binding to work, the GroomComponent must be a child of the SceneComponent
-				// so the Context ParentComponent is set to the SceneComponent temporarily
-				TGuardValue< USceneComponent* > ParentComponentGuard{Context->ParentComponent, SceneComponent};
-				const bool bNeedsActor = false;
-				UGroomComponent* GroomComponent = Cast< UGroomComponent >(CreateComponentsEx(TSubclassOf< USceneComponent >(UGroomComponent::StaticClass()), bNeedsActor));
-				if (GroomComponent)
+				// Check if the prim has the GroomBinding schema and setup the component and assets necessary to bind the groom to the GeometryCache
+				if (UsdUtils::PrimHasSchema(GetPrim(), UnrealIdentifiers::GroomBindingAPI))
 				{
-					UpdateComponents(SceneComponent);
+					UsdGroomTranslatorUtils::CreateGroomBindingAsset(
+						GetPrim(),
+						*Context->AssetCache,
+						*Context->InfoCache,
+						Context->ObjectFlags
+					);
+
+					// For the groom binding to work, the GroomComponent must be a child of the SceneComponent
+					// so the Context ParentComponent is set to the SceneComponent temporarily
+					TGuardValue< USceneComponent* > ParentComponentGuard{Context->ParentComponent, SceneComponent};
+					const bool bNeedsActor = false;
+					UGroomComponent* GroomComponent = Cast< UGroomComponent >(CreateComponentsEx(TSubclassOf< USceneComponent >(UGroomComponent::StaticClass()), bNeedsActor));
+					if (GroomComponent)
+					{
+						UpdateComponents(SceneComponent);
+					}
 				}
-			}
 #endif // !WITH_EDITOR
+			}
 		}
 	}
 
@@ -1623,7 +1680,14 @@ void FUsdGeomMeshTranslator::UpdateComponents(USceneComponent* SceneComponent)
 	{
 		ensure(!Context->bIsImporting);
 
-		UGeometryCache* GeometryCache = Cast< UGeometryCache >(Context->AssetCache->GetAssetForPrim(PrimPath.GetString()));
+		UGeometryCache* GeometryCache = nullptr;
+		if (Context->InfoCache)
+		{
+			GeometryCache = Cast<UGeometryCache>(Context->InfoCache->GetSingleAssetForPrim(
+				PrimPath,
+				UGeometryCache::StaticClass()
+			));
+		}
 
 		bool bShouldRegister = false;
 		if (GeometryCache != GeometryCacheUsdComponent->GetGeometryCache())
@@ -1673,7 +1737,7 @@ void FUsdGeomMeshTranslator::UpdateComponents(USceneComponent* SceneComponent)
 		// If the prim has a GroomBinding schema, apply the target groom to its associated GroomComponent
 		if (UsdUtils::PrimHasSchema(GetPrim(), UnrealIdentifiers::GroomBindingAPI))
 		{
-			UsdGroomTranslatorUtils::SetGroomFromPrim(GetPrim(), *Context->AssetCache, SceneComponent);
+			UsdGroomTranslatorUtils::SetGroomFromPrim(GetPrim(), *Context->InfoCache, SceneComponent);
 		}
 	}
 #endif // WITH_EDITOR
diff --git a/Engine/Plugins/Importers/USDImporter/Source/USDSchemas/Private/USDGeomPointInstancerTranslator.cpp b/Engine/Plugins/Importers/USDImporter/Source/USDSchemas/Private/USDGeomPointInstancerTranslator.cpp
--- a/Engine/Plugins/Importers/USDImporter/Source/USDSchemas/Private/USDGeomPointInstancerTranslator.cpp
+++ b/Engine/Plugins/Importers/USDImporter/Source/USDSchemas/Private/USDGeomPointInstancerTranslator.cpp
@@ -262,7 +262,12 @@ USceneComponent* FUsdGeomPointInstancerTranslator::CreateComponents()
 		return MainSceneComponent;
 	}
 
-	UUsdAssetCache& AssetCache = *Context->AssetCache.Get();
+	if (!Context->AssetCache.IsValid() || !Context->InfoCache.IsValid())
+	{
+		return MainSceneComponent;
+	}
+	UUsdAssetCache2& AssetCache = *Context->AssetCache.Get();
+	FUsdInfoCache& InfoCache = *Context->InfoCache.Get();
 
 	// Lets pretend ParentComponent is pointing to the parent USceneComponent while we create the child HISMs, so they get
 	// automatically attached to it as children
@@ -313,7 +318,10 @@ USceneComponent* FUsdGeomPointInstancerTranslator::CreateComponents()
 			continue;
 		}
 
-		UStaticMesh* StaticMesh = Cast< UStaticMesh >( AssetCache.GetAssetForPrim( PrototypePathStr ) );
+		UStaticMesh* StaticMesh = Cast< UStaticMesh >( InfoCache.GetSingleAssetForPrim(
+			UE::FSdfPath{*PrototypePathStr},
+			UStaticMesh::StaticClass()
+		));
 		UsdGeomPointInstancerTranslatorImpl::SetStaticMesh( StaticMesh, *HISMComponent );
 
 		// Evaluating point instancer can take a long time and is thread-safe. Move to async task while we work on something else.
@@ -335,32 +343,24 @@ USceneComponent* FUsdGeomPointInstancerTranslator::CreateComponents()
 		// Handle material overrides
 		if ( StaticMesh )
 		{
-#if WITH_EDITOR
-			// If the prim paths match, it means that it was this prim that created (and so "owns") the static mesh,
-			// so its material assignments will already be directly on the mesh. If they differ, we're using some other prim's mesh,
-			// so we may need material overrides on our component
-			UUsdAssetImportData* UsdImportData = Cast<UUsdAssetImportData>( StaticMesh->AssetImportData );
-			if ( UsdImportData && UsdImportData->PrimPath != PrototypePathStr )
-#endif // WITH_EDITOR
+			TArray<UMaterialInterface*> ExistingAssignments;
+			for ( FStaticMaterial& StaticMaterial : StaticMesh->GetStaticMaterials() )
 			{
-				TArray<UMaterialInterface*> ExistingAssignments;
-				for ( FStaticMaterial& StaticMaterial : StaticMesh->GetStaticMaterials() )
-				{
-					ExistingAssignments.Add( StaticMaterial.MaterialInterface );
-				}
-
-				MeshTranslationImpl::SetMaterialOverrides(
-					PrototypeUsdPrim,
-					ExistingAssignments,
-					*HISMComponent,
-					AssetCache,
-					Context->Time,
-					Context->ObjectFlags,
-					Context->bAllowInterpretingLODs,
-					Context->RenderContext,
-					Context->MaterialPurpose
-				);
+				ExistingAssignments.Add( StaticMaterial.MaterialInterface );
 			}
+
+			MeshTranslationImpl::SetMaterialOverrides(
+				PrototypeUsdPrim,
+				ExistingAssignments,
+				*HISMComponent,
+				AssetCache,
+				InfoCache,
+				Context->Time,
+				Context->ObjectFlags,
+				Context->bAllowInterpretingLODs,
+				Context->RenderContext,
+				Context->MaterialPurpose
+			);
 		}
 	}
 
diff --git a/Engine/Plugins/Importers/USDImporter/Source/USDSchemas/Private/USDGeomXformableTranslator.cpp b/Engine/Plugins/Importers/USDImporter/Source/USDSchemas/Private/USDGeomXformableTranslator.cpp
--- a/Engine/Plugins/Importers/USDImporter/Source/USDSchemas/Private/USDGeomXformableTranslator.cpp
+++ b/Engine/Plugins/Importers/USDImporter/Source/USDSchemas/Private/USDGeomXformableTranslator.cpp
@@ -539,7 +539,7 @@ void FUsdGeomXformableTranslator::UpdateComponents( USceneComponent* SceneCompon
 {
 	TRACE_CPUPROFILER_EVENT_SCOPE( FUsdGeomXformableTranslator::UpdateComponents );
 
-	if ( SceneComponent )
+	if (SceneComponent && Context->InfoCache)
 	{
 		SceneComponent->Modify();
 
@@ -556,7 +556,11 @@ void FUsdGeomXformableTranslator::UpdateComponents( USceneComponent* SceneCompon
 		bool bHasMultipleLODs = false;
 		if ( UStaticMeshComponent* StaticMeshComponent = Cast< UStaticMeshComponent >( SceneComponent ) )
 		{
-			UStaticMesh* PrimStaticMesh = Cast< UStaticMesh >( Context->AssetCache->GetAssetForPrim( PrimPath.GetString() ) );
+			UStaticMesh* PrimStaticMesh = Cast< UStaticMesh >( Context->InfoCache->GetSingleAssetForPrim(
+				PrimPath,
+				UStaticMesh::StaticClass()
+			));
+
 			if ( PrimStaticMesh )
 			{
 				bHasMultipleLODs = PrimStaticMesh->GetNumLODs() > 1;
diff --git a/Engine/Plugins/Importers/USDImporter/Source/USDSchemas/Private/USDGroomTranslator.cpp b/Engine/Plugins/Importers/USDImporter/Source/USDSchemas/Private/USDGroomTranslator.cpp
--- a/Engine/Plugins/Importers/USDImporter/Source/USDSchemas/Private/USDGroomTranslator.cpp
+++ b/Engine/Plugins/Importers/USDImporter/Source/USDSchemas/Private/USDGroomTranslator.cpp
@@ -192,9 +192,9 @@ protected:
 					}
 				}
 
-				if (GroomAsset)
+				if (GroomAsset && Context->InfoCache)
 				{
-					Context->AssetCache->LinkAssetToPrim(PrimPathString, GroomAsset);
+					Context->InfoCache->LinkAssetToPrim(PrimPath, GroomAsset);
 				}
 
 				// Next step is to parse the GroomCache data if it was determined that the groom has animated attributes
@@ -204,7 +204,15 @@ protected:
 		Then(ESchemaTranslationLaunchPolicy::Async,
 			[this]() -> bool
 			{
-				UGroomAsset* GroomAsset = Cast<UGroomAsset>(Context->AssetCache->GetAssetForPrim(PrimPath.GetString()));
+				if (!Context->InfoCache)
+				{
+					return false;
+				}
+
+				UGroomAsset* GroomAsset = Cast<UGroomAsset>(Context->InfoCache->GetSingleAssetForPrim(
+					PrimPath,
+					UGroomAsset::StaticClass()
+				));
 				if (!GroomAsset)
 				{
 					return false;
@@ -292,9 +300,12 @@ protected:
 						GroomCacheProcessor->AddGroomSample(MoveTemp(HairGroupsData));
 					}
 				}
-				else
+				else if (Context->InfoCache)
 				{
-					Context->AssetCache->LinkAssetToPrim(UsdGroomTranslatorUtils::GetStrandsGroomCachePrimPath(PrimPath), GroomCache);
+					Context->InfoCache->LinkAssetToPrim(
+						UE::FSdfPath{*UsdGroomTranslatorUtils::GetStrandsGroomCachePrimPath(PrimPath)},
+						GroomCache
+					);
 				}
 
 				if (!bSuccess)
@@ -309,8 +320,9 @@ protected:
 		Then(ESchemaTranslationLaunchPolicy::Sync,
 			[this]() -> bool
 			{
+				const FString StrandsGroomCachePrimPath = UsdGroomTranslatorUtils::GetStrandsGroomCachePrimPath(PrimPath);
 				FHairImportContext HairImportContext(nullptr, GetTransientPackage(), nullptr, FName(), Context->ObjectFlags | EObjectFlags::RF_Public);
-				FName UniqueName = MakeUniqueObjectName(GetTransientPackage(), UGroomCache::StaticClass(), *FPaths::GetBaseFilename(UsdGroomTranslatorUtils::GetStrandsGroomCachePrimPath(PrimPath)));
+				FName UniqueName = MakeUniqueObjectName(GetTransientPackage(), UGroomCache::StaticClass(), *FPaths::GetBaseFilename(StrandsGroomCachePrimPath));
 
 				// Once the processing has completed successfully, the data is transferred to the GroomCache
 				UGroomCache* GroomCache = FGroomCacheImporter::ProcessToGroomCache(*GroomCacheProcessor, AnimInfo, HairImportContext, UniqueName.ToString());
@@ -325,7 +337,10 @@ protected:
 #endif // WITH_EDITOR
 
 					Context->AssetCache->CacheAsset(GroomCacheHash.ToString(), GroomCache);
-					Context->AssetCache->LinkAssetToPrim(UsdGroomTranslatorUtils::GetStrandsGroomCachePrimPath(PrimPath), GroomCache);
+					Context->InfoCache->LinkAssetToPrim(
+						UE::FSdfPath{*StrandsGroomCachePrimPath},
+						GroomCache
+					);
 				}
 
 				return GroomCache != nullptr;
@@ -380,7 +395,14 @@ void FUsdGroomTranslator::UpdateComponents(USceneComponent* SceneComponent)
 	{
 		GroomComponent->Modify();
 
-		UGroomAsset* Groom = Cast<UGroomAsset>(Context->AssetCache->GetAssetForPrim(PrimPath.GetString()));
+		UGroomAsset* Groom = nullptr;
+		if(Context->InfoCache)
+		{
+			Groom = Cast<UGroomAsset>(Context->InfoCache->GetSingleAssetForPrim(
+				PrimPath,
+				UGroomAsset::StaticClass()
+			));
+		}
 
 		bool bShouldRegister = false;
 		if (Groom != GroomComponent->GroomAsset.Get())
@@ -396,7 +418,10 @@ void FUsdGroomTranslator::UpdateComponents(USceneComponent* SceneComponent)
 
 			if (Groom)
 			{
-				UGroomCache* GroomCache = Cast<UGroomCache>(Context->AssetCache->GetAssetForPrim(UsdGroomTranslatorUtils::GetStrandsGroomCachePrimPath(PrimPath)));
+				UGroomCache* GroomCache = Cast<UGroomCache>(Context->InfoCache->GetSingleAssetForPrim(
+					UE::FSdfPath{*UsdGroomTranslatorUtils::GetStrandsGroomCachePrimPath(PrimPath)},
+					UGroomCache::StaticClass()
+				));
 				if (GroomCache != GroomComponent->GroomCache.Get())
 				{
 					GroomComponent->SetGroomCache(GroomCache);
diff --git a/Engine/Plugins/Importers/USDImporter/Source/USDSchemas/Private/USDGroomTranslatorUtils.cpp b/Engine/Plugins/Importers/USDImporter/Source/USDSchemas/Private/USDGroomTranslatorUtils.cpp
--- a/Engine/Plugins/Importers/USDImporter/Source/USDSchemas/Private/USDGroomTranslatorUtils.cpp
+++ b/Engine/Plugins/Importers/USDImporter/Source/USDSchemas/Private/USDGroomTranslatorUtils.cpp
@@ -16,6 +16,7 @@
 #include "UnrealUSDWrapper.h"
 #include "USDAssetCache.h"
 #include "USDGeomXformableTranslator.h"
+#include "USDInfoCache.h"
 #include "USDIntegrationUtils.h"
 #include "USDSchemaTranslator.h"
 #include "USDTypesConversion.h"
@@ -146,7 +147,7 @@ namespace UE::UsdGroomTranslatorUtils::Private
 		return {};
 	}
 
-	UObject* GetGroomBindingSourceMesh(const pxr::UsdPrim& Prim, const UUsdAssetCache& AssetCache, EGroomBindingMeshType BindingType)
+	UObject* GetGroomBindingSourceMesh(const pxr::UsdPrim& Prim, const FUsdInfoCache& InfoCache, EGroomBindingMeshType BindingType)
 	{
 		FScopedUsdAllocs Allocs;
 
@@ -159,8 +160,7 @@ namespace UE::UsdGroomTranslatorUtils::Private
 			if (Targets.size() > 0)
 			{
 				const pxr::SdfPath& TargetPrimPath = Targets[0];
-				FString PrimPath(UsdToUnreal::ConvertPath(TargetPrimPath));
-				UObject* Mesh = AssetCache.GetAssetForPrim(PrimPath);
+				UObject* Mesh = InfoCache.GetSingleAssetForPrim(UE::FSdfPath{TargetPrimPath});
 
 				// Validate that the target prim and associated asset are of the expected type for the binding
 				pxr::UsdPrim TargetPrim = Prim.GetPrimAtPath(TargetPrimPath);
@@ -182,7 +182,12 @@ namespace UsdGroomTranslatorUtils
 {
 	using namespace UE::UsdGroomTranslatorUtils::Private;
 
-	void CreateGroomBindingAsset(const pxr::UsdPrim& Prim, UUsdAssetCache& AssetCache, EObjectFlags ObjectFlags)
+	void CreateGroomBindingAsset(
+		const pxr::UsdPrim& Prim,
+		UUsdAssetCache2& AssetCache,
+		FUsdInfoCache& InfoCache,
+		EObjectFlags ObjectFlags
+	)
 	{
 		// At this point, the prim (SkelRoot or GeomMesh) has already been checked to have the GroomBindingAPI,
 		// so we need to set up the GroomComponent and the groom binding asset to be able to bind it to the mesh
@@ -196,7 +201,10 @@ namespace UsdGroomTranslatorUtils
 		}
 
 		// The GroomAsset should already be processed and cached by the USDGroomTranslator
-		UGroomAsset* GroomAsset = Cast<UGroomAsset>(AssetCache.GetAssetForPrim(GroomPrimPath));
+		UGroomAsset* GroomAsset = Cast<UGroomAsset>(InfoCache.GetSingleAssetForPrim(
+			UE::FSdfPath{*GroomPrimPath},
+			UGroomAsset::StaticClass()
+		));
 		if (!GroomAsset)
 		{
 			return;
@@ -204,7 +212,7 @@ namespace UsdGroomTranslatorUtils
 
 		// Determine the type of binding needed based on the prim mesh type
 		const FString PrimPath(UsdToUnreal::ConvertPath(Prim.GetPath()));
-		UObject* PrimAsset = AssetCache.GetAssetForPrim(PrimPath);
+		UObject* PrimAsset = InfoCache.GetSingleAssetForPrim(UE::FSdfPath{*PrimPath});
 
 		EGroomBindingMeshType GroomBindingType = EGroomBindingMeshType::SkeletalMesh;
 		UObject* TargetMesh = Cast<USkeletalMesh>(PrimAsset);
@@ -222,7 +230,7 @@ namespace UsdGroomTranslatorUtils
 		Settings.GroomBindingType = GroomBindingType;
 		Settings.Groom = GroomAsset;
 		Settings.TargetMesh = TargetMesh;
-		Settings.SourceMesh = GetGroomBindingSourceMesh(Prim, AssetCache, GroomBindingType);
+		Settings.SourceMesh = GetGroomBindingSourceMesh(Prim, InfoCache, GroomBindingType);
 
 		// Try to get the GroomBindingAsset from the cache
 		const FSHAHash SHAHash = ComputeGroomBindingHash(Settings);
@@ -241,11 +249,11 @@ namespace UsdGroomTranslatorUtils
 
 		if (GroomBinding)
 		{
-			AssetCache.LinkAssetToPrim(GroomBindingPath, GroomBinding);
+			InfoCache.LinkAssetToPrim(UE::FSdfPath{*GroomBindingPath}, GroomBinding);
 		}
 	}
 
-	void SetGroomFromPrim(const pxr::UsdPrim& Prim, const UUsdAssetCache& AssetCache, USceneComponent* SceneComponent)
+	void SetGroomFromPrim(const pxr::UsdPrim& Prim, const FUsdInfoCache& InfoCache, USceneComponent* SceneComponent)
 	{
 		if (!SceneComponent)
 		{
@@ -260,7 +268,10 @@ namespace UsdGroomTranslatorUtils
 			return;
 		}
 
-		UGroomAsset* GroomAsset = Cast<UGroomAsset>(AssetCache.GetAssetForPrim(GroomPrimPath));
+		UGroomAsset* GroomAsset = Cast<UGroomAsset>(InfoCache.GetSingleAssetForPrim(
+			UE::FSdfPath{*GroomPrimPath},
+			UGroomAsset::StaticClass()
+		));
 		if (!GroomAsset)
 		{
 			return;
@@ -268,7 +279,10 @@ namespace UsdGroomTranslatorUtils
 
 		const FString PrimPath(UsdToUnreal::ConvertPath(Prim.GetPath()));
 		const FString GroomBindingPath = FString::Printf(TEXT("%s_groombinding"), *PrimPath);
-		UGroomBindingAsset* GroomBinding = Cast<UGroomBindingAsset>(AssetCache.GetAssetForPrim(GroomBindingPath));
+		UGroomBindingAsset* GroomBinding = Cast<UGroomBindingAsset>(InfoCache.GetSingleAssetForPrim(
+			UE::FSdfPath{*GroomBindingPath},
+			UGroomBindingAsset::StaticClass()
+		));
 
 		// Set the GroomAsset and GroomBindingAsset on the child GroomComponent of SceneComponent that was set up in the translator
 		TArray<USceneComponent*> Children;
diff --git a/Engine/Plugins/Importers/USDImporter/Source/USDSchemas/Private/USDGroomTranslatorUtils.h b/Engine/Plugins/Importers/USDImporter/Source/USDSchemas/Private/USDGroomTranslatorUtils.h
--- a/Engine/Plugins/Importers/USDImporter/Source/USDSchemas/Private/USDGroomTranslatorUtils.h
+++ b/Engine/Plugins/Importers/USDImporter/Source/USDSchemas/Private/USDGroomTranslatorUtils.h
@@ -15,8 +15,9 @@ PXR_NAMESPACE_OPEN_SCOPE
 class UsdPrim;
 PXR_NAMESPACE_CLOSE_SCOPE
 
+class FUsdInfoCache;
 class USceneComponent;
-class UUsdAssetCache;
+class UUsdAssetCache2;
 
 namespace UE
 {
@@ -26,10 +27,10 @@ namespace UE
 namespace UsdGroomTranslatorUtils
 {
 	// Create the groom binding asset for the given Prim with GroomBindingAPI and cache it in the AssetCache
-	void CreateGroomBindingAsset(const pxr::UsdPrim& Prim, UUsdAssetCache& AssetCache, EObjectFlags ObjectFlags);
+	void CreateGroomBindingAsset(const pxr::UsdPrim& Prim, UUsdAssetCache2& AssetCache, FUsdInfoCache& InfoCache, EObjectFlags ObjectFlags);
 
 	// Set the groom asset targeted by the given prim with GroomBindingAPI on the SceneComponent (with a GroomComponent as a direct child)
-	void SetGroomFromPrim(const pxr::UsdPrim& Prim, const UUsdAssetCache& AssetCache, USceneComponent* SceneComponent);
+	void SetGroomFromPrim(const pxr::UsdPrim& Prim, const FUsdInfoCache& InfoCache, USceneComponent* SceneComponent);
 
 	// Get the (pseudo) prim path for the strands GroomCache associated with a groom prim
 	FString GetStrandsGroomCachePrimPath(const UE::FSdfPath& PrimPath);
diff --git a/Engine/Plugins/Importers/USDImporter/Source/USDSchemas/Private/USDInfoCache.cpp b/Engine/Plugins/Importers/USDImporter/Source/USDSchemas/Private/USDInfoCache.cpp
--- a/Engine/Plugins/Importers/USDImporter/Source/USDSchemas/Private/USDInfoCache.cpp
+++ b/Engine/Plugins/Importers/USDImporter/Source/USDSchemas/Private/USDInfoCache.cpp
@@ -3,6 +3,7 @@
 #include "USDInfoCache.h"
 
 #include "USDGeomMeshConversion.h"
+#include "USDLog.h"
 #include "USDSchemasModule.h"
 #include "USDSchemaTranslator.h"
 #include "USDTypesConversion.h"
@@ -46,8 +47,13 @@ FArchive& operator <<( FArchive& Ar, UE::UsdInfoCache::Private::FUsdPrimInfo& In
 
 struct FUsdInfoCache::FUsdInfoCacheImpl
 {
-	mutable FRWLock Lock;
+	// Information we must have about all prims on the stage
 	TMap< UE::FSdfPath, UE::UsdInfoCache::Private::FUsdPrimInfo > InfoMap;
+	mutable FRWLock InfoMapLock;
+
+	// Information we may have about a subset of prims
+	TMap<UE::FSdfPath, TSet<TWeakObjectPtr<UObject>>> PrimPathToAssets;
+	mutable FRWLock PrimPathToAssetsLock;
 };
 
 FUsdInfoCache::FUsdInfoCache()
@@ -63,8 +69,14 @@ bool FUsdInfoCache::Serialize( FArchive& Ar )
 {
 	if ( FUsdInfoCacheImpl* ImplPtr = Impl.Get() )
 	{
-		FWriteScopeLock ScopeLock( ImplPtr->Lock );
-		Ar << ImplPtr->InfoMap;
+		{
+			FWriteScopeLock ScopeLock(ImplPtr->InfoMapLock);
+			Ar << ImplPtr->InfoMap;
+		}
+		{
+			FWriteScopeLock ScopeLock(ImplPtr->PrimPathToAssetsLock);
+			Ar << ImplPtr->PrimPathToAssets;
+		}
 	}
 
 	return true;
@@ -74,7 +86,7 @@ bool FUsdInfoCache::ContainsInfoAboutPrim( const UE::FSdfPath& Path ) const
 {
 	if ( FUsdInfoCacheImpl* ImplPtr = Impl.Get() )
 	{
-		FReadScopeLock ScopeLock( ImplPtr->Lock );
+		FReadScopeLock ScopeLock(ImplPtr->InfoMapLock);
 		return ImplPtr->InfoMap.Contains( Path );
 	}
 
@@ -85,7 +97,7 @@ bool FUsdInfoCache::IsPathCollapsed( const UE::FSdfPath& Path, ECollapsingType C
 {
 	if ( FUsdInfoCacheImpl* ImplPtr = Impl.Get() )
 	{
-		FReadScopeLock ScopeLock( ImplPtr->Lock );
+		FReadScopeLock ScopeLock(ImplPtr->InfoMapLock);
 
 		if ( const UE::UsdInfoCache::Private::FUsdPrimInfo* FoundInfo = ImplPtr->InfoMap.Find( Path ) )
 		{
@@ -108,7 +120,7 @@ bool FUsdInfoCache::DoesPathCollapseChildren( const UE::FSdfPath& Path, ECollaps
 {
 	if ( FUsdInfoCacheImpl* ImplPtr = Impl.Get() )
 	{
-		FReadScopeLock ScopeLock( ImplPtr->Lock );
+		FReadScopeLock ScopeLock(ImplPtr->InfoMapLock);
 
 		if ( const UE::UsdInfoCache::Private::FUsdPrimInfo* FoundInfo = ImplPtr->InfoMap.Find( Path ) )
 		{
@@ -132,7 +144,7 @@ UE::FSdfPath FUsdInfoCache::UnwindToNonCollapsedPath( const UE::FSdfPath& Path,
 {
 	if ( FUsdInfoCacheImpl* ImplPtr = Impl.Get() )
 	{
-		FReadScopeLock ScopeLock( ImplPtr->Lock );
+		FReadScopeLock ScopeLock(ImplPtr->InfoMapLock);
 
 		if ( const UE::UsdInfoCache::Private::FUsdPrimInfo* FoundInfo = ImplPtr->InfoMap.Find( Path ) )
 		{
@@ -219,7 +231,7 @@ namespace UE::USDInfoCacheImpl::Private
 				PrototypeVertexCounts.SetNumZeroed( PrototypePaths.size() );
 
 				{
-					FReadScopeLock ScopeLock( Impl.Lock );
+					FReadScopeLock ScopeLock(Impl.InfoMapLock);
 					for ( uint32 PrototypeIndex = 0; PrototypeIndex < PrototypePaths.size(); ++PrototypeIndex )
 					{
 						const pxr::SdfPath& PrototypePath = PrototypePaths[ PrototypeIndex ];
@@ -363,7 +375,7 @@ namespace UE::USDInfoCacheImpl::Private
 		}
 
 		{
-			FWriteScopeLock ScopeLock( Impl.Lock );
+			FWriteScopeLock ScopeLock(Impl.InfoMapLock);
 
 			UE::UsdInfoCache::Private::FUsdPrimInfo& Info = Impl.InfoMap.FindOrAdd( UE::FSdfPath{ UsdPrimPath } );
 
@@ -467,7 +479,7 @@ namespace UE::USDInfoCacheImpl::Private
 					PointInstancerMaterialSlots
 				);
 
-				FWriteScopeLock Lock{ Impl.Lock };
+				FWriteScopeLock Lock{Impl.InfoMapLock};
 				UE::UsdInfoCache::Private::FUsdPrimInfo& Info = Impl.InfoMap.FindOrAdd( UsdPointInstancerPath );
 
 				Info.ExpectedVertexCountForSubtree = PointInstancerVertexCount;
@@ -520,7 +532,7 @@ namespace UE::USDInfoCacheImpl::Private
 		bool bMergeIdenticalSlots
 	)
 	{
-		FWriteScopeLock Lock{ Impl.Lock };
+		FWriteScopeLock Lock{Impl.InfoMapLock};
 
 		if ( bMergeIdenticalSlots )
 		{
@@ -547,7 +559,7 @@ TOptional<uint64> FUsdInfoCache::GetSubtreeVertexCount( const UE::FSdfPath& Path
 {
 	if ( FUsdInfoCacheImpl* ImplPtr = Impl.Get() )
 	{
-		FReadScopeLock ScopeLock( ImplPtr->Lock );
+		FReadScopeLock ScopeLock(ImplPtr->InfoMapLock);
 
 		if ( const UE::UsdInfoCache::Private::FUsdPrimInfo* FoundInfo = ImplPtr->InfoMap.Find( Path ) )
 		{
@@ -565,7 +577,7 @@ TOptional<uint64> FUsdInfoCache::GetSubtreeMaterialSlotCount( const UE::FSdfPath
 {
 	if ( FUsdInfoCacheImpl* ImplPtr = Impl.Get() )
 	{
-		FReadScopeLock ScopeLock( ImplPtr->Lock );
+		FReadScopeLock ScopeLock(ImplPtr->InfoMapLock);
 
 		if ( const UE::UsdInfoCache::Private::FUsdPrimInfo* FoundInfo = ImplPtr->InfoMap.Find( Path ) )
 		{
@@ -579,6 +591,114 @@ TOptional<uint64> FUsdInfoCache::GetSubtreeMaterialSlotCount( const UE::FSdfPath
 	return {};
 }
 
+void FUsdInfoCache::LinkAssetToPrim(const UE::FSdfPath& Path, UObject* Asset)
+{
+	if (!Asset)
+	{
+		return;
+	}
+
+	FUsdInfoCacheImpl* ImplPtr = Impl.Get();
+	if (!ImplPtr)
+	{
+		return;
+	}
+	FWriteScopeLock ScopeLock(ImplPtr->PrimPathToAssetsLock);
+
+	UE_LOG(LogUsd, Verbose, TEXT("Linking asset '%s' to prim '%s'"),
+		*Asset->GetPathName(),
+		*Path.GetString()
+	);
+
+	ImplPtr->PrimPathToAssets.FindOrAdd(Path).Add(Asset);
+}
+
+TSet<TWeakObjectPtr<UObject>> FUsdInfoCache::RemoveAllAssetPrimLinks(const UE::FSdfPath& Path)
+{
+	FUsdInfoCacheImpl* ImplPtr = Impl.Get();
+	if (!ImplPtr)
+	{
+		return {};
+	}
+	FWriteScopeLock ScopeLock(ImplPtr->PrimPathToAssetsLock);
+
+	TSet<TWeakObjectPtr<UObject>> Values;
+	ImplPtr->PrimPathToAssets.RemoveAndCopyValue(Path, Values);
+	return Values;
+}
+
+UObject* FUsdInfoCache::GetSingleAssetForPrim(const UE::FSdfPath& Path, const UClass* FilterClass) const
+{
+	FUsdInfoCacheImpl* ImplPtr = Impl.Get();
+	if (!ImplPtr)
+	{
+		return nullptr;
+	}
+	FReadScopeLock ScopeLock(ImplPtr->PrimPathToAssetsLock);
+
+	if (const TSet<TWeakObjectPtr<UObject>>* FoundAssets = ImplPtr->PrimPathToAssets.Find(Path))
+	{
+		for (const TWeakObjectPtr<UObject>& FoundAsset : *FoundAssets)
+		{
+			UObject* FoundAssetPtr = FoundAsset.Get();
+			if (FoundAssetPtr && (!FilterClass || FoundAssetPtr->IsA(FilterClass)))
+			{
+				return FoundAssetPtr;
+			}
+		}
+	}
+
+	return nullptr;
+}
+
+TSet<TWeakObjectPtr<UObject>> FUsdInfoCache::GetAssetsForPrim(const UE::FSdfPath& Path, const UClass* FilterClass) const
+{
+	FUsdInfoCacheImpl* ImplPtr = Impl.Get();
+	if (!ImplPtr)
+	{
+		return {};
+	}
+	FReadScopeLock ScopeLock(ImplPtr->PrimPathToAssetsLock);
+
+	if (const TSet<TWeakObjectPtr<UObject>>* FoundAssets = ImplPtr->PrimPathToAssets.Find(Path))
+	{
+		return *FoundAssets;
+	}
+
+	return {};
+}
+
+UE::FSdfPath FUsdInfoCache::GetPrimForAsset(UObject* Asset) const
+{
+	FUsdInfoCacheImpl* ImplPtr = Impl.Get();
+	if (!ImplPtr)
+	{
+		return {};
+	}
+	FReadScopeLock ScopeLock(ImplPtr->PrimPathToAssetsLock);
+
+	for (const TPair<UE::FSdfPath, TSet<TWeakObjectPtr<UObject>>>& Pair : ImplPtr->PrimPathToAssets)
+	{
+		if (Pair.Value.Contains(Asset))
+		{
+			return Pair.Key;
+		}
+	}
+
+	return {};
+}
+
+TMap<UE::FSdfPath, TSet<TWeakObjectPtr<UObject>>> FUsdInfoCache::GetAllAssetPrimLinks() const
+{
+	FUsdInfoCacheImpl* ImplPtr = Impl.Get();
+	if (!ImplPtr)
+	{
+		return {};
+	}
+
+	return ImplPtr->PrimPathToAssets;
+}
+
 void FUsdInfoCache::RebuildCacheForSubtree( const UE::FUsdPrim& Prim, FUsdSchemaTranslationContext& Context )
 {
 #if USE_USD_SDK
@@ -646,7 +766,7 @@ void FUsdInfoCache::Clear()
 {
 	if ( FUsdInfoCacheImpl* ImplPtr = Impl.Get() )
 	{
-		FWriteScopeLock ScopeLock( ImplPtr->Lock );
+		FWriteScopeLock ScopeLock(ImplPtr->InfoMapLock);
 		ImplPtr->InfoMap.Empty();
 	}
 }
@@ -655,7 +775,7 @@ bool FUsdInfoCache::IsEmpty()
 {
 	if ( FUsdInfoCacheImpl* ImplPtr = Impl.Get() )
 	{
-		FReadScopeLock ScopeLock( ImplPtr->Lock );
+		FReadScopeLock ScopeLock(ImplPtr->InfoMapLock);
 		return ImplPtr->InfoMap.IsEmpty();
 	}
 
diff --git a/Engine/Plugins/Importers/USDImporter/Source/USDSchemas/Private/USDSchemaTranslator.cpp b/Engine/Plugins/Importers/USDImporter/Source/USDSchemas/Private/USDSchemaTranslator.cpp
--- a/Engine/Plugins/Importers/USDImporter/Source/USDSchemas/Private/USDSchemaTranslator.cpp
+++ b/Engine/Plugins/Importers/USDImporter/Source/USDSchemas/Private/USDSchemaTranslator.cpp
@@ -130,7 +130,7 @@ void FUsdSchemaTranslatorRegistry::Unregister( const FRegisteredSchemaTranslator
 	}
 }
 
-FUsdSchemaTranslationContext::FUsdSchemaTranslationContext( const UE::FUsdStage& InStage, UUsdAssetCache& InAssetCache )
+FUsdSchemaTranslationContext::FUsdSchemaTranslationContext( const UE::FUsdStage& InStage, UUsdAssetCache2& InAssetCache )
 	: Stage( InStage )
 	, AssetCache( &InAssetCache )
 {
diff --git a/Engine/Plugins/Importers/USDImporter/Source/USDSchemas/Private/USDShadeMaterialTranslator.cpp b/Engine/Plugins/Importers/USDImporter/Source/USDSchemas/Private/USDShadeMaterialTranslator.cpp
--- a/Engine/Plugins/Importers/USDImporter/Source/USDSchemas/Private/USDShadeMaterialTranslator.cpp
+++ b/Engine/Plugins/Importers/USDImporter/Source/USDSchemas/Private/USDShadeMaterialTranslator.cpp
@@ -129,17 +129,17 @@ namespace UE::UsdShadeTranslator::Private
 #endif // WITH_EDITOR
 						if ( UMaterialInstanceDynamic* MID = Cast<UMaterialInstanceDynamic>( UserMaterial ) )
 						{
-							if ( Context->AssetCache )
+							if ( Context->AssetCache && Context->InfoCache )
 							{
-								// This is super slow as it will essentially linear search the asset cache twice, but
+								// This is super slow as it will essentially linear search the asset cache, but
 								// not much we can do at runtime without asset import data
-								const FString PrimPath = Context->AssetCache->GetPrimForAsset( MID );
+								const UE::FSdfPath PrimPath = Context->InfoCache->GetPrimForAsset( MID );
 								const FString Hash = Context->AssetCache->GetHashForAsset( MID );
 
 								const FName NewInstanceName = MakeUniqueObjectName(
 									GetTransientPackage(),
 									UMaterialInstance::StaticClass(),
-									*FPaths::GetBaseFilename( PrimPath )
+									PrimPath.IsEmpty() ? TEXT("MaterialInstance") : *FPaths::GetBaseFilename(PrimPath.GetString())
 								);
 
 								// For MID we can't swap the reference material, so we need to create a brand new one and copy
@@ -155,9 +155,15 @@ namespace UE::UsdShadeTranslator::Private
 								}
 								NewMID->CopyParameterOverrides( MID );
 
-								Context->AssetCache->DiscardAsset( Hash );
-								Context->AssetCache->CacheAsset( Hash, NewMID, PrimPath );
-								NewMaterials.Add( NewMID );
+								if (Context->AssetCache->RemoveAsset(Hash))
+								{
+									Context->AssetCache->CacheAsset(Hash, NewMID);
+									if (!PrimPath.IsEmpty())
+									{
+										Context->InfoCache->LinkAssetToPrim(PrimPath, NewMID);
+									}
+									NewMaterials.Add( NewMID );
+								}
 							}
 						}
 						else
@@ -234,16 +240,23 @@ void FUsdShadeMaterialTranslator::CreateAssets()
 
 	if ( !ConvertedMaterial )
 	{
+		const FString PrimPathString = PrimPath.GetString();
 		const bool bIsTranslucent = UsdUtils::IsMaterialTranslucent( ShadeMaterial );
 		const FName InstanceName = MakeUniqueObjectName(
 			GetTransientPackage(),
 			UMaterialInstance::StaticClass(),
-			*FPaths::GetBaseFilename( PrimPath.GetString() )
+			*FPaths::GetBaseFilename(PrimPathString)
 		);
 
+		// TODO: There is an issue here: Note how we're only ever going to write the material prim's primvars into
+		// this PrimvarToUVIndex map when first creating the material, and not if we find it in the asset cache.
+		// This because finding the primvars to use essentially involves parsing the entire material again, so we
+		// likely shouldn't do it every time.
+		// This means that a mesh with float2 UV sets and materials that use them will likely end up with no UVs once
+		// the stage reloads...
 		TMap<FString, int32> Unused;
 		TMap<FString, int32>& PrimvarToUVIndex = Context->MaterialToPrimvarToUVIndex
-			? Context->MaterialToPrimvarToUVIndex->FindOrAdd( PrimPath.GetString() )
+			? Context->MaterialToPrimvarToUVIndex->FindOrAdd(PrimPathString)
 			: Unused;
 
 #if WITH_EDITOR
@@ -350,6 +363,8 @@ void FUsdShadeMaterialTranslator::CreateAssets()
 			{
 				if ( UMaterialInstanceDynamic* NewMaterial = UMaterialInstanceDynamic::Create( ReferenceMaterial, GetTransientPackage(), InstanceName ) )
 				{
+					NewMaterial->SetFlags(RF_Transient);
+
 					if ( UsdToUnreal::ConvertMaterial( ShadeMaterial, *NewMaterial, Context->AssetCache.Get(), PrimvarToUVIndex, *Context->RenderContext.ToString() ) )
 					{
 						TSet<UTexture*> VTTextures;
@@ -371,9 +386,10 @@ void FUsdShadeMaterialTranslator::CreateAssets()
 
 						// We must stash our material and textures *before* we call UpgradeMaterialsAndTexturesToVT, as that
 						// is what will actually swap our reference with a VT one if needed
-						if ( Context->AssetCache )
+						if ( Context->AssetCache && Context->InfoCache )
 						{
-							Context->AssetCache->CacheAsset( MaterialHashString, NewMaterial, PrimPath.GetString() );
+							Context->AssetCache->CacheAsset(MaterialHashString, NewMaterial);
+							Context->InfoCache->LinkAssetToPrim(PrimPath, NewMaterial);
 						}
 						for ( UTexture* Texture : VTTextures.Union( NonVTTextures ) )
 						{
@@ -397,18 +413,32 @@ void FUsdShadeMaterialTranslator::CreateAssets()
 			}
 		}
 	}
-	else if ( Context->MaterialToPrimvarToUVIndex && Context->AssetCache )
+	else if ( Context->MaterialToPrimvarToUVIndex && Context->InfoCache )
 	{
-		if ( TMap<FString, int32>* PrimvarToUVIndex = Context->MaterialToPrimvarToUVIndex->Find( Context->AssetCache->GetPrimForAsset( ConvertedMaterial ) ) )
+		const UE::FSdfPath FoundPrimPath = Context->InfoCache->GetPrimForAsset(ConvertedMaterial);
+		if (!FoundPrimPath.IsEmpty())
 		{
-			// Copy the Material -> Primvar -> UV index mapping from the cached material prim path to this prim path
-			Context->MaterialToPrimvarToUVIndex->FindOrAdd( PrimPath.GetString() ) = *PrimvarToUVIndex;
+			if (TMap<FString, int32>* PrimvarToUVIndex = Context->MaterialToPrimvarToUVIndex->Find(FoundPrimPath.GetString()))
+			{
+				// Copy the Material -> Primvar -> UV index mapping from the cached material prim path to this prim path
+				Context->MaterialToPrimvarToUVIndex->FindOrAdd(PrimPath.GetString()) = *PrimvarToUVIndex;
+			}
 		}
 	}
 
-	if ( ConvertedMaterial )
+	// Note that this needs to run even if we do find the material from the asset cache, otherwise we won't
+	// re-register the prim asset links when we reload a stage
+	if (ConvertedMaterial)
 	{
-		Context->AssetCache->CacheAsset( MaterialHashString, ConvertedMaterial, PrimPath.GetString() );
+		if (Context->InfoCache)
+		{
+			Context->InfoCache->LinkAssetToPrim(PrimPath, ConvertedMaterial);
+		}
+
+		if (Context->AssetCache)
+		{
+			Context->AssetCache->CacheAsset(MaterialHashString, ConvertedMaterial);
+		}
 	}
 }
 
diff --git a/Engine/Plugins/Importers/USDImporter/Source/USDSchemas/Private/USDSkelRootTranslator.cpp b/Engine/Plugins/Importers/USDImporter/Source/USDSchemas/Private/USDSkelRootTranslator.cpp
--- a/Engine/Plugins/Importers/USDImporter/Source/USDSchemas/Private/USDSkelRootTranslator.cpp
+++ b/Engine/Plugins/Importers/USDImporter/Source/USDSchemas/Private/USDSkelRootTranslator.cpp
@@ -11,6 +11,7 @@
 #include "USDErrorUtils.h"
 #include "USDGeomMeshConversion.h"
 #include "USDGroomTranslatorUtils.h"
+#include "USDInfoCache.h"
 #include "USDIntegrationUtils.h"
 #include "USDLayerUtils.h"
 #include "USDLog.h"
@@ -79,7 +80,8 @@ namespace UsdSkelRootTranslatorImpl
 		const pxr::UsdPrim& UsdPrim,
 		TArray<UsdUtils::FUsdPrimMaterialAssignmentInfo>& LODIndexToMaterialInfo,
 		USkeletalMesh* SkeletalMesh,
-		UUsdAssetCache& AssetCache,
+		UUsdAssetCache2& AssetCache,
+		FUsdInfoCache& InfoCache,
 		float Time, EObjectFlags Flags,
 		bool bSkeletalMeshHasMorphTargets
 	)
@@ -105,6 +107,7 @@ namespace UsdSkelRootTranslatorImpl
 			UsdPrim,
 			LODIndexToMaterialInfo,
 			AssetCache,
+			InfoCache,
 			Flags
 		);
 
@@ -374,7 +377,7 @@ namespace UsdSkelRootTranslatorImpl
 		TSet<FString>& InOutUsedMorphTargetNames,
 		const TMap< FString, TMap< FString, int32 > >& InMaterialToPrimvarsUVSetNames,
 		float InTime,
-		const UUsdAssetCache& AssetCache,
+		const UUsdAssetCache2& AssetCache,
 		bool bInInterpretLODs,
 		const pxr::TfToken& RenderContext,
 		const pxr::TfToken& MaterialPurpose
@@ -597,7 +600,8 @@ namespace UsdSkelRootTranslatorImpl
 		const pxr::UsdPrim& SkelRootPrim,
 		const TArray<UMaterialInterface*>& ExistingAssignments,
 		UMeshComponent& MeshComponent,
-		UUsdAssetCache& AssetCache,
+		UUsdAssetCache2& AssetCache,
+		FUsdInfoCache& InfoCache,
 		float Time,
 		EObjectFlags Flags,
 		bool bInterpretLODs,
@@ -728,6 +732,7 @@ namespace UsdSkelRootTranslatorImpl
 			ValidSkelRootPrim,
 			LODIndexToAssignments,
 			AssetCache,
+			InfoCache,
 			Flags
 		);
 
@@ -800,7 +805,7 @@ namespace UsdSkelRootTranslatorImpl
 		UPhysicsAsset* Result = NewObject< UPhysicsAsset >(
 			GetTransientPackage(),
 			AssetName,
-			Flags | EObjectFlags::RF_Public
+			Flags | EObjectFlags::RF_Public | EObjectFlags::RF_Transient
 		);
 
 		FPhysAssetCreateParams NewBodyData;
@@ -831,15 +836,17 @@ namespace UsdSkelRootTranslatorImpl
 		bool* bOutNeedsRecompile = nullptr // Whether this function wants the returned AnimBP to be recompiled
 	)
 	{
-		if ( !Prim )
+		if (!Prim || !Context.InfoCache || Context.AssetCache)
 		{
 			return nullptr;
 		}
 
-		FString PrimPath = UsdToUnreal::ConvertPath( Prim.GetPath() );
 		FString PrimName = UsdToUnreal::ConvertString( Prim.GetName() );
 
-		USkeletalMesh* SkeletalMesh = Cast< USkeletalMesh >( Context.AssetCache->GetAssetForPrim( PrimPath ) );
+		USkeletalMesh* SkeletalMesh = Cast< USkeletalMesh >(Context.InfoCache->GetSingleAssetForPrim(
+			UE::FSdfPath{Prim.GetPath()},
+			USkeletalMesh::StaticClass())
+		);
 		if ( !SkeletalMesh )
 		{
 			return nullptr;
@@ -882,7 +889,19 @@ namespace UsdSkelRootTranslatorImpl
 		const static FString DefaultAnimBPPath = TEXT( "/USDImporter/Blueprint/DefaultLiveLinkAnimBP.DefaultLiveLinkAnimBP" );
 		if ( DefaultAnimBPPath == AnimBPPath )
 		{
-			const FString CacheKey = UsdToUnreal::ConvertPath( Prim.GetPrimPath() ) + TEXT( "_DefaultAnimBlueprint" );
+			const FString PrimPath = UsdToUnreal::ConvertPath(Prim.GetPrimPath());
+
+			// Let's try to never reuse AnimBP between prims (as we want to be able to switch subject names
+			// independently and we likely won't have more than a handful of these anyway).
+			// We should have at least something deterministic though so that we don't repeatedly recreate assets for the same prim.
+			FSHAHash Hash;
+			FSHA1 SHA1;
+			// Each stage actor has a separate info cache that is assigned to the context
+			SHA1.Update(reinterpret_cast<const uint8*>(Context.InfoCache.Get()), sizeof(Context.InfoCache.Get()));
+			SHA1.UpdateWithString(*PrimPath, PrimPath.Len());
+			SHA1.Final();
+			SHA1.GetHash(&Hash.Hash[0]);
+			const FString CacheKey = Hash.ToString();
 
 			// Check if we can find an AnimBP for this prim in the asset cache (useful when doing Action->Import)
 			bool bReusedAnimBP = false;
@@ -1255,7 +1274,7 @@ namespace UsdSkelRootTranslatorImpl
 
 				if ( SkeletalMesh )
 				{
-					if ( bIsNew )
+					if (bIsNew && Context->AssetCache && Context->InfoCache)
 					{
 						UUsdMeshAssetImportData* ImportData = NewObject< UUsdMeshAssetImportData >( SkeletalMesh, TEXT( "USDAssetImportData" ) );
 						ImportData->PrimPath = SkelRootPath;
@@ -1266,6 +1285,7 @@ namespace UsdSkelRootTranslatorImpl
 							LODIndexToMaterialInfo,
 							SkeletalMesh,
 							*Context->AssetCache.Get(),
+							*Context->InfoCache.Get(),
 							Context->Time,
 							Context->ObjectFlags,
 							NewBlendShapes.Num() > 0
@@ -1305,7 +1325,10 @@ namespace UsdSkelRootTranslatorImpl
 						SkeletalMesh->SetPhysicsAsset( nullptr );
 					}
 
-					Context->AssetCache->LinkAssetToPrim( SkelRootPath, SkeletalMesh );
+					if (Context->InfoCache)
+					{
+						Context->InfoCache->LinkAssetToPrim(PrimPath, SkeletalMesh);
+					}
 
 					// Track our Skeleton by the source skeleton prim path
 					if ( USkeleton* Skeleton = SkeletalMesh->GetSkeleton() )
@@ -1316,11 +1339,14 @@ namespace UsdSkelRootTranslatorImpl
 						// be able to quickly query our skeleton again to find its path
 						std::vector< pxr::UsdSkelBinding > SkeletonBindings;
 						SkeletonCache.Get().ComputeSkelBindings( pxr::UsdSkelRoot( GetPrim() ), &SkeletonBindings, pxr::UsdTraverseInstanceProxies() );
-						if ( SkeletonBindings.size() > 0 )
+						if (SkeletonBindings.size() > 0 && Context->InfoCache)
 						{
 							pxr::UsdSkelBinding& SkeletonBinding = SkeletonBindings[ 0 ];
 							const pxr::UsdSkelSkeleton& UsdSkeleton = SkeletonBinding.GetSkeleton();
-							Context->AssetCache->LinkAssetToPrim( UsdToUnreal::ConvertPath( UsdSkeleton.GetPrim().GetPrimPath() ), Skeleton );
+							Context->InfoCache->LinkAssetToPrim(
+								UE::FSdfPath{UsdSkeleton.GetPrim().GetPrimPath()},
+								Skeleton
+							);
 						}
 					}
 
@@ -1384,12 +1410,15 @@ namespace UsdSkelRootTranslatorImpl
 		Then( ESchemaTranslationLaunchPolicy::Sync,
 			[ this ]()
 			{
-				if ( !Context->bAllowParsingSkeletalAnimations )
+				if (!Context->bAllowParsingSkeletalAnimations || !Context->InfoCache)
 				{
 					return false;
 				}
 
-				USkeletalMesh* SkeletalMesh = Cast< USkeletalMesh >( Context->AssetCache->GetAssetForPrim( PrimPath.GetString() ) );
+				USkeletalMesh* SkeletalMesh = Cast< USkeletalMesh >(Context->InfoCache->GetSingleAssetForPrim(
+					PrimPath,
+					USkeletalMesh::StaticClass()
+				));
 				if ( !SkeletalMesh )
 				{
 					return false;
@@ -1480,9 +1509,12 @@ namespace UsdSkelRootTranslatorImpl
 							}
 						}
 
-						if ( AnimSequence )
+						if (AnimSequence && Context->InfoCache)
 						{
-							Context->AssetCache->LinkAssetToPrim( SkelAnimationPrimPath, AnimSequence );
+							Context->InfoCache->LinkAssetToPrim(
+								UE::FSdfPath{*SkelAnimationPrimPath},
+								AnimSequence
+							);
 						}
 
 						// For now we shouldn't try to parse the SkelAnimations from skeletal bindings other than the first one as we only
@@ -1518,9 +1550,14 @@ USceneComponent* FUsdSkelRootTranslator::CreateComponents()
 
 #if WITH_EDITOR
 	// Check if the prim has the GroomBinding schema and setup the component and assets necessary to bind the groom to the SkeletalMesh
-	if ( UsdUtils::PrimHasSchema( GetPrim(), UnrealIdentifiers::GroomBindingAPI ) )
+	if ( UsdUtils::PrimHasSchema( GetPrim(), UnrealIdentifiers::GroomBindingAPI ) && Context->AssetCache && Context->InfoCache )
 	{
-		UsdGroomTranslatorUtils::CreateGroomBindingAsset( GetPrim(), *( Context->AssetCache ), Context->ObjectFlags );
+		UsdGroomTranslatorUtils::CreateGroomBindingAsset(
+			GetPrim(),
+			*Context->AssetCache,
+			*Context->InfoCache,
+			Context->ObjectFlags
+		);
 
 		// For the groom binding to work, the GroomComponent must be a child of the SceneComponent
 		// so the Context ParentComponent is set to the SceneComponent temporarily
@@ -1540,7 +1577,7 @@ USceneComponent* FUsdSkelRootTranslator::CreateComponents()
 void FUsdSkelRootTranslator::UpdateComponents( USceneComponent* SceneComponent )
 {
 	USkeletalMeshComponent* SkeletalMeshComponent = Cast< USkeletalMeshComponent >( SceneComponent );
-	if ( !SkeletalMeshComponent )
+	if (!SkeletalMeshComponent || !Context->InfoCache)
 	{
 		return;
 	}
@@ -1574,9 +1611,12 @@ void FUsdSkelRootTranslator::UpdateComponents( USceneComponent* SceneComponent )
 	);
 
 	UE::FUsdPrim SkelAnimPrim = UsdUtils::FindFirstAnimationSource( Prim );
-	if ( SkelAnimPrim )
+	if (SkelAnimPrim)
 	{
-		UAnimSequence* TargetAnimSequence = Cast< UAnimSequence >( Context->AssetCache->GetAssetForPrim( SkelAnimPrim.GetPrimPath().GetString() ) );
+		UAnimSequence* TargetAnimSequence = Cast< UAnimSequence >(Context->InfoCache->GetSingleAssetForPrim(
+			SkelAnimPrim.GetPrimPath(),
+			UAnimSequence::StaticClass()
+		));
 		if ( TargetAnimSequence != SkeletalMeshComponent->AnimationData.AnimToPlay )
 		{
 			SkeletalMeshComponent->AnimationData.AnimToPlay = TargetAnimSequence;
@@ -1594,7 +1634,10 @@ void FUsdSkelRootTranslator::UpdateComponents( USceneComponent* SceneComponent )
 
 #if WITH_EDITOR
 	// Re-set the skeletal mesh if we created a new one (maybe the hash changed, a skinned UsdGeomMesh was hidden, etc.)
-	USkeletalMesh* TargetSkeletalMesh = Cast< USkeletalMesh >( Context->AssetCache->GetAssetForPrim( PrimPath.GetString() ) );
+	USkeletalMesh* TargetSkeletalMesh = Cast< USkeletalMesh >(Context->InfoCache->GetSingleAssetForPrim(
+		PrimPath,
+		USkeletalMesh::StaticClass()
+	));
 	if ( SkeletalMeshComponent->GetSkeletalMeshAsset() != TargetSkeletalMesh )
 	{
 		SkeletalMeshComponent->SetSkeletalMesh(TargetSkeletalMesh);
@@ -1613,6 +1656,7 @@ void FUsdSkelRootTranslator::UpdateComponents( USceneComponent* SceneComponent )
 				ExistingAssignments,
 				*SkeletalMeshComponent,
 				*Context->AssetCache.Get(),
+				*Context->InfoCache.Get(),
 				Context->Time,
 				Context->ObjectFlags,
 				Context->bAllowInterpretingLODs,
@@ -1704,9 +1748,9 @@ void FUsdSkelRootTranslator::UpdateComponents( USceneComponent* SceneComponent )
 	}
 
 	// If the prim has a GroomBinding schema, apply the target groom to its associated GroomComponent
-	if ( UsdUtils::PrimHasSchema( Prim, UnrealIdentifiers::GroomBindingAPI ) )
+	if (UsdUtils::PrimHasSchema(Prim, UnrealIdentifiers::GroomBindingAPI) && Context->InfoCache)
 	{
-		UsdGroomTranslatorUtils::SetGroomFromPrim( Prim, *Context->AssetCache, SceneComponent );
+		UsdGroomTranslatorUtils::SetGroomFromPrim( Prim, *Context->InfoCache, SceneComponent );
 	}
 #endif // WITH_EDITOR
 }
diff --git a/Engine/Plugins/Importers/USDImporter/Source/USDSchemas/Public/USDInfoCache.h b/Engine/Plugins/Importers/USDImporter/Source/USDSchemas/Public/USDInfoCache.h
--- a/Engine/Plugins/Importers/USDImporter/Source/USDSchemas/Public/USDInfoCache.h
+++ b/Engine/Plugins/Importers/USDImporter/Source/USDSchemas/Public/USDInfoCache.h
@@ -21,38 +21,53 @@ enum class ECollapsingType
 };
 
 /**
- * Caches whether given prim paths represent prims that can have their assets or components collapsed when importing/opening
- * a USD Stage.
+ * Caches information about a specific USD Stage
  */
-struct USDSCHEMAS_API FUsdInfoCache
+class USDSCHEMAS_API FUsdInfoCache
 {
+public:
+	struct FUsdInfoCacheImpl;
+
 	FUsdInfoCache();
 	virtual ~FUsdInfoCache();
 
-	bool Serialize( FArchive& Ar );
+	bool Serialize(FArchive& Ar);
 
 	// Returns whether we contain any info about prim at 'Path' at all
-	bool ContainsInfoAboutPrim( const UE::FSdfPath& Path ) const;
+	bool ContainsInfoAboutPrim(const UE::FSdfPath& Path) const;
 
-	bool IsPathCollapsed( const UE::FSdfPath& Path, ECollapsingType CollapsingType ) const;
-	bool DoesPathCollapseChildren( const UE::FSdfPath& Path, ECollapsingType CollapsingType ) const;
+	void RebuildCacheForSubtree(const UE::FUsdPrim& Prim, FUsdSchemaTranslationContext& Context);
+
+	void Clear();
+	bool IsEmpty();
+
+public:
+	bool IsPathCollapsed(const UE::FSdfPath& Path, ECollapsingType CollapsingType) const;
+	bool DoesPathCollapseChildren(const UE::FSdfPath& Path, ECollapsingType CollapsingType) const;
 
 	// Returns Path in case it represents an uncollapsed prim, or returns the path to the prim that collapsed it
-	UE::FSdfPath UnwindToNonCollapsedPath( const UE::FSdfPath& Path, ECollapsingType CollapsingType ) const;
+	UE::FSdfPath UnwindToNonCollapsedPath(const UE::FSdfPath& Path, ECollapsingType CollapsingType) const;
 
+public:
 	// Provides the total vertex or material slots counts for each prim *and* its subtree.
 	// This is built inside RebuildCacheForSubtree, so it will factor in the used Context's bMergeIdenticalMaterialSlots.
 	// Note that these aren't affected by actual collapsing: A prim that doesn't collapse its children will still
 	// provide the total sum of vertex counts of its entire subtree when queried
-	TOptional<uint64> GetSubtreeVertexCount( const UE::FSdfPath& Path );
-	TOptional<uint64> GetSubtreeMaterialSlotCount( const UE::FSdfPath& Path );
+	TOptional<uint64> GetSubtreeVertexCount(const UE::FSdfPath& Path);
+	TOptional<uint64> GetSubtreeMaterialSlotCount(const UE::FSdfPath& Path);
 
-	void RebuildCacheForSubtree( const UE::FUsdPrim& Prim, FUsdSchemaTranslationContext& Context );
+public:
+	void LinkAssetToPrim(const UE::FSdfPath& Path, UObject* Asset);
 
-	void Clear();
-	bool IsEmpty();
+	TSet<TWeakObjectPtr<UObject>> RemoveAllAssetPrimLinks(const UE::FSdfPath& Path);
+
+	UObject* GetSingleAssetForPrim(const UE::FSdfPath& Path, const UClass* FilterClass = nullptr) const;
+	TSet<TWeakObjectPtr<UObject>> GetAssetsForPrim(const UE::FSdfPath& Path, const UClass* FilterClass = nullptr) const;
+
+	UE::FSdfPath GetPrimForAsset(UObject* Asset) const;
+	TMap<UE::FSdfPath, TSet<TWeakObjectPtr<UObject>>> GetAllAssetPrimLinks() const;
 
-	struct FUsdInfoCacheImpl;
 private:
 	TUniquePtr<FUsdInfoCacheImpl> Impl;
-};
\ No newline at end of file
+};
+
diff --git a/Engine/Plugins/Importers/USDImporter/Source/USDSchemas/Public/USDSchemaTranslator.h b/Engine/Plugins/Importers/USDImporter/Source/USDSchemas/Public/USDSchemaTranslator.h
--- a/Engine/Plugins/Importers/USDImporter/Source/USDSchemas/Public/USDSchemaTranslator.h
+++ b/Engine/Plugins/Importers/USDImporter/Source/USDSchemas/Public/USDSchemaTranslator.h
@@ -5,7 +5,7 @@
 #include "CoreMinimal.h"
 
 #include "UnrealUSDWrapper.h"
-#include "USDAssetCache.h"
+#include "USDAssetCache2.h"
 #include "USDInfoCache.h"
 #include "USDMemory.h"
 #include "USDSkeletalDataConversion.h"
@@ -128,7 +128,7 @@ protected:
 
 struct USDSCHEMAS_API FUsdSchemaTranslationContext : public TSharedFromThis< FUsdSchemaTranslationContext >
 {
-	explicit FUsdSchemaTranslationContext( const UE::FUsdStage& InStage, UUsdAssetCache& InAssetCache );
+	explicit FUsdSchemaTranslationContext( const UE::FUsdStage& InStage, UUsdAssetCache2& InAssetCache );
 
 	/** True if we're a context created by the USDStageImporter to fully import to persistent assets and actors */
 	bool bIsImporting = false;
@@ -164,7 +164,7 @@ struct USDSCHEMAS_API FUsdSchemaTranslationContext : public TSharedFromThis< FUs
 	int32 NaniteTriangleThreshold;
 
 	/** Where the translated assets will be stored */
-	TStrongObjectPtr< UUsdAssetCache > AssetCache;
+	TStrongObjectPtr< UUsdAssetCache2 > AssetCache;
 
 	/** Caches various information about prims that are expensive to query */
 	TSharedPtr<FUsdInfoCache> InfoCache;
diff --git a/Engine/Plugins/Importers/USDImporter/Source/USDStage/Private/USDLevelSequenceHelper.cpp b/Engine/Plugins/Importers/USDImporter/Source/USDStage/Private/USDLevelSequenceHelper.cpp
--- a/Engine/Plugins/Importers/USDImporter/Source/USDStage/Private/USDLevelSequenceHelper.cpp
+++ b/Engine/Plugins/Importers/USDImporter/Source/USDStage/Private/USDLevelSequenceHelper.cpp
@@ -3,10 +3,10 @@
 #include "USDLevelSequenceHelper.h"
 
 #include "GeometryCache.h"
-#include "USDAssetCache.h"
 #include "USDAssetImportData.h"
 #include "USDAttributeUtils.h"
 #include "USDConversionUtils.h"
+#include "USDInfoCache.h"
 #include "USDIntegrationUtils.h"
 #include "USDLayerUtils.h"
 #include "USDListener.h"
@@ -511,7 +511,7 @@ public:
 	~FUsdLevelSequenceHelperImpl();
 
 	ULevelSequence* Init(const UE::FUsdStage& InUsdStage);
-	void SetAssetCache(UUsdAssetCache* AssetCache);
+	void SetInfoCache(TSharedPtr<FUsdInfoCache> InfoCache);
 	bool HasData() const;
 	void Clear();
 
@@ -703,7 +703,7 @@ private:
 	FUsdLevelSequenceHelper::FOnSkelAnimationBaked OnSkelAnimationBaked;
 
 	TWeakObjectPtr<AUsdStageActor> StageActor = nullptr;
-	UUsdAssetCache* AssetCache = nullptr;  // We keep a pointer to this directly because we may be called via the USDStageImporter directly, when we don't have an available actor
+	TSharedPtr<FUsdInfoCache> InfoCache = nullptr;  // We keep a pointer to this directly because we may be called via the USDStageImporter directly, when we don't have an available actor
 	FGuid StageActorBinding;
 
 	// Only when this is zero we write LevelSequence object (tracks, moviescene, sections, etc.) transactions back to the USD stage
@@ -768,9 +768,9 @@ ULevelSequence* FUsdLevelSequenceHelperImpl::Init(const UE::FUsdStage& InUsdStag
 	return MainLevelSequence;
 }
 
-void FUsdLevelSequenceHelperImpl::SetAssetCache(UUsdAssetCache* InAssetCache)
+void FUsdLevelSequenceHelperImpl::SetInfoCache(TSharedPtr<FUsdInfoCache> InInfoCache)
 {
-	AssetCache = InAssetCache;
+	InfoCache = InInfoCache;
 }
 
 bool FUsdLevelSequenceHelperImpl::HasData() const
@@ -902,7 +902,7 @@ void FUsdLevelSequenceHelperImpl::BindToUsdStageActor(AUsdStageActor* InStageAct
 	UnbindFromUsdStageActor();
 
 	StageActor = InStageActor;
-	SetAssetCache(InStageActor ? InStageActor->GetAssetCache() : nullptr);
+	SetInfoCache(InStageActor ? InStageActor->GetInfoCache() : nullptr);
 
 	if (!StageActor.IsValid() || !MainLevelSequence || !MainLevelSequence->GetMovieScene())
 	{
@@ -949,7 +949,7 @@ void FUsdLevelSequenceHelperImpl::UnbindFromUsdStageActor()
 		StageActor.Reset();
 	}
 
-	SetAssetCache(nullptr);
+	SetInfoCache(nullptr);
 
 	OnStageEditTargetChangedHandle.Reset();
 }
@@ -1082,9 +1082,9 @@ ULevelSequence* FUsdLevelSequenceHelperImpl::FindOrAddSequenceForLayer(const UE:
 		// from the transaction buffer, so we would basically end up creating a identical new object on top of an existing one (the new object has the same address as the existing one).
 		// When importing we don't actually want to do this though, because we want these assets name to conflict so that we can publish/replace old assets if desired. The stage
 		// importer will make these names unique later if needed.
-		// We only get an AssetCache when importing (from UUsdStageImporter::ImportFromFile) or when BindToUsdStageActor is called,
+		// We only get an InfoCache when importing (from UUsdStageImporter::ImportFromFile) or when BindToUsdStageActor is called,
 		// which also gives us a stage actor. So if we don't have an actor but have a cache, we're importing
-		const bool bIsImporting = StageActor.IsExplicitlyNull() && AssetCache;
+		const bool bIsImporting = StageActor.IsExplicitlyNull() && InfoCache;
 		FName UniqueSequenceName = bIsImporting
 			? *UsdLevelSequenceHelperImpl::SanitizeObjectName(FPaths::GetBaseFilename(SequenceDisplayName))
 			: MakeUniqueObjectName(GetTransientPackage(), ULevelSequence::StaticClass(), *UsdLevelSequenceHelperImpl::SanitizeObjectName(FPaths::GetBaseFilename(SequenceDisplayName)));
@@ -1839,7 +1839,7 @@ void FUsdLevelSequenceHelperImpl::AddSkeletalTracks(const UUsdPrimTwin& PrimTwin
 		return;
 	}
 
-	if (!AssetCache)
+	if (!InfoCache)
 	{
 		return;
 	}
@@ -1855,8 +1855,11 @@ void FUsdLevelSequenceHelperImpl::AddSkeletalTracks(const UUsdPrimTwin& PrimTwin
 	// Fetch the UAnimSequence asset from the asset cache. Ideally we'd call AUsdStageActor::GetGeneratedAssets,
 	// but we may belong to a FUsdStageImportContext, and so there's no AUsdStageActor at all to use.
 	// At this point it doesn't matter much though, because we shouldn't need to uncollapse a SkelAnimation prim path anyway
-	FString PrimPath = SkelAnimationPrim.GetPrimPath().GetString();
-	UAnimSequence* Sequence = Cast<UAnimSequence>(AssetCache->GetAssetForPrim(PrimPath));
+	const UE::FSdfPath PrimPath = SkelAnimationPrim.GetPrimPath();
+	UAnimSequence* Sequence = Cast<UAnimSequence>(InfoCache->GetSingleAssetForPrim(
+		PrimPath,
+		UAnimSequence::StaticClass()
+	));
 	if (!Sequence)
 	{
 		return;
@@ -1912,20 +1915,23 @@ void FUsdLevelSequenceHelperImpl::AddSkeletalTracks(const UUsdPrimTwin& PrimTwin
 		NewSection->EvalOptions.CompletionMode = EMovieSceneCompletionMode::KeepState;
 	}
 
-	PrimPathByLevelSequenceName.AddUnique(SkelAnimationSequence->GetFName(), PrimPath);
+	PrimPathByLevelSequenceName.AddUnique(SkelAnimationSequence->GetFName(), PrimPath.GetString());
 }
 
 void FUsdLevelSequenceHelperImpl::AddGeometryCacheTracks(const UUsdPrimTwin& PrimTwin, const UE::FUsdPrim& Prim)
 {
 	UGeometryCacheComponent* ComponentToBind = Cast< UGeometryCacheComponent >(PrimTwin.GetSceneComponent());
-	if (!ComponentToBind || !AssetCache)
+	if (!ComponentToBind || !InfoCache)
 	{
 		return;
 	}
 
 	// Fetch the geometry cache asset from the asset cache. If there's none, don't actually need to create track
-	const FString PrimPath = Prim.GetPrimPath().GetString();
-	UGeometryCache* GeometryCache = Cast< UGeometryCache >(AssetCache->GetAssetForPrim(PrimPath));
+	const UE::FSdfPath PrimPath = Prim.GetPrimPath();
+	UGeometryCache* GeometryCache = Cast< UGeometryCache >(InfoCache->GetSingleAssetForPrim(
+		PrimPath,
+		UGeometryCache::StaticClass()
+	));
 	if (!GeometryCache)
 	{
 		return;
@@ -1972,7 +1978,7 @@ void FUsdLevelSequenceHelperImpl::AddGeometryCacheTracks(const UUsdPrimTwin& Pri
 		NewSection->EvalOptions.CompletionMode = EMovieSceneCompletionMode::KeepState;
 	}
 
-	PrimPathByLevelSequenceName.AddUnique(GeometryCacheSequence->GetFName(), PrimPath);
+	PrimPathByLevelSequenceName.AddUnique(GeometryCacheSequence->GetFName(), PrimPath.GetString());
 }
 
 void FUsdLevelSequenceHelperImpl::AddGroomTracks(const UUsdPrimTwin& PrimTwin, const UE::FUsdPrim& Prim)
@@ -1983,7 +1989,7 @@ void FUsdLevelSequenceHelperImpl::AddGroomTracks(const UUsdPrimTwin& PrimTwin, c
 		return;
 	}
 
-	if (!AssetCache)
+	if (!InfoCache)
 	{
 		return;
 	}
@@ -1991,7 +1997,10 @@ void FUsdLevelSequenceHelperImpl::AddGroomTracks(const UUsdPrimTwin& PrimTwin, c
 	// Fetch the groom cache asset from the asset cache. If there's none, don't actually need to create track
 	const FString PrimPath = Prim.GetPrimPath().GetString();
 	const FString GroomCachePath = FString::Printf(TEXT("%s_strands_cache"), *PrimPath);
-	UGroomCache* GroomCache = Cast< UGroomCache >(AssetCache->GetAssetForPrim(GroomCachePath));
+	UGroomCache* GroomCache = Cast<UGroomCache>(InfoCache->GetSingleAssetForPrim(
+		UE::FSdfPath{*GroomCachePath},
+		UGroomCache::StaticClass()
+	));
 	if (!GroomCache)
 	{
 		return;
@@ -2243,8 +2252,11 @@ void FUsdLevelSequenceHelperImpl::UpdateControlRigTracks(UUsdPrimTwin& PrimTwin)
 	// Fetch the UAnimSequence asset from the asset cache. Ideally we'd call AUsdStageActor::GetGeneratedAssets,
 	// but we may belong to a FUsdStageImportContext, and so there's no AUsdStageActor at all to use.
 	// At this point it doesn't matter much though, because we shouldn't need to uncollapse a SkelAnimation prim path anyway
-	FString SkelAnimationPrimPath = SkelAnimationPrim.GetPrimPath().GetString();
-	UAnimSequence* AnimSequence = Cast<UAnimSequence>(AssetCache->GetAssetForPrim(SkelAnimationPrimPath));
+	const UE::FSdfPath SkelAnimationPrimPath = SkelAnimationPrim.GetPrimPath();
+	UAnimSequence* AnimSequence = Cast<UAnimSequence>(InfoCache->GetSingleAssetForPrim(
+		SkelAnimationPrimPath,
+		UAnimSequence::StaticClass()
+	));
 
 	if (!SkelAnimationLayer)
 	{
@@ -3560,7 +3572,7 @@ public:
 	~FUsdLevelSequenceHelperImpl() {}
 
 	ULevelSequence* Init(const UE::FUsdStage& InUsdStage) { return nullptr; }
-	void SetAssetCache(UUsdAssetCache* AssetCache) {};
+	void SetInfoCache(TSharedPtr<FUsdInfoCache> InfoCache) {};
 	bool HasData() const { return false; };
 	void Clear() {};
 
@@ -3626,11 +3638,11 @@ void FUsdLevelSequenceHelper::OnStageActorRenamed()
 	}
 }
 
-void FUsdLevelSequenceHelper::SetAssetCache(UUsdAssetCache * AssetCache)
+void FUsdLevelSequenceHelper::SetInfoCache(TSharedPtr<FUsdInfoCache> InInfoCache)
 {
 	if (UsdSequencerImpl.IsValid())
 	{
-		UsdSequencerImpl->SetAssetCache(AssetCache);
+		UsdSequencerImpl->SetInfoCache(InInfoCache);
 	}
 }
 
diff --git a/Engine/Plugins/Importers/USDImporter/Source/USDStage/Private/USDStageActor.cpp b/Engine/Plugins/Importers/USDImporter/Source/USDStage/Private/USDStageActor.cpp
--- a/Engine/Plugins/Importers/USDImporter/Source/USDStage/Private/USDStageActor.cpp
+++ b/Engine/Plugins/Importers/USDImporter/Source/USDStage/Private/USDStageActor.cpp
@@ -20,6 +20,7 @@
 #include "USDLog.h"
 #include "USDPrimConversion.h"
 #include "USDPrimTwin.h"
+#include "USDProjectSettings.h"
 #include "USDSchemasModule.h"
 #include "USDSchemaTranslator.h"
 #include "USDSkelRootTranslator.h"
@@ -53,6 +54,7 @@
 #include "Engine/StaticMesh.h"
 #include "Engine/World.h"
 #include "EngineAnalytics.h"
+#include "Framework/Notifications/NotificationManager.h"
 #include "HAL/FileManager.h"
 #include "HAL/IConsoleManager.h"
 #include "LevelSequence.h"
@@ -70,8 +72,11 @@
 #include "StaticMeshAttributes.h"
 #include "StaticMeshOperations.h"
 #include "Tracks/MovieScene3DTransformTrack.h"
+#include "Widgets/Notifications/SNotificationList.h"
 
 #if WITH_EDITOR
+#include "USDClassesEditorModule.h"
+
 #include "Editor.h"
 #include "Editor/TransBuffer.h"
 #include "Editor/UnrealEdEngine.h"
@@ -102,8 +107,8 @@ struct FUsdStageActorImpl
 	{
 		TSharedRef< FUsdSchemaTranslationContext > TranslationContext = MakeShared< FUsdSchemaTranslationContext >(
 			StageActor->GetOrLoadUsdStage(),
-			*StageActor->AssetCache
-			);
+			*StageActor->UsdAssetCache
+		);
 
 		TranslationContext->Level = StageActor->GetLevel();
 		TranslationContext->ObjectFlags = DefaultObjFlag;
@@ -637,8 +642,6 @@ AUsdStageActor::AUsdStageActor()
 
 	RootComponent = SceneComponent;
 
-	AssetCache = CreateDefaultSubobject< UUsdAssetCache >(TEXT("AssetCache"));
-
 	// Note: We can't construct our RootUsdTwin as a default subobject here, it needs to be built on-demand.
 	// Even if we NewObject'd one it will work as a subobject in some contexts (maybe because the CDO will have a dedicated root twin?).
 	// As far as the engine is concerned, our prim twins are static assets like meshes or textures. However, they live on the transient
@@ -649,7 +652,8 @@ AUsdStageActor::AUsdStageActor()
 	IUsdSchemasModule& UsdSchemasModule = FModuleManager::Get().LoadModuleChecked< IUsdSchemasModule >(TEXT("USDSchemas"));
 	RenderContext = UsdSchemasModule.GetRenderContextRegistry().GetUnrealRenderContext();
 
-	Transactor = NewObject<UUsdTransactor>(this, TEXT("Transactor"), EObjectFlags::RF_Transactional);
+	const FName UniqueName = MakeUniqueObjectName(this, UUsdTransactor::StaticClass(), TEXT("Transactor"));
+	Transactor = NewObject<UUsdTransactor>(this, UniqueName, EObjectFlags::RF_Transactional);
 	Transactor->Initialize(this);
 
 	if (HasAuthorityOverStage())
@@ -704,7 +708,7 @@ AUsdStageActor::AUsdStageActor()
 						// We have a valid filepath but no objects/assets spawned, so it's likely we were just spawned on the
 						// other client, and were replicated here with our RootLayer path already filled out, meaning we should just load that stage
 						// Note that now our UUsdTransactor may have already caused the stage itself to be loaded, but we may still need to call LoadUsdStage on our end.
-						else if (ObjectsToWatch.Num() == 0 && (!AssetCache || AssetCache->GetNumAssets() == 0))
+						else if (ObjectsToWatch.Num() == 0 && (!UsdAssetCache || UsdAssetCache->GetNumAssets() == 0))
 						{
 							this->LoadUsdStage();
 							AUsdStageActor::OnActorLoaded.Broadcast(this);
@@ -918,6 +922,14 @@ void AUsdStageActor::OnUsdObjectsChanged(const UsdUtils::FObjectChangesByPath& I
 		}
 	}
 
+	EnsureAssetCache();
+
+	TOptional<FUsdScopedAssetCacheReferencer> ScopedReferencer;
+	if (UsdAssetCache)
+	{
+		ScopedReferencer.Emplace(UsdAssetCache, this);
+	}
+
 	// Mark the level as dirty since we received a notice about our stage having changed in some way.
 	// The main goal of this is to trigger the "save layers" dialog if we then save the UE level
 	const bool bAlwaysMarkDirty = true;
@@ -1125,14 +1137,20 @@ void AUsdStageActor::OnUsdObjectsChanged(const UsdUtils::FObjectChangesByPath& I
 			{
 				TSharedRef< FUsdSchemaTranslationContext > TranslationContext = FUsdStageActorImpl::CreateUsdSchemaTranslationContext(this, AssetsPrimPath.GetString());
 
-				if (bInResync)
+				if (bInResync && InfoCache)
 				{
-					UMaterialInterface* ExistingMaterial = Cast<UMaterialInterface>(AssetCache->GetAssetForPrim(AssetsPrimPath.GetString()));
+					UMaterialInterface* ExistingMaterial = Cast<UMaterialInterface>(InfoCache->GetSingleAssetForPrim(
+						AssetsPrimPath,
+						UMaterialInterface::StaticClass()
+					));
 
 					UE::FUsdPrim PrimToResync = GetOrLoadUsdStage().GetPrimAtPath(AssetsPrimPath);
 					LoadAssets(*TranslationContext, PrimToResync);
 
-					UMaterialInterface* NewMaterial = Cast<UMaterialInterface>(AssetCache->GetAssetForPrim(AssetsPrimPath.GetString()));
+					UMaterialInterface* NewMaterial = Cast<UMaterialInterface>(InfoCache->GetSingleAssetForPrim(
+						AssetsPrimPath,
+						UMaterialInterface::StaticClass()
+					));
 
 					// For UE-120185: If we recreated a material for a prim path we also need to update all components that were using it.
 					// This could be fleshed out further if other asset types require this refresh of "dependent components" but materials
@@ -1169,6 +1187,11 @@ void AUsdStageActor::OnUsdObjectsChanged(const UsdUtils::FObjectChangesByPath& I
 			OnPrimChanged.Broadcast(PrimChangedInfo.Key, PrimChangedInfo.Value);
 		}
 	}
+
+	if (bHasResync && UsdAssetCache)
+	{
+		UsdAssetCache->RefreshStorage();
+	}
 #endif // USE_USD_SDK
 }
 
@@ -1530,6 +1553,52 @@ void AUsdStageActor::SetRootLayer(const FString& RootFilePath)
 	LoadUsdStage();
 }
 
+void AUsdStageActor::SetAssetCache(UUsdAssetCache2* NewCache)
+{
+	const bool bMarkDirty = false;
+	Modify(bMarkDirty);
+
+	// Remove ourselves from our previous cache. We're going to have to create new assets anyway, so it doesn't matter
+	// if we discard our current assets.
+	// Note that these functions may be called when the cache is already set (if via
+	// AUsdStageActor::HandlePropertyChangedEvent) or to set the actual new cache (when called directly)
+	if (UsdAssetCache != NewCache)
+	{
+		if (UsdAssetCache)
+		{
+			UsdAssetCache->RemoveAllAssetReferences(this);
+			UsdAssetCache->RefreshStorage();
+		}
+
+		UsdAssetCache = NewCache;
+	}
+
+	// We can't have no cache while we have a stage loaded, so at least revert the property to a transient cache
+	// instead, as the intent may have been to just have the actor not point at the previous cache anymore.
+	if (!UsdAssetCache && UsdStage)
+	{
+		FNotificationInfo Toast(LOCTEXT("MustHaveCache", "Must have an Asset Cache"));
+		Toast.SubText = LOCTEXT("MustHaveCache_Subtext", "The Stage Actor must always have an Asset Cache while a stage is loaded, so a temporary cache will be created.\n\nClose the stage before clearing the cache if you wish to clear this property.");
+		Toast.Image = FCoreStyle::Get().GetBrush(TEXT("MessageLog.Warning"));
+		Toast.bUseLargeFont = false;
+		Toast.bFireAndForget = true;
+		Toast.FadeOutDuration = 1.0f;
+		Toast.ExpireDuration = 12.0f;
+		Toast.bUseThrobber = false;
+		Toast.bUseSuccessFailIcons = false;
+		FSlateNotificationManager::Get().AddNotification(Toast);
+
+		UsdAssetCache = NewObject< UUsdAssetCache2 >(GetTransientPackage(), NAME_None, GetMaskedFlags(RF_PropagateToSubObjects));
+	}
+
+	// Here we pretend we just received a root resync so that we re-fetch assets from the cache
+	// and update its components
+	UsdUtils::FObjectChangesByPath InfoChanges;
+	UsdUtils::FObjectChangesByPath ResyncChanges;
+	ResyncChanges.Add({ TEXT("/"), {} });
+	OnUsdObjectsChanged(InfoChanges, ResyncChanges);
+}
+
 void AUsdStageActor::SetInitialLoadSet(EUsdInitialLoadSet NewLoadSet)
 {
 	const bool bMarkDirty = false;
@@ -1692,44 +1761,28 @@ TArray<UObject*> AUsdStageActor::GetGeneratedAssets(const FString& PrimPath)
 		return {};
 	}
 
-	if (!AssetCache)
+	if (!InfoCache)
 	{
 		return {};
 	}
 
-	TSet<UObject*> Result;
-	FString PathToUse = PrimPath;
-
-	// If we found an asset generated specifically for a prim path, use that. If this prim has been collapsed,
-	// the asset generated for it will be assigned to its collapsing root, so we won't find anything here
-	if (UObject* FoundAsset = AssetCache->GetAssetForPrim(PrimPath))
+	// Prefer checking the prim directly, but also check its collapsed root if it is collapsed.
+	// This because we have some exception cases like USkeleton/UAnimSequences that can be found by querying the actual
+	// Skeleton/SKelAnimation prims even though they are collapsed into the SkelRoot prim.
+	TSet<TWeakObjectPtr<UObject>> AssetsSet = InfoCache->GetAssetsForPrim(UsdPath);
+	if (AssetsSet.Num() == 0 && InfoCache->IsPathCollapsed(UsdPath, ECollapsingType::Assets))
 	{
-		Result.Add(FoundAsset);
-	}
-	// If we haven't try seeing if we collapsed into another asset
-	else if (InfoCache.IsValid())
-	{
-		PathToUse = InfoCache->UnwindToNonCollapsedPath(UsdPath, ECollapsingType::Assets).GetString();
-		if (UObject* FoundCollapsedAsset = AssetCache->GetAssetForPrim(PathToUse))
-		{
-			Result.Add(FoundCollapsedAsset);
-		}
+		UsdPath = InfoCache->UnwindToNonCollapsedPath(UsdPath, ECollapsingType::Assets);
+		AssetsSet = InfoCache->GetAssetsForPrim(UsdPath);
 	}
 
-	// Collect any other asset that claims they came from the same prim (e.g. also return the skeleton if we query
-	// the SkelRoot, or return textures if we query a material prim that used them, etc.)
-	for (const TPair<FString, TObjectPtr<UObject>>& HashToAsset : AssetCache->GetCachedAssets())
+	TArray<UObject*> Assets;
+	Assets.Reserve(AssetsSet.Num());
+	for (const TWeakObjectPtr<UObject>& Asset : AssetsSet)
 	{
-		if (UUsdAssetImportData* ImportData = UsdUtils::GetAssetImportData(HashToAsset.Value))
-		{
-			if (ImportData->PrimPath == PathToUse)
-			{
-				Result.Add(HashToAsset.Value);
-			}
-		}
+		Assets.Add(Asset.Get());
 	}
-
-	return Result.Array();
+	return Assets;
 }
 
 FString AUsdStageActor::GetSourcePrimPath(UObject* Object)
@@ -1741,26 +1794,12 @@ FString AUsdStageActor::GetSourcePrimPath(UObject* Object)
 			return UsdPrimTwin->PrimPath;
 		}
 	}
-
-	if (AssetCache)
+	else if (InfoCache)
 	{
-		for (const TPair<FString, TWeakObjectPtr<UObject>>& PrimPathToAsset : AssetCache->GetAssetPrimLinks())
+		const UE::FSdfPath FoundPath = InfoCache->GetPrimForAsset(Object);
+		if (!FoundPath.IsEmpty())
 		{
-			if (PrimPathToAsset.Value.Get() == Object)
-			{
-				return PrimPathToAsset.Key;
-			}
-		}
-
-		for (const TPair<FString, TObjectPtr<UObject>>& HashToAsset : AssetCache->GetCachedAssets())
-		{
-			if (HashToAsset.Value == Object)
-			{
-				if (UUsdAssetImportData* ImportData = UsdUtils::GetAssetImportData(HashToAsset.Value))
-				{
-					return ImportData->PrimPath;
-				}
-			}
+			return FoundPath.GetString();
 		}
 	}
 
@@ -1934,10 +1973,7 @@ void AUsdStageActor::OnObjectsReplaced(const TMap<UObject*, UObject*>& ObjectRep
 			NewActor->OnPreStageChanged = OnPreStageChanged;
 			NewActor->OnPrimChanged = OnPrimChanged;
 
-			// UEngine::CopyPropertiesForUnrelatedObjects won't copy over the cache's transient assets, but we still
-			// need to ensure their lifetime here, so just take the previous asset cache instead, which still that has the transient assets
-			AssetCache->Rename(nullptr, NewActor);
-			NewActor->AssetCache = AssetCache;
+			NewActor->UsdAssetCache = UsdAssetCache;
 
 			NewActor->InfoCache = InfoCache;
 			InfoCache = nullptr;
@@ -1975,6 +2011,12 @@ void AUsdStageActor::LoadUsdStage()
 {
 	TRACE_CPUPROFILER_EVENT_SCOPE(AUsdStageActor::LoadUsdStage);
 
+	// Ensure an asset cache if we're going to load something
+	if (!RootLayer.FilePath.IsEmpty())
+	{
+		EnsureAssetCache();
+	}
+
 	double StartTime = FPlatformTime::Cycles64();
 
 	FScopedSlowTask SlowTask(1.f, LOCTEXT("LoadingUDStage", "Loading USD Stage"));
@@ -1982,11 +2024,6 @@ void AUsdStageActor::LoadUsdStage()
 
 	OnPreStageChanged.Broadcast();
 
-	if (!AssetCache)
-	{
-		AssetCache = NewObject< UUsdAssetCache >(this, TEXT("AssetCache"), GetMaskedFlags(RF_PropagateToSubObjects));
-	}
-
 	// Block writing level sequence changes back to the USD stage until we finished this transaction, because once we do
 	// the movie scene and tracks will all trigger OnObjectTransacted. We listen for those on FUsdLevelSequenceHelperImpl::OnObjectTransacted,
 	// and would otherwise end up writing all of the data we just loaded back to the USD stage
@@ -2002,6 +2039,12 @@ void AUsdStageActor::LoadUsdStage()
 
 	FScopedUsdMessageLog ScopedMessageLog;
 
+	TOptional<FUsdScopedAssetCacheReferencer> ScopedReferencer;
+	if (UsdAssetCache)
+	{
+		ScopedReferencer.Emplace(UsdAssetCache, this);
+	}
+
 	// If we have an isolated stage we may have been asked to refresh to display it: We should keep our UsdStage opened
 	UE::FUsdStage StageToLoad;
 	if (IsolatedStage)
@@ -2076,6 +2119,11 @@ void AUsdStageActor::LoadUsdStage()
 		}
 	}
 
+	if (UsdAssetCache)
+	{
+		UsdAssetCache->RefreshStorage();
+	}
+
 #if WITH_EDITOR
 	if (GIsEditor && GEditor && !IsGarbageCollecting())
 	{
@@ -2112,12 +2160,6 @@ void AUsdStageActor::UnloadUsdStage()
 	// and could lead to stage changes
 	BlockMonitoringLevelSequenceForThisTransaction();
 
-	if (AssetCache)
-	{
-		FUsdStageActorImpl::CloseEditorsForAssets(AssetCache->GetCachedAssets());
-		AssetCache->Reset();
-	}
-
 	if (InfoCache)
 	{
 		InfoCache->Clear();
@@ -2156,9 +2198,66 @@ void AUsdStageActor::UnloadUsdStage()
 
 	CloseUsdStage();
 
+	if (UsdAssetCache)
+	{
+		UsdAssetCache->RemoveAllAssetReferences(this);
+		UsdAssetCache->RefreshStorage();
+	}
+
 	OnStageChanged.Broadcast();
 }
 
+void AUsdStageActor::EnsureAssetCache()
+{
+	if (!UsdAssetCache)
+	{
+		if (const UUsdProjectSettings* ProjectSettings = GetDefault<UUsdProjectSettings>())
+		{
+			if (UUsdAssetCache2* DefaultCache = Cast<UUsdAssetCache2>(ProjectSettings->DefaultAssetCache.TryLoad()))
+			{
+				UE_LOG(LogUsd, Log, TEXT("USD Stage Actor '%s' had no previous USD Asset Cache, so it will use the default cache at '%s'. This can be configured on the project settings."),
+					*GetPathName(),
+					*DefaultCache->GetPathName()
+				);
+
+				UsdAssetCache = DefaultCache;
+			}
+		}
+	}
+#if WITH_EDITOR
+	// Show a dialog to let the user create a new default asset cache somewhere
+	if (!UsdAssetCache && GIsEditor && !IsRunningCommandlet() && !IsTemplate())
+	{
+		if (UUsdProjectSettings* ProjectSettings = GetMutableDefault<UUsdProjectSettings>())
+		{
+			if (ProjectSettings->bShowCreateDefaultAssetCacheDialog)
+			{
+				if (UUsdAssetCache2* NewCache = IUsdClassesEditorModule::ShowMissingDefaultAssetCacheDialog())
+				{
+					ProjectSettings->DefaultAssetCache = NewCache;
+					ProjectSettings->SaveConfig();
+
+					UsdAssetCache = NewCache;
+
+					UE_LOG(LogUsd, Log, TEXT("USD Stage Actor '%s' will use newly created, default USD Asset Cache at '%s'. This can be configured on the project settings."),
+						*GetPathName(),
+						*UsdAssetCache->GetPathName()
+					);
+				}
+			}
+		}
+	}
+#endif
+	if (!UsdAssetCache)
+	{
+		UE_LOG(LogUsd, Warning, TEXT("USD Stage Actor '%s' had no previous USD Asset Cache and no default cache is specified on the project settings, so a temporary cache will be generated. For better performance, create a persistent USD Asset Cache asset and point to it with this actor's UsdAssetCache property."),
+			*GetPathName()
+		);
+
+		UsdAssetCache = NewObject< UUsdAssetCache2 >(GetTransientPackage(), NAME_None, GetMaskedFlags(RF_PropagateToSubObjects));
+	}
+}
+
 UUsdPrimTwin* AUsdStageActor::GetRootPrimTwin()
 {
 	if (!RootUsdTwin)
@@ -2236,11 +2335,6 @@ void AUsdStageActor::ReloadAnimations()
 	}
 }
 
-UUsdAssetCache* AUsdStageActor::GetAssetCache()
-{
-	return AssetCache;
-}
-
 TSharedPtr<FUsdInfoCache> AUsdStageActor::GetInfoCache()
 {
 	return InfoCache;
@@ -2454,6 +2548,12 @@ void AUsdStageActor::PostDuplicate(bool bDuplicateForPIE)
 	}
 	else
 	{
+		// Temporary asset caches aren't meant to be shared
+		if (UsdAssetCache && UsdAssetCache->GetOutermost() == GetTransientPackage())
+		{
+			UsdAssetCache = NewObject< UUsdAssetCache2 >(GetTransientPackage(), NAME_None, GetMaskedFlags(RF_PropagateToSubObjects));
+		}
+
 		LoadUsdStage();
 	}
 }
@@ -2483,6 +2583,15 @@ void AUsdStageActor::Serialize(FArchive& Ar)
 
 		InfoCache->Serialize(Ar);
 	}
+
+	// Make sure we clear the old deprecated cache if we still have one, as old scenes may be persisting textures
+	// in there
+	PRAGMA_DISABLE_DEPRECATION_WARNINGS
+	if (AssetCache && Ar.IsLoading())
+	{
+		AssetCache->Reset();
+	}
+	PRAGMA_ENABLE_DEPRECATION_WARNINGS
 }
 
 void AUsdStageActor::Destroyed()
@@ -2560,7 +2669,7 @@ void AUsdStageActor::PostRegisterAllComponents()
 	// This may fail if our stage happened to not spawn any components, actors or assets, but by that
 	// point "being loaded" doesn't really mean anything anyway
 	const bool bStageIsLoaded = GetBaseUsdStage()
-		&& ((RootUsdTwin && RootUsdTwin->GetSceneComponent() != nullptr) || (AssetCache && AssetCache->GetNumAssets() > 0));
+		&& ((RootUsdTwin && RootUsdTwin->GetSceneComponent() != nullptr) || (UsdAssetCache && UsdAssetCache->GetNumAssets() > 0));
 
 	// Blocks loading stage when going into PIE, if we already have something loaded (we'll want to duplicate stuff instead).
 	// We need to allow loading when going into PIE when we have nothing loaded yet because the MovieRenderQueue (or other callers)
@@ -2653,7 +2762,7 @@ void AUsdStageActor::UnregisterAllComponents(bool bForReregister)
 #endif // WITH_EDITOR
 
 	const bool bStageIsLoaded = GetBaseUsdStage()
-		&& ((RootUsdTwin && RootUsdTwin->GetSceneComponent() != nullptr) || (AssetCache && AssetCache->GetNumAssets() > 0));
+		&& ((RootUsdTwin && RootUsdTwin->GetSceneComponent() != nullptr) || (UsdAssetCache && UsdAssetCache->GetNumAssets() > 0));
 
 	UWorld* World = GetWorld();
 	if (bIsTransitioningIntoPIE && bStageIsLoaded && (!World || World->WorldType == EWorldType::PIE))
@@ -2985,6 +3094,10 @@ void AUsdStageActor::HandlePropertyChangedEvent(FPropertyChangedEvent& PropertyC
 	{
 		SetRootMotionHandling(RootMotionHandling);
 	}
+	else if (PropertyName == GET_MEMBER_NAME_CHECKED(AUsdStageActor, UsdAssetCache))
+	{
+		SetAssetCache(UsdAssetCache);
+	}
 
 	bIsModifyingAProperty = false;
 }
@@ -3069,9 +3182,17 @@ void AUsdStageActor::LoadAsset(FUsdSchemaTranslationContext& TranslationContext,
 	PrimPath = UsdToUnreal::ConvertPath(Prim.GetPrimPath());
 #endif // #if USE_USD_SDK
 
-	if (AssetCache)
+	if (InfoCache)
 	{
-		AssetCache->RemoveAssetPrimLink(PrimPath);
+		TSet<TWeakObjectPtr<UObject>> OldAssets = InfoCache->RemoveAllAssetPrimLinks(UE::FSdfPath{*PrimPath});
+
+		if (UsdAssetCache)
+		{
+			for (const TWeakObjectPtr<UObject>& OldAsset : OldAssets)
+			{
+				UsdAssetCache->RemoveAssetReference(OldAsset.Get(), this);
+			}
+		}
 	}
 
 	IUsdSchemasModule& UsdSchemasModule = FModuleManager::Get().LoadModuleChecked< IUsdSchemasModule >(TEXT("USDSchemas"));
@@ -3092,21 +3213,29 @@ void AUsdStageActor::LoadAssets(FUsdSchemaTranslationContext& TranslationContext
 	TGuardValue< EObjectFlags > ContextFlagsGuard(TranslationContext.ObjectFlags, TranslationContext.ObjectFlags & ~RF_Transactional);
 
 	// Clear existing prim/asset association
-	if (AssetCache)
+	if (InfoCache)
 	{
-		FString StartPrimPath = StartPrim.GetPrimPath().GetString();
-		TSet<FString> PrimPathsToRemove;
-		for (const TPair< FString, TWeakObjectPtr<UObject> >& PrimPathToAssetIt : AssetCache->GetAssetPrimLinks())
+		UE::FSdfPath StartPrimPath = StartPrim.GetPrimPath();
+		TSet<UE::FSdfPath> PrimPathsToRemove;
+		for (const TPair<UE::FSdfPath, TSet<TWeakObjectPtr<UObject>>>& PrimPathToAssetIt : InfoCache->GetAllAssetPrimLinks())
 		{
-			const FString& PrimPath = PrimPathToAssetIt.Key;
-			if (PrimPath.StartsWith(StartPrimPath) || PrimPath == StartPrimPath)
+			const UE::FSdfPath& PrimPath = PrimPathToAssetIt.Key;
+			if (PrimPath.HasPrefix(StartPrimPath) || PrimPath == StartPrimPath)
 			{
 				PrimPathsToRemove.Add(PrimPath);
 			}
 		}
-		for (const FString& PrimPathToRemove : PrimPathsToRemove)
+
+		if (UsdAssetCache)
 		{
-			AssetCache->RemoveAssetPrimLink(PrimPathToRemove);
+			for (const UE::FSdfPath& PrimPathToRemove : PrimPathsToRemove)
+			{
+				TSet<TWeakObjectPtr<UObject>> OldAssets = InfoCache->RemoveAllAssetPrimLinks(PrimPathToRemove);
+				for (const TWeakObjectPtr<UObject>& OldAsset : OldAssets)
+				{
+					UsdAssetCache->RemoveAssetReference(OldAsset.Get(), this);
+				}
+			}
 		}
 	}
 
diff --git a/Engine/Plugins/Importers/USDImporter/Source/USDStage/Public/USDLevelSequenceHelper.h b/Engine/Plugins/Importers/USDImporter/Source/USDStage/Public/USDLevelSequenceHelper.h
--- a/Engine/Plugins/Importers/USDImporter/Source/USDStage/Public/USDLevelSequenceHelper.h
+++ b/Engine/Plugins/Importers/USDImporter/Source/USDStage/Public/USDLevelSequenceHelper.h
@@ -3,14 +3,15 @@
 #pragma once
 
 #include "CoreMinimal.h"
+#include "Templates/SharedPointer.h"
 
 #include "UsdWrappers/ForwardDeclarations.h"
 
 class AUsdStageActor;
 class FScopedBlockMonitoringChangesForTransaction;
+class FUsdInfoCache;
 class FUsdLevelSequenceHelperImpl;
 class ULevelSequence;
-class UUsdAssetCache;
 class UUsdPrimTwin;
 
 /**
@@ -34,7 +35,7 @@ public:
 	ULevelSequence* Init(const UE::FUsdStage& UsdStage);
 
 	/** Sets the asset cache to use when fetching assets and asset info required for the level sequence animation, like UAnimSequences */
-	void SetAssetCache( UUsdAssetCache* AssetCache );
+	void SetInfoCache(TSharedPtr<FUsdInfoCache> InInfoCache);
 
 	/* Returns true if we have at least one possessable or a reference to a subsequence */
 	bool HasData() const;
diff --git a/Engine/Plugins/Importers/USDImporter/Source/USDStage/Public/USDStageActor.h b/Engine/Plugins/Importers/USDImporter/Source/USDStage/Public/USDStageActor.h
--- a/Engine/Plugins/Importers/USDImporter/Source/USDStage/Public/USDStageActor.h
+++ b/Engine/Plugins/Importers/USDImporter/Source/USDStage/Public/USDStageActor.h
@@ -13,11 +13,12 @@
 
 #include "USDStageActor.generated.h"
 
+class FUsdInfoCache;
+class UUsdAssetCache2;
 class ULevelSequence;
 class UUsdAssetCache;
 class UUsdPrimTwin;
 class UUsdTransactor;
-struct FUsdInfoCache;
 struct FUsdSchemaTranslationContext;
 namespace UE
 {
@@ -33,6 +34,9 @@ public:
 	UPROPERTY(EditAnywhere, BlueprintReadOnly, Category = "USD", meta = (RelativeToGameDir))
 	FFilePath RootLayer;
 
+	UPROPERTY(EditAnywhere, BlueprintReadOnly, Category = "USD")
+	TObjectPtr<UUsdAssetCache2> UsdAssetCache;
+
 	// These properties are configs so that spawned actors read them from the CDO when spawned.
 	// This allows the defaults for them to be configured on EditorPerProjectUserSettings.ini, and allows us to write
 	// to that config from the USD Stage Editor, specifying our options before the editor is attached to any stage actor.
@@ -108,6 +112,9 @@ public:
 	UFUNCTION(BlueprintCallable, Category = "USD", meta = (CallInEditor = "true"))
 	USDSTAGE_API void SetRootLayer(const FString& RootFilePath );
 
+	UFUNCTION(BlueprintCallable, Category = "USD", meta = (CallInEditor = "true"))
+	USDSTAGE_API void SetAssetCache(UUsdAssetCache2* NewCache);
+
 	UFUNCTION(BlueprintCallable, Category = "USD", meta = (CallInEditor = "true"))
 	USDSTAGE_API void SetInitialLoadSet( EUsdInitialLoadSet NewLoadSet );
 
@@ -214,7 +221,6 @@ public:
 	// Regenerates our LevelSequence from the USD Stage
 	USDSTAGE_API void ReloadAnimations();
 
-	USDSTAGE_API UUsdAssetCache* GetAssetCache();
 	USDSTAGE_API TSharedPtr<FUsdInfoCache> GetInfoCache();
 	USDSTAGE_API TMap< FString, TMap< FString, int32 > > GetMaterialToPrimvarToUVIndex();
 	USDSTAGE_API const UsdUtils::FBlendShapeMap& GetBlendShapeMap();
@@ -277,6 +283,8 @@ protected:
 	void LoadUsdStage();
 	void UnloadUsdStage();
 
+	void EnsureAssetCache();
+
 	bool HasAuthorityOverStage() const;
 
 	void UpdateSpawnedObjectsTransientFlag(bool bTransient);
@@ -321,7 +329,8 @@ protected:
 	UPROPERTY( Transient )
 	TMap< TObjectPtr<UObject>, FString > ObjectsToWatch;
 
-	UPROPERTY( VisibleAnywhere, Category = "USD", AdvancedDisplay )
+	UE_DEPRECATED(5.2, "Use the new AssetCache property instead, that uses UUsdAssetCache2 objects")
+	UPROPERTY(AdvancedDisplay, meta=(DeprecatedProperty, DeprecationMessage = "Use the new AssetCache property instead, that uses UUsdAssetCache2 objects"))
 	TObjectPtr<UUsdAssetCache> AssetCache;
 
 	UPROPERTY( Transient )
diff --git a/Engine/Plugins/Importers/USDImporter/Source/USDStage/USDStage.Build.cs b/Engine/Plugins/Importers/USDImporter/Source/USDStage/USDStage.Build.cs
--- a/Engine/Plugins/Importers/USDImporter/Source/USDStage/USDStage.Build.cs
+++ b/Engine/Plugins/Importers/USDImporter/Source/USDStage/USDStage.Build.cs
@@ -37,7 +37,6 @@ namespace UnrealBuildTool.Rules
 					"MovieScene",
 					"MovieSceneTracks",
 					"Niagara",	// Needed by GroomComponent.h
-					"Projects", // So that we can use the plugin manager to find out our content dir and cook the base materials
 					"Slate",
 					"SlateCore",
 					"StaticMeshDescription",
@@ -61,6 +60,7 @@ namespace UnrealBuildTool.Rules
 						"PropertyEditor", // For the stage actor's details customization
 						"Sequencer",
 						"UnrealEd",
+						"USDClassesEditor",
 					}
 				);
 			}
diff --git a/Engine/Plugins/Importers/USDImporter/Source/USDStageEditor/Private/SUSDStageEditorStyle.cpp b/Engine/Plugins/Importers/USDImporter/Source/USDStageEditor/Private/SUSDStageEditorStyle.cpp
--- a/Engine/Plugins/Importers/USDImporter/Source/USDStageEditor/Private/SUSDStageEditorStyle.cpp
+++ b/Engine/Plugins/Importers/USDImporter/Source/USDStageEditor/Private/SUSDStageEditorStyle.cpp
@@ -12,6 +12,7 @@
 
 #define IMAGE_PLUGIN_BRUSH( RelativePath, ... ) FSlateImageBrush( FUsdStageEditorStyle::InContent( RelativePath, ".png" ), __VA_ARGS__ )
 #define IMAGE_BRUSH(RelativePath, ...) FSlateImageBrush(StyleSet->RootToContentDir(RelativePath, TEXT(".png")), __VA_ARGS__)
+#define IMAGE_BRUSH_SVG( RelativePath, ... ) FSlateVectorImageBrush(StyleSet->RootToContentDir(RelativePath, TEXT(".svg")), __VA_ARGS__)
 #define BOX_BRUSH(RelativePath, ...) FSlateBoxBrush(StyleSet->RootToContentDir(RelativePath, TEXT(".png")), __VA_ARGS__)
 #define DEFAULT_FONT(...) FCoreStyle::GetDefaultFontStyle(__VA_ARGS__)
 
diff --git a/Engine/Plugins/Importers/USDImporter/Source/USDStageEditorViewModels/Private/USDStageViewModel.cpp b/Engine/Plugins/Importers/USDImporter/Source/USDStageEditorViewModels/Private/USDStageViewModel.cpp
--- a/Engine/Plugins/Importers/USDImporter/Source/USDStageEditorViewModels/Private/USDStageViewModel.cpp
+++ b/Engine/Plugins/Importers/USDImporter/Source/USDStageEditorViewModels/Private/USDStageViewModel.cpp
@@ -3,6 +3,7 @@
 #include "USDStageViewModel.h"
 
 #include "UnrealUSDWrapper.h"
+#include "USDAssetCache2.h"
 #include "USDClassesModule.h"
 #include "USDConversionUtils.h"
 #include "USDErrorUtils.h"
@@ -305,7 +306,10 @@ void FUsdStageViewModel::ImportStage( const TCHAR* TargetContentFolder, UUsdStag
 		ImportContext.Stage = UsdStage;
 
 		const bool bIsAutomated = TargetContentFolder && Options;
-		if ( ImportContext.Init( StageName, RootPath, TEXT("/Game/"), RF_Public | RF_Transactional, bIsAutomated ) )
+
+		// Even when importing we still want to first create these assets as transient. Only the publishing process
+		// itself will remove the transient flag if everything succeeded
+		if ( ImportContext.Init( StageName, RootPath, TEXT("/Game/"), RF_Public | RF_Transactional | RF_Transient, bIsAutomated ) )
 		{
 			FScopedTransaction Transaction( FText::Format(LOCTEXT("ImportTransaction", "Import USD stage '{0}'"), FText::FromString(StageName)));
 
@@ -318,14 +322,18 @@ void FUsdStageViewModel::ImportStage( const TCHAR* TargetContentFolder, UUsdStag
 
 			// Let the importer reuse our assets, but force it to spawn new actors and components always
 			// This allows a different setting for asset/component collapsing, and doesn't require modifying the PrimTwins
-			ImportContext.AssetCache = StageActor->GetAssetCache();
+			ImportContext.AssetCache = StageActor->UsdAssetCache;
 			ImportContext.InfoCache = StageActor->GetInfoCache();
-			ImportContext.LevelSequenceHelper.SetAssetCache( StageActor->GetAssetCache() );
+			ImportContext.LevelSequenceHelper.SetInfoCache(StageActor->GetInfoCache());
 			ImportContext.MaterialToPrimvarToUVIndex = StageActor->GetMaterialToPrimvarToUVIndex();
 
 			ImportContext.TargetSceneActorAttachParent = StageActor->GetRootComponent()->GetAttachParent();
 			ImportContext.TargetSceneActorTargetTransform = StageActor->GetActorTransform();
 
+			// Preemptively remove the stage actor as a user of the assets on the cache so that the stage importer
+			// can just take the assets from the cache directly. Otherwise it will be forced to duplicate them
+			ImportContext.AssetCache->RemoveAllAssetReferences(StageActor);
+
 			UUsdStageImporter* USDImporter = IUsdStageImporterModule::Get().GetImporter();
 			USDImporter->ImportFromFile(ImportContext);
 
diff --git a/Engine/Plugins/Importers/USDImporter/Source/USDStageImporter/Private/USDStageImporter.cpp b/Engine/Plugins/Importers/USDImporter/Source/USDStageImporter/Private/USDStageImporter.cpp
--- a/Engine/Plugins/Importers/USDImporter/Source/USDStageImporter/Private/USDStageImporter.cpp
+++ b/Engine/Plugins/Importers/USDImporter/Source/USDStageImporter/Private/USDStageImporter.cpp
@@ -3,6 +3,7 @@
 #include "USDStageImporter.h"
 
 #include "USDAssetCache.h"
+#include "USDAssetCache2.h"
 #include "USDAssetImportData.h"
 #include "USDClassesModule.h"
 #include "USDConversionUtils.h"
@@ -661,10 +662,26 @@ namespace UsdStageImporterImpl
 
 		FSoftObjectPath OldPath = Asset;
 
+		// If the asset is currently being referenced on its asset cache we can't just take it, and must duplicate it
+		// below with the DuplicateObject call
+		TStrongObjectPtr<UObject> TakenAsset = nullptr;
+		if (UUsdAssetCache2* Cache = ImportContext.AssetCache.Get())
+		{
+			FString Hash = Cache->GetHashForAsset(Asset);
+			if (!Hash.IsEmpty())
+			{
+				TakenAsset.Reset(Cache->RemoveAsset(Hash));
+				if (!TakenAsset.IsValid())
+				{
+					UE_LOG(LogUsd, Log, TEXT("About to duplicate asset '%s' on import as it can't be taken directly from the asset cache"), *Asset->GetPathName());
+				}
+			}
+		}
+
 		// Strategy copied from FDatasmithImporterImpl::PublicizeAsset
 		// Replace existing asset (reimport or conflict) with new asset
 		UObject* MovedAsset = ExistingAsset;
-		if (ExistingAsset != nullptr && ExistingAsset != Asset && ReplacePolicy == EReplaceAssetPolicy::Replace)
+		if (!TakenAsset || (ExistingAsset != nullptr && ExistingAsset != Asset && ReplacePolicy == EReplaceAssetPolicy::Replace))
 		{
 			// We must do this here or else on the next engine tick the USkeletalMeshComponent::ParallelAnimationEvaluation
 			// will try to use the old anim instance, which will try using the old Skeleton's SmartNameContainer and crash.
@@ -681,7 +698,7 @@ namespace UsdStageImporterImpl
 
 			// If our DuplicateObject didn't stomp the old asset because TargetAssetName != ExistingAsset->GetName(),
 			// mark the existing asset as garbage so that we collect it when we run GC after the import
-			if ( ExistingAsset->GetName() != TargetAssetName )
+			if ( ExistingAsset && ExistingAsset->GetName() != TargetAssetName )
 			{
 				ExistingAsset->MarkAsGarbage();
 				ImportContext.bNeedsGarbageCollection = true;
@@ -749,81 +766,137 @@ namespace UsdStageImporterImpl
 	// Move imported assets from transient folder to their final package, updating AssetCache to point to the moved assets
 	void PublishAssets(FUsdStageImportContext& ImportContext, const TSet<UObject*>& AssetsToPublish, TMap<UObject*, UObject*>& ObjectsToRemap, TMap<FSoftObjectPath, FSoftObjectPath>& SoftObjectsToRemap )
 	{
+		TArray<UObject*> Textures;
+		TArray<UObject*> Materials;
+		TArray<UObject*> StaticMeshes;
+		TArray<UObject*> PhysicsAssets;
+		TArray<UObject*> Skeletons;
+		TArray<UObject*> SkeletalMeshes;
+		TArray<UObject*> AnimSequences;
+		TArray<UObject*> AnimBlueprints;
+		TArray<UObject*> LevelSequences;
+
 		TSet<FString> UniqueAssetNames;
+		TMap<UObject*, FString> AssetToContentFolder;
 
-		for ( UObject* Asset : AssetsToPublish )
+		for (UObject* Asset : AssetsToPublish)
 		{
-			if ( !Asset )
+			if (!Asset)
 			{
 				continue;
 			}
 
-			if ( Asset->IsA( UGeometryCache::StaticClass() ) )
+			const static FString TexturesFolder       = TEXT("Textures");
+			const static FString MaterialsFolder      = TEXT("Materials");
+			const static FString StaticMeshesFolder   = TEXT("StaticMeshes");
+			const static FString SkeletalMeshesFolder = TEXT("SkeletalMeshes");
+			const static FString LevelSequencesFolder = TEXT("LevelSequences");
+
+			const FString* AssetTypeFolderPtr = nullptr;
+			if (Asset->IsA(UGeometryCache::StaticClass()))
+			{
+				UE_LOG(LogUsd, Warning, TEXT("Ignoring asset '%s': Importing GeometryCaches assets from USD is not supported at this time"), *Asset->GetName());
+				continue;
+			}
+			else if (Asset->IsA(UMaterialInterface::StaticClass()))
+			{
+				AssetTypeFolderPtr = &MaterialsFolder;
+				Materials.Add(Asset);
+			}
+			else if (Asset->IsA(UStaticMesh::StaticClass()))
+			{
+				AssetTypeFolderPtr = &StaticMeshesFolder;
+				StaticMeshes.Add(Asset);
+			}
+			else if (Asset->IsA(UTexture::StaticClass()))
+			{
+				AssetTypeFolderPtr = &TexturesFolder;
+				Textures.Add(Asset);
+			}
+			else if (Asset->IsA(USkeletalMesh::StaticClass()))
+			{
+				AssetTypeFolderPtr = &SkeletalMeshesFolder;
+				SkeletalMeshes.Add(Asset);
+			}
+			else if (Asset->IsA(USkeleton::StaticClass()))
+			{
+				AssetTypeFolderPtr = &SkeletalMeshesFolder;
+				Skeletons.Add(Asset);
+			}
+			else if (Asset->IsA(UPhysicsAsset::StaticClass()))
+			{
+				AssetTypeFolderPtr = &SkeletalMeshesFolder;
+				PhysicsAssets.Add(Asset);
+			}
+			else if (Asset->IsA(UAnimSequence::StaticClass()))
+			{
+				AssetTypeFolderPtr = &SkeletalMeshesFolder;
+				AnimSequences.Add(Asset);
+			}
+			else if (Asset->IsA(UAnimBlueprint::StaticClass()))
+			{
+				AssetTypeFolderPtr = &SkeletalMeshesFolder;
+				AnimBlueprints.Add(Asset);
+			}
+			else if (Asset->IsA(ULevelSequence::StaticClass()))
+			{
+				AssetTypeFolderPtr = &LevelSequencesFolder;
+				LevelSequences.Add(Asset);
+			}
+			else
 			{
-				UE_LOG( LogUsd, Warning, TEXT( "Ignoring asset '%s': Importing GeometryCaches assets from USD is not supported at this time" ), *Asset->GetName() );
+				// We don't know what to do with this asset
+				ensure(false);
 				continue;
 			}
 
-			FString AssetTypeFolder;
-			if ( ImportContext.ImportOptions->bPrimPathFolderStructure )
+			if (ImportContext.ImportOptions->bPrimPathFolderStructure)
 			{
-				if ( UUsdAssetImportData* ImportData = UsdUtils::GetAssetImportData( Asset ) )
+				if (UUsdAssetImportData* ImportData = UsdUtils::GetAssetImportData(Asset))
 				{
 					// For skeletal stuff, the primpaths point to the SkelRoot, so it is useful to place the assets in there,
 					// as we'll always have at least the skeletal mesh and the skeleton
-					if ( Asset->IsA( USkeletalMesh::StaticClass() ) ||
-						 Asset->IsA( USkeleton::StaticClass() ) ||
-						 Asset->IsA( UPhysicsAsset::StaticClass() ) ||
-						 Asset->IsA( UAnimSequence::StaticClass() ) ||
-						 Asset->IsA( UAnimBlueprint::StaticClass() )
-					)
+					if (AssetTypeFolderPtr == &SkeletalMeshesFolder)
 					{
-						AssetTypeFolder = ImportData->PrimPath;
+						AssetToContentFolder.Add(Asset, ImportData->PrimPath);
 					}
 					else
 					{
-						AssetTypeFolder = FPaths::GetPath( ImportData->PrimPath );
+						AssetToContentFolder.Add(Asset, FPaths::GetPath(ImportData->PrimPath));
 					}
 				}
 			}
 			else
 			{
-				if (Asset->IsA(UMaterialInterface::StaticClass()))
-				{
-					AssetTypeFolder = TEXT("Materials");
-				}
-				else if (Asset->IsA(UStaticMesh::StaticClass()))
-				{
-					AssetTypeFolder = TEXT("StaticMeshes");
-				}
-				else if (Asset->IsA(UGeometryCache::StaticClass()))
-				{
-					AssetTypeFolder = TEXT("GeometryCaches");
-				}
-				else if (Asset->IsA(UTexture::StaticClass()))
-				{
-					AssetTypeFolder = TEXT("Textures");
-				}
-				else if (
-					Asset->IsA( USkeletalMesh::StaticClass() ) ||
-					Asset->IsA( USkeleton::StaticClass() ) ||
-					Asset->IsA( UPhysicsAsset::StaticClass() ) ||
-					Asset->IsA( UAnimSequence::StaticClass() ) ||
-					Asset->IsA( UAnimBlueprint::StaticClass() )
-				)
-				{
-					AssetTypeFolder = TEXT("SkeletalMeshes");
-				}
-				else if ( Asset->IsA( ULevelSequence::StaticClass() ) )
+				AssetToContentFolder.Add(Asset, *AssetTypeFolderPtr);
+			}
+		}
+
+		TFunction<void(const TArray<UObject*>&)> PublishAssetType = [&](const TArray<UObject*>& Assets)
+		{
+			for (UObject* Asset : Assets)
+			{
+				if (FString* AssetTypeFolder = AssetToContentFolder.Find(Asset))
 				{
-					AssetTypeFolder = TEXT("LevelSequences");
+					FString TargetAssetName = GetUserFriendlyName(Asset, UniqueAssetNames);
+					FString DestPackagePath = FPaths::Combine(ImportContext.PackagePath, *AssetTypeFolder, TargetAssetName);
+					PublishAsset(ImportContext, Asset, DestPackagePath, ObjectsToRemap, SoftObjectsToRemap);
 				}
 			}
+		};
 
-			FString TargetAssetName = GetUserFriendlyName(Asset, UniqueAssetNames);
-			FString DestPackagePath = FPaths::Combine(ImportContext.PackagePath, AssetTypeFolder, TargetAssetName);
-			PublishAsset(ImportContext, Asset, DestPackagePath, ObjectsToRemap, SoftObjectsToRemap);
-		}
+		// We should publish them in the order of dependents before dependencies, because publishing will involve
+		// taking these assets from the asset cache, which can only happen in that order. If we don't follow this order
+		// we may end up unnecessarily duplicating assets
+		PublishAssetType(LevelSequences);
+		PublishAssetType(AnimBlueprints);
+		PublishAssetType(AnimSequences);
+		PublishAssetType(SkeletalMeshes);
+		PublishAssetType(Skeletons);
+		PublishAssetType(PhysicsAssets);
+		PublishAssetType(StaticMeshes);
+		PublishAssetType(Materials);
+		PublishAssetType(Textures);
 	}
 
 	void ResolveComponentConflict(USceneComponent* NewRoot, USceneComponent* ExistingRoot, EReplaceActorPolicy ReplacePolicy, TMap<UObject*, UObject*>& ObjectsToRemap, TMap<FSoftObjectPath, FSoftObjectPath>& SoftObjectsToRemap )
@@ -1202,7 +1275,7 @@ namespace UsdStageImporterImpl
 	}
 
 	/**
-	 * FUsdAssetCache can track which assets are requested/added to itself during translation, but it may miss some dependencies
+	 * UUsdAssetCache2 can track which assets are requested/added to itself during translation, but it may miss some dependencies
 	 * that are only retrieved/added themselves when the original asset is first parsed. This function recursively collects all of those.
 	 * Example: An UMaterialInstance is already in the cache, so when translating we just retrieve the existing asset --> The textures that it's using won't be retrieved or marked as "Used"
 	 * Example: An USkeletalMesh is already in the cache, so in the same way we would miss its USkeleton, materials and textures of those materials
@@ -1241,124 +1314,13 @@ namespace UsdStageImporterImpl
 
 				// Only add it as a dependency if it's owned by the asset cache, but still traverse it because
 				// we may be in some strange situation where the material shouldn't be in this list, but one of its used textures should
-				if ( ImportContext.AssetCache->IsAssetOwnedByCache( Asset ) )
+				if ( Asset && ImportContext.AssetCache->IsAssetOwnedByCache( Asset->GetPathName() ) )
 				{
 					OutAssetsAndDependencies.Add( Asset );
 				}
 
-				if ( UMaterial* Material = Cast<UMaterial>( Asset ) )
-				{
-					TArray<UTexture*> UsedTextures;
-					const bool bAllQualityLevels = true;
-					const bool bAllFeatureLevels = true;
-					Material->GetUsedTextures( UsedTextures, EMaterialQualityLevel::High, bAllQualityLevels, ERHIFeatureLevel::SM5, bAllFeatureLevels );
-
-					OutAssetsAndDependencies.Reserve( OutAssetsAndDependencies.Num() + UsedTextures.Num() );
-					for ( UTexture* UsedTexture : UsedTextures )
-					{
-						if ( ImportContext.AssetCache->IsAssetOwnedByCache( UsedTexture ) )
-						{
-							OutAssetsAndDependencies.Add( UsedTexture );
-						}
-					}
-				}
-				else if ( UMaterialInstance* MaterialInstance = Cast<UMaterialInstance>( Asset ) )
-				{
-					OutAssetsAndDependencies.Reserve( OutAssetsAndDependencies.Num() + MaterialInstance->TextureParameterValues.Num() );
-					for ( const FTextureParameterValue& TextureValue : MaterialInstance->TextureParameterValues )
-					{
-						if ( UTexture* Texture = TextureValue.ParameterValue )
-						{
-							if ( ImportContext.AssetCache->IsAssetOwnedByCache( Texture ) )
-							{
-								OutAssetsAndDependencies.Add( Texture );
-							}
-						}
-					}
-				}
-				else if ( USkeletalMesh* SkeletalMesh = Cast<USkeletalMesh>( Asset ) )
-				{
-					if ( USkeleton* Skeleton = SkeletalMesh->GetSkeleton() )
-					{
-						if ( ImportContext.AssetCache->IsAssetOwnedByCache( Skeleton ) )
-						{
-							OutAssetsAndDependencies.Add( Skeleton );
-						}
-					}
-
-					if ( UPhysicsAsset* PhysicsAsset = SkeletalMesh->GetPhysicsAsset() )
-					{
-						if ( ImportContext.AssetCache->IsAssetOwnedByCache( PhysicsAsset ) )
-						{
-							OutAssetsAndDependencies.Add( PhysicsAsset );
-						}
-					}
-
-					AssetQueue.Reserve( AssetQueue.Num() + SkeletalMesh->GetMaterials().Num() );
-					for ( const FSkeletalMaterial& SkeletalMaterial : SkeletalMesh->GetMaterials() )
-					{
-						if ( UMaterialInterface* UsedMaterial = SkeletalMaterial.MaterialInterface )
-						{
-							AssetQueue.Add( UsedMaterial );
-						}
-					}
-				}
-				else if ( UStaticMesh* StaticMesh = Cast<UStaticMesh>( Asset ) )
-				{
-					AssetQueue.Reserve( AssetQueue.Num() + StaticMesh->GetStaticMaterials().Num() );
-					for ( const FStaticMaterial& StaticMaterial : StaticMesh->GetStaticMaterials() )
-					{
-						if ( UMaterialInterface* UsedMaterial = StaticMaterial.MaterialInterface )
-						{
-							AssetQueue.Add( UsedMaterial );
-						}
-					}
-				}
-				else if ( UGeometryCache* GeometryCache = Cast<UGeometryCache>( Asset ) )
-				{
-					for ( UMaterialInterface* UsedMaterial : GeometryCache->Materials )
-					{
-						if ( UsedMaterial )
-						{
-							AssetQueue.Add( UsedMaterial );
-						}
-					}
-				}
-				else if ( UAnimSequence* AnimSequence = Cast<UAnimSequence>( Asset ) )
-				{
-					if ( USkeletalMesh* Mesh = AnimSequence->GetPreviewMesh() )
-					{
-						AssetQueue.Add( Mesh );
-					}
-
-					if ( USkeleton* Skeleton = AnimSequence->GetSkeleton() )
-					{
-						if ( ImportContext.AssetCache->IsAssetOwnedByCache( Skeleton ) )
-						{
-							OutAssetsAndDependencies.Add( Skeleton );
-						}
-					}
-				}
-				else if ( UTexture* Texture = Cast<UTexture>( Asset ) )
-				{
-					// Do nothing. Textures have no additional dependencies
-				}
-				else if ( USkeleton* Skeleton = Cast<USkeleton>( Asset ) )
-				{
-					// Do nothing. Skeletons have no additional dependencies
-				}
-				else if ( UPhysicsAsset* PhysicsAsset = Cast<UPhysicsAsset>( Asset ) )
-				{
-					// Do nothing. PhysicsAssets have no additional dependencies
-				}
-				else if ( UAnimBlueprint* AnimBP = Cast<UAnimBlueprint>( Asset ) )
-				{
-					// Do nothing. AnimBlueprints have no additional dependencies
-				}
-				else
-				{
-					UE_LOG( LogUsd, Warning, TEXT( "Unknown asset '%s' encountered when collecting used assets before USD import." ), Asset ? *Asset->GetName() : TEXT( "nullptr" ) );
-				}
+				TSet<UObject*> Dependencies = IUsdClassesModule::GetAssetDependencies(Asset);
+				AssetQueue.Append(Dependencies.Array());
 			}
 		}
 	}
@@ -1706,11 +1668,21 @@ void UUsdStageImporter::ImportFromFile(FUsdStageImportContext& ImportContext)
 	// Ensure a valid asset cache
 	if ( !ImportContext.AssetCache )
 	{
-		ImportContext.AssetCache = NewObject<UUsdAssetCache>();
+		UE_LOG(LogUsd, Log, TEXT("Generating a temporary USD Asset Cache when importing '%s'."),
+			*ImportContext.FilePath
+		);
+		ImportContext.AssetCache = NewObject<UUsdAssetCache2>();
+	}
+	if (!ImportContext.InfoCache)
+	{
+		UE_LOG(LogUsd, Log, TEXT("Generating a temporary USD Info Cache when importing '%s'."),
+			*ImportContext.FilePath
+		);
+		ImportContext.InfoCache = MakeShared<FUsdInfoCache>();
 	}
 	ImportContext.AssetCache->MarkAssetsAsStale();
-	ImportContext.LevelSequenceHelper.SetAssetCache( ImportContext.AssetCache );
-	ImportContext.LevelSequenceHelper.Init( ImportContext.Stage );  // Must happen after the context gets an AssetCache!
+	ImportContext.LevelSequenceHelper.SetInfoCache(ImportContext.InfoCache);
+	ImportContext.LevelSequenceHelper.Init( ImportContext.Stage );  // Must happen after the context gets an InfoCache!
 
 	// Shotgun approach to recreate all render states because we may want to reimport/delete/reassign a material/static/skeletalmesh while it is currently being drawn
 	FGlobalComponentRecreateRenderStateContext RecreateRenderStateContext;
@@ -1812,9 +1784,21 @@ bool UUsdStageImporter::ReimportSingleAsset(FUsdStageImportContext& ImportContex
 	// Ensure a valid asset cache
 	if ( !ImportContext.AssetCache )
 	{
-		ImportContext.AssetCache = NewObject<UUsdAssetCache>();
+		UE_LOG(LogUsd, Log, TEXT("Generating a temporary USD Asset Cache when importing '%s'."),
+			*ImportContext.FilePath
+		);
+		ImportContext.AssetCache = NewObject<UUsdAssetCache2>();
+	}
+	if (!ImportContext.InfoCache)
+	{
+		UE_LOG(LogUsd, Log, TEXT("Generating a temporary USD Info Cache when importing '%s'."),
+			*ImportContext.FilePath
+		);
+		ImportContext.InfoCache = MakeShared<FUsdInfoCache>();
 	}
 	ImportContext.AssetCache->MarkAssetsAsStale();
+	ImportContext.LevelSequenceHelper.SetInfoCache(ImportContext.InfoCache);
+	ImportContext.LevelSequenceHelper.Init(ImportContext.Stage);  // Must happen after the context gets an InfoCache!
 
 	// Shotgun approach to recreate all render states because we may want to reimport/delete/reassign a material/static/skeletalmesh while it is currently being drawn
 	FGlobalComponentRecreateRenderStateContext RecreateRenderStateContext;
diff --git a/Engine/Plugins/Importers/USDImporter/Source/USDStageImporter/Public/USDStageImportContext.h b/Engine/Plugins/Importers/USDImporter/Source/USDStageImporter/Public/USDStageImportContext.h
--- a/Engine/Plugins/Importers/USDImporter/Source/USDStageImporter/Public/USDStageImportContext.h
+++ b/Engine/Plugins/Importers/USDImporter/Source/USDStageImporter/Public/USDStageImportContext.h
@@ -15,7 +15,7 @@
 
 #include "USDStageImportContext.generated.h"
 
-class UUsdAssetCache;
+class UUsdAssetCache2;
 class UUsdStageImportOptions;
 
 USTRUCT()
@@ -58,7 +58,7 @@ struct USDSTAGEIMPORTER_API FUsdStageImportContext
 	FUsdLevelSequenceHelper LevelSequenceHelper;
 
 	UPROPERTY()
-	TObjectPtr<UUsdAssetCache> AssetCache;
+	TObjectPtr<UUsdAssetCache2> AssetCache;
 
 	/** Caches various information about prims that are expensive to query */
 	TSharedPtr<FUsdInfoCache> InfoCache;
diff --git a/Engine/Plugins/Importers/USDImporter/Source/USDTests/Private/USDTestsBlueprintLibrary.cpp b/Engine/Plugins/Importers/USDImporter/Source/USDTests/Private/USDTestsBlueprintLibrary.cpp
--- a/Engine/Plugins/Importers/USDImporter/Source/USDTests/Private/USDTestsBlueprintLibrary.cpp
+++ b/Engine/Plugins/Importers/USDImporter/Source/USDTests/Private/USDTestsBlueprintLibrary.cpp
@@ -8,6 +8,7 @@
 
 #include "CoreMinimal.h"
 #include "Editor.h"
+#include "Editor/Transactor.h"
 #include "Kismet2/BlueprintEditorUtils.h"
 #include "Kismet2/KismetEditorUtilities.h"
 #include "UObject/ObjectMacros.h"
@@ -145,3 +146,14 @@ void USDTestsBlueprintLibrary::SetUsdStageCpp(AUsdStageActor* StageActor, const
 	UE::FUsdStage NewStage = UnrealUSDWrapper::OpenStage(*NewStageRootLayer, EUsdInitialLoadSet::LoadAll);
 	StageActor->SetUsdStage(NewStage);
 }
+
+void USDTestsBlueprintLibrary::ClearTransactionHistory()
+{
+	if (GEditor)
+	{
+		if (UTransactor* Transactor = GEditor->Trans)
+		{
+			Transactor->Reset(NSLOCTEXT("USDTests", "ClearTransactionHistoryReason", "USDTestsBlueprintLibrary::ClearTransactionHistory was called"));
+		}
+	}
+}
diff --git a/Engine/Plugins/Importers/USDImporter/Source/USDTests/Public/USDTestsBlueprintLibrary.h b/Engine/Plugins/Importers/USDImporter/Source/USDTests/Public/USDTestsBlueprintLibrary.h
--- a/Engine/Plugins/Importers/USDImporter/Source/USDTests/Public/USDTestsBlueprintLibrary.h
+++ b/Engine/Plugins/Importers/USDImporter/Source/USDTests/Public/USDTestsBlueprintLibrary.h
@@ -52,4 +52,10 @@ public:
 	 */
 	UFUNCTION(BlueprintCallable, Category = "Counts")
 	static void SetUsdStageCpp(AUsdStageActor* StageActor, const FString& NewStageRootLayer);
+
+	/**
+	 * Clears the editor transaction history (useful to be run before GC, to check for otherwise unreferenced assets)
+	 */
+	UFUNCTION(BlueprintCallable, Category = "Transactions")
+	static void ClearTransactionHistory();
 };
diff --git a/Engine/Plugins/Importers/USDImporter/Source/USDUtilities/Private/USDLightConversion.cpp b/Engine/Plugins/Importers/USDImporter/Source/USDUtilities/Private/USDLightConversion.cpp
--- a/Engine/Plugins/Importers/USDImporter/Source/USDUtilities/Private/USDLightConversion.cpp
+++ b/Engine/Plugins/Importers/USDImporter/Source/USDUtilities/Private/USDLightConversion.cpp
@@ -2,8 +2,7 @@
 
 #include "USDLightConversion.h"
 
-#include "Misc/Paths.h"
-#include "USDAssetCache.h"
+#include "USDAssetCache2.h"
 #include "USDAttributeUtils.h"
 #include "USDConversionUtils.h"
 #include "USDLayerUtils.h"
@@ -25,6 +24,7 @@
 #include "Components/SpotLightComponent.h"
 #include "EditorFramework/AssetImportData.h"
 #include "Engine/TextureCube.h"
+#include "Misc/Paths.h"
 
 #if USE_USD_SDK
 
@@ -215,7 +215,7 @@ bool UsdToUnreal::ConvertSphereLight( const pxr::UsdPrim& Prim, UPointLightCompo
 	return true;
 }
 
-bool UsdToUnreal::ConvertDomeLight( const pxr::UsdPrim& Prim, USkyLightComponent& LightComponent, UUsdAssetCache* TexturesCache )
+bool UsdToUnreal::ConvertDomeLight( const pxr::UsdPrim& Prim, USkyLightComponent& LightComponent, UUsdAssetCache2* TexturesCache )
 {
 	FScopedUsdAllocs UsdAllocs;
 
@@ -293,6 +293,12 @@ bool UsdToUnreal::ConvertLuxShapingAPI( const pxr::UsdPrim& Prim, USpotLightComp
 	return true;
 }
 
+bool UsdToUnreal::ConvertDomeLight(const pxr::UsdPrim& Prim, USkyLightComponent& LightComponent, UUsdAssetCache* TexturesCache)
+{
+	UUsdAssetCache2* NewCache = nullptr;
+	return UsdToUnreal::ConvertDomeLight(Prim, LightComponent, NewCache);
+}
+
 float UsdToUnreal::ConvertLightIntensityAttr( float UsdIntensity, float UsdExposure )
 {
 	return UsdIntensity * FMath::Exp2( UsdExposure );
diff --git a/Engine/Plugins/Importers/USDImporter/Source/USDUtilities/Private/USDShadeConversion.cpp b/Engine/Plugins/Importers/USDImporter/Source/USDUtilities/Private/USDShadeConversion.cpp
--- a/Engine/Plugins/Importers/USDImporter/Source/USDUtilities/Private/USDShadeConversion.cpp
+++ b/Engine/Plugins/Importers/USDImporter/Source/USDUtilities/Private/USDShadeConversion.cpp
@@ -6,6 +6,7 @@
 #if USE_USD_SDK
 
 #include "USDAssetCache.h"
+#include "USDAssetCache2.h"
 #include "USDAssetImportData.h"
 #include "USDAttributeUtils.h"
 #include "USDConversionUtils.h"
@@ -488,7 +489,7 @@ namespace UE
 			using FParameterValue = TVariant<float, FVector, FTextureParameterValue, FPrimvarReaderParameterValue, bool>;
 
 			bool GetTextureParameterValue( pxr::UsdShadeInput& ShadeInput, TextureGroup LODGroup, FParameterValue& OutValue,
-				UMaterialInterface* Material = nullptr, UUsdAssetCache* TexturesCache = nullptr, TMap<FString, int32>* PrimvarToUVIndex = nullptr )
+				UMaterialInterface* Material = nullptr, UUsdAssetCache2* TexturesCache = nullptr, TMap<FString, int32>* PrimvarToUVIndex = nullptr )
 			{
 				FScopedUsdAllocs UsdAllocs;
 
@@ -733,7 +734,7 @@ namespace UE
 			}
 
 			bool GetFloatParameterValue( pxr::UsdShadeConnectableAPI& Connectable, const pxr::TfToken& InputName, float DefaultValue, FParameterValue& OutValue,
-				UMaterialInterface* Material = nullptr, UUsdAssetCache* TexturesCache = nullptr, TMap<FString, int32>* PrimvarToUVIndex = nullptr)
+				UMaterialInterface* Material = nullptr, UUsdAssetCache2* TexturesCache = nullptr, TMap<FString, int32>* PrimvarToUVIndex = nullptr)
 			{
 				FScopedUsdAllocs Allocs;
 
@@ -852,7 +853,7 @@ namespace UE
 			}
 
 			bool GetVec3ParameterValue( pxr::UsdShadeConnectableAPI& Connectable, const pxr::TfToken& InputName, const FLinearColor& DefaultValue, FParameterValue& OutValue,
-				bool bIsNormalMap = false, UMaterialInterface* Material = nullptr, UUsdAssetCache* TexturesCache = nullptr, TMap<FString, int32>* PrimvarToUVIndex = nullptr )
+				bool bIsNormalMap = false, UMaterialInterface* Material = nullptr, UUsdAssetCache2* TexturesCache = nullptr, TMap<FString, int32>* PrimvarToUVIndex = nullptr )
 			{
 				FScopedUsdAllocs Allocs;
 
@@ -922,7 +923,7 @@ namespace UE
 				return true;
 			}
 
-			bool GetBoolParameterValue( pxr::UsdShadeConnectableAPI& Connectable, const pxr::TfToken& InputName, bool DefaultValue, FParameterValue& OutValue, UMaterialInterface* Material = nullptr, UUsdAssetCache* TexturesCache = nullptr, TMap<FString, int32>* PrimvarToUVIndex = nullptr)
+			bool GetBoolParameterValue( pxr::UsdShadeConnectableAPI& Connectable, const pxr::TfToken& InputName, bool DefaultValue, FParameterValue& OutValue, UMaterialInterface* Material = nullptr, UUsdAssetCache2* TexturesCache = nullptr, TMap<FString, int32>* PrimvarToUVIndex = nullptr)
 			{
 				FScopedUsdAllocs Allocs;
 
@@ -1258,17 +1259,12 @@ namespace UE
 					const FString ResolvedTexturePath = UsdUtils::GetResolvedTexturePath( TextureAssetPathAttr );
 					if ( !ResolvedTexturePath.IsEmpty() )
 					{
-						EObjectFlags ObjectFlags = RF_Transactional;
+						EObjectFlags ObjectFlags = RF_Transactional | RF_Transient;
 						if ( !Outer )
 						{
 							Outer = GetTransientPackage();
 						}
 
-						if ( Outer == GetTransientPackage() )
-						{
-							ObjectFlags = ObjectFlags | RF_Transient;
-						}
-
 						FString TextureExtension;
 						if ( IsInsideUsdzArchive( ResolvedTexturePath, TextureExtension ) )
 						{
@@ -1982,10 +1978,11 @@ namespace UsdShadeConversionImpl = UE::UsdShadeConversion::Private;
 bool UsdToUnreal::ConvertMaterial( const pxr::UsdShadeMaterial& UsdShadeMaterial, UMaterialInstance& Material )
 {
 	TMap<FString, int32> PrimvarToUVIndex;
-	return ConvertMaterial( UsdShadeMaterial, Material, nullptr, PrimvarToUVIndex );
+	UUsdAssetCache2* NewCache = nullptr;
+	return ConvertMaterial( UsdShadeMaterial, Material, NewCache, PrimvarToUVIndex );
 }
 
-bool UsdToUnreal::ConvertMaterial( const pxr::UsdShadeMaterial& UsdShadeMaterial, UMaterialInstance& Material, UUsdAssetCache* TexturesCache, TMap< FString, int32 >& PrimvarToUVIndex, const TCHAR* RenderContext )
+bool UsdToUnreal::ConvertMaterial( const pxr::UsdShadeMaterial& UsdShadeMaterial, UMaterialInstance& Material, UUsdAssetCache2* TexturesCache, TMap< FString, int32 >& PrimvarToUVIndex, const TCHAR* RenderContext )
 {
 	FScopedUsdAllocs UsdAllocs;
 
@@ -2142,10 +2139,11 @@ bool UsdToUnreal::ConvertMaterial( const pxr::UsdShadeMaterial& UsdShadeMaterial
 bool UsdToUnreal::ConvertMaterial( const pxr::UsdShadeMaterial& UsdShadeMaterial, UMaterial& Material )
 {
 	TMap<FString, int32> PrimvarToUVIndex;
-	return ConvertMaterial( UsdShadeMaterial, Material, nullptr, PrimvarToUVIndex );
+	UUsdAssetCache2* NewCache = nullptr;
+	return ConvertMaterial( UsdShadeMaterial, Material, NewCache, PrimvarToUVIndex );
 }
 
-bool UsdToUnreal::ConvertMaterial( const pxr::UsdShadeMaterial& UsdShadeMaterial, UMaterial& Material, UUsdAssetCache* TexturesCache, TMap< FString, int32 >& PrimvarToUVIndex, const TCHAR* RenderContext )
+bool UsdToUnreal::ConvertMaterial( const pxr::UsdShadeMaterial& UsdShadeMaterial, UMaterial& Material, UUsdAssetCache2* TexturesCache, TMap< FString, int32 >& PrimvarToUVIndex, const TCHAR* RenderContext )
 {
 	bool bHasMaterialInfo = false;
 
@@ -2308,7 +2306,7 @@ bool UsdToUnreal::ConvertMaterial( const pxr::UsdShadeMaterial& UsdShadeMaterial
 	return bHasMaterialInfo;
 }
 
-bool UsdToUnreal::ConvertShadeInputsToParameters( const pxr::UsdShadeMaterial& UsdShadeMaterial, UMaterialInstance& MaterialInstance, UUsdAssetCache* TexturesCache, const TCHAR* RenderContext )
+bool UsdToUnreal::ConvertShadeInputsToParameters( const pxr::UsdShadeMaterial& UsdShadeMaterial, UMaterialInstance& MaterialInstance, UUsdAssetCache2* TexturesCache, const TCHAR* RenderContext )
 {
 	FScopedUsdAllocs UsdAllocs;
 
@@ -2383,6 +2381,26 @@ bool UsdToUnreal::ConvertShadeInputsToParameters( const pxr::UsdShadeMaterial& U
 	return true;
 }
 
+PRAGMA_DISABLE_DEPRECATION_WARNINGS
+bool UsdToUnreal::ConvertMaterial(const pxr::UsdShadeMaterial& UsdShadeMaterial, UMaterialInstance& Material, UUsdAssetCache* TexturesCache, TMap<FString, int32>& PrimvarToUVIndex, const TCHAR* RenderContext)
+{
+	UUsdAssetCache2* NewCache = nullptr;
+	return UsdToUnreal::ConvertMaterial(UsdShadeMaterial, Material, NewCache, PrimvarToUVIndex, RenderContext);
+}
+
+bool UsdToUnreal::ConvertMaterial(const pxr::UsdShadeMaterial& UsdShadeMaterial, UMaterial& Material, UUsdAssetCache* TexturesCache, TMap<FString, int32>& PrimvarToUVIndex, const TCHAR* RenderContext)
+{
+	UUsdAssetCache2* NewCache = nullptr;
+	return UsdToUnreal::ConvertMaterial(UsdShadeMaterial, Material, NewCache, PrimvarToUVIndex, RenderContext);
+}
+
+bool UsdToUnreal::ConvertShadeInputsToParameters(const pxr::UsdShadeMaterial& UsdShadeMaterial, UMaterialInstance& MaterialInstance, UUsdAssetCache* TexturesCache, const TCHAR* RenderContext)
+{
+	UUsdAssetCache2* NewCache = nullptr;
+	return UsdToUnreal::ConvertShadeInputsToParameters(UsdShadeMaterial, MaterialInstance, NewCache, RenderContext);
+}
+PRAGMA_ENABLE_DEPRECATION_WARNINGS
+
 #if WITH_EDITOR
 bool UnrealToUsd::ConvertMaterialToBakedSurface(
 	const UMaterialInterface& InMaterial,
diff --git a/Engine/Plugins/Importers/USDImporter/Source/USDUtilities/Private/USDSkeletalDataConversion.cpp b/Engine/Plugins/Importers/USDImporter/Source/USDUtilities/Private/USDSkeletalDataConversion.cpp
--- a/Engine/Plugins/Importers/USDImporter/Source/USDUtilities/Private/USDSkeletalDataConversion.cpp
+++ b/Engine/Plugins/Importers/USDImporter/Source/USDUtilities/Private/USDSkeletalDataConversion.cpp
@@ -2531,7 +2531,7 @@ USkeletalMesh* UsdToUnreal::GetSkeletalMeshFromImportData(
 	// For now, create a new USkeletalMesh
 	// Note: Remember to initialize UsedMorphTargetNames with existing morph targets, whenever the SkeletalMesh is reused
 	FName UniqueMeshName = MakeUniqueObjectName( GetTransientPackage(), USkeletalMesh::StaticClass(), MeshName );
-	USkeletalMesh* SkeletalMesh = NewObject<USkeletalMesh>(GetTransientPackage(), UniqueMeshName, ObjectFlags | EObjectFlags::RF_Public);
+	USkeletalMesh* SkeletalMesh = NewObject<USkeletalMesh>(GetTransientPackage(), UniqueMeshName, ObjectFlags | EObjectFlags::RF_Public | EObjectFlags::RF_Transient);
 
 	// Process reference skeleton from import data
 	int32 SkeletalDepth = 0;
@@ -2655,7 +2655,7 @@ USkeletalMesh* UsdToUnreal::GetSkeletalMeshFromImportData(
 
 	// Generate a Skeleton and associate it to the SkeletalMesh
 	FName UniqueSkeletonName = MakeUniqueObjectName( GetTransientPackage(), USkeleton::StaticClass(), SkeletonName );
-	USkeleton* Skeleton = NewObject<USkeleton>( GetTransientPackage(), UniqueSkeletonName, ObjectFlags | EObjectFlags::RF_Public );
+	USkeleton* Skeleton = NewObject<USkeleton>( GetTransientPackage(), UniqueSkeletonName, ObjectFlags | EObjectFlags::RF_Public | EObjectFlags::RF_Transient );
 	Skeleton->MergeAllBonesToBoneTree(SkeletalMesh);
 	Skeleton->SetPreviewMesh(SkeletalMesh);
 	SkeletalMesh->SetSkeleton(Skeleton);
diff --git a/Engine/Plugins/Importers/USDImporter/Source/USDUtilities/Public/USDLightConversion.h b/Engine/Plugins/Importers/USDImporter/Source/USDUtilities/Public/USDLightConversion.h
--- a/Engine/Plugins/Importers/USDImporter/Source/USDUtilities/Public/USDLightConversion.h
+++ b/Engine/Plugins/Importers/USDImporter/Source/USDUtilities/Public/USDLightConversion.h
@@ -29,6 +29,7 @@ template< typename T > class TfRefPtr;
 using UsdStageRefPtr = TfRefPtr< UsdStage >;
 PXR_NAMESPACE_CLOSE_SCOPE
 
+class UUsdAssetCache2;
 class UDirectionalLightComponent;
 class ULightComponentBase;
 class UPointLightComponent;
@@ -59,9 +60,14 @@ namespace UsdToUnreal
 	USDUTILITIES_API bool ConvertRectLight( const pxr::UsdPrim& Prim, URectLightComponent& LightComponent, double UsdTimeCode = UsdUtils::GetDefaultTimeCode() );
 	USDUTILITIES_API bool ConvertDiskLight( const pxr::UsdPrim& Prim, URectLightComponent& LightComponent, double UsdTimeCode = UsdUtils::GetDefaultTimeCode() );
 	USDUTILITIES_API bool ConvertSphereLight( const pxr::UsdPrim& Prim, UPointLightComponent& LightComponent, double UsdTimeCode = UsdUtils::GetDefaultTimeCode() );
-	USDUTILITIES_API bool ConvertDomeLight( const pxr::UsdPrim& Prim, USkyLightComponent& LightComponent, UUsdAssetCache* TexturesCache );
+	USDUTILITIES_API bool ConvertDomeLight( const pxr::UsdPrim& Prim, USkyLightComponent& LightComponent, UUsdAssetCache2* TexturesCache );
 	USDUTILITIES_API bool ConvertLuxShapingAPI( const pxr::UsdPrim& Prim, USpotLightComponent& LightComponent, double UsdTimeCode = UsdUtils::GetDefaultTimeCode() );
 
+PRAGMA_DISABLE_DEPRECATION_WARNINGS
+	UE_DEPRECATED(5.2, "Use the other overload that receives an UUsdAssetCache2 object instead")
+	USDUTILITIES_API bool ConvertDomeLight( const pxr::UsdPrim& Prim, USkyLightComponent& LightComponent, UUsdAssetCache* TexturesCache );
+PRAGMA_ENABLE_DEPRECATION_WARNINGS
+
 	// These are separately exposed so that they can be reused when reading data into MovieScene tracks.
 	// The other attribute conversions are mostly trivial like a single call to UsdToUnreal::ConvertDistance.
 	// Intensity return values are always in Lumen
diff --git a/Engine/Plugins/Importers/USDImporter/Source/USDUtilities/Public/USDShadeConversion.h b/Engine/Plugins/Importers/USDImporter/Source/USDUtilities/Public/USDShadeConversion.h
--- a/Engine/Plugins/Importers/USDImporter/Source/USDUtilities/Public/USDShadeConversion.h
+++ b/Engine/Plugins/Importers/USDImporter/Source/USDUtilities/Public/USDShadeConversion.h
@@ -26,14 +26,15 @@ PXR_NAMESPACE_OPEN_SCOPE
 PXR_NAMESPACE_CLOSE_SCOPE
 
 class FSHAHash;
-class UUsdAssetCache;
+class UUsdAssetCache2;
 class UMaterial;
 class UMaterialOptions;
 class UTexture;
-struct FFlattenMaterial;
-struct FPropertyEntry;
+class UUsdAssetCache;
 enum class EFlattenMaterialProperties : uint8;
 enum EMaterialProperty : int;
+struct FFlattenMaterial;
+struct FPropertyEntry;
 namespace UE
 {
 	class FUsdPrim;
@@ -51,9 +52,9 @@ namespace UsdToUnreal
 	 * @return Whether the conversion was successful or not.
 	 */
 	USDUTILITIES_API bool ConvertMaterial( const pxr::UsdShadeMaterial& UsdShadeMaterial, UMaterialInstance& Material );
-	USDUTILITIES_API bool ConvertMaterial( const pxr::UsdShadeMaterial& UsdShadeMaterial, UMaterialInstance& Material, UUsdAssetCache* TexturesCache, TMap< FString, int32 >& PrimvarToUVIndex, const TCHAR* RenderContext = nullptr );
+	USDUTILITIES_API bool ConvertMaterial( const pxr::UsdShadeMaterial& UsdShadeMaterial, UMaterialInstance& Material, UUsdAssetCache2* TexturesCache, TMap< FString, int32 >& PrimvarToUVIndex, const TCHAR* RenderContext = nullptr );
 	USDUTILITIES_API bool ConvertMaterial( const pxr::UsdShadeMaterial& UsdShadeMaterial, UMaterial& Material );
-	USDUTILITIES_API bool ConvertMaterial( const pxr::UsdShadeMaterial& UsdShadeMaterial, UMaterial& Material, UUsdAssetCache* TexturesCache, TMap< FString, int32 >& PrimvarToUVIndex, const TCHAR* RenderContext = nullptr );
+	USDUTILITIES_API bool ConvertMaterial( const pxr::UsdShadeMaterial& UsdShadeMaterial, UMaterial& Material, UUsdAssetCache2* TexturesCache, TMap< FString, int32 >& PrimvarToUVIndex, const TCHAR* RenderContext = nullptr );
 
 	/**
 	 * Attemps to assign the values of the surface shader inputs to the MaterialInstance parameters by matching the inputs display names to the parameters names.
@@ -64,7 +65,16 @@ namespace UsdToUnreal
 	 * @return Whether the conversion was successful or not.
 	 *
 	 */
+	USDUTILITIES_API bool ConvertShadeInputsToParameters( const pxr::UsdShadeMaterial& UsdShadeMaterial, UMaterialInstance& MaterialInstance, UUsdAssetCache2* TexturesCache, const TCHAR* RenderContext = nullptr );
+
+PRAGMA_DISABLE_DEPRECATION_WARNINGS
+	UE_DEPRECATED(5.2, "Use the other overload that receives an UUsdAssetCache2 object instead")
+	USDUTILITIES_API bool ConvertMaterial( const pxr::UsdShadeMaterial& UsdShadeMaterial, UMaterialInstance& Material, UUsdAssetCache* TexturesCache, TMap< FString, int32 >& PrimvarToUVIndex, const TCHAR* RenderContext = nullptr );
+	UE_DEPRECATED(5.2, "Use the other overload that receives an UUsdAssetCache2 object instead")
+	USDUTILITIES_API bool ConvertMaterial( const pxr::UsdShadeMaterial& UsdShadeMaterial, UMaterial& Material, UUsdAssetCache* TexturesCache, TMap< FString, int32 >& PrimvarToUVIndex, const TCHAR* RenderContext = nullptr );
+	UE_DEPRECATED(5.2, "Use the other overload that receives an UUsdAssetCache2 object instead")
 	USDUTILITIES_API bool ConvertShadeInputsToParameters( const pxr::UsdShadeMaterial& UsdShadeMaterial, UMaterialInstance& MaterialInstance, UUsdAssetCache* TexturesCache, const TCHAR* RenderContext = nullptr );
+PRAGMA_ENABLE_DEPRECATION_WARNINGS
 }
 
 #if WITH_EDITOR
diff --git a/Engine/Plugins/Importers/USDImporter/Source/UnrealUSDWrapper/Private/UsdWrappers/SdfPath.cpp b/Engine/Plugins/Importers/USDImporter/Source/UnrealUSDWrapper/Private/UsdWrappers/SdfPath.cpp
--- a/Engine/Plugins/Importers/USDImporter/Source/UnrealUSDWrapper/Private/UsdWrappers/SdfPath.cpp
+++ b/Engine/Plugins/Importers/USDImporter/Source/UnrealUSDWrapper/Private/UsdWrappers/SdfPath.cpp
@@ -311,4 +311,13 @@ namespace UE
 #endif // #if USE_USD_SDK
 		return Result;
 	}
+
+	bool FSdfPath::HasPrefix(const UE::FSdfPath& Prefix) const
+	{
+#if USE_USD_SDK
+		return Impl->PxrSdfPath.Get().HasPrefix(Prefix);
+#else
+		return false;
+#endif // #if USE_USD_SDK
+	}
 }
diff --git a/Engine/Plugins/Importers/USDImporter/Source/UnrealUSDWrapper/Public/UsdWrappers/SdfPath.h b/Engine/Plugins/Importers/USDImporter/Source/UnrealUSDWrapper/Public/UsdWrappers/SdfPath.h
--- a/Engine/Plugins/Importers/USDImporter/Source/UnrealUSDWrapper/Public/UsdWrappers/SdfPath.h
+++ b/Engine/Plugins/Importers/USDImporter/Source/UnrealUSDWrapper/Public/UsdWrappers/SdfPath.h
@@ -82,6 +82,8 @@ namespace UE
 
 		TArray<FSdfPath> GetPrefixes() const;
 
+		bool HasPrefix(const UE::FSdfPath& Prefix) const;
+
 	private:
 		TUniquePtr< Internal::FSdfPathImpl > Impl;
 	};
diff --git a/Engine/Plugins/Importers/USDImporter/USDImporter.uplugin b/Engine/Plugins/Importers/USDImporter/USDImporter.uplugin
--- a/Engine/Plugins/Importers/USDImporter/USDImporter.uplugin
+++ b/Engine/Plugins/Importers/USDImporter/USDImporter.uplugin
@@ -63,6 +63,11 @@
 			"Type": "Runtime",
 			"LoadingPhase": "Default"
 		},
+        {
+			"Name": "USDClassesEditor",
+			"Type": "Editor",
+			"LoadingPhase": "Default"
+		},
 		{
 			"Name": "GeometryCacheUSD",
 			"Type": "Editor",
