USD: Fix static analysis nullptr dereference warnings.

diff --git a/Engine/Plugins/Importers/USDImporter/Source/USDClasses/Private/USDAssetCache2.cpp b/Engine/Plugins/Importers/USDImporter/Source/USDClasses/Private/USDAssetCache2.cpp
--- a/Engine/Plugins/Importers/USDImporter/Source/USDClasses/Private/USDAssetCache2.cpp
+++ b/Engine/Plugins/Importers/USDImporter/Source/USDClasses/Private/USDAssetCache2.cpp
@@ -1491,8 +1491,11 @@ void UUsdAssetCache2::Serialize(FArchive& Ar)
 FUsdScopedAssetCacheReferencer::FUsdScopedAssetCacheReferencer(UUsdAssetCache2* InAssetCache, const UObject* Referencer)
 {
 	// For now we're assuming you can't nest these objects
-	ensure(InAssetCache && Referencer && !InAssetCache->CurrentScopedReferencer);
-	InAssetCache->CurrentScopedReferencer = Referencer;
+	if (ensure(InAssetCache))
+	{
+		ensure(Referencer && !InAssetCache->CurrentScopedReferencer);
+		InAssetCache->CurrentScopedReferencer = Referencer;
+	}
 
 	AssetCache = InAssetCache;
 }
