USD: Fix nullptr deref static analysis warning.

diff --git a/Engine/Plugins/Importers/USDImporter/Source/USDClasses/Private/USDAssetCache2.cpp b/Engine/Plugins/Importers/USDImporter/Source/USDClasses/Private/USDAssetCache2.cpp
--- a/Engine/Plugins/Importers/USDImporter/Source/USDClasses/Private/USDAssetCache2.cpp
+++ b/Engine/Plugins/Importers/USDImporter/Source/USDClasses/Private/USDAssetCache2.cpp
@@ -339,6 +339,11 @@ namespace UE::AssetCache::Private
 	// Copied from DataprepSnapshot.cpp, the "ReadSnapshotData" function
 	void DeserializeObjectAndSubObjects(UObject* Object, const TArray<uint8>& InBuffer)
 	{
+		if (!Object)
+		{
+			return;
+		}
+
 		// Remove all objects created by default that InObject is dependent on
 		// This method must obviously be called just after the InObject is created
 		auto RemoveDefaultDependencies = [](UObject* InObject)
