USD: Make remaining assets transient, as that's what's expected by the asset cache.

diff --git a/Engine/Plugins/Importers/USDImporter/Source/USDSchemas/Private/USDShadeMaterialTranslator.cpp b/Engine/Plugins/Importers/USDImporter/Source/USDSchemas/Private/USDShadeMaterialTranslator.cpp
--- a/Engine/Plugins/Importers/USDImporter/Source/USDSchemas/Private/USDShadeMaterialTranslator.cpp
+++ b/Engine/Plugins/Importers/USDImporter/Source/USDSchemas/Private/USDShadeMaterialTranslator.cpp
@@ -262,7 +262,7 @@ void FUsdShadeMaterialTranslator::CreateAssets()
 #if WITH_EDITOR
 		if ( GIsEditor ) // Also have to prevent Standalone game from going with MaterialInstanceConstants
 		{
-			if ( UMaterialInstanceConstant* NewMaterial = NewObject<UMaterialInstanceConstant>( GetTransientPackage(), InstanceName, Context->ObjectFlags ) )
+			if ( UMaterialInstanceConstant* NewMaterial = NewObject<UMaterialInstanceConstant>( GetTransientPackage(), InstanceName, Context->ObjectFlags | EObjectFlags::RF_Transient ) )
 			{
 				UUsdAssetImportData* ImportData = NewObject< UUsdAssetImportData >( NewMaterial, TEXT( "USDAssetImportData" ) );
 				ImportData->PrimPath = PrimPath.GetString();
diff --git a/Engine/Plugins/Importers/USDImporter/Source/USDSchemas/Private/USDSkelRootTranslator.cpp b/Engine/Plugins/Importers/USDImporter/Source/USDSchemas/Private/USDSkelRootTranslator.cpp
--- a/Engine/Plugins/Importers/USDImporter/Source/USDSchemas/Private/USDSkelRootTranslator.cpp
+++ b/Engine/Plugins/Importers/USDImporter/Source/USDSchemas/Private/USDSkelRootTranslator.cpp
@@ -1475,7 +1475,11 @@ namespace UsdSkelRootTranslatorImpl
 							// Bone/CurveCompressionSettings. Undoing that transaction would call UAnimSequence::Serialize with nullptr values for both, which crashes.
 							// Besides, this particular asset type is only ever created when we import to content folder assets (so never for realtime), and
 							// in that case we don't need it to be transactional anyway
-							AnimSequence = NewObject<UAnimSequence>( GetTransientPackage(), NAME_None, Context->ObjectFlags & ~EObjectFlags::RF_Transactional );
+							AnimSequence = NewObject<UAnimSequence>(
+								GetTransientPackage(),
+								NAME_None,
+								(Context->ObjectFlags & ~EObjectFlags::RF_Transactional) | EObjectFlags::RF_Transient
+							);
 							AnimSequence->SetSkeleton( SkeletalMesh->GetSkeleton() );
 
 							// This is read back in the USDImporter, so that if we ever import this AnimSequence we will always also import the SkeletalMesh for it
