USD: Improve InfoCache prim path to assets interface.

diff --git a/Engine/Plugins/Importers/USDImporter/Source/USDSchemas/Private/MeshTranslationImpl.cpp b/Engine/Plugins/Importers/USDImporter/Source/USDSchemas/Private/MeshTranslationImpl.cpp
--- a/Engine/Plugins/Importers/USDImporter/Source/USDSchemas/Private/MeshTranslationImpl.cpp
+++ b/Engine/Plugins/Importers/USDImporter/Source/USDSchemas/Private/MeshTranslationImpl.cpp
@@ -106,15 +106,13 @@ TMap<const UsdUtils::FUsdPrimMaterialSlot*, UMaterialInterface*> MeshTranslation
 					PrimPath += TwoSidedToken;
 				}
 
-				Material = Cast<UMaterialInterface>(InfoCache.GetSingleAssetForPrim(UE::FSdfPath{*PrimPath}, UMaterialInterface::StaticClass()));
+				Material = InfoCache.GetSingleAssetForPrim<UMaterialInterface>(UE::FSdfPath{*PrimPath});
 
 				// Need to create a two-sided material on-demand
 				if (!Material && Slot.bMeshIsDoubleSided)
 				{
 					// By now we parsed all materials so we must have the single-sided version of this material
-					UMaterialInstance* OneSidedMat = Cast<UMaterialInstance>(
-						InfoCache.GetSingleAssetForPrim(UE::FSdfPath{*Slot.MaterialSource}, UMaterialInterface::StaticClass())
-					);
+					UMaterialInstance* OneSidedMat = InfoCache.GetSingleAssetForPrim<UMaterialInstance>(UE::FSdfPath{*Slot.MaterialSource});
 					if (!OneSidedMat)
 					{
 						UE_LOG(LogUsd, Warning, TEXT("Failed to generate a two-sided material from the material prim at path '%s' as no single-sided material was generated for it."), *PrimPath);
diff --git a/Engine/Plugins/Importers/USDImporter/Source/USDSchemas/Private/USDGeomMeshTranslator.cpp b/Engine/Plugins/Importers/USDImporter/Source/USDSchemas/Private/USDGeomMeshTranslator.cpp
--- a/Engine/Plugins/Importers/USDImporter/Source/USDSchemas/Private/USDGeomMeshTranslator.cpp
+++ b/Engine/Plugins/Importers/USDImporter/Source/USDSchemas/Private/USDGeomMeshTranslator.cpp
@@ -417,9 +417,9 @@ namespace UsdGeomMeshTranslatorImpl
 						TUsdStore< pxr::UsdPrim > MaterialPrim = UsdPrim.GetStage()->GetPrimAtPath(MaterialPrimPath);
 						if (MaterialPrim.Get() && InfoCache)
 						{
-							Material = Cast<UMaterialInterface>(InfoCache->GetSingleAssetForPrim(
+							Material = InfoCache->GetSingleAssetForPrim<UMaterialInterface>(
 								UE::FSdfPath{MaterialPrim.Get().GetPrimPath()}
-							));
+							);
 						}
 					}
 
@@ -1570,10 +1570,9 @@ USceneComponent* FUsdGeomMeshTranslator::CreateComponents()
 	{
 		if (Context->InfoCache)
 		{
-			if (UStaticMesh* StaticMesh = Cast<UStaticMesh>(Context->InfoCache->GetSingleAssetForPrim(
-				PrimPath,
-				UStaticMesh::StaticClass()
-			)))
+			if (UStaticMesh* StaticMesh = Context->InfoCache->GetSingleAssetForPrim<UStaticMesh>(
+				PrimPath
+			))
 			{
 				TArray<UMaterialInterface*> ExistingAssignments;
 				for (FStaticMaterial& StaticMaterial : StaticMesh->GetStaticMaterials())
@@ -1603,10 +1602,9 @@ USceneComponent* FUsdGeomMeshTranslator::CreateComponents()
 
 		if (Context->InfoCache && Context->AssetCache)
 		{
-			if (UGeometryCache* GeometryCache = Cast< UGeometryCache >(Context->InfoCache->GetSingleAssetForPrim(
-				PrimPath,
-				UGeometryCache::StaticClass()
-			)))
+			if (UGeometryCache* GeometryCache = Context->InfoCache->GetSingleAssetForPrim<UGeometryCache>(
+				PrimPath
+			))
 			{
 				// Geometry caches don't support LODs
 				const bool bAllowInterpretingLODs = false;
@@ -1683,10 +1681,9 @@ void FUsdGeomMeshTranslator::UpdateComponents(USceneComponent* SceneComponent)
 		UGeometryCache* GeometryCache = nullptr;
 		if (Context->InfoCache)
 		{
-			GeometryCache = Cast<UGeometryCache>(Context->InfoCache->GetSingleAssetForPrim(
-				PrimPath,
-				UGeometryCache::StaticClass()
-			));
+			GeometryCache = Context->InfoCache->GetSingleAssetForPrim<UGeometryCache>(
+				PrimPath
+			);
 		}
 
 		bool bShouldRegister = false;
diff --git a/Engine/Plugins/Importers/USDImporter/Source/USDSchemas/Private/USDGeomPointInstancerTranslator.cpp b/Engine/Plugins/Importers/USDImporter/Source/USDSchemas/Private/USDGeomPointInstancerTranslator.cpp
--- a/Engine/Plugins/Importers/USDImporter/Source/USDSchemas/Private/USDGeomPointInstancerTranslator.cpp
+++ b/Engine/Plugins/Importers/USDImporter/Source/USDSchemas/Private/USDGeomPointInstancerTranslator.cpp
@@ -318,10 +318,9 @@ USceneComponent* FUsdGeomPointInstancerTranslator::CreateComponents()
 			continue;
 		}
 
-		UStaticMesh* StaticMesh = Cast< UStaticMesh >( InfoCache.GetSingleAssetForPrim(
-			UE::FSdfPath{*PrototypePathStr},
-			UStaticMesh::StaticClass()
-		));
+		UStaticMesh* StaticMesh = InfoCache.GetSingleAssetForPrim<UStaticMesh>(
+			UE::FSdfPath{*PrototypePathStr}
+		);
 		UsdGeomPointInstancerTranslatorImpl::SetStaticMesh( StaticMesh, *HISMComponent );
 
 		// Evaluating point instancer can take a long time and is thread-safe. Move to async task while we work on something else.
diff --git a/Engine/Plugins/Importers/USDImporter/Source/USDSchemas/Private/USDGeomXformableTranslator.cpp b/Engine/Plugins/Importers/USDImporter/Source/USDSchemas/Private/USDGeomXformableTranslator.cpp
--- a/Engine/Plugins/Importers/USDImporter/Source/USDSchemas/Private/USDGeomXformableTranslator.cpp
+++ b/Engine/Plugins/Importers/USDImporter/Source/USDSchemas/Private/USDGeomXformableTranslator.cpp
@@ -556,10 +556,9 @@ void FUsdGeomXformableTranslator::UpdateComponents( USceneComponent* SceneCompon
 		bool bHasMultipleLODs = false;
 		if ( UStaticMeshComponent* StaticMeshComponent = Cast< UStaticMeshComponent >( SceneComponent ) )
 		{
-			UStaticMesh* PrimStaticMesh = Cast< UStaticMesh >( Context->InfoCache->GetSingleAssetForPrim(
-				PrimPath,
-				UStaticMesh::StaticClass()
-			));
+			UStaticMesh* PrimStaticMesh = Context->InfoCache->GetSingleAssetForPrim<UStaticMesh>(
+				PrimPath
+			);
 
 			if ( PrimStaticMesh )
 			{
diff --git a/Engine/Plugins/Importers/USDImporter/Source/USDSchemas/Private/USDGroomTranslator.cpp b/Engine/Plugins/Importers/USDImporter/Source/USDSchemas/Private/USDGroomTranslator.cpp
--- a/Engine/Plugins/Importers/USDImporter/Source/USDSchemas/Private/USDGroomTranslator.cpp
+++ b/Engine/Plugins/Importers/USDImporter/Source/USDSchemas/Private/USDGroomTranslator.cpp
@@ -209,10 +209,9 @@ protected:
 					return false;
 				}
 
-				UGroomAsset* GroomAsset = Cast<UGroomAsset>(Context->InfoCache->GetSingleAssetForPrim(
-					PrimPath,
-					UGroomAsset::StaticClass()
-				));
+				UGroomAsset* GroomAsset = Context->InfoCache->GetSingleAssetForPrim<UGroomAsset>(
+					PrimPath
+				);
 				if (!GroomAsset)
 				{
 					return false;
@@ -398,10 +397,9 @@ void FUsdGroomTranslator::UpdateComponents(USceneComponent* SceneComponent)
 		UGroomAsset* Groom = nullptr;
 		if(Context->InfoCache)
 		{
-			Groom = Cast<UGroomAsset>(Context->InfoCache->GetSingleAssetForPrim(
-				PrimPath,
-				UGroomAsset::StaticClass()
-			));
+			Groom = Context->InfoCache->GetSingleAssetForPrim<UGroomAsset>(
+				PrimPath
+			);
 		}
 
 		bool bShouldRegister = false;
@@ -418,10 +416,9 @@ void FUsdGroomTranslator::UpdateComponents(USceneComponent* SceneComponent)
 
 			if (Groom)
 			{
-				UGroomCache* GroomCache = Cast<UGroomCache>(Context->InfoCache->GetSingleAssetForPrim(
-					UE::FSdfPath{*UsdGroomTranslatorUtils::GetStrandsGroomCachePrimPath(PrimPath)},
-					UGroomCache::StaticClass()
-				));
+				UGroomCache* GroomCache = Context->InfoCache->GetSingleAssetForPrim<UGroomCache>(
+					UE::FSdfPath{*UsdGroomTranslatorUtils::GetStrandsGroomCachePrimPath(PrimPath)}
+				);
 				if (GroomCache != GroomComponent->GroomCache.Get())
 				{
 					GroomComponent->SetGroomCache(GroomCache);
diff --git a/Engine/Plugins/Importers/USDImporter/Source/USDSchemas/Private/USDGroomTranslatorUtils.cpp b/Engine/Plugins/Importers/USDImporter/Source/USDSchemas/Private/USDGroomTranslatorUtils.cpp
--- a/Engine/Plugins/Importers/USDImporter/Source/USDSchemas/Private/USDGroomTranslatorUtils.cpp
+++ b/Engine/Plugins/Importers/USDImporter/Source/USDSchemas/Private/USDGroomTranslatorUtils.cpp
@@ -201,10 +201,9 @@ namespace UsdGroomTranslatorUtils
 		}
 
 		// The GroomAsset should already be processed and cached by the USDGroomTranslator
-		UGroomAsset* GroomAsset = Cast<UGroomAsset>(InfoCache.GetSingleAssetForPrim(
-			UE::FSdfPath{*GroomPrimPath},
-			UGroomAsset::StaticClass()
-		));
+		UGroomAsset* GroomAsset = InfoCache.GetSingleAssetForPrim<UGroomAsset>(
+			UE::FSdfPath{*GroomPrimPath}
+		);
 		if (!GroomAsset)
 		{
 			return;
@@ -268,10 +267,9 @@ namespace UsdGroomTranslatorUtils
 			return;
 		}
 
-		UGroomAsset* GroomAsset = Cast<UGroomAsset>(InfoCache.GetSingleAssetForPrim(
-			UE::FSdfPath{*GroomPrimPath},
-			UGroomAsset::StaticClass()
-		));
+		UGroomAsset* GroomAsset = InfoCache.GetSingleAssetForPrim<UGroomAsset>(
+			UE::FSdfPath{*GroomPrimPath}
+		);
 		if (!GroomAsset)
 		{
 			return;
@@ -279,10 +277,9 @@ namespace UsdGroomTranslatorUtils
 
 		const FString PrimPath(UsdToUnreal::ConvertPath(Prim.GetPath()));
 		const FString GroomBindingPath = FString::Printf(TEXT("%s_groombinding"), *PrimPath);
-		UGroomBindingAsset* GroomBinding = Cast<UGroomBindingAsset>(InfoCache.GetSingleAssetForPrim(
-			UE::FSdfPath{*GroomBindingPath},
-			UGroomBindingAsset::StaticClass()
-		));
+		UGroomBindingAsset* GroomBinding = InfoCache.GetSingleAssetForPrim<UGroomBindingAsset>(
+			UE::FSdfPath{*GroomBindingPath}
+		);
 
 		// Set the GroomAsset and GroomBindingAsset on the child GroomComponent of SceneComponent that was set up in the translator
 		TArray<USceneComponent*> Children;
diff --git a/Engine/Plugins/Importers/USDImporter/Source/USDSchemas/Private/USDInfoCache.cpp b/Engine/Plugins/Importers/USDImporter/Source/USDSchemas/Private/USDInfoCache.cpp
--- a/Engine/Plugins/Importers/USDImporter/Source/USDSchemas/Private/USDInfoCache.cpp
+++ b/Engine/Plugins/Importers/USDImporter/Source/USDSchemas/Private/USDInfoCache.cpp
@@ -627,31 +627,7 @@ TSet<TWeakObjectPtr<UObject>> FUsdInfoCache::RemoveAllAssetPrimLinks(const UE::F
 	return Values;
 }
 
-UObject* FUsdInfoCache::GetSingleAssetForPrim(const UE::FSdfPath& Path, const UClass* FilterClass) const
-{
-	FUsdInfoCacheImpl* ImplPtr = Impl.Get();
-	if (!ImplPtr)
-	{
-		return nullptr;
-	}
-	FReadScopeLock ScopeLock(ImplPtr->PrimPathToAssetsLock);
-
-	if (const TSet<TWeakObjectPtr<UObject>>* FoundAssets = ImplPtr->PrimPathToAssets.Find(Path))
-	{
-		for (const TWeakObjectPtr<UObject>& FoundAsset : *FoundAssets)
-		{
-			UObject* FoundAssetPtr = FoundAsset.Get();
-			if (FoundAssetPtr && (!FilterClass || FoundAssetPtr->IsA(FilterClass)))
-			{
-				return FoundAssetPtr;
-			}
-		}
-	}
-
-	return nullptr;
-}
-
-TSet<TWeakObjectPtr<UObject>> FUsdInfoCache::GetAssetsForPrim(const UE::FSdfPath& Path, const UClass* FilterClass) const
+TSet<TWeakObjectPtr<UObject>> FUsdInfoCache::GetAllAssetsForPrim(const UE::FSdfPath& Path) const
 {
 	FUsdInfoCacheImpl* ImplPtr = Impl.Get();
 	if (!ImplPtr)
diff --git a/Engine/Plugins/Importers/USDImporter/Source/USDSchemas/Private/USDSkelRootTranslator.cpp b/Engine/Plugins/Importers/USDImporter/Source/USDSchemas/Private/USDSkelRootTranslator.cpp
--- a/Engine/Plugins/Importers/USDImporter/Source/USDSchemas/Private/USDSkelRootTranslator.cpp
+++ b/Engine/Plugins/Importers/USDImporter/Source/USDSchemas/Private/USDSkelRootTranslator.cpp
@@ -843,9 +843,8 @@ namespace UsdSkelRootTranslatorImpl
 
 		FString PrimName = UsdToUnreal::ConvertString( Prim.GetName() );
 
-		USkeletalMesh* SkeletalMesh = Cast< USkeletalMesh >(Context.InfoCache->GetSingleAssetForPrim(
-			UE::FSdfPath{Prim.GetPath()},
-			USkeletalMesh::StaticClass())
+		USkeletalMesh* SkeletalMesh = Context.InfoCache->GetSingleAssetForPrim<USkeletalMesh>(
+			UE::FSdfPath{Prim.GetPath()}
 		);
 		if ( !SkeletalMesh )
 		{
@@ -1415,10 +1414,9 @@ namespace UsdSkelRootTranslatorImpl
 					return false;
 				}
 
-				USkeletalMesh* SkeletalMesh = Cast< USkeletalMesh >(Context->InfoCache->GetSingleAssetForPrim(
-					PrimPath,
-					USkeletalMesh::StaticClass()
-				));
+				USkeletalMesh* SkeletalMesh = Context->InfoCache->GetSingleAssetForPrim<USkeletalMesh>(
+					PrimPath
+				);
 				if ( !SkeletalMesh )
 				{
 					return false;
@@ -1617,10 +1615,9 @@ void FUsdSkelRootTranslator::UpdateComponents( USceneComponent* SceneComponent )
 	UE::FUsdPrim SkelAnimPrim = UsdUtils::FindFirstAnimationSource( Prim );
 	if (SkelAnimPrim)
 	{
-		UAnimSequence* TargetAnimSequence = Cast< UAnimSequence >(Context->InfoCache->GetSingleAssetForPrim(
-			SkelAnimPrim.GetPrimPath(),
-			UAnimSequence::StaticClass()
-		));
+		UAnimSequence* TargetAnimSequence = Context->InfoCache->GetSingleAssetForPrim<UAnimSequence>(
+			SkelAnimPrim.GetPrimPath()
+		);
 		if ( TargetAnimSequence != SkeletalMeshComponent->AnimationData.AnimToPlay )
 		{
 			SkeletalMeshComponent->AnimationData.AnimToPlay = TargetAnimSequence;
@@ -1638,10 +1635,9 @@ void FUsdSkelRootTranslator::UpdateComponents( USceneComponent* SceneComponent )
 
 #if WITH_EDITOR
 	// Re-set the skeletal mesh if we created a new one (maybe the hash changed, a skinned UsdGeomMesh was hidden, etc.)
-	USkeletalMesh* TargetSkeletalMesh = Cast< USkeletalMesh >(Context->InfoCache->GetSingleAssetForPrim(
-		PrimPath,
-		USkeletalMesh::StaticClass()
-	));
+	USkeletalMesh* TargetSkeletalMesh = Context->InfoCache->GetSingleAssetForPrim<USkeletalMesh>(
+		PrimPath
+	);
 	if ( SkeletalMeshComponent->GetSkeletalMeshAsset() != TargetSkeletalMesh )
 	{
 		SkeletalMeshComponent->SetSkeletalMesh(TargetSkeletalMesh);
diff --git a/Engine/Plugins/Importers/USDImporter/Source/USDSchemas/Public/USDInfoCache.h b/Engine/Plugins/Importers/USDImporter/Source/USDSchemas/Public/USDInfoCache.h
--- a/Engine/Plugins/Importers/USDImporter/Source/USDSchemas/Public/USDInfoCache.h
+++ b/Engine/Plugins/Importers/USDImporter/Source/USDSchemas/Public/USDInfoCache.h
@@ -4,6 +4,7 @@
 
 #include "Containers/Set.h"
 #include "CoreMinimal.h"
+#include "UObject/WeakObjectPtr.h"
 
 #include "UsdWrappers/SdfPath.h"
 
@@ -61,8 +62,41 @@ public:
 
 	TSet<TWeakObjectPtr<UObject>> RemoveAllAssetPrimLinks(const UE::FSdfPath& Path);
 
-	UObject* GetSingleAssetForPrim(const UE::FSdfPath& Path, const UClass* FilterClass = nullptr) const;
-	TSet<TWeakObjectPtr<UObject>> GetAssetsForPrim(const UE::FSdfPath& Path, const UClass* FilterClass = nullptr) const;
+	TSet<TWeakObjectPtr<UObject>> GetAllAssetsForPrim(const UE::FSdfPath& Path) const;
+
+	template<typename T = UObject>
+	T* GetSingleAssetForPrim(const UE::FSdfPath& Path) const
+	{
+		TSet<TWeakObjectPtr<UObject>> Assets = GetAllAssetsForPrim(Path);
+		for (const TWeakObjectPtr<UObject>& Asset : Assets)
+		{
+			if (T* CastAsset = Cast<T>(Asset.Get()))
+			{
+				return CastAsset;
+			}
+		}
+
+		return nullptr;
+	}
+
+	template<typename T>
+	TSet<T*> GetAssetsForPrim(const UE::FSdfPath& Path) const
+	{
+		TSet<TWeakObjectPtr<UObject>> Assets = GetAllAssetsForPrim(Path);
+
+		TSet<T*> CastAssets;
+		CastAssets.Reserve(Assets.Num());
+
+		for (const TWeakObjectPtr<UObject>& Asset : Assets)
+		{
+			if (T* CastAsset = Cast<T>(Asset.Get()))
+			{
+				CastAssets.Add(CastAsset);
+			}
+		}
+
+		return CastAssets;
+	}
 
 	UE::FSdfPath GetPrimForAsset(UObject* Asset) const;
 	TMap<UE::FSdfPath, TSet<TWeakObjectPtr<UObject>>> GetAllAssetPrimLinks() const;
diff --git a/Engine/Plugins/Importers/USDImporter/Source/USDStage/Private/USDLevelSequenceHelper.cpp b/Engine/Plugins/Importers/USDImporter/Source/USDStage/Private/USDLevelSequenceHelper.cpp
--- a/Engine/Plugins/Importers/USDImporter/Source/USDStage/Private/USDLevelSequenceHelper.cpp
+++ b/Engine/Plugins/Importers/USDImporter/Source/USDStage/Private/USDLevelSequenceHelper.cpp
@@ -1856,10 +1856,9 @@ void FUsdLevelSequenceHelperImpl::AddSkeletalTracks(const UUsdPrimTwin& PrimTwin
 	// but we may belong to a FUsdStageImportContext, and so there's no AUsdStageActor at all to use.
 	// At this point it doesn't matter much though, because we shouldn't need to uncollapse a SkelAnimation prim path anyway
 	const UE::FSdfPath PrimPath = SkelAnimationPrim.GetPrimPath();
-	UAnimSequence* Sequence = Cast<UAnimSequence>(InfoCache->GetSingleAssetForPrim(
-		PrimPath,
-		UAnimSequence::StaticClass()
-	));
+	UAnimSequence* Sequence = InfoCache->GetSingleAssetForPrim<UAnimSequence>(
+		PrimPath
+	);
 	if (!Sequence)
 	{
 		return;
@@ -1928,10 +1927,9 @@ void FUsdLevelSequenceHelperImpl::AddGeometryCacheTracks(const UUsdPrimTwin& Pri
 
 	// Fetch the geometry cache asset from the asset cache. If there's none, don't actually need to create track
 	const UE::FSdfPath PrimPath = Prim.GetPrimPath();
-	UGeometryCache* GeometryCache = Cast< UGeometryCache >(InfoCache->GetSingleAssetForPrim(
-		PrimPath,
-		UGeometryCache::StaticClass()
-	));
+	UGeometryCache* GeometryCache = InfoCache->GetSingleAssetForPrim<UGeometryCache>(
+		PrimPath
+	);
 	if (!GeometryCache)
 	{
 		return;
@@ -1997,10 +1995,9 @@ void FUsdLevelSequenceHelperImpl::AddGroomTracks(const UUsdPrimTwin& PrimTwin, c
 	// Fetch the groom cache asset from the asset cache. If there's none, don't actually need to create track
 	const FString PrimPath = Prim.GetPrimPath().GetString();
 	const FString GroomCachePath = FString::Printf(TEXT("%s_strands_cache"), *PrimPath);
-	UGroomCache* GroomCache = Cast<UGroomCache>(InfoCache->GetSingleAssetForPrim(
-		UE::FSdfPath{*GroomCachePath},
-		UGroomCache::StaticClass()
-	));
+	UGroomCache* GroomCache = InfoCache->GetSingleAssetForPrim<UGroomCache>(
+		UE::FSdfPath{*GroomCachePath}
+	);
 	if (!GroomCache)
 	{
 		return;
@@ -2253,10 +2250,9 @@ void FUsdLevelSequenceHelperImpl::UpdateControlRigTracks(UUsdPrimTwin& PrimTwin)
 	// but we may belong to a FUsdStageImportContext, and so there's no AUsdStageActor at all to use.
 	// At this point it doesn't matter much though, because we shouldn't need to uncollapse a SkelAnimation prim path anyway
 	const UE::FSdfPath SkelAnimationPrimPath = SkelAnimationPrim.GetPrimPath();
-	UAnimSequence* AnimSequence = Cast<UAnimSequence>(InfoCache->GetSingleAssetForPrim(
-		SkelAnimationPrimPath,
-		UAnimSequence::StaticClass()
-	));
+	UAnimSequence* AnimSequence = InfoCache->GetSingleAssetForPrim<UAnimSequence>(
+		SkelAnimationPrimPath
+	);
 
 	if (!SkelAnimationLayer)
 	{
diff --git a/Engine/Plugins/Importers/USDImporter/Source/USDStage/Private/USDStageActor.cpp b/Engine/Plugins/Importers/USDImporter/Source/USDStage/Private/USDStageActor.cpp
--- a/Engine/Plugins/Importers/USDImporter/Source/USDStage/Private/USDStageActor.cpp
+++ b/Engine/Plugins/Importers/USDImporter/Source/USDStage/Private/USDStageActor.cpp
@@ -1139,18 +1139,16 @@ void AUsdStageActor::OnUsdObjectsChanged(const UsdUtils::FObjectChangesByPath& I
 
 				if (bInResync && InfoCache)
 				{
-					UMaterialInterface* ExistingMaterial = Cast<UMaterialInterface>(InfoCache->GetSingleAssetForPrim(
-						AssetsPrimPath,
-						UMaterialInterface::StaticClass()
-					));
+					UMaterialInterface* ExistingMaterial = InfoCache->GetSingleAssetForPrim<UMaterialInterface>(
+						AssetsPrimPath
+					);
 
 					UE::FUsdPrim PrimToResync = GetOrLoadUsdStage().GetPrimAtPath(AssetsPrimPath);
 					LoadAssets(*TranslationContext, PrimToResync);
 
-					UMaterialInterface* NewMaterial = Cast<UMaterialInterface>(InfoCache->GetSingleAssetForPrim(
-						AssetsPrimPath,
-						UMaterialInterface::StaticClass()
-					));
+					UMaterialInterface* NewMaterial = InfoCache->GetSingleAssetForPrim<UMaterialInterface>(
+						AssetsPrimPath
+					);
 
 					// For UE-120185: If we recreated a material for a prim path we also need to update all components that were using it.
 					// This could be fleshed out further if other asset types require this refresh of "dependent components" but materials
@@ -1769,11 +1767,11 @@ TArray<UObject*> AUsdStageActor::GetGeneratedAssets(const FString& PrimPath)
 	// Prefer checking the prim directly, but also check its collapsed root if it is collapsed.
 	// This because we have some exception cases like USkeleton/UAnimSequences that can be found by querying the actual
 	// Skeleton/SKelAnimation prims even though they are collapsed into the SkelRoot prim.
-	TSet<TWeakObjectPtr<UObject>> AssetsSet = InfoCache->GetAssetsForPrim(UsdPath);
+	TSet<TWeakObjectPtr<UObject>> AssetsSet = InfoCache->GetAllAssetsForPrim(UsdPath);
 	if (AssetsSet.Num() == 0 && InfoCache->IsPathCollapsed(UsdPath, ECollapsingType::Assets))
 	{
 		UsdPath = InfoCache->UnwindToNonCollapsedPath(UsdPath, ECollapsingType::Assets);
-		AssetsSet = InfoCache->GetAssetsForPrim(UsdPath);
+		AssetsSet = InfoCache->GetAllAssetsForPrim(UsdPath);
 	}
 
 	TArray<UObject*> Assets;
