USD: Improve handling of twosided materials to not need a fake "_TwoSided" prim path suffix anymore.

diff --git a/Engine/Plugins/Importers/USDImporter/Source/USDSchemas/Private/MeshTranslationImpl.cpp b/Engine/Plugins/Importers/USDImporter/Source/USDSchemas/Private/MeshTranslationImpl.cpp
--- a/Engine/Plugins/Importers/USDImporter/Source/USDSchemas/Private/MeshTranslationImpl.cpp
+++ b/Engine/Plugins/Importers/USDImporter/Source/USDSchemas/Private/MeshTranslationImpl.cpp
@@ -95,37 +95,39 @@ TMap<const UsdUtils::FUsdPrimMaterialSlot*, UMaterialInterface*> MeshTranslation
 			}
 			case UsdUtils::EPrimAssignmentType::MaterialPrim:
 			{
-				// TODO: Instead of using a custom prim suffix, add the material to the same prim and then pick
-				// between all assets on a prim when needed.
-				// Previously the suffix was fine because PrimPathToAssets dealt with strings, now that it has an
-				// SdfPath API we can't use that. It's sort of a hack anyway
-				FString PrimPath = Slot.MaterialSource;
-				const static FString TwoSidedToken = TEXT("_twosided");
-				if (Slot.bMeshIsDoubleSided)
+				UMaterialInstance* OneSidedMat = nullptr;
+
+				TSet<UMaterialInstance*> ExistingMaterials = InfoCache.GetAssetsForPrim<UMaterialInstance>(UE::FSdfPath{*Slot.MaterialSource});
+				for (UMaterialInstance* ExistingMaterial: ExistingMaterials)
 				{
-					PrimPath += TwoSidedToken;
-				}
+					const bool bExistingIsTwoSided = ExistingMaterial->IsTwoSided();
 
-				Material = InfoCache.GetSingleAssetForPrim<UMaterialInterface>(UE::FSdfPath{*PrimPath});
+					if (!bExistingIsTwoSided)
+					{
+						OneSidedMat = ExistingMaterial;
+					}
+
+					if (Slot.bMeshIsDoubleSided == bExistingIsTwoSided)
+					{
+						Material = ExistingMaterial;
+					}
+				}
 
 				// Need to create a two-sided material on-demand
 				if (!Material && Slot.bMeshIsDoubleSided)
 				{
 					// By now we parsed all materials so we must have the single-sided version of this material
-					UMaterialInstance* OneSidedMat = InfoCache.GetSingleAssetForPrim<UMaterialInstance>(UE::FSdfPath{*Slot.MaterialSource});
 					if (!OneSidedMat)
 					{
-						UE_LOG(LogUsd, Warning, TEXT("Failed to generate a two-sided material from the material prim at path '%s' as no single-sided material was generated for it."), *PrimPath);
+						UE_LOG(LogUsd, Warning, TEXT("Failed to generate a two-sided material from the material prim at path '%s' as no single-sided material was generated for it."), *Slot.MaterialSource);
 						continue;
 					}
 
 					// Check if for some reason we already have a two-sided material ready due to a complex scenario
 					// related to the global cache
 					const FString OneSidedHash = AssetCache.GetHashForAsset(OneSidedMat);
-					const FString TwoSidedHash = OneSidedHash + TEXT("_TwoSided");
-					UMaterialInstance* TwoSidedMat = Cast<UMaterialInstance>(
-						AssetCache.GetCachedAsset(TwoSidedHash)
-					);
+					const FString TwoSidedHash = OneSidedHash + UnrealIdentifiers::TwoSidedMaterialSuffix;
+					UMaterialInstance* TwoSidedMat = Cast<UMaterialInstance>(AssetCache.GetCachedAsset(TwoSidedHash));
 					if (!TwoSidedMat)
 					{
 						// Important to not use GetBaseMaterial() here because if our parent is the translucent we'll
@@ -190,7 +192,7 @@ TMap<const UsdUtils::FUsdPrimMaterialSlot*, UMaterialInterface*> MeshTranslation
 					}
 
 					AssetCache.CacheAsset(TwoSidedHash, TwoSidedMat);
-					InfoCache.LinkAssetToPrim(UE::FSdfPath{*PrimPath}, TwoSidedMat);
+					InfoCache.LinkAssetToPrim(UE::FSdfPath{*Slot.MaterialSource}, TwoSidedMat);
 					Material = TwoSidedMat;
 				}
 
@@ -207,8 +209,6 @@ TMap<const UsdUtils::FUsdPrimMaterialSlot*, UMaterialInterface*> MeshTranslation
 				}
 				else if (!Material->IsTwoSided() && Slot.bMeshIsDoubleSided)
 				{
-					// TODO: Update this message with the proper source prim paths when UE-138122 is submitted,
-					// as 'UsdPrim' may just be e.g. a SkelRoot
 					UE_LOG(LogUsd, Warning, TEXT("Using one-sided UE material '%s' for doubleSided prim '%s'"),
 						*Slot.MaterialSource,
 						*UsdToUnreal::ConvertPath(UsdPrim.GetPrimPath())
