USDImporter: add "RelativeToGameContentDir" metadata to AdditionalPluginDirectories setting

This metadata addition improves support for UE project-based USD plugins and extensions.
Previously, it was not clear when a relative path was used as an AdditionalPluginDirectories
entry what path it was expected to be relative to. Using a relative path allows the USD
plugins to live within and travel with the UE project, and when the same relative path is
set in the project packaging setting "Additional Non-Asset Directories To Copy", the USD
plugins will get included in a packaged game.

diff --git a/Engine/Plugins/Importers/USDImporter/Source/USDClasses/Public/USDProjectSettings.h b/Engine/Plugins/Importers/USDImporter/Source/USDClasses/Public/USDProjectSettings.h
--- a/Engine/Plugins/Importers/USDImporter/Source/USDClasses/Public/USDProjectSettings.h
+++ b/Engine/Plugins/Importers/USDImporter/Source/USDClasses/Public/USDProjectSettings.h
@@ -24,8 +24,19 @@ class UUsdProjectSettings : public UDeveloperSettings
 	GENERATED_BODY()
 
 public:
-	// Additional paths to check for USD plugins
-	UPROPERTY( config, EditAnywhere, Category = USD )
+	// Additional paths to check for USD plugins.
+	//
+	// If you want the USD plugins to be included in a packaged game, you must use a relative
+	// path to a location within your project directory, and you must also add that same path
+	// to the "Additional Non-Asset Directories To Copy" Project Packaging setting.
+	//
+	// For example, this relative path could be used to locate USD plugins in a directory at
+	// the root of your project:
+	//     ../USD_Plugins
+	//
+	// The packaging process cannot use an absolute path and will raise an error if given one
+	// when it tries to concatenate the game content directory path with an absolute path.
+	UPROPERTY( config, EditAnywhere, Category = USD, meta = (RelativeToGameContentDir) )
 	TArray<FDirectoryPath> AdditionalPluginDirectories;
 
 	// Material purposes to show on drop-downs in addition to the standard "preview" and "full"
diff --git a/Engine/Plugins/Importers/USDImporter/Source/UnrealUSDWrapper/Private/UnrealUSDWrapper.cpp b/Engine/Plugins/Importers/USDImporter/Source/UnrealUSDWrapper/Private/UnrealUSDWrapper.cpp
--- a/Engine/Plugins/Importers/USDImporter/Source/UnrealUSDWrapper/Private/UnrealUSDWrapper.cpp
+++ b/Engine/Plugins/Importers/USDImporter/Source/UnrealUSDWrapper/Private/UnrealUSDWrapper.cpp
@@ -939,7 +939,20 @@ public:
 		{
 			if ( !Directory.Path.IsEmpty() )
 			{
-				PluginDirectories.Add( Directory.Path );
+				// The directory chooser will try to enforce that the paths
+				// in the setting are within the game content dir, but it's
+				// still possible for the user to manually enter either a
+				// relative or an absolute path that points elsewhere, so aim
+				// to support any of these.
+				FString PluginDirPath = Directory.Path;
+				if ( FPaths::IsRelative( PluginDirPath ) )
+				{
+					PluginDirPath = FPaths::Combine( FPaths::ProjectContentDir(), PluginDirPath );
+					PluginDirPath = FPaths::ConvertRelativePathToFull( PluginDirPath );
+				}
+				FPaths::NormalizeDirectoryName( PluginDirPath );
+				FPaths::CollapseRelativeDirectories( PluginDirPath );
+				PluginDirectories.Add( PluginDirPath );
 			}
 		}
 
