USD: Fix issues handling MDL materials.

Fix MDL materials and textures not being created as transient;
Fix MDL textures not being added to the asset cache;
Fix MDL reference materials not being accounted for as dependencies for their instances;

diff --git a/Engine/Plugins/Importers/USDImporter/Source/USDClasses/Private/USDClassesModule.cpp b/Engine/Plugins/Importers/USDImporter/Source/USDClasses/Private/USDClassesModule.cpp
--- a/Engine/Plugins/Importers/USDImporter/Source/USDClasses/Private/USDClassesModule.cpp
+++ b/Engine/Plugins/Importers/USDImporter/Source/USDClasses/Private/USDClassesModule.cpp
@@ -318,6 +318,13 @@ TSet<UObject*> IUsdClassesModule::GetAssetDependencies(UObject* Asset)
 				Result.Add(Texture);
 			}
 		}
+
+		// We'll have a dependency on our reference material too of course (this happens for Mdl
+		// materials for example, that create new UMaterial assets every time, and also material instances).
+		if (UMaterialInterface* ReferenceMaterial = MaterialInstance->Parent.Get())
+		{
+			Result.Add(ReferenceMaterial);
+		}
 	}
 	else if (USkeletalMesh* SkeletalMesh = Cast<USkeletalMesh>(Asset))
 	{
diff --git a/Engine/Plugins/Importers/USDImporter/Source/USDSchemas/Private/Custom/MDLUSDShadeMaterialTranslator.cpp b/Engine/Plugins/Importers/USDImporter/Source/USDSchemas/Private/Custom/MDLUSDShadeMaterialTranslator.cpp
--- a/Engine/Plugins/Importers/USDImporter/Source/USDSchemas/Private/Custom/MDLUSDShadeMaterialTranslator.cpp
+++ b/Engine/Plugins/Importers/USDImporter/Source/USDSchemas/Private/Custom/MDLUSDShadeMaterialTranslator.cpp
@@ -16,6 +16,7 @@
 #include "UsdWrappers/SdfPath.h"
 #include "UsdWrappers/UsdPrim.h"
 
+#include "Engine/Texture.h"
 #include "Materials/Material.h"
 #include "Materials/MaterialInstance.h"
 #include "Materials/MaterialInstanceConstant.h"
@@ -149,7 +150,7 @@ void FMdlUsdShadeMaterialTranslator::CreateAssets()
 
 			TStrongObjectPtr< UMDLImporterOptions > ImportOptions( NewObject< UMDLImporterOptions >() );
 
-			MdlMaterial = FMdlMaterialImporter::ImportMaterialFromModule( GetTransientPackage(), Context->ObjectFlags, MdlModuleName, MdlDefinitionName, *ImportOptions.Get() );
+			MdlMaterial = FMdlMaterialImporter::ImportMaterialFromModule( GetTransientPackage(), Context->ObjectFlags | RF_Transient, MdlModuleName, MdlDefinitionName, *ImportOptions.Get() );
 
 			if ( MdlMaterial )
 			{
@@ -172,7 +173,7 @@ void FMdlUsdShadeMaterialTranslator::CreateAssets()
 		UMaterialInstanceConstant* MdlMaterialInstance = Cast< UMaterialInstanceConstant >(Context->AssetCache->GetCachedAsset(MdlFullInstanceName));
 		if ( !MdlMaterialInstance && MdlMaterial )
 		{
-			MdlMaterialInstance = NewObject< UMaterialInstanceConstant >(GetTransientPackage());
+			MdlMaterialInstance = NewObject< UMaterialInstanceConstant >(GetTransientPackage(), NAME_None, Context->ObjectFlags | RF_Transient);
 
 			UUsdAssetImportData* ImportData = NewObject< UUsdAssetImportData >( MdlMaterialInstance, TEXT( "USDAssetImportData" ) );
 			ImportData->PrimPath = PrimPath.GetString();
@@ -203,7 +204,68 @@ void FMdlUsdShadeMaterialTranslator::CreateAssets()
 
 		if (Context->InfoCache)
 		{
+			Context->InfoCache->LinkAssetToPrim(PrimPath, MdlMaterial);
 			Context->InfoCache->LinkAssetToPrim(PrimPath, MdlMaterialInstance);
+
+			if (UMaterial* MdlReference = Cast<UMaterial>(MdlMaterial))
+			{
+				TArray<UTexture*> UsedTextures;
+				const bool bAllQualityLevels = true;
+				const bool bAllFeatureLevels = true;
+				MdlReference->GetUsedTextures(
+					UsedTextures,
+					EMaterialQualityLevel::High,
+					bAllQualityLevels,
+					ERHIFeatureLevel::SM5,
+					bAllFeatureLevels
+				);
+
+				for (UTexture* Texture : UsedTextures)
+				{
+					if (Texture->GetOutermost() == GetTransientPackage())
+					{
+						Context->InfoCache->LinkAssetToPrim(PrimPath, Texture);
+
+						const FString FilePath = Texture->AssetImportData ? Texture->AssetImportData->GetFirstFilename() : Texture->GetName();
+						const FString TextureHash = UsdUtils::GetTextureHash(
+							FilePath,
+							Texture->SRGB,
+							Texture->CompressionSettings,
+							Texture->GetTextureAddressX(),
+							Texture->GetTextureAddressY()
+						);
+
+						Texture->SetFlags(RF_Transient);
+						Context->AssetCache->CacheAsset(TextureHash, Texture);
+					}
+				}
+			}
+
+			if (MdlMaterialInstance)
+			{
+				for (const FTextureParameterValue& TextureValue : MdlMaterialInstance->TextureParameterValues)
+				{
+					if (UTexture* Texture = TextureValue.ParameterValue)
+					{
+						if (Texture->GetOutermost() == GetTransientPackage())
+						{
+							Context->InfoCache->LinkAssetToPrim(PrimPath, Texture);
+
+							const FString FilePath = Texture->AssetImportData ? Texture->AssetImportData->GetFirstFilename() : Texture->GetName();
+							const FString TextureHash = UsdUtils::GetTextureHash(
+								FilePath,
+								Texture->SRGB,
+								Texture->CompressionSettings,
+								Texture->GetTextureAddressX(),
+								Texture->GetTextureAddressY()
+							);
+
+							Texture->SetFlags(RF_Transient);
+							Context->AssetCache->CacheAsset(TextureHash, Texture);
+						}
+					}
+				}
+			}
 		}
 	}
 	else
diff --git a/Engine/Plugins/Importers/USDImporter/Source/USDUtilities/Private/USDShadeConversion.cpp b/Engine/Plugins/Importers/USDImporter/Source/USDUtilities/Private/USDShadeConversion.cpp
--- a/Engine/Plugins/Importers/USDImporter/Source/USDUtilities/Private/USDShadeConversion.cpp
+++ b/Engine/Plugins/Importers/USDImporter/Source/USDUtilities/Private/USDShadeConversion.cpp
@@ -320,67 +320,6 @@ namespace UE
 				return NewTexture;
 			}
 
-			// Computes and returns the hash string for the texture at the given path.
-			// Handles regular texture asset paths as well as asset paths identifying textures inside Usdz archives.
-			// Returns an empty string if the texture could not be hashed.
-			FString GetTextureHash(
-				const FString& ResolvedTexturePath,
-				bool bSRGB,
-				TextureCompressionSettings CompressionSettings,
-				TextureAddress AddressX,
-				TextureAddress AddressY
-			)
-			{
-				FMD5 MD5;
-
-				// Hash the actual texture
-				FString TextureExtension;
-				if (IsInsideUsdzArchive(ResolvedTexturePath, TextureExtension))
-				{
-					uint64 BufferSize = 0;
-					TUsdStore<std::shared_ptr<const char>> Buffer = ReadTextureBufferFromUsdzArchive(ResolvedTexturePath, BufferSize);
-					const uint8* BufferStart = reinterpret_cast<const uint8*>(Buffer.Get().get());
-
-					if (BufferSize > 0 && BufferStart != nullptr)
-					{
-						MD5.Update( BufferStart, BufferSize );
-					}
-				}
-				// Copied from FMD5Hash::HashFileFromArchive as it doesn't expose its FMD5
-				else if ( TUniquePtr<FArchive> Ar{ IFileManager::Get().CreateFileReader( *ResolvedTexturePath ) } )
-				{
-					TArray<uint8> LocalScratch;
-					LocalScratch.SetNumUninitialized( 1024 * 64 );
-
-					const int64 Size = Ar->TotalSize();
-					int64 Position = 0;
-
-					// Read in BufferSize chunks
-					while ( Position < Size )
-					{
-						const auto ReadNum = FMath::Min( Size - Position, ( int64 ) LocalScratch.Num() );
-						Ar->Serialize( LocalScratch.GetData(), ReadNum );
-						MD5.Update( LocalScratch.GetData(), ReadNum );
-
-						Position += ReadNum;
-					}
-				}
-				else
-				{
-					UE_LOG( LogUsd, Warning, TEXT( "Failed to find texture at path '%s' when trying to generate a hash for it" ), *ResolvedTexturePath );
-				}
-
-				// Hash the additional data
-				MD5.Update( reinterpret_cast< uint8* >( &bSRGB ), sizeof( bool ) );
-				MD5.Update( reinterpret_cast< uint8* >( &CompressionSettings ), sizeof( CompressionSettings ) );
-				MD5.Update( reinterpret_cast< uint8* >( &AddressX ), sizeof( AddressX ) );
-				MD5.Update( reinterpret_cast< uint8* >( &AddressY ), sizeof( AddressY ) );
-
-				FMD5Hash Hash;
-				Hash.Set( MD5 );
-				return LexToString( Hash );
-			}
-
 			// Will traverse the shade material graph backwards looking for a string/token value and return it.
 			// Returns the empty string if it didn't find anything.
 			FString RecursivelySearchForStringValue( pxr::UsdShadeInput Input )
@@ -624,7 +563,7 @@ namespace UE
 							}
 						}
 
-						const FString TextureHash = GetTextureHash( TexturePath, bSRGB, CompressionSettings, AddressX, AddressY );
+						const FString TextureHash = UsdUtils::GetTextureHash( TexturePath, bSRGB, CompressionSettings, AddressX, AddressY );
 
 						// We only actually want to retrieve the textures if we have a cache to put them in
 						if ( TexturesCache )
@@ -2504,6 +2443,66 @@ FString UsdUtils::GetResolvedTexturePath( const pxr::UsdAttribute& TextureAssetP
 	return ResolvedTexturePath;
 }
 
+FString UsdUtils::GetTextureHash(
+	const FString& ResolvedTexturePath,
+	bool bSRGB,
+	TextureCompressionSettings CompressionSettings,
+	TextureAddress AddressX,
+	TextureAddress AddressY
+)
+{
+	using namespace UE::UsdShadeConversion::Private;
+
+	FMD5 MD5;
+
+	// Hash the actual texture
+	FString TextureExtension;
+	if (IsInsideUsdzArchive(ResolvedTexturePath, TextureExtension))
+	{
+		uint64 BufferSize = 0;
+		TUsdStore<std::shared_ptr<const char>> Buffer = ReadTextureBufferFromUsdzArchive(ResolvedTexturePath, BufferSize);
+		const uint8* BufferStart = reinterpret_cast<const uint8*>(Buffer.Get().get());
+
+		if (BufferSize > 0 && BufferStart != nullptr)
+		{
+			MD5.Update(BufferStart, BufferSize);
+		}
+	}
+	// Copied from FMD5Hash::HashFileFromArchive as it doesn't expose its FMD5
+	else if (TUniquePtr<FArchive> Ar{IFileManager::Get().CreateFileReader(*ResolvedTexturePath)})
+	{
+		TArray<uint8> LocalScratch;
+		LocalScratch.SetNumUninitialized(1024 * 64);
+
+		const int64 Size = Ar->TotalSize();
+		int64 Position = 0;
+
+		// Read in BufferSize chunks
+		while (Position < Size)
+		{
+			const auto ReadNum = FMath::Min(Size - Position, (int64)LocalScratch.Num());
+			Ar->Serialize(LocalScratch.GetData(), ReadNum);
+			MD5.Update(LocalScratch.GetData(), ReadNum);
+
+			Position += ReadNum;
+		}
+	}
+	else
+	{
+		UE_LOG(LogUsd, Warning, TEXT("Failed to find texture at path '%s' when trying to generate a hash for it"), *ResolvedTexturePath);
+	}
+
+	// Hash the additional data
+	MD5.Update(reinterpret_cast<uint8*>(&bSRGB), sizeof(bool));
+	MD5.Update(reinterpret_cast<uint8*>(&CompressionSettings), sizeof(CompressionSettings));
+	MD5.Update(reinterpret_cast<uint8*>(&AddressX), sizeof(AddressX));
+	MD5.Update(reinterpret_cast<uint8*>(&AddressY), sizeof(AddressY));
+
+	FMD5Hash Hash;
+	Hash.Set(MD5);
+	return LexToString(Hash);
+}
+
 UTexture* UsdUtils::CreateTexture( const pxr::UsdAttribute& TextureAssetPathAttr, const FString& PrimPath, TextureGroup LODGroup, UObject* Outer )
 {
 	TRACE_CPUPROFILER_EVENT_SCOPE( UsdUtils::CreateTexture );
diff --git a/Engine/Plugins/Importers/USDImporter/Source/USDUtilities/Public/USDShadeConversion.h b/Engine/Plugins/Importers/USDImporter/Source/USDUtilities/Public/USDShadeConversion.h
--- a/Engine/Plugins/Importers/USDImporter/Source/USDUtilities/Public/USDShadeConversion.h
+++ b/Engine/Plugins/Importers/USDImporter/Source/USDUtilities/Public/USDShadeConversion.h
@@ -26,13 +26,15 @@ PXR_NAMESPACE_OPEN_SCOPE
 PXR_NAMESPACE_CLOSE_SCOPE
 
 class FSHAHash;
-class UUsdAssetCache2;
 class UMaterial;
 class UMaterialOptions;
 class UTexture;
 class UUsdAssetCache;
+class UUsdAssetCache2;
 enum class EFlattenMaterialProperties : uint8;
 enum EMaterialProperty : int;
+enum TextureAddress : int;
+enum TextureCompressionSettings : int;
 struct FFlattenMaterial;
 struct FPropertyEntry;
 namespace UE
@@ -159,6 +161,19 @@ namespace UsdUtils
 	/** Returns the resolved path from a pxr::SdfAssetPath attribute. For UDIMs path, returns the path to the 1001 tile. */
 	USDUTILITIES_API FString GetResolvedTexturePath( const pxr::UsdAttribute& TextureAssetPathAttr );
 
+	/**
+	 * Computes and returns the hash string for the texture at the given path.
+	 * Handles regular texture asset paths as well as asset paths identifying textures inside Usdz archives.
+     * Returns an empty string if the texture could not be hashed.
+	 */
+	USDUTILITIES_API FString GetTextureHash(
+		const FString& ResolvedTexturePath,
+		bool bSRGB,
+		TextureCompressionSettings CompressionSettings,
+		TextureAddress AddressX,
+		TextureAddress AddressY
+	);
+
 	/** Creates a texture from a pxr::SdfAssetPath attribute. PrimPath is optional, and should point to the source shadematerial prim path. It will be placed in its UUsdAssetImportData */
 	USDUTILITIES_API UTexture* CreateTexture( const pxr::UsdAttribute& TextureAssetPathAttr, const FString& PrimPath = FString(), TextureGroup LODGroup = TEXTUREGROUP_World, UObject* Outer = GetTransientPackage() );
 
