USD: Fix unused materials being parsed and consuming too much memory.

Fix two sided materials never being imported.

diff --git a/Engine/Plugins/Importers/USDImporter/Source/USDSchemas/Private/Custom/MDLUSDShadeMaterialTranslator.cpp b/Engine/Plugins/Importers/USDImporter/Source/USDSchemas/Private/Custom/MDLUSDShadeMaterialTranslator.cpp
--- a/Engine/Plugins/Importers/USDImporter/Source/USDSchemas/Private/Custom/MDLUSDShadeMaterialTranslator.cpp
+++ b/Engine/Plugins/Importers/USDImporter/Source/USDSchemas/Private/Custom/MDLUSDShadeMaterialTranslator.cpp
@@ -9,6 +9,7 @@
 #include "USDAssetImportData.h"
 #include "USDConversionUtils.h"
 #include "USDErrorUtils.h"
+#include "USDLog.h"
 #include "USDShadeConversion.h"
 #include "USDTypesConversion.h"
 
@@ -83,6 +84,15 @@ void FMdlUsdShadeMaterialTranslator::CreateAssets()
 		return;
 	}
 
+	if (Context->bTranslateOnlyUsedMaterials && Context->InfoCache)
+	{
+		if (!Context->InfoCache->IsMaterialUsed(PrimPath))
+		{
+			UE_LOG(LogUsd, Verbose, TEXT("Skipping creating assets for material prim '%s' as it is not currently bound by any prim."), *PrimPath.GetString());
+			return;
+		}
+	}
+
 	FScopedUsdAllocs UsdAllocs;
 
 	pxr::UsdShadeMaterial ShadeMaterial( GetPrim() );
diff --git a/Engine/Plugins/Importers/USDImporter/Source/USDSchemas/Private/MeshTranslationImpl.cpp b/Engine/Plugins/Importers/USDImporter/Source/USDSchemas/Private/MeshTranslationImpl.cpp
--- a/Engine/Plugins/Importers/USDImporter/Source/USDSchemas/Private/MeshTranslationImpl.cpp
+++ b/Engine/Plugins/Importers/USDImporter/Source/USDSchemas/Private/MeshTranslationImpl.cpp
@@ -98,7 +98,7 @@ TMap<const UsdUtils::FUsdPrimMaterialSlot*, UMaterialInterface*> MeshTranslation
 				UMaterialInstance* OneSidedMat = nullptr;
 
 				TSet<UMaterialInstance*> ExistingMaterials = InfoCache.GetAssetsForPrim<UMaterialInstance>(UE::FSdfPath{*Slot.MaterialSource});
-				for (UMaterialInstance* ExistingMaterial: ExistingMaterials)
+				for (UMaterialInstance* ExistingMaterial : ExistingMaterials)
 				{
 					const bool bExistingIsTwoSided = ExistingMaterial->IsTwoSided();
 
@@ -113,89 +113,101 @@ TMap<const UsdUtils::FUsdPrimMaterialSlot*, UMaterialInterface*> MeshTranslation
 					}
 				}
 
-				// Need to create a two-sided material on-demand
-				if (!Material && Slot.bMeshIsDoubleSided)
+				FString MaterialHash = AssetCache.GetHashForAsset(OneSidedMat);
+
+				if (Slot.bMeshIsDoubleSided)
 				{
-					// By now we parsed all materials so we must have the single-sided version of this material
-					if (!OneSidedMat)
-					{
-						UE_LOG(LogUsd, Warning, TEXT("Failed to generate a two-sided material from the material prim at path '%s' as no single-sided material was generated for it."), *Slot.MaterialSource);
-						continue;
-					}
+					MaterialHash = MaterialHash + UnrealIdentifiers::TwoSidedMaterialSuffix;
 
-					// Check if for some reason we already have a two-sided material ready due to a complex scenario
-					// related to the global cache
-					const FString OneSidedHash = AssetCache.GetHashForAsset(OneSidedMat);
-					const FString TwoSidedHash = OneSidedHash + UnrealIdentifiers::TwoSidedMaterialSuffix;
-					UMaterialInstance* TwoSidedMat = Cast<UMaterialInstance>(AssetCache.GetCachedAsset(TwoSidedHash));
-					if (!TwoSidedMat)
+					// Need to create a two-sided material on-demand
+					if (!Material)
 					{
-						// Important to not use GetBaseMaterial() here because if our parent is the translucent we'll
-						// get the reference UsdPreviewSurface instead, as that is also *its* reference
-						UMaterialInterface* ReferenceMaterial = OneSidedMat->Parent.Get();
-						UMaterialInterface* ReferenceMaterialTwoSided =
-							MeshTranslationImpl::GetTwoSidedVersionOfReferencePreviewSurfaceMaterial(ReferenceMaterial);
-						if (!ensure(ReferenceMaterialTwoSided && ReferenceMaterialTwoSided != ReferenceMaterial))
+						// By now we parsed all materials so we must have the single-sided version of this material
+						if (!OneSidedMat)
 						{
+							UE_LOG(LogUsd, Warning, TEXT("Failed to generate a two-sided material from the material prim at path '%s' as no single-sided material was generated for it."), *Slot.MaterialSource);
 							continue;
 						}
 
-						const FName NewInstanceName = MakeUniqueObjectName(
-							GetTransientPackage(),
-							UMaterialInstance::StaticClass(),
-							*(FPaths::GetBaseFilename(Slot.MaterialSource) + UnrealIdentifiers::TwoSidedMaterialSuffix)
-						);
-
-#if WITH_EDITOR
-						UMaterialInstanceConstant* MIC = Cast<UMaterialInstanceConstant>(OneSidedMat);
-						if (GIsEditor && MIC)
+						// Check if for some reason we already have a two-sided material ready due to a complex scenario
+						// related to the global cache
+						UMaterialInstance* TwoSidedMat = Cast<UMaterialInstance>(AssetCache.GetCachedAsset(MaterialHash));
+						if (!TwoSidedMat)
 						{
-							UMaterialInstanceConstant* TwoSidedMIC = NewObject<UMaterialInstanceConstant>(
-								GetTransientPackage(),
-								NewInstanceName,
-								Flags
-								);
-							if (TwoSidedMIC)
+							// Important to not use GetBaseMaterial() here because if our parent is the translucent we'll
+							// get the reference UsdPreviewSurface instead, as that is also *its* reference
+							UMaterialInterface* ReferenceMaterial = OneSidedMat->Parent.Get();
+							UMaterialInterface* ReferenceMaterialTwoSided =
+								MeshTranslationImpl::GetTwoSidedVersionOfReferencePreviewSurfaceMaterial(ReferenceMaterial);
+							if (!ensure(ReferenceMaterialTwoSided && ReferenceMaterialTwoSided != ReferenceMaterial))
 							{
-								UUsdAssetImportData* ImportData = NewObject< UUsdAssetImportData >(
-									TwoSidedMIC,
-									TEXT("USDAssetImportData")
-									);
-								ImportData->PrimPath = Slot.MaterialSource;
-								TwoSidedMIC->AssetImportData = ImportData;
+								continue;
 							}
 
-							TwoSidedMIC->SetParentEditorOnly(ReferenceMaterialTwoSided);
-							TwoSidedMIC->CopyMaterialUniformParametersEditorOnly(OneSidedMat);
-
-							TwoSidedMat = TwoSidedMIC;
-						}
-						else
-#endif // WITH_EDITOR
-						if (UMaterialInstanceDynamic* MID = Cast<UMaterialInstanceDynamic>(OneSidedMat))
-						{
-							UMaterialInstanceDynamic* TwoSidedMID = UMaterialInstanceDynamic::Create(
-								ReferenceMaterialTwoSided,
+							const FName NewInstanceName = MakeUniqueObjectName(
 								GetTransientPackage(),
-								NewInstanceName
+								UMaterialInstance::StaticClass(),
+								*(FPaths::GetBaseFilename(Slot.MaterialSource) + UnrealIdentifiers::TwoSidedMaterialSuffix)
 							);
-							if (!ensure(TwoSidedMID))
+
+#if WITH_EDITOR
+							UMaterialInstanceConstant* MIC = Cast<UMaterialInstanceConstant>(OneSidedMat);
+							if (GIsEditor && MIC)
 							{
-								continue;
+								UMaterialInstanceConstant* TwoSidedMIC = NewObject<UMaterialInstanceConstant>(
+									GetTransientPackage(),
+									NewInstanceName,
+									Flags
+								);
+								if (TwoSidedMIC)
+								{
+									UUsdAssetImportData* ImportData = NewObject< UUsdAssetImportData >(
+										TwoSidedMIC,
+										TEXT("USDAssetImportData")
+									);
+									ImportData->PrimPath = Slot.MaterialSource;
+									TwoSidedMIC->AssetImportData = ImportData;
+								}
+
+								TwoSidedMIC->SetParentEditorOnly(ReferenceMaterialTwoSided);
+								TwoSidedMIC->CopyMaterialUniformParametersEditorOnly(OneSidedMat);
+
+								TwoSidedMat = TwoSidedMIC;
 							}
+							else
+#endif // WITH_EDITOR
+							if (UMaterialInstanceDynamic* MID = Cast<UMaterialInstanceDynamic>(OneSidedMat))
+							{
+								UMaterialInstanceDynamic* TwoSidedMID = UMaterialInstanceDynamic::Create(
+									ReferenceMaterialTwoSided,
+									GetTransientPackage(),
+									NewInstanceName
+								);
+								if (!ensure(TwoSidedMID))
+								{
+									continue;
+								}
+
+								TwoSidedMID->CopyParameterOverrides(MID);
 
-							TwoSidedMID->CopyParameterOverrides(MID);
+								TwoSidedMat = TwoSidedMID;
+							}
 
-							TwoSidedMat = TwoSidedMID;
 						}
 
+						if (TwoSidedMat)
+						{
+							TwoSidedMat->SetFlags(RF_Transient);
+							Material = TwoSidedMat;
+						}
 					}
-
-					AssetCache.CacheAsset(TwoSidedHash, TwoSidedMat);
-					InfoCache.LinkAssetToPrim(UE::FSdfPath{*Slot.MaterialSource}, TwoSidedMat);
-					Material = TwoSidedMat;
 				}
 
+				// Cache the material to "ping it" as active, but also register two sided materials for the
+				// first time
+				AssetCache.CacheAsset(MaterialHash, Material);
+				InfoCache.LinkAssetToPrim(UE::FSdfPath{*Slot.MaterialSource}, Material);
+
 				break;
 			}
 			case UsdUtils::EPrimAssignmentType::UnrealMaterial:
diff --git a/Engine/Plugins/Importers/USDImporter/Source/USDSchemas/Private/USDInfoCache.cpp b/Engine/Plugins/Importers/USDImporter/Source/USDSchemas/Private/USDInfoCache.cpp
--- a/Engine/Plugins/Importers/USDImporter/Source/USDSchemas/Private/USDInfoCache.cpp
+++ b/Engine/Plugins/Importers/USDImporter/Source/USDSchemas/Private/USDInfoCache.cpp
@@ -19,6 +19,7 @@
 	#include "pxr/usd/sdf/path.h"
 	#include "pxr/usd/usd/prim.h"
 	#include "pxr/usd/usdGeom/mesh.h"
+	#include "pxr/usd/usdShade/materialBindingAPI.h"
 	#include "pxr/usd/usdGeom/pointInstancer.h"
 	#include "pxr/usd/usdGeom/subset.h"
 #include "USDIncludesEnd.h"
@@ -54,6 +55,11 @@ struct FUsdInfoCache::FUsdInfoCacheImpl
 	// Information we may have about a subset of prims
 	TMap<UE::FSdfPath, TSet<TWeakObjectPtr<UObject>>> PrimPathToAssets;
 	mutable FRWLock PrimPathToAssetsLock;
+
+	// Paths to material prims that are actually used by mesh prims in the scene, given the current settings for
+	// render context, material purpose, variant selections, etc.
+	TSet<UE::FSdfPath> UsedMaterialPaths;
+	mutable FRWLock UsedMaterialPathsLock;
 };
 
 FUsdInfoCache::FUsdInfoCache()
@@ -77,6 +83,10 @@ bool FUsdInfoCache::Serialize( FArchive& Ar )
 			FWriteScopeLock ScopeLock(ImplPtr->PrimPathToAssetsLock);
 			Ar << ImplPtr->PrimPathToAssets;
 		}
+		{
+			FWriteScopeLock ScopeLock(ImplPtr->UsedMaterialPathsLock);
+			Ar << ImplPtr->UsedMaterialPaths;
+		}
 	}
 
 	return true;
@@ -171,6 +181,18 @@ UE::FSdfPath FUsdInfoCache::UnwindToNonCollapsedPath( const UE::FSdfPath& Path,
 	return Path;
 }
 
+bool FUsdInfoCache::IsMaterialUsed(const UE::FSdfPath& Path) const
+{
+	if (FUsdInfoCacheImpl* ImplPtr = Impl.Get())
+	{
+		FReadScopeLock ScopeLock(ImplPtr->UsedMaterialPathsLock);
+
+		return ImplPtr->UsedMaterialPaths.Contains(Path);
+	}
+
+	return false;
+}
+
 namespace UE::USDInfoCacheImpl::Private
 {
 #if USE_USD_SDK
@@ -268,6 +290,7 @@ namespace UE::USDInfoCacheImpl::Private
 	void RecursiveRebuildCache(
 		const pxr::UsdPrim& UsdPrim,
 		FUsdSchemaTranslationContext& Context,
+		const pxr::TfToken& MaterialPurposeToken,
 		FUsdInfoCache::FUsdInfoCacheImpl& Impl,
 		FUsdSchemaTranslatorRegistry& Registry,
 		TMap< UE::FSdfPath, TArray<UsdUtils::FUsdPrimMaterialSlot>>& InOutSubtreeToMaterialSlots,
@@ -312,6 +335,16 @@ namespace UE::USDInfoCacheImpl::Private
 			}
 		}
 
+		pxr::UsdShadeMaterialBindingAPI BindingAPI{UsdPrim};
+		if (BindingAPI)
+		{
+			if (pxr::UsdShadeMaterial ShadeMaterial = BindingAPI.ComputeBoundMaterial(MaterialPurposeToken))
+			{
+				FWriteScopeLock ScopeLock(Impl.UsedMaterialPathsLock);
+				Impl.UsedMaterialPaths.Add(UE::FSdfPath{ShadeMaterial.GetPrim().GetPath()});
+			}
+		}
+
 		pxr::UsdPrimSiblingRange PrimChildren = UsdPrim.GetFilteredChildren( pxr::UsdTraverseInstanceProxies( pxr::UsdPrimAllPrimsPredicate ) );
 
 		TArray<pxr::UsdPrim> Prims;
@@ -336,6 +369,7 @@ namespace UE::USDInfoCacheImpl::Private
 				RecursiveRebuildCache(
 					Prims[ Index ],
 					Context,
+					MaterialPurposeToken,
 					Impl,
 					Registry,
 					InOutSubtreeToMaterialSlots,
@@ -700,7 +734,16 @@ void FUsdInfoCache::RebuildCacheForSubtree( const UE::FUsdPrim& Prim, FUsdSchema
 			return;
 		}
 
-		Clear();
+		// Don't call Clear() here as we don't want to get rid of PrimPathToAssets because we're rebuilding the cache,
+		// as that info is also linked to the asset cache
+		{
+			FWriteScopeLock ScopeLock(ImplPtr->InfoMapLock);
+			ImplPtr->InfoMap.Empty();
+		}
+		{
+			FWriteScopeLock ScopeLock(ImplPtr->UsedMaterialPathsLock);
+			ImplPtr->UsedMaterialPaths.Empty();
+		}
 
 		IUsdSchemasModule& UsdSchemasModule = FModuleManager::Get().LoadModuleChecked< IUsdSchemasModule >( TEXT( "USDSchemas" ) );
 		FUsdSchemaTranslatorRegistry& Registry = UsdSchemasModule.GetTranslatorRegistry();
@@ -708,11 +751,18 @@ void FUsdInfoCache::RebuildCacheForSubtree( const UE::FUsdPrim& Prim, FUsdSchema
 		TMap< UE::FSdfPath, TArray<UsdUtils::FUsdPrimMaterialSlot>> TempSubtreeSlots;
 		TArray< FString > PointInstancerPaths;
 
+		pxr::TfToken MaterialPurposeToken = pxr::UsdShadeTokens->allPurpose;
+		if (!Context.MaterialPurpose.IsNone())
+		{
+			MaterialPurposeToken = UnrealToUsd::ConvertToken(*Context.MaterialPurpose.ToString()).Get();
+		}
+
 		uint64 SubtreeVertexCount = 0;
 		TArray<UsdUtils::FUsdPrimMaterialSlot> SubtreeSlots;
 		UE::USDInfoCacheImpl::Private::RecursiveRebuildCache(
 			UsdPrim,
 			Context,
+			MaterialPurposeToken,
 			*ImplPtr,
 			Registry,
 			TempSubtreeSlots,
@@ -742,8 +792,18 @@ void FUsdInfoCache::Clear()
 {
 	if ( FUsdInfoCacheImpl* ImplPtr = Impl.Get() )
 	{
-		FWriteScopeLock ScopeLock(ImplPtr->InfoMapLock);
-		ImplPtr->InfoMap.Empty();
+		{
+			FWriteScopeLock ScopeLock(ImplPtr->InfoMapLock);
+			ImplPtr->InfoMap.Empty();
+		}
+		{
+			FWriteScopeLock ScopeLock(ImplPtr->PrimPathToAssetsLock);
+			ImplPtr->PrimPathToAssets.Empty();
+		}
+		{
+			FWriteScopeLock ScopeLock(ImplPtr->UsedMaterialPathsLock);
+			ImplPtr->UsedMaterialPaths.Empty();
+		}
 	}
 }
 
diff --git a/Engine/Plugins/Importers/USDImporter/Source/USDSchemas/Private/USDShadeMaterialTranslator.cpp b/Engine/Plugins/Importers/USDImporter/Source/USDSchemas/Private/USDShadeMaterialTranslator.cpp
--- a/Engine/Plugins/Importers/USDImporter/Source/USDSchemas/Private/USDShadeMaterialTranslator.cpp
+++ b/Engine/Plugins/Importers/USDImporter/Source/USDSchemas/Private/USDShadeMaterialTranslator.cpp
@@ -223,6 +223,15 @@ void FUsdShadeMaterialTranslator::CreateAssets()
 		return;
 	}
 
+	if (Context->bTranslateOnlyUsedMaterials && Context->InfoCache)
+	{
+		if (!Context->InfoCache->IsMaterialUsed(PrimPath))
+		{
+			UE_LOG(LogUsd, Verbose, TEXT("Skipping creating assets for material prim '%s' as it is not currently bound by any prim."), *PrimPath.GetString());
+			return;
+		}
+	}
+
 	TRACE_CPUPROFILER_EVENT_SCOPE( FUsdShadeMaterialTranslator::CreateAssets );
 
 	const pxr::TfToken RenderContextToken =
@@ -433,6 +442,20 @@ void FUsdShadeMaterialTranslator::CreateAssets()
 		if (Context->InfoCache)
 		{
 			Context->InfoCache->LinkAssetToPrim(PrimPath, ConvertedMaterial);
+
+			// Also link the textures to the same material prim.
+			// This is important because it lets the stage actor drop its references to old unused textures in the
+			// asset cache if they aren't being used by any other material
+			if (UMaterialInstance* MaterialInstance = Cast<UMaterialInstance>(ConvertedMaterial))
+			{
+				for (const FTextureParameterValue& TextureValue : MaterialInstance->TextureParameterValues)
+				{
+					if (UTexture* Texture = TextureValue.ParameterValue)
+					{
+						Context->InfoCache->LinkAssetToPrim(PrimPath, Texture);
+					}
+				}
+			}
 		}
 
 		if (Context->AssetCache)
diff --git a/Engine/Plugins/Importers/USDImporter/Source/USDSchemas/Public/USDInfoCache.h b/Engine/Plugins/Importers/USDImporter/Source/USDSchemas/Public/USDInfoCache.h
--- a/Engine/Plugins/Importers/USDImporter/Source/USDSchemas/Public/USDInfoCache.h
+++ b/Engine/Plugins/Importers/USDImporter/Source/USDSchemas/Public/USDInfoCache.h
@@ -49,6 +49,9 @@ public:
 	// Returns Path in case it represents an uncollapsed prim, or returns the path to the prim that collapsed it
 	UE::FSdfPath UnwindToNonCollapsedPath(const UE::FSdfPath& Path, ECollapsingType CollapsingType) const;
 
+public:
+	bool IsMaterialUsed(const UE::FSdfPath& Path) const;
+
 public:
 	// Provides the total vertex or material slots counts for each prim *and* its subtree.
 	// This is built inside RebuildCacheForSubtree, so it will factor in the used Context's bMergeIdenticalMaterialSlots.
diff --git a/Engine/Plugins/Importers/USDImporter/Source/USDSchemas/Public/USDSchemaTranslator.h b/Engine/Plugins/Importers/USDImporter/Source/USDSchemas/Public/USDSchemaTranslator.h
--- a/Engine/Plugins/Importers/USDImporter/Source/USDSchemas/Public/USDSchemaTranslator.h
+++ b/Engine/Plugins/Importers/USDImporter/Source/USDSchemas/Public/USDSchemaTranslator.h
@@ -216,6 +216,9 @@ struct USDSCHEMAS_API FUsdSchemaTranslationContext : public TSharedFromThis< FUs
 	/** If true, we will also try creating UAnimSequence skeletal animation assets when parsing SkelRoot prims */
 	bool bAllowParsingSkeletalAnimations = true;
 
+	/** Skip the import of materials that aren't being used by any prim on the stage */
+	bool bTranslateOnlyUsedMaterials = false;
+
 	/** Groom group interpolation settings */
 	TArray<FHairGroupsInterpolation> GroomInterpolationSettings;
 
diff --git a/Engine/Plugins/Importers/USDImporter/Source/USDStage/Private/USDStageActor.cpp b/Engine/Plugins/Importers/USDImporter/Source/USDStage/Private/USDStageActor.cpp
--- a/Engine/Plugins/Importers/USDImporter/Source/USDStage/Private/USDStageActor.cpp
+++ b/Engine/Plugins/Importers/USDImporter/Source/USDStage/Private/USDStageActor.cpp
@@ -121,6 +121,7 @@ struct FUsdStageActorImpl
 		TranslationContext->MaterialToPrimvarToUVIndex = &StageActor->MaterialToPrimvarToUVIndex;
 		TranslationContext->BlendShapesByPath = &StageActor->BlendShapesByPath;
 		TranslationContext->InfoCache = StageActor->InfoCache;
+		TranslationContext->bTranslateOnlyUsedMaterials = true;  // There is no point in turning this off when just opening a stage
 
 		// Its more convenient to toggle between variants using the USDStage window, as opposed to parsing LODs
 		TranslationContext->bAllowInterpretingLODs = false;
@@ -630,7 +631,7 @@ AUsdStageActor::AUsdStageActor()
 	, bMergeIdenticalMaterialSlots(true)
 	, PurposesToLoad((int32)EUsdPurpose::Proxy)
 	, NaniteTriangleThreshold((uint64)1000000)
-	, MaterialPurpose(*UnrealIdentifiers::MaterialAllPurpose)
+	, MaterialPurpose(*UnrealIdentifiers::MaterialPreviewPurpose)
 	, RootMotionHandling(EUsdRootMotionHandling::NoAdditionalRootMotion)
 	, Time(0.0f)
 	, bIsTransitioningIntoPIE(false)
@@ -2075,13 +2076,14 @@ void AUsdStageActor::LoadUsdStage()
 	// the tracks
 	PrimsToAnimate.Reset();
 
-	TSharedRef< FUsdSchemaTranslationContext > TranslationContext = FUsdStageActorImpl::CreateUsdSchemaTranslationContext(this, RootTwin->PrimPath);
-
-	SlowTask.EnterProgressFrame(0.1f);
 	if (!InfoCache.IsValid())
 	{
 		InfoCache = MakeShared<FUsdInfoCache>();
 	}
+
+	TSharedRef< FUsdSchemaTranslationContext > TranslationContext = FUsdStageActorImpl::CreateUsdSchemaTranslationContext(this, RootTwin->PrimPath);
+
+	SlowTask.EnterProgressFrame(0.1f);
 	InfoCache->RebuildCacheForSubtree(StageToLoad.GetPseudoRoot(), TranslationContext.Get());
 
 	SlowTask.EnterProgressFrame(0.7f);
diff --git a/Engine/Plugins/Importers/USDImporter/Source/USDStageImporter/Private/USDStageImportOptions.cpp b/Engine/Plugins/Importers/USDImporter/Source/USDStageImporter/Private/USDStageImportOptions.cpp
--- a/Engine/Plugins/Importers/USDImporter/Source/USDStageImporter/Private/USDStageImportOptions.cpp
+++ b/Engine/Plugins/Importers/USDImporter/Source/USDStageImporter/Private/USDStageImportOptions.cpp
@@ -17,12 +17,13 @@ UUsdStageImportOptions::UUsdStageImportOptions(const FObjectInitializer& ObjectI
 	bImportSkeletalAnimations = true;
 	bImportLevelSequences = true;
 	bImportMaterials = true;
+	bImportOnlyUsedMaterials = false;
 
 	PurposesToImport = (int32) (EUsdPurpose::Default | EUsdPurpose::Proxy | EUsdPurpose::Render | EUsdPurpose::Guide);
 	NaniteTriangleThreshold = INT32_MAX;
 	IUsdSchemasModule& UsdSchemasModule = FModuleManager::Get().LoadModuleChecked< IUsdSchemasModule >( TEXT("USDSchemas") );
 	RenderContextToImport = UsdSchemasModule.GetRenderContextRegistry().GetUnrealRenderContext();
-	MaterialPurpose = *UnrealIdentifiers::MaterialAllPurpose;
+	MaterialPurpose = *UnrealIdentifiers::MaterialPreviewPurpose;
 	bOverrideStageOptions = false;
 	StageOptions.MetersPerUnit = 0.01f;
 	StageOptions.UpAxis = EUsdUpAxis::ZAxis;
@@ -81,6 +82,7 @@ void UsdUtils::AddAnalyticsAttributes(
 	InOutAttributes.Emplace( TEXT( "ImportSkeletalAnimations" ), LexToString( Options.bImportSkeletalAnimations ) );
 	InOutAttributes.Emplace( TEXT( "ImportLevelSequences" ), LexToString( Options.bImportLevelSequences ) );
 	InOutAttributes.Emplace( TEXT( "ImportMaterials" ), LexToString( Options.bImportMaterials ) );
+	InOutAttributes.Emplace( TEXT( "ImportOnlyUsedMaterials" ), LexToString( Options.bImportOnlyUsedMaterials) );
 	if ( Options.PrimsToImport != TArray<FString>{ TEXT( "/" ) } )
 	{
 		InOutAttributes.Emplace( TEXT( "NumPrimsToImport" ), LexToString( Options.PrimsToImport.Num() ) );
diff --git a/Engine/Plugins/Importers/USDImporter/Source/USDStageImporter/Private/USDStageImporter.cpp b/Engine/Plugins/Importers/USDImporter/Source/USDStageImporter/Private/USDStageImporter.cpp
--- a/Engine/Plugins/Importers/USDImporter/Source/USDStageImporter/Private/USDStageImporter.cpp
+++ b/Engine/Plugins/Importers/USDImporter/Source/USDStageImporter/Private/USDStageImporter.cpp
@@ -668,7 +668,7 @@ namespace UsdStageImporterImpl
 		if (UUsdAssetCache2* Cache = ImportContext.AssetCache.Get())
 		{
 			FString Hash = Cache->GetHashForAsset(Asset);
-			if (!Hash.IsEmpty())
+			if (!Hash.IsEmpty() && Cache->CanRemoveAsset(Hash))
 			{
 				TakenAsset.Reset(Cache->RemoveAsset(Hash));
 				if (!TakenAsset.IsValid())
@@ -1704,6 +1704,7 @@ void UUsdStageImporter::ImportFromFile(FUsdStageImportContext& ImportContext)
 	TranslationContext->bMergeIdenticalMaterialSlots = ImportContext.ImportOptions->bMergeIdenticalMaterialSlots;
 	TranslationContext->bAllowInterpretingLODs = ImportContext.ImportOptions->bInterpretLODs;
 	TranslationContext->bAllowParsingSkeletalAnimations = ImportContext.ImportOptions->bImportGeometry && ImportContext.ImportOptions->bImportSkeletalAnimations;
+	TranslationContext->bTranslateOnlyUsedMaterials = ImportContext.ImportOptions->bImportOnlyUsedMaterials;
 	TranslationContext->MaterialToPrimvarToUVIndex = &ImportContext.MaterialToPrimvarToUVIndex;
 	TranslationContext->InfoCache = ImportContext.InfoCache;
 	TranslationContext->BlendShapesByPath = &BlendShapesByPath;
@@ -1819,6 +1820,7 @@ bool UUsdStageImporter::ReimportSingleAsset(FUsdStageImportContext& ImportContex
 	TranslationContext->bMergeIdenticalMaterialSlots = ImportContext.ImportOptions->bMergeIdenticalMaterialSlots;
 	TranslationContext->bAllowInterpretingLODs = ImportContext.ImportOptions->bInterpretLODs;
 	TranslationContext->bAllowParsingSkeletalAnimations = ImportContext.ImportOptions->bImportGeometry && ImportContext.ImportOptions->bImportSkeletalAnimations;
+	TranslationContext->bTranslateOnlyUsedMaterials = ImportContext.ImportOptions->bImportOnlyUsedMaterials;
 	TranslationContext->MaterialToPrimvarToUVIndex = &ImportContext.MaterialToPrimvarToUVIndex;
 	TranslationContext->InfoCache = ImportContext.InfoCache;
 	TranslationContext->BlendShapesByPath = &BlendShapesByPath;
diff --git a/Engine/Plugins/Importers/USDImporter/Source/USDStageImporter/Public/USDStageImportOptions.h b/Engine/Plugins/Importers/USDImporter/Source/USDStageImporter/Public/USDStageImportOptions.h
--- a/Engine/Plugins/Importers/USDImporter/Source/USDStageImporter/Public/USDStageImportOptions.h
+++ b/Engine/Plugins/Importers/USDImporter/Source/USDStageImporter/Public/USDStageImportOptions.h
@@ -66,6 +66,13 @@ public:
 	UPROPERTY(BlueprintReadWrite, config, EditAnywhere, Category = "DataToImport", meta = (DisplayName = "Materials & Textures"))
 	bool bImportMaterials;
 
+	/**
+	 * If this is checked, only materials actively used by the stage and import settings will be parsed.
+	 * If this is unchecked, all materials present on the stage will be parsed.
+	 */
+	UPROPERTY(BlueprintReadWrite, config, EditAnywhere, Category = "DataToImport", AdvancedDisplay, meta = (EditCondition = bImportMaterials))
+	bool bImportOnlyUsedMaterials;
+
 	/**
 	 * List of paths of prims to import (e.g. ["/Root/MyBox", "/Root/OtherPrim"]).
 	 * Importing a prim will import its entire subtree.
