USD: Fix ConvertGeomMesh potentially reading texture coordinates out of bounds and leading to random FMeshDescription hashes.

diff --git a/Engine/Plugins/Importers/USDImporter/Source/USDUtilities/Private/USDGeomMeshConversion.cpp b/Engine/Plugins/Importers/USDImporter/Source/USDUtilities/Private/USDGeomMeshConversion.cpp
--- a/Engine/Plugins/Importers/USDImporter/Source/USDUtilities/Private/USDGeomMeshConversion.cpp
+++ b/Engine/Plugins/Importers/USDImporter/Source/USDUtilities/Private/USDGeomMeshConversion.cpp
@@ -1011,9 +1011,16 @@ bool UsdToUnreal::ConvertGeomMesh(
 
 					if ( UVSet.UVIndices.IsSet() )
 					{
-						if ( ensure( UVSet.UVIndices.GetValue().size() > ValueIndex ) )
+						const pxr::VtIntArray& UVIndices = UVSet.UVIndices.GetValue();
+
+						if (ensure(ValueIndex < UVIndices.size()))
 						{
-							UV = UVSet.UVs[ UVSet.UVIndices.GetValue()[ ValueIndex ] ];
+							int UVIndex = UVIndices[ValueIndex];
+
+							if (ensure(UVIndex < UVSet.UVs.size()))
+							{
+								UV = UVSet.UVs[UVIndex];
+							}
 						}
 					}
 					else if ( ensure( UVSet.UVs.size() > ValueIndex ) )
