USD: Fix nullptr dereference when failing to reimport an asset.

diff --git a/Engine/Plugins/Importers/USDImporter/Source/USDStageImporter/Private/USDStageImporter.cpp b/Engine/Plugins/Importers/USDImporter/Source/USDStageImporter/Private/USDStageImporter.cpp
--- a/Engine/Plugins/Importers/USDImporter/Source/USDStageImporter/Private/USDStageImporter.cpp
+++ b/Engine/Plugins/Importers/USDImporter/Source/USDStageImporter/Private/USDStageImporter.cpp
@@ -1397,7 +1397,7 @@ namespace UsdStageImporterImpl
 		for ( UObject* Object : UsedAssetsAndDependencies )
 		{
 			// If it's still on the transient package it means we abandoned this one (maybe we had asset replace policy ignore and hit a conflict)
-			if ( Object->GetOutermost() != GetTransientPackage() )
+			if ( Object && Object->GetOutermost() != GetTransientPackage() )
 			{
 				FAssetRegistryModule::AssetCreated( Object );
 			}
