USD: Fix some stage options not being autofilled in the import options dialog when doing Actions->Import

This was getting in the way when testing imports with the global asset cache, so I just added the remaining lines. We have a bunch of Jiras for these already too

diff --git a/Engine/Plugins/Importers/USDImporter/Source/USDStageEditorViewModels/Private/USDStageViewModel.cpp b/Engine/Plugins/Importers/USDImporter/Source/USDStageEditorViewModels/Private/USDStageViewModel.cpp
--- a/Engine/Plugins/Importers/USDImporter/Source/USDStageEditorViewModels/Private/USDStageViewModel.cpp
+++ b/Engine/Plugins/Importers/USDImporter/Source/USDStageEditorViewModels/Private/USDStageViewModel.cpp
@@ -291,6 +291,11 @@ void FUsdStageViewModel::ImportStage( const TCHAR* TargetContentFolder, UUsdStag
 		ImportContext.ImportOptions->StageOptions.MetersPerUnit = UsdUtils::GetUsdStageMetersPerUnit( UsdStage );
 		ImportContext.ImportOptions->StageOptions.UpAxis = UsdUtils::GetUsdStageUpAxisAsEnum( UsdStage );
 		ImportContext.ImportOptions->ImportTimeCode = StageActor->GetTime();
+		ImportContext.ImportOptions->NaniteTriangleThreshold = StageActor->NaniteTriangleThreshold;
+		ImportContext.ImportOptions->RootMotionHandling = StageActor->RootMotionHandling;
+		ImportContext.ImportOptions->KindsToCollapse = StageActor->KindsToCollapse;
+		ImportContext.ImportOptions->bMergeIdenticalMaterialSlots = StageActor->bMergeIdenticalMaterialSlots;
+
 		ImportContext.bReadFromStageCache = true; // So that we import whatever the user has open right now, even if the file has changes
 
 		const FString RootPath = UsdStage.GetRootLayer().GetRealPath();
