USD: Fix material purpose info not being passed to mesh conversion options.

Not quite sure how we missed this, but we weren't using the material purpose when parsing the *actual meshes*. We only care about it for material assignments when we handled material overrides it seems, which also seems to have hidden the issue

diff --git a/Engine/Plugins/Importers/USDImporter/Source/USDSchemas/Private/USDGeomMeshTranslator.cpp b/Engine/Plugins/Importers/USDImporter/Source/USDSchemas/Private/USDGeomMeshTranslator.cpp
--- a/Engine/Plugins/Importers/USDImporter/Source/USDSchemas/Private/USDGeomMeshTranslator.cpp
+++ b/Engine/Plugins/Importers/USDImporter/Source/USDSchemas/Private/USDGeomMeshTranslator.cpp
@@ -1009,10 +1009,17 @@ namespace UsdGeomMeshTranslatorImpl
 				RenderContextToken = UnrealToUsd::ConvertToken(*Context->RenderContext.ToString()).Get();
 			}
 
+			pxr::TfToken MaterialPurposeToken = pxr::UsdShadeTokens->allPurpose;
+			if (!Context->MaterialPurpose.IsNone())
+			{
+				MaterialPurposeToken = UnrealToUsd::ConvertToken(*Context->MaterialPurpose.ToString()).Get();
+			}
+
 			UsdToUnreal::FUsdMeshConversionOptions Options;
 			Options.PurposesToLoad = Context->PurposesToLoad;
 			Options.bMergeIdenticalMaterialSlots = Context->bMergeIdenticalMaterialSlots;
 			Options.RenderContext = RenderContextToken;
+			Options.MaterialPurpose = MaterialPurposeToken;
 
 			// Create and configure a new USDTrack to be added to the GeometryCache
 			UGeometryCacheTrackUsd* UsdTrack = NewObject< UGeometryCacheTrackUsd >(GeometryCache);
@@ -1398,10 +1405,17 @@ void FGeomMeshCreateAssetsTaskChain::SetupTasks()
 				RenderContextToken = UnrealToUsd::ConvertToken(*Context->RenderContext.ToString()).Get();
 			}
 
+			pxr::TfToken MaterialPurposeToken = pxr::UsdShadeTokens->allPurpose;
+			if (!Context->MaterialPurpose.IsNone())
+			{
+				MaterialPurposeToken = UnrealToUsd::ConvertToken(*Context->MaterialPurpose.ToString()).Get();
+			}
+
 			UsdToUnreal::FUsdMeshConversionOptions Options;
 			Options.TimeCode = Context->Time;
 			Options.PurposesToLoad = Context->PurposesToLoad;
 			Options.RenderContext = RenderContextToken;
+			Options.MaterialPurpose = MaterialPurposeToken;
 			Options.MaterialToPrimvarToUVIndex = MaterialToPrimvarToUVIndex;
 			Options.bMergeIdenticalMaterialSlots = Context->bMergeIdenticalMaterialSlots;
 			Options.AdditionalTransform = AdditionalTransform;
@@ -1463,10 +1477,17 @@ void FGeometryCacheCreateAssetsTaskChain::SetupTasks()
 				RenderContextToken = UnrealToUsd::ConvertToken(*Context->RenderContext.ToString()).Get();
 			}
 
+			pxr::TfToken MaterialPurposeToken = pxr::UsdShadeTokens->allPurpose;
+			if (!Context->MaterialPurpose.IsNone())
+			{
+				MaterialPurposeToken = UnrealToUsd::ConvertToken(*Context->MaterialPurpose.ToString()).Get();
+			}
+
 			UsdToUnreal::FUsdMeshConversionOptions Options;
 			Options.TimeCode = UsdUtils::GetEarliestTimeCode();
 			Options.PurposesToLoad = Context->PurposesToLoad;
 			Options.RenderContext = RenderContextToken;
+			Options.MaterialPurpose = MaterialPurposeToken;
 			Options.MaterialToPrimvarToUVIndex = MaterialToPrimvarToUVIndex;
 			Options.bMergeIdenticalMaterialSlots = Context->bMergeIdenticalMaterialSlots;
 
diff --git a/Engine/Plugins/Importers/USDImporter/Source/USDSchemas/Private/USDGeomPointInstancerTranslator.cpp b/Engine/Plugins/Importers/USDImporter/Source/USDSchemas/Private/USDGeomPointInstancerTranslator.cpp
--- a/Engine/Plugins/Importers/USDImporter/Source/USDSchemas/Private/USDGeomPointInstancerTranslator.cpp
+++ b/Engine/Plugins/Importers/USDImporter/Source/USDSchemas/Private/USDGeomPointInstancerTranslator.cpp
@@ -115,10 +115,17 @@ void FUsdGeomPointInstancerCreateAssetsTaskChain::SetupTasks()
 				RenderContextToken = UnrealToUsd::ConvertToken( *Context->RenderContext.ToString() ).Get();
 			}
 
+			pxr::TfToken MaterialPurposeToken = pxr::UsdShadeTokens->allPurpose;
+			if (!Context->MaterialPurpose.IsNone())
+			{
+				MaterialPurposeToken = UnrealToUsd::ConvertToken(*Context->MaterialPurpose.ToString()).Get();
+			}
+
 			UsdToUnreal::FUsdMeshConversionOptions Options;
 			Options.TimeCode = Context->Time;
 			Options.PurposesToLoad = Context->PurposesToLoad;
 			Options.RenderContext = RenderContextToken;
+			Options.MaterialPurpose = MaterialPurposeToken;
 			Options.MaterialToPrimvarToUVIndex = MaterialToPrimvarToUVIndex;
 			Options.bMergeIdenticalMaterialSlots = Context->bMergeIdenticalMaterialSlots;
 
diff --git a/Engine/Plugins/Importers/USDImporter/Source/USDSchemas/Private/USDGeomXformableTranslator.cpp b/Engine/Plugins/Importers/USDImporter/Source/USDSchemas/Private/USDGeomXformableTranslator.cpp
--- a/Engine/Plugins/Importers/USDImporter/Source/USDSchemas/Private/USDGeomXformableTranslator.cpp
+++ b/Engine/Plugins/Importers/USDImporter/Source/USDSchemas/Private/USDGeomXformableTranslator.cpp
@@ -235,6 +235,12 @@ void FUsdGeomXformableCreateAssetsTaskChain::SetupTasks()
 				RenderContextToken = UnrealToUsd::ConvertToken( *Context->RenderContext.ToString() ).Get();
 			}
 
+			pxr::TfToken MaterialPurposeToken = pxr::UsdShadeTokens->allPurpose;
+			if (!Context->MaterialPurpose.IsNone())
+			{
+				MaterialPurposeToken = UnrealToUsd::ConvertToken(*Context->MaterialPurpose.ToString()).Get();
+			}
+
 			// We're going to put Prim's transform and visibility on the component, so we don't need to bake it into the combined mesh
 			const bool bSkipRootPrimTransformAndVis = true;
 
@@ -242,6 +248,7 @@ void FUsdGeomXformableCreateAssetsTaskChain::SetupTasks()
 			Options.TimeCode = Context->Time;
 			Options.PurposesToLoad = Context->PurposesToLoad;
 			Options.RenderContext = RenderContextToken;
+			Options.MaterialPurpose = MaterialPurposeToken;
 			Options.MaterialToPrimvarToUVIndex = MaterialToPrimvarToUVIndex;
 			Options.bMergeIdenticalMaterialSlots = Context->bMergeIdenticalMaterialSlots;
 
