USD: Fix minor bugs encountered when implementing global asset cache.

Fix reloading one stage triggering all stage actors to reload;
Fix missing Modify() call when resetting a StageActor;
Fix info cache not being initialized before it is needed by the CreateUsdSchemaTranslationContext when loading stages;
Fix some log message being output as LogTemp instead of LogUsd;

Another bizarre bug: We were never checking *which* layers were modified when they were reloaded/added/removed, and it seems to be a global event. When one of the stages reloaded, all stages got an event about "layers being reloaded" and all regenerated their animations...

diff --git a/Engine/Plugins/Importers/USDImporter/Source/USDStage/Private/USDStageActor.cpp b/Engine/Plugins/Importers/USDImporter/Source/USDStage/Private/USDStageActor.cpp
--- a/Engine/Plugins/Importers/USDImporter/Source/USDStage/Private/USDStageActor.cpp
+++ b/Engine/Plugins/Importers/USDImporter/Source/USDStage/Private/USDStageActor.cpp
@@ -740,7 +740,7 @@ AUsdStageActor::AUsdStageActor()
 		UsdListener.GetOnObjectsChanged().AddUObject(this, &AUsdStageActor::OnUsdObjectsChanged);
 
 		UsdListener.GetOnLayersChanged().AddLambda(
-			[&](const TArray< FString >& ChangeVec)
+			[&, this](const TArray< FString >& ChangedLayerIdentifiers)
 			{
 				if (!IsListeningToUsdNotices())
 				{
@@ -753,11 +753,30 @@ AUsdStageActor::AUsdStageActor()
 					SuppressTransaction.Emplace(GUndo, nullptr);
 				}
 
+				const UE::FUsdStage& Stage = GetUsdStage();
+				if (!Stage)
+				{
+					return;
+				}
+
+				TSet<FString> UsedLayerIdentifiers;
+				TArray<UE::FSdfLayer> UsedLayers = Stage.GetUsedLayers();
+				UsedLayerIdentifiers.Reserve(UsedLayers.Num());
+				for (const UE::FSdfLayer& UsedLayer : UsedLayers)
+				{
+					UsedLayerIdentifiers.Add(UsedLayer.GetIdentifier());
+				}
+
 				// Check to see if any layer reloaded. If so, rebuild all of our animations as a single layer changing
 				// might propagate timecodes through all level sequences
-				for (const FString& ChangeVecItem : ChangeVec)
+				for (const FString& ChangedLayerIdentifier : ChangedLayerIdentifiers)
 				{
-					UE_LOG(LogUsd, Verbose, TEXT("Reloading animations because layer '%s' was added/removed/reloaded"), *ChangeVecItem);
+					if (!UsedLayerIdentifiers.Contains(ChangedLayerIdentifier))
+					{
+						continue;
+					}
+
+					UE_LOG(LogUsd, Verbose, TEXT("Reloading animations because layer '%s' was added/removed/reloaded"), *ChangedLayerIdentifier);
 					ReloadAnimations();
 
 					// Make sure our PrimsToAnimate and the LevelSequenceHelper are kept in sync, because we'll use PrimsToAnimate to
@@ -859,7 +878,7 @@ void AUsdStageActor::IsolateLayer( const UE::FSdfLayer& Layer )
 		TArray<UE::FSdfLayer> CurrentLayerStack = FreshCurrentStage.GetLayerStack();
 		if (!CurrentLayerStack.Contains(Layer))
 		{
-			UE_LOG(LogTemp, Warning, TEXT("Failed to isolate layer '%s' as it is not part of the currently opened USD Stage's local layer stack"),
+			UE_LOG(LogUsd, Warning, TEXT("Failed to isolate layer '%s' as it is not part of the currently opened USD Stage's local layer stack"),
 				*Layer.GetIdentifier(),
 				*UsdStage.GetRootLayer().GetIdentifier()
 			);
@@ -1196,6 +1215,8 @@ void AUsdStageActor::OnUsdObjectsChanged(const UsdUtils::FObjectChangesByPath& I
 
 USDSTAGE_API void AUsdStageActor::Reset()
 {
+	Modify();
+
 	Super::Reset();
 
 	UnloadUsdStage();
diff --git a/Engine/Plugins/Importers/USDImporter/Source/UnrealUSDWrapper/Private/USDListener.cpp b/Engine/Plugins/Importers/USDImporter/Source/UnrealUSDWrapper/Private/USDListener.cpp
--- a/Engine/Plugins/Importers/USDImporter/Source/UnrealUSDWrapper/Private/USDListener.cpp
+++ b/Engine/Plugins/Importers/USDImporter/Source/UnrealUSDWrapper/Private/USDListener.cpp
@@ -388,7 +388,7 @@ void FUsdListenerImpl::HandleLayersChangedNotice( const pxr::SdfNotice::LayersDi
 				const pxr::SdfChangeList::Entry::_Flags& Flags = Change.second.flags;
 				if (Flags.didReloadContent)
 				{
-					LayersNames.Add( ANSI_TO_TCHAR( Change.first.GetString().c_str() ) );
+					LayersNames.Add(ANSI_TO_TCHAR(ChangeVecItem.first->GetIdentifier().c_str()));
 				}
 			}
 		}
