USD: Also fetch "used material" info by reading the relationship directly, even if no material binding schema is present.

In the USD version we have in 5.3 when a prim just has the material:binding relationship but not the MaterialBindingAPI schema USD itself emits a warning saying "hey these material bindings are invalid since the prim doesn't have the schema".

In the USD version we have in 5.2 we don't get such warning... This is an issue because the new mechanism we added to only parse materials used by the current material purpose would consider a "valid assignment" only one done with the schema. This means that if you're missing the schema we'd ignore the binding (as if it was from an unused material purpose), not generated the materials and there would be no feedback at all.

I'm not sure if no warning means it's an allowed thing to do or not, but we have always parsed material assignments like that anyway (only this filtering of "used materials" that was exclusively checking via the schema). I think we should probably support this in 5.2, even if only due to the fact that if the user is missing these schemas on their prims, he will get no materials parsed and zero feedback about it.

diff --git a/Engine/Plugins/Importers/USDImporter/Source/USDSchemas/Private/USDInfoCache.cpp b/Engine/Plugins/Importers/USDImporter/Source/USDSchemas/Private/USDInfoCache.cpp
--- a/Engine/Plugins/Importers/USDImporter/Source/USDSchemas/Private/USDInfoCache.cpp
+++ b/Engine/Plugins/Importers/USDImporter/Source/USDSchemas/Private/USDInfoCache.cpp
@@ -302,6 +302,12 @@ namespace UE::USDInfoCacheImpl::Private
 	)
 	{
 		TRACE_CPUPROFILER_EVENT_SCOPE( UE::USDInfoCacheImpl::Private::RecursiveRebuildCache );
+
+		if (!UsdPrim)
+		{
+			return;
+		}
+
 		FScopedUsdAllocs Allocs;
 
 		pxr::SdfPath UsdPrimPath = UsdPrim.GetPrimPath();
@@ -335,13 +341,33 @@ namespace UE::USDInfoCacheImpl::Private
 			}
 		}
 
-		pxr::UsdShadeMaterialBindingAPI BindingAPI{UsdPrim};
-		if (BindingAPI)
+		if (!UsdPrim.IsPseudoRoot())
 		{
-			if (pxr::UsdShadeMaterial ShadeMaterial = BindingAPI.ComputeBoundMaterial(MaterialPurposeToken))
+			pxr::UsdShadeMaterialBindingAPI BindingAPI{UsdPrim};
+			if (BindingAPI)
+			{
+				if (pxr::UsdShadeMaterial ShadeMaterial = BindingAPI.ComputeBoundMaterial(MaterialPurposeToken))
+				{
+					FWriteScopeLock ScopeLock(Impl.UsedMaterialPathsLock);
+					Impl.UsedMaterialPaths.Add(UE::FSdfPath{ShadeMaterial.GetPrim().GetPath()});
+				}
+			}
+			else if (pxr::UsdRelationship Relationship = UsdPrim.GetRelationship(pxr::UsdShadeTokens->materialBinding))
 			{
-				FWriteScopeLock ScopeLock(Impl.UsedMaterialPathsLock);
-				Impl.UsedMaterialPaths.Add(UE::FSdfPath{ShadeMaterial.GetPrim().GetPath()});
+				pxr::SdfPathVector Targets;
+				Relationship.GetTargets(&Targets);
+
+				if (Targets.size() > 0)
+				{
+					const pxr::SdfPath& TargetMaterialPrimPath = Targets[0];
+					pxr::UsdPrim MaterialPrim = UsdPrim.GetStage()->GetPrimAtPath(TargetMaterialPrimPath);
+					pxr::UsdShadeMaterial UsdShadeMaterial{MaterialPrim};
+					if (UsdShadeMaterial)
+					{
+						FWriteScopeLock ScopeLock(Impl.UsedMaterialPathsLock);
+						Impl.UsedMaterialPaths.Add(UE::FSdfPath{TargetMaterialPrimPath});
+					}
+				}
 			}
 		}
 
