USD: Adapt bake to control rig code path to handle rigs that modify their own topology on construction.

diff --git a/Engine/Plugins/Importers/USDImporter/Source/USDUtilities/Private/USDSkeletalDataConversion.cpp b/Engine/Plugins/Importers/USDImporter/Source/USDUtilities/Private/USDSkeletalDataConversion.cpp
--- a/Engine/Plugins/Importers/USDImporter/Source/USDUtilities/Private/USDSkeletalDataConversion.cpp
+++ b/Engine/Plugins/Importers/USDImporter/Source/USDUtilities/Private/USDSkeletalDataConversion.cpp
@@ -3326,9 +3326,9 @@ bool UnrealToUsd::ConvertControlRigSection(
 
 	double StartTime = FPlatformTime::Cycles64();
 
-	// TODO: This does not seem necessary
 	ControlRig->Initialize();
 	ControlRig->RequestInit();
+	ControlRig->Evaluate_AnyThread(); // Important as it runs the Construction event, which can change topology
 
 	// Record how the topology looks while we setup our arrays and maps. If this changes during
 	// baking we'll just drop everything and return
@@ -3343,11 +3343,22 @@ bool UnrealToUsd::ConvertControlRigSection(
 	// This works because the topology won't change in here, and bone names are unique across the entire skeleton
 	// Its possible we'll be putting INDEX_NONEs into RigIndexToRefSkeletonIndex, but that's alright.
 	TArray<int32> RigJointIndexToRefSkeletonIndex;
-	TArray<FRigBoneElement*> InitialBoneElements = InitialHierarchy->GetBones();
-	for ( FRigBoneElement* RigBone : InitialBoneElements )
-	{
-		RigJointIndexToRefSkeletonIndex.Add( InRefSkeleton.FindBoneIndex( RigBone->GetName() ) );
-	}
+	TFunction<void()> RegenerateRigJointIndexToRefSkeletonIndex =
+		[ControlRig, &RigJointIndexToRefSkeletonIndex, &InRefSkeleton]()
+		{
+			URigHierarchy* Hierarchy = ControlRig->GetHierarchy();
+			if (!Hierarchy)
+			{
+				return;
+			}
+
+			RigJointIndexToRefSkeletonIndex.Reset();
+			for (FRigBoneElement* RigBone : Hierarchy->GetBones())
+			{
+				RigJointIndexToRefSkeletonIndex.Add(InRefSkeleton.FindBoneIndex(RigBone->GetName()));
+			}
+		};
+	RegenerateRigJointIndexToRefSkeletonIndex();
 
 	pxr::UsdAttribute JointsAttr = SkelAnim.CreateJointsAttr();
 	UnrealToUsd::ConvertJointsAttribute( InRefSkeleton, JointsAttr );
@@ -3604,12 +3615,19 @@ bool UnrealToUsd::ConvertControlRigSection(
 		ControlRig->Evaluate_AnyThread();
 
 		URigHierarchy* Hierarchy = ControlRig->GetHierarchy();
-		if ( !Hierarchy || Hierarchy->GetTopologyVersion() != TopologyVersion )
+		if ( !Hierarchy )
 		{
-			UE_LOG( LogUsd, Error, TEXT( "Baking Control Rig tracks for rig '%s' failed because the rig changed topology during baking process" ), *ControlRig->GetPathName() );
+			UE_LOG( LogUsd, Error, TEXT( "Baking Control Rig tracks for rig '%s' failed" ), *ControlRig->GetPathName() );
 			return false;
 		}
 
+		if (Hierarchy->GetTopologyVersion() != TopologyVersion)
+		{
+			UE_LOG(LogUsd, Log, TEXT("Regenerating ControlRig to reference skeleton mapping for rig '%s' as its topology changed"), *ControlRig->GetPathName());
+			RegenerateRigJointIndexToRefSkeletonIndex();
+			TopologyVersion = Hierarchy->GetTopologyVersion();
+		}
+
 		// Sadly we have to fetch these each frame as these are regenerated on each evaluation of the Sequencer
 		// (c.f. FControlRigBindingHelper::BindToSequencerInstance, URigHierarchy::CopyHierarchy)
 		TArray<FRigBoneElement*> BoneElements = Hierarchy->GetBones();
