USDImporter: align engine-side prim IsA() and HasAPI() function signatures with USD API

As part of the work to support schema versioning in USD, a change went into v23.02
that refactors the IsA() and HasAPI() functions on UsdPrim. In particular with HasAPI(),
the instanceName parameter is only relevant when the schema being tested is a multiple
apply API schema, and the refactor makes it so that passing an empty instanceName
is now a coding error.

This change deprecates the single HasAPI() function that takes an optional instance name
parameter and replaces it with one- and two-FName parameter overloads that match the
one- and two-token counterparts on UsdPrim.

Using the engine's current v22.11 version of USD, we continue to get the TfType by its
schema identifier before passing it to the USD-side IsA() or HasAPI() functions that take
the schema by type. After the engine's USD version is upgraded to v23.02, we can instead
call the UsdPrim functions that take the schema identifier as a token directly and skip
getting the type ourselves.

More detail on the relevant change in USD is available in the commit here:

[REDACTED_PROFILE_URL]/USD/commit/5933c201ebc5eacb9198c615e30666e672e726c2

diff --git a/Engine/Plugins/Importers/USDImporter/Source/UnrealUSDWrapper/Private/UsdWrappers/UsdPrim.cpp b/Engine/Plugins/Importers/USDImporter/Source/UnrealUSDWrapper/Private/UsdWrappers/UsdPrim.cpp
--- a/Engine/Plugins/Importers/USDImporter/Source/UnrealUSDWrapper/Private/UsdWrappers/UsdPrim.cpp
+++ b/Engine/Plugins/Importers/USDImporter/Source/UnrealUSDWrapper/Private/UsdWrappers/UsdPrim.cpp
@@ -242,12 +242,12 @@ namespace UE
 		return AppliedSchemas;
 	}
 
-	bool FUsdPrim::IsA( FName SchemaType ) const
+	bool FUsdPrim::IsA( FName SchemaIdentifier ) const
 	{
 #if USE_USD_SDK
 		FScopedUsdAllocs Allocs;
 
-		pxr::TfType Type = pxr::UsdSchemaRegistry::GetTypeFromName( pxr::TfToken( TCHAR_TO_ANSI( *SchemaType.ToString() ) ) );
+		pxr::TfType Type = pxr::UsdSchemaRegistry::GetTypeFromName( pxr::TfToken( TCHAR_TO_ANSI( *SchemaIdentifier.ToString() ) ) );
 		if ( Type.IsUnknown() )
 		{
 			return false;
@@ -259,18 +259,35 @@ namespace UE
 #endif // #if USE_USD_SDK
 	}
 
-	bool FUsdPrim::HasAPI( FName SchemaType, TOptional<FName> InstanceName ) const
+	bool FUsdPrim::HasAPI( FName SchemaIdentifier ) const
+	{
+#if USE_USD_SDK
+		FScopedUsdAllocs Allocs;
+
+		pxr::TfType Type = pxr::UsdSchemaRegistry::GetTypeFromName( pxr::TfToken( TCHAR_TO_ANSI( *SchemaIdentifier.ToString() ) ) );
+		if ( Type.IsUnknown() )
+		{
+			return false;
+		}
+
+		return Impl->PxrUsdPrim.Get().HasAPI( Type );
+#else
+		return false;
+#endif // #if USE_USD_SDK
+	}
+
+	bool FUsdPrim::HasAPI( FName SchemaIdentifier, FName InstanceName ) const
 	{
 #if USE_USD_SDK
 		FScopedUsdAllocs Allocs;
 
-		pxr::TfType Type = pxr::UsdSchemaRegistry::GetTypeFromName( pxr::TfToken( TCHAR_TO_ANSI( *SchemaType.ToString() ) ) );
+		pxr::TfType Type = pxr::UsdSchemaRegistry::GetTypeFromName( pxr::TfToken( TCHAR_TO_ANSI( *SchemaIdentifier.ToString() ) ) );
 		if ( Type.IsUnknown() )
 		{
 			return false;
 		}
 
-		pxr::TfToken UsdInstanceName = InstanceName.IsSet() ? pxr::TfToken( TCHAR_TO_ANSI( *InstanceName.GetValue().ToString() ) ) : pxr::TfToken();
+		pxr::TfToken UsdInstanceName = pxr::TfToken( TCHAR_TO_ANSI( *InstanceName.ToString() ) );
 
 		return Impl->PxrUsdPrim.Get().HasAPI( Type, UsdInstanceName );
 #else
@@ -278,6 +295,17 @@ namespace UE
 #endif // #if USE_USD_SDK
 	}
 
+	// Deprecated in 5.3
+	bool FUsdPrim::HasAPI( FName SchemaType, TOptional<FName> InstanceName ) const
+	{
+		if ( InstanceName.IsSet() )
+		{
+			return HasAPI( SchemaType, InstanceName.GetValue() );
+		}
+
+		return HasAPI( SchemaType );
+	}
+
 	const FSdfPath FUsdPrim::GetPrimPath() const
 	{
 #if USE_USD_SDK
diff --git a/Engine/Plugins/Importers/USDImporter/Source/UnrealUSDWrapper/Public/UsdWrappers/UsdPrim.h b/Engine/Plugins/Importers/USDImporter/Source/UnrealUSDWrapper/Public/UsdWrappers/UsdPrim.h
--- a/Engine/Plugins/Importers/USDImporter/Source/UnrealUSDWrapper/Public/UsdWrappers/UsdPrim.h
+++ b/Engine/Plugins/Importers/USDImporter/Source/UnrealUSDWrapper/Public/UsdWrappers/UsdPrim.h
@@ -85,8 +85,12 @@ namespace UE
 
 		TArray<FName> GetAppliedSchemas() const;
 
-		bool IsA( FName SchemaType ) const;
-		bool HasAPI( FName SchemaType, TOptional<FName> InstanceName = {} ) const;
+		bool IsA( FName SchemaIdentifier ) const;
+		bool HasAPI( FName SchemaIdentifier ) const;
+		bool HasAPI( FName SchemaIdentifier, FName InstanceName ) const;
+
+		UE_DEPRECATED( 5.3, "Please use the one FName overload for single apply API schemas, or the two FName overload for multiple apply API schemas." )
+		bool HasAPI( FName SchemaType, TOptional<FName> InstanceName ) const;
 
 		const FSdfPath GetPrimPath() const;
 		FUsdStage GetStage() const;
