USD: Fix overridden opinions being incorrectly displayed in some scenarios.

diff --git a/Engine/Plugins/Importers/USDImporter/Source/USDUtilities/Private/USDPrimConversion.cpp b/Engine/Plugins/Importers/USDImporter/Source/USDUtilities/Private/USDPrimConversion.cpp
--- a/Engine/Plugins/Importers/USDImporter/Source/USDUtilities/Private/USDPrimConversion.cpp
+++ b/Engine/Plugins/Importers/USDImporter/Source/USDUtilities/Private/USDPrimConversion.cpp
@@ -2858,279 +2858,344 @@ UnrealToUsd::FPropertyTrackWriter UnrealToUsd::CreatePropertyTrackWriter( const
 	pxr::UsdStageRefPtr UsdStage = UsdPrim.GetStage();
 	FUsdStageInfo StageInfo{ UsdStage };
 
-	pxr::SdfChangeBlock ChangeBlock;
-
 	pxr::UsdAttribute Attr;
+	{
+		pxr::SdfChangeBlock ChangeBlock;
 
-	const FName& PropertyPath = Track.GetPropertyPath();
+		const FName& PropertyPath = Track.GetPropertyPath();
 
-	// SceneComponent
-	{
-		if ( PropertyPath == UnrealIdentifiers::TransformPropertyName )
+		// Note that it's important that each individual case authors a spec for the relevant attribute right now,
+		// even though it returns a FPropertyTrackWriter to actually do the baking later. This will be done with
+		// CreateXformOpOrderAttr() or CreateVisibilityAttr() or CreateFocalLengthAttr(), etc.
+		// This just ensures that our check for overridden attributes at the bottom of this function call works
+		// properly.
+
+		// SceneComponent
 		{
-			if ( pxr::UsdGeomXformable Xformable{ UsdPrim } )
+			if ( PropertyPath == UnrealIdentifiers::TransformPropertyName )
 			{
-				Xformable.CreateXformOpOrderAttr();
-
-				if ( pxr::UsdGeomXformOp TransformOp = UE::USDPrimConversionImpl::Private::ForceMatrixXform( Xformable ) )
+				if ( pxr::UsdGeomXformable Xformable{ UsdPrim } )
 				{
-					Attr = TransformOp.GetAttr();
+					Xformable.CreateXformOpOrderAttr();
 
-					// Compensate different orientation for light or camera components
-					FTransform AdditionalRotation = FTransform::Identity;
-					if ( UsdPrim.IsA< pxr::UsdGeomCamera >() || UsdPrim.HasAPI< pxr::UsdLuxLightAPI >() )
+					if ( pxr::UsdGeomXformOp TransformOp = UE::USDPrimConversionImpl::Private::ForceMatrixXform( Xformable ) )
 					{
-						AdditionalRotation = FTransform( FRotator( 0.0f, 90.0f, 0.0f ) );
+						Attr = TransformOp.GetAttr();
 
-						if ( StageInfo.UpAxis == EUsdUpAxis::ZAxis )
+						// Compensate different orientation for light or camera components
+						FTransform AdditionalRotation = FTransform::Identity;
+						if ( UsdPrim.IsA< pxr::UsdGeomCamera >() || UsdPrim.HasAPI< pxr::UsdLuxLightAPI >() )
 						{
-							AdditionalRotation *= FTransform( FRotator( 90.0f, 0.0f, 0.0f ) );
+							AdditionalRotation = FTransform( FRotator( 0.0f, 90.0f, 0.0f ) );
+
+							if ( StageInfo.UpAxis == EUsdUpAxis::ZAxis )
+							{
+								AdditionalRotation *= FTransform( FRotator( 90.0f, 0.0f, 0.0f ) );
+							}
 						}
-					}
 
-					// Invert compensation applied to parent if it's a light or camera component
-					if ( const USceneComponent* AttachParent = Component.GetAttachParent() )
-					{
-						if ( AttachParent->IsA( UCineCameraComponent::StaticClass() ) ||
-							AttachParent->IsA( ULightComponent::StaticClass() ) )
+						// Invert compensation applied to parent if it's a light or camera component
+						if ( const USceneComponent* AttachParent = Component.GetAttachParent() )
 						{
-							FTransform InverseCompensation = FTransform( FRotator( 0.0f, 90.f, 0.0f ) );
-
-							if ( UsdUtils::GetUsdStageUpAxis( UsdStage ) == pxr::UsdGeomTokens->z )
+							if ( AttachParent->IsA( UCineCameraComponent::StaticClass() ) ||
+								AttachParent->IsA( ULightComponent::StaticClass() ) )
 							{
-								InverseCompensation *= FTransform( FRotator( 90.0f, 0.f, 0.0f ) );
-							}
+								FTransform InverseCompensation = FTransform( FRotator( 0.0f, 90.f, 0.0f ) );
 
-							AdditionalRotation = AdditionalRotation * InverseCompensation.Inverse();
+								if ( UsdUtils::GetUsdStageUpAxis( UsdStage ) == pxr::UsdGeomTokens->z )
+								{
+									InverseCompensation *= FTransform( FRotator( 90.0f, 0.f, 0.0f ) );
+								}
+
+								AdditionalRotation = AdditionalRotation * InverseCompensation.Inverse();
+							}
 						}
-					}
 
-					Result.TransformWriter = [&Component, AdditionalRotation, StageInfo, Attr]( const FTransform& UEValue, double UsdTimeCode )
-					{
-						FTransform FinalUETransform = AdditionalRotation * UEValue;
-						pxr::GfMatrix4d UsdTransform = UnrealToUsd::ConvertTransform( StageInfo, FinalUETransform );
-						Attr.Set< pxr::GfMatrix4d >( UsdTransform, UsdTimeCode );
-					};
+						Result.TransformWriter = [&Component, AdditionalRotation, StageInfo, Attr]( const FTransform& UEValue, double UsdTimeCode )
+						{
+							FTransform FinalUETransform = AdditionalRotation * UEValue;
+							pxr::GfMatrix4d UsdTransform = UnrealToUsd::ConvertTransform( StageInfo, FinalUETransform );
+							Attr.Set< pxr::GfMatrix4d >( UsdTransform, UsdTimeCode );
+						};
+					}
 				}
 			}
-		}
-		// bHidden is for the actor, and bHiddenInGame is for a component
-		// A component is only visible when it's not hidden and its actor is not hidden
-		// A bHidden is just handled like a bHiddenInGame for the actor's root component
-		// Whenever we handle a bHiddenInGame, we always combine it with the actor's bHidden
-		else if ( PropertyPath == UnrealIdentifiers::HiddenPropertyName || PropertyPath == UnrealIdentifiers::HiddenInGamePropertyName )
-		{
-			if ( pxr::UsdGeomImageable Imageable{ UsdPrim } )
+			// bHidden is for the actor, and bHiddenInGame is for a component
+			// A component is only visible when it's not hidden and its actor is not hidden
+			// A bHidden is just handled like a bHiddenInGame for the actor's root component
+			// Whenever we handle a bHiddenInGame, we always combine it with the actor's bHidden
+			else if ( PropertyPath == UnrealIdentifiers::HiddenPropertyName || PropertyPath == UnrealIdentifiers::HiddenInGamePropertyName )
 			{
-				Attr = Imageable.CreateVisibilityAttr();
-				if ( Attr )
+				if ( pxr::UsdGeomImageable Imageable{ UsdPrim } )
 				{
-					Result.BoolWriter = [Imageable, Attr]( bool UEValue, double UsdTimeCode )
+					Attr = Imageable.CreateVisibilityAttr();
+					if ( Attr )
 					{
-						if ( UEValue )
+						Result.BoolWriter = [Imageable, Attr]( bool UEValue, double UsdTimeCode )
 						{
-							// We have to do both here as MakeVisible will ensure we also flip any parent prims,
-							// and setting the attribute will ensure we write a timeSample. Otherwise if MakeVisible
-							// finds that the prim should already be visible due to a stronger opinion, it won't write anything
-							Attr.Set<pxr::TfToken>( pxr::UsdGeomTokens->inherited, UsdTimeCode );
-							Imageable.MakeVisible( UsdTimeCode );
-						}
-						else
-						{
-							Attr.Set<pxr::TfToken>( pxr::UsdGeomTokens->invisible, UsdTimeCode );
-							Imageable.MakeInvisible( UsdTimeCode );
-						}
-					};
+							if ( UEValue )
+							{
+								// We have to do both here as MakeVisible will ensure we also flip any parent prims,
+								// and setting the attribute will ensure we write a timeSample. Otherwise if MakeVisible
+								// finds that the prim should already be visible due to a stronger opinion, it won't write anything
+								Attr.Set<pxr::TfToken>( pxr::UsdGeomTokens->inherited, UsdTimeCode );
+								Imageable.MakeVisible( UsdTimeCode );
+							}
+							else
+							{
+								Attr.Set<pxr::TfToken>( pxr::UsdGeomTokens->invisible, UsdTimeCode );
+								Imageable.MakeInvisible( UsdTimeCode );
+							}
+						};
+					}
 				}
 			}
 		}
-	}
 
-	if ( pxr::UsdGeomCamera Camera{ UsdPrim } )
-	{
-		bool bConvertDistance = true;
-
-		if ( PropertyPath == UnrealIdentifiers::CurrentFocalLengthPropertyName )
-		{
-			Attr = Camera.CreateFocalLengthAttr();
-		}
-		else if ( PropertyPath == UnrealIdentifiers::ManualFocusDistancePropertyName )
-		{
-			Attr = Camera.CreateFocusDistanceAttr();
-		}
-		else if ( PropertyPath == UnrealIdentifiers::CurrentAperturePropertyName )
-		{
-			bConvertDistance = false;
-			Attr = Camera.CreateFStopAttr();
-		}
-		else if ( PropertyPath == UnrealIdentifiers::SensorWidthPropertyName )
+		if ( pxr::UsdGeomCamera Camera{ UsdPrim } )
 		{
-			Attr = Camera.CreateHorizontalApertureAttr();
-		}
-		else if ( PropertyPath == UnrealIdentifiers::SensorHeightPropertyName )
-		{
-			Attr = Camera.CreateVerticalApertureAttr();
-		}
+			bool bConvertDistance = true;
 
-		if ( Attr )
-		{
-			if ( bConvertDistance )
+			if ( PropertyPath == UnrealIdentifiers::CurrentFocalLengthPropertyName )
 			{
-				Result.FloatWriter = [Attr, StageInfo]( float UEValue, double UsdTimeCode ) mutable
-				{
-					Attr.Set( UnrealToUsd::ConvertDistance( StageInfo, UEValue ), UsdTimeCode );
-				};
+				Attr = Camera.CreateFocalLengthAttr();
 			}
-			else
+			else if ( PropertyPath == UnrealIdentifiers::ManualFocusDistancePropertyName )
 			{
-				Result.FloatWriter = [Attr]( float UEValue, double UsdTimeCode ) mutable
-				{
-					Attr.Set( UEValue, UsdTimeCode );
-				};
+				Attr = Camera.CreateFocusDistanceAttr();
 			}
-		}
-	}
-	else if ( pxr::UsdLuxLightAPI LightAPI{ Prim } )
-	{
-		if ( PropertyPath == UnrealIdentifiers::LightColorPropertyName )
-		{
-			Attr = LightAPI.GetColorAttr();
+			else if ( PropertyPath == UnrealIdentifiers::CurrentAperturePropertyName )
+			{
+				bConvertDistance = false;
+				Attr = Camera.CreateFStopAttr();
+			}
+			else if ( PropertyPath == UnrealIdentifiers::SensorWidthPropertyName )
+			{
+				Attr = Camera.CreateHorizontalApertureAttr();
+			}
+			else if ( PropertyPath == UnrealIdentifiers::SensorHeightPropertyName )
+			{
+				Attr = Camera.CreateVerticalApertureAttr();
+			}
+
 			if ( Attr )
 			{
-				Result.ColorWriter = [Attr]( const FLinearColor& UEValue, double UsdTimeCode )
+				if ( bConvertDistance )
 				{
-					pxr::GfVec4f Vec4 = UnrealToUsd::ConvertColor( UEValue );
-					Attr.Set( pxr::GfVec3f{ Vec4[ 0 ], Vec4[ 1 ], Vec4[ 2 ] }, UsdTimeCode );
-				};
+					Result.FloatWriter = [Attr, StageInfo]( float UEValue, double UsdTimeCode ) mutable
+					{
+						Attr.Set( UnrealToUsd::ConvertDistance( StageInfo, UEValue ), UsdTimeCode );
+					};
+				}
+				else
+				{
+					Result.FloatWriter = [Attr]( float UEValue, double UsdTimeCode ) mutable
+					{
+						Attr.Set( UEValue, UsdTimeCode );
+					};
+				}
 			}
 		}
-		else if ( PropertyPath == UnrealIdentifiers::UseTemperaturePropertyName )
+		else if ( pxr::UsdLuxLightAPI LightAPI{ Prim } )
 		{
-			Attr = LightAPI.GetEnableColorTemperatureAttr();
-			if ( Attr )
+			if ( PropertyPath == UnrealIdentifiers::LightColorPropertyName )
 			{
-				Result.BoolWriter = [Attr]( bool UEValue, double UsdTimeCode )
+				Attr = LightAPI.GetColorAttr();
+				if ( Attr )
 				{
-					Attr.Set( UEValue, UsdTimeCode );
-				};
+					Result.ColorWriter = [Attr]( const FLinearColor& UEValue, double UsdTimeCode )
+					{
+						pxr::GfVec4f Vec4 = UnrealToUsd::ConvertColor( UEValue );
+						Attr.Set( pxr::GfVec3f{ Vec4[ 0 ], Vec4[ 1 ], Vec4[ 2 ] }, UsdTimeCode );
+					};
+				}
 			}
-		}
-		else if ( PropertyPath == UnrealIdentifiers::TemperaturePropertyName )
-		{
-			Attr = LightAPI.GetColorTemperatureAttr();
-			if ( Attr )
+			else if ( PropertyPath == UnrealIdentifiers::UseTemperaturePropertyName )
 			{
-				Result.FloatWriter = [Attr]( float UEValue, double UsdTimeCode )
+				Attr = LightAPI.GetEnableColorTemperatureAttr();
+				if ( Attr )
 				{
-					Attr.Set( UEValue, UsdTimeCode );
-				};
+					Result.BoolWriter = [Attr]( bool UEValue, double UsdTimeCode )
+					{
+						Attr.Set( UEValue, UsdTimeCode );
+					};
+				}
 			}
-		}
-
-		else if ( pxr::UsdLuxSphereLight SphereLight{ UsdPrim } )
-		{
-			if ( PropertyPath == UnrealIdentifiers::SourceRadiusPropertyName )
+			else if ( PropertyPath == UnrealIdentifiers::TemperaturePropertyName )
 			{
-				OutPropertyPathsToRefresh.Add( UnrealIdentifiers::IntensityPropertyName );
-
-				Attr = SphereLight.GetRadiusAttr();
+				Attr = LightAPI.GetColorTemperatureAttr();
 				if ( Attr )
 				{
-					Result.FloatWriter = [Attr, StageInfo]( float UEValue, double UsdTimeCode )
+					Result.FloatWriter = [Attr]( float UEValue, double UsdTimeCode )
 					{
-						Attr.Set( UnrealToUsd::ConvertDistance( StageInfo, UEValue ), UsdTimeCode );
+						Attr.Set( UEValue, UsdTimeCode );
 					};
 				}
 			}
 
-			// Spot light
-			else if ( UsdPrim.HasAPI< pxr::UsdLuxShapingAPI >() )
+			else if ( pxr::UsdLuxSphereLight SphereLight{ UsdPrim } )
 			{
-				pxr::UsdLuxShapingAPI ShapingAPI{ UsdPrim };
-
-				if ( PropertyPath == UnrealIdentifiers::IntensityPropertyName )
+				if ( PropertyPath == UnrealIdentifiers::SourceRadiusPropertyName )
 				{
-					Attr = SphereLight.GetIntensityAttr();
-					pxr::UsdAttribute RadiusAttr = SphereLight.GetRadiusAttr();
-					pxr::UsdAttribute ConeAngleAttr = ShapingAPI.GetShapingConeAngleAttr();
-					pxr::UsdAttribute ConeSoftnessAttr = ShapingAPI.GetShapingConeSoftnessAttr();
+					OutPropertyPathsToRefresh.Add( UnrealIdentifiers::IntensityPropertyName );
 
-					// Always clear exposure because we'll put all of our "light intensity" on the intensity attr and assume exposure
-					// is zero, as we can't manipulate something like that exposure directly from UE anyway
-					if ( pxr::UsdAttribute ExposureAttr = SphereLight.GetExposureAttr() )
+					Attr = SphereLight.GetRadiusAttr();
+					if ( Attr )
 					{
-						ExposureAttr.Clear();
+						Result.FloatWriter = [Attr, StageInfo]( float UEValue, double UsdTimeCode )
+						{
+							Attr.Set( UnrealToUsd::ConvertDistance( StageInfo, UEValue ), UsdTimeCode );
+						};
 					}
+				}
 
-					// For now we'll assume the light intensity units are constant and the user doesn't have any light intensity unit tracks...
-					ELightUnits Units = ELightUnits::Lumens;
-					if ( const ULocalLightComponent* LightComponent = Cast<const ULocalLightComponent>( &Component ) )
+				// Spot light
+				else if ( UsdPrim.HasAPI< pxr::UsdLuxShapingAPI >() )
+				{
+					pxr::UsdLuxShapingAPI ShapingAPI{ UsdPrim };
+
+					if ( PropertyPath == UnrealIdentifiers::IntensityPropertyName )
 					{
-						Units = LightComponent->IntensityUnits;
+						Attr = SphereLight.GetIntensityAttr();
+						pxr::UsdAttribute RadiusAttr = SphereLight.GetRadiusAttr();
+						pxr::UsdAttribute ConeAngleAttr = ShapingAPI.GetShapingConeAngleAttr();
+						pxr::UsdAttribute ConeSoftnessAttr = ShapingAPI.GetShapingConeSoftnessAttr();
+
+						// Always clear exposure because we'll put all of our "light intensity" on the intensity attr and assume exposure
+						// is zero, as we can't manipulate something like that exposure directly from UE anyway
+						if ( pxr::UsdAttribute ExposureAttr = SphereLight.GetExposureAttr() )
+						{
+							ExposureAttr.Clear();
+						}
+
+						// For now we'll assume the light intensity units are constant and the user doesn't have any light intensity unit tracks...
+						ELightUnits Units = ELightUnits::Lumens;
+						if ( const ULocalLightComponent* LightComponent = Cast<const ULocalLightComponent>( &Component ) )
+						{
+							Units = LightComponent->IntensityUnits;
+						}
+
+						if ( Attr && RadiusAttr && ConeAngleAttr && ConeSoftnessAttr )
+						{
+							Result.FloatWriter = [Attr, RadiusAttr, ConeAngleAttr, ConeSoftnessAttr, StageInfo, Units]( float UEValue, double UsdTimeCode )
+							{
+								const float UsdConeAngle = UsdUtils::GetUsdValue< float >( ConeAngleAttr, UsdTimeCode );
+								const float UsdConeSoftness = UsdUtils::GetUsdValue< float >( ConeSoftnessAttr, UsdTimeCode );
+								const float UsdRadius = UsdUtils::GetUsdValue< float >( RadiusAttr, UsdTimeCode );
+
+								float InnerConeAngle = 0.0f;
+								const float OuterConeAngle = UsdToUnreal::ConvertConeAngleSoftnessAttr( UsdConeAngle, UsdConeSoftness, InnerConeAngle );
+								const float SourceRadius = UsdToUnreal::ConvertDistance( StageInfo, UsdRadius );
+
+								Attr.Set(
+									UnrealToUsd::ConvertSpotLightIntensityProperty( UEValue, OuterConeAngle, InnerConeAngle, SourceRadius, StageInfo, Units ),
+									UsdTimeCode
+								);
+							};
+						}
 					}
+					else if ( PropertyPath == UnrealIdentifiers::OuterConeAnglePropertyName )
+					{
+						Attr = ShapingAPI.GetShapingConeAngleAttr();
+						if ( Attr )
+						{
+							// InnerConeAngle is calculated based on ConeAngleAttr, so we need to refresh it
+							OutPropertyPathsToRefresh.Add( UnrealIdentifiers::InnerConeAnglePropertyName );
 
-					if ( Attr && RadiusAttr && ConeAngleAttr && ConeSoftnessAttr )
+							Result.FloatWriter = [Attr]( float UEValue, double UsdTimeCode )
+							{
+								Attr.Set( UEValue, UsdTimeCode );
+							};
+						}
+					}
+					else if ( PropertyPath == UnrealIdentifiers::InnerConeAnglePropertyName )
 					{
-						Result.FloatWriter = [Attr, RadiusAttr, ConeAngleAttr, ConeSoftnessAttr, StageInfo, Units]( float UEValue, double UsdTimeCode )
+						Attr = ShapingAPI.GetShapingConeSoftnessAttr();
+						pxr::UsdAttribute ConeAngleAttr = ShapingAPI.GetShapingConeAngleAttr();
+
+						if ( ConeAngleAttr && Attr )
 						{
-							const float UsdConeAngle = UsdUtils::GetUsdValue< float >( ConeAngleAttr, UsdTimeCode );
-							const float UsdConeSoftness = UsdUtils::GetUsdValue< float >( ConeSoftnessAttr, UsdTimeCode );
-							const float UsdRadius = UsdUtils::GetUsdValue< float >( RadiusAttr, UsdTimeCode );
+							Result.FloatWriter = [Attr, ConeAngleAttr]( float UEValue, double UsdTimeCode )
+							{
+								const float UsdConeAngle = UsdUtils::GetUsdValue< float >( ConeAngleAttr, UsdTimeCode );
+								const float OuterConeAngle = UsdConeAngle;
+
+								const float OutNewSoftness = UnrealToUsd::ConvertInnerConeAngleProperty( UEValue, OuterConeAngle );
+								Attr.Set( OutNewSoftness, UsdTimeCode );
+							};
+						}
+					}
+				}
+				// Just a point light
+				else
+				{
+					if ( PropertyPath == UnrealIdentifiers::IntensityPropertyName )
+					{
+						Attr = SphereLight.GetIntensityAttr();
+						pxr::UsdAttribute RadiusAttr = SphereLight.GetRadiusAttr();
 
-							float InnerConeAngle = 0.0f;
-							const float OuterConeAngle = UsdToUnreal::ConvertConeAngleSoftnessAttr( UsdConeAngle, UsdConeSoftness, InnerConeAngle );
-							const float SourceRadius = UsdToUnreal::ConvertDistance( StageInfo, UsdRadius );
+						// Always clear exposure because we'll put all of our "light intensity" on the intensity attr and assume exposure
+						// is zero, as we can't manipulate something like that exposure directly from UE anyway
+						if ( pxr::UsdAttribute ExposureAttr = SphereLight.GetExposureAttr() )
+						{
+							ExposureAttr.Clear();
+						}
 
-							Attr.Set(
-								UnrealToUsd::ConvertSpotLightIntensityProperty( UEValue, OuterConeAngle, InnerConeAngle, SourceRadius, StageInfo, Units ),
-								UsdTimeCode
-							);
-						};
+						// For now we'll assume the light intensity units are constant and the user doesn't have any light intensity unit tracks...
+						ELightUnits Units = ELightUnits::Lumens;
+						if ( const ULocalLightComponent* LightComponent = Cast<const ULocalLightComponent>( &Component ) )
+						{
+							Units = LightComponent->IntensityUnits;
+						}
+
+						if ( Attr && RadiusAttr )
+						{
+							Result.FloatWriter = [Attr, RadiusAttr, StageInfo, Units]( float UEValue, double UsdTimeCode )
+							{
+								const float SourceRadius = UsdToUnreal::ConvertDistance( StageInfo, UsdUtils::GetUsdValue< float >( RadiusAttr, UsdTimeCode ) );
+								Attr.Set( UnrealToUsd::ConvertPointLightIntensityProperty( UEValue, SourceRadius, StageInfo, Units ), UsdTimeCode );
+							};
+						}
 					}
 				}
-				else if ( PropertyPath == UnrealIdentifiers::OuterConeAnglePropertyName )
+			}
+			else if ( pxr::UsdLuxRectLight RectLight{ UsdPrim } )
+			{
+				if ( PropertyPath == UnrealIdentifiers::SourceWidthPropertyName )
 				{
-					Attr = ShapingAPI.GetShapingConeAngleAttr();
+					Attr = RectLight.GetWidthAttr();
 					if ( Attr )
 					{
-						// InnerConeAngle is calculated based on ConeAngleAttr, so we need to refresh it
-						OutPropertyPathsToRefresh.Add( UnrealIdentifiers::InnerConeAnglePropertyName );
+						OutPropertyPathsToRefresh.Add( UnrealIdentifiers::IntensityPropertyName );
 
-						Result.FloatWriter = [Attr]( float UEValue, double UsdTimeCode )
+						Result.FloatWriter = [Attr, StageInfo]( float UEValue, double UsdTimeCode )
 						{
-							Attr.Set( UEValue, UsdTimeCode );
+							Attr.Set( UnrealToUsd::ConvertDistance( StageInfo, UEValue ), UsdTimeCode );
 						};
 					}
 				}
-				else if ( PropertyPath == UnrealIdentifiers::InnerConeAnglePropertyName )
+				else if ( PropertyPath == UnrealIdentifiers::SourceHeightPropertyName )
 				{
-					Attr = ShapingAPI.GetShapingConeSoftnessAttr();
-					pxr::UsdAttribute ConeAngleAttr = ShapingAPI.GetShapingConeAngleAttr();
-
-					if ( ConeAngleAttr && Attr )
+					Attr = RectLight.GetHeightAttr();
+					if ( Attr )
 					{
-						Result.FloatWriter = [Attr, ConeAngleAttr]( float UEValue, double UsdTimeCode )
-						{
-							const float UsdConeAngle = UsdUtils::GetUsdValue< float >( ConeAngleAttr, UsdTimeCode );
-							const float OuterConeAngle = UsdConeAngle;
+						OutPropertyPathsToRefresh.Add( UnrealIdentifiers::IntensityPropertyName );
 
-							const float OutNewSoftness = UnrealToUsd::ConvertInnerConeAngleProperty( UEValue, OuterConeAngle );
-							Attr.Set( OutNewSoftness, UsdTimeCode );
+						Result.FloatWriter = [Attr, StageInfo]( float UEValue, double UsdTimeCode )
+						{
+							Attr.Set( UnrealToUsd::ConvertDistance( StageInfo, UEValue ), UsdTimeCode );
 						};
 					}
 				}
-			}
-			// Just a point light
-			else
-			{
-				if ( PropertyPath == UnrealIdentifiers::IntensityPropertyName )
+				else if ( PropertyPath == UnrealIdentifiers::IntensityPropertyName )
 				{
-					Attr = SphereLight.GetIntensityAttr();
-					pxr::UsdAttribute RadiusAttr = SphereLight.GetRadiusAttr();
+					Attr = RectLight.GetIntensityAttr();
+					pxr::UsdAttribute WidthAttr = RectLight.GetWidthAttr();
+					pxr::UsdAttribute HeightAttr = RectLight.GetHeightAttr();
 
 					// Always clear exposure because we'll put all of our "light intensity" on the intensity attr and assume exposure
 					// is zero, as we can't manipulate something like that exposure directly from UE anyway
-					if ( pxr::UsdAttribute ExposureAttr = SphereLight.GetExposureAttr() )
+					if ( pxr::UsdAttribute ExposureAttr = RectLight.GetExposureAttr() )
 					{
 						ExposureAttr.Clear();
 					}
@@ -3142,158 +3207,100 @@ UnrealToUsd::FPropertyTrackWriter UnrealToUsd::CreatePropertyTrackWriter( const
 						Units = LightComponent->IntensityUnits;
 					}
 
-					if ( Attr && RadiusAttr )
+					if ( Attr && WidthAttr && HeightAttr )
 					{
-						Result.FloatWriter = [Attr, RadiusAttr, StageInfo, Units]( float UEValue, double UsdTimeCode )
+						Result.FloatWriter = [Attr, WidthAttr, HeightAttr, StageInfo, Units]( float UEValue, double UsdTimeCode )
 						{
-							const float SourceRadius = UsdToUnreal::ConvertDistance( StageInfo, UsdUtils::GetUsdValue< float >( RadiusAttr, UsdTimeCode ) );
-							Attr.Set( UnrealToUsd::ConvertPointLightIntensityProperty( UEValue, SourceRadius, StageInfo, Units ), UsdTimeCode );
+							const float Width = UsdToUnreal::ConvertDistance( StageInfo, UsdUtils::GetUsdValue< float >( WidthAttr, UsdTimeCode ) );
+							const float Height = UsdToUnreal::ConvertDistance( StageInfo, UsdUtils::GetUsdValue< float >( HeightAttr, UsdTimeCode ) );
+
+							Attr.Set( UnrealToUsd::ConvertRectLightIntensityProperty( UEValue, Width, Height, StageInfo, Units ), UsdTimeCode );
 						};
 					}
 				}
 			}
-		}
-		else if ( pxr::UsdLuxRectLight RectLight{ UsdPrim } )
-		{
-			if ( PropertyPath == UnrealIdentifiers::SourceWidthPropertyName )
+			else if ( pxr::UsdLuxDiskLight DiskLight{ UsdPrim } )
 			{
-				Attr = RectLight.GetWidthAttr();
-				if ( Attr )
+				if ( PropertyPath == UnrealIdentifiers::SourceWidthPropertyName || PropertyPath == UnrealIdentifiers::SourceHeightPropertyName )
 				{
-					OutPropertyPathsToRefresh.Add( UnrealIdentifiers::IntensityPropertyName );
-
-					Result.FloatWriter = [Attr, StageInfo]( float UEValue, double UsdTimeCode )
+					Attr = DiskLight.GetRadiusAttr();
+					if ( Attr )
 					{
-						Attr.Set( UnrealToUsd::ConvertDistance( StageInfo, UEValue ), UsdTimeCode );
-					};
-				}
-			}
-			else if ( PropertyPath == UnrealIdentifiers::SourceHeightPropertyName )
-			{
-				Attr = RectLight.GetHeightAttr();
-				if ( Attr )
-				{
-					OutPropertyPathsToRefresh.Add( UnrealIdentifiers::IntensityPropertyName );
+						OutPropertyPathsToRefresh.Add( UnrealIdentifiers::IntensityPropertyName );
 
-					Result.FloatWriter = [Attr, StageInfo]( float UEValue, double UsdTimeCode )
-					{
-						Attr.Set( UnrealToUsd::ConvertDistance( StageInfo, UEValue ), UsdTimeCode );
-					};
-				}
-			}
-			else if ( PropertyPath == UnrealIdentifiers::IntensityPropertyName )
-			{
-				Attr = RectLight.GetIntensityAttr();
-				pxr::UsdAttribute WidthAttr = RectLight.GetWidthAttr();
-				pxr::UsdAttribute HeightAttr = RectLight.GetHeightAttr();
+						// Resync the other to match this one after we bake it, effectively always enforcing the UE rect light into a square shape
+						OutPropertyPathsToRefresh.Add( PropertyPath == UnrealIdentifiers::SourceWidthPropertyName
+							? UnrealIdentifiers::SourceHeightPropertyName
+							: UnrealIdentifiers::SourceWidthPropertyName
+						);
 
-				// Always clear exposure because we'll put all of our "light intensity" on the intensity attr and assume exposure
-				// is zero, as we can't manipulate something like that exposure directly from UE anyway
-				if ( pxr::UsdAttribute ExposureAttr = RectLight.GetExposureAttr() )
-				{
-					ExposureAttr.Clear();
+						Result.FloatWriter = [Attr, StageInfo]( float UEValue, double UsdTimeCode )
+						{
+							Attr.Set( UnrealToUsd::ConvertDistance( StageInfo, UEValue * 0.5f ), UsdTimeCode );
+						};
+					}
 				}
-
-				// For now we'll assume the light intensity units are constant and the user doesn't have any light intensity unit tracks...
-				ELightUnits Units = ELightUnits::Lumens;
-				if ( const ULocalLightComponent* LightComponent = Cast<const ULocalLightComponent>( &Component ) )
+				else if ( PropertyPath == UnrealIdentifiers::IntensityPropertyName )
 				{
-					Units = LightComponent->IntensityUnits;
-				}
+					Attr = RectLight.GetIntensityAttr();
+					pxr::UsdAttribute RadiusAttr = DiskLight.GetRadiusAttr();
 
-				if ( Attr && WidthAttr && HeightAttr )
-				{
-					Result.FloatWriter = [Attr, WidthAttr, HeightAttr, StageInfo, Units]( float UEValue, double UsdTimeCode )
+					// Always clear exposure because we'll put all of our "light intensity" on the intensity attr and assume exposure
+					// is zero, as we can't manipulate something like that exposure directly from UE anyway
+					if ( pxr::UsdAttribute ExposureAttr = RectLight.GetExposureAttr() )
 					{
-						const float Width = UsdToUnreal::ConvertDistance( StageInfo, UsdUtils::GetUsdValue< float >( WidthAttr, UsdTimeCode ) );
-						const float Height = UsdToUnreal::ConvertDistance( StageInfo, UsdUtils::GetUsdValue< float >( HeightAttr, UsdTimeCode ) );
-
-						Attr.Set( UnrealToUsd::ConvertRectLightIntensityProperty( UEValue, Width, Height, StageInfo, Units ), UsdTimeCode );
-					};
-				}
-			}
-		}
-		else if ( pxr::UsdLuxDiskLight DiskLight{ UsdPrim } )
-		{
-			if ( PropertyPath == UnrealIdentifiers::SourceWidthPropertyName || PropertyPath == UnrealIdentifiers::SourceHeightPropertyName )
-			{
-				Attr = DiskLight.GetRadiusAttr();
-				if ( Attr )
-				{
-					OutPropertyPathsToRefresh.Add( UnrealIdentifiers::IntensityPropertyName );
-
-					// Resync the other to match this one after we bake it, effectively always enforcing the UE rect light into a square shape
-					OutPropertyPathsToRefresh.Add( PropertyPath == UnrealIdentifiers::SourceWidthPropertyName
-						? UnrealIdentifiers::SourceHeightPropertyName
-						: UnrealIdentifiers::SourceWidthPropertyName
-					);
+						ExposureAttr.Clear();
+					}
 
-					Result.FloatWriter = [Attr, StageInfo]( float UEValue, double UsdTimeCode )
+					// For now we'll assume the light intensity units are constant and the user doesn't have any light intensity unit tracks...
+					ELightUnits Units = ELightUnits::Lumens;
+					if ( const ULocalLightComponent* LightComponent = Cast<const ULocalLightComponent>( &Component ) )
 					{
-						Attr.Set( UnrealToUsd::ConvertDistance( StageInfo, UEValue * 0.5f ), UsdTimeCode );
-					};
-				}
-			}
-			else if ( PropertyPath == UnrealIdentifiers::IntensityPropertyName )
-			{
-				Attr = RectLight.GetIntensityAttr();
-				pxr::UsdAttribute RadiusAttr = DiskLight.GetRadiusAttr();
-
-				// Always clear exposure because we'll put all of our "light intensity" on the intensity attr and assume exposure
-				// is zero, as we can't manipulate something like that exposure directly from UE anyway
-				if ( pxr::UsdAttribute ExposureAttr = RectLight.GetExposureAttr() )
-				{
-					ExposureAttr.Clear();
-				}
-
-				// For now we'll assume the light intensity units are constant and the user doesn't have any light intensity unit tracks...
-				ELightUnits Units = ELightUnits::Lumens;
-				if ( const ULocalLightComponent* LightComponent = Cast<const ULocalLightComponent>( &Component ) )
-				{
-					Units = LightComponent->IntensityUnits;
-				}
+						Units = LightComponent->IntensityUnits;
+					}
 
-				if ( Attr && RadiusAttr )
-				{
-					Result.FloatWriter = [Attr, RadiusAttr, StageInfo, Units]( float UEValue, double UsdTimeCode )
+					if ( Attr && RadiusAttr )
 					{
-						const float Radius = UsdToUnreal::ConvertDistance( StageInfo, UsdUtils::GetUsdValue< float >( RadiusAttr, UsdTimeCode ) );
+						Result.FloatWriter = [Attr, RadiusAttr, StageInfo, Units]( float UEValue, double UsdTimeCode )
+						{
+							const float Radius = UsdToUnreal::ConvertDistance( StageInfo, UsdUtils::GetUsdValue< float >( RadiusAttr, UsdTimeCode ) );
 
-						Attr.Set( UnrealToUsd::ConvertRectLightIntensityProperty( UEValue, Radius, StageInfo, Units ), UsdTimeCode );
-					};
+							Attr.Set( UnrealToUsd::ConvertRectLightIntensityProperty( UEValue, Radius, StageInfo, Units ), UsdTimeCode );
+						};
+					}
 				}
 			}
-		}
-		else if ( pxr::UsdLuxDistantLight DistantLight{ UsdPrim } )
-		{
-			if ( PropertyPath == UnrealIdentifiers::LightSourceAnglePropertyName )
+			else if ( pxr::UsdLuxDistantLight DistantLight{ UsdPrim } )
 			{
-				Attr = DistantLight.GetAngleAttr();
-				if ( Attr )
+				if ( PropertyPath == UnrealIdentifiers::LightSourceAnglePropertyName )
 				{
-					Result.FloatWriter = [Attr]( float UEValue, double UsdTimeCode )
+					Attr = DistantLight.GetAngleAttr();
+					if ( Attr )
 					{
-						Attr.Set( UEValue, UsdTimeCode );
-					};
+						Result.FloatWriter = [Attr]( float UEValue, double UsdTimeCode )
+						{
+							Attr.Set( UEValue, UsdTimeCode );
+						};
+					}
 				}
-			}
-			else if ( PropertyPath == UnrealIdentifiers::IntensityPropertyName )
-			{
-				Attr = DistantLight.GetIntensityAttr();
-
-				// Always clear exposure because we'll put all of our "light intensity" on the intensity attr and assume exposure
-				// is zero, as we can't manipulate something like that exposure directly from UE anyway
-				if ( pxr::UsdAttribute ExposureAttr = RectLight.GetExposureAttr() )
+				else if ( PropertyPath == UnrealIdentifiers::IntensityPropertyName )
 				{
-					ExposureAttr.Clear();
-				}
+					Attr = DistantLight.GetIntensityAttr();
 
-				if ( Attr )
-				{
-					Result.FloatWriter = [Attr]( float UEValue, double UsdTimeCode )
+					// Always clear exposure because we'll put all of our "light intensity" on the intensity attr and assume exposure
+					// is zero, as we can't manipulate something like that exposure directly from UE anyway
+					if ( pxr::UsdAttribute ExposureAttr = RectLight.GetExposureAttr() )
 					{
-						Attr.Set( UnrealToUsd::ConvertLightIntensityProperty( UEValue ), UsdTimeCode );
-					};
+						ExposureAttr.Clear();
+					}
+
+					if ( Attr )
+					{
+						Result.FloatWriter = [Attr]( float UEValue, double UsdTimeCode )
+						{
+							Attr.Set( UnrealToUsd::ConvertLightIntensityProperty( UEValue ), UsdTimeCode );
+						};
+					}
 				}
 			}
 		}
@@ -3308,6 +3315,12 @@ UnrealToUsd::FPropertyTrackWriter UnrealToUsd::CreatePropertyTrackWriter( const
 			Attr.ClearAtTime( TimeSample );
 		}
 
+		// Note that we must do this only after the change block is destroyed!
+		// This is important because if we don't have spec for this attribute on the current edit target, we're relying
+		// on the previous code to create it, and we need to let USD emit its internal notices and fully commit the
+		// "attribute creation" spec first. This because NotifyIfOverriddenOpinion will go through the attribute's
+		// spec stack and consider our attribute overriden if it finds a stronger opinion than the one on the edit
+		// target. Well if our own spec hasn't been created yet it will misfire when it runs into any other spec
 		UsdUtils::NotifyIfOverriddenOpinion( UE::FUsdAttribute{ Attr } );
 	}
 
