USD: Fix regression on how adding/removing layers didn't automatically reload animations.

Fix LevelSequences being created for layers not in the local layer stack;
Fix LevelSequences being created for particular prims (we should only have sequences dedicated for layers);

diff --git a/Engine/Plugins/Importers/USDImporter/Source/USDStage/Private/USDLevelSequenceHelper.cpp b/Engine/Plugins/Importers/USDImporter/Source/USDStage/Private/USDLevelSequenceHelper.cpp
--- a/Engine/Plugins/Importers/USDImporter/Source/USDStage/Private/USDLevelSequenceHelper.cpp
+++ b/Engine/Plugins/Importers/USDImporter/Source/USDStage/Private/USDLevelSequenceHelper.cpp
@@ -1057,9 +1057,16 @@ ULevelSequence* FUsdLevelSequenceHelperImpl::FindOrAddSequenceForAttribute(const
 	{
 		if (UE::FSdfLayer AttributeLayer = UsdUtils::FindLayerForAttribute(Attribute, 0.0))
 		{
-			const FString SequenceIdentifier = Attribute.GetPrim().GetPrimPath().GetString();
-
-			Sequence = FindOrAddSequenceForLayer(AttributeLayer, SequenceIdentifier, SequenceIdentifier);
+			if (UsdStage)
+			{
+				// Only create level sequences for layers in the local layer stack: Our current approach is that only
+				// those layers can be written to, so there is no point in doing otherwise
+				if (UsdStage.HasLocalLayer(AttributeLayer))
+				{
+					const FString SequenceIdentifier = AttributeLayer.GetIdentifier();
+					Sequence = FindOrAddSequenceForLayer(AttributeLayer, SequenceIdentifier, SequenceIdentifier);
+				}
+			}
 		}
 	}
 
diff --git a/Engine/Plugins/Importers/USDImporter/Source/USDStage/Private/USDStageActor.cpp b/Engine/Plugins/Importers/USDImporter/Source/USDStage/Private/USDStageActor.cpp
--- a/Engine/Plugins/Importers/USDImporter/Source/USDStage/Private/USDStageActor.cpp
+++ b/Engine/Plugins/Importers/USDImporter/Source/USDStage/Private/USDStageActor.cpp
@@ -742,7 +742,7 @@ AUsdStageActor::AUsdStageActor()
 		UsdListener.GetOnLayersChanged().AddLambda(
 			[&, this](const TArray< FString >& ChangedLayerIdentifiers)
 			{
-				if (!IsListeningToUsdNotices())
+				if (!IsListeningToUsdNotices() || ChangedLayerIdentifiers.Num() == 0)
 				{
 					return;
 				}
@@ -961,6 +961,8 @@ void AUsdStageActor::OnUsdObjectsChanged(const UsdUtils::FObjectChangesByPath& I
 
 	bool bHasResync = ResyncChanges.Num() > 0;
 
+	bool bNeedsAnimationReload = false;
+
 	// The most important thing here is to iterate in parent to child order, so build SortedPrimsChangedList
 	TMap< FString, bool > SortedPrimsChangedList;
 	for (const TPair<FString, TArray<UsdUtils::FObjectChangeNotice>>& InfoChange : InfoChanges)
@@ -996,6 +998,14 @@ void AUsdStageActor::OnUsdObjectsChanged(const UsdUtils::FObjectChangesByPath& I
 					break;
 				}
 			}
+
+			// Any sublayer change (even offsets) means we need to regenerate our LevelSequence to add (or shift)
+			// the corresponding subsequences
+			if (PrimPath.IsAbsoluteRootPath() && ObjectChange.SubLayerChanges.Num() > 0)
+			{
+				bNeedsAnimationReload = true;
+				break;
+			}
 		}
 
 		// We may need the full spec path with variant selections later, but for traversal and retrieving prims from the stage we always need
@@ -1206,6 +1216,12 @@ void AUsdStageActor::OnUsdObjectsChanged(const UsdUtils::FObjectChangesByPath& I
 		}
 	}
 
+	if (bNeedsAnimationReload)
+	{
+		ReloadAnimations();
+		PrimsToAnimate.Reset();
+	}
+
 	if (bHasResync && UsdAssetCache)
 	{
 		UsdAssetCache->RefreshStorage();
diff --git a/Engine/Plugins/Importers/USDImporter/Source/UnrealUSDWrapper/Private/USDListener.cpp b/Engine/Plugins/Importers/USDImporter/Source/UnrealUSDWrapper/Private/USDListener.cpp
--- a/Engine/Plugins/Importers/USDImporter/Source/UnrealUSDWrapper/Private/USDListener.cpp
+++ b/Engine/Plugins/Importers/USDImporter/Source/UnrealUSDWrapper/Private/USDListener.cpp
@@ -375,15 +375,9 @@ void FUsdListenerImpl::HandleLayersChangedNotice( const pxr::SdfNotice::LayersDi
 			const pxr::SdfChangeList::EntryList& ChangeList = ChangeVecItem.second.GetEntryList();
 			for (const std::pair<pxr::SdfPath, pxr::SdfChangeList::Entry>& Change : ChangeList)
 			{
-				for (const pxr::SdfChangeList::Entry::SubLayerChange& SubLayerChange : Change.second.subLayerChanges)
-				{
-					const pxr::SdfChangeList::SubLayerChangeType ChangeType = SubLayerChange.second;
-					if (ChangeType == pxr::SdfChangeList::SubLayerChangeType::SubLayerAdded ||
-						ChangeType == pxr::SdfChangeList::SubLayerChangeType::SubLayerRemoved)
-					{
-						LayersNames.Add( ANSI_TO_TCHAR( SubLayerChange.first.c_str() ) );
-					}
-				}
+				// Sadly we can't do much with these notices as USD doesn't provide information of which layer the
+				// layer was added *to*. Since this is a global notice not even tied to a particular stage,
+				// it doesn't help us know when to update what
 
 				const pxr::SdfChangeList::Entry::_Flags& Flags = Change.second.flags;
 				if (Flags.didReloadContent)
