USD: Fix how the mesh vertex count collapsing limit and Nanite material slot collapsing limit weren't being checked at all anymore;

Simplify Nanite enabling behavior when collapsing.

diff --git a/Engine/Plugins/Importers/USDImporter/Source/USDSchemas/Private/USDGeomMeshTranslator.cpp b/Engine/Plugins/Importers/USDImporter/Source/USDSchemas/Private/USDGeomMeshTranslator.cpp
--- a/Engine/Plugins/Importers/USDImporter/Source/USDSchemas/Private/USDGeomMeshTranslator.cpp
+++ b/Engine/Plugins/Importers/USDImporter/Source/USDSchemas/Private/USDGeomMeshTranslator.cpp
@@ -98,6 +98,8 @@ namespace UsdGeomMeshTranslatorImpl
 
 		FScopedUsdAllocs UsdAllocs;
 
+		UE::FSdfPath PrimPath = Prim.GetPrimPath();
+
 		pxr::UsdPrim UsdPrim{Prim};
 		if (!UsdPrim)
 		{
@@ -116,7 +118,7 @@ namespace UsdGeomMeshTranslatorImpl
 				{
 					bHasNaniteOverrideEnabled = true;
 					UE_LOG(LogUsd, Log, TEXT("Trying to enable Nanite for mesh generated for prim '%s' as the '%s' attribute is set to '%s'"),
-						*Prim.GetPrimPath().GetString(),
+						*PrimPath.GetString(),
 						*UsdToUnreal::ConvertToken(UnrealIdentifiers::UnrealNaniteOverride),
 						*UsdToUnreal::ConvertToken(UnrealIdentifiers::UnrealNaniteOverrideEnable)
 					);
@@ -125,7 +127,7 @@ namespace UsdGeomMeshTranslatorImpl
 				else if (OverrideValue == UnrealIdentifiers::UnrealNaniteOverrideDisable)
 				{
 					UE_LOG(LogUsd, Log, TEXT("Not enabling Nanite for mesh generated for prim '%s' as the '%s' attribute is set to '%s'"),
-						*Prim.GetPrimPath().GetString(),
+						*PrimPath.GetString(),
 						*UsdToUnreal::ConvertToken(UnrealIdentifiers::UnrealNaniteOverride),
 						*UsdToUnreal::ConvertToken(UnrealIdentifiers::UnrealNaniteOverrideDisable)
 					);
@@ -135,13 +137,13 @@ namespace UsdGeomMeshTranslatorImpl
 		}
 
 		// We want Nanite because the mesh is large enough for the threshold, which is set to something valid
-		if (!bHasNaniteOverrideEnabled)
+		if (!bHasNaniteOverrideEnabled && Context.InfoCache.IsValid())
 		{
-			const int32 NumTriangles = LODIndexToMeshDescription[0].Triangles().Num();
+			uint64 NumTriangles = Context.InfoCache->GetSubtreeVertexCount(PrimPath).Get(0);
 			if (NumTriangles >= Context.NaniteTriangleThreshold)
 			{
 				UE_LOG(LogUsd, Verbose, TEXT("Trying to enable Nanite for mesh generated for prim '%s' as it has '%d' triangles, and the threshold is '%d'"),
-					*Prim.GetPrimPath().GetString(),
+					*PrimPath.GetString(),
 					NumTriangles,
 					Context.NaniteTriangleThreshold
 				);
@@ -149,7 +151,7 @@ namespace UsdGeomMeshTranslatorImpl
 			else
 			{
 				UE_LOG(LogUsd, Verbose, TEXT("Not enabling Nanite for mesh generated for prim '%s' as it has '%d' triangles, and the threshold is '%d'"),
-					*Prim.GetPrimPath().GetString(),
+					*PrimPath.GetString(),
 					NumTriangles,
 					Context.NaniteTriangleThreshold
 				);
@@ -165,21 +167,21 @@ namespace UsdGeomMeshTranslatorImpl
 		if (LODIndexToMeshDescription.Num() > 1)
 		{
 			UE_LOG(LogUsd, Warning, TEXT("Not enabling Nanite for mesh generated for prim '%s' as it has more than one generated LOD (and so came from a LOD variant set setup)"),
-				*Prim.GetPrimPath().GetString()
+				*PrimPath.GetString()
 			);
 			return false;
 		}
 
-		// We cannot enable Nanite here due to a limit on the number of material slots
-		if (LODIndexToMaterialInfo.Num() > 0)
+		if (Context.InfoCache.IsValid())
 		{
-			const int32 NumSections = LODIndexToMaterialInfo[0].Slots.Num();
+			TOptional<uint64> SubtreeSectionCount = Context.InfoCache->GetSubtreeMaterialSlotCount(PrimPath);
+
 			const int32 MaxNumSections = 64; // There is no define for this, but it's checked for on NaniteBuilder.cpp, FBuilderModule::Build
-			if (NumSections > MaxNumSections)
+			if (!SubtreeSectionCount.IsSet() || SubtreeSectionCount.GetValue() > MaxNumSections)
 			{
 				UE_LOG(LogUsd, Warning, TEXT("Not enabling Nanite for mesh generated for prim '%s' as LOD0 has '%d' material slots, which is above the Nanite limit of '%d'"),
-					*Prim.GetPrimPath().GetString(),
-					NumSections,
+					*PrimPath.GetString(),
+					SubtreeSectionCount.GetValue(),
 					MaxNumSections
 				);
 				return false;
@@ -188,12 +190,12 @@ namespace UsdGeomMeshTranslatorImpl
 
 #if !WITH_EDITOR
 		UE_LOG(LogUsd, Warning, TEXT("Not enabling Nanite for mesh generated for prim '%s' as we can't setup Nanite during runtime"),
-			*Prim.GetPrimPath().GetString()
+			*PrimPath.GetString()
 		);
 		return false;
-#endif
-
+#else
 		return true;
+#endif
 	}
 
 	bool IsAnimated(const pxr::UsdPrim& Prim)
diff --git a/Engine/Plugins/Importers/USDImporter/Source/USDSchemas/Private/USDGeomXformableTranslator.cpp b/Engine/Plugins/Importers/USDImporter/Source/USDSchemas/Private/USDGeomXformableTranslator.cpp
--- a/Engine/Plugins/Importers/USDImporter/Source/USDSchemas/Private/USDGeomXformableTranslator.cpp
+++ b/Engine/Plugins/Importers/USDImporter/Source/USDSchemas/Private/USDGeomXformableTranslator.cpp
@@ -55,12 +55,6 @@ static FAutoConsoleVariableRef CVarCollapsePrimsWithoutKind(
 	GCollapsePrimsWithoutKind,
 	TEXT( "Allow collapsing prims that have no authored 'Kind' value" ) );
 
-static int32 GMaxNumVerticesCollapsedMesh = 5000000;
-static FAutoConsoleVariableRef CVarMaxNumVerticesCollapsedMesh(
-	TEXT( "USD.MaxNumVerticesCollapsedMesh" ),
-	GMaxNumVerticesCollapsedMesh,
-	TEXT( "Maximum number of vertices that a combined Mesh can have for us to collapse it into a single StaticMesh" ) );
-
 static bool GEnableCollision = true;
 static FAutoConsoleVariableRef CVarEnableCollision(
 	TEXT( "USD.EnableCollision" ),
@@ -683,80 +677,6 @@ bool FUsdGeomXformableTranslator::CollapsesChildren( ECollapsingType CollapsingT
 		}
 	}
 
-	if ( bCollapsesChildren )
-	{
-		TArray< TUsdStore< pxr::UsdPrim > > ChildGeomMeshes = UsdUtils::GetAllPrimsOfType( Prim, pxr::TfType::Find< pxr::UsdGeomMesh >() );
-
-		// Don't collapse children if they have different Nanite override values
-		bool bChildrenWantNanite = false;
-		bool bOtherChildHasNaniteOpinion = false;
-		for ( const TUsdStore< pxr::UsdPrim >& StoredPrim : ChildGeomMeshes )
-		{
-			pxr::UsdPrim ChildGeomMeshPrim = StoredPrim.Get();
-			if ( !pxr::UsdGeomMesh{ ChildGeomMeshPrim } )
-			{
-				continue;
-			}
-
-			if ( pxr::UsdAttribute NaniteOverride = ChildGeomMeshPrim.GetAttribute( UnrealIdentifiers::UnrealNaniteOverride ) )
-			{
-				pxr::TfToken OverrideValue;
-				if ( NaniteOverride.Get( &OverrideValue ) )
-				{
-					const bool bChildWantsNanite = ( OverrideValue == UnrealIdentifiers::UnrealNaniteOverrideEnable );
-
-					if ( bOtherChildHasNaniteOpinion && ( bChildWantsNanite != bChildrenWantNanite ) )
-					{
-						UE_LOG( LogUsd, Log, TEXT( "Not collapsing down from prim '%s' as child meshes have different values for the '%s' attribute" ),
-							*PrimPath.GetString(),
-							*UsdToUnreal::ConvertToken( UnrealIdentifiers::UnrealNaniteOverride )
-						);
-						return false;
-					}
-					else if ( !bOtherChildHasNaniteOpinion )
-					{
-						bChildrenWantNanite = bChildWantsNanite;
-						bOtherChildHasNaniteOpinion = true;
-					}
-				}
-			}
-		}
-
-		if ( Context->InfoCache.IsValid() )
-		{
-			TOptional<uint64> NumExpectedVertices = Context->InfoCache->GetSubtreeVertexCount( PrimPath );
-			if ( !NumExpectedVertices.IsSet() || NumExpectedVertices.GetValue() > GMaxNumVerticesCollapsedMesh )
-			{
-				bCollapsesChildren = false;
-			}
-
-			if ( bChildrenWantNanite )
-			{
-				TOptional<uint64> NumExpectedMaterialSlots = Context->InfoCache->GetSubtreeMaterialSlotCount( PrimPath ).Get( 0 );
-
-				// Note that we wont try to prevent collapsing in general if the combined mesh would have a triangle count above the threshold but too many material slots:
-				// We'll just disable Nanite with a message on the log instead.
-				// This because not only is it difficult to estimate the total number of UE triangles we'll get from the combined USD Mesh prims, but also because
-				// it doesn't really work very well: Imagine we have the Kitchen Set scene (that collapses to like 500 material slots) and we put the threshold just under
-				// the combined number of triangles. This means we can't collapse then, as the combined mesh would want Nanite (as it has more triangles than the threshold)
-				// but has too many material slots. If we prevent if from collapsing, what do we do with the uncollapsed meshes then?
-				//  - If we enable Nanite for them its a bit unexpected because now we'll have a bunch of tiny meshes that for some reason have Nanite enabled;
-				//  - If we don't enable Nanite for them, then what is the benefit? Now the mesh hasn't collapsed but we don't have Nanite anywhere anyway...
-				// At least for the case below (with the explicit overrides) we would end up with some meshes having Nanite, according to how the user set them.
-				// In the future we could expose the collapsing controls on the stage actor to let the user control this a bit better
-				const int32 MaxNumSections = 64; // There is no define for this, but it's checked for on NaniteBuilder.cpp, FBuilderModule::Build
-				if ( !NumExpectedMaterialSlots.IsSet() || NumExpectedMaterialSlots.GetValue() > MaxNumSections )
-				{
-					UE_LOG( LogUsd, Log, TEXT( "Not collapsing down from prim '%s' as child meshes want Nanite to be abled but the generated static mesh would have more than '%d' material slots" ),
-						*PrimPath.GetString(),
-						MaxNumSections
-					);
-					bCollapsesChildren = false;
-				}
-			}
-		}
-	}
-
 	return bCollapsesChildren;
 }
 
diff --git a/Engine/Plugins/Importers/USDImporter/Source/USDSchemas/Private/USDInfoCache.cpp b/Engine/Plugins/Importers/USDImporter/Source/USDSchemas/Private/USDInfoCache.cpp
--- a/Engine/Plugins/Importers/USDImporter/Source/USDSchemas/Private/USDInfoCache.cpp
+++ b/Engine/Plugins/Importers/USDImporter/Source/USDSchemas/Private/USDInfoCache.cpp
@@ -3,6 +3,7 @@
 #include "USDInfoCache.h"
 
 #include "USDGeomMeshConversion.h"
+#include "USDGeomXformableTranslator.h"
 #include "USDLog.h"
 #include "USDSchemasModule.h"
 #include "USDSchemaTranslator.h"
@@ -25,6 +26,12 @@
 #include "USDIncludesEnd.h"
 #endif // USE_USD_SDK
 
+static int32 GMaxNumVerticesCollapsedMesh = 5000000;
+static FAutoConsoleVariableRef CVarMaxNumVerticesCollapsedMesh(
+	TEXT("USD.MaxNumVerticesCollapsedMesh"),
+	GMaxNumVerticesCollapsedMesh,
+	TEXT("Maximum number of vertices that a combined Mesh can have for us to collapse it into a single StaticMesh"));
+
 namespace UE::UsdInfoCache::Private
 {
 	struct FUsdPrimInfo
@@ -287,7 +294,7 @@ namespace UE::USDInfoCacheImpl::Private
 		}
 	}
 
-	void RecursiveRebuildCache(
+	void RecursivePropagateVertexAndMaterialSlotCounts(
 		const pxr::UsdPrim& UsdPrim,
 		FUsdSchemaTranslationContext& Context,
 		const pxr::TfToken& MaterialPurposeToken,
@@ -296,12 +303,10 @@ namespace UE::USDInfoCacheImpl::Private
 		TMap< UE::FSdfPath, TArray<UsdUtils::FUsdPrimMaterialSlot>>& InOutSubtreeToMaterialSlots,
 		TArray< FString >& InOutPointInstancerPaths,
 		uint64& OutSubtreeVertexCount,
-		TArray<UsdUtils::FUsdPrimMaterialSlot>& OutSubtreeSlots,
-		const pxr::SdfPath& AssetCollapsedRoot = pxr::SdfPath::EmptyPath(),
-		const pxr::SdfPath& ComponentCollapsedRoot = pxr::SdfPath::EmptyPath()
+		TArray<UsdUtils::FUsdPrimMaterialSlot>& OutSubtreeSlots
 	)
 	{
-		TRACE_CPUPROFILER_EVENT_SCOPE( UE::USDInfoCacheImpl::Private::RecursiveRebuildCache );
+		TRACE_CPUPROFILER_EVENT_SCOPE( UE::USDInfoCacheImpl::Private::RecursivePropagateVertexAndMaterialSlotCounts );
 
 		if (!UsdPrim)
 		{
@@ -311,35 +316,7 @@ namespace UE::USDInfoCacheImpl::Private
 		FScopedUsdAllocs Allocs;
 
 		pxr::SdfPath UsdPrimPath = UsdPrim.GetPrimPath();
-
-		// Prevents allocation by referencing instead of copying
-		const pxr::SdfPath* AssetCollapsedRootOverride = &AssetCollapsedRoot;
-		const pxr::SdfPath* ComponentCollapsedRootOverride = &ComponentCollapsedRoot;
-
-		bool bIsAssetCollapsed = !AssetCollapsedRoot.IsEmpty();
-		bool bIsComponentCollapsed = !ComponentCollapsedRoot.IsEmpty();
-
-		if ( !bIsAssetCollapsed || !bIsComponentCollapsed )
-		{
-			if ( TSharedPtr< FUsdSchemaTranslator > SchemaTranslator = Registry.CreateTranslatorForSchema( Context.AsShared(), UE::FUsdTyped( UsdPrim ) ) )
-			{
-				if ( !bIsAssetCollapsed )
-				{
-					if ( SchemaTranslator->CollapsesChildren( ECollapsingType::Assets ) )
-					{
-						AssetCollapsedRootOverride = &UsdPrimPath;
-					}
-				}
-
-				if ( !bIsComponentCollapsed )
-				{
-					if ( SchemaTranslator->CollapsesChildren( ECollapsingType::Components ) )
-					{
-						ComponentCollapsedRootOverride = &UsdPrimPath;
-					}
-				}
-			}
-		}
+		UE::FSdfPath PrimPath{UsdPrimPath};
 
 		if (!UsdPrim.IsPseudoRoot())
 		{
@@ -389,10 +366,10 @@ namespace UE::USDInfoCacheImpl::Private
 
 		const int32 MinBatchSize = 1;
 
-		ParallelFor( TEXT( "RecursiveRebuildCache" ), Prims.Num(), MinBatchSize,
+		ParallelFor( TEXT( "RecursivePropagateVertexAndMaterialSlotCounts" ), Prims.Num(), MinBatchSize,
 			[&]( int32 Index )
 			{
-				RecursiveRebuildCache(
+				RecursivePropagateVertexAndMaterialSlotCounts(
 					Prims[ Index ],
 					Context,
 					MaterialPurposeToken,
@@ -401,9 +378,7 @@ namespace UE::USDInfoCacheImpl::Private
 					InOutSubtreeToMaterialSlots,
 					InOutPointInstancerPaths,
 					ChildSubtreeVertexCounts[ Index ],
-					ChildSubtreeMaterialSlots[ Index ],
-					*AssetCollapsedRootOverride,
-					*ComponentCollapsedRootOverride
+					ChildSubtreeMaterialSlots[ Index ]
 				);
 			}
 		);
@@ -454,12 +429,6 @@ namespace UE::USDInfoCacheImpl::Private
 				Info.ExpectedVertexCountForSubtree = OutSubtreeVertexCount;
 				InOutSubtreeToMaterialSlots.Emplace( UsdPrimPath, OutSubtreeSlots );
 			}
-
-			// These paths will be still empty in case nothing has collapsed yet, hold UsdPrimPath in case UsdPrim
-			// collapses that type, or hold the path to the collapsed root passed in via our caller, in case we're
-			// collapsed
-			Info.AssetCollapsedRoot = *AssetCollapsedRootOverride;
-			Info.ComponentCollapsedRoot = *ComponentCollapsedRootOverride;
 		}
 	}
 
@@ -612,6 +581,121 @@ namespace UE::USDInfoCacheImpl::Private
 			}
 		}
 	}
+
+	bool CanMeshSubtreeBeCollapsed(
+		const pxr::UsdPrim& UsdPrim,
+		FUsdSchemaTranslationContext& Context,
+		FUsdInfoCache::FUsdInfoCacheImpl& Impl,
+		const TSharedPtr< FUsdSchemaTranslator >& Translator
+	)
+	{
+		if (!UsdPrim)
+		{
+			return false;
+		}
+
+		// Only care about collapsing into a StaticMesh: We should always collapse into a SkeletalMesh as we have
+		// no real alternative for handling them
+		TSharedPtr<FUsdGeomXformableTranslator> XformableTranslator = StaticCastSharedPtr<FUsdGeomXformableTranslator>(Translator);
+		if (!XformableTranslator.IsValid())
+		{
+			return false;
+		}
+
+		pxr::SdfPath UsdPrimPath = UsdPrim.GetPrimPath();
+
+		UE::UsdInfoCache::Private::FUsdPrimInfo& Info = Impl.InfoMap.FindOrAdd(UE::FSdfPath{UsdPrimPath});
+
+		if (Info.ExpectedVertexCountForSubtree > GMaxNumVerticesCollapsedMesh)
+		{
+			return false;
+		}
+
+		return true;
+	}
+
+	void RecursiveQueryCollapsesChildren(
+		const pxr::UsdPrim& UsdPrim,
+		FUsdSchemaTranslationContext& Context,
+		FUsdInfoCache::FUsdInfoCacheImpl& Impl,
+		FUsdSchemaTranslatorRegistry& Registry,
+		const pxr::SdfPath& AssetCollapsedRoot = pxr::SdfPath::EmptyPath(),
+		const pxr::SdfPath& ComponentCollapsedRoot = pxr::SdfPath::EmptyPath()
+	)
+	{
+		TRACE_CPUPROFILER_EVENT_SCOPE(UE::USDInfoCacheImpl::Private::RecursiveQueryCollapsesChildren);
+		FScopedUsdAllocs Allocs;
+
+		pxr::SdfPath UsdPrimPath = UsdPrim.GetPrimPath();
+		UE::FSdfPath PrimPath{UsdPrimPath};
+
+		// Prevents allocation by referencing instead of copying
+		const pxr::SdfPath* AssetCollapsedRootOverride = &AssetCollapsedRoot;
+		const pxr::SdfPath* ComponentCollapsedRootOverride = &ComponentCollapsedRoot;
+
+		bool bIsAssetCollapsed = !AssetCollapsedRoot.IsEmpty();
+		bool bIsComponentCollapsed = !ComponentCollapsedRoot.IsEmpty();
+
+		if (!bIsAssetCollapsed || !bIsComponentCollapsed)
+		{
+			if (TSharedPtr< FUsdSchemaTranslator > SchemaTranslator = Registry.CreateTranslatorForSchema(Context.AsShared(), UE::FUsdTyped(UsdPrim)))
+			{
+				const bool bCanMeshSubtreeBeCollapsed = CanMeshSubtreeBeCollapsed(UsdPrim, Context, Impl, SchemaTranslator);
+
+				if (!bIsAssetCollapsed)
+				{
+					if (SchemaTranslator->CollapsesChildren(ECollapsingType::Assets) && bCanMeshSubtreeBeCollapsed)
+					{
+						AssetCollapsedRootOverride = &UsdPrimPath;
+					}
+				}
+
+				if (!bIsComponentCollapsed)
+				{
+					if (SchemaTranslator->CollapsesChildren(ECollapsingType::Components) && bCanMeshSubtreeBeCollapsed)
+					{
+						ComponentCollapsedRootOverride = &UsdPrimPath;
+					}
+				}
+			}
+		}
+
+		pxr::UsdPrimSiblingRange PrimChildren = UsdPrim.GetFilteredChildren(pxr::UsdTraverseInstanceProxies(pxr::UsdPrimAllPrimsPredicate));
+
+		TArray<pxr::UsdPrim> Prims;
+		for (pxr::UsdPrim Child : PrimChildren)
+		{
+			Prims.Emplace(Child);
+		}
+
+		const int32 MinBatchSize = 1;
+
+		ParallelFor(TEXT("RecursiveQueryCollapsesChildren"), Prims.Num(), MinBatchSize,
+			[&](int32 Index)
+			{
+				RecursiveQueryCollapsesChildren(
+					Prims[Index],
+					Context,
+					Impl,
+					Registry,
+					*AssetCollapsedRootOverride,
+					*ComponentCollapsedRootOverride
+				);
+			}
+		);
+
+		{
+			FWriteScopeLock ScopeLock(Impl.InfoMapLock);
+
+			UE::UsdInfoCache::Private::FUsdPrimInfo& Info = Impl.InfoMap.FindOrAdd(UE::FSdfPath{UsdPrimPath});
+
+			// These paths will be still empty in case nothing has collapsed yet, hold UsdPrimPath in case UsdPrim
+			// collapses that type, or hold the path to the collapsed root passed in via our caller, in case we're
+			// collapsed
+			Info.AssetCollapsedRoot = *AssetCollapsedRootOverride;
+			Info.ComponentCollapsedRoot = *ComponentCollapsedRootOverride;
+		}
+	}
 #endif // USE_USD_SDK
 }
 
@@ -783,9 +867,11 @@ void FUsdInfoCache::RebuildCacheForSubtree( const UE::FUsdPrim& Prim, FUsdSchema
 			MaterialPurposeToken = UnrealToUsd::ConvertToken(*Context.MaterialPurpose.ToString()).Get();
 		}
 
+		// Propagate vertex and material slot counts before we query CollapsesChildren because the Xformable
+		// translator needs to know when it would generate too large a static mesh
 		uint64 SubtreeVertexCount = 0;
 		TArray<UsdUtils::FUsdPrimMaterialSlot> SubtreeSlots;
-		UE::USDInfoCacheImpl::Private::RecursiveRebuildCache(
+		UE::USDInfoCacheImpl::Private::RecursivePropagateVertexAndMaterialSlotCounts(
 			UsdPrim,
 			Context,
 			MaterialPurposeToken,
@@ -810,6 +896,13 @@ void FUsdInfoCache::RebuildCacheForSubtree( const UE::FUsdPrim& Prim, FUsdSchema
 			TempSubtreeSlots,
 			Context.bMergeIdenticalMaterialSlots
 		);
+
+		UE::USDInfoCacheImpl::Private::RecursiveQueryCollapsesChildren(
+			UsdPrim,
+			Context,
+			*ImplPtr,
+			Registry
+		);
 	}
 #endif // USE_USD_SDK
 }
