USD: Fix anim sequences no longer being able to drive morph target animations.

diff --git a/Engine/Plugins/Importers/USDImporter/Source/USDUtilities/Private/USDSkeletalDataConversion.cpp b/Engine/Plugins/Importers/USDImporter/Source/USDUtilities/Private/USDSkeletalDataConversion.cpp
--- a/Engine/Plugins/Importers/USDImporter/Source/USDUtilities/Private/USDSkeletalDataConversion.cpp
+++ b/Engine/Plugins/Importers/USDImporter/Source/USDUtilities/Private/USDSkeletalDataConversion.cpp
@@ -82,6 +82,12 @@
 
 #define LOCTEXT_NAMESPACE "UsdSkeletalDataConversion"
 
+static bool bAddCurveMetadataToSkeleton = false;
+static FAutoConsoleVariableRef CVarAddCurveMetadataToSkeleton(
+	TEXT("USD.AddCurveMetadataToSkeleton"),
+	bAddCurveMetadataToSkeleton,
+	TEXT("When true will cause blend shape / morph target float curve data to be added to generated USkeleton assets. When false, this curve data will be added to the generated USkeletalMesh assets instead."));
+
 #if USE_USD_SDK && WITH_EDITOR
 namespace SkelDataConversionImpl
 {
@@ -412,6 +418,35 @@ namespace SkelDataConversionImpl
 			return;
 		}
 
+		if (bAddCurveMetadataToSkeleton)
+		{
+			const bool bMaterialCurve = false;
+			const bool bMorphTargetCurve = true;
+			Skeleton->AccumulateCurveMetaData(CurveName, bMaterialCurve, bMorphTargetCurve);
+		}
+		else
+		{
+			// We should always have the preview mesh here, as we need to declare these curve names on it now
+			USkeletalMesh* Mesh = Sequence->GetPreviewMesh();
+			if (ensure(Mesh))
+			{
+				UAnimCurveMetaData* AnimCurveMetaData = Mesh->GetAssetUserData<UAnimCurveMetaData>();
+				if (AnimCurveMetaData == nullptr)
+				{
+					AnimCurveMetaData = NewObject<UAnimCurveMetaData>(Mesh, NAME_None, RF_Transactional);
+					Mesh->AddAssetUserData(AnimCurveMetaData);
+				}
+
+				AnimCurveMetaData->AddCurveMetaData(CurveName);
+
+				// Ensure we have a morph flag set
+				if (FCurveMetaData* CurveMetaData = AnimCurveMetaData->GetCurveMetaData(CurveName))
+				{
+					CurveMetaData->Type.bMorphtarget = true;
+				}
+			}
+		}
+
 		// Ignore curves that don't contribute to the animation
 		bool bHasNonZeroKey = false;
 		for ( const FRichCurveKey& Key : SourceData.Keys )
@@ -426,7 +461,7 @@ namespace SkelDataConversionImpl
 		{
 			return;
 		}
-		
+
 		const bool bShouldTransact = false;
 		const IAnimationDataModel* DataModel = Sequence->GetDataModel();
 		IAnimationDataController& Controller = Sequence->GetController();
@@ -452,7 +487,7 @@ namespace SkelDataConversionImpl
 
 			Controller.SetCurveFlags( CurveId, Curve->GetCurveTypeFlags() | AACF_DefaultCurve, bShouldTransact );
 		}
-		
+
 		if ( Curve )
 		{
 			Controller.SetCurveKeys( CurveId, SourceData.GetConstRefOfKeys(), bShouldTransact );
@@ -2355,7 +2390,6 @@ bool UsdToUnreal::ConvertSkelAnim(
 	OutSkeletalAnimationAsset->ImportFileFramerate = LayerTimeCodesPerSecond;
 	OutSkeletalAnimationAsset->ImportResampleFramerate = LayerTimeCodesPerSecond;
 
-
 	const FFrameRate FrameRate(LayerTimeCodesPerSecond, 1);
 	Controller.SetFrameRate(FrameRate, bShouldTransact);
 	const FFrameNumber FrameNumber = FrameRate.AsFrameNumber(LayerSequenceLengthSeconds);
