USDImporter: use token-based IsA() and HasAPI() functions in UsdPrim UE wrapper

With the engine USD version upgraded to v23.02 and the IsA() and HasAPI() functions
aligned between the UE FUsdPrim type and the pxr UsdPrim type, we can now call the
token-based overloads of those functions on the underlying UsdPrim. This also means
that we no longer need to deal with the UsdSchemaRegistry directly in any of these
functions in the UE wrapper.

diff --git a/Engine/Plugins/Importers/USDImporter/Source/UnrealUSDWrapper/Private/UsdWrappers/UsdPrim.cpp b/Engine/Plugins/Importers/USDImporter/Source/UnrealUSDWrapper/Private/UsdWrappers/UsdPrim.cpp
--- a/Engine/Plugins/Importers/USDImporter/Source/UnrealUSDWrapper/Private/UsdWrappers/UsdPrim.cpp
+++ b/Engine/Plugins/Importers/USDImporter/Source/UnrealUSDWrapper/Private/UsdWrappers/UsdPrim.cpp
@@ -12,7 +12,6 @@
 
 #include "USDIncludesStart.h"
 	#include "pxr/usd/sdf/schema.h"
-	#include "pxr/usd/usd/schemaRegistry.h"
 	#include "pxr/usd/usd/attribute.h"
 	#include "pxr/usd/usd/prim.h"
 	#include "pxr/usd/usd/schemaBase.h"
@@ -247,13 +246,9 @@ namespace UE
 #if USE_USD_SDK
 		FScopedUsdAllocs Allocs;
 
-		pxr::TfType Type = pxr::UsdSchemaRegistry::GetTypeFromName( pxr::TfToken( TCHAR_TO_ANSI( *SchemaIdentifier.ToString() ) ) );
-		if ( Type.IsUnknown() )
-		{
-			return false;
-		}
+		const pxr::TfToken UsdSchemaIdentifier( TCHAR_TO_ANSI( *SchemaIdentifier.ToString() ) );
 
-		return Impl->PxrUsdPrim.Get().IsA( Type );
+		return Impl->PxrUsdPrim.Get().IsA( UsdSchemaIdentifier );
 #else
 		return false;
 #endif // #if USE_USD_SDK
@@ -264,13 +259,9 @@ namespace UE
 #if USE_USD_SDK
 		FScopedUsdAllocs Allocs;
 
-		pxr::TfType Type = pxr::UsdSchemaRegistry::GetTypeFromName( pxr::TfToken( TCHAR_TO_ANSI( *SchemaIdentifier.ToString() ) ) );
-		if ( Type.IsUnknown() )
-		{
-			return false;
-		}
+		const pxr::TfToken UsdSchemaIdentifier( TCHAR_TO_ANSI( *SchemaIdentifier.ToString() ) );
 
-		return Impl->PxrUsdPrim.Get().HasAPI( Type );
+		return Impl->PxrUsdPrim.Get().HasAPI( UsdSchemaIdentifier );
 #else
 		return false;
 #endif // #if USE_USD_SDK
@@ -281,15 +272,10 @@ namespace UE
 #if USE_USD_SDK
 		FScopedUsdAllocs Allocs;
 
-		pxr::TfType Type = pxr::UsdSchemaRegistry::GetTypeFromName( pxr::TfToken( TCHAR_TO_ANSI( *SchemaIdentifier.ToString() ) ) );
-		if ( Type.IsUnknown() )
-		{
-			return false;
-		}
-
-		pxr::TfToken UsdInstanceName = pxr::TfToken( TCHAR_TO_ANSI( *InstanceName.ToString() ) );
+		const pxr::TfToken UsdSchemaIdentifier( TCHAR_TO_ANSI( *SchemaIdentifier.ToString() ) );
+		const pxr::TfToken UsdInstanceName( TCHAR_TO_ANSI( *InstanceName.ToString() ) );
 
-		return Impl->PxrUsdPrim.Get().HasAPI( Type, UsdInstanceName );
+		return Impl->PxrUsdPrim.Get().HasAPI( UsdSchemaIdentifier, UsdInstanceName );
 #else
 		return false;
 #endif // #if USE_USD_SDK
