USD: Fix assets being discarded by the asset cache before the stage actor is done with them;

diff --git a/Engine/Plugins/Importers/USDImporter/Source/USDSchemas/Private/USDInfoCache.cpp b/Engine/Plugins/Importers/USDImporter/Source/USDSchemas/Private/USDInfoCache.cpp
--- a/Engine/Plugins/Importers/USDImporter/Source/USDSchemas/Private/USDInfoCache.cpp
+++ b/Engine/Plugins/Importers/USDImporter/Source/USDSchemas/Private/USDInfoCache.cpp
@@ -61,6 +61,7 @@ struct FUsdInfoCache::FUsdInfoCacheImpl
 
 	// Information we may have about a subset of prims
 	TMap<UE::FSdfPath, TSet<TWeakObjectPtr<UObject>>> PrimPathToAssets;
+	TMap<TWeakObjectPtr<UObject>, TSet<UE::FSdfPath>> AssetToPrimPaths;
 	mutable FRWLock PrimPathToAssetsLock;
 
 	// Paths to material prims that are actually used by mesh prims in the scene, given the current settings for
@@ -89,6 +90,7 @@ bool FUsdInfoCache::Serialize( FArchive& Ar )
 		{
 			FWriteScopeLock ScopeLock(ImplPtr->PrimPathToAssetsLock);
 			Ar << ImplPtr->PrimPathToAssets;
+			Ar << ImplPtr->AssetToPrimPaths;
 		}
 		{
 			FWriteScopeLock ScopeLock(ImplPtr->UsedMaterialPathsLock);
@@ -755,6 +757,7 @@ void FUsdInfoCache::LinkAssetToPrim(const UE::FSdfPath& Path, UObject* Asset)
 	);
 
 	ImplPtr->PrimPathToAssets.FindOrAdd(Path).Add(Asset);
+	ImplPtr->AssetToPrimPaths.FindOrAdd(Asset).Add(Path);
 }
 
 TSet<TWeakObjectPtr<UObject>> FUsdInfoCache::RemoveAllAssetPrimLinks(const UE::FSdfPath& Path)
@@ -766,9 +769,18 @@ TSet<TWeakObjectPtr<UObject>> FUsdInfoCache::RemoveAllAssetPrimLinks(const UE::F
 	}
 	FWriteScopeLock ScopeLock(ImplPtr->PrimPathToAssetsLock);
 
-	TSet<TWeakObjectPtr<UObject>> Values;
-	ImplPtr->PrimPathToAssets.RemoveAndCopyValue(Path, Values);
-	return Values;
+	TSet<TWeakObjectPtr<UObject>> Assets;
+	ImplPtr->PrimPathToAssets.RemoveAndCopyValue(Path, Assets);
+
+	for (const TWeakObjectPtr<UObject>& Asset : Assets)
+	{
+		if (TSet<UE::FSdfPath>* PrimPaths = ImplPtr->AssetToPrimPaths.Find(Asset))
+		{
+			PrimPaths->Remove(Path);
+		}
+	}
+
+	return Assets;
 }
 
 TSet<TWeakObjectPtr<UObject>> FUsdInfoCache::GetAllAssetsForPrim(const UE::FSdfPath& Path) const
@@ -788,8 +800,13 @@ TSet<TWeakObjectPtr<UObject>> FUsdInfoCache::GetAllAssetsForPrim(const UE::FSdfP
 	return {};
 }
 
-UE::FSdfPath FUsdInfoCache::GetPrimForAsset(UObject* Asset) const
+TSet<UE::FSdfPath> FUsdInfoCache::GetPrimsForAsset(UObject* Asset) const
 {
+	if (!Asset)
+	{
+		return {};
+	}
+
 	FUsdInfoCacheImpl* ImplPtr = Impl.Get();
 	if (!ImplPtr)
 	{
@@ -797,12 +814,9 @@ UE::FSdfPath FUsdInfoCache::GetPrimForAsset(UObject* Asset) const
 	}
 	FReadScopeLock ScopeLock(ImplPtr->PrimPathToAssetsLock);
 
-	for (const TPair<UE::FSdfPath, TSet<TWeakObjectPtr<UObject>>>& Pair : ImplPtr->PrimPathToAssets)
+	if (const TSet<UE::FSdfPath>* FoundPrims = ImplPtr->AssetToPrimPaths.Find(Asset))
 	{
-		if (Pair.Value.Contains(Asset))
-		{
-			return Pair.Key;
-		}
+		return *FoundPrims;
 	}
 
 	return {};
@@ -918,6 +932,7 @@ void FUsdInfoCache::Clear()
 		{
 			FWriteScopeLock ScopeLock(ImplPtr->PrimPathToAssetsLock);
 			ImplPtr->PrimPathToAssets.Empty();
+			ImplPtr->AssetToPrimPaths.Empty();
 		}
 		{
 			FWriteScopeLock ScopeLock(ImplPtr->UsedMaterialPathsLock);
diff --git a/Engine/Plugins/Importers/USDImporter/Source/USDSchemas/Private/USDShadeMaterialTranslator.cpp b/Engine/Plugins/Importers/USDImporter/Source/USDSchemas/Private/USDShadeMaterialTranslator.cpp
--- a/Engine/Plugins/Importers/USDImporter/Source/USDSchemas/Private/USDShadeMaterialTranslator.cpp
+++ b/Engine/Plugins/Importers/USDImporter/Source/USDSchemas/Private/USDShadeMaterialTranslator.cpp
@@ -131,9 +131,12 @@ namespace UE::UsdShadeTranslator::Private
 						{
 							if ( Context->AssetCache && Context->InfoCache )
 							{
-								// This is super slow as it will essentially linear search the asset cache, but
-								// not much we can do at runtime without asset import data
-								const UE::FSdfPath PrimPath = Context->InfoCache->GetPrimForAsset( MID );
+								UE::FSdfPath PrimPath;
+								for (const UE::FSdfPath& Path : Context->InfoCache->GetPrimsForAsset(MID))
+								{
+									PrimPath = Path;
+									break;
+								}
 								const FString Hash = Context->AssetCache->GetHashForAsset( MID );
 
 								const FName NewInstanceName = MakeUniqueObjectName(
@@ -155,7 +158,7 @@ namespace UE::UsdShadeTranslator::Private
 								}
 								NewMID->CopyParameterOverrides( MID );
 
-								if (Context->AssetCache->RemoveAsset(Hash))
+								if (Context->AssetCache->CanRemoveAsset(Hash) && Context->AssetCache->RemoveAsset(Hash))
 								{
 									Context->AssetCache->CacheAsset(Hash, NewMID);
 									if (!PrimPath.IsEmpty())
@@ -424,9 +427,11 @@ void FUsdShadeMaterialTranslator::CreateAssets()
 	}
 	else if ( Context->MaterialToPrimvarToUVIndex && Context->InfoCache )
 	{
-		const UE::FSdfPath FoundPrimPath = Context->InfoCache->GetPrimForAsset(ConvertedMaterial);
-		if (!FoundPrimPath.IsEmpty())
+		const TSet<UE::FSdfPath> FoundPrimPaths = Context->InfoCache->GetPrimsForAsset(ConvertedMaterial);
+		if (FoundPrimPaths.Num() > 0)
 		{
+			const UE::FSdfPath& FoundPrimPath = *FoundPrimPaths.CreateConstIterator();
+
 			if (TMap<FString, int32>* PrimvarToUVIndex = Context->MaterialToPrimvarToUVIndex->Find(FoundPrimPath.GetString()))
 			{
 				// Copy the Material -> Primvar -> UV index mapping from the cached material prim path to this prim path
diff --git a/Engine/Plugins/Importers/USDImporter/Source/USDSchemas/Public/USDInfoCache.h b/Engine/Plugins/Importers/USDImporter/Source/USDSchemas/Public/USDInfoCache.h
--- a/Engine/Plugins/Importers/USDImporter/Source/USDSchemas/Public/USDInfoCache.h
+++ b/Engine/Plugins/Importers/USDImporter/Source/USDSchemas/Public/USDInfoCache.h
@@ -101,7 +101,7 @@ public:
 		return CastAssets;
 	}
 
-	UE::FSdfPath GetPrimForAsset(UObject* Asset) const;
+	TSet<UE::FSdfPath> GetPrimsForAsset(UObject* Asset) const;
 	TMap<UE::FSdfPath, TSet<TWeakObjectPtr<UObject>>> GetAllAssetPrimLinks() const;
 
 private:
diff --git a/Engine/Plugins/Importers/USDImporter/Source/USDStage/Private/USDStageActor.cpp b/Engine/Plugins/Importers/USDImporter/Source/USDStage/Private/USDStageActor.cpp
--- a/Engine/Plugins/Importers/USDImporter/Source/USDStage/Private/USDStageActor.cpp
+++ b/Engine/Plugins/Importers/USDImporter/Source/USDStage/Private/USDStageActor.cpp
@@ -1832,10 +1832,10 @@ FString AUsdStageActor::GetSourcePrimPath(UObject* Object)
 	}
 	else if (InfoCache)
 	{
-		const UE::FSdfPath FoundPath = InfoCache->GetPrimForAsset(Object);
-		if (!FoundPath.IsEmpty())
+		const TSet<UE::FSdfPath> FoundPaths = InfoCache->GetPrimsForAsset(Object);
+		if (FoundPaths.Num() > 0)
 		{
-			return FoundPath.GetString();
+			return FoundPaths.CreateConstIterator()->GetString();
 		}
 	}
 
@@ -3221,12 +3221,29 @@ void AUsdStageActor::LoadAsset(FUsdSchemaTranslationContext& TranslationContext,
 
 	if (InfoCache)
 	{
-		TSet<TWeakObjectPtr<UObject>> OldAssets = InfoCache->RemoveAllAssetPrimLinks(UE::FSdfPath{*PrimPath});
-
 		if (UsdAssetCache)
 		{
+			UE::FSdfPath UsdPrimPath = UE::FSdfPath{*PrimPath};
+			TSet<TWeakObjectPtr<UObject>> OldAssets = InfoCache->RemoveAllAssetPrimLinks(UsdPrimPath);
+
 			for (const TWeakObjectPtr<UObject>& OldAsset : OldAssets)
 			{
+				// If there are any other prim paths linked to this asset that we *won't* be removing/reparsing
+				// in here, it means our stage actor as a whole is still "referencing" that asset
+				bool bAssetStillReferenced = false;
+				for (const UE::FSdfPath& LinkedPrim : InfoCache->GetPrimsForAsset(OldAsset.Get()))
+				{
+					if (LinkedPrim != UsdPrimPath)
+					{
+						bAssetStillReferenced = true;
+						break;
+					}
+				}
+				if (bAssetStillReferenced)
+				{
+					continue;
+				}
+
 				UsdAssetCache->RemoveAssetReference(OldAsset.Get(), this);
 			}
 		}
@@ -3270,6 +3287,22 @@ void AUsdStageActor::LoadAssets(FUsdSchemaTranslationContext& TranslationContext
 				TSet<TWeakObjectPtr<UObject>> OldAssets = InfoCache->RemoveAllAssetPrimLinks(PrimPathToRemove);
 				for (const TWeakObjectPtr<UObject>& OldAsset : OldAssets)
 				{
+					// If there are any other prim paths linked to this asset that we *won't* be removing/reparsing
+					// in here, it means our stage actor as a whole is still "referencing" that asset
+					bool bAssetStillReferenced = false;
+					for (const UE::FSdfPath& LinkedPrim : InfoCache->GetPrimsForAsset(OldAsset.Get()))
+					{
+						if (!PrimPathsToRemove.Contains(LinkedPrim))
+						{
+							bAssetStillReferenced = true;
+							break;
+						}
+					}
+					if (bAssetStillReferenced)
+					{
+						continue;
+					}
+
 					UsdAssetCache->RemoveAssetReference(OldAsset.Get(), this);
 				}
 			}
