USD: Fix asset names not being properly sanitized in some scenarios;

Cleanup the multiple copies of the name sanitizing function.

diff --git a/Engine/Plugins/Importers/USDImporter/Source/USDClasses/Private/USDClassesModule.cpp b/Engine/Plugins/Importers/USDImporter/Source/USDClasses/Private/USDClassesModule.cpp
--- a/Engine/Plugins/Importers/USDImporter/Source/USDClasses/Private/USDClassesModule.cpp
+++ b/Engine/Plugins/Importers/USDImporter/Source/USDClasses/Private/USDClassesModule.cpp
@@ -25,6 +25,7 @@
 #include "Misc/SecureHash.h"
 #include "PhysicsEngine/PhysicsAsset.h"
 #include "Serialization/JsonSerializer.h"
+#include "UObject/NameTypes.h"
 #include "UObject/ObjectSaveContext.h"
 #include "UObject/Package.h"
 
@@ -404,6 +405,19 @@ TSet<UObject*> IUsdClassesModule::GetAssetDependencies(UObject* Asset)
 	return Result;
 }
 
+FString IUsdClassesModule::SanitizeObjectName(const FString& InObjectName)
+{
+	FString SanitizedText = InObjectName;
+	const TCHAR* InvalidChar = INVALID_OBJECTNAME_CHARACTERS;
+	while (*InvalidChar)
+	{
+		SanitizedText.ReplaceCharInline(*InvalidChar, TCHAR('_'), ESearchCase::CaseSensitive);
+		++InvalidChar;
+	}
+
+	return SanitizedText;
+}
+
 class FUsdClassesModule : public IUsdClassesModule
 {
 
diff --git a/Engine/Plugins/Importers/USDImporter/Source/USDClasses/Public/USDClassesModule.h b/Engine/Plugins/Importers/USDImporter/Source/USDClasses/Public/USDClassesModule.h
--- a/Engine/Plugins/Importers/USDImporter/Source/USDClasses/Public/USDClassesModule.h
+++ b/Engine/Plugins/Importers/USDImporter/Source/USDClasses/Public/USDClassesModule.h
@@ -46,4 +46,7 @@ public:
 	 * When given a mesh, will return materials, etc.)
 	 */
 	USDCLASSES_API static TSet<UObject*> GetAssetDependencies(UObject* Asset);
+
+	// Adapted from ObjectTools as it is within an Editor-only module
+	USDCLASSES_API static FString SanitizeObjectName(const FString& InObjectName);
 };
\ No newline at end of file
diff --git a/Engine/Plugins/Importers/USDImporter/Source/USDSchemas/Private/USDGeomMeshTranslator.cpp b/Engine/Plugins/Importers/USDImporter/Source/USDSchemas/Private/USDGeomMeshTranslator.cpp
--- a/Engine/Plugins/Importers/USDImporter/Source/USDSchemas/Private/USDGeomMeshTranslator.cpp
+++ b/Engine/Plugins/Importers/USDImporter/Source/USDSchemas/Private/USDGeomMeshTranslator.cpp
@@ -777,7 +777,11 @@ namespace UsdGeomMeshTranslatorImpl
 		{
 			bOutIsNew = true;
 
-			FName AssetName = MakeUniqueObjectName(GetTransientPackage(), UStaticMesh::StaticClass(), *FPaths::GetBaseFilename(MeshName));
+			FName AssetName = MakeUniqueObjectName(
+				GetTransientPackage(),
+				UStaticMesh::StaticClass(),
+				*IUsdClassesModule::SanitizeObjectName(FPaths::GetBaseFilename(MeshName))
+			);
 			StaticMesh = NewObject< UStaticMesh >(GetTransientPackage(), AssetName, Context.ObjectFlags | EObjectFlags::RF_Public | EObjectFlags::RF_Transient);
 
 #if WITH_EDITOR
@@ -989,7 +993,11 @@ namespace UsdGeomMeshTranslatorImpl
 		{
 			bOutIsNew = true;
 
-			const FName AssetName = MakeUniqueObjectName(GetTransientPackage(), UGeometryCache::StaticClass(), *FPaths::GetBaseFilename(InPrimPath));
+			const FName AssetName = MakeUniqueObjectName(
+				GetTransientPackage(),
+				UGeometryCache::StaticClass(),
+				*IUsdClassesModule::SanitizeObjectName(FPaths::GetBaseFilename(InPrimPath))
+			);
 			GeometryCache = NewObject< UGeometryCache >(GetTransientPackage(), AssetName, Context->ObjectFlags | EObjectFlags::RF_Public);
 
 			const UE::FUsdStage& Stage = Context->Stage;
diff --git a/Engine/Plugins/Importers/USDImporter/Source/USDSchemas/Private/USDGeomXformableTranslator.cpp b/Engine/Plugins/Importers/USDImporter/Source/USDSchemas/Private/USDGeomXformableTranslator.cpp
--- a/Engine/Plugins/Importers/USDImporter/Source/USDSchemas/Private/USDGeomXformableTranslator.cpp
+++ b/Engine/Plugins/Importers/USDImporter/Source/USDSchemas/Private/USDGeomXformableTranslator.cpp
@@ -2,9 +2,9 @@
 
 #include "USDGeomXformableTranslator.h"
 
-#include "Engine/Level.h"
 #include "UnrealUSDWrapper.h"
 #include "USDAssetCache.h"
+#include "USDClassesModule.h"
 #include "USDConversionUtils.h"
 #include "USDErrorUtils.h"
 #include "USDGeomMeshConversion.h"
@@ -21,6 +21,7 @@
 #include "Components/SceneComponent.h"
 #include "Components/SkeletalMeshComponent.h"
 #include "Components/StaticMeshComponent.h"
+#include "Engine/Level.h"
 #include "Engine/StaticMesh.h"
 #include "Engine/World.h"
 #include "LiveLinkComponentController.h"
@@ -472,7 +473,11 @@ USceneComponent* FUsdGeomXformableTranslator::CreateComponentsEx( TOptional< TSu
 
 		if ( ComponentType.IsSet() && ComponentType.GetValue() != nullptr )
 		{
-			const FName ComponentName = MakeUniqueObjectName( ComponentOuter, ComponentType.GetValue(), FName( Prim.GetName() ) );
+			const FName ComponentName = MakeUniqueObjectName(
+				ComponentOuter,
+				ComponentType.GetValue(),
+				*IUsdClassesModule::SanitizeObjectName(Prim.GetName().ToString())
+			);
 			SceneComponent = NewObject< USceneComponent >( ComponentOuter, ComponentType.GetValue(), ComponentName, ComponentFlags );
 
 			if ( AActor* Owner = SceneComponent->GetOwner() )
diff --git a/Engine/Plugins/Importers/USDImporter/Source/USDSchemas/Private/USDGroomTranslator.cpp b/Engine/Plugins/Importers/USDImporter/Source/USDSchemas/Private/USDGroomTranslator.cpp
--- a/Engine/Plugins/Importers/USDImporter/Source/USDSchemas/Private/USDGroomTranslator.cpp
+++ b/Engine/Plugins/Importers/USDImporter/Source/USDSchemas/Private/USDGroomTranslator.cpp
@@ -19,6 +19,7 @@
 
 #include "UnrealUSDWrapper.h"
 #include "USDAssetImportData.h"
+#include "USDClassesModule.h"
 #include "USDGroomConversion.h"
 #include "USDGroomTranslatorUtils.h"
 #include "USDIntegrationUtils.h"
@@ -174,7 +175,11 @@ protected:
 				UGroomAsset* GroomAsset = Cast<UGroomAsset>(Context->AssetCache->GetCachedAsset(SHAHash.ToString()));
 				if (!GroomAsset)
 				{
-					FName AssetName = MakeUniqueObjectName(GetTransientPackage(), UGroomAsset::StaticClass(), *FPaths::GetBaseFilename(PrimPathString));
+					FName AssetName = MakeUniqueObjectName(
+						GetTransientPackage(),
+						UGroomAsset::StaticClass(),
+						*IUsdClassesModule::SanitizeObjectName(FPaths::GetBaseFilename(PrimPathString))
+					);
 
 					FHairImportContext HairImportContext(ImportOptions.Get(), GetTransientPackage(), UGroomAsset::StaticClass(), AssetName, Context->ObjectFlags | EObjectFlags::RF_Public);
 					UGroomAsset* ExistingAsset = nullptr;
@@ -321,7 +326,11 @@ protected:
 			{
 				const FString StrandsGroomCachePrimPath = UsdGroomTranslatorUtils::GetStrandsGroomCachePrimPath(PrimPath);
 				FHairImportContext HairImportContext(nullptr, GetTransientPackage(), nullptr, FName(), Context->ObjectFlags | EObjectFlags::RF_Public);
-				FName UniqueName = MakeUniqueObjectName(GetTransientPackage(), UGroomCache::StaticClass(), *FPaths::GetBaseFilename(StrandsGroomCachePrimPath));
+				FName UniqueName = MakeUniqueObjectName(
+					GetTransientPackage(),
+					UGroomCache::StaticClass(),
+					*IUsdClassesModule::SanitizeObjectName(FPaths::GetBaseFilename(StrandsGroomCachePrimPath))
+				);
 
 				// Once the processing has completed successfully, the data is transferred to the GroomCache
 				UGroomCache* GroomCache = FGroomCacheImporter::ProcessToGroomCache(*GroomCacheProcessor, AnimInfo, HairImportContext, UniqueName.ToString());
diff --git a/Engine/Plugins/Importers/USDImporter/Source/USDSchemas/Private/USDGroomTranslatorUtils.cpp b/Engine/Plugins/Importers/USDImporter/Source/USDSchemas/Private/USDGroomTranslatorUtils.cpp
--- a/Engine/Plugins/Importers/USDImporter/Source/USDSchemas/Private/USDGroomTranslatorUtils.cpp
+++ b/Engine/Plugins/Importers/USDImporter/Source/USDSchemas/Private/USDGroomTranslatorUtils.cpp
@@ -15,6 +15,7 @@
 
 #include "UnrealUSDWrapper.h"
 #include "USDAssetCache.h"
+#include "USDClassesModule.h"
 #include "USDGeomXformableTranslator.h"
 #include "USDInfoCache.h"
 #include "USDIntegrationUtils.h"
@@ -98,7 +99,11 @@ namespace UE::UsdGroomTranslatorUtils::Private
 			return nullptr;
 		}
 
-		const FName BindingAssetName = MakeUniqueObjectName(GetTransientPackage(), UGroomBindingAsset::StaticClass(), *FPaths::GetBaseFilename(GroomBindingPath));
+		const FName BindingAssetName = MakeUniqueObjectName(
+			GetTransientPackage(),
+			UGroomBindingAsset::StaticClass(),
+			*IUsdClassesModule::SanitizeObjectName(FPaths::GetBaseFilename(GroomBindingPath))
+		);
 		UGroomBindingAsset* GroomBinding = NewObject<UGroomBindingAsset>(GetTransientPackage(), BindingAssetName, ObjectFlags | RF_Public);
 		if (GroomBinding)
 		{
diff --git a/Engine/Plugins/Importers/USDImporter/Source/USDSchemas/Private/USDShadeMaterialTranslator.cpp b/Engine/Plugins/Importers/USDImporter/Source/USDSchemas/Private/USDShadeMaterialTranslator.cpp
--- a/Engine/Plugins/Importers/USDImporter/Source/USDSchemas/Private/USDShadeMaterialTranslator.cpp
+++ b/Engine/Plugins/Importers/USDImporter/Source/USDSchemas/Private/USDShadeMaterialTranslator.cpp
@@ -6,6 +6,7 @@
 #include "MeshTranslationImpl.h"
 #include "USDAssetCache.h"
 #include "USDAssetImportData.h"
+#include "USDClassesModule.h"
 #include "USDLog.h"
 #include "USDShadeConversion.h"
 #include "USDTypesConversion.h"
@@ -142,7 +143,9 @@ namespace UE::UsdShadeTranslator::Private
 								const FName NewInstanceName = MakeUniqueObjectName(
 									GetTransientPackage(),
 									UMaterialInstance::StaticClass(),
-									PrimPath.IsEmpty() ? TEXT("MaterialInstance") : *FPaths::GetBaseFilename(PrimPath.GetString())
+									PrimPath.IsEmpty()
+										? TEXT("MaterialInstance")
+										: *IUsdClassesModule::SanitizeObjectName(FPaths::GetBaseFilename(PrimPath.GetString()))
 								);
 
 								// For MID we can't swap the reference material, so we need to create a brand new one and copy
@@ -257,7 +260,7 @@ void FUsdShadeMaterialTranslator::CreateAssets()
 		const FName InstanceName = MakeUniqueObjectName(
 			GetTransientPackage(),
 			UMaterialInstance::StaticClass(),
-			*FPaths::GetBaseFilename(PrimPathString)
+			*IUsdClassesModule::SanitizeObjectName(FPaths::GetBaseFilename(PrimPathString))
 		);
 
 		// TODO: There is an issue here: Note how we're only ever going to write the material prim's primvars into
diff --git a/Engine/Plugins/Importers/USDImporter/Source/USDSchemas/Private/USDSkelRootTranslator.cpp b/Engine/Plugins/Importers/USDImporter/Source/USDSchemas/Private/USDSkelRootTranslator.cpp
--- a/Engine/Plugins/Importers/USDImporter/Source/USDSchemas/Private/USDSkelRootTranslator.cpp
+++ b/Engine/Plugins/Importers/USDImporter/Source/USDSchemas/Private/USDSkelRootTranslator.cpp
@@ -7,6 +7,7 @@
 #include "MeshTranslationImpl.h"
 #include "USDAssetCache.h"
 #include "USDAssetImportData.h"
+#include "USDClassesModule.h"
 #include "USDConversionUtils.h"
 #include "USDErrorUtils.h"
 #include "USDGeomMeshConversion.h"
@@ -799,7 +800,11 @@ namespace UsdSkelRootTranslatorImpl
 		FName AssetName = MakeUniqueObjectName(
 			GetTransientPackage(),
 			UPhysicsAsset::StaticClass(),
-			*FPaths::GetBaseFilename( TEXT( "PHYS_" ) + SkeletalMesh->GetName() )
+			*IUsdClassesModule::SanitizeObjectName(
+				FPaths::GetBaseFilename(
+					TEXT("PHYS_") + SkeletalMesh->GetName()
+				)
+			)
 		);
 
 		UPhysicsAsset* Result = NewObject< UPhysicsAsset >(
@@ -919,7 +924,9 @@ namespace UsdSkelRootTranslatorImpl
 				FName UniqueName = MakeUniqueObjectName(
 					GetTransientPackage(),
 					UAnimBlueprint::StaticClass(),
-					*( PrimName + TEXT( "_DefaultAnimBlueprint" ) )
+					*IUsdClassesModule::SanitizeObjectName(
+						PrimName + TEXT( "_DefaultAnimBlueprint" )
+					)
 				);
 
 				// Duplicate and never reuse these so that they can be assigned independent subject names if desired.
diff --git a/Engine/Plugins/Importers/USDImporter/Source/USDStage/Private/USDLevelSequenceHelper.cpp b/Engine/Plugins/Importers/USDImporter/Source/USDStage/Private/USDLevelSequenceHelper.cpp
--- a/Engine/Plugins/Importers/USDImporter/Source/USDStage/Private/USDLevelSequenceHelper.cpp
+++ b/Engine/Plugins/Importers/USDImporter/Source/USDStage/Private/USDLevelSequenceHelper.cpp
@@ -5,6 +5,7 @@
 #include "GeometryCache.h"
 #include "USDAssetImportData.h"
 #include "USDAttributeUtils.h"
+#include "USDClassesModule.h"
 #include "USDConversionUtils.h"
 #include "USDInfoCache.h"
 #include "USDIntegrationUtils.h"
@@ -13,8 +14,8 @@
 #include "USDLog.h"
 #include "USDMemory.h"
 #include "USDPrimConversion.h"
-#include "USDProjectSettings.h"
 #include "USDPrimTwin.h"
+#include "USDProjectSettings.h"
 #include "USDStageActor.h"
 #include "USDTypesConversion.h"
 #include "USDValueConversion.h"
@@ -97,20 +98,6 @@
 
 namespace UsdLevelSequenceHelperImpl
 {
-	// Adapted from ObjectTools as it is within an Editor-only module
-	FString SanitizeObjectName(const FString& InObjectName)
-	{
-		FString SanitizedText = InObjectName;
-		const TCHAR* InvalidChar = INVALID_OBJECTNAME_CHARACTERS;
-		while (*InvalidChar)
-		{
-			SanitizedText.ReplaceCharInline(*InvalidChar, TCHAR('_'), ESearchCase::CaseSensitive);
-			++InvalidChar;
-		}
-
-		return SanitizedText;
-	}
-
 	/**
 	 * Similar to FrameRate.AsFrameNumber(TimeSeconds) except that it uses RoundToDouble instead of FloorToDouble, to
 	 * prevent issues with floating point precision
@@ -1093,8 +1080,12 @@ ULevelSequence* FUsdLevelSequenceHelperImpl::FindOrAddSequenceForLayer(const UE:
 		// which also gives us a stage actor. So if we don't have an actor but have a cache, we're importing
 		const bool bIsImporting = StageActor.IsExplicitlyNull() && InfoCache;
 		FName UniqueSequenceName = bIsImporting
-			? *UsdLevelSequenceHelperImpl::SanitizeObjectName(FPaths::GetBaseFilename(SequenceDisplayName))
-			: MakeUniqueObjectName(GetTransientPackage(), ULevelSequence::StaticClass(), *UsdLevelSequenceHelperImpl::SanitizeObjectName(FPaths::GetBaseFilename(SequenceDisplayName)));
+			? *IUsdClassesModule::SanitizeObjectName(FPaths::GetBaseFilename(SequenceDisplayName))
+			: MakeUniqueObjectName(
+				GetTransientPackage(),
+				ULevelSequence::StaticClass(),
+				*IUsdClassesModule::SanitizeObjectName(FPaths::GetBaseFilename(SequenceDisplayName))
+			);
 
 		Sequence = NewObject< ULevelSequence >(GetTransientPackage(), UniqueSequenceName, FUsdLevelSequenceHelperImpl::DefaultObjFlags);
 		Sequence->Initialize();
diff --git a/Engine/Plugins/Importers/USDImporter/Source/USDStageImporter/Private/USDStageImporter.cpp b/Engine/Plugins/Importers/USDImporter/Source/USDStageImporter/Private/USDStageImporter.cpp
--- a/Engine/Plugins/Importers/USDImporter/Source/USDStageImporter/Private/USDStageImporter.cpp
+++ b/Engine/Plugins/Importers/USDImporter/Source/USDStageImporter/Private/USDStageImporter.cpp
@@ -49,7 +49,6 @@
 #include "Misc/Paths.h"
 #include "Misc/ScopedSlowTask.h"
 #include "Modules/ModuleManager.h"
-#include "ObjectTools.h"
 #include "PackageTools.h"
 #include "PhysicsEngine/PhysicsAsset.h"
 #include "Serialization/ArchiveReplaceObjectRef.h"
@@ -190,7 +189,7 @@ namespace UsdStageImporterImpl
 		// We always spawn another scene actor regardless of collision or whether the level already has one,
 		// so that we can fully build our hierarchy separately before resolving collisions according to ExistingActorPolicy
 		AActor* Actor = ImportContext.World->SpawnActor(AActor::StaticClass(), nullptr, SpawnParameters);
-		Actor->SetActorLabel(ObjectTools::SanitizeObjectName(ImportContext.ObjectName));
+		Actor->SetActorLabel(IUsdClassesModule::SanitizeObjectName(ImportContext.ObjectName));
 
 		USceneComponent* RootComponent = Actor->GetRootComponent();
 		if (!RootComponent)
@@ -221,7 +220,7 @@ namespace UsdStageImporterImpl
 	AActor* GetExistingSceneActor(FUsdStageImportContext& ImportContext)
 	{
 		// We always reuse the existing scene actor for a scene, regardless of ReplacePolicy
-		FString TargetActorLabel = ObjectTools::SanitizeObjectName(ImportContext.ObjectName);
+		FString TargetActorLabel = IUsdClassesModule::SanitizeObjectName(ImportContext.ObjectName);
 		for (TActorIterator<AActor> ActorItr(ImportContext.World); ActorItr; ++ActorItr)
 		{
 			AActor* ThisActor = *ActorItr;
@@ -559,7 +558,7 @@ namespace UsdStageImporterImpl
 		// We don't care if our assets overwrite something in the final destination package (that conflict will be
 		// handled according to EReplaceAssetPolicy). But we do want these assets to have unique names amongst themselves
 		// or else they will overwrite each other when publishing
-		AssetName = UsdUtils::GetUniqueName( ObjectTools::SanitizeObjectName( AssetName ), UniqueAssetNames );
+		AssetName = UsdUtils::GetUniqueName(IUsdClassesModule::SanitizeObjectName(AssetName), UniqueAssetNames);
 		UniqueAssetNames.Add(AssetName);
 
 		return AssetName;
diff --git a/Engine/Plugins/Importers/USDImporter/Source/USDUtilities/Private/USDConversionUtils.cpp b/Engine/Plugins/Importers/USDImporter/Source/USDUtilities/Private/USDConversionUtils.cpp
--- a/Engine/Plugins/Importers/USDImporter/Source/USDUtilities/Private/USDConversionUtils.cpp
+++ b/Engine/Plugins/Importers/USDImporter/Source/USDUtilities/Private/USDConversionUtils.cpp
@@ -2,8 +2,8 @@
 
 #include "USDConversionUtils.h"
 
-#include "PhysicsEngine/PhysicsAsset.h"
 #include "USDAssetImportData.h"
+#include "USDClassesModule.h"
 #include "USDDuplicateType.h"
 #include "USDErrorUtils.h"
 #include "USDGeomMeshConversion.h"
@@ -46,12 +46,9 @@
 #include "InstancedFoliageActor.h"
 #include "LandscapeProxy.h"
 #include "Misc/PackageName.h"
+#include "PhysicsEngine/PhysicsAsset.h"
 #include "Widgets/Notifications/SNotificationList.h"
 
-#if WITH_EDITOR
-#include "ObjectTools.h"
-#endif // WITH_EDITOR
-
 #if USE_USD_SDK
 #include "USDIncludesStart.h"
 	#include "pxr/base/tf/stringUtils.h"
@@ -92,20 +89,6 @@
 
 namespace USDConversionUtilsImpl
 {
-	// Adapted from ObjectTools as it is within an Editor-only module
-	FString SanitizeObjectName( const FString& InObjectName )
-	{
-		FString SanitizedText = InObjectName;
-		const TCHAR* InvalidChar = INVALID_OBJECTNAME_CHARACTERS;
-		while ( *InvalidChar )
-		{
-			SanitizedText.ReplaceCharInline( *InvalidChar, TCHAR( '_' ), ESearchCase::CaseSensitive );
-			++InvalidChar;
-		}
-
-		return SanitizedText;
-	}
-
 	/** Show some warnings if the UVSet primvars show some unsupported/problematic behavior */
 	void CheckUVSetPrimvars( TMap<int32, TArray<pxr::UsdGeomPrimvar>> UsablePrimvars, TMap<int32, TArray<pxr::UsdGeomPrimvar>> UsedPrimvars, const FString& MeshPath )
 	{
@@ -1114,7 +1097,7 @@ FString UsdUtils::GetAssetPathFromPrimPath( const FString& RootContentPath, cons
 	ModelApi.GetAssetName( &RawAssetName );
 
 	FString AssetName = UsdToUnreal::ConvertString( RawAssetName );
-	FString MeshName = USDConversionUtilsImpl::SanitizeObjectName( RawPrimName );
+	FString MeshName = IUsdClassesModule::SanitizeObjectName(RawPrimName);
 
 	FString USDPath = UsdToUnreal::ConvertString( Prim.GetPrimPath().GetString().c_str() );
 
diff --git a/Engine/Plugins/Importers/USDImporter/Source/USDUtilities/Private/USDShadeConversion.cpp b/Engine/Plugins/Importers/USDImporter/Source/USDUtilities/Private/USDShadeConversion.cpp
--- a/Engine/Plugins/Importers/USDImporter/Source/USDUtilities/Private/USDShadeConversion.cpp
+++ b/Engine/Plugins/Importers/USDImporter/Source/USDUtilities/Private/USDShadeConversion.cpp
@@ -9,6 +9,7 @@
 #include "USDAssetCache2.h"
 #include "USDAssetImportData.h"
 #include "USDAttributeUtils.h"
+#include "USDClassesModule.h"
 #include "USDConversionUtils.h"
 #include "USDErrorUtils.h"
 #include "USDLayerUtils.h"
@@ -46,7 +47,6 @@
 	#include "MaterialEditingLibrary.h"
 	#include "MaterialOptions.h"
 	#include "MaterialUtilities.h"
-	#include "ObjectTools.h"
 #endif // WITH_EDITOR
 
 #include "USDIncludesStart.h"
@@ -224,7 +224,11 @@ namespace UE
 
 				FScopedUnrealAllocs UEAllocs;
 
-				FName TextureName = MakeUniqueObjectName( Outer, UTexture::StaticClass(), *FPaths::GetBaseFilename( ResolvedTexturePath ) );
+				FName TextureName = MakeUniqueObjectName(
+					Outer,
+					UTexture::StaticClass(),
+					*IUsdClassesModule::SanitizeObjectName(FPaths::GetBaseFilename(ResolvedTexturePath))
+				);
 
 				return Cast<UTexture>(TextureFactory->FactoryCreateBinary(
 					UTexture::StaticClass(),
@@ -1213,7 +1217,13 @@ namespace UE
 						// Not inside an USDZ archive, just a regular texture
 						else
 						{
-							FName TextureName = MakeUniqueObjectName( Outer, UTexture::StaticClass(), *FPaths::GetBaseFilename( ResolvedTexturePath ) );
+							FName TextureName = MakeUniqueObjectName(
+								Outer,
+								UTexture::StaticClass(),
+								*IUsdClassesModule::SanitizeObjectName(
+									FPaths::GetBaseFilename(ResolvedTexturePath)
+								)
+							);
 							Texture = Cast< UTexture >( TextureFactory->ImportObject( UTexture::StaticClass(), Outer, TextureName, ObjectFlags, ResolvedTexturePath, TEXT(""), bOutCancelled));
 						}
 
@@ -1406,9 +1416,9 @@ namespace UE
 
 					// Final FilePath will be something like "C:/TexturesFolder/Game_ContentFolder_Materials_Red_BaseColor.png", which automatically guarantees
 					// it won't overwrite another texture from the same export, but will overwrite old textures from previous exports
-					FString TextureFileName = ObjectTools::SanitizeObjectName( FPaths::ChangeExtension( TextureNamePrefix, TEXT( "" ) ) );
+					FString TextureFileName = IUsdClassesModule::SanitizeObjectName( FPaths::ChangeExtension( TextureNamePrefix, TEXT( "" ) ) );
 					TextureFileName.RemoveFromStart( TEXT( "_" ) );
-					FString TextureFilePath = FPaths::Combine( TexturesFolder.Path, FString::Printf( TEXT( "%s_%s.exr" ), *ObjectTools::SanitizeObjectName( TextureFileName ), *TrimmedPropertyName ) );
+					FString TextureFilePath = FPaths::Combine( TexturesFolder.Path, FString::Printf( TEXT( "%s_%s.exr" ), *IUsdClassesModule::SanitizeObjectName( TextureFileName ), *TrimmedPropertyName ) );
 
 					// For some reason the baked samples always have zero alpha and there is nothing we can do about it... It seems like the material baking module is made
 					// with the intent that the data ends up in UTexture2Ds, where they can be set to be compressed without alpha and have the value ignored.
@@ -1464,9 +1474,9 @@ namespace UE
 
 					// Final FilePath will be something like "C:/TexturesFolder/Game_ContentFolder_Materials_Red_BaseColor.png", which automatically guarantees
 					// it won't overwrite another texture from the same export, but will overwrite old textures from previous exports
-					FString TextureFileName = ObjectTools::SanitizeObjectName( FPaths::ChangeExtension( TextureNamePrefix, TEXT( "" ) ) );
+					FString TextureFileName = IUsdClassesModule::SanitizeObjectName( FPaths::ChangeExtension( TextureNamePrefix, TEXT( "" ) ) );
 					TextureFileName.RemoveFromStart( TEXT( "_" ) );
-					FString TextureFilePath = FPaths::Combine( TexturesFolder.Path, FString::Printf( TEXT( "%s_%s.png" ), *ObjectTools::SanitizeObjectName( TextureFileName ), *TrimmedPropertyName ) );
+					FString TextureFilePath = FPaths::Combine( TexturesFolder.Path, FString::Printf( TEXT( "%s_%s.png" ), *IUsdClassesModule::SanitizeObjectName( TextureFileName ), *TrimmedPropertyName ) );
 
 					// For some reason the baked samples always have zero alpha and there is nothing we can do about it... It seems like the material baking module is made
 					// with the intent that the data ends up in UTexture2Ds, where they can be set to be compressed without alpha and have the value ignored.
diff --git a/Engine/Plugins/Importers/USDImporter/Source/USDUtilities/Private/USDSkeletalDataConversion.cpp b/Engine/Plugins/Importers/USDImporter/Source/USDUtilities/Private/USDSkeletalDataConversion.cpp
--- a/Engine/Plugins/Importers/USDImporter/Source/USDUtilities/Private/USDSkeletalDataConversion.cpp
+++ b/Engine/Plugins/Importers/USDImporter/Source/USDUtilities/Private/USDSkeletalDataConversion.cpp
@@ -91,20 +91,6 @@ static FAutoConsoleVariableRef CVarAddCurveMetadataToSkeleton(
 #if USE_USD_SDK && WITH_EDITOR
 namespace SkelDataConversionImpl
 {
-	// Adapted from ObjectTools as it is within an Editor-only module
-	FString SanitizeObjectName( const FString& InObjectName )
-	{
-		FString SanitizedText = InObjectName;
-		const TCHAR* InvalidChar = INVALID_OBJECTNAME_CHARACTERS;
-		while ( *InvalidChar )
-		{
-			SanitizedText.ReplaceCharInline( *InvalidChar, TCHAR( '_' ), ESearchCase::CaseSensitive );
-			++InvalidChar;
-		}
-
-		return SanitizedText;
-	}
-
 	// Adapted from LODUtilities.cpp
 	struct FMeshDataBundle
 	{
@@ -1127,13 +1113,13 @@ namespace UnrealToUsdImpl
 
 		// The first bone is the root, and has ParentIndex == -1, so do it separately here to void checking the indices for all bones
 		// Sanitize because ExportName can have spaces, which USD doesn't like
-		OutFullPaths[0] = SkelDataConversionImpl::SanitizeObjectName(BoneNamesInOrder[0].ExportName);
+		OutFullPaths[0] = IUsdClassesModule::SanitizeObjectName(BoneNamesInOrder[0].ExportName);
 
 		// Bones are always stored in an increasing order, so we can do all paths in a single pass
 		for ( int32 BoneIndex = 1; BoneIndex < NumBones; ++BoneIndex )
 		{
 			const FMeshBoneInfo& BoneInfo = BoneNamesInOrder[ BoneIndex ];
-			FString SanitizedBoneName = SkelDataConversionImpl::SanitizeObjectName(BoneInfo.ExportName);
+			FString SanitizedBoneName = IUsdClassesModule::SanitizeObjectName(BoneInfo.ExportName);
 
 			OutFullPaths[BoneIndex] = FString::Printf(TEXT("%s/%s"), *OutFullPaths[ BoneInfo.ParentIndex ], *SanitizedBoneName );
 		}
@@ -2438,7 +2424,7 @@ bool UsdToUnreal::ConvertBlendShape( const pxr::UsdSkelBlendShape& UsdBlendShape
 	// because although the path is usually unique, USD has case sensitive paths and the FNames of the
 	// UMorphTargets are case insensitive
 	FString PrimaryName = UsdUtils::GetUniqueName(
-		SkelDataConversionImpl::SanitizeObjectName( UsdToUnreal::ConvertString( UsdBlendShape.GetPrim().GetName() ) ),
+		IUsdClassesModule::SanitizeObjectName( UsdToUnreal::ConvertString( UsdBlendShape.GetPrim().GetName() ) ),
 		UsedMorphTargetNames );
 	FString PrimaryPath = UsdToUnreal::ConvertPath( UsdBlendShape.GetPrim().GetPath() );
 	if ( UsdUtils::FUsdBlendShape* ExistingBlendShape = OutBlendShapes.Find( PrimaryPath ) )
@@ -2471,7 +2457,7 @@ bool UsdToUnreal::ConvertBlendShape( const pxr::UsdSkelBlendShape& UsdBlendShape
 		FString OrigInbetweenName = UsdToUnreal::ConvertString( Inbetween.GetAttr().GetName() );
 		FString InbetweenPath = FString::Printf(TEXT("%s_%s"), *PrimaryPath, *OrigInbetweenName );
 		FString InbetweenName = UsdUtils::GetUniqueName(
-			SkelDataConversionImpl::SanitizeObjectName( FPaths::GetCleanFilename( InbetweenPath ) ),
+			IUsdClassesModule::SanitizeObjectName( FPaths::GetCleanFilename( InbetweenPath ) ),
 			UsedMorphTargetNames );
 
 		if ( Weight > 1.0f || Weight < 0.0f || FMath::IsNearlyZero(Weight) || FMath::IsNearlyEqual(Weight, 1.0f) )
@@ -2546,7 +2532,11 @@ USkeletalMesh* UsdToUnreal::GetSkeletalMeshFromImportData(
 	// A SkeletalMesh could be retrieved for re-use and updated for animations
 	// For now, create a new USkeletalMesh
 	// Note: Remember to initialize UsedMorphTargetNames with existing morph targets, whenever the SkeletalMesh is reused
-	FName UniqueMeshName = MakeUniqueObjectName( GetTransientPackage(), USkeletalMesh::StaticClass(), MeshName );
+	FName UniqueMeshName = MakeUniqueObjectName(
+		GetTransientPackage(),
+		USkeletalMesh::StaticClass(),
+		*IUsdClassesModule::SanitizeObjectName(MeshName.ToString())
+	);
 	USkeletalMesh* SkeletalMesh = NewObject<USkeletalMesh>(GetTransientPackage(), UniqueMeshName, ObjectFlags | EObjectFlags::RF_Public | EObjectFlags::RF_Transient);
 
 	// Process reference skeleton from import data
@@ -2670,7 +2660,11 @@ USkeletalMesh* UsdToUnreal::GetSkeletalMeshFromImportData(
 	SkeletalMesh->CalculateInvRefMatrices();
 
 	// Generate a Skeleton and associate it to the SkeletalMesh
-	FName UniqueSkeletonName = MakeUniqueObjectName( GetTransientPackage(), USkeleton::StaticClass(), SkeletonName );
+	FName UniqueSkeletonName = MakeUniqueObjectName(
+		GetTransientPackage(),
+		USkeleton::StaticClass(),
+		*IUsdClassesModule::SanitizeObjectName(SkeletonName.ToString())
+	);
 	USkeleton* Skeleton = NewObject<USkeleton>( GetTransientPackage(), UniqueSkeletonName, ObjectFlags | EObjectFlags::RF_Public | EObjectFlags::RF_Transient );
 	Skeleton->MergeAllBonesToBoneTree(SkeletalMesh);
 	Skeleton->SetPreviewMesh(SkeletalMesh);
