USD: Fix crash when importing materials twice in a row.

diff --git a/Engine/Plugins/Importers/USDImporter/Source/USDClasses/Private/USDAssetCache2.cpp b/Engine/Plugins/Importers/USDImporter/Source/USDClasses/Private/USDAssetCache2.cpp
--- a/Engine/Plugins/Importers/USDImporter/Source/USDClasses/Private/USDAssetCache2.cpp
+++ b/Engine/Plugins/Importers/USDImporter/Source/USDClasses/Private/USDAssetCache2.cpp
@@ -358,7 +358,7 @@ namespace UE::AssetCache::Private
 				{
 					if (ObjectWithOuter != nullptr)
 					{
-						ObjectWithOuter->Rename(nullptr, GetTransientPackage(), REN_NonTransactional | REN_DontCreateRedirectors);
+						ObjectWithOuter->Rename(nullptr, GetTransientPackage(), REN_NonTransactional | REN_DontCreateRedirectors | REN_ForceNoResetLoaders);
 					}
 					ObjectWithOuter->MarkAsGarbage();
 				}
@@ -420,13 +420,13 @@ namespace UE::AssetCache::Private
 				if (NewOuter != SubObject->GetOuter())
 				{
 					const TCHAR* NewName = nullptr;
-					SubObject->Rename(NewName, NewOuter, REN_NonTransactional | REN_DontCreateRedirectors);
+					SubObject->Rename(NewName, NewOuter, REN_NonTransactional | REN_DontCreateRedirectors | REN_ForceNoResetLoaders);
 				}
 
 				UObject* NoNewOuter = nullptr;
 				if (SubObjectName != SubObject->GetName() && SubObject->Rename(*SubObjectName, NoNewOuter, REN_Test))
 				{
-					SubObject->Rename(*SubObjectName, NoNewOuter, REN_NonTransactional | REN_DontCreateRedirectors);
+					SubObject->Rename(*SubObjectName, NoNewOuter, REN_NonTransactional | REN_DontCreateRedirectors | REN_ForceNoResetLoaders);
 				}
 			}
 		}
@@ -606,7 +606,7 @@ void UUsdAssetCache2::CacheAsset(const FString& Hash, UObject* Asset, const UObj
 	{
 		UObject* NewOuter = this;
 		const FName NewName = MakeUniqueObjectName(NewOuter, Asset->GetClass(), Asset->GetFName());
-		Asset->Rename(*NewName.ToString(), NewOuter);
+		Asset->Rename(*NewName.ToString(), NewOuter, REN_DontCreateRedirectors | REN_NonTransactional | REN_ForceNoResetLoaders);
 	}
 
 	FCachedAssetInfo NewInfo;
@@ -1347,7 +1347,7 @@ bool UUsdAssetCache2::TryUnloadAsset(FCachedAssetInfo& InOutInfo)
 
 			const TCHAR* NewName = nullptr;
 			UObject* NewOuter = GetTransientPackage();
-			Asset->Rename(NewName, NewOuter);
+			Asset->Rename(NewName, NewOuter, REN_DontCreateRedirectors | REN_NonTransactional | REN_ForceNoResetLoaders);
 		}
 
 		return InOutInfo.BulkData.UnloadBulkData();
diff --git a/Engine/Plugins/Importers/USDImporter/Source/USDStageImporter/Private/USDStageImporter.cpp b/Engine/Plugins/Importers/USDImporter/Source/USDStageImporter/Private/USDStageImporter.cpp
--- a/Engine/Plugins/Importers/USDImporter/Source/USDStageImporter/Private/USDStageImporter.cpp
+++ b/Engine/Plugins/Importers/USDImporter/Source/USDStageImporter/Private/USDStageImporter.cpp
@@ -693,6 +693,23 @@ namespace UsdStageImporterImpl
 				}
 			}
 
+			// We have to rename the existing asset away, because some objects manage subobjects (like
+			// UMaterialInterface and its EditorOnlyData), and try to rename them whenever they're duplicated/renamed.
+			// That can lead to issues if the target names conflict with the names of ExistingAsset's subobjects			 
+			if (ExistingAsset)
+			{
+				FName UniqueName = MakeUniqueObjectName(
+					GetTransientPackage(),
+					ExistingAsset->GetClass(),
+					ExistingAsset->GetFName()
+				);
+				ExistingAsset->Rename(
+					*UniqueName.ToString(), 
+					GetTransientPackage(), 
+					REN_DontCreateRedirectors | REN_NonTransactional | REN_DoNotDirty
+				);
+			}
+
 			MovedAsset = DuplicateObject<UObject>( Asset, Package, *TargetAssetName );
 
 			// If our DuplicateObject didn't stomp the old asset because TargetAssetName != ExistingAsset->GetName(),
