USDImporter: add optional EUsdLoadPolicy parameter to FUsdPrim::Load()

This aligns the Load() function of the engine-side FUsdPrim wrapper with its
pxr UsdPrim counterpart, allowing the underlying prim to optionally be loaded
without loading any of its descendants.

The declaration of the engine-side EUsdLoadPolicy was moved from UsdStage.h
to UnrealUSDWrapper.h (which is already #include'd in UsdStage.h) to support this.

diff --git a/Engine/Plugins/Importers/USDImporter/Source/UnrealUSDWrapper/Private/UsdWrappers/UsdPrim.cpp b/Engine/Plugins/Importers/USDImporter/Source/UnrealUSDWrapper/Private/UsdWrappers/UsdPrim.cpp
--- a/Engine/Plugins/Importers/USDImporter/Source/UnrealUSDWrapper/Private/UsdWrappers/UsdPrim.cpp
+++ b/Engine/Plugins/Importers/USDImporter/Source/UnrealUSDWrapper/Private/UsdWrappers/UsdPrim.cpp
@@ -416,10 +416,13 @@ namespace UE
 #endif // #if USE_USD_SDK
 	}
 
-	void FUsdPrim::Load()
+	void FUsdPrim::Load( EUsdLoadPolicy Policy )
 	{
 #if USE_USD_SDK
-		Impl->PxrUsdPrim.Get().Load();
+		static_assert( ( int32 ) EUsdLoadPolicy::UsdLoadWithDescendants == ( int32 ) pxr::UsdLoadWithDescendants, "EUsdLoadPolicy enum doesn't match USD!" );
+		static_assert( ( int32 ) EUsdLoadPolicy::UsdLoadWithoutDescendants == ( int32 ) pxr::UsdLoadWithoutDescendants, "EUsdLoadPolicy enum doesn't match USD!" );
+
+		Impl->PxrUsdPrim.Get().Load( static_cast< pxr::UsdLoadPolicy > ( Policy ) );
 #endif // #if USE_USD_SDK
 	}
 
diff --git a/Engine/Plugins/Importers/USDImporter/Source/UnrealUSDWrapper/Public/UnrealUSDWrapper.h b/Engine/Plugins/Importers/USDImporter/Source/UnrealUSDWrapper/Public/UnrealUSDWrapper.h
--- a/Engine/Plugins/Importers/USDImporter/Source/UnrealUSDWrapper/Public/UnrealUSDWrapper.h
+++ b/Engine/Plugins/Importers/USDImporter/Source/UnrealUSDWrapper/Public/UnrealUSDWrapper.h
@@ -96,6 +96,14 @@ enum class EUsdDefaultKind : int32
 };
 ENUM_CLASS_FLAGS( EUsdDefaultKind );
 
+/** Corresponds to pxr::UsdLoadPolicy, refer to the USD SDK documentation */
+UENUM()
+enum class EUsdLoadPolicy : uint8
+{
+	UsdLoadWithDescendants,    // Load a prim plus all its descendants.
+	UsdLoadWithoutDescendants  // Load a prim by itself with no descendants.
+};
+
 UENUM()
 enum class EUsdInitialLoadSet : uint8
 {
diff --git a/Engine/Plugins/Importers/USDImporter/Source/UnrealUSDWrapper/Public/UsdWrappers/UsdPrim.h b/Engine/Plugins/Importers/USDImporter/Source/UnrealUSDWrapper/Public/UsdWrappers/UsdPrim.h
--- a/Engine/Plugins/Importers/USDImporter/Source/UnrealUSDWrapper/Public/UsdWrappers/UsdPrim.h
+++ b/Engine/Plugins/Importers/USDImporter/Source/UnrealUSDWrapper/Public/UsdWrappers/UsdPrim.h
@@ -6,6 +6,7 @@
 #include "Templates/UniquePtr.h"
 #include "UObject/NameTypes.h"
 
+#include "UnrealUSDWrapper.h"
 #include "UsdWrappers/ForwardDeclarations.h"
 
 #if USE_USD_SDK
@@ -108,7 +109,7 @@ namespace UE
 
 		bool HasPayload() const;
 		bool IsLoaded() const;
-		void Load();
+		void Load( EUsdLoadPolicy Policy = EUsdLoadPolicy::UsdLoadWithDescendants );
 		void Unload();
 
 		bool RemoveProperty( FName PropName ) const;
diff --git a/Engine/Plugins/Importers/USDImporter/Source/UnrealUSDWrapper/Public/UsdWrappers/UsdStage.h b/Engine/Plugins/Importers/USDImporter/Source/UnrealUSDWrapper/Public/UsdWrappers/UsdStage.h
--- a/Engine/Plugins/Importers/USDImporter/Source/UnrealUSDWrapper/Public/UsdWrappers/UsdStage.h
+++ b/Engine/Plugins/Importers/USDImporter/Source/UnrealUSDWrapper/Public/UsdWrappers/UsdStage.h
@@ -18,13 +18,6 @@ namespace UE
 		template< typename PtrType > class FUsdStageImpl;
 	}
 
-	/** Corresponds to pxr::UsdLoadPolicy, refer to the USD SDK documentation */
-	enum class EUsdLoadPolicy
-	{
-		UsdLoadWithDescendants,    // Load a prim plus all its descendants.
-		UsdLoadWithoutDescendants  // Load a prim by itself with no descendants.
-	};
-
 	/**
 	 * Minimal pxr::UsdStage pointer wrapper for Unreal that can be used from no-rtti modules.
 	 * Use the aliases FUsdStage and FUsdStageWeak instead (defined on ForwardDeclarations.h)
