USD: Fixes for UDIM texture handling.

diff --git a/Engine/Plugins/Importers/USDImporter/Source/USDUtilities/Private/USDShadeConversion.cpp b/Engine/Plugins/Importers/USDImporter/Source/USDUtilities/Private/USDShadeConversion.cpp
--- a/Engine/Plugins/Importers/USDImporter/Source/USDUtilities/Private/USDShadeConversion.cpp
+++ b/Engine/Plugins/Importers/USDImporter/Source/USDUtilities/Private/USDShadeConversion.cpp
@@ -131,22 +131,55 @@ namespace UE
 #endif // WITH_EDITOR
 
 			// Given an AssetPath, resolve it to an actual file path
-			FString ResolveAssetPath( const pxr::SdfLayerHandle& LayerHandle, const FString& AssetPath )
+			FString ResolveTexturePath( const pxr::SdfLayerHandle& LayerHandle, const FString& AssetPath )
 			{
+				// TODO: Most of this stuff is incompatible with custom resolvers, as these asset paths may be URLs
+				// or just GUIDs or anything like that, where relative vs absolute path make no sense. We will need
+				// a different approach whenever we properly handle USD resolvers
+
 				FScopedUsdAllocs UsdAllocs;
 
-				FString AssetPathToResolve = AssetPath;
-				AssetPathToResolve.ReplaceInline( TEXT("<UDIM>"), TEXT("1001") ); // If it's a UDIM path, we replace the UDIM tag with the start tile
+				FString PathToResolve = LayerHandle
+					? UsdToUnreal::ConvertString(pxr::SdfComputeAssetPathRelativeToLayer(LayerHandle, UnrealToUsd::ConvertString(*AssetPath).Get()))
+					: AssetPath;
 
-				pxr::ArResolverScopedCache ResolverCache;
-				pxr::ArResolver& Resolver = pxr::ArGetResolver();
+				// We need to provide absolute paths to the resolver later: It has no idea what to do with a relative
+				// path relative to some random location
+				if (FPaths::IsRelative(PathToResolve) && LayerHandle)
+				{
+					FString LayerDirectory = FPaths::GetPath(UsdToUnreal::ConvertString(LayerHandle->GetRealPath()));
+					PathToResolve = FPaths::Combine(LayerDirectory, PathToResolve);
+				}
+
+				// If this path has an UDIM placeholder in it (e.g. "textures/red.<UDIM>.exr"), let's try to find an
+				// actual existing UDIM tile instead, or else Resolver.Resolve will just give us the empty string.
+				// Previously we directly tried to find tile 1001, but it seems there is no guarantee that the
+				// particular 1001 tile will exist.
+				// Also note that for UDIMs the UE texture factory expects to receive path to any one tile, and it
+				// will itself discover the remaining tiles (need to also keep that in mind when doing the support
+				// for custom resolvers...)
+				if (PathToResolve.Contains(TEXT("<UDIM>")))
+				{
+					FString UdimFileFilter = PathToResolve.Replace(TEXT("<UDIM>"), TEXT("*"));
+
+					TArray<FString> UDIMFiles;
+					IFileManager::Get().FindFiles(UDIMFiles, *UdimFileFilter, /*files*/ true, /*directories*/ false);
+					if (UDIMFiles.Num() > 0)
+					{
+						// Sort here to enforce some sort of consistency between repeated calls
+						UDIMFiles.Sort();
 
-				const FString RelativePathToResolve =
-					LayerHandle
-					? UsdToUnreal::ConvertString( pxr::SdfComputeAssetPathRelativeToLayer( LayerHandle, UnrealToUsd::ConvertString( *AssetPathToResolve ).Get() ) )
-					: AssetPathToResolve;
+						// FindFiles will just return the filename with no path info, so lets put the file in the same
+						// location our original <UDIM> path was
+						PathToResolve = FPaths::Combine(FPaths::GetPath(PathToResolve), UDIMFiles[0]);
+					}
+				}
 
-				std::string AssetIdentifier = Resolver.Resolve( UnrealToUsd::ConvertString( *RelativePathToResolve ).Get().c_str() );
+				pxr::ArResolverScopedCache ResolverCache;
+				pxr::ArResolver& Resolver = pxr::ArGetResolver();
+				std::string AssetIdentifier = Resolver.Resolve(
+					UnrealToUsd::ConvertString(*PathToResolve).Get().c_str()
+				);
 
 				// Don't normalize an empty path as the result will be "."
 				if ( AssetIdentifier.size() > 0 )
@@ -155,7 +188,6 @@ namespace UE
 				}
 
 				FString ResolvedAssetPath = UsdToUnreal::ConvertString( AssetIdentifier );
-
 				return ResolvedAssetPath;
 			}
 
@@ -2421,12 +2453,11 @@ FString UsdUtils::GetResolvedTexturePath( const pxr::UsdAttribute& TextureAssetP
 	pxr::SdfAssetPath TextureAssetPath;
 	TextureAssetPathAttr.Get< pxr::SdfAssetPath >( &TextureAssetPath );
 
-	pxr::ArResolver& Resolver = pxr::ArGetResolver();
-
 	std::string AssetIdentifier = TextureAssetPath.GetResolvedPath();
 	// Don't normalize an empty path as the result will be "."
 	if ( AssetIdentifier.size() > 0 )
 	{
+		pxr::ArResolver& Resolver = pxr::ArGetResolver();
 		AssetIdentifier = Resolver.CreateIdentifier( AssetIdentifier );
 	}
 
@@ -2441,7 +2472,7 @@ FString UsdUtils::GetResolvedTexturePath( const pxr::UsdAttribute& TextureAssetP
 		if ( !TexturePath.IsEmpty() )
 		{
 			pxr::SdfLayerRefPtr TextureLayer = UsdUtils::FindLayerForAttribute( TextureAssetPathAttr, pxr::UsdTimeCode::EarliestTime().GetValue() );
-			ResolvedTexturePath = UsdShadeConversionImpl::ResolveAssetPath( TextureLayer, TexturePath );
+			ResolvedTexturePath = UsdShadeConversionImpl::ResolveTexturePath( TextureLayer, TexturePath );
 		}
 	}
 
