USD: Fix reference/payload paths added relative to binary causing issues.

diff --git a/Engine/Plugins/Importers/USDImporter/Source/USDStageEditor/Private/SUSDStageTreeView.cpp b/Engine/Plugins/Importers/USDImporter/Source/USDStageEditor/Private/SUSDStageTreeView.cpp
--- a/Engine/Plugins/Importers/USDImporter/Source/USDStageEditor/Private/SUSDStageTreeView.cpp
+++ b/Engine/Plugins/Importers/USDImporter/Source/USDStageEditor/Private/SUSDStageTreeView.cpp
@@ -1157,9 +1157,24 @@ void SUsdStageTreeView::OnAddReference()
 		FText::FromString( Referencer.GetPrimPath().GetString() )
 	) );
 
+	// AddReference expects absolute file paths, so let's try ensuring that
+	FString ReferencedLayerPath = OptionsPtr->TargetFile.FilePath;
+	if (OptionsPtr->bInternalReference)
+	{
+		ReferencedLayerPath = TEXT("");
+	}
+	else if (FPaths::IsRelative(ReferencedLayerPath))
+	{
+		FString AbsoluteLayerFromBinary = FPaths::ConvertRelativePathToFull(ReferencedLayerPath);
+		if (!FPaths::IsRelative(AbsoluteLayerFromBinary) && FPaths::FileExists(AbsoluteLayerFromBinary))
+		{
+			ReferencedLayerPath = AbsoluteLayerFromBinary;
+		}
+	}
+
 	UsdUtils::AddReference(
 		Referencer,
-		OptionsPtr->bInternalReference ? TEXT("") : *OptionsPtr->TargetFile.FilePath,
+		*ReferencedLayerPath,
 		OptionsPtr->bUseDefaultPrim ? UE::FSdfPath{} : UE::FSdfPath{ *OptionsPtr->TargetPrimPath },
 		OptionsPtr->TimeCodeOffset,
 		OptionsPtr->TimeCodeScale
@@ -1216,9 +1231,24 @@ void SUsdStageTreeView::OnAddPayload()
 		FText::FromString( Referencer.GetPrimPath().GetString() )
 	) );
 
+	// AddPayload expects absolute file paths too, so let's try ensuring that
+	FString ReferencedLayerPath = OptionsPtr->TargetFile.FilePath;
+	if (OptionsPtr->bInternalReference)
+	{
+		ReferencedLayerPath = TEXT("");
+	}
+	else if (FPaths::IsRelative(ReferencedLayerPath))
+	{
+		FString AbsoluteLayerFromBinary = FPaths::ConvertRelativePathToFull(ReferencedLayerPath);
+		if (!FPaths::IsRelative(AbsoluteLayerFromBinary) && FPaths::FileExists(AbsoluteLayerFromBinary))
+		{
+			ReferencedLayerPath = AbsoluteLayerFromBinary;
+		}
+	}
+
 	UsdUtils::AddPayload(
 		Referencer,
-		OptionsPtr->bInternalReference ? TEXT( "" ) : *OptionsPtr->TargetFile.FilePath,
+		*ReferencedLayerPath,
 		OptionsPtr->bUseDefaultPrim ? UE::FSdfPath{} : UE::FSdfPath{ *OptionsPtr->TargetPrimPath },
 		OptionsPtr->TimeCodeOffset,
 		OptionsPtr->TimeCodeScale
