USD: Fix missing material overrides on components spawned for collapsed subtrees.

diff --git a/Engine/Plugins/Importers/USDImporter/Source/USDSchemas/Private/MeshTranslationImpl.cpp b/Engine/Plugins/Importers/USDImporter/Source/USDSchemas/Private/MeshTranslationImpl.cpp
--- a/Engine/Plugins/Importers/USDImporter/Source/USDSchemas/Private/MeshTranslationImpl.cpp
+++ b/Engine/Plugins/Importers/USDImporter/Source/USDSchemas/Private/MeshTranslationImpl.cpp
@@ -259,11 +259,6 @@ void MeshTranslationImpl::SetMaterialOverrides(
 {
 	FScopedUsdAllocs Allocs;
 
-	pxr::UsdGeomMesh Mesh{Prim};
-	if (!Mesh)
-	{
-		return;
-	}
 	pxr::SdfPath PrimPath = Prim.GetPath();
 	pxr::UsdStageRefPtr Stage = Prim.GetStage();
 
@@ -289,6 +284,11 @@ void MeshTranslationImpl::SetMaterialOverrides(
 		TMap<int32, UsdUtils::FUsdPrimMaterialAssignmentInfo> LODIndexToAssignmentsMap;
 		TFunction<bool(const pxr::UsdGeomMesh&, int32)> IterateLODs = [&](const pxr::UsdGeomMesh& LODMesh, int32 LODIndex)
 		{
+			// In here we need to parse the assignments again and can't rely on the cache because the info cache
+			// only has info about the default variant selection state of the stage: It won't have info about the
+			// LOD variant set setups as that involves actively toggling variants.
+			// TODO: Make the cache rebuild collect this info. Right now is not a good time for this as that would
+			// break the parallel-for setup that that function has
 			UsdUtils::FUsdPrimMaterialAssignmentInfo LODInfo = UsdUtils::GetPrimMaterialAssignments(
 				LODMesh.GetPrim(),
 				pxr::UsdTimeCode(Time),
@@ -319,15 +319,29 @@ void MeshTranslationImpl::SetMaterialOverrides(
 	// Extract material assignment info from prim if its *not* a LOD mesh, or if we failed to parse LODs
 	if (!bInterpretedLODs)
 	{
-		LODIndexToAssignments = {
-			UsdUtils::GetPrimMaterialAssignments(
-				ValidPrim,
-				pxr::UsdTimeCode(Time),
-				bProvideMaterialIndices,
-				RenderContextToken,
-				MaterialPurposeToken
-			)
-		};
+		// Try to pull the material slot info from the info cache first, which is useful if ValidPrim is a collapsed
+		// prim subtree: Querying it's material assignments directly is likely not what we want, as ValidPrim is
+		// likely just some root Xform prim.
+		// Note: This only works because we'll rebuild the cache when our material purpose/render context changes,
+		// and because in USD relationships (and so material bindings) can't vary with time
+		TOptional<TArray<UsdUtils::FUsdPrimMaterialSlot>> SubtreeSlots = InfoCache.GetSubtreeMaterialSlots(UE::FSdfPath{PrimPath});
+		if (SubtreeSlots.IsSet())
+		{
+			UsdUtils::FUsdPrimMaterialAssignmentInfo& NewInfo = LODIndexToAssignments.Emplace_GetRef();
+			NewInfo.Slots = MoveTemp(SubtreeSlots.GetValue());
+		}
+		else
+		{
+			LODIndexToAssignments = {
+				UsdUtils::GetPrimMaterialAssignments(
+					ValidPrim,
+					pxr::UsdTimeCode(Time),
+					bProvideMaterialIndices,
+					RenderContextToken,
+					MaterialPurposeToken
+				)
+			};
+		}
 	}
 
 	TMap<const UsdUtils::FUsdPrimMaterialSlot*, UMaterialInterface*> ResolvedMaterials = MeshTranslationImpl::ResolveMaterialAssignmentInfo(
diff --git a/Engine/Plugins/Importers/USDImporter/Source/USDSchemas/Private/USDGeomPointInstancerTranslator.cpp b/Engine/Plugins/Importers/USDImporter/Source/USDSchemas/Private/USDGeomPointInstancerTranslator.cpp
--- a/Engine/Plugins/Importers/USDImporter/Source/USDSchemas/Private/USDGeomPointInstancerTranslator.cpp
+++ b/Engine/Plugins/Importers/USDImporter/Source/USDSchemas/Private/USDGeomPointInstancerTranslator.cpp
@@ -99,7 +99,6 @@ void FUsdGeomPointInstancerCreateAssetsTaskChain::SetupTasks()
 	Do( ESchemaTranslationLaunchPolicy::Async,
 		[this]() -> bool
 		{
-			// TODO: Restore support for LOD prototypes
 			LODIndexToMeshDescription.Reset( 1 );
 			LODIndexToMaterialInfo.Reset( 1 );
 
diff --git a/Engine/Plugins/Importers/USDImporter/Source/USDSchemas/Private/USDGeomXformableTranslator.cpp b/Engine/Plugins/Importers/USDImporter/Source/USDSchemas/Private/USDGeomXformableTranslator.cpp
--- a/Engine/Plugins/Importers/USDImporter/Source/USDSchemas/Private/USDGeomXformableTranslator.cpp
+++ b/Engine/Plugins/Importers/USDImporter/Source/USDSchemas/Private/USDGeomXformableTranslator.cpp
@@ -2,6 +2,7 @@
 
 #include "USDGeomXformableTranslator.h"
 
+#include "MeshTranslationImpl.h"
 #include "UnrealUSDWrapper.h"
 #include "USDAssetCache.h"
 #include "USDClassesModule.h"
@@ -282,12 +283,49 @@ FUsdGeomXformableTranslator::FUsdGeomXformableTranslator( TSubclassOf< USceneCom
 
 USceneComponent* FUsdGeomXformableTranslator::CreateComponents()
 {
-	USceneComponent* Component = CreateComponentsEx( {}, {} );
+	USceneComponent* SceneComponent = CreateComponentsEx( {}, {} );
 
 	// We pulled UpdateComponents outside CreateComponentsEx as in some cases we don't want to do it
 	// right away (like on FUsdGeomPointInstancerTranslator::CreateComponents)
-	UpdateComponents( Component );
-	return Component;
+	UpdateComponents(SceneComponent);
+
+	// Handle material overrides for collapsed meshes. This can happen if we have two separate subtrees that collapse
+	// the same: A single static mesh will be shared between them and one of the task chains will manage to put their
+	// material assignments on the mesh directly. To ensure the correct materials for the second subtree, we need to
+	// set overrides.
+	// Note how we don't have to handle geometry caches in here like the geom mesh translator does though: We don't
+	// collapse geometry caches, at least for now
+	if (UStaticMeshComponent* StaticMeshComponent = Cast<UStaticMeshComponent>(SceneComponent))
+	{
+		if (Context->InfoCache)
+		{
+			if (UStaticMesh* StaticMesh = Context->InfoCache->GetSingleAssetForPrim<UStaticMesh>(
+				PrimPath
+			))
+			{
+				TArray<UMaterialInterface*> ExistingAssignments;
+				for (FStaticMaterial& StaticMaterial : StaticMesh->GetStaticMaterials())
+				{
+					ExistingAssignments.Add(StaticMaterial.MaterialInterface);
+				}
+
+				MeshTranslationImpl::SetMaterialOverrides(
+					GetPrim(),
+					ExistingAssignments,
+					*StaticMeshComponent,
+					*Context->AssetCache.Get(),
+					*Context->InfoCache.Get(),
+					Context->Time,
+					Context->ObjectFlags,
+					Context->bAllowInterpretingLODs,
+					Context->RenderContext,
+					Context->MaterialPurpose
+				);
+			}
+		}
+	}
+
+	return SceneComponent;
 }
 
 USceneComponent* FUsdGeomXformableTranslator::CreateComponentsEx( TOptional< TSubclassOf< USceneComponent > > ComponentType, TOptional< bool > bNeedsActor )
diff --git a/Engine/Plugins/Importers/USDImporter/Source/USDSchemas/Private/USDInfoCache.cpp b/Engine/Plugins/Importers/USDImporter/Source/USDSchemas/Private/USDInfoCache.cpp
--- a/Engine/Plugins/Importers/USDImporter/Source/USDSchemas/Private/USDInfoCache.cpp
+++ b/Engine/Plugins/Importers/USDImporter/Source/USDSchemas/Private/USDInfoCache.cpp
@@ -39,7 +39,7 @@ namespace UE::UsdInfoCache::Private
 		UE::FSdfPath AssetCollapsedRoot;
 		UE::FSdfPath ComponentCollapsedRoot;
 		uint64 ExpectedVertexCountForSubtree;
-		uint64 ExpectedMaterialSlotCountForSubtree;
+		TArray<UsdUtils::FUsdPrimMaterialSlot> SubtreeMaterialSlots;
 	};
 }
 
@@ -48,7 +48,7 @@ FArchive& operator <<( FArchive& Ar, UE::UsdInfoCache::Private::FUsdPrimInfo& In
 	Ar << Info.AssetCollapsedRoot;
 	Ar << Info.ComponentCollapsedRoot;
 	Ar << Info.ExpectedVertexCountForSubtree;
-	Ar << Info.ExpectedMaterialSlotCountForSubtree;
+	Ar << Info.SubtreeMaterialSlots;
 
 	return Ar;
 }
@@ -267,6 +267,24 @@ namespace UE::USDInfoCacheImpl::Private
 					{
 						const pxr::SdfPath& PrototypePath = PrototypePaths[ PrototypeIndex ];
 
+						// Skip invisible prototypes here to mirror how they're skipped within
+						// USDGeomMeshConversion.cpp, in the RecursivelyCollapseChildMeshes function. Those two
+						// traversals have to match at least with respect to the material slots, so that we can use
+						// the data collected here to apply material overrides to the meshes generated for the point
+						// instancers when they're collapsed
+						pxr::UsdPrim PrototypePrim = UsdPrim.GetStage()->GetPrimAtPath(PrototypePath);
+						if (pxr::UsdGeomImageable UsdGeomImageable = pxr::UsdGeomImageable(PrototypePrim))
+						{
+							if (pxr::UsdAttribute VisibilityAttr = UsdGeomImageable.GetVisibilityAttr())
+							{
+								pxr::TfToken VisibilityToken;
+								if (VisibilityAttr.Get(&VisibilityToken) && VisibilityToken == pxr::UsdGeomTokens->invisible)
+								{
+									continue;
+								}
+							}
+						}
+
 						// If we're calling this for a point instancer we should have parsed the results for our
 						// prototype subtrees already
 						if ( const UE::UsdInfoCache::Private::FUsdPrimInfo* FoundInfo = Impl.InfoMap.Find( UE::FSdfPath{ PrototypePath } ) )
@@ -388,12 +406,34 @@ namespace UE::USDInfoCacheImpl::Private
 		OutSubtreeVertexCount = 0;
 		OutSubtreeSlots.Empty();
 
+		// We will still step into invisible prims to collect all info we can, but we won't count their material slots
+		// or vertex counts: The main usage of those counts is to handle collapsed meshes, and during collapse we just
+		// early out whenever we encounter an invisible prim
+		bool bPrimIsInvisible = false;
+		if (pxr::UsdGeomImageable UsdGeomImageable = pxr::UsdGeomImageable(UsdPrim))
+		{
+			if (pxr::UsdAttribute VisibilityAttr = UsdGeomImageable.GetVisibilityAttr())
+			{
+				pxr::TfToken VisibilityToken;
+				if (VisibilityAttr.Get(&VisibilityToken) && VisibilityToken == pxr::UsdGeomTokens->invisible)
+				{
+					bPrimIsInvisible = true;
+				}
+			}
+		}
+
+		// If the mesh prim has an unselected geometry purpose, it is also essentially invisible
+		if (!EnumHasAllFlags(Context.PurposesToLoad, IUsdPrim::GetPurpose(UsdPrim)))
+		{
+			bPrimIsInvisible = true;
+		}
+
 		bool bIsPointInstancer = false;
 		if ( pxr::UsdGeomPointInstancer PointInstancer{ UsdPrim } )
 		{
 			bIsPointInstancer = true;
 		}
-		else
+		else if (!bPrimIsInvisible)
 		{
 			GetPrimVertexCountAndSlots(
 				UsdPrim,
@@ -571,7 +611,7 @@ namespace UE::USDInfoCacheImpl::Private
 			{
 				TSet<UsdUtils::FUsdPrimMaterialSlot> SlotsSet{ Pair.Value };
 				UE::UsdInfoCache::Private::FUsdPrimInfo& Info = Impl.InfoMap.FindOrAdd( Pair.Key );
-				Info.ExpectedMaterialSlotCountForSubtree = SlotsSet.Num();
+				Info.SubtreeMaterialSlots = SlotsSet.Array();
 			}
 		}
 		else
@@ -579,7 +619,7 @@ namespace UE::USDInfoCacheImpl::Private
 			for ( const TPair< UE::FSdfPath, TArray<UsdUtils::FUsdPrimMaterialSlot>>& Pair : SubtreeMaterialSlots )
 			{
 				UE::UsdInfoCache::Private::FUsdPrimInfo& Info = Impl.InfoMap.FindOrAdd( Pair.Key );
-				Info.ExpectedMaterialSlotCountForSubtree = Pair.Value.Num();
+				Info.SubtreeMaterialSlots = Pair.Value;
 			}
 		}
 	}
@@ -727,7 +767,7 @@ TOptional<uint64> FUsdInfoCache::GetSubtreeMaterialSlotCount( const UE::FSdfPath
 
 		if ( const UE::UsdInfoCache::Private::FUsdPrimInfo* FoundInfo = ImplPtr->InfoMap.Find( Path ) )
 		{
-			return FoundInfo->ExpectedMaterialSlotCountForSubtree;
+			return FoundInfo->SubtreeMaterialSlots.Num();
 		}
 
 		// This should never happen: We should have cached the entire tree
@@ -737,6 +777,24 @@ TOptional<uint64> FUsdInfoCache::GetSubtreeMaterialSlotCount( const UE::FSdfPath
 	return {};
 }
 
+TOptional<TArray<UsdUtils::FUsdPrimMaterialSlot>> FUsdInfoCache::GetSubtreeMaterialSlots(const UE::FSdfPath& Path)
+{
+	if (FUsdInfoCacheImpl* ImplPtr = Impl.Get())
+	{
+		FReadScopeLock ScopeLock(ImplPtr->InfoMapLock);
+
+		if (const UE::UsdInfoCache::Private::FUsdPrimInfo* FoundInfo = ImplPtr->InfoMap.Find(Path))
+		{
+			return FoundInfo->SubtreeMaterialSlots;
+		}
+
+		// This should never happen: We should have cached the entire tree
+		ensureMsgf(false, TEXT("Prim path '%s' has not been cached!"), *Path.GetString());
+	}
+
+	return {};
+}
+
 void FUsdInfoCache::LinkAssetToPrim(const UE::FSdfPath& Path, UObject* Asset)
 {
 	if (!Asset)
diff --git a/Engine/Plugins/Importers/USDImporter/Source/USDSchemas/Public/USDInfoCache.h b/Engine/Plugins/Importers/USDImporter/Source/USDSchemas/Public/USDInfoCache.h
--- a/Engine/Plugins/Importers/USDImporter/Source/USDSchemas/Public/USDInfoCache.h
+++ b/Engine/Plugins/Importers/USDImporter/Source/USDSchemas/Public/USDInfoCache.h
@@ -21,6 +21,11 @@ enum class ECollapsingType
 	Components
 };
 
+namespace UsdUtils
+{
+	struct FUsdPrimMaterialSlot;
+}
+
 /**
  * Caches information about a specific USD Stage
  */
@@ -59,6 +64,7 @@ public:
 	// provide the total sum of vertex counts of its entire subtree when queried
 	TOptional<uint64> GetSubtreeVertexCount(const UE::FSdfPath& Path);
 	TOptional<uint64> GetSubtreeMaterialSlotCount(const UE::FSdfPath& Path);
+	TOptional<TArray<UsdUtils::FUsdPrimMaterialSlot>> GetSubtreeMaterialSlots(const UE::FSdfPath& Path);
 
 public:
 	void LinkAssetToPrim(const UE::FSdfPath& Path, UObject* Asset);
diff --git a/Engine/Plugins/Importers/USDImporter/Source/USDStageEditor/Private/SUSDStage.cpp b/Engine/Plugins/Importers/USDImporter/Source/USDStageEditor/Private/SUSDStage.cpp
--- a/Engine/Plugins/Importers/USDImporter/Source/USDStageEditor/Private/SUSDStage.cpp
+++ b/Engine/Plugins/Importers/USDImporter/Source/USDStageEditor/Private/SUSDStage.cpp
@@ -1435,6 +1435,12 @@ void SUsdStage::FillMaterialPurposeSubMenu( FMenuBuilder& MenuBuilder )
 			{
 				if ( AUsdStageActor* StageActor = GetStageActorOrCDO() )
 				{
+					FScopedTransaction Transaction(FText::Format(
+						LOCTEXT("ChangeMaterialPurpose", "Change selected material purpose for USD stage actor '{0}' to '{1}'"),
+						FText::FromString(StageActor->GetActorLabel()),
+						FText::FromString(ChosenOption.IsValid() ? *ChosenOption : TEXT(""))
+					));
+
 					StageActor->SetMaterialPurpose( **ChosenOption );
 					if ( StageActor->IsTemplate() )
 					{
@@ -1489,6 +1495,12 @@ void SUsdStage::FillMaterialPurposeSubMenu( FMenuBuilder& MenuBuilder )
 
 				if ( AUsdStageActor* StageActor = GetStageActorOrCDO() )
 				{
+					FScopedTransaction Transaction(FText::Format(
+						LOCTEXT("TypeMaterialPurpose", "Change selected material purpose for USD stage actor '{0}' to '{1}'"),
+						FText::FromString(StageActor->GetActorLabel()),
+						NewText
+					));
+
 					StageActor->SetMaterialPurpose( NewPurpose );
 					if ( StageActor->IsTemplate() )
 					{
diff --git a/Engine/Plugins/Importers/USDImporter/Source/USDUtilities/Public/USDGeomMeshConversion.h b/Engine/Plugins/Importers/USDImporter/Source/USDUtilities/Public/USDGeomMeshConversion.h
--- a/Engine/Plugins/Importers/USDImporter/Source/USDUtilities/Public/USDGeomMeshConversion.h
+++ b/Engine/Plugins/Importers/USDImporter/Source/USDUtilities/Public/USDGeomMeshConversion.h
@@ -192,6 +192,7 @@ namespace UnrealToUsd
 	 */
 	USDUTILITIES_API bool ConvertMeshDescriptions( const TArray<FMeshDescription>& LODIndexToMeshDescription, pxr::UsdPrim& UsdPrim, const FMatrix& AdditionalTransform, const pxr::UsdTimeCode TimeCode = pxr::UsdTimeCode::Default() );
 }
+#endif // USE_USD_SDK
 
 namespace UsdUtils
 {
@@ -268,7 +269,11 @@ namespace UsdUtils
 		/** Describes the index of the Slot that each polygon/face of a mesh uses. Matches the order Slots */
 		TArray<int32> MaterialIndices;
 	};
+}
 
+#if USE_USD_SDK
+namespace UsdUtils
+{
 	/** Creates a FDisplayColorMaterial object describing the vertex color/opacity data from UsdMesh at time TimeCode */
 	USDUTILITIES_API TOptional<FDisplayColorMaterial> ExtractDisplayColorMaterial( const pxr::UsdGeomMesh& UsdMesh, const pxr::UsdTimeCode TimeCode = pxr::UsdTimeCode::EarliestTime() );
 
