USDImporter: add config setting and API for setting the default search path for USD's default resolver

diff --git a/Engine/Plugins/Importers/USDImporter/Source/USDClasses/Public/USDProjectSettings.h b/Engine/Plugins/Importers/USDImporter/Source/USDClasses/Public/USDProjectSettings.h
--- a/Engine/Plugins/Importers/USDImporter/Source/USDClasses/Public/USDProjectSettings.h
+++ b/Engine/Plugins/Importers/USDImporter/Source/USDClasses/Public/USDProjectSettings.h
@@ -39,6 +39,18 @@ public:
 	UPROPERTY( config, EditAnywhere, Category = USD, meta = (RelativeToGameContentDir) )
 	TArray<FDirectoryPath> AdditionalPluginDirectories;
 
+	// The directories that will be used as the default search path by USD's default resolver
+	// during asset resolution.
+	//
+	// Each directory in the search path should be an absolute path. If it is not, it will be
+	// anchored to the current working directory.
+	//
+	// Note that the default search path must be set before the first invocation of USD's
+	// resolver system, so changing this setting will require a restart of the engine in order
+	// for the new setting to take effect.
+	UPROPERTY( config, EditAnywhere, Category = USD, meta = (ConfigRestartRequired = true) )
+	TArray<FDirectoryPath> DefaultResolverSearchPath;
+
 	// Material purposes to show on drop-downs in addition to the standard "preview" and "full"
 	UPROPERTY( config, EditAnywhere, Category = USD )
 	TArray<FName> AdditionalMaterialPurposes;
diff --git a/Engine/Plugins/Importers/USDImporter/Source/UnrealUSDWrapper/Private/UnrealUSDWrapper.cpp b/Engine/Plugins/Importers/USDImporter/Source/UnrealUSDWrapper/Private/UnrealUSDWrapper.cpp
--- a/Engine/Plugins/Importers/USDImporter/Source/UnrealUSDWrapper/Private/UnrealUSDWrapper.cpp
+++ b/Engine/Plugins/Importers/USDImporter/Source/UnrealUSDWrapper/Private/UnrealUSDWrapper.cpp
@@ -856,6 +856,23 @@ void UnrealUSDWrapper::EraseStageFromCache( const UE::FUsdStage& Stage )
 #endif // #if USE_USD_SDK
 }
 
+void UnrealUSDWrapper::SetDefaultResolverDefaultSearchPath( const TArray<FDirectoryPath>& SearchPath )
+{
+#if USE_USD_SDK
+	FScopedUsdAllocs UsdAllocs;
+
+	std::vector< std::string > DefaultResolverSearchPath;
+	DefaultResolverSearchPath.reserve( SearchPath.Num() );
+
+	for ( const FDirectoryPath& Directory : SearchPath )
+	{
+		DefaultResolverSearchPath.push_back( TCHAR_TO_UTF8( *Directory.Path ) );
+	}
+
+	pxr::ArDefaultResolver::SetDefaultSearchPath( DefaultResolverSearchPath );
+#endif // #if USE_USD_SDK
+}
+
 void UnrealUSDWrapper::SetupDiagnosticDelegate()
 {
 #if USE_USD_SDK
@@ -970,6 +987,10 @@ public:
 			PlugRegistry::GetInstance().RegisterPlugins( UsdPluginDirectories );
 		}
 
+		// Set the default search path for USD's default resolver using any path specified in the settings.
+		const TArray<FDirectoryPath> SearchPath = GetDefault<UUsdProjectSettings>()->DefaultResolverSearchPath;
+		UnrealUSDWrapper::SetDefaultResolverDefaultSearchPath( SearchPath );
+
 #endif // USE_USD_SDK
 
 		FUsdMemoryManager::Initialize();
diff --git a/Engine/Plugins/Importers/USDImporter/Source/UnrealUSDWrapper/Public/UnrealUSDWrapper.h b/Engine/Plugins/Importers/USDImporter/Source/UnrealUSDWrapper/Public/UnrealUSDWrapper.h
--- a/Engine/Plugins/Importers/USDImporter/Source/UnrealUSDWrapper/Public/UnrealUSDWrapper.h
+++ b/Engine/Plugins/Importers/USDImporter/Source/UnrealUSDWrapper/Public/UnrealUSDWrapper.h
@@ -4,6 +4,7 @@
 
 #include "CoreTypes.h"
 
+#include "Engine/EngineTypes.h"
 #include "Modules/ModuleInterface.h"
 #include "Templates/Tuple.h"
 #include "UObject/ObjectMacros.h"
@@ -224,6 +225,16 @@ public:
 	/** Removes the stage from the stage cache. See UsdStageCache::Erase. */
 	UNREALUSDWRAPPER_API static void EraseStageFromCache( const UE::FUsdStage& Stage );
 
+	/**
+	 * Set the directories that will be used as the default search path by USD's default resolver during asset resolution.
+	 *
+	 * Each directory in the search path should be an absolute path. If it is not, it will be anchored to the current working directory.
+	 *
+	 * Note that the default search path must be set before the first invocation of USD's resolver system, so this function
+	 * must be called before that to have any effect.
+	 */
+	UNREALUSDWRAPPER_API static void SetDefaultResolverDefaultSearchPath( const TArray<FDirectoryPath>& SearchPath );
+
 	/** Starts listening to error/warning/log messages emitted by USD */
 	UNREALUSDWRAPPER_API static void SetupDiagnosticDelegate();
 
