USD: Expand asset cache to support persisting Materials, Skeletal Meshes and Static Meshes.

diff --git a/Engine/Plugins/Importers/USDImporter/Source/USDClasses/Private/USDAssetCache2.cpp b/Engine/Plugins/Importers/USDImporter/Source/USDClasses/Private/USDAssetCache2.cpp
--- a/Engine/Plugins/Importers/USDImporter/Source/USDClasses/Private/USDAssetCache2.cpp
+++ b/Engine/Plugins/Importers/USDImporter/Source/USDClasses/Private/USDAssetCache2.cpp
@@ -5,12 +5,21 @@
 #include "USDClassesModule.h"
 #include "USDLog.h"
 
+#include "Animation/Skeleton.h"
+#include "Engine/SkeletalMesh.h"
 #include "Engine/StaticMesh.h"
 #include "Engine/Texture2D.h"
+#include "Interfaces/ITargetPlatform.h"
+#include "Interfaces/ITargetPlatformManagerModule.h"
+#include "Materials/MaterialInterface.h"
+#include "MaterialShared.h"
+#include "MeshDescription.h"
 #include "Misc/ScopeRWLock.h"
+#include "PhysicsEngine/PhysicsAsset.h"
+#include "Serialization/ObjectAndNameAsStringProxyArchive.h"
 #include "Serialization/ObjectReader.h"
 #include "Serialization/ObjectWriter.h"
-#include "Serialization/ObjectAndNameAsStringProxyArchive.h"
+#include "StaticMeshResources.h"
 #include "UObject/ObjectRedirector.h"
 #include "UObject/Package.h"
 
@@ -50,8 +59,14 @@ namespace UE::AssetCache::Private
 		}
 		UClass* AssetClass = FindClass(ClassName);
 
-		// Just textures for now
-		static TArray<const UClass*> AllowedPersistentClasses = { UTexture2D::StaticClass() };
+		static TArray<const UClass*> AllowedPersistentClasses = {
+			UTexture2D::StaticClass(),
+			UMaterialInterface::StaticClass(),
+			UStaticMesh::StaticClass(),
+			UPhysicsAsset::StaticClass(),
+			USkeleton::StaticClass(),
+			USkeletalMesh::StaticClass()
+		};
 
 		for (const UClass* AllowedClass : AllowedPersistentClasses)
 		{
@@ -119,6 +134,42 @@ namespace UE::AssetCache::Private
 
 			return *this;
 		}
+
+		// Override FObjectAndNameAsStringProxyArchive::operator<< because we don't want to blindly call LoadObject
+		// on 'None' strings, which is what gets serialized in case we have a nullptr (check
+		// UObjectBaseUtility::GetPathName)
+		virtual FArchive& operator<<(UObject*& Obj) override
+		{
+			if (IsLoading())
+			{
+				// load the path name to the object
+				FString LoadedString;
+				InnerArchive << LoadedString;
+
+				const static FString NullptrString = TEXT("None");
+				if (LoadedString == NullptrString || LoadedString.IsEmpty())
+				{
+					Obj = nullptr;
+				}
+				else
+				{
+					// look up the object by fully qualified pathname
+					Obj = FindObject<UObject>(nullptr, *LoadedString, false);
+					// If we couldn't find it, and we want to load it, do that
+					if (!Obj && bLoadIfFindFails)
+					{
+						Obj = LoadObject<UObject>(nullptr, *LoadedString);
+					}
+				}
+			}
+			else
+			{
+				// save out the fully qualified object name
+				FString SavedString(Obj->GetPathName());
+				InnerArchive << SavedString;
+			}
+			return *this;
+		}
 	};
 
 	// The main archiver used to write UObjects to byte buffers
@@ -198,10 +249,12 @@ namespace UE::AssetCache::Private
 		Ar.SetLicenseeUEVer(DestArchive.LicenseeUEVer());
 		Ar.SetCustomVersions(DestArchive.GetCustomVersions());
 
+		ensure(!Object->IsEditorOnly() || !Ar.IsCooking());
+
 		// Collect sub-objects depending on input object including nested objects
 		TArray< UObject* > SubObjectsArray;
 		GetObjectsWithOuter(Object, SubObjectsArray, /*bIncludeNestedObjects = */ true);
-		TSet<UObject*> SubObjectsSet(SubObjectsArray);
+		TSet<UObject*> AllowedSubObjectsSet(SubObjectsArray);
 
 		// Keep track of which subobjects are actually referenced by our UObject tree.
 		// Other objects may just be hiding inside of them if they're referenced by an external UObject, and so
@@ -210,7 +263,7 @@ namespace UE::AssetCache::Private
 		// referenced by the transaction buffer);
 		TSet<UObject*> ActualDependencies;
 		{
-			FObjectDependencyAnalyzer Analyzer(Object, SubObjectsSet);
+			FObjectDependencyAnalyzer Analyzer(Object, AllowedSubObjectsSet);
 			Object->Serialize(Analyzer);
 			ActualDependencies.Append(Analyzer.DependentObjects);
 		}
@@ -221,37 +274,21 @@ namespace UE::AssetCache::Private
 			TMap< UObject*, TSet<UObject*> > SubObjectDependencyGraph;
 			SubObjectDependencyGraph.Reserve(SubObjectsArray.Num());
 
+			// Build graph of dependency: each entry contains the set of sub-objects to create before itself.
+			// Note that we're not using ActualDependencies for this because we want to map out
+			// all transitive dependencies as well, and FObjectDependencyAnalyzer is not recursive
 			for (UObject* SubObject : SubObjectsArray)
 			{
-				SubObjectDependencyGraph.Add(SubObject);
-			}
-
-			// Build graph of dependency: each entry contains the set of sub-objects to create before itself
-			for (UObject* SubObject : SubObjectsArray)
-			{
-				FObjectDependencyAnalyzer SubAnalyzer(SubObject, SubObjectsSet);
+				FObjectDependencyAnalyzer SubAnalyzer(SubObject, AllowedSubObjectsSet);
 				SubObject->Serialize(SubAnalyzer);
 
-				SubObjectDependencyGraph[SubObject].Append(SubAnalyzer.DependentObjects);
+				SubObjectDependencyGraph.FindOrAdd(SubObject).Append(SubAnalyzer.DependentObjects);
 				ActualDependencies.Append(SubAnalyzer.DependentObjects);
 			}
 
-			// Only keep the UObjects that we're actually referencing
-			TArray<UObject*> UsedSubObjectsArray;
-			UsedSubObjectsArray.Reserve(SubObjectsArray.Num());
-			for (UObject* SubObject : SubObjectsArray)
-			{
-				if (ActualDependencies.Contains(SubObject))
-				{
-					UsedSubObjectsArray.Add(SubObject);
-				}
-			}
-			Swap(UsedSubObjectsArray, SubObjectsArray);
-
 			// Sort array of sub-objects: first objects do not depend on ones below
-			int32 Count = SubObjectsArray.Num();
+			int32 Count = ActualDependencies.Num();
 			SubObjectsArray.Empty(Count);
-
 			while (Count != SubObjectsArray.Num())
 			{
 				for (auto& Entry : SubObjectDependencyGraph)
@@ -262,7 +299,10 @@ namespace UE::AssetCache::Private
 
 						SubObjectDependencyGraph.Remove(SubObject);
 
-						SubObjectsArray.Add(SubObject);
+						if (ActualDependencies.Contains(SubObject))
+						{
+							SubObjectsArray.Add(SubObject);
+						}
 
 						for (auto& SubEntry : SubObjectDependencyGraph)
 						{
@@ -278,6 +318,26 @@ namespace UE::AssetCache::Private
 			}
 		}
 
+		// Filter objects we'll be serializing beforehand. This is a significant difference from the reference code,
+		// but essentially if we aren't going to be serializing an object, we shouldn't serialize the info about it
+		// either: If it's a default subobject we'll find another instance of it with the same name when deserializing
+		// anyway. If it's not, it probably shouldn't be serialized anyway
+		{
+			TArray<UObject*> FilteredSubObjects;
+			FilteredSubObjects.Reserve(SubObjectsArray.Num());
+			for (UObject* SubObject : SubObjectsArray)
+			{
+				if (SubObject &&
+					!SubObject->HasAnyFlags(RF_DefaultSubObject) &&
+					(!SubObject->IsEditorOnly() || !Ar.IsCooking()))
+				{
+					FilteredSubObjects.Add(SubObject);
+				}
+			}
+			Swap(FilteredSubObjects, SubObjectsArray);
+		}
+
+
 		// Serialize size of array
 		int32 SubObjectsCount = SubObjectsArray.Num();
 		Ar << SubObjectsCount;
@@ -314,18 +374,15 @@ namespace UE::AssetCache::Private
 
 		for (UObject* SubObject : SubObjectsArray)
 		{
-			if (SubObject && !SubObject->HasAnyFlags(RF_DefaultSubObject))
+			// Ensure these objects are serialized with flags as if they were regular
+			// persistent assets, in case they have branching on their serialization functions that watch
+			// out for this
+			SubObject->ClearFlags(RF_Transient);
+			SubObject->SetFlags(RF_Public);
 			{
-				// Ensure these objects are serialized with flags as if they were regular
-				// persistent assets, in case they have branching on their serialization functions that watch
-				// out for this
-				SubObject->ClearFlags(RF_Transient);
-				SubObject->SetFlags(RF_Public);
-				{
-					SubObject->Serialize(Ar);
-				}
-				SubObject->SetFlags(RF_Transient);
+				SubObject->Serialize(Ar);
 			}
+			SubObject->SetFlags(RF_Transient);
 		}
 
 		Object->ClearFlags(RF_Transient);
@@ -369,6 +426,7 @@ namespace UE::AssetCache::Private
 
 		FAssetCacheObjectReader Reader{ *Object, InBuffer };
 		FArchiveObjectPtrAsStringWrapper Ar{ Reader };
+		Ar.bLoadIfFindFails = true;
 
 		// Deserialize count of sub-objects
 		int32 SubObjectsCount = 0;
@@ -431,13 +489,17 @@ namespace UE::AssetCache::Private
 			}
 		}
 
+		TArray<UTexture*> DeserializedTextures;
+		TArray<UStaticMesh*> DeserializedStaticMeshes;
+		TArray<USkeletalMesh*> DeserializedSkeletalMeshes;
+		TArray<UMaterialInterface*> DeserializedMaterials;
+
 		for (UObject* SubObject : SubObjectsArray)
 		{
-			if (SubObject && !SubObject->HasAnyFlags(RF_DefaultSubObject))
-			{
-				SubObject->Serialize(Ar);
-				SubObject->SetFlags(RF_Transient | RF_Public);
-			}
+			ensure(SubObject && !SubObject->HasAnyFlags(RF_DefaultSubObject));
+
+			SubObject->Serialize(Ar);
+			SubObject->SetFlags(RF_Transient | RF_Public);
 		}
 
 		Object->Serialize(Ar);
@@ -447,6 +509,35 @@ namespace UE::AssetCache::Private
 		{
 			Texture->UpdateResource();
 		}
+#if WITH_EDITOR
+		else if (UStaticMesh* StaticMesh = Cast<UStaticMesh>(Object))
+		{
+			// These steps are copied from USDGeomMeshTranslator.cpp, from the PreBuildStaticMesh, BuildStaticMesh
+			// and PostBuildStaticMesh functions, although we will have serialized our body setup and extended bounds
+			// already. Ideally this will hit DDC and should be fast.
+			// We only need these in the editor: When we're cooking we'll just save the cooked mesh data into the
+			// asset itself, so at runtime there is no need for an additional "build" step that we have to handle here
+
+			StaticMesh->SetRenderData(MakeUnique< FStaticMeshRenderData >());
+			ITargetPlatformManagerModule& TargetPlatformManager = GetTargetPlatformManagerRef();
+			ITargetPlatform* RunningPlatform = TargetPlatformManager.GetRunningTargetPlatform();
+			const FStaticMeshLODSettings& LODSettings = RunningPlatform->GetStaticMeshLODSettings();
+			StaticMesh->GetRenderData()->Cache(RunningPlatform, StaticMesh, LODSettings);
+			StaticMesh->InitResources();
+		}
+		else if (USkeletalMesh* SkeletalMesh = Cast<USkeletalMesh>(Object))
+		{
+			SkeletalMesh->CalculateInvRefMatrices();
+			SkeletalMesh->InitMorphTargetsAndRebuildRenderData();
+		}
+#endif // WITH_EDITOR
+		if (UMaterialInterface* Material = Cast<UMaterialInterface>(Object))
+		{
+			// TODO: If this presents a performance issue we could possibly just do part of the work that this seems
+			// to be doing, but until then it's probably good to run the entire FMaterialUpdateContext for robustness
+			FMaterialUpdateContext UpdateContext;
+			UpdateContext.AddMaterialInterface(Material);
+		}
 	}
 }
 
@@ -516,7 +607,7 @@ void UUsdAssetCache2::CacheAsset(const FString& Hash, UObject* Asset, const UObj
 		return;
 	}
 
-	UE_LOG(LogUsd, Verbose, TEXT("Caching asset '%s' with for hash '%s' into the USD Asset Cache '%s'"),
+	UE_LOG(LogUsd, Verbose, TEXT("Caching asset '%s' with hash '%s' into the USD Asset Cache '%s'"),
 		*Asset->GetPathName(),
 		*Hash,
 		*GetPathName()
@@ -553,13 +644,7 @@ void UUsdAssetCache2::CacheAsset(const FString& Hash, UObject* Asset, const UObj
 			Modify();
 
 			FWriteScopeLock Lock(RWLock);
-
-			if (ReferencerToUse)
-			{
-				FoundInfo->Referencers.Add(FObjectKey{ ReferencerToUse });
-			}
-			LRUCache.FindAndTouch(Asset);
-			ActiveAssets.Add(Asset);
+			TouchAsset(Asset, ReferencerToUse);
 			return;
 		}
 		else
@@ -703,121 +788,150 @@ UObject* UUsdAssetCache2::RemoveAsset(const FString& Hash)
 
 UObject* UUsdAssetCache2::GetCachedAsset(const FString& Hash)
 {
-	// Write here as this will affect LRUCache order and PendingPersistentStorage
-	TOptional<FRWScopeLock> Lock;
-	Lock.Emplace(RWLock, SLT_Write);
-
-	UObject* FoundObject = AssetStorage.FindRef(Hash);
-
+	// Try to fetch a loaded asset from the cache with this hash
+	UObject* FoundAsset = nullptr;
+	{
+		FReadScopeLock Lock{RWLock};
+		FoundAsset = AssetStorage.FindRef(Hash);
+	}
 	UE_LOG(LogUsd, Verbose, TEXT("Fetching cached asset with hash '%s' for cache '%s'. Loaded in AssetStorage? %d. Size of LRU cache: %d, Pending persistent assets: %d"),
 		*Hash,
 		*GetPathName(),
-		FoundObject != nullptr,
+		FoundAsset != nullptr,
 		LRUCache.Num(),
 		PendingPersistentStorage.Num()
 	);
+	if (FoundAsset)
+	{
+		FWriteScopeLock Lock{RWLock};
+		TouchAsset(FoundAsset, CurrentScopedReferencer);
+		return FoundAsset;
+	}
 
-	// Maybe this is a persistent asset we just haven't loaded yet
-	if (!FoundObject)
+	// Check if we know about this persistent asset but just haven't loaded it yet
+	FSoftObjectPath* FoundAssetPath = nullptr;
 	{
-		if (FSoftObjectPath* FoundSoftPath = PendingPersistentStorage.Find(Hash))
+		FReadScopeLock Lock{RWLock};
+		FoundAssetPath = PendingPersistentStorage.Find(Hash);
+		if (!FoundAssetPath)
 		{
-			if (FCachedAssetInfo* FoundCachedInfo = LRUCache.Find(*FoundSoftPath))
-			{
-				const FString ClassNameString = FoundCachedInfo->AssetClassName.ToString();
+			UE_LOG(LogUsd, Verbose, TEXT("Full cache miss: We do not know about any assets with hash '%s' in cache '%s'"),
+				*Hash,
+				*GetPathName()
+			);
+			return nullptr;
+		}
+	}
 
-				// Actually load the UObject and subobjects from disk
-				if (UClass* FoundClass = UE::AssetCache::Private::FindClass(*ClassNameString))
-				{
-					// Deserialize the referenced asset from the Info's bulkdata
-					// TODO: If performance is an issue here we can probably prevent this copy by having our own
-					// FObjectReader that uses a pointer and offset and reads it off the bulkdata directly
-					TArray<uint8> BulkDataCopy;
-					const void* BulkDataBytes = FoundCachedInfo->BulkData.LockReadOnly();
-					{
-						BulkDataCopy.SetNumUninitialized(FoundCachedInfo->BulkData.GetBulkDataSize());
-						FMemory::Memcpy(BulkDataCopy.GetData(), BulkDataBytes, BulkDataCopy.Num());
-					}
-					FoundCachedInfo->BulkData.Unlock();
+	TFunction<UObject* (FCachedAssetInfo* ObjectInfo)> LoadAssetFromBulkData = [this](FCachedAssetInfo* ObjectInfo) -> UObject*
+	{
+		if (!ObjectInfo)
+		{
+			return nullptr;
+		}
 
-					FName AssetName = FoundCachedInfo->AssetName.IsEmpty() ? NAME_None : FName{ *FoundCachedInfo->AssetName };
+		const FString ClassNameString = ObjectInfo->AssetClassName.ToString();
+		UClass* FoundClass = UE::AssetCache::Private::FindClass(*ClassNameString);
+		if (!FoundClass)
+		{
+			UE_LOG(LogUsd, Warning, TEXT("Failed to find object class '%s' when deserializing asset '%s' with hash '%s'"),
+				*ClassNameString,
+				*ObjectInfo->AssetName,
+				*ObjectInfo->Hash
+			);
+			return nullptr;
+		}
 
-					// They're always "transient", since we'll persist them ourselves
-					UObject* NewAsset = NewObject<UObject>(this, FoundClass, AssetName, FoundCachedInfo->AssetFlags | RF_Transient);
-					UE::AssetCache::Private::DeserializeObjectAndSubObjects(NewAsset, BulkDataCopy);
-					ensure(NewAsset);
+		// Deserialize the referenced asset from the Info's bulkdata
+		// TODO: If performance is an issue here we can probably prevent this copy by having our own
+		// FObjectReader that uses a pointer and offset and reads it off the bulkdata directly
+		TArray<uint8> BulkDataCopy;
+		const void* BulkDataBytes = ObjectInfo->BulkData.LockReadOnly();
+		{
+			BulkDataCopy.SetNumUninitialized(ObjectInfo->BulkData.GetBulkDataSize());
+			FMemory::Memcpy(BulkDataCopy.GetData(), BulkDataBytes, BulkDataCopy.Num());
+		}
+		ObjectInfo->BulkData.Unlock();
 
-					UE_LOG(LogUsd, Verbose, TEXT("Deserialized asset at '%s' (%s, hash '%s', storage %d) from bulkdata (%.3f MB)"),
-						*FoundSoftPath->ToString(),
-						NewAsset ? *NewAsset->GetPathName() : TEXT("nullptr"),
-						*Hash,
-						FoundCachedInfo->CurrentStorageType,
-						BulkDataCopy.Num() / 1000000.0
-					);
+		FName AssetName = ObjectInfo->AssetName.IsEmpty() ? NAME_None : FName{*ObjectInfo->AssetName};
 
-					// Explicitly clearing this here for safety, as it will be invalidated when we remove Hash from
-					// PendingPersistentStorage just below
-					FoundSoftPath = nullptr;
+		// They're always "transient", since we'll persist them ourselves
+		UObject* NewAsset = NewObject<UObject>(this, FoundClass, AssetName, ObjectInfo->AssetFlags | RF_Transient);
+		UE::AssetCache::Private::DeserializeObjectAndSubObjects(NewAsset, BulkDataCopy);
+		ensure(NewAsset);
 
-					FoundObject = NewAsset;
-					AssetStorage.Add(Hash, NewAsset);
-					PendingPersistentStorage.Remove(Hash);
-				}
-				else
+		UE_LOG(LogUsd, Verbose, TEXT("Deserialized asset at '%s' (%s, hash '%s', storage %d) from bulkdata (%.3f MB)"),
+			*ObjectInfo->AssetName,
+			NewAsset ? *NewAsset->GetPathName() : TEXT("nullptr"),
+			*ObjectInfo->Hash,
+			ObjectInfo->CurrentStorageType,
+			BulkDataCopy.Num() / 1000000.0
+		);
+
+		AssetStorage.Add(ObjectInfo->Hash, NewAsset);
+		PendingPersistentStorage.Remove(ObjectInfo->Hash);
+		return NewAsset;
+	};
+
+	// Since LoadAssetFromBulkData will modify our member maps, the safest is to collect all hashes that we'll be
+	// loading first and then do a pass
+	TArray<FSoftObjectPath> AssetsToLoad;
+	TFunction<void(const FSoftObjectPath&)> RecursivelyCollectAssetsToLoad;
+	RecursivelyCollectAssetsToLoad = [this, &RecursivelyCollectAssetsToLoad, &AssetsToLoad](const FSoftObjectPath& AssetPath)
+	{
+		if (FCachedAssetInfo* FoundCachedInfo = LRUCache.Find(AssetPath))
+		{
+			if (PendingPersistentStorage.Contains(FoundCachedInfo->Hash))
+			{
+				// Before we deserialize the asset, we need to make sure all of it's dependencies are also
+				// deserialized, so push our dependency assets first
+				for (const FSoftObjectPath& DependencyPath : FoundCachedInfo->Dependencies)
 				{
-					UE_LOG(LogUsd, Warning, TEXT("Failed to find object class '%s' when deserializing asset '%s' with hash '%s'"),
-						*ClassNameString,
-						*FoundSoftPath->ToString(),
-						*Hash
-					);
+					RecursivelyCollectAssetsToLoad(DependencyPath);
 				}
-			}
-			else
-			{
-				UE_LOG(LogUsd, Warning, TEXT("Failed to find info about asset '%s' with hash '%s' when deserializing on-demand"),
-					*FoundSoftPath->ToString(),
-					*Hash
-				);
+
+				AssetsToLoad.Add(AssetPath);
 			}
 		}
+		else
+		{
+			UE_LOG(LogUsd, Warning, TEXT("Failed to find info about asset '%s' when deserializing on-demand"),
+				*AssetPath.ToString()
+			);
+		}
+	};
+	{
+		FReadScopeLock Lock{RWLock};
+		RecursivelyCollectAssetsToLoad(*FoundAssetPath);
 	}
 
-	if (FoundObject)
+	ensure(AssetsToLoad.Num() > 0);
+	UE_LOG(LogUsd, Verbose, TEXT("Partial cache miss: Will attempt to load %d assets from bulk data of cache '%s'"),
+		AssetsToLoad.Num(),
+		*GetPathName()
+	);
 	{
-		ActiveAssets.Add(FoundObject);
-		ensure(LRUCache.FindAndTouch(FoundObject));
-
-		if (CurrentScopedReferencer)
+		FWriteScopeLock Lock{RWLock};
+		for (const FSoftObjectPath& AssetToLoad : AssetsToLoad)
 		{
-			Lock.Reset(); // Release our lock as AddAssetReference will want to acquire it
-			AddAssetReference(FoundObject, CurrentScopedReferencer);
+			if (FCachedAssetInfo* FoundCachedInfo = LRUCache.Find(AssetToLoad))
+			{
+				UObject* LoadedAsset = LoadAssetFromBulkData(FoundCachedInfo);
+				ensure(LoadedAsset);
+
+				TouchAsset(LoadedAsset, CurrentScopedReferencer);
+				FoundAsset = LoadedAsset;  // The very last asset we loaded is our main asset, as it's pushed last into AssetsToLoad
+			}
 		}
 	}
 
-	return FoundObject;
+	return FoundAsset;
 }
 
 bool UUsdAssetCache2::AddAssetReference(const UObject* Asset, const UObject* Referencer)
 {
-	if (!Asset || !Referencer)
-	{
-		return false;
-	}
-
 	FWriteScopeLock Lock(RWLock);
-
-	if (FCachedAssetInfo* Info = LRUCache.Find(Asset))
-	{
-		Info->Referencers.Add(FObjectKey{Referencer});
-
-		UE_LOG(LogUsd, Verbose, TEXT("Added referencer '%s' for asset '%s'"),
-			*Referencer->GetPathName(),
-			*Asset->GetPathName()
-		);
-		return true;
-	}
-
-	return false;
+	return AddAssetReferenceInternal(Asset, Referencer);
 }
 
 bool UUsdAssetCache2::RemoveAssetReference(const UObject* Asset, const UObject* Referencer)
@@ -1361,6 +1475,43 @@ bool UUsdAssetCache2::IsAssetOwnedByCacheInternal(const FString& AssetPath) cons
 	return !AssetPath.IsEmpty() && LRUCache.Contains(AssetPath);
 }
 
+bool UUsdAssetCache2::AddAssetReferenceInternal(const UObject* Asset, const UObject* Referencer)
+{
+	if (!Asset || !Referencer)
+	{
+		return false;
+	}
+
+	if (FCachedAssetInfo* Info = LRUCache.Find(Asset))
+	{
+		Info->Referencers.Add(FObjectKey{Referencer});
+
+		UE_LOG(LogUsd, Verbose, TEXT("Added referencer '%s' for asset '%s'"),
+			*Referencer->GetPathName(),
+			*Asset->GetPathName()
+		);
+		return true;
+	}
+
+	return false;
+}
+
+void UUsdAssetCache2::TouchAsset(UObject* Asset, const UObject* Referencer)
+{
+	if (!Asset)
+	{
+		return;
+	}
+
+	ActiveAssets.Add(Asset);
+	ensure(LRUCache.FindAndTouch(Asset));
+
+	if (Referencer)
+	{
+		ensure(AddAssetReferenceInternal(Asset, Referencer));
+	}
+}
+
 void UUsdAssetCache2::Serialize(FArchive& Ar)
 {
 	if (Ar.IsSaving())
@@ -1374,14 +1525,31 @@ void UUsdAssetCache2::Serialize(FArchive& Ar)
 		ensure(SavedCacheVersion <= GCurrentCacheVersion);
 	}
 
-	FWriteScopeLock Lock(RWLock);
-
 	Super::Serialize(Ar);
 
 	if (Ar.IsPersistent() || Ar.IsTransacting() || (Ar.GetPortFlags() & (PPF_DuplicateForPIE | PPF_Duplicate)))
 	{
 		if (Ar.IsSaving())
 		{
+			// When we cook we must generate a new cooked bulkdata as it may need to filter editor-only stuff during
+			// the serialization process, so make sure none of these assets are pending
+			if (Ar.IsCooking())
+			{
+				// Take a copy here because GetCachedAsset will modify PendingPersistentStorage
+				TMap<FString, FSoftObjectPath> PendingStorageCopy;
+				{
+					FReadScopeLock Lock(RWLock);
+					PendingStorageCopy = PendingPersistentStorage;
+				}
+				for (const TPair<FString, FSoftObjectPath>& PendingAsset : PendingStorageCopy)
+				{
+					UObject* AssetFromBulkData = GetCachedAsset(PendingAsset.Key);
+					ensure(AssetFromBulkData);
+				}
+			}
+
+			FWriteScopeLock Lock(RWLock);
+
 			TSet<FSoftObjectPath> PersistentAssets;
 			for (TLruCache<FSoftObjectPath, FCachedAssetInfo>::TIterator Iter{ LRUCache }; Iter; ++Iter)
 			{
@@ -1394,18 +1562,9 @@ void UUsdAssetCache2::Serialize(FArchive& Ar)
 			int32 NumPersistent = PersistentAssets.Num();
 			Ar << NumPersistent;
 
-			// When we cook we must generate a new cooked bulkdata as it may need to filter editor-only stuff during
-			// the serialization process, so make sure none of these assets are pending
-			if (Ar.IsCooking())
-			{
-				// Take a copy here because GetCachedAsset will modify PendingPersistentStorage
-				TMap<FString, FSoftObjectPath> PendingStorageCopy = PendingPersistentStorage;
-				for (const TPair<FString, FSoftObjectPath>& PendingAsset : PendingStorageCopy)
-				{
-					UObject* AssetFromBulkData = GetCachedAsset(PendingAsset.Key);
-					ensure(AssetFromBulkData);
-				}
-			}
+			UE_LOG(LogUsd, Verbose, TEXT("Serializing USD Asset Cache '%s'"),
+				*GetPathName()
+			);
 
 			// By using the iterators we'll traverse the most recently used object first
 			for (TLruCache<FSoftObjectPath, FCachedAssetInfo>::TIterator Iter{ LRUCache }; Iter; ++Iter)
@@ -1441,6 +1600,13 @@ void UUsdAssetCache2::Serialize(FArchive& Ar)
 					Info.SizeOnDiskInBytes = Buffer.Num();
 				}
 
+				UE_LOG(LogUsd, Verbose, TEXT("\tPersistent asset '%s' ('%s') with hash '%s' (%u disk bytes)"),
+					*Info.AssetName,
+					*Info.AssetClassName.ToString(),
+					*Info.Hash,
+					Info.SizeOnDiskInBytes
+				);
+
 				Ar << AssetPath;
 
 				// We only want to serialize dependencies on persistent assets, or else when we deserialize this
@@ -1459,6 +1625,8 @@ void UUsdAssetCache2::Serialize(FArchive& Ar)
 		}
 		else if (Ar.IsLoading())
 		{
+			FWriteScopeLock Lock(RWLock);
+
 			int32 NumPersistent = INDEX_NONE;
 			Ar << NumPersistent;
 			ensure(NumPersistent >= 0);
@@ -1469,6 +1637,10 @@ void UUsdAssetCache2::Serialize(FArchive& Ar)
 			TArray<FSoftObjectPath> AssetPaths;
 			AssetPaths.Reserve(NumPersistent);
 
+			UE_LOG(LogUsd, Verbose, TEXT("Deserializing USD Asset Cache '%s'"),
+				*GetPathName()
+			);
+
 			for (int32 Index = 0; Index < NumPersistent; ++Index)
 			{
 				FSoftObjectPath& AssetPath = AssetPaths.Emplace_GetRef();
@@ -1478,6 +1650,13 @@ void UUsdAssetCache2::Serialize(FArchive& Ar)
 				NewEl.Serialize(Ar, this);
 				NewEl.CurrentStorageType = ECacheStorageType::Persistent;
 
+				UE_LOG(LogUsd, Verbose, TEXT("\tPersistent asset '%s' ('%s') with hash '%s' (%u disk bytes)"),
+					*NewEl.AssetName,
+					*NewEl.AssetClassName.ToString(),
+					*NewEl.Hash,
+					NewEl.SizeOnDiskInBytes
+				);
+
 				PendingPersistentStorage.Add(NewEl.Hash, AssetPath);
 			}
 
diff --git a/Engine/Plugins/Importers/USDImporter/Source/USDClasses/Public/USDAssetCache2.h b/Engine/Plugins/Importers/USDImporter/Source/USDClasses/Public/USDAssetCache2.h
--- a/Engine/Plugins/Importers/USDImporter/Source/USDClasses/Public/USDAssetCache2.h
+++ b/Engine/Plugins/Importers/USDImporter/Source/USDClasses/Public/USDAssetCache2.h
@@ -170,6 +170,8 @@ private:
 	// Tries unloading unreferenced, persistent assets back to bulkdata on disk if we can
 	bool TryUnloadAsset(FCachedAssetInfo& InOutInfo);
 	bool IsAssetOwnedByCacheInternal(const FString& AssetPath) const;
+	bool AddAssetReferenceInternal(const UObject* Asset, const UObject* Referencer);
+	void TouchAsset(UObject* Asset, const UObject* Referencer = nullptr);
 
 private:
 	friend struct FUsdScopedAssetCacheReferencer;
diff --git a/Engine/Plugins/Importers/USDImporter/Source/USDClasses/USDClasses.Build.cs b/Engine/Plugins/Importers/USDImporter/Source/USDClasses/USDClasses.Build.cs
--- a/Engine/Plugins/Importers/USDImporter/Source/USDClasses/USDClasses.Build.cs
+++ b/Engine/Plugins/Importers/USDImporter/Source/USDClasses/USDClasses.Build.cs
@@ -27,6 +27,8 @@ namespace UnrealBuildTool.Rules
 					"DeveloperSettings",
 					"GeometryCache",
 					"Json", // To read/write plugInfo.json files from UnrealUSDWrapper.cpp
+					"MeshDescription",
+					"RHI",
 				}
 			);
 		}
diff --git a/Engine/Plugins/Importers/USDImporter/Source/USDSchemas/Private/MeshTranslationImpl.cpp b/Engine/Plugins/Importers/USDImporter/Source/USDSchemas/Private/MeshTranslationImpl.cpp
--- a/Engine/Plugins/Importers/USDImporter/Source/USDSchemas/Private/MeshTranslationImpl.cpp
+++ b/Engine/Plugins/Importers/USDImporter/Source/USDSchemas/Private/MeshTranslationImpl.cpp
@@ -212,13 +212,22 @@ TMap<const UsdUtils::FUsdPrimMaterialSlot*, UMaterialInterface*> MeshTranslation
 			}
 			case UsdUtils::EPrimAssignmentType::UnrealMaterial:
 			{
-				Material = Cast< UMaterialInterface >(FSoftObjectPath(Slot.MaterialSource).TryLoad());
-				if (!Material)
+				UObject* Object = FSoftObjectPath(Slot.MaterialSource).TryLoad();
+				Material = Cast< UMaterialInterface >(Object);
+				if (!Object)
 				{
 					UE_LOG(LogUsd, Warning, TEXT("UE material '%s' for prim '%s' could not be loaded or was not found."),
 						*Slot.MaterialSource,
 						*UsdToUnreal::ConvertPath(UsdPrim.GetPrimPath()));
 				}
+				else if (!Material)
+				{
+					UE_LOG(LogUsd, Warning, TEXT("Object '%s' assigned as an Unreal Material for prim '%s' is not actually a material (but instead a '%s') and will not be used"),
+						*Slot.MaterialSource,
+						*UsdToUnreal::ConvertPath(UsdPrim.GetPrimPath()),
+						*Object->GetClass()->GetName()
+					);
+				}
 				else if (!Material->IsTwoSided() && Slot.bMeshIsDoubleSided)
 				{
 					UE_LOG(LogUsd, Warning, TEXT("Using one-sided UE material '%s' for doubleSided prim '%s'"),
diff --git a/Engine/Plugins/Importers/USDImporter/Source/USDSchemas/Private/USDGeomMeshTranslator.cpp b/Engine/Plugins/Importers/USDImporter/Source/USDSchemas/Private/USDGeomMeshTranslator.cpp
--- a/Engine/Plugins/Importers/USDImporter/Source/USDSchemas/Private/USDGeomMeshTranslator.cpp
+++ b/Engine/Plugins/Importers/USDImporter/Source/USDSchemas/Private/USDGeomMeshTranslator.cpp
@@ -431,6 +431,7 @@ namespace UsdGeomMeshTranslatorImpl
 				}
 				case UsdUtils::EPrimAssignmentType::UnrealMaterial:
 				{
+					// TODO: Try reusing the MeshTranslationImpl code for this as it has more warnings/feedback
 					Material = Cast< UMaterialInterface >(FSoftObjectPath(Slot.MaterialSource).TryLoad());
 					break;
 				}
@@ -834,6 +835,7 @@ namespace UsdGeomMeshTranslatorImpl
 
 		StaticMesh.SetRenderData(MakeUnique< FStaticMeshRenderData >());
 		StaticMesh.CreateBodySetup();
+		StaticMesh.MarkAsNotHavingNavigationData(); // Needed or else it will warn if we try cooking with body setup
 	}
 
 	bool BuildStaticMesh(UStaticMesh& StaticMesh, const FStaticFeatureLevel& FeatureLevel, TArray<FMeshDescription>& LODIndexToMeshDescription)
diff --git a/Engine/Plugins/Importers/USDImporter/Source/USDSchemas/Private/USDGeomPointInstancerTranslator.cpp b/Engine/Plugins/Importers/USDImporter/Source/USDSchemas/Private/USDGeomPointInstancerTranslator.cpp
--- a/Engine/Plugins/Importers/USDImporter/Source/USDSchemas/Private/USDGeomPointInstancerTranslator.cpp
+++ b/Engine/Plugins/Importers/USDImporter/Source/USDSchemas/Private/USDGeomPointInstancerTranslator.cpp
@@ -72,6 +72,7 @@ namespace UsdGeomPointInstancerTranslatorImpl
 		if ( StaticMesh )
 		{
 			StaticMesh->CreateBodySetup(); // BodySetup is required for HISM component
+			StaticMesh->MarkAsNotHavingNavigationData(); // Needed or else it will warn if we try cooking with body setup
 		}
 
 		MeshComponent.SetStaticMesh( StaticMesh );
diff --git a/Engine/Plugins/Importers/USDImporter/Source/USDUtilities/Private/USDGeomMeshConversion.cpp b/Engine/Plugins/Importers/USDImporter/Source/USDUtilities/Private/USDGeomMeshConversion.cpp
--- a/Engine/Plugins/Importers/USDImporter/Source/USDUtilities/Private/USDGeomMeshConversion.cpp
+++ b/Engine/Plugins/Importers/USDImporter/Source/USDUtilities/Private/USDGeomMeshConversion.cpp
@@ -1542,16 +1542,13 @@ UsdUtils::FUsdPrimMaterialAssignmentInfo UsdUtils::GetPrimMaterialAssignments(
 
 		if ( !ValidPackagePath.IsEmpty() )
 		{
-			// We can't TryLoad() this right now as we may be in an Async thread, so settle for checking with the asset registry module
+			// We can't TryLoad() or LoadObject<> this right now as we may be in an Async thread.
+			// The FAssetData may not be ready yet however, in case we're loading a stage right when launching the
+			// editor, so here we just settle for finding any valid object
 			FSoftObjectPath SoftObjectPath{ ValidPackagePath };
 			if ( SoftObjectPath.IsValid() )
 			{
-				FAssetRegistryModule& AssetRegistryModule = FModuleManager::LoadModuleChecked<FAssetRegistryModule>( "AssetRegistry" );
-				FAssetData AssetData = AssetRegistryModule.Get().GetAssetByObjectPath( SoftObjectPath );
-				if ( AssetData.IsValid() && AssetData.GetClass()->IsChildOf( UMaterialInterface::StaticClass() ) )
-				{
-					return ValidPackagePath;
-				}
+				return ValidPackagePath;
 			}
 
 			UE_LOG( LogUsd, Warning, TEXT( "Could not find a valid material at path '%s', targetted by prim '%s's unrealMaterial attribute. Material assignment will fallback to USD materials and display color data." ),
diff --git a/Engine/Plugins/Importers/USDImporter/Source/USDUtilities/Private/USDShadeConversion.cpp b/Engine/Plugins/Importers/USDImporter/Source/USDUtilities/Private/USDShadeConversion.cpp
--- a/Engine/Plugins/Importers/USDImporter/Source/USDUtilities/Private/USDShadeConversion.cpp
+++ b/Engine/Plugins/Importers/USDImporter/Source/USDUtilities/Private/USDShadeConversion.cpp
@@ -1944,6 +1944,13 @@ namespace UE
 						FString ResolvedPath = UsdUtils::GetResolvedTexturePath( ShadeInput.GetAttr() );
 						InOutHashState.UpdateWithString( *ResolvedPath, ResolvedPath.Len() );
 					}
+					else if (ShadeInputValue.IsHolding<pxr::TfToken>())
+					{
+						// Stringify instead of using GetHash because if ShadeInputValue contains a pxr::TfToken then
+						// it will actually just contain some non-deterministic integer IDs
+						FString Stringified = UsdToUnreal::ConvertString(pxr::TfStringify(ShadeInputValue));
+						InOutHashState.UpdateWithString( *Stringified, Stringified.Len() );
+					}
 					else
 					{
 						uint64 ValueHash = ( uint64 ) ShadeInputValue.GetHash();
