USD: Added USDSTAGE_API to AUsdStageActor() to allow derivation from other modules/plugins.

diff --git a/Engine/Plugins/Importers/USDImporter/Source/USDStage/Public/USDStageActor.h b/Engine/Plugins/Importers/USDImporter/Source/USDStage/Public/USDStageActor.h
--- a/Engine/Plugins/Importers/USDImporter/Source/USDStage/Public/USDStageActor.h
+++ b/Engine/Plugins/Importers/USDImporter/Source/USDStage/Public/USDStageActor.h
@@ -259,7 +259,7 @@ public:
 	USDSTAGE_API virtual void PostUnregisterAllComponents() override;
 	// End AActor interface
 
-	AUsdStageActor();
+	USDSTAGE_API AUsdStageActor();
 
 protected:
 	/** Loads the asset for a single prim */
