USD: Track resource sizes for UGeometryCache and UGeometryTrackUsd

diff --git a/Engine/Plugins/Importers/USDImporter/Source/GeometryCacheUSD/Private/GeometryCacheTrackUSD.cpp b/Engine/Plugins/Importers/USDImporter/Source/GeometryCacheUSD/Private/GeometryCacheTrackUSD.cpp
--- a/Engine/Plugins/Importers/USDImporter/Source/GeometryCacheUSD/Private/GeometryCacheTrackUSD.cpp
+++ b/Engine/Plugins/Importers/USDImporter/Source/GeometryCacheUSD/Private/GeometryCacheTrackUSD.cpp
@@ -36,6 +36,33 @@ void UGeometryCacheTrackUsd::BeginDestroy()
 	Super::BeginDestroy();
 }
 
+void UGeometryCacheTrackUsd::GetResourceSizeEx(FResourceSizeEx& CumulativeResourceSize)
+{
+	Super::GetResourceSizeEx(CumulativeResourceSize);
+
+	// This is an additional copy that lives on the track
+	MeshData.GetResourceSizeEx(CumulativeResourceSize);
+
+	if (FGeometryCacheUsdStream* UsdStreamPtr = UsdStream.Get())
+	{
+		const FGeometryCacheStreamStats& Stats = UsdStreamPtr->GetStreamStats();
+
+		int32 NumFramesFullyLoaded = Stats.NumCachedFrames;
+		float TotalMemoryForLoadedFramesMB = Stats.MemoryUsed;
+		float MemoryPerFrameMB = TotalMemoryForLoadedFramesMB / NumFramesFullyLoaded;
+
+		// Also account for frames we're going to prefetch
+		uint32 NumRemainingFrames = UsdStreamPtr->GetNumFramesNeeded();
+		float TotalMemoryMB = TotalMemoryForLoadedFramesMB + NumRemainingFrames * MemoryPerFrameMB;
+
+		// Using 1048576 instead of 1000000 for megabyte as that's what's used by FGeometryCacheStreamBase memory stats
+		double TotalMemoryBytes = TotalMemoryMB * 1048576;
+		CumulativeResourceSize.AddDedicatedSystemMemoryBytes(static_cast<SIZE_T>(TotalMemoryBytes + 0.5));
+	}
+
+	CumulativeResourceSize.AddDedicatedSystemMemoryBytes(SampleInfos.GetAllocatedSize());
+}
+
 const bool UGeometryCacheTrackUsd::UpdateMeshData(const float Time, const bool bLooping, int32& InOutMeshSampleIndex, FGeometryCacheMeshData*& OutMeshData)
 {
 	const int32 SampleIndex = FindSampleIndexFromTime(Time, bLooping);
diff --git a/Engine/Plugins/Importers/USDImporter/Source/GeometryCacheUSD/Public/GeometryCacheTrackUSD.h b/Engine/Plugins/Importers/USDImporter/Source/GeometryCacheUSD/Public/GeometryCacheTrackUSD.h
--- a/Engine/Plugins/Importers/USDImporter/Source/GeometryCacheUSD/Public/GeometryCacheTrackUSD.h
+++ b/Engine/Plugins/Importers/USDImporter/Source/GeometryCacheUSD/Public/GeometryCacheTrackUSD.h
@@ -32,6 +32,7 @@ public:
 
 	//~ Begin UObject Interface.
 	virtual void BeginDestroy() override;
+	virtual void GetResourceSizeEx(FResourceSizeEx& CumulativeResourceSize) override;
 	//~ End UObject Interface.
 
 	//~ Begin UGeometryCacheTrack Interface.
diff --git a/Engine/Plugins/Runtime/GeometryCache/Source/GeometryCache/Classes/GeometryCache.h b/Engine/Plugins/Runtime/GeometryCache/Source/GeometryCache/Classes/GeometryCache.h
--- a/Engine/Plugins/Runtime/GeometryCache/Source/GeometryCache/Classes/GeometryCache.h
+++ b/Engine/Plugins/Runtime/GeometryCache/Source/GeometryCache/Classes/GeometryCache.h
@@ -26,6 +26,7 @@ public:
 	virtual void PostInitProperties() override;
 	virtual FString GetDesc() override;
 	virtual void GetAssetRegistryTags(TArray<FAssetRegistryTag>& OutTags) const override;
+	virtual void GetResourceSizeEx(FResourceSizeEx& CumulativeResourceSize) override;
 
 	virtual void BeginDestroy() override;
 	virtual bool IsReadyForFinishDestroy() override;
diff --git a/Engine/Plugins/Runtime/GeometryCache/Source/GeometryCache/Private/GeometryCache.cpp b/Engine/Plugins/Runtime/GeometryCache/Source/GeometryCache/Private/GeometryCache.cpp
--- a/Engine/Plugins/Runtime/GeometryCache/Source/GeometryCache/Private/GeometryCache.cpp
+++ b/Engine/Plugins/Runtime/GeometryCache/Source/GeometryCache/Private/GeometryCache.cpp
@@ -108,6 +108,19 @@ void UGeometryCache::GetAssetRegistryTags(TArray<FAssetRegistryTag>& OutTags) co
 	Super::GetAssetRegistryTags(OutTags);
 }
 
+void UGeometryCache::GetResourceSizeEx(FResourceSizeEx& CumulativeResourceSize)
+{
+	Super::GetResourceSizeEx(CumulativeResourceSize);
+
+	for(TObjectPtr<UGeometryCacheTrack> Track : Tracks)
+	{
+		if (Track)
+		{
+			Track->GetResourceSizeEx(CumulativeResourceSize);
+		}
+	}
+}
+
 void UGeometryCache::BeginDestroy()
 {
 	Super::BeginDestroy();
