USD: Display relationships, prim metadata and attribute metadata.

diff --git a/Engine/Plugins/Importers/USDImporter/Source/USDClasses/Private/USDAssetCache2.cpp b/Engine/Plugins/Importers/USDImporter/Source/USDClasses/Private/USDAssetCache2.cpp
--- a/Engine/Plugins/Importers/USDImporter/Source/USDClasses/Private/USDAssetCache2.cpp
+++ b/Engine/Plugins/Importers/USDImporter/Source/USDClasses/Private/USDAssetCache2.cpp
@@ -1678,7 +1678,7 @@ FUsdScopedAssetCacheReferencer::FUsdScopedAssetCacheReferencer(UUsdAssetCache2*
 	// For now we're assuming you can't nest these objects
 	if (ensure(InAssetCache))
 	{
-		ensure(Referencer && !InAssetCache->CurrentScopedReferencer);
+		ensure(Referencer && (!InAssetCache->CurrentScopedReferencer || InAssetCache->CurrentScopedReferencer == Referencer));
 		InAssetCache->CurrentScopedReferencer = Referencer;
 	}
 
diff --git a/Engine/Plugins/Importers/USDImporter/Source/USDStage/Private/USDStageActor.cpp b/Engine/Plugins/Importers/USDImporter/Source/USDStage/Private/USDStageActor.cpp
--- a/Engine/Plugins/Importers/USDImporter/Source/USDStage/Private/USDStageActor.cpp
+++ b/Engine/Plugins/Importers/USDImporter/Source/USDStage/Private/USDStageActor.cpp
@@ -973,6 +973,16 @@ void AUsdStageActor::OnUsdObjectsChanged(const UsdUtils::FObjectChangesByPath& I
 		bool bIsResync = false;
 		for (const UsdUtils::FObjectChangeNotice& ObjectChange : InfoChange.Value)
 		{
+			// We upgrade these mainly so that material binding changes refresh assets and components. Note that we
+			// just set material overrides on CreateComponents and not on UpdateComponents. In general, relationships
+			// likely won't be updated very often and can't have timeSamples, so it should be fine to resync on their
+			// changes
+			if (ObjectChange.Flags.bDidChangeRelationshipTargets)
+			{
+				bIsResync = true;
+				bHasResync = true;
+			}
+
 			for (const UsdUtils::FAttributeChange& AttributeChange : ObjectChange.AttributeChanges)
 			{
 				static const TSet<FString> StageResyncProperties = {
diff --git a/Engine/Plugins/Importers/USDImporter/Source/USDStageEditor/Private/SUSDPrimInfo.cpp b/Engine/Plugins/Importers/USDImporter/Source/USDStageEditor/Private/SUSDPrimInfo.cpp
--- a/Engine/Plugins/Importers/USDImporter/Source/USDStageEditor/Private/SUSDPrimInfo.cpp
+++ b/Engine/Plugins/Importers/USDImporter/Source/USDStageEditor/Private/SUSDPrimInfo.cpp
@@ -35,15 +35,48 @@ void SUsdPrimInfo::Construct( const FArguments& InArgs )
 {
 	ChildSlot
 	[
-		SNew( SVerticalBox )
+		SNew(SVerticalBox)
 
 		+SVerticalBox::Slot()
-		.FillHeight( 1.f )
+		.FillHeight(1.f)
+		[
+			SNew(SBox)
+			.Content()
+			[
+				SAssignNew(PropertiesList, SUsdPrimPropertiesList)
+				.NameColumnText(LOCTEXT("PropertyName", "Name"))
+				.OnSelectionChanged_Lambda([this](const TSharedPtr<FUsdPrimAttributeViewModel>& NewSelection, ESelectInfo::Type SelectionType)
+				{
+					// Display property metadata if we have exactly one selected
+					TArray<FString> SelectedProperties = PropertiesList->GetSelectedPropertyNames();
+					if (PropertiesList && SelectedProperties.Num() == 1)
+					{
+						PropertyMetadataPanel->SetPrimPath(
+							PropertiesList->GetUsdStage(),
+							*(FString{PropertiesList->GetPrimPath()} + "." + SelectedProperties[0])
+						);
+					}
+					else
+					{
+						PropertyMetadataPanel->SetPrimPath({}, TEXT(""));
+					}
+				})
+			]
+		]
+
+		+SVerticalBox::Slot()
+		.AutoHeight()
 		[
 			SNew( SBox )
 			.Content()
 			[
-				SAssignNew( PropertiesList, SUsdPrimPropertiesList )
+				SAssignNew(PropertyMetadataPanel, SUsdPrimPropertiesList)
+				.NameColumnText_Lambda([this]() -> FText
+				{
+					FString PropertyName = FPaths::GetExtension(PropertyMetadataPanel->GetPrimPath());
+					return FText::FromString(FString::Printf(TEXT("%s metadata"), *PropertyName));
+				})
+				.Visibility(EVisibility::Collapsed)
 			]
 		]
 
@@ -84,6 +117,9 @@ void SUsdPrimInfo::SetPrimPath( const UE::FUsdStageWeak& UsdStage, const TCHAR*
 	if ( PropertiesList )
 	{
 		PropertiesList->SetPrimPath( UsdStage, PrimPath );
+
+		PropertyMetadataPanel->SetPrimPath({}, TEXT(""));
+		PropertyMetadataPanel->ClearSelection();
 	}
 
 	if ( IntegrationsPanel )
diff --git a/Engine/Plugins/Importers/USDImporter/Source/USDStageEditor/Private/SUSDPrimInfo.h b/Engine/Plugins/Importers/USDImporter/Source/USDStageEditor/Private/SUSDPrimInfo.h
--- a/Engine/Plugins/Importers/USDImporter/Source/USDStageEditor/Private/SUSDPrimInfo.h
+++ b/Engine/Plugins/Importers/USDImporter/Source/USDStageEditor/Private/SUSDPrimInfo.h
@@ -21,6 +21,7 @@ public:
 private:
 	friend class SUsdStage;
 	TSharedPtr< class SUsdPrimPropertiesList > PropertiesList;
+	TSharedPtr< class SUsdPrimPropertiesList > PropertyMetadataPanel;
 	TSharedPtr< class SUsdIntegrationsPanel > IntegrationsPanel;
 	TSharedPtr< class SVariantsList > VariantsList;
 	TSharedPtr< class SUsdReferencesList > ReferencesList;
diff --git a/Engine/Plugins/Importers/USDImporter/Source/USDStageEditor/Private/SUSDStage.cpp b/Engine/Plugins/Importers/USDImporter/Source/USDStageEditor/Private/SUSDStage.cpp
--- a/Engine/Plugins/Importers/USDImporter/Source/USDStageEditor/Private/SUSDStage.cpp
+++ b/Engine/Plugins/Importers/USDImporter/Source/USDStageEditor/Private/SUSDStage.cpp
@@ -463,6 +463,24 @@ void SUsdStage::SetSelectedPropertyNames( const TArray<FString>& NewSelection )
 	}
 }
 
+TArray<FString> SUsdStage::GetSelectedPropertyMetadataNames() const
+{
+	if (UsdPrimInfoWidget && UsdPrimInfoWidget->PropertyMetadataPanel)
+	{
+		return UsdPrimInfoWidget->PropertyMetadataPanel->GetSelectedPropertyNames();
+	}
+
+	return {};
+}
+
+void SUsdStage::SetSelectedPropertyMetadataNames(const TArray<FString>& NewSelection)
+{
+	if (UsdPrimInfoWidget && UsdPrimInfoWidget->PropertyMetadataPanel)
+	{
+		UsdPrimInfoWidget->PropertyMetadataPanel->SetSelectedPropertyNames(NewSelection);
+	}
+}
+
 void SUsdStage::AttachToStageActor( AUsdStageActor* InUsdStageActor, bool bFlashButton )
 {
 	if ( ViewModel.UsdStageActor == InUsdStageActor )
diff --git a/Engine/Plugins/Importers/USDImporter/Source/USDStageEditor/Private/SUSDStage.h b/Engine/Plugins/Importers/USDImporter/Source/USDStageEditor/Private/SUSDStage.h
--- a/Engine/Plugins/Importers/USDImporter/Source/USDStageEditor/Private/SUSDStage.h
+++ b/Engine/Plugins/Importers/USDImporter/Source/USDStageEditor/Private/SUSDStage.h
@@ -56,6 +56,9 @@ public:
 	TArray<FString> GetSelectedPropertyNames() const;
 	void SetSelectedPropertyNames( const TArray<FString>& NewSelection );
 
+	TArray<FString> GetSelectedPropertyMetadataNames() const;
+	void SetSelectedPropertyMetadataNames(const TArray<FString>& NewSelection);
+
 	// Main menu actions.
 	// For all of these, providing an empty path will cause us to pop open a dialog to let the user pick the path
 	// instead.
diff --git a/Engine/Plugins/Importers/USDImporter/Source/USDStageEditor/Private/USDStageEditorBlueprintLibrary.cpp b/Engine/Plugins/Importers/USDImporter/Source/USDStageEditor/Private/USDStageEditorBlueprintLibrary.cpp
--- a/Engine/Plugins/Importers/USDImporter/Source/USDStageEditor/Private/USDStageEditorBlueprintLibrary.cpp
+++ b/Engine/Plugins/Importers/USDImporter/Source/USDStageEditor/Private/USDStageEditorBlueprintLibrary.cpp
@@ -119,6 +119,18 @@ void UUsdStageEditorBlueprintLibrary::SetSelectedPropertyNames( const TArray<FSt
 	StageEditorModule.SetSelectedPropertyNames( NewSelection );
 }
 
+TArray<FString> UUsdStageEditorBlueprintLibrary::GetSelectedPropertyMetadataNames()
+{
+	IUsdStageEditorModule& StageEditorModule = FModuleManager::GetModuleChecked< IUsdStageEditorModule >("USDStageEditor");
+	return StageEditorModule.GetSelectedPropertyMetadataNames();
+}
+
+void UUsdStageEditorBlueprintLibrary::SetSelectedPropertyMetadataNames(const TArray<FString>& NewSelection)
+{
+	IUsdStageEditorModule& StageEditorModule = FModuleManager::GetModuleChecked< IUsdStageEditorModule >("USDStageEditor");
+	StageEditorModule.SetSelectedPropertyMetadataNames(NewSelection);
+}
+
 void UUsdStageEditorBlueprintLibrary::FileNew()
 {
 	IUsdStageEditorModule& StageEditorModule = FModuleManager::GetModuleChecked< IUsdStageEditorModule >( "USDStageEditor" );
diff --git a/Engine/Plugins/Importers/USDImporter/Source/USDStageEditor/Private/USDStageEditorModule.cpp b/Engine/Plugins/Importers/USDImporter/Source/USDStageEditor/Private/USDStageEditorModule.cpp
--- a/Engine/Plugins/Importers/USDImporter/Source/USDStageEditor/Private/USDStageEditorModule.cpp
+++ b/Engine/Plugins/Importers/USDImporter/Source/USDStageEditor/Private/USDStageEditorModule.cpp
@@ -538,6 +538,28 @@ void IUsdStageEditorModule::SetSelectedPropertyNames( const TArray<FString>& New
 #endif // USE_USD_SDK
 }
 
+TArray<FString> IUsdStageEditorModule::GetSelectedPropertyMetadataNames() const
+{
+#if USE_USD_SDK
+	if (TSharedPtr<SUsdStage> UsdStageEditor = UE::UsdStageEditorModule::Private::GetUsdStageEditor())
+	{
+		return UsdStageEditor->GetSelectedPropertyMetadataNames();
+	}
+#endif // USE_USD_SDK
+
+	return {};
+}
+
+void IUsdStageEditorModule::SetSelectedPropertyMetadataNames(const TArray<FString>& NewSelection) const
+{
+#if USE_USD_SDK
+	if (TSharedPtr<SUsdStage> UsdStageEditor = UE::UsdStageEditorModule::Private::GetUsdStageEditor())
+	{
+		UsdStageEditor->SetSelectedPropertyMetadataNames(NewSelection);
+	}
+#endif // USE_USD_SDK
+}
+
 void IUsdStageEditorModule::FileNew() const
 {
 #if USE_USD_SDK
diff --git a/Engine/Plugins/Importers/USDImporter/Source/USDStageEditor/Private/Widgets/SUSDPrimPropertiesList.cpp b/Engine/Plugins/Importers/USDImporter/Source/USDStageEditor/Private/Widgets/SUSDPrimPropertiesList.cpp
--- a/Engine/Plugins/Importers/USDImporter/Source/USDStageEditor/Private/Widgets/SUSDPrimPropertiesList.cpp
+++ b/Engine/Plugins/Importers/USDImporter/Source/USDStageEditor/Private/Widgets/SUSDPrimPropertiesList.cpp
@@ -30,6 +30,7 @@ namespace UsdPrimPropertiesListConstants
 	const FMargin ComboBoxItemPadding( 3.0f, 0.0f, 2.0f, 0.0f );
 	const FMargin NumericEntryBoxItemPadding( 0.0f, 0.0f, 2.0f, 0.0f );
 	const float DesiredNumericEntryBoxWidth = 80.0f;
+	const float DesiredTextEntryBoxWidth = 80.0f;
 
 	const TCHAR* NormalFont = TEXT("PropertyWindow.NormalFont");
 }
@@ -227,6 +228,33 @@ public:
 
 protected:
 	FText GetLabel() const { return FText::FromString( UsdPrimAttribute->Label ); }
+	FText GetType() const
+	{
+		const static FText TypeTexts[] = {
+			FText::FromString(TEXT("M")),
+			FText::FromString(TEXT("A")),
+			FText::FromString(TEXT("R"))
+		};
+
+		const int32 Index = static_cast<int32>(UsdPrimAttribute->Type);
+		return Index >= 0 && Index < static_cast<int32>(EAttributeModelType::MAX)
+			? TypeTexts[Index]
+			: FText::GetEmpty();
+	}
+
+	FText GetTypeToolTip() const
+	{
+		const static FText ToolTipTexts[] = {
+			LOCTEXT("MetadataTooltip", "Metadata"),
+			LOCTEXT("AttributeTooltip", "Attribute"),
+			LOCTEXT("RelationshipTooltip", "Relationship")
+		};
+
+		const int32 Index = static_cast<int32>(UsdPrimAttribute->Type);
+		return Index >= 0 && Index < static_cast<int32>(EAttributeModelType::MAX)
+			? ToolTipTexts[Index]
+			: FText::GetEmpty();
+	}
 
 	// Optional here because SNumericEntryBox uses optional values as input
 	template<typename T>
@@ -283,7 +311,7 @@ private:
 	template<typename T>
 	TArray<TAttribute<TOptional<T>>> GetAttributeArray();
 
-	TSharedRef< SWidget > GenerateTextWidget( const TAttribute<FText>& Attribute );
+	TSharedRef< SWidget > GenerateTextWidget(const TAttribute<FText>& Attribute, const TAttribute<FText>& ToolTipAttribute = {}, ETextJustify::Type TextJustify = ETextJustify::Left);
 	TSharedRef< SWidget > GenerateEditableTextWidget( const TAttribute<FText>& Attribute, bool bIsReadOnly );
 	TSharedRef< SWidget > GenerateCheckboxWidget( const TAttribute<ECheckBoxState>& Attribute );
 
@@ -321,9 +349,17 @@ TSharedRef< SWidget > SUsdPrimPropertyRow::GenerateWidgetForColumn( const FName&
 
 	TSharedRef< SWidget > ColumnWidget = SNullWidget::NullWidget;
 
-	if ( ColumnName == TEXT("PropertyName") )
+	if (ColumnName == PrimAttributeColumnIds::TypeColumn)
+	{
+		ColumnWidget = GenerateTextWidget(
+			{this, &SUsdPrimPropertyRow::GetType},
+			{this, &SUsdPrimPropertyRow::GetTypeToolTip},
+			ETextJustify::Center
+		);
+	}
+	else if (ColumnName == PrimAttributeColumnIds::NameColumn)
 	{
-		ColumnWidget = GenerateTextWidget( { this, &SUsdPrimPropertyRow::GetLabel } );
+		ColumnWidget = GenerateTextWidget({this, &SUsdPrimPropertyRow::GetLabel});
 	}
 	else
 	{
@@ -366,6 +402,7 @@ TSharedRef< SWidget > SUsdPrimPropertyRow::GenerateWidgetForColumn( const FName&
 				SAssignNew( ColumnWidget, SBox )
 				.HeightOverride( FUsdStageEditorStyle::Get()->GetFloat( "UsdStageEditor.ListItemHeight" ) )
 				.VAlign( VAlign_Center )
+				.HAlign( HAlign_Left )
 				[
 					SNew( SComboBox< TSharedPtr< FString > > )
 					.OptionsSource( Options )
@@ -404,18 +441,11 @@ TSharedRef< SWidget > SUsdPrimPropertyRow::GenerateWidgetForColumn( const FName&
 		}
 	}
 
-	return SNew( SHorizontalBox )
-		+SHorizontalBox::Slot()
-		.HAlign( HAlign_Left )
-		.VAlign( VAlign_Fill )
-		.AutoWidth()
+	return SNew(SBox)
+		.VAlign(VAlign_Center)
+		.HAlign(HAlign_Fill)
 		[
-			SNew(SBox)
-			.MinDesiredWidth( UsdPrimPropertiesListConstants::DesiredNumericEntryBoxWidth )
-			.VAlign(VAlign_Center)
-			[
-				ColumnWidget
-			]
+			ColumnWidget
 		];
 }
 
@@ -441,28 +471,30 @@ TArray<TAttribute<TOptional<T>>> SUsdPrimPropertyRow::GetAttributeArray()
 	return Attributes;
 }
 
-TSharedRef< SWidget > SUsdPrimPropertyRow::GenerateTextWidget( const TAttribute<FText>& Attribute )
+TSharedRef< SWidget > SUsdPrimPropertyRow::GenerateTextWidget(
+	const TAttribute<FText>& Attribute,
+	const TAttribute<FText>& ToolTipAttribute,
+	ETextJustify::Type TextJustify
+)
 {
-	return SNew( SBox )
-		.HeightOverride( FUsdStageEditorStyle::Get()->GetFloat( "UsdStageEditor.ListItemHeight" ) )
-		.MinDesiredWidth( UsdPrimPropertiesListConstants::DesiredNumericEntryBoxWidth )
-		.VAlign( VAlign_Center )
-		[
-			SNew( STextBlock )
-			.Text( Attribute )
-			.Font( FAppStyle::GetFontStyle( "PropertyWindow.NormalFont" ) )
-			.Margin( UsdPrimPropertiesListConstants::RightRowPadding )
-		];
+	return SNew(STextBlock)
+		.Text(Attribute)
+		.Justification(TextJustify)
+		.ToolTipText(ToolTipAttribute)
+		.Font(FAppStyle::GetFontStyle("PropertyWindow.NormalFont"))
+		.Margin(UsdPrimPropertiesListConstants::RightRowPadding);
 }
 
 TSharedRef< SWidget > SUsdPrimPropertyRow::GenerateEditableTextWidget( const TAttribute<FText>& Attribute, bool bIsReadOnly )
 {
 	return SNew( SBox )
 		.HeightOverride( FUsdStageEditorStyle::Get()->GetFloat( "UsdStageEditor.ListItemHeight" ) )
-		.MinDesiredWidth( UsdPrimPropertiesListConstants::DesiredNumericEntryBoxWidth )
-		.VAlign( VAlign_Center )
+		.MinDesiredWidth( UsdPrimPropertiesListConstants::DesiredTextEntryBoxWidth )
+		.VAlign(VAlign_Center)
+		.HAlign(HAlign_Left)
 		[
 			SNew( SEditableTextBox )
+			.MinDesiredWidth( UsdPrimPropertiesListConstants::DesiredTextEntryBoxWidth )
 			.Text( Attribute )
 			.IsReadOnly( bIsReadOnly )
 			.Font( FAppStyle::GetFontStyle( "PropertyWindow.NormalFont" ) )
@@ -477,6 +509,7 @@ TSharedRef< SWidget > SUsdPrimPropertyRow::GenerateCheckboxWidget( const TAttrib
 	return SNew( SBox )
 		.HeightOverride( FUsdStageEditorStyle::Get()->GetFloat( "UsdStageEditor.ListItemHeight" ) )
 		.VAlign( VAlign_Center )
+		.HAlign(HAlign_Left)
 		[
 			SNew( SCheckBox )
 			.IsChecked( Attribute )
@@ -561,7 +594,13 @@ TSharedRef< SWidget > SUsdPrimPropertyRow::GenerateSpinboxWidgets( const TArray<
 			];
 	}
 
-	return VertBox.ToSharedRef();
+	return SNew(SBox)
+		.MinDesiredWidth(UsdPrimPropertiesListConstants::DesiredNumericEntryBoxWidth)
+		.VAlign(VAlign_Center)
+		.HAlign(HAlign_Left)
+		[
+			VertBox.ToSharedRef()
+		];
 }
 
 template<typename T>
@@ -658,13 +697,21 @@ void SUsdPrimPropertiesList::Construct( const FArguments& InArgs )
 
 	SAssignNew( HeaderRowWidget, SHeaderRow )
 
-	+SHeaderRow::Column( FName( TEXT("PropertyName") ) )
-	.DefaultLabel( LOCTEXT( "PropertyName", "Property Name" ) )
-	.FillWidth( 25.f )
+	+SHeaderRow::Column(PrimAttributeColumnIds::TypeColumn)
+	.DefaultLabel(FText::GetEmpty())
+	.FixedWidth(24.0f)
+	.SortMode(this, &SUsdPrimPropertiesList::GetColumnSortMode, PrimAttributeColumnIds::TypeColumn)
+	.OnSort(this, &SUsdPrimPropertiesList::Sort)
 
-	+SHeaderRow::Column( FName( TEXT("PropertyValue") ) )
+	+SHeaderRow::Column(PrimAttributeColumnIds::NameColumn)
+	.DefaultLabel(InArgs._NameColumnText)
+	.FillWidth( 40.f )
+	.SortMode(this, &SUsdPrimPropertiesList::GetColumnSortMode, PrimAttributeColumnIds::NameColumn)
+	.OnSort(this, &SUsdPrimPropertiesList::Sort)
+
+	+SHeaderRow::Column(PrimAttributeColumnIds::ValueColumn)
 	.DefaultLabel( LOCTEXT( "PropertyValue", "Value" ) )
-	.FillWidth( 75.f );
+	.FillWidth( 60.f );
 
 	SListView::Construct
 	(
@@ -673,6 +720,8 @@ void SUsdPrimPropertiesList::Construct( const FArguments& InArgs )
 		.OnGenerateRow( this, &SUsdPrimPropertiesList::OnGenerateRow )
 		.HeaderRow( HeaderRowWidget )
 	);
+
+	OnSelectionChanged = InArgs._OnSelectionChanged;
 }
 
 TSharedRef< ITableRow > SUsdPrimPropertiesList::OnGenerateRow( TSharedPtr< FUsdPrimAttributeViewModel > InDisplayNode, const TSharedRef< STableViewBase >& OwnerTable )
@@ -686,9 +735,35 @@ void SUsdPrimPropertiesList::GeneratePropertiesList( const UE::FUsdStageWeak& Us
 	ViewModel.Refresh( UsdStage, InPrimPath, TimeCode );
 }
 
+void SUsdPrimPropertiesList::Sort(
+	const EColumnSortPriority::Type SortPriority,
+	const FName& ColumnId,
+	const EColumnSortMode::Type NewSortMode
+)
+{
+	ViewModel.CurrentSortColumn = ColumnId;
+	ViewModel.CurrentSortMode = NewSortMode;
+	ViewModel.Sort();
+
+	RequestListRefresh();
+}
+
+EColumnSortMode::Type SUsdPrimPropertiesList::GetColumnSortMode(const FName ColumnId) const
+{
+	if (ViewModel.CurrentSortColumn != ColumnId)
+	{
+		return EColumnSortMode::None;
+	}
+
+	return ViewModel.CurrentSortMode;
+}
+
 void SUsdPrimPropertiesList::SetPrimPath( const UE::FUsdStageWeak& UsdStage, const TCHAR* InPrimPath )
 {
 	GeneratePropertiesList( UsdStage, InPrimPath );
+
+	SetVisibility(ViewModel.PrimAttributes.Num() > 0 ? EVisibility::Visible : EVisibility::Collapsed);
+
 	RequestListRefresh();
 }
 
@@ -731,6 +806,16 @@ void SUsdPrimPropertiesList::SetSelectedPropertyNames( const TArray<FString>& Ne
 	SetItemSelection( NewItemSelection, bSelected );
 }
 
+UE::FUsdStageWeak SUsdPrimPropertiesList::GetUsdStage() const
+{
+	return ViewModel.GetUsdStage();
+}
+
+FString SUsdPrimPropertiesList::GetPrimPath() const
+{
+	return ViewModel.GetPrimPath();
+}
+
 #undef LOCTEXT_NAMESPACE
 
 #endif // #if USE_USD_SDK
diff --git a/Engine/Plugins/Importers/USDImporter/Source/USDStageEditor/Private/Widgets/SUSDPrimPropertiesList.h b/Engine/Plugins/Importers/USDImporter/Source/USDStageEditor/Private/Widgets/SUSDPrimPropertiesList.h
--- a/Engine/Plugins/Importers/USDImporter/Source/USDStageEditor/Private/Widgets/SUSDPrimPropertiesList.h
+++ b/Engine/Plugins/Importers/USDImporter/Source/USDStageEditor/Private/Widgets/SUSDPrimPropertiesList.h
@@ -7,9 +7,13 @@
 
 #if USE_USD_SDK
 
+
 class SUsdPrimPropertiesList : public SListView< TSharedPtr< FUsdPrimAttributeViewModel > >
 {
+public:
 	SLATE_BEGIN_ARGS( SUsdPrimPropertiesList ) {}
+		SLATE_EVENT(FOnSelectionChanged, OnSelectionChanged)
+		SLATE_ATTRIBUTE(FText, NameColumnText)
 	SLATE_END_ARGS()
 
 public:
@@ -19,10 +23,16 @@ public:
 	TArray<FString> GetSelectedPropertyNames() const;
 	void SetSelectedPropertyNames( const TArray<FString>& NewSelection );
 
+	UE::FUsdStageWeak GetUsdStage() const;
+	FString GetPrimPath() const;
+
 protected:
 	TSharedRef< ITableRow > OnGenerateRow( TSharedPtr< FUsdPrimAttributeViewModel > InDisplayNode, const TSharedRef< STableViewBase >& OwnerTable );
 	void GeneratePropertiesList( const UE::FUsdStageWeak& UsdStage, const TCHAR* InPrimPath );
 
+	void Sort(const EColumnSortPriority::Type SortPriority, const FName& ColumnId, const EColumnSortMode::Type NewSortMode);
+	EColumnSortMode::Type GetColumnSortMode(const FName ColumnId) const;
+
 private:
 	FUsdPrimAttributesViewModel ViewModel;
 	TSharedPtr< SHeaderRow > HeaderRowWidget;
diff --git a/Engine/Plugins/Importers/USDImporter/Source/USDStageEditor/Public/USDStageEditorBlueprintLibrary.h b/Engine/Plugins/Importers/USDImporter/Source/USDStageEditor/Public/USDStageEditorBlueprintLibrary.h
--- a/Engine/Plugins/Importers/USDImporter/Source/USDStageEditor/Public/USDStageEditorBlueprintLibrary.h
+++ b/Engine/Plugins/Importers/USDImporter/Source/USDStageEditor/Public/USDStageEditorBlueprintLibrary.h
@@ -94,6 +94,22 @@ public:
 	UFUNCTION( BlueprintCallable, Category = "USD|Selection" )
 	static void SetSelectedPropertyNames( const TArray<FString>& NewSelection );
 
+	/**
+	 * Returns the names of the currently selected property metadata entries on the right panel of the USD Stage
+	 * Editor.
+	 * @return The names of selected metadata (e.g. ["documentation", "typeName"])
+	 */
+	UFUNCTION( BlueprintCallable, Category = "USD|Selection" )
+	static TArray<FString> GetSelectedPropertyMetadataNames();
+
+	/**
+	 * Sets the USD Stage Editor property metadata selection to the entries with names contained in NewSelection.
+	 * Provide an empty array to clear the selection.
+	 * @param NewSelection - The list of property names to select (e.g. ["documentation", "typeName"])
+	 */
+	UFUNCTION( BlueprintCallable, Category = "USD|Selection" )
+	static void SetSelectedPropertyMetadataNames( const TArray<FString>& NewSelection );
+
     /**
 	 * Creates a new memory-only layer and opens an USD Stage with that layer as its root.
 	 * Corresponds to the "File -> New" action on the USD Stage Editor menu bar.
diff --git a/Engine/Plugins/Importers/USDImporter/Source/USDStageEditor/Public/USDStageEditorModule.h b/Engine/Plugins/Importers/USDImporter/Source/USDStageEditor/Public/USDStageEditorModule.h
--- a/Engine/Plugins/Importers/USDImporter/Source/USDStageEditor/Public/USDStageEditorModule.h
+++ b/Engine/Plugins/Importers/USDImporter/Source/USDStageEditor/Public/USDStageEditorModule.h
@@ -38,6 +38,9 @@ public:
 	TArray<FString> GetSelectedPropertyNames() const;
 	void SetSelectedPropertyNames( const TArray<FString>& NewSelection ) const;
 
+	TArray<FString> GetSelectedPropertyMetadataNames() const;
+	void SetSelectedPropertyMetadataNames(const TArray<FString>& NewSelection) const;
+
 	// For all of these, providing an empty path will cause us to pop open a dialog to let the user pick the path
 	// instead.
 	void FileNew() const;
diff --git a/Engine/Plugins/Importers/USDImporter/Source/USDStageEditorViewModels/Private/USDPrimAttributesViewModel.cpp b/Engine/Plugins/Importers/USDImporter/Source/USDStageEditorViewModels/Private/USDPrimAttributesViewModel.cpp
--- a/Engine/Plugins/Importers/USDImporter/Source/USDStageEditorViewModels/Private/USDPrimAttributesViewModel.cpp
+++ b/Engine/Plugins/Importers/USDImporter/Source/USDStageEditorViewModels/Private/USDPrimAttributesViewModel.cpp
@@ -28,8 +28,10 @@
 	#include "pxr/usd/kind/registry.h"
 	#include "pxr/usd/sdf/types.h"
 	#include "pxr/usd/usd/attribute.h"
+	#include "pxr/usd/usd/common.h"
 	#include "pxr/usd/usd/editContext.h"
 	#include "pxr/usd/usd/prim.h"
+	#include "pxr/usd/usd/relationship.h"
 	#include "pxr/usd/usdGeom/tokens.h"
 #include "USDIncludesEnd.h"
 
@@ -106,13 +108,21 @@ void FUsdPrimAttributeViewModel::SetAttributeValue( const UsdUtils::FConvertedVt
 }
 
 template<typename T>
-void FUsdPrimAttributesViewModel::CreatePrimAttribute( const FString& AttributeName, const T& Value, UsdUtils::EUsdBasicDataTypes SourceType, const FString& ValueRole, bool bReadOnly )
+void FUsdPrimAttributesViewModel::CreatePrimAttribute(
+	EAttributeModelType Type,
+	const FString& AttributeName,
+	const T& Value,
+	UsdUtils::EUsdBasicDataTypes SourceType,
+	const FString& ValueRole,
+	bool bReadOnly
+)
 {
 	UsdUtils::FConvertedVtValue VtValue;
 	VtValue.Entries = { { UsdUtils::FConvertedVtValueComponent{ TInPlaceType<T>(), Value } } };
 	VtValue.SourceType = SourceType;
 
 	FUsdPrimAttributeViewModel Property( this );
+	Property.Type = Type;
 	Property.Label = AttributeName;
 	Property.Value = VtValue;
 	Property.ValueRole = ValueRole;
@@ -121,9 +131,15 @@ void FUsdPrimAttributesViewModel::CreatePrimAttribute( const FString& AttributeN
 	PrimAttributes.Add( MakeSharedUnreal< FUsdPrimAttributeViewModel >( MoveTemp( Property ) ) );
 }
 
-void FUsdPrimAttributesViewModel::CreatePrimAttribute( const FString& AttributeName, const UsdUtils::FConvertedVtValue& Value, bool bReadOnly )
+void FUsdPrimAttributesViewModel::CreatePrimAttribute(
+	EAttributeModelType Type,
+	const FString& AttributeName,
+	const UsdUtils::FConvertedVtValue& Value,
+	bool bReadOnly
+)
 {
 	FUsdPrimAttributeViewModel Property( this );
+	Property.Type = Type;
 	Property.Label = AttributeName;
 	Property.Value = Value;
 	Property.bReadOnly = bReadOnly;
@@ -143,12 +159,13 @@ void FUsdPrimAttributesViewModel::SetPrimAttribute( const FString& AttributeName
 
 	// Transact here as setting this attribute may trigger USD events that affect assets/components
 	FScopedTransaction Transaction( FText::Format(
-		LOCTEXT( "SetPrimAttribute", "Set value for attribute '{0}' of prim '{1}'" ),
+		LOCTEXT( "SetPrimAttribute", "Set value for field '{0}' of prim '{1}'" ),
 		FText::FromString( AttributeName ),
 		FText::FromString( PrimPath )
 	));
 
 	const bool bIsStageAttribute = PrimPath == TEXT( "/" ) || PrimPath.IsEmpty();
+	bool bIsPropertyPath = false;
 
 	UE::FVtValue VtValue;
 	if ( UnrealToUsd::ConvertValue( InValue, VtValue ) )
@@ -164,6 +181,11 @@ void FUsdPrimAttributesViewModel::SetPrimAttribute( const FString& AttributeName
 		}
 		else if ( UE::FUsdPrim UsdPrim = UsdStage.GetPrimAtPath( UE::FSdfPath( *PrimPath ) ) )
 		{
+			FScopedUsdAllocs UsdAllocs;
+
+			pxr::UsdPrim PxrUsdPrim{UsdPrim};
+			pxr::TfToken AttributeNameToken = UnrealToUsd::ConvertToken(*AttributeName).Get();
+
 			// Single value, single component of FString
 			if ( AttributeName == TEXT( "kind" ) && InValue.Entries.Num() == 1 && InValue.Entries[ 0 ].Num() == 1 && InValue.Entries[ 0 ][ 0 ].IsType<FString>() )
 			{
@@ -177,14 +199,56 @@ void FUsdPrimAttributesViewModel::SetPrimAttribute( const FString& AttributeName
 				bSuccess = Attribute.Set( VtValue );
 				UsdUtils::NotifyIfOverriddenOpinion( Attribute );
 			}
+			else if (pxr::UsdRelationship Relationship = PxrUsdPrim.GetRelationship(AttributeNameToken))
+			{
+				if (InValue.Entries.Num() == 1 && InValue.Entries[0].Num() == 1 && InValue.Entries[0][0].IsType<FString>())
+				{
+					const FString& NewValue = InValue.Entries[0][0].Get<FString>();
+
+					// We only allow editing relationship attributes that have a single target anyway
+					bSuccess = Relationship.SetTargets(std::vector<pxr::SdfPath>{UnrealToUsd::ConvertPath(*NewValue).Get()});
+					UsdUtils::NotifyIfOverriddenOpinion(Relationship);
+				}
+			}
+			else if (PxrUsdPrim.HasMetadata(AttributeNameToken))
+			{
+				bSuccess = PxrUsdPrim.SetMetadata(AttributeNameToken, VtValue.GetUsdValue());
+			}
+		}
+		// Trying to set property metadata
+		else if (!PrimPath.IsEmpty())
+		{
+			FScopedUsdAllocs Allocs;
+
+			pxr::SdfPath UsdPath = UnrealToUsd::ConvertPath(*PrimPath).Get();
+			if (UsdPath.IsPropertyPath())
+			{
+				bIsPropertyPath = true;
+
+				pxr::SdfPath PxrPrimPath = UsdPath.GetPrimPath();
+				pxr::TfToken PropertyName = UsdPath.GetNameToken();
+
+				pxr::UsdStageRefPtr PxrUsdStage{UsdStage};
+				if (pxr::UsdPrim Prim = PxrUsdStage->GetPrimAtPath(PxrPrimPath))
+				{
+					if (pxr::UsdProperty Property = Prim.GetProperty(PropertyName))
+					{
+						pxr::TfToken AttributeNameToken = UnrealToUsd::ConvertToken(*AttributeName).Get();
+						if (Property.HasMetadata(AttributeNameToken))
+						{
+							bSuccess = Property.SetMetadata(AttributeNameToken, VtValue.GetUsdValue());
+						}
+					}
+				}
+			}
 		}
 	}
 
 	if ( !bSuccess )
 	{
-		const FText ErrorMessage = FText::Format( LOCTEXT( "FailToSetAttributeMessage", "Failed to set attribute '{0}' on {1} '{2}'" ),
+		const FText ErrorMessage = FText::Format( LOCTEXT( "FailToSetAttributeMessage", "Failed to set '{0}' on {1} '{2}'" ),
 			FText::FromString( AttributeName ),
-			FText::FromString( bIsStageAttribute ? TEXT( "stage" ) : TEXT( "prim" ) ),
+			FText::FromString( bIsStageAttribute ? TEXT( "stage" ) : bIsPropertyPath ? TEXT("property") : TEXT("prim")),
 			FText::FromString( bIsStageAttribute ? UsdStage.GetRootLayer().GetRealPath() : PrimPath )
 		);
 
@@ -209,6 +273,70 @@ void FUsdPrimAttributesViewModel::Refresh( const UE::FUsdStageWeak& InUsdStage,
 	PrimAttributes.Reset();
 
 #if USE_USD_SDK
+	TFunction<void(const FString& MetadataName, const pxr::VtValue& Value)> DisplayEntryAsString =
+		[this](const FString& MetadataName, const pxr::VtValue& Value)
+		{
+			const bool bAttrReadOnly = true;
+			FString Stringified = UsdToUnreal::ConvertString(pxr::TfStringify(Value));
+
+			// STextBlock can get very slow calculating its desired size for very long string so chop it if needed
+			const int32 MaxValueLength = 300;
+			if (Stringified.Len() > MaxValueLength)
+			{
+				Stringified.LeftInline(MaxValueLength);
+				Stringified.Append(TEXT("..."));
+			}
+
+			CreatePrimAttribute(
+				EAttributeModelType::Metadata,
+				MetadataName,
+				Stringified,
+				UsdUtils::EUsdBasicDataTypes::String,
+				TEXT(""),
+				bAttrReadOnly
+			);
+		};
+
+	// Lambda to add rows to the model based on USD object metadata. Extracted here because we can reuse it for
+	// properties and prims
+	TFunction<void(pxr::UsdObject)> AddMetadataEntries = [this, &DisplayEntryAsString](pxr::UsdObject Object)
+	{
+		std::map<class pxr::TfToken, pxr::VtValue, pxr::TfDictionaryLessThan> MetadataMap = Object.GetAllMetadata();
+		for (std::map<class pxr::TfToken, pxr::VtValue, pxr::TfDictionaryLessThan>::iterator Iter = MetadataMap.begin();
+			 Iter != MetadataMap.end();
+			 ++Iter)
+		{
+			const pxr::TfToken& Key = Iter->first;
+			const pxr::VtValue& Value = Iter->second;
+
+			FString MetadataName = UsdToUnreal::ConvertToken(Key);
+
+			UsdUtils::FConvertedVtValue ConvertedValue;
+			if (UsdToUnreal::ConvertValue(UE::FVtValue{Value}, ConvertedValue))
+			{
+				if (Value.IsArrayValued())
+				{
+					DisplayEntryAsString(MetadataName, Value);
+				}
+				else
+				{
+					const bool bAttrReadOnly = false;
+					CreatePrimAttribute(
+						EAttributeModelType::Metadata,
+						MetadataName,
+						ConvertedValue,
+						bAttrReadOnly
+					);
+				}
+			}
+			// If we don't know how to display this type, try stringifying it
+			else
+			{
+				DisplayEntryAsString(MetadataName, Value);
+			}
+		}
+	};
+
 	if ( UsdStage )
 	{
 		// Show info about the stage
@@ -216,7 +344,14 @@ void FUsdPrimAttributesViewModel::Refresh( const UE::FUsdStageWeak& InUsdStage,
 		{
 			const bool bReadOnly = true;
 			const FString Role = TEXT( "" );
-			CreatePrimAttribute( TEXT( "path" ), UsdStage.GetRootLayer().GetRealPath(), UsdUtils::EUsdBasicDataTypes::String, Role, bReadOnly );
+			CreatePrimAttribute(
+				EAttributeModelType::Metadata,
+				TEXT( "path" ),
+				UsdStage.GetRootLayer().GetRealPath(),
+				UsdUtils::EUsdBasicDataTypes::String,
+				Role,
+				bReadOnly
+			);
 
 			FScopedUsdAllocs UsdAllocs;
 
@@ -235,7 +370,12 @@ void FUsdPrimAttributesViewModel::Refresh( const UE::FUsdStageWeak& InUsdStage,
 					}
 
 					const bool bAttrReadOnly = ConvertedValue.bIsArrayValued;
-					CreatePrimAttribute( AttributeName, ConvertedValue, bAttrReadOnly );
+					CreatePrimAttribute(
+						EAttributeModelType::Metadata,
+						AttributeName,
+						ConvertedValue,
+						bAttrReadOnly
+					);
 				}
 			}
 
@@ -245,14 +385,35 @@ void FUsdPrimAttributesViewModel::Refresh( const UE::FUsdStageWeak& InUsdStage,
 		{
 			// For now we can't rename/reparent prims through this
 			const bool bReadOnly = true;
-			const FString Role = TEXT( "" );
-			CreatePrimAttribute( TEXT( "name" ), UsdPrim.GetName().ToString(), UsdUtils::EUsdBasicDataTypes::String, Role, bReadOnly );
-			CreatePrimAttribute( TEXT( "path" ), FString( InPrimPath ), UsdUtils::EUsdBasicDataTypes::String, Role, bReadOnly );
-			CreatePrimAttribute( TEXT( "kind" ), UsdToUnreal::ConvertString( IUsdPrim::GetKind( UsdPrim ).GetString() ), UsdUtils::EUsdBasicDataTypes::Token );
+			const FString Role = TEXT("");
+			CreatePrimAttribute(
+				EAttributeModelType::Metadata,
+				TEXT("name"),
+				UsdPrim.GetName().ToString(),
+				UsdUtils::EUsdBasicDataTypes::String,
+				Role,
+				bReadOnly
+			);
+			CreatePrimAttribute(
+				EAttributeModelType::Metadata,
+				TEXT("path"),
+				FString(InPrimPath),
+				UsdUtils::EUsdBasicDataTypes::String,
+				Role,
+				bReadOnly
+			);
+			CreatePrimAttribute(
+				EAttributeModelType::Metadata,
+				TEXT("kind"),
+				UsdToUnreal::ConvertString(IUsdPrim::GetKind(UsdPrim).GetString()),
+				UsdUtils::EUsdBasicDataTypes::Token
+			);
 
 			FScopedUsdAllocs UsdAllocs;
 
-			for ( const pxr::UsdAttribute& PrimAttribute : pxr::UsdPrim( UsdPrim ).GetAttributes() )
+			pxr::UsdPrim PxrUsdPrim{UsdPrim};
+
+			for ( const pxr::UsdAttribute& PrimAttribute : PxrUsdPrim.GetAttributes() )
 			{
 				FString AttributeName = UsdToUnreal::ConvertString( PrimAttribute.GetName().GetString() );
 
@@ -273,7 +434,14 @@ void FUsdPrimAttributesViewModel::Refresh( const UE::FUsdStageWeak& InUsdStage,
 						Stringified.Append( TEXT( "..." ) );
 					}
 
-					CreatePrimAttribute( AttributeName, Stringified, UsdUtils::EUsdBasicDataTypes::String, TEXT( "" ), bAttrReadOnly );
+					CreatePrimAttribute(
+						EAttributeModelType::Attribute,
+						AttributeName,
+						Stringified,
+						UsdUtils::EUsdBasicDataTypes::String,
+						TEXT(""),
+						bAttrReadOnly
+					);
 				}
 				else
 				{
@@ -281,7 +449,12 @@ void FUsdPrimAttributesViewModel::Refresh( const UE::FUsdStageWeak& InUsdStage,
 					if ( UsdToUnreal::ConvertValue( UE::FVtValue{ VtValue }, ConvertedValue ) )
 					{
 						const bool bAttrReadOnly = false;
-						CreatePrimAttribute( AttributeName, ConvertedValue, bAttrReadOnly );
+						CreatePrimAttribute(
+							EAttributeModelType::Attribute,
+							AttributeName,
+							ConvertedValue,
+							bAttrReadOnly
+						);
 					}
 
 					if ( PrimAttribute.HasAuthoredConnections() )
@@ -301,31 +474,159 @@ void FUsdPrimAttributesViewModel::Refresh( const UE::FUsdStageWeak& InUsdStage,
 							ConnectionPropertyValue.Entries = { Entry };
 
 							const bool bConnectionValueReadOnly = true;
-							CreatePrimAttribute( ConnectionAttributeName, ConnectionPropertyValue, bConnectionValueReadOnly );
+							CreatePrimAttribute(
+								EAttributeModelType::Attribute,
+								ConnectionAttributeName,
+								ConnectionPropertyValue,
+								bConnectionValueReadOnly
+							);
 						}
 					}
 				}
 			}
+
+			for (const pxr::UsdRelationship& Relationship : PxrUsdPrim.GetRelationships())
+			{
+				FString RelationshipName = UsdToUnreal::ConvertString(Relationship.GetName().GetString());
+
+				std::vector<pxr::SdfPath> Targets;
+				if (Relationship.GetTargets(&Targets))
+				{
+					if (Targets.size() == 1)
+					{
+						FString UETarget = UsdToUnreal::ConvertPath(Targets[0]);
+						const bool bAttrReadOnly = false;
+						CreatePrimAttribute(
+							EAttributeModelType::Relationship,
+							RelationshipName,
+							UETarget,
+							UsdUtils::EUsdBasicDataTypes::String,
+							TEXT(""),
+							bAttrReadOnly
+						);
+					}
+					else if (Targets.size() > 1)
+					{
+						FString CombinedTargets;
+						for (const pxr::SdfPath& Target : Targets)
+						{
+							CombinedTargets += UsdToUnreal::ConvertPath(Target) + TEXT(", ");
+						}
+						CombinedTargets.RemoveFromEnd(TEXT(", "));
+
+						// STextBlock can get very slow calculating its desired size for very long string so chop it if needed
+						const int32 MaxValueLength = 300;
+						if (CombinedTargets.Len() > MaxValueLength)
+						{
+							CombinedTargets.LeftInline(MaxValueLength);
+							CombinedTargets.Append(TEXT("..."));
+						}
+
+						const bool bAttrReadOnly = true;
+						CreatePrimAttribute(
+							EAttributeModelType::Relationship,
+							RelationshipName,
+							CombinedTargets,
+							UsdUtils::EUsdBasicDataTypes::String,
+							TEXT(""),
+							bAttrReadOnly
+						);
+					}
+				}
+			}
+
+			AddMetadataEntries(PxrUsdPrim);
+		}
+		// It's a property path?
+		else if (!PrimPath.IsEmpty())
+		{
+			FScopedUsdAllocs Allocs;
+
+			pxr::SdfPath UsdPath = UnrealToUsd::ConvertPath(*PrimPath).Get();
+			if (UsdPath.IsPropertyPath())
+			{
+				pxr::SdfPath PxrPrimPath = UsdPath.GetPrimPath();
+				pxr::TfToken PropertyName = UsdPath.GetNameToken();
+
+				pxr::UsdStageRefPtr PxrUsdStage{UsdStage};
+				if (pxr::UsdPrim Prim = PxrUsdStage->GetPrimAtPath(PxrPrimPath))
+				{
+					if (pxr::UsdProperty Property = Prim.GetProperty(PropertyName))
+					{
+						AddMetadataEntries(Property);
+					}
+				}
+			}
 		}
 	}
 
-	PrimAttributes.Sort([](const TSharedPtr<FUsdPrimAttributeViewModel>& A, const TSharedPtr<FUsdPrimAttributeViewModel>& B)
-	{
-		return A->Label < B->Label;
-	});
+	Sort();
 #endif // #if USE_USD_SDK
 }
 
-// One for each variant of FPrimPropertyValue. These should all be implicitly instantiated from the above code anyway, but just in case that changes somehow
-template void FUsdPrimAttributesViewModel::CreatePrimAttribute( const FString&, const bool&,	UsdUtils::EUsdBasicDataTypes, const FString&, bool );
-template void FUsdPrimAttributesViewModel::CreatePrimAttribute( const FString&, const uint8&,	UsdUtils::EUsdBasicDataTypes, const FString&, bool );
-template void FUsdPrimAttributesViewModel::CreatePrimAttribute( const FString&, const int32&,	UsdUtils::EUsdBasicDataTypes, const FString&, bool );
-template void FUsdPrimAttributesViewModel::CreatePrimAttribute( const FString&, const uint32&,	UsdUtils::EUsdBasicDataTypes, const FString&, bool );
-template void FUsdPrimAttributesViewModel::CreatePrimAttribute( const FString&, const int64&,	UsdUtils::EUsdBasicDataTypes, const FString&, bool );
-template void FUsdPrimAttributesViewModel::CreatePrimAttribute( const FString&, const uint64&,	UsdUtils::EUsdBasicDataTypes, const FString&, bool );
-template void FUsdPrimAttributesViewModel::CreatePrimAttribute( const FString&, const float&,	UsdUtils::EUsdBasicDataTypes, const FString&, bool );
-template void FUsdPrimAttributesViewModel::CreatePrimAttribute( const FString&, const double&,	UsdUtils::EUsdBasicDataTypes, const FString&, bool );
-template void FUsdPrimAttributesViewModel::CreatePrimAttribute( const FString&, const FString&, UsdUtils::EUsdBasicDataTypes, const FString&, bool );
+void FUsdPrimAttributesViewModel::Sort()
+{
+	// When sorting by type, we'll sort according to the types first, but still sort alphabetically within each type
+	if (CurrentSortColumn == PrimAttributeColumnIds::TypeColumn)
+	{
+		if (CurrentSortMode == EColumnSortMode::Ascending)
+		{
+			PrimAttributes.Sort([](const TSharedPtr<FUsdPrimAttributeViewModel>& A, const TSharedPtr<FUsdPrimAttributeViewModel>& B)
+			{
+				if(A->Type == B->Type)
+				{
+					return A->Label < B->Label;
+				}
+				else
+				{
+					return A->Type < B->Type;
+				}
+			});
+		}
+		else
+		{
+			PrimAttributes.Sort([](const TSharedPtr<FUsdPrimAttributeViewModel>& A, const TSharedPtr<FUsdPrimAttributeViewModel>& B)
+			{
+				if(A->Type == B->Type)
+				{
+					return A->Label < B->Label;
+				}
+				else
+				{
+					return A->Type > B->Type;
+				}
+			});
+		}
+	}
+	// When sorting by name we'll sort alphabetically regardless of type
+	else if (CurrentSortColumn == PrimAttributeColumnIds::NameColumn)
+	{
+		if (CurrentSortMode == EColumnSortMode::Ascending)
+		{
+			PrimAttributes.Sort([](const TSharedPtr<FUsdPrimAttributeViewModel>& A, const TSharedPtr<FUsdPrimAttributeViewModel>& B)
+			{
+				return A->Label < B->Label;
+			});
+		}
+		else
+		{
+			PrimAttributes.Sort([](const TSharedPtr<FUsdPrimAttributeViewModel>& A, const TSharedPtr<FUsdPrimAttributeViewModel>& B)
+			{
+				return A->Label > B->Label;
+			});
+		}
+	}
+}
+
+UE::FUsdStageWeak FUsdPrimAttributesViewModel::GetUsdStage() const
+{
+	return UsdStage;
+}
+
+FString FUsdPrimAttributesViewModel::GetPrimPath() const
+{
+	return PrimPath;
+}
 
 #undef LOCTEXT_NAMESPACE
 
diff --git a/Engine/Plugins/Importers/USDImporter/Source/USDStageEditorViewModels/Public/USDPrimAttributesViewModel.h b/Engine/Plugins/Importers/USDImporter/Source/USDStageEditorViewModels/Public/USDPrimAttributesViewModel.h
--- a/Engine/Plugins/Importers/USDImporter/Source/USDStageEditorViewModels/Public/USDPrimAttributesViewModel.h
+++ b/Engine/Plugins/Importers/USDImporter/Source/USDStageEditorViewModels/Public/USDPrimAttributesViewModel.h
@@ -7,9 +7,25 @@
 #include "UsdWrappers/UsdStage.h"
 
 #include "Templates/SharedPointer.h"
+#include "Widgets/Views/SHeaderRow.h"
 
 class FUsdPrimAttributesViewModel;
 
+enum class EAttributeModelType : int32
+{
+	Metadata = 0,
+	Attribute = 1,
+	Relationship = 2,
+	MAX = 3
+};
+
+namespace PrimAttributeColumnIds
+{
+	inline const FName TypeColumn = TEXT("PropertyType");
+	inline const FName NameColumn = TEXT("PropertyName");
+	inline const FName ValueColumn = TEXT("PropertyValue");
+};
+
 class USDSTAGEEDITORVIEWMODELS_API FUsdPrimAttributeViewModel : public TSharedFromThis< FUsdPrimAttributeViewModel >
 {
 public:
@@ -20,6 +36,7 @@ public:
 	void SetAttributeValue( const UsdUtils::FConvertedVtValue& InValue );
 
 public:
+	EAttributeModelType Type;
 	FString Label;
 	UsdUtils::FConvertedVtValue Value;
 	FString ValueRole;
@@ -34,14 +51,20 @@ class USDSTAGEEDITORVIEWMODELS_API FUsdPrimAttributesViewModel
 {
 public:
 	template<typename T>
-	void CreatePrimAttribute( const FString& AttributeName, const T& Value, UsdUtils::EUsdBasicDataTypes UsdType, const FString& ValueRole = FString(), bool bReadOnly = false );
-	void CreatePrimAttribute( const FString& AttributeName, const UsdUtils::FConvertedVtValue& Value, bool bReadOnly = false );
+	void CreatePrimAttribute(EAttributeModelType Type, const FString& AttributeName, const T& Value, UsdUtils::EUsdBasicDataTypes UsdType, const FString& ValueRole = FString(), bool bReadOnly = false);
+	void CreatePrimAttribute(EAttributeModelType Type, const FString& AttributeName, const UsdUtils::FConvertedVtValue& Value, bool bReadOnly = false);
 
 	void SetPrimAttribute( const FString& AttributeName, const UsdUtils::FConvertedVtValue& Value );
 
 	void Refresh( const UE::FUsdStageWeak& UsdStage, const TCHAR* PrimPath, float TimeCode );
+	void Sort();
+
+	UE::FUsdStageWeak GetUsdStage() const;
+	FString GetPrimPath() const;
 
 public:
+	EColumnSortMode::Type CurrentSortMode = EColumnSortMode::Ascending;
+	FName CurrentSortColumn = PrimAttributeColumnIds::TypeColumn;
 	TArray< TSharedPtr< FUsdPrimAttributeViewModel > > PrimAttributes;
 
 private:
diff --git a/Engine/Plugins/Importers/USDImporter/Source/USDUtilities/Private/USDAttributeUtils.cpp b/Engine/Plugins/Importers/USDImporter/Source/USDUtilities/Private/USDAttributeUtils.cpp
--- a/Engine/Plugins/Importers/USDImporter/Source/USDUtilities/Private/USDAttributeUtils.cpp
+++ b/Engine/Plugins/Importers/USDImporter/Source/USDUtilities/Private/USDAttributeUtils.cpp
@@ -186,15 +186,29 @@ void UsdUtils::NotifyIfOverriddenOpinion( const UE::FUsdAttribute& Attribute )
 		return;
 	}
 
-	pxr::UsdStageRefPtr Stage = UsdAttribute.GetPrim().GetStage();
-	if ( !Stage )
+	NotifyIfOverriddenOpinion(UsdAttribute);
+#endif // #if USE_USD_SDK
+}
+
+#if USE_USD_SDK
+void UsdUtils::NotifyIfOverriddenOpinion(const pxr::UsdProperty& Property)
+{
+	FScopedUsdAllocs Allocs;
+
+	if (!Property)
+	{
+		return;
+	}
+
+	pxr::UsdStageRefPtr Stage = Property.GetPrim().GetStage();
+	if (!Stage)
 	{
 		return;
 	}
 
 	const pxr::UsdEditTarget& EditTarget = Stage->GetEditTarget();
 	const pxr::SdfLayerHandle& Layer = EditTarget.GetLayer();
-	if ( !Layer )
+	if (!Layer)
 	{
 		return;
 	}
@@ -207,36 +221,36 @@ void UsdUtils::NotifyIfOverriddenOpinion( const UE::FUsdAttribute& Attribute )
 	// References:
 	// - https://graphics.pixar.com/usd/release/glossary.html#livrps-strength-ordering
 	// - https://groups.google.com/g/usd-interest/c/xTxFYQA_bRs/m/qbGkvx3yAgAJ
-	std::vector<pxr::SdfPropertySpecHandle> SpecStack = UsdAttribute.GetPropertyStack();
-	for ( const pxr::SdfPropertySpecHandle& Spec : SpecStack )
+	std::vector<pxr::SdfPropertySpecHandle> SpecStack = Property.GetPropertyStack();
+	for (const pxr::SdfPropertySpecHandle& Spec : SpecStack)
 	{
-		if ( !Spec )
+		if (!Spec)
 		{
 			continue;
 		}
 
 		const pxr::SdfLayerHandle& SpecLayer = Spec->GetLayer();
-		if ( SpecLayer != Layer )
+		if (SpecLayer != Layer)
 		{
-			const FText Text = LOCTEXT( "OverridenOpinionText", "USD: Overridden opinion" );
+			const FText Text = LOCTEXT("OverridenOpinionText", "USD: Overridden opinion");
 
 			const FText SubText = FText::Format(
-				LOCTEXT( "OverridenOpinionSubText", "Opinion authored for this attribute:\n\n{0}\n\nAt this layer:\n\n{1}\n\nIs overridden by another spec at this layer:\n\n{2}\n\nAnd so may not be visible on the composed stage. This means this edit may not be visible once the stage is reloaded." ),
-				FText::FromString( UsdToUnreal::ConvertPath( Spec->GetPath() ) ),
-				FText::FromString( UsdToUnreal::ConvertString( Layer->GetIdentifier() ) ),
-				FText::FromString( UsdToUnreal::ConvertString( SpecLayer->GetIdentifier() ) )
+				LOCTEXT("OverridenOpinionSubText", "Opinion authored for this attribute:\n\n{0}\n\nAt this layer:\n\n{1}\n\nIs overridden by another spec at this layer:\n\n{2}\n\nAnd so may not be visible on the composed stage. This means this edit may not be visible once the stage is reloaded."),
+				FText::FromString(UsdToUnreal::ConvertPath(Spec->GetPath())),
+				FText::FromString(UsdToUnreal::ConvertString(Layer->GetIdentifier())),
+				FText::FromString(UsdToUnreal::ConvertString(SpecLayer->GetIdentifier()))
 			);
 
-			UE_LOG( LogUsd, Warning, TEXT( "%s" ), *SubText.ToString().Replace( TEXT( "\n\n" ), TEXT( " " ) ) );
+			UE_LOG(LogUsd, Warning, TEXT("%s"), *SubText.ToString().Replace(TEXT("\n\n"), TEXT(" ")));
 
 			const UUsdProjectSettings* Settings = GetDefault<UUsdProjectSettings>();
-			if ( Settings && Settings->bShowOverriddenOpinionsWarning )
+			if (Settings && Settings->bShowOverriddenOpinionsWarning)
 			{
 				static TWeakPtr<SNotificationItem> Notification;
 
-				FNotificationInfo Toast( Text );
+				FNotificationInfo Toast(Text);
 				Toast.SubText = SubText;
-				Toast.Image = FCoreStyle::Get().GetBrush( TEXT( "MessageLog.Warning" ) );
+				Toast.Image = FCoreStyle::Get().GetBrush(TEXT("MessageLog.Warning"));
 				Toast.CheckBoxText = LOCTEXT("DontAskAgain", "Don't prompt again");
 				Toast.bUseLargeFont = false;
 				Toast.bFireAndForget = false;
@@ -245,37 +259,37 @@ void UsdUtils::NotifyIfOverriddenOpinion( const UE::FUsdAttribute& Attribute )
 				Toast.bUseThrobber = false;
 				Toast.bUseSuccessFailIcons = false;
 				Toast.ButtonDetails.Emplace(
-					LOCTEXT( "OverridenOpinionMessageOk", "Ok" ),
+					LOCTEXT("OverridenOpinionMessageOk", "Ok"),
 					FText::GetEmpty(),
-					FSimpleDelegate::CreateLambda([](){
-						if ( TSharedPtr<SNotificationItem> PinnedNotification = Notification.Pin() )
+					FSimpleDelegate::CreateLambda([]() {
+						if (TSharedPtr<SNotificationItem> PinnedNotification = Notification.Pin())
 						{
-							PinnedNotification->SetCompletionState( SNotificationItem::CS_Success );
+							PinnedNotification->SetCompletionState(SNotificationItem::CS_Success);
 							PinnedNotification->ExpireAndFadeout();
 						}
-					})
+						})
 				);
 				// This is flipped because the default checkbox message is "Don't prompt again"
 				Toast.CheckBoxState = Settings->bShowOverriddenOpinionsWarning ? ECheckBoxState::Unchecked : ECheckBoxState::Checked;
 				Toast.CheckBoxStateChanged = FOnCheckStateChanged::CreateStatic([](ECheckBoxState NewState)
-				{
-					if ( UUsdProjectSettings* Settings = GetMutableDefault<UUsdProjectSettings>() )
 					{
-						// This is flipped because the default checkbox message is "Don't prompt again"
-						Settings->bShowOverriddenOpinionsWarning = NewState == ECheckBoxState::Unchecked;
-						Settings->SaveConfig();
-					}
-				});
+						if (UUsdProjectSettings* Settings = GetMutableDefault<UUsdProjectSettings>())
+						{
+							// This is flipped because the default checkbox message is "Don't prompt again"
+							Settings->bShowOverriddenOpinionsWarning = NewState == ECheckBoxState::Unchecked;
+							Settings->SaveConfig();
+						}
+					});
 
 				// Only show one at a time
-				if ( !Notification.IsValid() )
+				if (!Notification.IsValid())
 				{
-					Notification = FSlateNotificationManager::Get().AddNotification( Toast );
+					Notification = FSlateNotificationManager::Get().AddNotification(Toast);
 				}
 
-				if ( TSharedPtr<SNotificationItem> PinnedNotification = Notification.Pin() )
+				if (TSharedPtr<SNotificationItem> PinnedNotification = Notification.Pin())
 				{
-					PinnedNotification->SetCompletionState( SNotificationItem::CS_Pending );
+					PinnedNotification->SetCompletionState(SNotificationItem::CS_Pending);
 				}
 			}
 		}
@@ -284,7 +298,7 @@ void UsdUtils::NotifyIfOverriddenOpinion( const UE::FUsdAttribute& Attribute )
 			break;
 		}
 	}
-#endif // #if USE_USD_SDK
 }
+#endif // #if USE_USD_SDK
 
 #undef LOCTEXT_NAMESPACE
\ No newline at end of file
diff --git a/Engine/Plugins/Importers/USDImporter/Source/USDUtilities/Private/USDGeomMeshConversion.cpp b/Engine/Plugins/Importers/USDImporter/Source/USDUtilities/Private/USDGeomMeshConversion.cpp
--- a/Engine/Plugins/Importers/USDImporter/Source/USDUtilities/Private/USDGeomMeshConversion.cpp
+++ b/Engine/Plugins/Importers/USDImporter/Source/USDUtilities/Private/USDGeomMeshConversion.cpp
@@ -1599,6 +1599,18 @@ UsdUtils::FUsdPrimMaterialAssignmentInfo UsdUtils::GetPrimMaterialAssignments(
 			{
 				const pxr::SdfPath& TargetMaterialPrimPath = Targets[0];
 				pxr::UsdPrim MaterialPrim = UsdPrim.GetStage()->GetPrimAtPath( TargetMaterialPrimPath );
+				if (!MaterialPrim)
+				{
+					FUsdLogManager::LogMessage(
+						EMessageSeverity::Warning,
+						FText::Format(LOCTEXT("IgnoringMaterialInvalid", "Ignoring target material '{0}' bound to prim '{1}' as that prim doesn't exist on this stage"),
+							FText::FromString(UsdToUnreal::ConvertPath(TargetMaterialPrimPath)),
+							FText::FromString(UsdToUnreal::ConvertPath(UsdPrim.GetPath()))
+						)
+					);
+					return {};
+				}
+
 				pxr::UsdShadeMaterial UsdShadeMaterial{ MaterialPrim };
 				if ( !UsdShadeMaterial )
 				{
diff --git a/Engine/Plugins/Importers/USDImporter/Source/USDUtilities/Public/USDAttributeUtils.h b/Engine/Plugins/Importers/USDImporter/Source/USDUtilities/Public/USDAttributeUtils.h
--- a/Engine/Plugins/Importers/USDImporter/Source/USDUtilities/Public/USDAttributeUtils.h
+++ b/Engine/Plugins/Importers/USDImporter/Source/USDUtilities/Public/USDAttributeUtils.h
@@ -6,6 +6,12 @@
 
 #include "UsdWrappers/ForwardDeclarations.h"
 
+#if USE_USD_SDK
+PXR_NAMESPACE_OPEN_SCOPE
+	class UsdProperty;
+PXR_NAMESPACE_CLOSE_SCOPE
+#endif // USE_USD_SDK
+
 namespace UE
 {
 	class FUsdAttribute;
@@ -31,5 +37,8 @@ namespace UsdUtils
 	 * the strongest opinion for the attribute already
 	 */
 	USDUTILITIES_API void NotifyIfOverriddenOpinion( const UE::FUsdAttribute& Attribute );
+#if USE_USD_SDK
+	USDUTILITIES_API void NotifyIfOverriddenOpinion( const pxr::UsdProperty& Property );
+#endif // USE_USD_SDK
 }
 
