USD: Fix localization warning due to using the same LOCTEXT key twice (Horde issue 281514).

diff --git a/Engine/Plugins/Importers/USDImporter/Source/USDUtilities/Private/USDGeomMeshConversion.cpp b/Engine/Plugins/Importers/USDImporter/Source/USDUtilities/Private/USDGeomMeshConversion.cpp
--- a/Engine/Plugins/Importers/USDImporter/Source/USDUtilities/Private/USDGeomMeshConversion.cpp
+++ b/Engine/Plugins/Importers/USDImporter/Source/USDUtilities/Private/USDGeomMeshConversion.cpp
@@ -1616,7 +1616,7 @@ UsdUtils::FUsdPrimMaterialAssignmentInfo UsdUtils::GetPrimMaterialAssignments(
 				{
 					FUsdLogManager::LogMessage(
 						EMessageSeverity::Warning,
-						FText::Format( LOCTEXT( "IgnoringMaterialInvalid", "Ignoring material '{0}' bound to prim '{1}' as it does not possess the UsdShadeMaterial schema" ),
+						FText::Format( LOCTEXT( "IgnoringMaterialNoSchema", "Ignoring material '{0}' bound to prim '{1}' as it does not possess the UsdShadeMaterial schema" ),
 							FText::FromString( UsdToUnreal::ConvertPath( TargetMaterialPrimPath ) ),
 							FText::FromString( UsdToUnreal::ConvertPath( UsdPrim.GetPath() ) )
 						)
