USD: Update generated assets when relevant prims are modified.

diff --git a/Engine/Plugins/Importers/USDImporter/Source/USDSchemas/Private/Custom/MDLUSDShadeMaterialTranslator.cpp b/Engine/Plugins/Importers/USDImporter/Source/USDSchemas/Private/Custom/MDLUSDShadeMaterialTranslator.cpp
--- a/Engine/Plugins/Importers/USDImporter/Source/USDSchemas/Private/Custom/MDLUSDShadeMaterialTranslator.cpp
+++ b/Engine/Plugins/Importers/USDImporter/Source/USDSchemas/Private/Custom/MDLUSDShadeMaterialTranslator.cpp
@@ -102,6 +102,9 @@ void FMdlUsdShadeMaterialTranslator::CreateAssets()
 		return;
 	}
 
+	// Note that we rely on FUsdShadeMaterialTranslator's implementation of CollectAuxiliaryPrims
+	RegisterAuxiliaryPrims();
+
 	const pxr::TfToken MdlToken = UnrealToUsd::ConvertToken( *MdlRenderContext.ToString() ).Get();
 
 	pxr::UsdShadeShader SurfaceShader = ShadeMaterial.ComputeSurfaceSource( MdlToken );
diff --git a/Engine/Plugins/Importers/USDImporter/Source/USDSchemas/Private/USDGeomMeshTranslator.cpp b/Engine/Plugins/Importers/USDImporter/Source/USDSchemas/Private/USDGeomMeshTranslator.cpp
--- a/Engine/Plugins/Importers/USDImporter/Source/USDSchemas/Private/USDGeomMeshTranslator.cpp
+++ b/Engine/Plugins/Importers/USDImporter/Source/USDSchemas/Private/USDGeomMeshTranslator.cpp
@@ -1224,9 +1224,14 @@ namespace UsdGeomMeshTranslatorImpl
 #endif // WITH_EDITOR
 }
 
-FBuildStaticMeshTaskChain::FBuildStaticMeshTaskChain(const TSharedRef< FUsdSchemaTranslationContext >& InContext, const UE::FSdfPath& InPrimPath)
+FBuildStaticMeshTaskChain::FBuildStaticMeshTaskChain(
+	const TSharedRef< FUsdSchemaTranslationContext >& InContext,
+	const UE::FSdfPath& InPrimPath,
+	const TOptional<UE::FSdfPath>& InAlternativePrimToLinkAssetsTo
+)
 	: PrimPath(InPrimPath)
 	, Context(InContext)
+	, AlternativePrimToLinkAssetsTo(InAlternativePrimToLinkAssetsTo)
 {
 }
 
@@ -1271,7 +1276,10 @@ void FBuildStaticMeshTaskChain::SetupTasks()
 			{
 				if (Context->InfoCache)
 				{
-					Context->InfoCache->LinkAssetToPrim(PrimPath, StaticMesh);
+					const UE::FSdfPath& TargetPath = AlternativePrimToLinkAssetsTo.IsSet()
+						? AlternativePrimToLinkAssetsTo.GetValue()
+						: PrimPath;
+					Context->InfoCache->LinkAssetToPrim(TargetPath, StaticMesh);
 				}
 
 #if WITH_EDITOR
@@ -1280,7 +1288,7 @@ void FBuildStaticMeshTaskChain::SetupTasks()
 				if (bIsNew)
 				{
 					UUsdMeshAssetImportData* ImportData = NewObject<UUsdMeshAssetImportData>(StaticMesh, TEXT("UUSDAssetImportData"));
-					ImportData->PrimPath = PrimPathString;
+					ImportData->PrimPath = MeshName;
 					StaticMesh->AssetImportData = ImportData;
 				}
 #endif // WITH_EDITOR
@@ -1385,8 +1393,13 @@ void FBuildStaticMeshTaskChain::SetupTasks()
 		});
 }
 
-FGeomMeshCreateAssetsTaskChain::FGeomMeshCreateAssetsTaskChain(const TSharedRef< FUsdSchemaTranslationContext >& InContext, const UE::FSdfPath& InPrimPath, const FTransform& InAdditionalTransform)
-	: FBuildStaticMeshTaskChain(InContext, InPrimPath)
+FGeomMeshCreateAssetsTaskChain::FGeomMeshCreateAssetsTaskChain(
+	const TSharedRef< FUsdSchemaTranslationContext >& InContext,
+	const UE::FSdfPath& InPrimPath,
+	const TOptional<UE::FSdfPath>& AlternativePrimToLinkAssetsTo,
+	const FTransform& InAdditionalTransform
+)
+	: FBuildStaticMeshTaskChain(InContext, InPrimPath, AlternativePrimToLinkAssetsTo)
 	, AdditionalTransform(InAdditionalTransform)
 {
 	SetupTasks();
@@ -1552,7 +1565,10 @@ void FGeometryCacheCreateAssetsTaskChain::SetupTasks()
 
 			if (GeometryCache && Context->InfoCache)
 			{
-				Context->InfoCache->LinkAssetToPrim(PrimPath, GeometryCache);
+				const UE::FSdfPath& TargetPath = AlternativePrimToLinkAssetsTo.IsSet()
+					? AlternativePrimToLinkAssetsTo.GetValue()
+					: PrimPath;
+				Context->InfoCache->LinkAssetToPrim(TargetPath, GeometryCache);
 			}
 
 			// We don't have any additional step
@@ -1565,6 +1581,8 @@ void FUsdGeomMeshTranslator::CreateAssets()
 {
 	TRACE_CPUPROFILER_EVENT_SCOPE(FUsdGeomMeshTranslator::CreateAssets);
 
+	RegisterAuxiliaryPrims();
+
 #if WITH_EDITOR
 	if (!Context->bIsImporting && GUseGeometryCacheUSD && UsdGeomMeshTranslatorImpl::IsAnimated(GetPrim()))
 	{
@@ -1802,4 +1820,28 @@ bool FUsdGeomMeshTranslator::CanBeCollapsed(ECollapsingType CollapsingType) cons
 	return Super::CanBeCollapsed(CollapsingType);
 }
 
+TSet<UE::FSdfPath> FUsdGeomMeshTranslator::CollectAuxiliaryPrims() const
+{
+	TSet<UE::FSdfPath> Result;
+	{
+		FScopedUsdAllocs UsdAllocs;
+
+		pxr::UsdPrim Prim = GetPrim();
+
+		// The UsdGeomSubset prims are used to express multiple material assignments per mesh. A change in them could
+		// mean a change in triangle to material slot mapping
+		TArray<TUsdStore<pxr::UsdPrim>> ChildPrims = UsdUtils::GetAllPrimsOfType(
+			Prim,
+			pxr::TfType::Find<pxr::UsdGeomSubset>()
+		);
+
+		Result.Reserve(ChildPrims.Num());
+		for (const TUsdStore<pxr::UsdPrim>& ChildPrim : ChildPrims)
+		{
+			Result.Add(UE::FSdfPath{ChildPrim.Get().GetPrimPath()});
+		}
+	}
+	return Result;
+}
+
 #endif // #if USE_USD_SDK
diff --git a/Engine/Plugins/Importers/USDImporter/Source/USDSchemas/Private/USDGeomPointInstancerTranslator.cpp b/Engine/Plugins/Importers/USDImporter/Source/USDSchemas/Private/USDGeomPointInstancerTranslator.cpp
--- a/Engine/Plugins/Importers/USDImporter/Source/USDSchemas/Private/USDGeomPointInstancerTranslator.cpp
+++ b/Engine/Plugins/Importers/USDImporter/Source/USDSchemas/Private/USDGeomPointInstancerTranslator.cpp
@@ -29,6 +29,7 @@
 
 #include "USDIncludesStart.h"
 	#include "pxr/usd/usd/prim.h"
+	#include "pxr/usd/usd/primRange.h"
 	#include "pxr/usd/usd/stage.h"
 	#include "pxr/usd/usdGeom/camera.h"
 	#include "pxr/usd/usdGeom/mesh.h"
@@ -53,6 +54,7 @@ namespace UsdGeomPointInstancerTranslatorImpl
 		{
 			TRACE_CPUPROFILER_EVENT_SCOPE(ApplyPointInstanceTransforms);
 
+			Component->ClearInstances();
 			Component->AddInstances( InstanceTransforms, false );
 		}
 	}
@@ -64,6 +66,8 @@ namespace UsdGeomPointInstancerTranslatorImpl
 			return;
 		}
 
+		MeshComponent.Modify();
+
 		if ( MeshComponent.IsRegistered() )
 		{
 			MeshComponent.UnregisterComponent();
@@ -84,10 +88,15 @@ namespace UsdGeomPointInstancerTranslatorImpl
 FUsdGeomPointInstancerCreateAssetsTaskChain::FUsdGeomPointInstancerCreateAssetsTaskChain(
 	const TSharedRef< FUsdSchemaTranslationContext >& InContext,
 	const UE::FSdfPath& InPrimPath,
-	bool bInIgnoreTopLevelTransformAndVisibility
+	bool bInIgnoreTopLevelTransformAndVisibility,
+	const TOptional<UE::FSdfPath>& InAlternativePrimToLinkAssetsTo
 )
-	: FBuildStaticMeshTaskChain( InContext, InPrimPath )
-	, bIgnoreTopLevelTransformAndVisibility( bInIgnoreTopLevelTransformAndVisibility )
+	: FBuildStaticMeshTaskChain(
+		InContext,
+		InPrimPath,
+		InAlternativePrimToLinkAssetsTo
+	)
+	, bIgnoreTopLevelTransformAndVisibility(bInIgnoreTopLevelTransformAndVisibility)
 {
 	SetupTasks();
 }
@@ -147,6 +156,8 @@ void FUsdGeomPointInstancerTranslator::CreateAssets()
 {
 	TRACE_CPUPROFILER_EVENT_SCOPE( FUsdGeomPointInstancerTranslator::CreateAssets );
 
+	RegisterAuxiliaryPrims();
+
 	FScopedUsdAllocs Allocs;
 
 	pxr::UsdPrim Prim = GetPrim();
@@ -163,7 +174,11 @@ void FUsdGeomPointInstancerTranslator::CreateAssets()
 	{
 		// Don't bake our actual point instancer's transform or visibility into the mesh as its nice to have these on the static mesh component instead
 		const bool bIgnoreTopLevelTransformAndVisibility = true;
-		Context->TranslatorTasks.Add( MakeShared< FUsdGeomPointInstancerCreateAssetsTaskChain >( Context, PrimPath, bIgnoreTopLevelTransformAndVisibility ) );
+		Context->TranslatorTasks.Add(MakeShared<FUsdGeomPointInstancerCreateAssetsTaskChain>(
+			Context,
+			PrimPath,
+			bIgnoreTopLevelTransformAndVisibility
+		));
 	}
 	// Otherwise we're just going to spawn HISM components instead
 	else
@@ -178,6 +193,11 @@ void FUsdGeomPointInstancerTranslator::CreateAssets()
 
 		for ( const pxr::SdfPath& PrototypePath : PrototypePaths )
 		{
+			// Note how we will spawn a task chain for the prototype *regardless of where it is*. This prototype
+			// could even be external to the point instancer itself, and so will already have been handled by
+			// another translator. Unfortunately we need to do this because we need to generate a task chain for
+			// it in case it is another point instancer itself
+
 			pxr::UsdPrim PrototypeUsdPrim = Prim.GetStage()->GetPrimAtPath( PrototypePath );
 			UE::FSdfPath UEPrototypePath{ PrototypePath };
 
@@ -221,13 +241,23 @@ void FUsdGeomPointInstancerTranslator::CreateAssets()
 					UsdToUnreal::ConvertXformable( Prim.GetStage(), ParentXform, AdditionalUESpaceTransform, Context->Time, &bOutResetTransformStack );
 				}
 
-				Context->TranslatorTasks.Add( MakeShared< FGeomMeshCreateAssetsTaskChain >( Context, ChildMeshPath, AdditionalUESpaceTransform ) );
+				Context->TranslatorTasks.Add(MakeShared<FGeomMeshCreateAssetsTaskChain>(
+					Context,
+					ChildMeshPath,
+					PrimPath,
+					AdditionalUESpaceTransform
+				));
 			}
 			else
 			{
 				// Fully bake the prototype top level transform and visibility into its own static mesh
 				const bool bIgnoreTopLevelTransformAndVisibility = false;
-				Context->TranslatorTasks.Add( MakeShared< FUsdGeomPointInstancerCreateAssetsTaskChain >( Context, UEPrototypePath, bIgnoreTopLevelTransformAndVisibility ) );
+				Context->TranslatorTasks.Add(MakeShared<FUsdGeomPointInstancerCreateAssetsTaskChain>(
+					Context,
+					UEPrototypePath,
+					bIgnoreTopLevelTransformAndVisibility,
+					PrimPath
+				));
 			}
 		}
 	}
@@ -252,36 +282,144 @@ USceneComponent* FUsdGeomPointInstancerTranslator::CreateComponents()
 	{
 		return MainSceneComponent;
 	}
+
+	// Actually create the child HISM components for each point instancer prototype
+	{
+		FScopedUsdAllocs UsdAllocs;
+
+		pxr::UsdPrim Prim = GetPrim();
+		pxr::UsdGeomPointInstancer PointInstancer(Prim);
+		if (!PointInstancer)
+		{
+			return MainSceneComponent;
+		}
+
+		pxr::SdfPathVector PrototypePaths;
+		if (!PointInstancer.GetPrototypesRel().GetTargets(&PrototypePaths))
+		{
+			return MainSceneComponent;
+		}
+
+		if (!Context->AssetCache.IsValid() || !Context->InfoCache.IsValid())
+		{
+			return MainSceneComponent;
+		}
+		UUsdAssetCache2& AssetCache = *Context->AssetCache.Get();
+		FUsdInfoCache& InfoCache = *Context->InfoCache.Get();
+
+		// Lets pretend ParentComponent is pointing to the parent USceneComponent while we create the child HISMs, so they get
+		// automatically attached to it as children
+		TGuardValue< USceneComponent* > ParentComponentGuard{Context->ParentComponent, MainSceneComponent};
+
+		TArray<TFuture<TTuple<UHierarchicalInstancedStaticMeshComponent*, TArray<FTransform>>>> Tasks;
+		FScopedSlowTask PrototypePathsSlowTask((float)PrototypePaths.size(), LOCTEXT("GeomPointCreateComponents", "Creating HierarchicalInstancedStaticMeshComponents for point instancers"));
+		for (int32 PrototypeIndex = 0; PrototypeIndex < PrototypePaths.size(); ++PrototypeIndex)
+		{
+			PrototypePathsSlowTask.EnterProgressFrame();
+
+			const pxr::SdfPath& PrototypePath = PrototypePaths[PrototypeIndex];
+			FString PrototypePathStr = UsdToUnreal::ConvertPath(PrototypePath);
+
+			pxr::UsdPrim PrototypeUsdPrim = Prim.GetStage()->GetPrimAtPath(PrototypePath);
+			if (!PrototypeUsdPrim)
+			{
+				UE_LOG(LogUsd, Warning, TEXT("Failed to find prototype '%s' for PointInstancer '%s' when creating components"), *PrototypePathStr, *PrimPath.GetString());
+				continue;
+			}
+
+			// If our prototype was a LOD mesh we will have used the path of one of the actual LOD meshes to start the FGeomMeshCreateAssetsTaskChain,
+			// so we have to look for our resulting mesh with the same path
+			if (Context->bAllowInterpretingLODs && UsdUtils::DoesPrimContainMeshLODs(PrototypeUsdPrim))
+			{
+				pxr::UsdPrimSiblingRange PrimRange = PrototypeUsdPrim.GetChildren();
+				for (pxr::UsdPrimSiblingRange::iterator PrimRangeIt = PrimRange.begin(); PrimRangeIt != PrimRange.end(); ++PrimRangeIt)
+				{
+					const pxr::UsdPrim& Child = *PrimRangeIt;
+					if (pxr::UsdGeomMesh ChildMesh{Child})
+					{
+						PrototypeUsdPrim = Child;
+						PrototypePathStr = UsdToUnreal::ConvertPath(Child.GetPrimPath());
+						break;
+					}
+				}
+			}
+
+			using UHISMComponent = UHierarchicalInstancedStaticMeshComponent;
+
+			FUsdGeomXformableTranslator XformableTranslator{UHISMComponent::StaticClass(), Context, UE::FUsdTyped(PrototypeUsdPrim)};
+
+			const bool bNeedsActor = false;
+			UHISMComponent* HISMComponent = Cast<UHISMComponent>(XformableTranslator.CreateComponentsEx({UHISMComponent::StaticClass()}, bNeedsActor));
+			if (!HISMComponent)
+			{
+				UE_LOG(LogUsd, Warning, TEXT("Failed to generate HISM component for prototype '%s' for PointInstancer '%s'"), *PrototypePathStr, *PrimPath.GetString());
+				continue;
+			}
+		}
+	}
+
 	UpdateComponents( MainSceneComponent );
 
+	return MainSceneComponent;
+}
+
+void FUsdGeomPointInstancerTranslator::UpdateComponents(USceneComponent* PointInstancerRootComponent)
+{
 	FScopedUsdAllocs UsdAllocs;
 
 	pxr::UsdPrim Prim = GetPrim();
 	pxr::UsdGeomPointInstancer PointInstancer( Prim );
 	if ( !PointInstancer )
 	{
-		return MainSceneComponent;
+		return;
 	}
 
 	pxr::SdfPathVector PrototypePaths;
 	if ( !PointInstancer.GetPrototypesRel().GetTargets( &PrototypePaths ) )
 	{
-		return MainSceneComponent;
+		return;
 	}
 
 	if (!Context->AssetCache.IsValid() || !Context->InfoCache.IsValid())
 	{
-		return MainSceneComponent;
+		return;
 	}
 	UUsdAssetCache2& AssetCache = *Context->AssetCache.Get();
 	FUsdInfoCache& InfoCache = *Context->InfoCache.Get();
 
 	// Lets pretend ParentComponent is pointing to the parent USceneComponent while we create the child HISMs, so they get
 	// automatically attached to it as children
-	TGuardValue< USceneComponent* > ParentComponentGuard{ Context->ParentComponent, MainSceneComponent };
+	TGuardValue< USceneComponent* > ParentComponentGuard{Context->ParentComponent, PointInstancerRootComponent};
+
+	const TArray<TObjectPtr<USceneComponent>>& AttachedChildren = PointInstancerRootComponent->GetAttachChildren();
+	TArray<UHierarchicalInstancedStaticMeshComponent*> AttachedHISMs;
+	AttachedHISMs.Reserve(AttachedChildren.Num());
+	for (const TObjectPtr<USceneComponent>& AttachedChild : AttachedChildren)
+	{
+		if (UHierarchicalInstancedStaticMeshComponent* HISM = Cast<UHierarchicalInstancedStaticMeshComponent>(AttachedChild))
+		{
+			AttachedHISMs.Add(HISM);
+		}
+	}
+
+	// We "link" the prototype meshes to the point instancer, but we don't know which mesh corresponds to each
+	// prototype, as we translate these with task pools and some of those prototypes may have generated nullptr.
+	// We always put the prototype path on the asset import data though, so here we use that to figure out where
+	// each mesh should go
+	TSet<UStaticMesh*> PrototypeMeshSet = InfoCache.GetAssetsForPrim<UStaticMesh>(PrimPath);
+	std::unordered_map<pxr::SdfPath, UStaticMesh*, pxr::SdfPath::Hash> PrototypeMeshes;
+	PrototypeMeshes.reserve(PrototypeMeshSet.Num());
+	for (UStaticMesh* PrototypeMesh : PrototypeMeshSet)
+	{
+		if (UUsdAssetImportData* ImportData = Cast<UUsdAssetImportData>(PrototypeMesh->GetAssetImportData()))
+		{
+			pxr::SdfPath PrototypePath = UnrealToUsd::ConvertPath(*ImportData->PrimPath).Get();
+			PrototypeMeshes.insert({PrototypePath, PrototypeMesh});
+		}
+	}
 
 	TArray<TFuture<TTuple<UHierarchicalInstancedStaticMeshComponent*, TArray<FTransform>>>> Tasks;
-	FScopedSlowTask PrototypePathsSlowTask( ( float ) PrototypePaths.size(), LOCTEXT( "GeomPointCreateComponents", "Creating HierarchicalInstancedStaticMeshComponents for point instancers" ) );
+	FScopedSlowTask PrototypePathsSlowTask( ( float ) PrototypePaths.size(), LOCTEXT( "GeomPointUpdateComponents", "Updating HierarchicalInstancedStaticMeshComponents for point instancers" ) );
 	for ( int32 PrototypeIndex = 0; PrototypeIndex < PrototypePaths.size(); ++PrototypeIndex )
 	{
 		PrototypePathsSlowTask.EnterProgressFrame();
@@ -292,10 +430,20 @@ USceneComponent* FUsdGeomPointInstancerTranslator::CreateComponents()
 		pxr::UsdPrim PrototypeUsdPrim = Prim.GetStage()->GetPrimAtPath( PrototypePath );
 		if ( !PrototypeUsdPrim )
 		{
-			UE_LOG( LogUsd, Warning, TEXT( "Failed to find prototype '%s' for PointInstancer '%s' when creating components" ), *PrototypePathStr, *PrimPath.GetString() );
+			UE_LOG( LogUsd, Warning, TEXT( "Failed to find prototype '%s' for PointInstancer '%s' when updating components" ), *PrototypePathStr, *PrimPath.GetString() );
 			continue;
 		}
 
+		UHierarchicalInstancedStaticMeshComponent* HISMComponent = nullptr;
+
+		// The user could have just manually deleted the component, so we must check
+		if (!AttachedHISMs.IsValidIndex(PrototypeIndex))
+		{
+			UE_LOG(LogUsd, Warning, TEXT("Failed to find corresponding HISM component for prototype '%s' of PointInstancer '%s'. Cancelling component update"), *PrototypePathStr, *PrimPath.GetString());
+			break;
+		}
+		HISMComponent = AttachedHISMs[PrototypeIndex];
+
 		// If our prototype was a LOD mesh we will have used the path of one of the actual LOD meshes to start the FGeomMeshCreateAssetsTaskChain,
 		// so we have to look for our resulting mesh with the same path
 		if ( Context->bAllowInterpretingLODs && UsdUtils::DoesPrimContainMeshLODs( PrototypeUsdPrim ) )
@@ -313,21 +461,13 @@ USceneComponent* FUsdGeomPointInstancerTranslator::CreateComponents()
 			}
 		}
 
-		using UHISMComponent = UHierarchicalInstancedStaticMeshComponent;
-
-		FUsdGeomXformableTranslator XformableTranslator{ UHISMComponent::StaticClass(), Context, UE::FUsdTyped( PrototypeUsdPrim ) };
-
-		const bool bNeedsActor = false;
-		UHISMComponent* HISMComponent = Cast<UHISMComponent>( XformableTranslator.CreateComponentsEx( { UHISMComponent::StaticClass() }, bNeedsActor ) );
-		if ( !HISMComponent )
+		// This mesh could be nullptr, but that's OK
+		UStaticMesh* StaticMesh = nullptr;
+		std::unordered_map<pxr::SdfPath, UStaticMesh*, pxr::SdfPath::Hash>::iterator Iter = PrototypeMeshes.find(PrototypePath);
+		if (Iter != PrototypeMeshes.end())
 		{
-			UE_LOG( LogUsd, Warning, TEXT( "Failed to generate HISM component for prototype '%s' for PointInstancer '%s'" ), *PrototypePathStr, *PrimPath.GetString() );
-			continue;
+			StaticMesh = Iter->second;
 		}
-
-		UStaticMesh* StaticMesh = InfoCache.GetSingleAssetForPrim<UStaticMesh>(
-			UE::FSdfPath{*PrototypePathStr}
-		);
 		UsdGeomPointInstancerTranslatorImpl::SetStaticMesh( StaticMesh, *HISMComponent );
 
 		// Evaluating point instancer can take a long time and is thread-safe. Move to async task while we work on something else.
@@ -377,17 +517,91 @@ USceneComponent* FUsdGeomPointInstancerTranslator::CreateComponents()
 		UsdGeomPointInstancerTranslatorImpl::ApplyPointInstanceTransforms( Result.Key, Result.Value );
 	}
 
-	return MainSceneComponent;
+	Super::UpdateComponents(PointInstancerRootComponent);
 }
 
-void FUsdGeomPointInstancerTranslator::UpdateComponents( USceneComponent* PointInstancerRootComponent )
+bool FUsdGeomPointInstancerTranslator::CanBeCollapsed( ECollapsingType CollapsingType ) const
 {
-	Super::UpdateComponents( PointInstancerRootComponent );
+	return GCollapseTopLevelPointInstancers;
 }
 
-bool FUsdGeomPointInstancerTranslator::CanBeCollapsed( ECollapsingType CollapsingType ) const
+TSet<UE::FSdfPath> FUsdGeomPointInstancerTranslator::CollectAuxiliaryPrims() const
 {
-	return GCollapseTopLevelPointInstancers;
+	FScopedUsdAllocs UsdAllocs;
+
+	pxr::UsdPrim Prim = GetPrim();
+	pxr::UsdGeomPointInstancer PointInstancer(Prim);
+	if (!PointInstancer)
+	{
+		return {};
+	}
+
+	pxr::SdfPathVector PrototypePaths;
+	if (!PointInstancer.GetPrototypesRel().GetTargets(&PrototypePaths))
+	{
+		return {};
+	}
+
+	IUsdSchemasModule& UsdSchemasModule = FModuleManager::Get().LoadModuleChecked< IUsdSchemasModule >(TEXT("USDSchemas"));
+
+	TSet<UE::FSdfPath> Result;
+	Result.Reserve(PrototypePaths.size());
+	for (int32 PrototypeIndex = 0; PrototypeIndex < PrototypePaths.size(); ++PrototypeIndex)
+	{
+		UE::FSdfPath PrototypePath = UE::FSdfPath{PrototypePaths[PrototypeIndex]};
+		UE::FUsdPrim PrototypePrim = Context->Stage.GetPrimAtPath(PrototypePath);
+
+		Result.Add(PrototypePath);
+
+		// Internal prototype
+		// We must depend on all prims of the prototype subtree, because we're in charge of collapsing it
+		if (PrototypePath.HasPrefix(PrimPath))
+		{
+			TArray<TUsdStore<pxr::UsdPrim>> ChildPrims = UsdUtils::GetAllPrimsOfType(
+				PrototypePrim,
+				pxr::TfType::Find<pxr::UsdGeomImageable>()
+			);
+
+			for (const TUsdStore<pxr::UsdPrim>& ChildPrim : ChildPrims)
+			{
+				Result.Add(UE::FSdfPath{ChildPrim.Get().GetPrimPath()});
+
+				TSharedPtr< FUsdSchemaTranslator > ChildSchemaTranslator = UsdSchemasModule.GetTranslatorRegistry().CreateTranslatorForSchema(
+					Context,
+					UE::FUsdTyped(ChildPrim.Get())
+				);
+
+				TSet<UE::FSdfPath> RecursiveDependencies = ChildSchemaTranslator->CollectAuxiliaryPrims();
+				for (const UE::FSdfPath& RecursiveDependency : RecursiveDependencies)
+				{
+					Result.Add(RecursiveDependency);
+				}
+			}
+		}
+		// External prototype
+		// Depend on prims until they collapse into something, at which point we can stop as they will depend on their
+		// own subtree by themselves already
+		else
+		{
+			pxr::UsdPrimRange PrimRange(PrototypePrim, pxr::UsdTraverseInstanceProxies());
+
+			for (pxr::UsdPrimRange::iterator PrimRangeIt = PrimRange.begin(); PrimRangeIt != PrimRange.end(); ++PrimRangeIt)
+			{
+				Result.Add(UE::FSdfPath{PrimRangeIt->GetPrimPath()});
+
+				TSharedPtr< FUsdSchemaTranslator > SchemaTranslator = UsdSchemasModule.GetTranslatorRegistry().CreateTranslatorForSchema(
+					Context,
+					UE::FUsdTyped(*PrimRangeIt)
+				);
+				if (SchemaTranslator->CollapsesChildren(ECollapsingType::Assets))
+				{
+					PrimRangeIt.PruneChildren();
+				}
+			}
+		}
+	}
+
+	return Result;
 }
 
 #undef LOCTEXT_NAMESPACE
diff --git a/Engine/Plugins/Importers/USDImporter/Source/USDSchemas/Private/USDGeomXformableTranslator.cpp b/Engine/Plugins/Importers/USDImporter/Source/USDSchemas/Private/USDGeomXformableTranslator.cpp
--- a/Engine/Plugins/Importers/USDImporter/Source/USDSchemas/Private/USDGeomXformableTranslator.cpp
+++ b/Engine/Plugins/Importers/USDImporter/Source/USDSchemas/Private/USDGeomXformableTranslator.cpp
@@ -264,13 +264,15 @@ void FUsdGeomXformableCreateAssetsTaskChain::SetupTasks()
 
 void FUsdGeomXformableTranslator::CreateAssets()
 {
+	TRACE_CPUPROFILER_EVENT_SCOPE( FUsdGeomMeshTranslator::CreateAssets );
+
 	if ( !CollapsesChildren( ECollapsingType::Assets ) )
 	{
 		// We only have to create assets when our children are collapsed together
 		return;
 	}
 
-	TRACE_CPUPROFILER_EVENT_SCOPE( FUsdGeomMeshTranslator::CreateAssets );
+	RegisterAuxiliaryPrims();
 
 	Context->TranslatorTasks.Add( MakeShared< FUsdGeomXformableCreateAssetsTaskChain >( Context, PrimPath ) );
 }
@@ -746,4 +748,38 @@ bool FUsdGeomXformableTranslator::CanBeCollapsed( ECollapsingType CollapsingType
 	return true;
 }
 
+TSet<UE::FSdfPath> FUsdGeomXformableTranslator::CollectAuxiliaryPrims() const
+{
+	if (!Context->InfoCache.IsValid())
+	{
+		return {};
+	}
+
+	if (!Context->InfoCache->DoesPathCollapseChildren(PrimPath, ECollapsingType::Assets))
+	{
+		return {};
+	}
+
+	TSet<UE::FSdfPath> Result;
+	{
+		FScopedUsdAllocs UsdAllocs;
+
+		pxr::UsdPrim Prim = GetPrim();
+
+		// We check imageable because that is the most basal schema that is still relevant for collapsed meshes (it
+		// holds the visibility attribute)
+		TArray<TUsdStore<pxr::UsdPrim>> ChildPrims = UsdUtils::GetAllPrimsOfType(
+			Prim,
+			pxr::TfType::Find<pxr::UsdGeomImageable>()
+		);
+
+		Result.Reserve(ChildPrims.Num());
+		for (const TUsdStore<pxr::UsdPrim>& ChildPrim : ChildPrims)
+		{
+			Result.Add(UE::FSdfPath{ChildPrim.Get().GetPrimPath()});
+		}
+	}
+	return Result;
+}
+
 #endif // #if USE_USD_SDK
diff --git a/Engine/Plugins/Importers/USDImporter/Source/USDSchemas/Private/USDGroomTranslator.cpp b/Engine/Plugins/Importers/USDImporter/Source/USDSchemas/Private/USDGroomTranslator.cpp
--- a/Engine/Plugins/Importers/USDImporter/Source/USDSchemas/Private/USDGroomTranslator.cpp
+++ b/Engine/Plugins/Importers/USDImporter/Source/USDSchemas/Private/USDGroomTranslator.cpp
@@ -27,10 +27,12 @@
 #include "USDTypesConversion.h"
 
 #include "USDIncludesStart.h"
-#include "pxr/base/tf/type.h"
-#include "pxr/usd/usd/prim.h"
-#include "pxr/usd/usd/timeCode.h"
-#include "pxr/usd/usdGeom/curves.h"
+	#include "pxr/base/tf/type.h"
+	#include "pxr/usd/usd/prim.h"
+	#include "pxr/usd/usd/timeCode.h"
+	#include "pxr/usd/usdGeom/basisCurves.h"
+	#include "pxr/usd/usdGeom/curves.h"
+	#include "pxr/usd/usdGeom/imageable.h"
 #include "USDIncludesEnd.h"
 
 namespace UE::UsdGroomTranslator::Private
@@ -306,10 +308,7 @@ protected:
 				}
 				else if (Context->InfoCache)
 				{
-					Context->InfoCache->LinkAssetToPrim(
-						UE::FSdfPath{*UsdGroomTranslatorUtils::GetStrandsGroomCachePrimPath(PrimPath)},
-						GroomCache
-					);
+					Context->InfoCache->LinkAssetToPrim(PrimPath, GroomCache);
 				}
 
 				if (!bSuccess)
@@ -338,17 +337,14 @@ protected:
 				{
 #if WITH_EDITOR
 					UUsdAssetImportData* ImportData = NewObject<UUsdAssetImportData>(GroomCache, TEXT("UUSDAssetImportData"));
-					ImportData->PrimPath = UsdGroomTranslatorUtils::GetStrandsGroomCachePrimPath(PrimPath);
+					ImportData->PrimPath = PrimPath.GetString();
 					ImportData->ImportOptions = ImportOptions.Get();
 
 					GroomCache->AssetImportData = ImportData;
 #endif // WITH_EDITOR
 
 					Context->AssetCache->CacheAsset(GroomCacheHash.ToString(), GroomCache);
-					Context->InfoCache->LinkAssetToPrim(
-						UE::FSdfPath{*StrandsGroomCachePrimPath},
-						GroomCache
-					);
+					Context->InfoCache->LinkAssetToPrim(PrimPath, GroomCache);
 				}
 
 				return GroomCache != nullptr;
@@ -368,6 +364,8 @@ void FUsdGroomTranslator::CreateAssets()
 		return Super::CreateAssets();
 	}
 
+	RegisterAuxiliaryPrims();
+
 	Context->TranslatorTasks.Add(MakeShared<FUsdGroomCreateAssetsTaskChain>(Context, PrimPath));
 }
 
@@ -425,9 +423,7 @@ void FUsdGroomTranslator::UpdateComponents(USceneComponent* SceneComponent)
 
 			if (Groom)
 			{
-				UGroomCache* GroomCache = Context->InfoCache->GetSingleAssetForPrim<UGroomCache>(
-					UE::FSdfPath{*UsdGroomTranslatorUtils::GetStrandsGroomCachePrimPath(PrimPath)}
-				);
+				UGroomCache* GroomCache = Context->InfoCache->GetSingleAssetForPrim<UGroomCache>(PrimPath);
 				if (GroomCache != GroomComponent->GroomCache.Get())
 				{
 					GroomComponent->SetGroomCache(GroomCache);
@@ -467,4 +463,49 @@ bool FUsdGroomTranslator::CanBeCollapsed(ECollapsingType CollapsingType) const
 	return true;
 }
 
+TSet<UE::FSdfPath> FUsdGroomTranslator::CollectAuxiliaryPrims() const
+{
+	if (!IsGroomPrim())
+	{
+		return Super::CollectAuxiliaryPrims();
+	}
+
+	if (!Context->InfoCache.IsValid())
+	{
+		return {};
+	}
+
+	if (!Context->InfoCache->DoesPathCollapseChildren(PrimPath, ECollapsingType::Assets))
+	{
+		return {};
+	}
+
+	TSet<UE::FSdfPath> Result;
+	{
+		FScopedUsdAllocs UsdAllocs;
+
+		TFunction<void(const pxr::UsdPrim&)> RecursivelyRegisterPrims;
+		RecursivelyRegisterPrims = [&](const pxr::UsdPrim& UsdPrim)
+		{
+			if (pxr::UsdGeomBasisCurves Curves = pxr::UsdGeomBasisCurves(UsdPrim))
+			{
+				Result.Add(UE::FSdfPath{UsdPrim.GetPrimPath()});
+			}
+			else if (pxr::UsdGeomImageable(UsdPrim))
+			{
+				Result.Add(UE::FSdfPath{UsdPrim.GetPrimPath()});
+
+				for (const pxr::UsdPrim& Child : UsdPrim.GetChildren())
+				{
+					RecursivelyRegisterPrims(Child);
+				}
+			}
+		};
+
+		pxr::UsdPrim Prim = GetPrim();
+		RecursivelyRegisterPrims(Prim);
+	}
+	return Result;
+}
+
 #endif // #if USE_USD_SDK
\ No newline at end of file
diff --git a/Engine/Plugins/Importers/USDImporter/Source/USDSchemas/Private/USDInfoCache.cpp b/Engine/Plugins/Importers/USDImporter/Source/USDSchemas/Private/USDInfoCache.cpp
--- a/Engine/Plugins/Importers/USDImporter/Source/USDSchemas/Private/USDInfoCache.cpp
+++ b/Engine/Plugins/Importers/USDImporter/Source/USDSchemas/Private/USDInfoCache.cpp
@@ -2,6 +2,7 @@
 
 #include "USDInfoCache.h"
 
+#include "USDAssetCache2.h"
 #include "USDGeomMeshConversion.h"
 #include "USDGeomXformableTranslator.h"
 #include "USDLog.h"
@@ -64,10 +65,16 @@ struct FUsdInfoCache::FUsdInfoCacheImpl
 	TMap<TWeakObjectPtr<UObject>, TSet<UE::FSdfPath>> AssetToPrimPaths;
 	mutable FRWLock PrimPathToAssetsLock;
 
-	// Paths to material prims that are actually used by mesh prims in the scene, given the current settings for
+	// Paths to material prims to the mesh prims they are bound to in the scene, given the current settings for
 	// render context, material purpose, variant selections, etc.
-	TSet<UE::FSdfPath> UsedMaterialPaths;
-	mutable FRWLock UsedMaterialPathsLock;
+	TMap<UE::FSdfPath, TSet<UE::FSdfPath>> MaterialUsers;
+	mutable FRWLock MaterialUsersLock;
+
+	// Maps from prims, to all the prims that require also reading this prim to be translated into an asset.
+	// Mainly used to update these assets whenever the depencency prim is updated.
+	TMap<UE::FSdfPath, TSet<UE::FSdfPath>> AuxToMainPrims;
+	TMap<UE::FSdfPath, TSet<UE::FSdfPath>> MainToAuxPrims;
+	mutable FRWLock AuxiliaryPrimsLock;
 };
 
 FUsdInfoCache::FUsdInfoCache()
@@ -93,8 +100,13 @@ bool FUsdInfoCache::Serialize( FArchive& Ar )
 			Ar << ImplPtr->AssetToPrimPaths;
 		}
 		{
-			FWriteScopeLock ScopeLock(ImplPtr->UsedMaterialPathsLock);
-			Ar << ImplPtr->UsedMaterialPaths;
+			FWriteScopeLock ScopeLock(ImplPtr->MaterialUsersLock);
+			Ar << ImplPtr->MaterialUsers;
+		}
+		{
+			FWriteScopeLock ScopeLock(ImplPtr->AuxiliaryPrimsLock);
+			Ar << ImplPtr->AuxToMainPrims;
+			Ar << ImplPtr->MainToAuxPrims;
 		}
 	}
 
@@ -190,13 +202,123 @@ UE::FSdfPath FUsdInfoCache::UnwindToNonCollapsedPath( const UE::FSdfPath& Path,
 	return Path;
 }
 
+void FUsdInfoCache::RegisterAuxiliaryPrim(const UE::FSdfPath& MainPrimPath, const UE::FSdfPath& AuxPrimPath) const
+{
+	if (FUsdInfoCacheImpl* ImplPtr = Impl.Get())
+	{
+		FWriteScopeLock ScopeLock(ImplPtr->AuxiliaryPrimsLock);
+
+		UE_LOG(LogUsd, Verbose, TEXT("Registering main prim '%s' and aux prim '%s'"),
+			*MainPrimPath.GetString(),
+			*AuxPrimPath.GetString()
+		);
+
+		ImplPtr->MainToAuxPrims.FindOrAdd(MainPrimPath).Add(AuxPrimPath);
+		ImplPtr->AuxToMainPrims.FindOrAdd(AuxPrimPath).Add(MainPrimPath);
+	}
+}
+
+void FUsdInfoCache::RegisterAuxiliaryPrims(const UE::FSdfPath& MainPrimPath, const TSet<UE::FSdfPath>& AuxPrimPaths) const
+{
+	if (FUsdInfoCacheImpl* ImplPtr = Impl.Get())
+	{
+		FWriteScopeLock ScopeLock(ImplPtr->AuxiliaryPrimsLock);
+
+		ImplPtr->MainToAuxPrims.FindOrAdd(MainPrimPath).Append(AuxPrimPaths);
+
+		for (const UE::FSdfPath& AuxPrimPath : AuxPrimPaths)
+		{
+			UE_LOG(LogUsd, Verbose, TEXT("Registering main prim '%s' and aux prim '%s'"),
+				*MainPrimPath.GetString(),
+				*AuxPrimPath.GetString()
+			);
+
+			ImplPtr->AuxToMainPrims.FindOrAdd(AuxPrimPath).Add(MainPrimPath);
+		}
+	}
+}
+
+void FUsdInfoCache::RemoveAuxiliaryPrims(const UE::FSdfPath& MainPrimPath) const
+{
+	if (FUsdInfoCacheImpl* ImplPtr = Impl.Get())
+	{
+		FWriteScopeLock ScopeLock(ImplPtr->AuxiliaryPrimsLock);
+
+		TSet<UE::FSdfPath> AuxPrimPaths;
+		ImplPtr->MainToAuxPrims.RemoveAndCopyValue(MainPrimPath, AuxPrimPaths);
+
+		for (const UE::FSdfPath& AuxPrimPath : AuxPrimPaths)
+		{
+			UE_LOG(LogUsd, Verbose, TEXT("Removing auxilliary prim '%s' from main prim '%s'"),
+				*AuxPrimPath.GetString(),
+				*MainPrimPath.GetString()
+			);
+
+			TSet<UE::FSdfPath>* MainPrimPaths = ImplPtr->AuxToMainPrims.Find(AuxPrimPath);
+			if (ensure(MainPrimPaths))
+			{
+				MainPrimPaths->Remove(MainPrimPath);
+			}
+		}
+	}
+}
+
+TSet<UE::FSdfPath> FUsdInfoCache::GetMainPrims(const UE::FSdfPath& AuxPrimPath) const
+{
+	if (FUsdInfoCacheImpl* ImplPtr = Impl.Get())
+	{
+		FReadScopeLock ScopeLock(ImplPtr->AuxiliaryPrimsLock);
+
+		if (const TSet<UE::FSdfPath>* FoundMainPrims = ImplPtr->AuxToMainPrims.Find(AuxPrimPath))
+		{
+			TSet<UE::FSdfPath> Result = *FoundMainPrims;
+			Result.Add(AuxPrimPath);
+			return Result;
+		}
+	}
+
+	return {AuxPrimPath};
+}
+
+TSet<UE::FSdfPath> FUsdInfoCache::GetAuxiliaryPrims(const UE::FSdfPath& MainPrimPath) const
+{
+	if (FUsdInfoCacheImpl* ImplPtr = Impl.Get())
+	{
+		FReadScopeLock ScopeLock(ImplPtr->AuxiliaryPrimsLock);
+
+		if (const TSet<UE::FSdfPath>* FoundAuxPrims = ImplPtr->MainToAuxPrims.Find(MainPrimPath))
+		{
+			TSet<UE::FSdfPath> Result = *FoundAuxPrims;
+			Result.Add(MainPrimPath);
+			return Result;
+		}
+	}
+
+	return {MainPrimPath};
+}
+
+TSet<UE::FSdfPath> FUsdInfoCache::GetMaterialUsers(const UE::FSdfPath& Path) const
+{
+	if (FUsdInfoCacheImpl* ImplPtr = Impl.Get())
+	{
+		FReadScopeLock ScopeLock(ImplPtr->MaterialUsersLock);
+
+		if (const TSet<UE::FSdfPath>* FoundUsers = ImplPtr->MaterialUsers.Find(Path))
+		{
+			return *FoundUsers;
+		}
+	}
+
+	return {};
+}
+
 bool FUsdInfoCache::IsMaterialUsed(const UE::FSdfPath& Path) const
 {
 	if (FUsdInfoCacheImpl* ImplPtr = Impl.Get())
 	{
-		FReadScopeLock ScopeLock(ImplPtr->UsedMaterialPathsLock);
+		FReadScopeLock ScopeLock(ImplPtr->MaterialUsersLock);
 
-		return ImplPtr->UsedMaterialPaths.Contains(Path);
+		return ImplPtr->MaterialUsers.Contains(Path);
 	}
 
 	return false;
@@ -345,8 +467,8 @@ namespace UE::USDInfoCacheImpl::Private
 			{
 				if (pxr::UsdShadeMaterial ShadeMaterial = BindingAPI.ComputeBoundMaterial(MaterialPurposeToken))
 				{
-					FWriteScopeLock ScopeLock(Impl.UsedMaterialPathsLock);
-					Impl.UsedMaterialPaths.Add(UE::FSdfPath{ShadeMaterial.GetPrim().GetPath()});
+					FWriteScopeLock ScopeLock(Impl.MaterialUsersLock);
+					Impl.MaterialUsers.FindOrAdd(UE::FSdfPath{ShadeMaterial.GetPrim().GetPath()}).Add(UE::FSdfPath{UsdPrimPath});
 				}
 			}
 			else if (pxr::UsdRelationship Relationship = UsdPrim.GetRelationship(pxr::UsdShadeTokens->materialBinding))
@@ -361,8 +483,8 @@ namespace UE::USDInfoCacheImpl::Private
 					pxr::UsdShadeMaterial UsdShadeMaterial{MaterialPrim};
 					if (UsdShadeMaterial)
 					{
-						FWriteScopeLock ScopeLock(Impl.UsedMaterialPathsLock);
-						Impl.UsedMaterialPaths.Add(UE::FSdfPath{TargetMaterialPrimPath});
+						FWriteScopeLock ScopeLock(Impl.MaterialUsersLock);
+						Impl.MaterialUsers.FindOrAdd(UE::FSdfPath{TargetMaterialPrimPath}).Add(UE::FSdfPath{UsdPrimPath});
 					}
 				}
 			}
@@ -797,11 +919,6 @@ TOptional<TArray<UsdUtils::FUsdPrimMaterialSlot>> FUsdInfoCache::GetSubtreeMater
 
 void FUsdInfoCache::LinkAssetToPrim(const UE::FSdfPath& Path, UObject* Asset)
 {
-	if (!Asset)
-	{
-		return;
-	}
-
 	FUsdInfoCacheImpl* ImplPtr = Impl.Get();
 	if (!ImplPtr)
 	{
@@ -827,6 +944,10 @@ TSet<TWeakObjectPtr<UObject>> FUsdInfoCache::RemoveAllAssetPrimLinks(const UE::F
 	}
 	FWriteScopeLock ScopeLock(ImplPtr->PrimPathToAssetsLock);
 
+	UE_LOG(LogUsd, Verbose, TEXT("Removing asset prim links for path '%s'"),
+		*Path.GetString()
+	);
+
 	TSet<TWeakObjectPtr<UObject>> Assets;
 	ImplPtr->PrimPathToAssets.RemoveAndCopyValue(Path, Assets);
 
@@ -917,14 +1038,15 @@ void FUsdInfoCache::RebuildCacheForSubtree( const UE::FUsdPrim& Prim, FUsdSchema
 		}
 
 		// Don't call Clear() here as we don't want to get rid of PrimPathToAssets because we're rebuilding the cache,
-		// as that info is also linked to the asset cache
+		// as that info is also linked to the asset cache. The same for the aux prim maps, that currently gets
+		// filled in on-demand.
 		{
 			FWriteScopeLock ScopeLock(ImplPtr->InfoMapLock);
 			ImplPtr->InfoMap.Empty();
 		}
 		{
-			FWriteScopeLock ScopeLock(ImplPtr->UsedMaterialPathsLock);
-			ImplPtr->UsedMaterialPaths.Empty();
+			FWriteScopeLock ScopeLock(ImplPtr->MaterialUsersLock);
+			ImplPtr->MaterialUsers.Empty();
 		}
 
 		IUsdSchemasModule& UsdSchemasModule = FModuleManager::Get().LoadModuleChecked< IUsdSchemasModule >( TEXT( "USDSchemas" ) );
@@ -993,8 +1115,13 @@ void FUsdInfoCache::Clear()
 			ImplPtr->AssetToPrimPaths.Empty();
 		}
 		{
-			FWriteScopeLock ScopeLock(ImplPtr->UsedMaterialPathsLock);
-			ImplPtr->UsedMaterialPaths.Empty();
+			FWriteScopeLock ScopeLock(ImplPtr->MaterialUsersLock);
+			ImplPtr->MaterialUsers.Empty();
+		}
+		{
+			FWriteScopeLock ScopeLock(ImplPtr->AuxiliaryPrimsLock);
+			ImplPtr->AuxToMainPrims.Empty();
+			ImplPtr->MainToAuxPrims.Empty();
 		}
 	}
 }
diff --git a/Engine/Plugins/Importers/USDImporter/Source/USDSchemas/Private/USDSchemaTranslator.cpp b/Engine/Plugins/Importers/USDImporter/Source/USDSchemas/Private/USDSchemaTranslator.cpp
--- a/Engine/Plugins/Importers/USDImporter/Source/USDSchemas/Private/USDSchemaTranslator.cpp
+++ b/Engine/Plugins/Importers/USDImporter/Source/USDSchemas/Private/USDSchemaTranslator.cpp
@@ -201,6 +201,18 @@ void FUsdSchemaTranslationContext::CompleteTasks()
 	}
 }
 
+void FUsdSchemaTranslator::RegisterAuxiliaryPrims()
+{
+	if (Context->InfoCache.IsValid())
+	{
+		TSet<UE::FSdfPath> AuxPrims = CollectAuxiliaryPrims();
+		if (AuxPrims.Num() > 0)
+		{
+			Context->InfoCache->RegisterAuxiliaryPrims(PrimPath, AuxPrims);
+		}
+	}
+}
+
 bool FUsdSchemaTranslator::IsCollapsed( ECollapsingType CollapsingType ) const
 {
 #if USE_USD_SDK
diff --git a/Engine/Plugins/Importers/USDImporter/Source/USDSchemas/Private/USDShadeMaterialTranslator.cpp b/Engine/Plugins/Importers/USDImporter/Source/USDSchemas/Private/USDShadeMaterialTranslator.cpp
--- a/Engine/Plugins/Importers/USDImporter/Source/USDSchemas/Private/USDShadeMaterialTranslator.cpp
+++ b/Engine/Plugins/Importers/USDImporter/Source/USDSchemas/Private/USDShadeMaterialTranslator.cpp
@@ -222,6 +222,8 @@ namespace UE::UsdShadeTranslator::Private
 
 void FUsdShadeMaterialTranslator::CreateAssets()
 {
+	TRACE_CPUPROFILER_EVENT_SCOPE( FUsdShadeMaterialTranslator::CreateAssets );
+
 	pxr::UsdShadeMaterial ShadeMaterial( GetPrim() );
 
 	if ( !ShadeMaterial )
@@ -238,7 +240,7 @@ void FUsdShadeMaterialTranslator::CreateAssets()
 		}
 	}
 
-	TRACE_CPUPROFILER_EVENT_SCOPE( FUsdShadeMaterialTranslator::CreateAssets );
+	RegisterAuxiliaryPrims();
 
 	const pxr::TfToken RenderContextToken =
 		Context->RenderContext.IsNone() ?
@@ -483,4 +485,64 @@ bool FUsdShadeMaterialTranslator::CanBeCollapsed( ECollapsingType CollapsingType
 	return false;
 }
 
+TSet<UE::FSdfPath> FUsdShadeMaterialTranslator::CollectAuxiliaryPrims() const
+{
+	TSet<UE::FSdfPath> Result;
+	{
+		TFunction<void(const pxr::UsdShadeInput&)> TraverseShadeInput;
+		TraverseShadeInput = [&TraverseShadeInput, &Result](const pxr::UsdShadeInput& ShadeInput)
+		{
+			if (!ShadeInput)
+			{
+				return;
+			}
+
+			pxr::UsdShadeConnectableAPI Source;
+			pxr::TfToken SourceName;
+			pxr::UsdShadeAttributeType AttributeType;
+			if (pxr::UsdShadeConnectableAPI::GetConnectedSource(ShadeInput.GetAttr(), &Source, &SourceName, &AttributeType))
+			{
+				pxr::UsdPrim ConnectedPrim = Source.GetPrim();
+				UE::FSdfPath ConnectedPrimPath = UE::FSdfPath{ConnectedPrim.GetPrimPath()};
+
+				if (!Result.Contains(ConnectedPrimPath))
+				{
+					Result.Add(ConnectedPrimPath);
+
+					for (const pxr::UsdShadeInput& ChildInput : Source.GetInputs())
+					{
+						TraverseShadeInput(ChildInput);
+					}
+				}
+			}
+		};
+
+		FScopedUsdAllocs UsdAllocs;
+
+		pxr::UsdPrim Prim = GetPrim();
+		pxr::UsdShadeMaterial UsdShadeMaterial{Prim};
+		if (!UsdShadeMaterial)
+		{
+			return {};
+		}
+
+		const pxr::TfToken RenderContextToken = Context->RenderContext.IsNone()
+			? pxr::UsdShadeTokens->universalRenderContext
+			: UnrealToUsd::ConvertToken(*Context->RenderContext.ToString()).Get();
+		pxr::UsdShadeShader SurfaceShader = UsdShadeMaterial.ComputeSurfaceSource({RenderContextToken});
+		if (!SurfaceShader)
+		{
+			return {};
+		}
+
+		Result.Add(UE::FSdfPath{SurfaceShader.GetPrim().GetPrimPath()});
+
+		for (const pxr::UsdShadeInput& ShadeInput : SurfaceShader.GetInputs())
+		{
+			TraverseShadeInput(ShadeInput);
+		}
+	}
+	return Result;
+}
+
 #endif // #if USE_USD_SDK
diff --git a/Engine/Plugins/Importers/USDImporter/Source/USDSchemas/Private/USDSkelRootTranslator.cpp b/Engine/Plugins/Importers/USDImporter/Source/USDSchemas/Private/USDSkelRootTranslator.cpp
--- a/Engine/Plugins/Importers/USDImporter/Source/USDSchemas/Private/USDSkelRootTranslator.cpp
+++ b/Engine/Plugins/Importers/USDImporter/Source/USDSchemas/Private/USDSkelRootTranslator.cpp
@@ -55,6 +55,7 @@
 #endif // WITH_EDITOR
 
 #include "USDIncludesStart.h"
+	#include "pxr/usd/usd/primRange.h"
 	#include "pxr/usd/usdGeom/mesh.h"
 	#include "pxr/usd/usdShade/material.h"
 	#include "pxr/usd/usdSkel/binding.h"
@@ -204,7 +205,11 @@ namespace UsdSkelRootTranslatorImpl
 		return bMaterialsHaveChanged;
 	}
 
-	FSHAHash ComputeSHAHash( const TArray<FSkeletalMeshImportData>& LODIndexToSkeletalMeshImportData, TArray<SkeletalMeshImportData::FBone>& ImportedBones )
+	FSHAHash ComputeSHAHash(
+		const TArray<FSkeletalMeshImportData>& LODIndexToSkeletalMeshImportData,
+		TArray<SkeletalMeshImportData::FBone>& ImportedBones,
+		UsdUtils::FBlendShapeMap* BlendShapes
+	)
 	{
 		FSHA1 HashState;
 
@@ -227,6 +232,25 @@ namespace UsdSkelRootTranslatorImpl
 			HashState.Update( ( uint8* ) &Bone.BonePos, sizeof( Bone.BonePos ) );
 		}
 
+		if (BlendShapes)
+		{
+			for (const TPair<FString, UsdUtils::FUsdBlendShape>& Pair : *BlendShapes)
+			{
+				const UsdUtils::FUsdBlendShape& BlendShape = Pair.Value;
+				HashState.UpdateWithString(*BlendShape.Name, BlendShape.Name.Len());
+
+				HashState.Update((uint8*)BlendShape.Vertices.GetData(), BlendShape.Vertices.Num() * BlendShape.Vertices.GetTypeSize());
+
+				for (const UsdUtils::FUsdBlendShapeInbetween& InBetween : BlendShape.Inbetweens)
+				{
+					HashState.UpdateWithString(*InBetween.Name, InBetween.Name.Len());
+					HashState.Update((uint8*)&InBetween.InbetweenWeight, sizeof(InBetween.InbetweenWeight));
+				}
+
+				HashState.Update((uint8*)&BlendShape.bHasAuthoredTangents, sizeof(BlendShape.bHasAuthoredTangents));
+			}
+		}
+
 		FSHAHash OutHash;
 
 		HashState.Final();
@@ -235,7 +259,11 @@ namespace UsdSkelRootTranslatorImpl
 		return OutHash;
 	}
 
-	FSHAHash ComputeSHAHash( const pxr::UsdSkelSkeletonQuery& InUsdSkeletonQuery, const pxr::UsdPrim& RootMotionPrim )
+	FSHAHash ComputeSHAHash(
+		const pxr::UsdSkelSkeletonQuery& InUsdSkeletonQuery,
+		const pxr::UsdPrim& RootMotionPrim,
+		const FString& SkeletalMeshHashString
+	)
 	{
 		TRACE_CPUPROFILER_EVENT_SCOPE( UsdSkelRootTranslatorImpl::ComputeSHAHash_SkelQuery );
 
@@ -265,6 +293,19 @@ namespace UsdSkelRootTranslatorImpl
 		int32 InterpolationType = static_cast< int32 >( Stage->GetInterpolationType() );
 		HashState.Update( ( uint8* ) &InterpolationType, sizeof( int32 ) );
 
+		// Hash blend shape and joint order tokens
+		TFunction<void(const pxr::VtArray<pxr::TfToken>&)> HashTokens = [&HashState](const pxr::VtArray<pxr::TfToken>& Tokens)
+		{
+			for (const pxr::TfToken& Token : Tokens)
+			{
+				const std::string& TokenString = Token.GetString();
+				const char* TokenCStr = TokenString.c_str();
+				HashState.Update((const uint8*)TokenCStr, sizeof(char) * TokenString.size());
+			}
+		};
+		HashTokens(AnimQuery.GetJointOrder());
+		HashTokens(AnimQuery.GetBlendShapeOrder());
+
 		// Time samples for joint transforms
 		std::vector<double> TimeData;
 		AnimQuery.GetJointTransformTimeSamples( &TimeData );
@@ -318,6 +359,14 @@ namespace UsdSkelRootTranslatorImpl
 			HashState.Update( ( uint8* ) &bResetsXformSack, sizeof( bool ) );
 		}
 
+		// An anim sequence matches a particular skeleton. If the skeleton is different, we'll likely
+		// need a new AnimSequence, even if the SkelAnimation prim itself hashes the same. The same
+		// applies to SkeletalMesh morph targets too: In USD a SkelAnimation has decoupled "blend
+		// shape channel" curves that can animate blend shapes specific to each mesh, but in UE the
+		// AnimSequence curves are specific to each morph target name, so if we have a different mesh,
+		// we're probably better off ensuring we have a different AnimSequence too
+		HashState.UpdateWithString(*SkeletalMeshHashString, SkeletalMeshHashString.Len());
+
 		HashState.Final();
 		HashState.GetHash( &OutHash.Hash[ 0 ] );
 
@@ -1182,6 +1231,7 @@ namespace UsdSkelRootTranslatorImpl
 		// Note that we want this to be case insensitive so that our UMorphTarget FNames are unique not only due to case differences
 		TSet<FString> UsedMorphTargetNames;
 		TUsdStore< pxr::UsdSkelCache > SkeletonCache;
+		FString SkeletalMeshHashString;
 
 		// Don't keep a live reference to the prim because other translators may mutate the stage in an ExclusiveSync translation step, invalidating the reference
 		UE::FUsdPrim GetPrim() const { return Context->Stage.GetPrimAtPath( PrimPath ); }
@@ -1258,9 +1308,16 @@ namespace UsdSkelRootTranslatorImpl
 		Then( ESchemaTranslationLaunchPolicy::Sync,
 			[ this ]()
 			{
+				if (!Context->InfoCache || !Context->AssetCache)
+				{
+					return false;
+				}
+
+				UsdUtils::FBlendShapeMap* BlendShapes = Context->BlendShapesByPath ? &NewBlendShapes : nullptr;
+
 				FString SkelRootPath = PrimPath.GetString();
-				FSHAHash SkeletalMeshHash = UsdSkelRootTranslatorImpl::ComputeSHAHash( LODIndexToSkeletalMeshImportData, SkeletonBones );
-				const FString SkeletalMeshHashString = SkeletalMeshHash.ToString();
+				FSHAHash SkeletalMeshHash = UsdSkelRootTranslatorImpl::ComputeSHAHash( LODIndexToSkeletalMeshImportData, SkeletonBones, BlendShapes );
+				SkeletalMeshHashString = SkeletalMeshHash.ToString();
 
 				USkeletalMesh* SkeletalMesh = Cast< USkeletalMesh >( Context->AssetCache->GetCachedAsset( SkeletalMeshHashString ) );
 
@@ -1280,7 +1337,7 @@ namespace UsdSkelRootTranslatorImpl
 
 				if ( SkeletalMesh )
 				{
-					if (bIsNew && Context->AssetCache && Context->InfoCache)
+					if (bIsNew)
 					{
 						UUsdMeshAssetImportData* ImportData = NewObject< UUsdMeshAssetImportData >( SkeletalMesh, TEXT( "USDAssetImportData" ) );
 						ImportData->PrimPath = SkelRootPath;
@@ -1322,6 +1379,7 @@ namespace UsdSkelRootTranslatorImpl
 									SkeletalMeshHashString + TEXT( "_PhysicsAsset" ),
 									PhysicsAsset
 								);
+								Context->InfoCache->LinkAssetToPrim(PrimPath, PhysicsAsset);
 							}
 						}
 					}
@@ -1331,10 +1389,7 @@ namespace UsdSkelRootTranslatorImpl
 						SkeletalMesh->SetPhysicsAsset( nullptr );
 					}
 
-					if (Context->InfoCache)
-					{
-						Context->InfoCache->LinkAssetToPrim(PrimPath, SkeletalMesh);
-					}
+					Context->InfoCache->LinkAssetToPrim(PrimPath, SkeletalMesh);
 
 					// Track our Skeleton by the source skeleton prim path
 					if ( USkeleton* Skeleton = SkeletalMesh->GetSkeleton() )
@@ -1349,10 +1404,7 @@ namespace UsdSkelRootTranslatorImpl
 						{
 							pxr::UsdSkelBinding& SkeletonBinding = SkeletonBindings[ 0 ];
 							const pxr::UsdSkelSkeleton& UsdSkeleton = SkeletonBinding.GetSkeleton();
-							Context->InfoCache->LinkAssetToPrim(
-								UE::FSdfPath{UsdSkeleton.GetPrim().GetPrimPath()},
-								Skeleton
-							);
+							Context->InfoCache->LinkAssetToPrim(PrimPath, Skeleton);
 						}
 					}
 
@@ -1468,7 +1520,11 @@ namespace UsdSkelRootTranslatorImpl
 								? Skeleton.GetPrim()
 								: pxr::UsdPrim{};
 
-						FSHAHash Hash = UsdSkelRootTranslatorImpl::ComputeSHAHash( SkelQuery, RootMotionPrim );
+						FSHAHash Hash = UsdSkelRootTranslatorImpl::ComputeSHAHash(
+							SkelQuery,
+							RootMotionPrim,
+							SkeletalMeshHashString
+						);
 						FString HashString = Hash.ToString();
 						UAnimSequence* AnimSequence = Cast< UAnimSequence >( Context->AssetCache->GetCachedAsset( HashString ) );
 
@@ -1480,6 +1536,8 @@ namespace UsdSkelRootTranslatorImpl
 							// Bone/CurveCompressionSettings. Undoing that transaction would call UAnimSequence::Serialize with nullptr values for both, which crashes.
 							// Besides, this particular asset type is only ever created when we import to content folder assets (so never for realtime), and
 							// in that case we don't need it to be transactional anyway
+							// TODO: I think this changed and now we generate AnimSequences when opening the stage too,
+							// so this may lead to issues
 							AnimSequence = NewObject<UAnimSequence>(
 								GetTransientPackage(),
 								NAME_None,
@@ -1520,10 +1578,7 @@ namespace UsdSkelRootTranslatorImpl
 
 						if (AnimSequence && Context->InfoCache)
 						{
-							Context->InfoCache->LinkAssetToPrim(
-								UE::FSdfPath{*SkelAnimationPrimPath},
-								AnimSequence
-							);
+							Context->InfoCache->LinkAssetToPrim(PrimPath, AnimSequence);
 						}
 
 						// For now we shouldn't try to parse the SkelAnimations from skeletal bindings other than the first one as we only
@@ -1543,7 +1598,11 @@ namespace UsdSkelRootTranslatorImpl
 
 void FUsdSkelRootTranslator::CreateAssets()
 {
+	TRACE_CPUPROFILER_EVENT_SCOPE(FUsdSkelRootTranslator::CreateAssets);
+
 #if WITH_EDITOR
+	RegisterAuxiliaryPrims();
+
 	// Importing skeletal meshes actually works in Standalone mode, but we intentionally block it here
 	// to not confuse users as to why it doesn't work at runtime
 	TSharedRef< UsdSkelRootTranslatorImpl::FSkelRootCreateAssetsTaskChain > AssetsTaskChain =
@@ -1555,6 +1614,8 @@ void FUsdSkelRootTranslator::CreateAssets()
 
 USceneComponent* FUsdSkelRootTranslator::CreateComponents()
 {
+	TRACE_CPUPROFILER_EVENT_SCOPE(FUsdSkelRootTranslator::CreateComponents);
+
 	USceneComponent* SceneComponent = FUsdGeomXformableTranslator::CreateComponents();
 
 #if WITH_EDITOR
@@ -1611,6 +1672,8 @@ void FUsdSkelRootTranslator::UpdateComponents( USceneComponent* SceneComponent )
 		}
 	}
 
+	SkeletalMeshComponent->Modify();
+
 	SkeletalMeshComponent->SetAnimationMode(
 		bPrimHasLiveLinkEnabled
 			? EAnimationMode::AnimationBlueprint
@@ -1623,7 +1686,7 @@ void FUsdSkelRootTranslator::UpdateComponents( USceneComponent* SceneComponent )
 	if (SkelAnimPrim)
 	{
 		UAnimSequence* TargetAnimSequence = Context->InfoCache->GetSingleAssetForPrim<UAnimSequence>(
-			SkelAnimPrim.GetPrimPath()
+			Prim.GetPrimPath()
 		);
 		if ( TargetAnimSequence != SkeletalMeshComponent->AnimationData.AnimToPlay )
 		{
@@ -1762,6 +1825,97 @@ void FUsdSkelRootTranslator::UpdateComponents( USceneComponent* SceneComponent )
 #endif // WITH_EDITOR
 }
 
+bool FUsdSkelRootTranslator::CollapsesChildren(ECollapsingType CollapsingType) const
+{
+	return true;
+}
+
+bool FUsdSkelRootTranslator::CanBeCollapsed(ECollapsingType CollapsingType) const
+{
+	return false;
+}
+
+TSet<UE::FSdfPath> FUsdSkelRootTranslator::CollectAuxiliaryPrims() const
+{
+	if (!Context->InfoCache.IsValid())
+	{
+		return {};
+	}
+
+	if (!Context->InfoCache->DoesPathCollapseChildren(PrimPath, ECollapsingType::Assets))
+	{
+		return {};
+	}
+
+	TSet<UE::FSdfPath> Result;
+	{
+		FScopedUsdAllocs UsdAllocs;
+
+		pxr::UsdPrim Prim = GetPrim();
+
+		pxr::TfType ImageableSchema = pxr::TfType::Find<pxr::UsdGeomImageable>();
+		pxr::TfType MeshSchema = pxr::TfType::Find<pxr::UsdGeomMesh>();
+
+		pxr::UsdPrimRange PrimRange{Prim, pxr::UsdTraverseInstanceProxies()};
+		for (pxr::UsdPrimRange::iterator PrimRangeIt = PrimRange.begin(); PrimRangeIt != PrimRange.end(); ++PrimRangeIt)
+		{
+			if (PrimRangeIt->IsA(MeshSchema))
+			{
+				Result.Add(UE::FSdfPath{PrimRangeIt->GetPrimPath()});
+
+				if (pxr::UsdSkelBindingAPI SkelBindingAPI{*PrimRangeIt})
+				{
+					// Collect blend shapes, which don't have to be within the Mesh or SkelRoot prim at all
+					if (pxr::UsdSkelBlendShapeQuery BlendShapeQuery{SkelBindingAPI})
+					{
+						for (int32 BlendShapeIndex = 0; BlendShapeIndex < BlendShapeQuery.GetNumBlendShapes(); ++BlendShapeIndex)
+						{
+							if (pxr::UsdSkelBlendShape BlendShape = BlendShapeQuery.GetBlendShape(BlendShapeIndex))
+							{
+								Result.Add(UE::FSdfPath{BlendShape.GetPrim().GetPrimPath()});
+							}
+						}
+					}
+				}
+			}
+			// All meshes, xforms, skeleton are imageables. Registering any imageable is also a good idea because their
+			// visibility could affect child Mesh prims, and so the combined skeletal mesh.
+			else if (PrimRangeIt->IsA(ImageableSchema))
+			{
+				Result.Add(UE::FSdfPath{PrimRangeIt->GetPrimPath()});
+			}
+		}
+
+		// Collect the bound SkelAnimation, that doesn't have to be within the actual SkelRoot
+		if (pxr::UsdSkelRoot SkeletonRoot{Prim})
+		{
+			pxr::UsdSkelCache SkeletonCache;
+
+			std::vector< pxr::UsdSkelBinding > SkeletonBindings;
+			SkeletonCache.Populate(SkeletonRoot, pxr::UsdTraverseInstanceProxies());
+			SkeletonCache.ComputeSkelBindings(SkeletonRoot, &SkeletonBindings, pxr::UsdTraverseInstanceProxies());
+
+			if (SkeletonBindings.size() > 0)
+			{
+				// We only handle the first skeleton binding
+				const pxr::UsdSkelBinding& Binding = SkeletonBindings[0];
+
+				const pxr::UsdSkelSkeleton& Skeleton = Binding.GetSkeleton();
+				pxr::UsdSkelSkeletonQuery SkelQuery = SkeletonCache.GetSkelQuery(Skeleton);
+
+				if (pxr::UsdSkelAnimQuery AnimQuery = SkelQuery.GetAnimQuery())
+				{
+					if (pxr::UsdPrim SkelAnimationPrim = AnimQuery.GetPrim())
+					{
+						Result.Add(UE::FSdfPath{SkelAnimationPrim.GetPrimPath()});
+					}
+				}
+			}
+		}
+	}
+	return Result;
+}
+
 #undef LOCTEXT_NAMESPACE
 
 #endif // #if USE_USD_SDK
diff --git a/Engine/Plugins/Importers/USDImporter/Source/USDSchemas/Public/USDGeomMeshTranslator.h b/Engine/Plugins/Importers/USDImporter/Source/USDSchemas/Public/USDGeomMeshTranslator.h
--- a/Engine/Plugins/Importers/USDImporter/Source/USDSchemas/Public/USDGeomMeshTranslator.h
+++ b/Engine/Plugins/Importers/USDImporter/Source/USDSchemas/Public/USDGeomMeshTranslator.h
@@ -30,7 +30,11 @@ namespace UsdUtils
 class FBuildStaticMeshTaskChain : public FUsdSchemaTranslatorTaskChain
 {
 public:
-	explicit FBuildStaticMeshTaskChain( const TSharedRef< FUsdSchemaTranslationContext >& InContext, const UE::FSdfPath& InPrimPath );
+	explicit FBuildStaticMeshTaskChain(
+		const TSharedRef< FUsdSchemaTranslationContext >& InContext,
+		const UE::FSdfPath& InPrimPath,
+		const TOptional<UE::FSdfPath>& AlternativePrimToLinkAssetsTo = {}
+	);
 
 protected:
 	// Inputs
@@ -41,6 +45,7 @@ protected:
 	TArray<UsdUtils::FUsdPrimMaterialAssignmentInfo> LODIndexToMaterialInfo;
 
 	// Outputs
+	TOptional<UE::FSdfPath> AlternativePrimToLinkAssetsTo;
 	UStaticMesh* StaticMesh = nullptr;
 
 	// Required to prevent StaticMesh from being used for drawing while it is being rebuilt
@@ -55,7 +60,12 @@ protected:
 class FGeomMeshCreateAssetsTaskChain : public FBuildStaticMeshTaskChain
 {
 public:
-	explicit FGeomMeshCreateAssetsTaskChain( const TSharedRef< FUsdSchemaTranslationContext >& InContext, const UE::FSdfPath& PrimPath, const FTransform& AdditionalTransform = FTransform::Identity );
+	explicit FGeomMeshCreateAssetsTaskChain(
+		const TSharedRef< FUsdSchemaTranslationContext >& InContext,
+		const UE::FSdfPath& PrimPath,
+		const TOptional<UE::FSdfPath>& AlternativePrimToLinkAssetsTo = {},
+		const FTransform& AdditionalTransform = FTransform::Identity
+	);
 
 protected:
 	// Inputs
@@ -82,6 +92,7 @@ public:
 	virtual bool CollapsesChildren( ECollapsingType CollapsingType ) const override;
 	virtual bool CanBeCollapsed( ECollapsingType CollapsingType ) const override;
 
+	virtual TSet<UE::FSdfPath> CollectAuxiliaryPrims() const override;
 };
 
 #endif // #if USE_USD_SDK
diff --git a/Engine/Plugins/Importers/USDImporter/Source/USDSchemas/Public/USDGeomPointInstancerTranslator.h b/Engine/Plugins/Importers/USDImporter/Source/USDSchemas/Public/USDGeomPointInstancerTranslator.h
--- a/Engine/Plugins/Importers/USDImporter/Source/USDSchemas/Public/USDGeomPointInstancerTranslator.h
+++ b/Engine/Plugins/Importers/USDImporter/Source/USDSchemas/Public/USDGeomPointInstancerTranslator.h
@@ -12,7 +12,12 @@ struct FUsdSchemaTranslationContext;
 class FUsdGeomPointInstancerCreateAssetsTaskChain : public FBuildStaticMeshTaskChain
 {
 public:
-	explicit FUsdGeomPointInstancerCreateAssetsTaskChain( const TSharedRef< FUsdSchemaTranslationContext >& InContext, const UE::FSdfPath& InPrimPath, bool bIgnoreTopLevelTransformAndVisibility );
+	explicit FUsdGeomPointInstancerCreateAssetsTaskChain(
+		const TSharedRef<FUsdSchemaTranslationContext>& InContext,
+		const UE::FSdfPath& InPrimPath,
+		bool bIgnoreTopLevelTransformAndVisibility,
+		const TOptional<UE::FSdfPath>& AlternativePrimToLinkAssetsTo = {}
+	);
 
 protected:
 	virtual void SetupTasks() override;
@@ -38,6 +43,8 @@ public:
 	// collapsing even simple, not-nested point instancers).
 	virtual bool CollapsesChildren( ECollapsingType CollapsingType ) const override { return true; }
 	virtual bool CanBeCollapsed( ECollapsingType CollapsingType ) const override;
+
+	virtual TSet<UE::FSdfPath> CollectAuxiliaryPrims() const override;
 };
 
 #endif // #if USE_USD_SDK
diff --git a/Engine/Plugins/Importers/USDImporter/Source/USDSchemas/Public/USDGeomXformableTranslator.h b/Engine/Plugins/Importers/USDImporter/Source/USDSchemas/Public/USDGeomXformableTranslator.h
--- a/Engine/Plugins/Importers/USDImporter/Source/USDSchemas/Public/USDGeomXformableTranslator.h
+++ b/Engine/Plugins/Importers/USDImporter/Source/USDSchemas/Public/USDGeomXformableTranslator.h
@@ -32,6 +32,8 @@ public:
 	virtual bool CollapsesChildren( ECollapsingType CollapsingType ) const override;
 	virtual bool CanBeCollapsed( ECollapsingType CollapsingType ) const override;
 
+	virtual TSet<UE::FSdfPath> CollectAuxiliaryPrims() const override;
+
 	// If the optional parameters are not set, we'll figure them out automatically.
 	USceneComponent* CreateComponentsEx( TOptional< TSubclassOf< USceneComponent > > ComponentType, TOptional< bool > bNeedsActor );
 
diff --git a/Engine/Plugins/Importers/USDImporter/Source/USDSchemas/Public/USDGroomTranslator.h b/Engine/Plugins/Importers/USDImporter/Source/USDSchemas/Public/USDGroomTranslator.h
--- a/Engine/Plugins/Importers/USDImporter/Source/USDSchemas/Public/USDGroomTranslator.h
+++ b/Engine/Plugins/Importers/USDImporter/Source/USDSchemas/Public/USDGroomTranslator.h
@@ -20,6 +20,8 @@ public:
 	virtual bool CollapsesChildren(ECollapsingType CollapsingType) const override;
 	virtual bool CanBeCollapsed(ECollapsingType CollapsingType) const override;
 
+	virtual TSet<UE::FSdfPath> CollectAuxiliaryPrims() const override;
+
 private:
 	bool IsGroomPrim() const;
 };
diff --git a/Engine/Plugins/Importers/USDImporter/Source/USDSchemas/Public/USDInfoCache.h b/Engine/Plugins/Importers/USDImporter/Source/USDSchemas/Public/USDInfoCache.h
--- a/Engine/Plugins/Importers/USDImporter/Source/USDSchemas/Public/USDInfoCache.h
+++ b/Engine/Plugins/Importers/USDImporter/Source/USDSchemas/Public/USDInfoCache.h
@@ -55,6 +55,26 @@ public:
 	UE::FSdfPath UnwindToNonCollapsedPath(const UE::FSdfPath& Path, ECollapsingType CollapsingType) const;
 
 public:
+	// Records that whenever we want to translate MainPrim, we need to read one or more attributes off AuxPrim
+	void RegisterAuxiliaryPrim(const UE::FSdfPath& MainPrimPath, const UE::FSdfPath& AuxPrimPath) const;
+	void RegisterAuxiliaryPrims(const UE::FSdfPath& MainPrimPath, const TSet<UE::FSdfPath>& AuxPrimPaths) const;
+
+	void RemoveAuxiliaryPrims(const UE::FSdfPath& MainPrimPath) const;
+
+	// Returns the paths to prims that, when translated into assets or components, also require reading the prim at
+	// 'Path'. e.g. providing the path to a Shader prim will return the paths to all Material prims for which the
+	// translation involves reading that particular Shader.
+	TSet<UE::FSdfPath> GetMainPrims(const UE::FSdfPath& AuxPrimPath) const;
+
+	// The inverse of the function above: Provide it with the path to a Material prim and it will return the set of
+	// paths to all Shader prims that need to be read to translate that Material prim into material assets
+	TSet<UE::FSdfPath> GetAuxiliaryPrims(const UE::FSdfPath& MainPrimPath) const;
+
+public:
+	// Returns the set of paths to all prims that have a material:binding relationship to the particular material at
+	// 'Path', if any.
+	// Returns a copy for thread safety.
+	TSet<UE::FSdfPath> GetMaterialUsers(const UE::FSdfPath& Path) const;
 	bool IsMaterialUsed(const UE::FSdfPath& Path) const;
 
 public:
diff --git a/Engine/Plugins/Importers/USDImporter/Source/USDSchemas/Public/USDSchemaTranslator.h b/Engine/Plugins/Importers/USDImporter/Source/USDSchemas/Public/USDSchemaTranslator.h
--- a/Engine/Plugins/Importers/USDImporter/Source/USDSchemas/Public/USDSchemaTranslator.h
+++ b/Engine/Plugins/Importers/USDImporter/Source/USDSchemas/Public/USDSchemaTranslator.h
@@ -276,6 +276,12 @@ public:
 
 	virtual ~FUsdSchemaTranslator() = default;
 
+	/**
+	 * Translates the prim at PrimPath, creating assets.
+	 * Derived classes that wish to implement FUsdSchemaTranslator::CreateAssets should manually call
+	 * FUsdSchemaTranslator::RegisterAuxiliaryPrims and implement CollectAuxiliaryPrims to catalog their
+	 * dependency prims and regenerate assets when those prims are updated.
+	 */
 	virtual void CreateAssets() {}
 
 	virtual USceneComponent* CreateComponents() { return nullptr; }
@@ -283,11 +289,22 @@ public:
 
 	virtual bool CollapsesChildren( ECollapsingType CollapsingType ) const { return false; }
 
+	/**
+	 * Returns the set of prims that also need to be read in order to translate the prim at PrimPath.
+	 * Note: This function never needs to return PrimPath itself, as the query function in the InfoCache will always
+	 * append it to the result
+	 */
+	virtual TSet<UE::FSdfPath> CollectAuxiliaryPrims() const { return {}; }
+
 	bool IsCollapsed( ECollapsingType CollapsingType ) const;
 	virtual bool CanBeCollapsed( ECollapsingType CollapsingType ) const { return false; }
 
 	UE::FUsdPrim GetPrim() const { return Context->Stage.GetPrimAtPath(PrimPath); }
 
+protected:
+	/** Calls CollectAuxiliaryPrims to record the current dependencies in the Context's InfoCache */
+	void RegisterAuxiliaryPrims();
+
 protected:
 	UE::FSdfPath PrimPath;
 	TSharedRef< FUsdSchemaTranslationContext > Context;
diff --git a/Engine/Plugins/Importers/USDImporter/Source/USDSchemas/Public/USDShadeMaterialTranslator.h b/Engine/Plugins/Importers/USDImporter/Source/USDSchemas/Public/USDShadeMaterialTranslator.h
--- a/Engine/Plugins/Importers/USDImporter/Source/USDSchemas/Public/USDShadeMaterialTranslator.h
+++ b/Engine/Plugins/Importers/USDImporter/Source/USDSchemas/Public/USDShadeMaterialTranslator.h
@@ -19,6 +19,8 @@ public:
 
 	virtual bool CollapsesChildren( ECollapsingType CollapsingType ) const override;
 	virtual bool CanBeCollapsed( ECollapsingType CollapsingType ) const override;
+
+	virtual TSet<UE::FSdfPath> CollectAuxiliaryPrims() const override;
 };
 
 #endif // #if USE_USD_SDK
diff --git a/Engine/Plugins/Importers/USDImporter/Source/USDSchemas/Public/USDSkelRootTranslator.h b/Engine/Plugins/Importers/USDImporter/Source/USDSchemas/Public/USDSkelRootTranslator.h
--- a/Engine/Plugins/Importers/USDImporter/Source/USDSchemas/Public/USDSkelRootTranslator.h
+++ b/Engine/Plugins/Importers/USDImporter/Source/USDSchemas/Public/USDSkelRootTranslator.h
@@ -27,8 +27,10 @@ public:
 	virtual USceneComponent* CreateComponents() override;
 	virtual void UpdateComponents( USceneComponent* SceneComponent ) override;
 
-	virtual bool CollapsesChildren( ECollapsingType CollapsingType ) const override { return true; }
-	virtual bool CanBeCollapsed( ECollapsingType CollapsingType ) const override { return false; }
+	virtual bool CollapsesChildren(ECollapsingType CollapsingType) const override;
+	virtual bool CanBeCollapsed(ECollapsingType CollapsingType) const override;
+
+	virtual TSet<UE::FSdfPath> CollectAuxiliaryPrims() const override;
 };
 
 #endif //#if USE_USD_SDK
diff --git a/Engine/Plugins/Importers/USDImporter/Source/USDStage/Private/USDLevelSequenceHelper.cpp b/Engine/Plugins/Importers/USDImporter/Source/USDStage/Private/USDLevelSequenceHelper.cpp
--- a/Engine/Plugins/Importers/USDImporter/Source/USDStage/Private/USDLevelSequenceHelper.cpp
+++ b/Engine/Plugins/Importers/USDImporter/Source/USDStage/Private/USDLevelSequenceHelper.cpp
@@ -1853,7 +1853,7 @@ void FUsdLevelSequenceHelperImpl::AddSkeletalTracks(const UUsdPrimTwin& PrimTwin
 	// Fetch the UAnimSequence asset from the asset cache. Ideally we'd call AUsdStageActor::GetGeneratedAssets,
 	// but we may belong to a FUsdStageImportContext, and so there's no AUsdStageActor at all to use.
 	// At this point it doesn't matter much though, because we shouldn't need to uncollapse a SkelAnimation prim path anyway
-	const UE::FSdfPath PrimPath = SkelAnimationPrim.GetPrimPath();
+	const UE::FSdfPath PrimPath = Prim.GetPrimPath();
 	UAnimSequence* Sequence = InfoCache->GetSingleAssetForPrim<UAnimSequence>(
 		PrimPath
 	);
@@ -2204,6 +2204,8 @@ void FUsdLevelSequenceHelperImpl::UpdateControlRigTracks(UUsdPrimTwin& PrimTwin)
 		return;
 	}
 
+	ensure(UsdPrim.IsA(TEXT("SkelRoot")));
+
 	USkeletalMeshComponent* ComponentToBind = Cast<USkeletalMeshComponent>(PrimTwin.GetSceneComponent());
 	if (!ComponentToBind)
 	{
@@ -2229,8 +2231,6 @@ void FUsdLevelSequenceHelperImpl::UpdateControlRigTracks(UUsdPrimTwin& PrimTwin)
 		// If this SkelRoot doesn't have any animation, lets create a new one on the current edit target
 		SkelAnimationLayer = UsdStage.GetEditTarget();
 
-		ensure(UsdPrim.IsA(TEXT("SkelRoot")));
-
 		FString UniqueChildName = UsdUtils::GetValidChildName(TEXT("Animation"), UsdPrim);
 		SkelAnimationPrim = UsdStage.DefinePrim(
 			UsdPrim.GetPrimPath().AppendChild(*UniqueChildName),
@@ -2243,19 +2243,18 @@ void FUsdLevelSequenceHelperImpl::UpdateControlRigTracks(UUsdPrimTwin& PrimTwin)
 
 		UsdUtils::BindAnimationSource(UsdPrim, SkelAnimationPrim);
 	}
+	if (!SkelAnimationLayer)
+	{
+		return;
+	}
 
 	// Fetch the UAnimSequence asset from the asset cache. Ideally we'd call AUsdStageActor::GetGeneratedAssets,
 	// but we may belong to a FUsdStageImportContext, and so there's no AUsdStageActor at all to use.
 	// At this point it doesn't matter much though, because we shouldn't need to uncollapse a SkelAnimation prim path anyway
-	const UE::FSdfPath SkelAnimationPrimPath = SkelAnimationPrim.GetPrimPath();
 	UAnimSequence* AnimSequence = InfoCache->GetSingleAssetForPrim<UAnimSequence>(
-		SkelAnimationPrimPath
+		PrimPath
 	);
 
-	if (!SkelAnimationLayer)
-	{
-		return;
-	}
 	UE::FUsdEditContext EditContext{UsdStage, SkelAnimationLayer};
 	FString Identifier = SkelAnimationLayer.GetIdentifier();
 
diff --git a/Engine/Plugins/Importers/USDImporter/Source/USDStage/Private/USDStageActor.cpp b/Engine/Plugins/Importers/USDImporter/Source/USDStage/Private/USDStageActor.cpp
--- a/Engine/Plugins/Importers/USDImporter/Source/USDStage/Private/USDStageActor.cpp
+++ b/Engine/Plugins/Importers/USDImporter/Source/USDStage/Private/USDStageActor.cpp
@@ -905,7 +905,7 @@ void AUsdStageActor::IsolateLayer( const UE::FSdfLayer& Layer )
 void AUsdStageActor::OnUsdObjectsChanged(const UsdUtils::FObjectChangesByPath& InfoChanges, const UsdUtils::FObjectChangesByPath& ResyncChanges)
 {
 #if USE_USD_SDK
-	if (!IsListeningToUsdNotices())
+	if (!IsListeningToUsdNotices() || !InfoCache.IsValid())
 	{
 		return;
 	}
@@ -973,14 +973,14 @@ void AUsdStageActor::OnUsdObjectsChanged(const UsdUtils::FObjectChangesByPath& I
 		bool bIsResync = false;
 		for (const UsdUtils::FObjectChangeNotice& ObjectChange : InfoChange.Value)
 		{
-			// We upgrade these mainly so that material binding changes refresh assets and components. Note that we
-			// just set material overrides on CreateComponents and not on UpdateComponents. In general, relationships
-			// likely won't be updated very often and can't have timeSamples, so it should be fine to resync on their
-			// changes
+			// This is in charge of resyncing components (and so calling CreateComponents instead of just
+			// UpdateComponents) whenever we change material assignments. This is important because
+			// CreateComponents is where we set our material overrides
 			if (ObjectChange.Flags.bDidChangeRelationshipTargets)
 			{
 				bIsResync = true;
 				bHasResync = true;
+				continue;
 			}
 
 			for (const UsdUtils::FAttributeChange& AttributeChange : ObjectChange.AttributeChanges)
@@ -1037,7 +1037,6 @@ void AUsdStageActor::OnUsdObjectsChanged(const UsdUtils::FObjectChangesByPath& I
 	// UpdateComponents/ReloadAssets below will already do this of course, but the point of this is that we unwind
 	// the collapsed prim paths *before* we rebuild the cache: We're really interested in whether this prim was
 	// collapsed or not before this change, as the change itself may have made the prim uncollapsible now
-	if (InfoCache.IsValid())
 	{
 		TMap<FString, bool> NewEntries;
 		NewEntries.Reserve(SortedPrimsChangedList.Num());
@@ -1077,7 +1076,7 @@ void AUsdStageActor::OnUsdObjectsChanged(const UsdUtils::FObjectChangesByPath& I
 	}
 
 	FScopedSlowTask RefreshStageTask(SortedPrimsChangedList.Num(), LOCTEXT("RefreshingUSDStage", "Refreshing USD Stage"));
-	RefreshStageTask.MakeDialog();
+	RefreshStageTask.MakeDialogDelayed(0.25f);
 
 	FScopedUsdMessageLog ScopedMessageLog;
 
@@ -1095,10 +1094,13 @@ void AUsdStageActor::OnUsdObjectsChanged(const UsdUtils::FObjectChangesByPath& I
 		// part of the tree (here we're just rebuilding the whole thing for now).
 		// However, if we know that only Prim resynced, we can traverse the tree root down and if we reach Prim and its collapsing state
 		// hasn't updated from before, we don't have to update its subtree at all, or sibling subtrees.
+		// Note that whenver we do a selective rebuild of the info cache we'll need to be very careful when to update
+		// certain info cache maps: For example whenever we delete a prim we need to make sure it's removed from MaterialUsers, etc.
 		TSharedRef< FUsdSchemaTranslationContext > TranslationContext = FUsdStageActorImpl::CreateUsdSchemaTranslationContext(this, TEXT("/"));
 		InfoCache->RebuildCacheForSubtree(Stage.GetPseudoRoot(), TranslationContext.Get());
 	}
 
+	bool bHasLoadedOrAbandonedAssets = false;
 	for (const TPair< FString, bool >& PrimChangedInfo : SortedPrimsChangedList)
 	{
 		RefreshStageTask.EnterProgressFrame();
@@ -1138,87 +1140,121 @@ void AUsdStageActor::OnUsdObjectsChanged(const UsdUtils::FObjectChangesByPath& I
 			return !PathToProcess.IsEmpty() && PathsProcessed.Contains(PathToProcess);
 		};
 
-		auto UpdateComponents = [&](const UE::FUsdPrim& Prim, const UE::FSdfPath& InPrimPath, const bool bInResync)
+		TFunction<void(const UE::FSdfPath& InPrimPath, const bool bInResync)> UpdateComponents;
+		UpdateComponents = [&](const UE::FSdfPath& InPrimPath, const bool bInResync)
 		{
-			// Don't query the InfoCache about prims that don't exist since that will ensure.
-			// We need to carry on though, because we must update components in case a prim is renamed or deleted,
-			// which will emit notices for prims that don't actually exist on the stage
-			UE::FSdfPath ComponentsPrimPath = (Prim && InfoCache.IsValid())
-				? InfoCache->UnwindToNonCollapsedPath(InPrimPath, ECollapsingType::Components)
-				: InPrimPath;
-
 			TSet< UE::FSdfPath >& RefreshedComponents = bInResync ? ResyncedComponents : UpdatedComponents;
 
-			if (!IsPathAlreadyProcessed(RefreshedComponents, ComponentsPrimPath))
+			TSet<UE::FSdfPath> MainPrimPaths = InfoCache->GetMainPrims(InPrimPath);
+			for (const UE::FSdfPath& MainPrimPath : MainPrimPaths)
 			{
-				TSharedRef< FUsdSchemaTranslationContext > TranslationContext = FUsdStageActorImpl::CreateUsdSchemaTranslationContext(this, ComponentsPrimPath.GetString());
-				UpdatePrim(ComponentsPrimPath, bInResync, *TranslationContext);
-				TranslationContext->CompleteTasks();
-
-				RefreshedComponents.Add(ComponentsPrimPath);
+				// Note: We don't need to "unwind" DependentPrim to a collapse root here: The entries for dependent
+				// prims on the TranslationDependency maps are always prims that were translated into assets before,
+				// so we know they must be "unwound" paths already
 
+				if (IsPathAlreadyProcessed(RefreshedComponents, MainPrimPath))
+				{
+					continue;
+				}
+				RefreshedComponents.Add(MainPrimPath);
 				if (bInResync)
 				{
 					// Consider that the path has been updated in the case of a resync
-					UpdatedComponents.Add(ComponentsPrimPath);
+					UpdatedComponents.Add(MainPrimPath);
 				}
+
+				UE_LOG(LogUsd, Verbose, TEXT("Prim '%s' was updated. Updating components for prim '%s'"),
+					*InPrimPath.GetString(),
+					*MainPrimPath.GetString()
+				);
+
+				// We could share this context, add all tasks for all Dependents and then call CompleteTasks once, but
+				// that could cause trouble in case we had a triangle of dependencies (A depends on B an C,
+				// and B depends on C) as that would prevent dependencies from being parsed before dependents
+				TSharedRef< FUsdSchemaTranslationContext > TranslationContext = FUsdStageActorImpl::CreateUsdSchemaTranslationContext(
+					this,
+					MainPrimPath.GetString()
+				);
+
+				UpdatePrim(MainPrimPath, bInResync, *TranslationContext);
+				TranslationContext->CompleteTasks();
+
+				// Now update the dependents of *this* prim
+				UpdateComponents(MainPrimPath, bInResync);
 			}
 		};
 
-		auto ReloadAssets = [&](const UE::FUsdPrim& Prim, const UE::FSdfPath& InPrimPath, const bool bInResync)
+		TFunction<void(const UE::FSdfPath& InPrimPath, const bool bInResync)> ReloadAssets;
+		ReloadAssets = [&](const UE::FSdfPath& InPrimPath, const bool bInResync)
 		{
-			UE::FSdfPath AssetsPrimPath = (Prim && InfoCache.IsValid())
-				? InfoCache->UnwindToNonCollapsedPath(InPrimPath, ECollapsingType::Assets)
-				: InPrimPath;
-
 			TSet< UE::FSdfPath >& RefreshedAssets = bInResync ? ResyncedAssets : UpdatedAssets;
 
-			if (!IsPathAlreadyProcessed(RefreshedAssets, AssetsPrimPath))
+			TSet<UE::FSdfPath> MainPrimPaths = InfoCache->GetMainPrims(InPrimPath);
+			for (const UE::FSdfPath& MainPrimPath : MainPrimPaths)
 			{
-				TSharedRef< FUsdSchemaTranslationContext > TranslationContext = FUsdStageActorImpl::CreateUsdSchemaTranslationContext(this, AssetsPrimPath.GetString());
+				if (IsPathAlreadyProcessed(RefreshedAssets, MainPrimPath))
+				{
+					continue;
+				}
+				RefreshedAssets.Add(MainPrimPath);
+
+				UE_LOG(LogUsd, Verbose, TEXT("Prim '%s' was updated. Reloading assets for prim '%s'"),
+					*InPrimPath.GetString(),
+					*MainPrimPath.GetString()
+				);
 
-				if (bInResync && InfoCache)
+				UE::FUsdPrim PrimToUpdate = GetOrLoadUsdStage().GetPrimAtPath(MainPrimPath);
+				if (!PrimToUpdate)
 				{
-					UMaterialInterface* ExistingMaterial = InfoCache->GetSingleAssetForPrim<UMaterialInterface>(
-						AssetsPrimPath
-					);
+					continue;
+				}
 
-					UE::FUsdPrim PrimToResync = GetOrLoadUsdStage().GetPrimAtPath(AssetsPrimPath);
-					LoadAssets(*TranslationContext, PrimToResync);
+				TSharedRef< FUsdSchemaTranslationContext > TranslationContext = FUsdStageActorImpl::CreateUsdSchemaTranslationContext(
+					this,
+					MainPrimPath.GetString()
+				);
 
-					UMaterialInterface* NewMaterial = InfoCache->GetSingleAssetForPrim<UMaterialInterface>(
-						AssetsPrimPath
-					);
+				UMaterialInterface* ExistingMaterial = InfoCache->GetSingleAssetForPrim<UMaterialInterface>(MainPrimPath);
 
-					// For UE-120185: If we recreated a material for a prim path we also need to update all components that were using it.
-					// This could be fleshed out further if other asset types require this refresh of "dependent components" but materials
-					// seem to be the only ones that do at the moment
-					if (ExistingMaterial && (ExistingMaterial != NewMaterial))
-					{
-						for (const FString& MaterialUserPrim : UsdUtils::GetMaterialUsers(PrimToResync, MaterialPurpose))
-						{
-							const bool bResyncComponent = true; // We need to force resync to reassign materials
-							UE::FSdfPath MaterialPrimPath{*MaterialUserPrim};
-							UE::FUsdPrim MaterialPrim = GetOrLoadUsdStage().GetPrimAtPath(MaterialPrimPath);
-							UpdateComponents(MaterialPrim, MaterialPrimPath, bResyncComponent);
-						}
-					}
+				if (bInResync)
+				{
+					bHasLoadedOrAbandonedAssets |= LoadAssets(*TranslationContext, PrimToUpdate);
 
 					// Resyncing also includes "updating" the prim
-					UpdatedAssets.Add(AssetsPrimPath);
+					UpdatedAssets.Add(MainPrimPath);
 				}
 				else
 				{
-					LoadAsset(*TranslationContext, GetOrLoadUsdStage().GetPrimAtPath(AssetsPrimPath));
+					bHasLoadedOrAbandonedAssets |= LoadAsset(*TranslationContext, PrimToUpdate);
+				}
+
+				UMaterialInterface* NewMaterial = InfoCache->GetSingleAssetForPrim<UMaterialInterface>(
+					MainPrimPath
+				);
+
+				// For UE-120185: If we recreated a material for a prim path we also need to update all components that were using it.
+				// This could be fleshed out further if other asset types require this refresh of "dependent components" but materials
+				// seem to be the only ones that do at the moment.
+				// Note that even after UE-157644 this is still useful: Material prims are not marked as a dependency
+				// of Mesh prims, otherwise we'd have to regenerate the StaticMesh itself when material info changed.
+				// Ideally we'd just set a new material override on the component, which is what this does
+				if (ExistingMaterial && (ExistingMaterial != NewMaterial) && InfoCache.IsValid())
+				{
+					for (const UE::FSdfPath& MaterialUserPath : InfoCache->GetMaterialUsers(MainPrimPath))
+					{
+						const bool bResyncComponent = true; // We need to force resync to reassign materials
+						UE::FUsdPrim MaterialPrim = GetOrLoadUsdStage().GetPrimAtPath(MaterialUserPath);
+						UpdateComponents(MaterialUserPath, bResyncComponent);
+					}
 				}
 
-				RefreshedAssets.Add(AssetsPrimPath);
+				// Now update the main prims of *this* prim as an auxiliary prim
+				ReloadAssets(MainPrimPath, bInResync);
 			}
 		};
 
-		UE::FUsdPrim Prim = GetOrLoadUsdStage().GetPrimAtPath(PrimPath);
-		ReloadAssets(Prim, PrimPath, bIsResync);
-		UpdateComponents(Prim, PrimPath, bIsResync);
+		ReloadAssets(PrimPath, bIsResync);
+		UpdateComponents(PrimPath, bIsResync);
 
 		if (HasAuthorityOverStage())
 		{
@@ -1232,7 +1268,7 @@ void AUsdStageActor::OnUsdObjectsChanged(const UsdUtils::FObjectChangesByPath& I
 		PrimsToAnimate.Reset();
 	}
 
-	if (bHasResync && UsdAssetCache)
+	if (bHasLoadedOrAbandonedAssets && UsdAssetCache)
 	{
 		UsdAssetCache->RefreshStorage();
 	}
@@ -1474,7 +1510,7 @@ void AUsdStageActor::UpdatePrim(const UE::FSdfPath& InUsdPrimPath, bool bResync,
 	TRACE_CPUPROFILER_EVENT_SCOPE(AUsdStageActor::UpdatePrim);
 
 	FScopedSlowTask SlowTask(1.f, LOCTEXT("UpdatingUSDPrim", "Updating USD Prim"));
-	SlowTask.MakeDialog();
+	SlowTask.MakeDialogDelayed(0.25f);
 	SlowTask.EnterProgressFrame();
 
 	UE::FSdfPath UsdPrimPath = InUsdPrimPath;
@@ -2066,7 +2102,7 @@ void AUsdStageActor::LoadUsdStage()
 	double StartTime = FPlatformTime::Cycles64();
 
 	FScopedSlowTask SlowTask(1.f, LOCTEXT("LoadingUDStage", "Loading USD Stage"));
-	SlowTask.MakeDialog();
+	SlowTask.MakeDialogDelayed(0.25f);
 
 	OnPreStageChanged.Broadcast();
 
@@ -2134,7 +2170,7 @@ void AUsdStageActor::LoadUsdStage()
 	InfoCache->RebuildCacheForSubtree(StageToLoad.GetPseudoRoot(), TranslationContext.Get());
 
 	SlowTask.EnterProgressFrame(0.7f);
-	LoadAssets(*TranslationContext, StageToLoad.GetPseudoRoot());
+	const bool bLoadedOrAbandonedAssets = LoadAssets(*TranslationContext, StageToLoad.GetPseudoRoot());
 
 	SlowTask.EnterProgressFrame(0.2f);
 	UpdatePrim(StageToLoad.GetPseudoRoot().GetPrimPath(), true, *TranslationContext);
@@ -2166,7 +2202,7 @@ void AUsdStageActor::LoadUsdStage()
 		}
 	}
 
-	if (UsdAssetCache)
+	if (bLoadedOrAbandonedAssets && UsdAssetCache)
 	{
 		UsdAssetCache->RefreshStorage();
 	}
@@ -3217,67 +3253,75 @@ void AUsdStageActor::OnSkelAnimationBaked(const FString& SkelRootPrimPath)
 #endif // #if USE_USD_SDK
 }
 
-void AUsdStageActor::LoadAsset(FUsdSchemaTranslationContext& TranslationContext, const UE::FUsdPrim& Prim)
+bool AUsdStageActor::LoadAsset(FUsdSchemaTranslationContext& TranslationContext, const UE::FUsdPrim& Prim)
 {
 	TRACE_CPUPROFILER_EVENT_SCOPE(AUsdStageActor::LoadAsset);
 
-	// Mark the assets as non transactional so that they don't get serialized in the transaction buffer
-	TGuardValue< EObjectFlags > ContextFlagsGuard(TranslationContext.ObjectFlags, TranslationContext.ObjectFlags & ~RF_Transactional);
-
-	FString PrimPath;
-#if USE_USD_SDK
-	PrimPath = UsdToUnreal::ConvertPath(Prim.GetPrimPath());
-#endif // #if USE_USD_SDK
+	if (!UsdAssetCache || !InfoCache)
+	{
+		return false;
+	}
 
-	if (InfoCache)
+	// Clear existing prim/asset association
+	bool bAssetsDiscarded = false;
 	{
-		if (UsdAssetCache)
-		{
-			UE::FSdfPath UsdPrimPath = UE::FSdfPath{*PrimPath};
-			TSet<TWeakObjectPtr<UObject>> OldAssets = InfoCache->RemoveAllAssetPrimLinks(UsdPrimPath);
+		UE::FSdfPath UsdPrimPath = Prim.GetPrimPath();
+		TSet<TWeakObjectPtr<UObject>> OldAssets = InfoCache->RemoveAllAssetPrimLinks(UsdPrimPath);
 
-			for (const TWeakObjectPtr<UObject>& OldAsset : OldAssets)
+		for (const TWeakObjectPtr<UObject>& OldAsset : OldAssets)
+		{
+			// If there are any other prim paths linked to this asset that we *won't* be removing/reparsing
+			// in here, it means our stage actor as a whole is still "referencing" that asset
+			bool bAssetStillReferenced = false;
+			for (const UE::FSdfPath& LinkedPrim : InfoCache->GetPrimsForAsset(OldAsset.Get()))
 			{
-				// If there are any other prim paths linked to this asset that we *won't* be removing/reparsing
-				// in here, it means our stage actor as a whole is still "referencing" that asset
-				bool bAssetStillReferenced = false;
-				for (const UE::FSdfPath& LinkedPrim : InfoCache->GetPrimsForAsset(OldAsset.Get()))
-				{
-					if (LinkedPrim != UsdPrimPath)
-					{
-						bAssetStillReferenced = true;
-						break;
-					}
-				}
-				if (bAssetStillReferenced)
+				if (LinkedPrim != UsdPrimPath)
 				{
-					continue;
+					bAssetStillReferenced = true;
+					break;
 				}
+			}
 
-				UsdAssetCache->RemoveAssetReference(OldAsset.Get(), this);
+			if (bAssetStillReferenced)
+			{
+				continue;
 			}
+
+			bAssetsDiscarded |= UsdAssetCache->RemoveAssetReference(OldAsset.Get(), this);
 		}
 	}
 
+	// Mark the assets as non transactional so that they don't get serialized in the transaction buffer
+	TGuardValue< EObjectFlags > ContextFlagsGuard(TranslationContext.ObjectFlags, TranslationContext.ObjectFlags & ~RF_Transactional);
+
+	int32 StartNumAssets = UsdAssetCache->GetNumAssets();
+
 	IUsdSchemasModule& UsdSchemasModule = FModuleManager::Get().LoadModuleChecked< IUsdSchemasModule >(TEXT("USDSchemas"));
-	if (TSharedPtr< FUsdSchemaTranslator > SchemaTranslator = UsdSchemasModule.GetTranslatorRegistry().CreateTranslatorForSchema(TranslationContext.AsShared(), UE::FUsdTyped(Prim)))
+	if (TSharedPtr< FUsdSchemaTranslator > SchemaTranslator = UsdSchemasModule.GetTranslatorRegistry().CreateTranslatorForSchema(
+		TranslationContext.AsShared(),
+		UE::FUsdTyped(Prim)
+	))
 	{
 		TRACE_CPUPROFILER_EVENT_SCOPE(AUsdStageActor::CreateAssetsForPrim);
 		SchemaTranslator->CreateAssets();
 	}
 
 	TranslationContext.CompleteTasks(); // Finish the asset tasks before moving on
+
+	return bAssetsDiscarded || UsdAssetCache->GetNumAssets() != StartNumAssets;
 }
 
-void AUsdStageActor::LoadAssets(FUsdSchemaTranslationContext& TranslationContext, const UE::FUsdPrim& StartPrim)
+bool AUsdStageActor::LoadAssets(FUsdSchemaTranslationContext& TranslationContext, const UE::FUsdPrim& StartPrim)
 {
 	TRACE_CPUPROFILER_EVENT_SCOPE(AUsdStageActor::LoadAssets);
 
-	// Mark the assets as non transactional so that they don't get serialized in the transaction buffer
-	TGuardValue< EObjectFlags > ContextFlagsGuard(TranslationContext.ObjectFlags, TranslationContext.ObjectFlags & ~RF_Transactional);
+	if (!UsdAssetCache || !InfoCache)
+	{
+		return false;
+	}
 
 	// Clear existing prim/asset association
-	if (InfoCache)
+	bool bAssetsDiscarded = false;
 	{
 		UE::FSdfPath StartPrimPath = StartPrim.GetPrimPath();
 		TSet<UE::FSdfPath> PrimPathsToRemove;
@@ -3290,35 +3334,37 @@ void AUsdStageActor::LoadAssets(FUsdSchemaTranslationContext& TranslationContext
 			}
 		}
 
-		if (UsdAssetCache)
+		for (const UE::FSdfPath& PrimPathToRemove : PrimPathsToRemove)
 		{
-			for (const UE::FSdfPath& PrimPathToRemove : PrimPathsToRemove)
+			TSet<TWeakObjectPtr<UObject>> OldAssets = InfoCache->RemoveAllAssetPrimLinks(PrimPathToRemove);
+			for (const TWeakObjectPtr<UObject>& OldAsset : OldAssets)
 			{
-				TSet<TWeakObjectPtr<UObject>> OldAssets = InfoCache->RemoveAllAssetPrimLinks(PrimPathToRemove);
-				for (const TWeakObjectPtr<UObject>& OldAsset : OldAssets)
+				// If there are any other prim paths linked to this asset that we *won't* be removing/reparsing
+				// in here, it means our stage actor as a whole is still "referencing" that asset
+				bool bAssetStillReferenced = false;
+				for (const UE::FSdfPath& LinkedPrim : InfoCache->GetPrimsForAsset(OldAsset.Get()))
 				{
-					// If there are any other prim paths linked to this asset that we *won't* be removing/reparsing
-					// in here, it means our stage actor as a whole is still "referencing" that asset
-					bool bAssetStillReferenced = false;
-					for (const UE::FSdfPath& LinkedPrim : InfoCache->GetPrimsForAsset(OldAsset.Get()))
-					{
-						if (!PrimPathsToRemove.Contains(LinkedPrim))
-						{
-							bAssetStillReferenced = true;
-							break;
-						}
-					}
-					if (bAssetStillReferenced)
+					if (!PrimPathsToRemove.Contains(LinkedPrim))
 					{
-						continue;
+						bAssetStillReferenced = true;
+						break;
 					}
-
-					UsdAssetCache->RemoveAssetReference(OldAsset.Get(), this);
 				}
+				if (bAssetStillReferenced)
+				{
+					continue;
+				}
+
+				bAssetsDiscarded |= UsdAssetCache->RemoveAssetReference(OldAsset.Get(), this);
 			}
 		}
 	}
 
+	// Mark the assets as non transactional so that they don't get serialized in the transaction buffer
+	TGuardValue< EObjectFlags > ContextFlagsGuard(TranslationContext.ObjectFlags, TranslationContext.ObjectFlags & ~RF_Transactional);
+
+	int32 StartNumAssets = UsdAssetCache->GetNumAssets();
+
 	IUsdSchemasModule& UsdSchemasModule = FModuleManager::Get().LoadModuleChecked< IUsdSchemasModule >(TEXT("USDSchemas"));
 
 	auto CreateAssetsForPrims = [&UsdSchemasModule, &TranslationContext](const TArray< UE::FUsdPrim >& AllPrimAssets, FSlowTask& Progress)
@@ -3329,7 +3375,10 @@ void AUsdStageActor::LoadAssets(FUsdSchemaTranslationContext& TranslationContext
 		{
 			Progress.EnterProgressFrame(1.f);
 
-			if (TSharedPtr< FUsdSchemaTranslator > SchemaTranslator = UsdSchemasModule.GetTranslatorRegistry().CreateTranslatorForSchema(TranslationContext.AsShared(), UE::FUsdTyped(UsdPrim)))
+			if (TSharedPtr< FUsdSchemaTranslator > SchemaTranslator = UsdSchemasModule.GetTranslatorRegistry().CreateTranslatorForSchema(
+				TranslationContext.AsShared(),
+				UE::FUsdTyped(UsdPrim)
+			))
 			{
 				TRACE_CPUPROFILER_EVENT_SCOPE(AUsdStageActor::CreateAssetsForPrim);
 				SchemaTranslator->CreateAssets();
@@ -3341,7 +3390,10 @@ void AUsdStageActor::LoadAssets(FUsdSchemaTranslationContext& TranslationContext
 
 	auto PruneChildren = [&UsdSchemasModule, &TranslationContext](const UE::FUsdPrim& UsdPrim) -> bool
 	{
-		if (TSharedPtr< FUsdSchemaTranslator > SchemaTranslator = UsdSchemasModule.GetTranslatorRegistry().CreateTranslatorForSchema(TranslationContext.AsShared(), UE::FUsdTyped(UsdPrim)))
+		if (TSharedPtr< FUsdSchemaTranslator > SchemaTranslator = UsdSchemasModule.GetTranslatorRegistry().CreateTranslatorForSchema(
+			TranslationContext.AsShared(),
+			UE::FUsdTyped(UsdPrim)
+		))
 		{
 			return SchemaTranslator->CollapsesChildren(ECollapsingType::Assets);
 		}
@@ -3362,6 +3414,8 @@ void AUsdStageActor::LoadAssets(FUsdSchemaTranslationContext& TranslationContext
 		FScopedSlowTask AssetsProgress(AllPrimAssets.Num(), LOCTEXT("CreateAssets", "Creating assets"));
 		CreateAssetsForPrims(AllPrimAssets, AssetsProgress);
 	}
+
+	return bAssetsDiscarded || UsdAssetCache->GetNumAssets() != StartNumAssets;
 }
 
 void AUsdStageActor::AnimatePrims()
diff --git a/Engine/Plugins/Importers/USDImporter/Source/USDStage/Public/USDStageActor.h b/Engine/Plugins/Importers/USDImporter/Source/USDStage/Public/USDStageActor.h
--- a/Engine/Plugins/Importers/USDImporter/Source/USDStage/Public/USDStageActor.h
+++ b/Engine/Plugins/Importers/USDImporter/Source/USDStage/Public/USDStageActor.h
@@ -263,10 +263,10 @@ public:
 
 protected:
 	/** Loads the asset for a single prim */
-	void LoadAsset(FUsdSchemaTranslationContext& TranslationContext, const UE::FUsdPrim& Prim);
+	bool LoadAsset(FUsdSchemaTranslationContext& TranslationContext, const UE::FUsdPrim& Prim);
 
 	/** Loads the assets for all prims from StartPrim and its children */
-	void LoadAssets(FUsdSchemaTranslationContext& TranslationContext, const UE::FUsdPrim& StartPrim);
+	bool LoadAssets(FUsdSchemaTranslationContext& TranslationContext, const UE::FUsdPrim& StartPrim);
 
 	void Refresh() const;
 	void AnimatePrims();
diff --git a/Engine/Plugins/Importers/USDImporter/Source/USDStageImporter/Private/USDOptionsWindow.cpp b/Engine/Plugins/Importers/USDImporter/Source/USDStageImporter/Private/USDOptionsWindow.cpp
--- a/Engine/Plugins/Importers/USDImporter/Source/USDStageImporter/Private/USDOptionsWindow.cpp
+++ b/Engine/Plugins/Importers/USDImporter/Source/USDStageImporter/Private/USDOptionsWindow.cpp
@@ -117,7 +117,7 @@ bool SUsdOptionsWindow::ShowOptions(
 	// will be pushed to the end of FSlateApplication::ActiveModalWindows (SlateApplication.cpp) and cause our options
 	// dialog to pop out of its modal loop (FSlateApplication::AddModalWindow), instantly returning false to our caller
 	FScopedSlowTask Progress( 1, LOCTEXT( "ShowingDialog", "Picking options..." ) );
-	Progress.MakeDialog();
+	Progress.MakeDialogDelayed(0.25f);
 
 	const bool bSlowTaskWindow = false;
 	FSlateApplication::Get().AddModalWindow( Window, ParentWindow, bSlowTaskWindow );
diff --git a/Engine/Plugins/Importers/USDImporter/Source/USDUtilities/Public/USDGeomMeshConversion.h b/Engine/Plugins/Importers/USDImporter/Source/USDUtilities/Public/USDGeomMeshConversion.h
--- a/Engine/Plugins/Importers/USDImporter/Source/USDUtilities/Public/USDGeomMeshConversion.h
+++ b/Engine/Plugins/Importers/USDImporter/Source/USDUtilities/Public/USDGeomMeshConversion.h
@@ -303,6 +303,7 @@ namespace UsdUtils
 	 * Returns an array of prim paths to prims within MaterialPrim's stage that have a UsdShadeMaterialBindingAPI, and whose computed material binding points to MaterialPrim.
 	 * The user prim's schema is not even checked, so this list will naturally include regular UsdGeomMesh prims, UsdGeomSubset prims, and any other with the UsdShadeMaterialBindingAPI.
 	 */
+	UE_DEPRECATED(5.3, "No longer used as the FUsdInfoCache now retrieves and retains this information")
 	USDUTILITIES_API TArray<FString> GetMaterialUsers( const UE::FUsdPrim& MaterialPrim, FName MaterialPurpose = *UnrealIdentifiers::MaterialAllPurpose );
 
 	/** Returns whether this prim can be interpreted as describing a static mesh with multiple LODs (i.e. if the prim holds the LOD variant set itself) */
