USD: Fix visual glitch where the prim properties panel would disappear after closing the stage.

diff --git a/Engine/Plugins/Importers/USDImporter/Source/USDStageEditor/Private/SUSDPrimInfo.cpp b/Engine/Plugins/Importers/USDImporter/Source/USDStageEditor/Private/SUSDPrimInfo.cpp
--- a/Engine/Plugins/Importers/USDImporter/Source/USDStageEditor/Private/SUSDPrimInfo.cpp
+++ b/Engine/Plugins/Importers/USDImporter/Source/USDStageEditor/Private/SUSDPrimInfo.cpp
@@ -49,16 +49,20 @@ void SUsdPrimInfo::Construct( const FArguments& InArgs )
 				{
 					// Display property metadata if we have exactly one selected
 					TArray<FString> SelectedProperties = PropertiesList->GetSelectedPropertyNames();
-					if (PropertiesList && SelectedProperties.Num() == 1)
+					if (PropertiesList && SelectedProperties.Num() == 1 && NewSelection &&
+						(NewSelection->Type == EAttributeModelType::Attribute || NewSelection->Type == EAttributeModelType::Relationship)
+					)
 					{
 						PropertyMetadataPanel->SetPrimPath(
 							PropertiesList->GetUsdStage(),
 							*(FString{PropertiesList->GetPrimPath()} + "." + SelectedProperties[0])
 						);
+						PropertyMetadataPanel->SetVisibility(EVisibility::Visible);
 					}
 					else
 					{
 						PropertyMetadataPanel->SetPrimPath({}, TEXT(""));
+						PropertyMetadataPanel->SetVisibility(EVisibility::Collapsed);
 					}
 				})
 			]
diff --git a/Engine/Plugins/Importers/USDImporter/Source/USDStageEditor/Private/Widgets/SUSDPrimPropertiesList.cpp b/Engine/Plugins/Importers/USDImporter/Source/USDStageEditor/Private/Widgets/SUSDPrimPropertiesList.cpp
--- a/Engine/Plugins/Importers/USDImporter/Source/USDStageEditor/Private/Widgets/SUSDPrimPropertiesList.cpp
+++ b/Engine/Plugins/Importers/USDImporter/Source/USDStageEditor/Private/Widgets/SUSDPrimPropertiesList.cpp
@@ -762,8 +762,6 @@ void SUsdPrimPropertiesList::SetPrimPath( const UE::FUsdStageWeak& UsdStage, con
 {
 	GeneratePropertiesList( UsdStage, InPrimPath );
 
-	SetVisibility(ViewModel.PrimAttributes.Num() > 0 ? EVisibility::Visible : EVisibility::Collapsed);
-
 	RequestListRefresh();
 }
 
