USD: Remove usage of AssetImportData outside EDITORONLY ifdefs.

diff --git a/Engine/Plugins/Importers/USDImporter/Source/USDSchemas/Private/USDGeomPointInstancerTranslator.cpp b/Engine/Plugins/Importers/USDImporter/Source/USDSchemas/Private/USDGeomPointInstancerTranslator.cpp
--- a/Engine/Plugins/Importers/USDImporter/Source/USDSchemas/Private/USDGeomPointInstancerTranslator.cpp
+++ b/Engine/Plugins/Importers/USDImporter/Source/USDSchemas/Private/USDGeomPointInstancerTranslator.cpp
@@ -409,6 +409,8 @@ void FUsdGeomPointInstancerTranslator::UpdateComponents(USceneComponent* PointIn
 	TSet<UStaticMesh*> PrototypeMeshSet = InfoCache.GetAssetsForPrim<UStaticMesh>(PrimPath);
 	std::unordered_map<pxr::SdfPath, UStaticMesh*, pxr::SdfPath::Hash> PrototypeMeshes;
 	PrototypeMeshes.reserve(PrototypeMeshSet.Num());
+	// TODO: Replace AssetImportData with AssetUserData, which is available at runtime
+#if WITH_EDITOR
 	for (UStaticMesh* PrototypeMesh : PrototypeMeshSet)
 	{
 		if (UUsdAssetImportData* ImportData = Cast<UUsdAssetImportData>(PrototypeMesh->GetAssetImportData()))
@@ -417,6 +419,7 @@ void FUsdGeomPointInstancerTranslator::UpdateComponents(USceneComponent* PointIn
 			PrototypeMeshes.insert({PrototypePath, PrototypeMesh});
 		}
 	}
+#endif // WITH_EDITOR
 
 	TArray<TFuture<TTuple<UHierarchicalInstancedStaticMeshComponent*, TArray<FTransform>>>> Tasks;
 	FScopedSlowTask PrototypePathsSlowTask( ( float ) PrototypePaths.size(), LOCTEXT( "GeomPointUpdateComponents", "Updating HierarchicalInstancedStaticMeshComponents for point instancers" ) );
