USD: Fix ensure when deserializing a cache with a persistent asset that is too large to fit in persistent storage;

Upgrade deserialization log message to Log level.

diff --git a/Engine/Plugins/Importers/USDImporter/Source/USDClasses/Private/USDAssetCache2.cpp b/Engine/Plugins/Importers/USDImporter/Source/USDClasses/Private/USDAssetCache2.cpp
--- a/Engine/Plugins/Importers/USDImporter/Source/USDClasses/Private/USDAssetCache2.cpp
+++ b/Engine/Plugins/Importers/USDImporter/Source/USDClasses/Private/USDAssetCache2.cpp
@@ -860,7 +860,7 @@ UObject* UUsdAssetCache2::GetCachedAsset(const FString& Hash)
 		UE::AssetCache::Private::DeserializeObjectAndSubObjects(NewAsset, BulkDataCopy);
 		ensure(NewAsset);
 
-		UE_LOG(LogUsd, Verbose, TEXT("Deserialized asset at '%s' (%s, hash '%s', storage %d) from bulkdata (%.3f MB)"),
+		UE_LOG(LogUsd, Log, TEXT("Deserialized asset at '%s' (%s, hash '%s', storage %d) from bulkdata (%.3f MB)"),
 			*ObjectInfo->AssetName,
 			NewAsset ? *NewAsset->GetPathName() : TEXT("nullptr"),
 			*ObjectInfo->Hash,
@@ -1289,8 +1289,20 @@ void UUsdAssetCache2::RefreshStorage()
 				{
 					Info.BulkData.UnloadBulkData();
 
-					ensure(Asset);
-					NewAssetStorage.Add(Info.Hash, Asset);
+					if (Asset)
+					{
+						ensure(AssetStorage.Contains(Info.Hash));
+						NewAssetStorage.Add(Info.Hash, Asset);
+					}
+					// It's possible to end up with an asset that is currently persistent but living in
+					// Unreferenced storage if we load an asset cache and discover the asset fits in the latter
+					// but not in the former (e.g. if size on disk > size in memory somehow, or if the storage
+					// sizes are tweaked to cause that)
+					else
+					{
+						ensure(!AssetStorage.Contains(Info.Hash));
+						NewPendingPersistentStorage.Add(Info.Hash, AssetPath);
+					}
 
 					UnreferencedAssetsSumBytes += Info.SizeOnMemoryInBytes;
 					UE_LOG(LogUsd, Verbose, TEXT("Updated unreferenced sum to %.3f out of max %u MB"),
