USD: Rename "dependent" to "consumer".

diff --git a/Engine/Plugins/Importers/USDImporter/Source/USDClasses/Private/USDAssetCache2.cpp b/Engine/Plugins/Importers/USDImporter/Source/USDClasses/Private/USDAssetCache2.cpp
--- a/Engine/Plugins/Importers/USDImporter/Source/USDClasses/Private/USDAssetCache2.cpp
+++ b/Engine/Plugins/Importers/USDImporter/Source/USDClasses/Private/USDAssetCache2.cpp
@@ -224,9 +224,9 @@ namespace UE::AssetCache::Private
 						return FArchiveUObject::operator<<(Obj);
 					}
 					// Stop serialization when a dependency is found or has been found
-					else if (Obj != SourceObject && !DependentObjects.Contains(Obj) && ValidObjects.Contains(Obj))
+					else if (Obj != SourceObject && !ConsumerObjects.Contains(Obj) && ValidObjects.Contains(Obj))
 					{
-						DependentObjects.Add(Obj);
+						ConsumerObjects.Add(Obj);
 					}
 				}
 
@@ -235,7 +235,7 @@ namespace UE::AssetCache::Private
 
 			UObject* SourceObject;
 			const TSet<UObject*>& ValidObjects;
-			TSet<UObject*> DependentObjects;
+			TSet<UObject*> ConsumerObjects;
 		};
 
 		FAssetCacheObjectWriter Writer{ *Object, Buffer };
@@ -265,7 +265,7 @@ namespace UE::AssetCache::Private
 		{
 			FObjectDependencyAnalyzer Analyzer(Object, AllowedSubObjectsSet);
 			Object->Serialize(Analyzer);
-			ActualDependencies.Append(Analyzer.DependentObjects);
+			ActualDependencies.Append(Analyzer.ConsumerObjects);
 		}
 
 		// Sort array of sub-objects based on their inter-dependency
@@ -282,8 +282,8 @@ namespace UE::AssetCache::Private
 				FObjectDependencyAnalyzer SubAnalyzer(SubObject, AllowedSubObjectsSet);
 				SubObject->Serialize(SubAnalyzer);
 
-				SubObjectDependencyGraph.FindOrAdd(SubObject).Append(SubAnalyzer.DependentObjects);
-				ActualDependencies.Append(SubAnalyzer.DependentObjects);
+				SubObjectDependencyGraph.FindOrAdd(SubObject).Append(SubAnalyzer.ConsumerObjects);
+				ActualDependencies.Append(SubAnalyzer.ConsumerObjects);
 			}
 
 			// Sort array of sub-objects: first objects do not depend on ones below
@@ -401,7 +401,7 @@ namespace UE::AssetCache::Private
 			return;
 		}
 
-		// Remove all objects created by default that InObject is dependent on
+		// Remove all objects created by default that InObject is an outer of.
 		// This method must obviously be called just after the InObject is created
 		auto RemoveDefaultDependencies = [](UObject* InObject)
 		{
@@ -568,7 +568,7 @@ void UUsdAssetCache2::FCachedAssetInfo::Serialize(FArchive& Ar, UObject* Owner)
 	Ar << SizeOnDiskInBytes;
 	Ar << SizeOnMemoryInBytes;
 	Ar << Dependencies;
-	Ar << Dependents;
+	Ar << Consumers;
 
 	// Very important: This flag ensures that when we call Serialize below we won't actually
 	// load the full bulkdata right away, so that we can do it on-demand
@@ -729,7 +729,7 @@ bool UUsdAssetCache2::CanRemoveAsset(const FString& Hash)
 		return true;
 	}
 
-	if (Info->Referencers.Num() > 0 || Info->Dependents.Num() > 0)
+	if (Info->Referencers.Num() > 0 || Info->Consumers.Num() > 0)
 	{
 		return false;
 	}
@@ -765,14 +765,14 @@ UObject* UUsdAssetCache2::RemoveAsset(const FString& Hash)
 		if (FCachedAssetInfo* Info = LRUCache.Find(FoundObjectPath))
 		{
 			ensure(Info->Referencers.Num() == 0);
-			ensure(Info->Dependents.Num() == 0);
+			ensure(Info->Consumers.Num() == 0);
 
-			// Make sure all of our dependencies know we're no longer dependent on them
+			// Make sure all of our dependencies know we're no longer Consumer on them
 			for (const FSoftObjectPath& Dependency : Info->Dependencies)
 			{
 				if (FCachedAssetInfo* DependencyInfo = LRUCache.Find(Dependency))
 				{
-					DependencyInfo->Dependents.Remove(FoundObjectPath);
+					DependencyInfo->Consumers.Remove(FoundObjectPath);
 				}
 			}
 			Info->Dependencies.Empty();
@@ -1146,27 +1146,27 @@ void UUsdAssetCache2::RefreshStorage()
 					Info.Dependencies.Add(Dependency);
 					if (FCachedAssetInfo* DependencyInfo = LRUCache.Find(Dependency))
 					{
-						DependencyInfo->Dependents.Add(AssetPath);
+						DependencyInfo->Consumers.Add(AssetPath);
 					}
 				}
 			}
 		}
 	}
 
-	TFunction<bool(const FCachedAssetInfo&)> DependentTreeHasReferencedAsset;
-	DependentTreeHasReferencedAsset = [this, &DependentTreeHasReferencedAsset](const FCachedAssetInfo& Info) -> bool
+	TFunction<bool(const FCachedAssetInfo&)> ConsumerTreeHasReferencedAsset;
+	ConsumerTreeHasReferencedAsset = [this, &ConsumerTreeHasReferencedAsset](const FCachedAssetInfo& Info) -> bool
 	{
-		for (const FSoftObjectPath& Dependent : Info.Dependents)
+		for (const FSoftObjectPath& Consumer : Info.Consumers)
 		{
-			FCachedAssetInfo* DependentInfo = LRUCache.Find(Dependent);
-			if (ensure(DependentInfo))
+			FCachedAssetInfo* ConsumerInfo = LRUCache.Find(Consumer);
+			if (ensure(ConsumerInfo))
 			{
-				if (DependentInfo->Referencers.Num() > 0)
+				if (ConsumerInfo->Referencers.Num() > 0)
 				{
 					return true;
 				}
 
-				if (DependentTreeHasReferencedAsset(*DependentInfo))
+				if (ConsumerTreeHasReferencedAsset(*ConsumerInfo))
 				{
 					return true;
 				}
@@ -1180,7 +1180,7 @@ void UUsdAssetCache2::RefreshStorage()
 	VisitEntryRecursively =
 		[
 			this,
-			&DependentTreeHasReferencedAsset,
+			&ConsumerTreeHasReferencedAsset,
 			&VisitEntryRecursively,
 			&VisitedAssets,
 			bIsTransientAssetCache,
@@ -1228,7 +1228,7 @@ void UUsdAssetCache2::RefreshStorage()
 			const bool bAssetFitsInUnreferenced = UnreferencedAssetsSumBytes + Info.SizeOnMemoryInBytes <= MaxUnreferencedTransientBytes;
 			const bool bDependenciesArePersistent = WorstDependencyStorage == ECacheStorageType::Persistent;
 			const bool bAllDependenciesInStorage = WorstDependencyStorage != ECacheStorageType::None;
-			const bool bDependentIsReferenced = DependentTreeHasReferencedAsset(Info);
+			const bool bConsumerIsReferenced = ConsumerTreeHasReferencedAsset(Info);
 
 			const FString AssetPathStr = AssetPath.ToString();
 			const FString AssetClassNameStr = Info.AssetClassName.ToString();
@@ -1240,8 +1240,8 @@ void UUsdAssetCache2::RefreshStorage()
 				Info.CurrentStorageType = ECacheStorageType::Persistent;
 			}
 
-			// Priority 2: Asset is kept on transient storage because it (or a dependent) is referenced
-			else if (Info.Referencers.Num() > 0 || bDependentIsReferenced)
+			// Priority 2: Asset is kept on transient storage because it (or a Consumer) is referenced
+			else if (Info.Referencers.Num() > 0 || bConsumerIsReferenced)
 			{
 				Info.CurrentStorageType = ECacheStorageType::Referenced;
 			}
@@ -1369,12 +1369,12 @@ void UUsdAssetCache2::RefreshStorage()
 				TryUnloadAsset(Info);
 				Info.BulkData.RemoveBulkData();
 
-				// Make sure all of our dependencies know we're no longer a dependent
+				// Make sure all of our dependencies know we're no longer a Consumer
 				for (const FSoftObjectPath& Dependency : Info.Dependencies)
 				{
 					if (FCachedAssetInfo* DepInfo = LRUCache.Find(Dependency))
 					{
-						DepInfo->Dependents.Remove(AssetPath);
+						DepInfo->Consumers.Remove(AssetPath);
 					}
 				}
 
@@ -1447,24 +1447,24 @@ bool UUsdAssetCache2::TryUnloadAsset(FCachedAssetInfo& InOutInfo)
 		// Asset is currently loaded
 		if (TObjectPtr<UObject> Asset = AssetStorage.FindRef(InOutInfo.Hash))
 		{
-			// Check to see if all of its dependents can be unloaded too
-			for (const FSoftObjectPath& Dependent : InOutInfo.Dependents)
+			// Check to see if all of its Consumers can be unloaded too
+			for (const FSoftObjectPath& Consumer : InOutInfo.Consumers)
 			{
-				if (FCachedAssetInfo* DependentInfo = LRUCache.Find(Dependent))
+				if (FCachedAssetInfo* ConsumerInfo = LRUCache.Find(Consumer))
 				{
-					ensure(DependentInfo->Dependencies.Contains(Asset));
+					ensure(ConsumerInfo->Dependencies.Contains(Asset));
 
-					bool bDependentIsUnloaded = TryUnloadAsset(*DependentInfo);
-					if (!bDependentIsUnloaded)
+					bool bConsumerIsUnloaded = TryUnloadAsset(*ConsumerInfo);
+					if (!bConsumerIsUnloaded)
 					{
-						// Our dependent must stay loaded -> So must we
+						// Our Consumer must stay loaded -> So must we
 						return false;
 					}
 				}
-				// We don't know about this dependent anymore (maybe it was evicted?)
+				// We don't know about this Consumer anymore (maybe it was evicted?)
 				else
 				{
-					InOutInfo.Dependents.Remove(Dependent);
+					InOutInfo.Consumers.Remove(Consumer);
 				}
 			}
 
@@ -1624,15 +1624,15 @@ void UUsdAssetCache2::Serialize(FArchive& Ar)
 				// We only want to serialize dependencies on persistent assets, or else when we deserialize this
 				// we'll have a bunch of dependencies on assets that don't exist
 				TSet<FSoftObjectPath> PersistentDependencies = Info.Dependencies.Intersect(PersistentAssets);
-				TSet<FSoftObjectPath> PersistentDependents = Info.Dependents.Intersect(PersistentAssets);
+				TSet<FSoftObjectPath> PersistentConsumers = Info.Consumers.Intersect(PersistentAssets);
 				Swap(Info.Dependencies, PersistentDependencies);
-				Swap(Info.Dependents, PersistentDependents);
+				Swap(Info.Consumers, PersistentConsumers);
 				{
 					Info.CurrentStorageType = ECacheStorageType::Persistent;
 					Info.Serialize(Ar, this);
 				}
 				Swap(Info.Dependencies, PersistentDependencies);
-				Swap(Info.Dependents, PersistentDependents);
+				Swap(Info.Consumers, PersistentConsumers);
 			}
 		}
 		else if (Ar.IsLoading())
diff --git a/Engine/Plugins/Importers/USDImporter/Source/USDClasses/Public/USDAssetCache2.h b/Engine/Plugins/Importers/USDImporter/Source/USDClasses/Public/USDAssetCache2.h
--- a/Engine/Plugins/Importers/USDImporter/Source/USDClasses/Public/USDAssetCache2.h
+++ b/Engine/Plugins/Importers/USDImporter/Source/USDClasses/Public/USDAssetCache2.h
@@ -47,7 +47,7 @@ public:
 
 	/**
 	 * Returns true if the asset with the given hash can be removed from the cache. It will return false in case the
-	 * asset is still being used, either by another dependent asset or directly by some referencer.
+	 * asset is still being used, either by another consumer asset or directly by some referencer.
 	 */
 	UFUNCTION(BlueprintCallable, Category = "Caching", meta = (CallInEditor = "true"))
 	bool CanRemoveAsset(const FString& Hash);
@@ -210,7 +210,7 @@ private:
 		uint64 SizeOnDiskInBytes = 0;
 		uint64 SizeOnMemoryInBytes = 0;
 		TSet<FSoftObjectPath> Dependencies;
-		TSet<FSoftObjectPath> Dependents;
+		TSet<FSoftObjectPath> Consumers;
 
 		// Transient stuff
 		ECacheStorageType CurrentStorageType = ECacheStorageType::None;  // Only used internally during RefreshStorage()
diff --git a/Engine/Plugins/Importers/USDImporter/Source/USDStage/Private/USDStageActor.cpp b/Engine/Plugins/Importers/USDImporter/Source/USDStage/Private/USDStageActor.cpp
--- a/Engine/Plugins/Importers/USDImporter/Source/USDStage/Private/USDStageActor.cpp
+++ b/Engine/Plugins/Importers/USDImporter/Source/USDStage/Private/USDStageActor.cpp
@@ -1148,9 +1148,9 @@ void AUsdStageActor::OnUsdObjectsChanged(const UsdUtils::FObjectChangesByPath& I
 			TSet<UE::FSdfPath> MainPrimPaths = InfoCache->GetMainPrims(InPrimPath);
 			for (const UE::FSdfPath& MainPrimPath : MainPrimPaths)
 			{
-				// Note: We don't need to "unwind" DependentPrim to a collapse root here: The entries for dependent
-				// prims on the TranslationDependency maps are always prims that were translated into assets before,
-				// so we know they must be "unwound" paths already
+				// Note: We don't need to "unwind" MainPrimPath to a collapse root here: The entries for main prims
+				// on the Main/Aux maps are always prims that were translated into assets before, so we know they
+				// must be "unwound" paths already
 
 				if (IsPathAlreadyProcessed(RefreshedComponents, MainPrimPath))
 				{
@@ -1168,7 +1168,7 @@ void AUsdStageActor::OnUsdObjectsChanged(const UsdUtils::FObjectChangesByPath& I
 					*MainPrimPath.GetString()
 				);
 
-				// We could share this context, add all tasks for all Dependents and then call CompleteTasks once, but
+				// We could share this context, add all tasks for all consumers and then call CompleteTasks once, but
 				// that could cause trouble in case we had a triangle of dependencies (A depends on B an C,
 				// and B depends on C) as that would prevent dependencies from being parsed before dependents
 				TSharedRef< FUsdSchemaTranslationContext > TranslationContext = FUsdStageActorImpl::CreateUsdSchemaTranslationContext(
@@ -1179,7 +1179,7 @@ void AUsdStageActor::OnUsdObjectsChanged(const UsdUtils::FObjectChangesByPath& I
 				UpdatePrim(MainPrimPath, bInResync, *TranslationContext);
 				TranslationContext->CompleteTasks();
 
-				// Now update the dependents of *this* prim
+				// Now update the consumers of *this* prim
 				UpdateComponents(MainPrimPath, bInResync);
 			}
 		};
@@ -1233,7 +1233,7 @@ void AUsdStageActor::OnUsdObjectsChanged(const UsdUtils::FObjectChangesByPath& I
 				);
 
 				// For UE-120185: If we recreated a material for a prim path we also need to update all components that were using it.
-				// This could be fleshed out further if other asset types require this refresh of "dependent components" but materials
+				// This could be fleshed out further if other asset types require this refresh of "consumer components" but materials
 				// seem to be the only ones that do at the moment.
 				// Note that even after UE-157644 this is still useful: Material prims are not marked as a dependency
 				// of Mesh prims, otherwise we'd have to regenerate the StaticMesh itself when material info changed.
