USD: Added a cvar USD.SkipMeshTangentComputation to skip computation of tangents for static meshes.

diff --git a/Engine/Plugins/Importers/USDImporter/Source/USDSchemas/Private/USDGeomMeshTranslator.cpp b/Engine/Plugins/Importers/USDImporter/Source/USDSchemas/Private/USDGeomMeshTranslator.cpp
--- a/Engine/Plugins/Importers/USDImporter/Source/USDSchemas/Private/USDGeomMeshTranslator.cpp
+++ b/Engine/Plugins/Importers/USDImporter/Source/USDSchemas/Private/USDGeomMeshTranslator.cpp
@@ -82,6 +82,12 @@ static FAutoConsoleVariableRef CVarMeshNormalRepairThreshold(
 	GMeshNormalRepairThreshold,
 	TEXT("We will try repairing up to this fraction of a Mesh's normals when invalid. If a Mesh has more invalid normals than this, we will recompute all of them. Defaults to 0.05 (5% of all normals)."));
 
+static bool GSkipMeshTangentComputation = false;
+static FAutoConsoleVariableRef CVarSkipMeshTangentComputation(
+	TEXT("USD.SkipMeshTangentComputation"),
+	GSkipMeshTangentComputation,
+	TEXT("Skip computing tangents for meshes. With meshes with a huge numer of vertices, it can take a very long time to compute them."));
+
 namespace UsdGeomMeshTranslatorImpl
 {
 	bool ShouldEnableNanite(
@@ -701,7 +707,7 @@ namespace UsdGeomMeshTranslatorImpl
 		// Make sure our normals can be rebuilt from MeshDescription::InitializeAutoGeneratedAttributes in case some tool needs them.
 		// Always force-compute tangents here as we never have them anyway. If we don't force them to be recomputed we'll get
 		// the worst of both worlds as some of these will be arbitrarily recomputed anyway, and some will be left invalid
-		EComputeNTBsFlags Options = EComputeNTBsFlags::UseMikkTSpace | EComputeNTBsFlags::Tangents;
+		EComputeNTBsFlags Options = GSkipMeshTangentComputation ? EComputeNTBsFlags::None : EComputeNTBsFlags::UseMikkTSpace | EComputeNTBsFlags::Tangents;
 		if (InvalidNormalFraction >= GMeshNormalRepairThreshold)
 		{
 			Options |= EComputeNTBsFlags::Normals;
