USD: Fixed possible double free with USD stream.

diff --git a/Engine/Plugins/Importers/USDImporter/Source/GeometryCacheUSD/Private/GeometryCacheUSDStream.cpp b/Engine/Plugins/Importers/USDImporter/Source/GeometryCacheUSD/Private/GeometryCacheUSDStream.cpp
--- a/Engine/Plugins/Importers/USDImporter/Source/GeometryCacheUSD/Private/GeometryCacheUSDStream.cpp
+++ b/Engine/Plugins/Importers/USDImporter/Source/GeometryCacheUSD/Private/GeometryCacheUSDStream.cpp
@@ -64,6 +64,7 @@ FGeometryCacheUsdStream::FGeometryCacheUsdStream(UGeometryCacheTrackUsd* InUsdTr
 		InUsdTrack->EndFrameIndex})
 , UsdTrack(InUsdTrack)
 , ReadFunc(InReadFunc)
+, bReadyForRead(false)
 {
 }
 
@@ -90,7 +91,7 @@ bool FGeometryCacheUsdStream::GetFrameData(int32 FrameIndex, FGeometryCacheMeshD
 	return true;
 }
 
-void FGeometryCacheUsdStream::UpdateRequestStatus( TArray<int32>& OutFramesCompleted )
+void FGeometryCacheUsdStream::UpdateRequestStatus(TArray<int32>& OutFramesCompleted)
 {
 	FGeometryCacheStreamBase::UpdateRequestStatus( OutFramesCompleted );
 
@@ -99,6 +100,16 @@ void FGeometryCacheUsdStream::UpdateRequestStatus( TArray<int32>& OutFramesCompl
 	if ( FramesNeeded.Num() == 0 && FramesRequested.Num() == 0 && UsdTrack && UsdTrack->CurrentStagePinned )
 	{
 		UsdTrack->UnloadUsdStage();
+		bReadyForRead = false;
+	}
+}
+
+void FGeometryCacheUsdStream::PrepareRead()
+{
+	// Do this before calling the ReadFunc as FUsdStreamDDCUtils::GetUsdStreamDDCKey also needs a valid stage to work with
+	if (!bReadyForRead)
+	{
+		bReadyForRead = UsdTrack->LoadUsdStage();
 	}
 }
 
@@ -106,8 +117,7 @@ void FGeometryCacheUsdStream::GetMeshData(int32 FrameIndex, int32 ConcurrencyInd
 {
 	TRACE_CPUPROFILER_EVENT_SCOPE(FGeometryCacheUsdStream::GetMeshData);
 
-	// Do this before calling the ReadFunc as FUsdStreamDDCUtils::GetUsdStreamDDCKey also needs a valid stage to work with
-	if ( !UsdTrack->LoadUsdStage() )
+	if (!bReadyForRead)
 	{
 		return;
 	}
diff --git a/Engine/Plugins/Importers/USDImporter/Source/GeometryCacheUSD/Public/GeometryCacheUSDStream.h b/Engine/Plugins/Importers/USDImporter/Source/GeometryCacheUSD/Public/GeometryCacheUSDStream.h
--- a/Engine/Plugins/Importers/USDImporter/Source/GeometryCacheUSD/Public/GeometryCacheUSDStream.h
+++ b/Engine/Plugins/Importers/USDImporter/Source/GeometryCacheUSD/Public/GeometryCacheUSDStream.h
@@ -5,6 +5,8 @@
 #include "GeometryCacheStreamBase.h"
 #include "GeometryCacheTrackUSDTypes.h"
 
+#include <atomic>
+
 class UGeometryCacheTrackUsd;
 
 class FGeometryCacheUsdStream : public FGeometryCacheStreamBase
@@ -14,13 +16,15 @@ public:
 
 	//~ Begin IGeometryCacheStream Interface
 	virtual bool GetFrameData(int32 FrameIndex, FGeometryCacheMeshData& OutMeshData) override;
-	virtual void UpdateRequestStatus( TArray<int32>& OutFramesCompleted ) override;
+	virtual void UpdateRequestStatus(TArray<int32>& OutFramesCompleted) override;
+	//~ End IGeometryCacheStream Interface
 
 protected:
+	virtual void PrepareRead() override;
 	virtual void GetMeshData(int32 FrameIndex, int32 ConcurrencyIndex, FGeometryCacheMeshData& OutMeshData) override;
-	//~ End IGeometryCacheStream Interface
 
 	UGeometryCacheTrackUsd* UsdTrack;
 	FReadUsdMeshFunction ReadFunc;
+	std::atomic<bool> bReadyForRead;
 };
 
diff --git a/Engine/Plugins/Runtime/GeometryCache/Source/GeometryCacheStreamer/Private/GeometryCacheStreamBase.cpp b/Engine/Plugins/Runtime/GeometryCache/Source/GeometryCacheStreamer/Private/GeometryCacheStreamBase.cpp
--- a/Engine/Plugins/Runtime/GeometryCache/Source/GeometryCacheStreamer/Private/GeometryCacheStreamBase.cpp
+++ b/Engine/Plugins/Runtime/GeometryCache/Source/GeometryCacheStreamer/Private/GeometryCacheStreamBase.cpp
@@ -128,6 +128,8 @@ bool FGeometryCacheStreamBase::RequestFrameData()
 
 	if (ReadIndices.Num() > 0)
 	{
+		PrepareRead();
+
 		// Get any ReadIndex available
 		const int32 ReadIndex = ReadIndices.Pop(false);
 
@@ -373,6 +375,8 @@ void FGeometryCacheStreamBase::LoadFrameData(int32 FrameIndex)
 		return;
 	}
 
+	PrepareRead();
+
 	FGeometryCacheMeshData* MeshData = new FGeometryCacheMeshData;
 	GetMeshData(FrameIndex, 0, *MeshData);
 	FramesAvailable.Add(FrameIndex, MeshData);
diff --git a/Engine/Plugins/Runtime/GeometryCache/Source/GeometryCacheStreamer/Public/GeometryCacheStreamBase.h b/Engine/Plugins/Runtime/GeometryCache/Source/GeometryCacheStreamer/Public/GeometryCacheStreamBase.h
--- a/Engine/Plugins/Runtime/GeometryCache/Source/GeometryCacheStreamer/Public/GeometryCacheStreamBase.h
+++ b/Engine/Plugins/Runtime/GeometryCache/Source/GeometryCacheStreamer/Public/GeometryCacheStreamBase.h
@@ -46,8 +46,10 @@ public:
 	void UpdateCurrentFrameIndex(int32 FrameIndex);
 
 protected:
+	/* Function called from main thread to prepare for GetMeshData */
+	virtual void PrepareRead() {}
 
-	/* Derived class must provide a way to get the mesh data for the given FrameIndex */
+	/* Derived class must provide a way to get the mesh data for the given FrameIndex, called from worker threads */
 	virtual void GetMeshData(int32 FrameIndex, int32 ReadConcurrencyIndex, FGeometryCacheMeshData& OutMeshData) = 0;
 
 	void LoadFrameData(int32 FrameIndex);
