USD: Rename/code cleanup to better handle the property metadata panel.

diff --git a/Engine/Plugins/Importers/USDImporter/Source/USDStageEditor/Private/SUSDPrimInfo.cpp b/Engine/Plugins/Importers/USDImporter/Source/USDStageEditor/Private/SUSDPrimInfo.cpp
--- a/Engine/Plugins/Importers/USDImporter/Source/USDStageEditor/Private/SUSDPrimInfo.cpp
+++ b/Engine/Plugins/Importers/USDImporter/Source/USDStageEditor/Private/SUSDPrimInfo.cpp
@@ -7,7 +7,7 @@
 #include "USDStageModule.h"
 #include "USDTypesConversion.h"
 #include "Widgets/SUSDIntegrationsPanel.h"
-#include "Widgets/SUSDPrimPropertiesList.h"
+#include "Widgets/SUSDObjectFieldList.h"
 #include "Widgets/SUSDReferencesList.h"
 #include "Widgets/SUSDVariantSetsList.h"
 
@@ -43,25 +43,25 @@ void SUsdPrimInfo::Construct( const FArguments& InArgs )
 			SNew(SBox)
 			.Content()
 			[
-				SAssignNew(PropertiesList, SUsdPrimPropertiesList)
-				.NameColumnText(LOCTEXT("PropertyName", "Name"))
-				.OnSelectionChanged_Lambda([this](const TSharedPtr<FUsdPrimAttributeViewModel>& NewSelection, ESelectInfo::Type SelectionType)
+				SAssignNew(PropertiesList, SUsdObjectFieldList)
+				.NameColumnText(LOCTEXT("NameColumnText", "Name"))
+				.OnSelectionChanged_Lambda([this](const TSharedPtr<FUsdObjectFieldViewModel>& NewSelection, ESelectInfo::Type SelectionType)
 				{
 					// Display property metadata if we have exactly one selected
-					TArray<FString> SelectedProperties = PropertiesList->GetSelectedPropertyNames();
-					if (PropertiesList && SelectedProperties.Num() == 1 && NewSelection &&
-						(NewSelection->Type == EAttributeModelType::Attribute || NewSelection->Type == EAttributeModelType::Relationship)
+					TArray<FString> SelectedFields = PropertiesList->GetSelectedFieldNames();
+					if (PropertiesList && SelectedFields.Num() == 1 && NewSelection &&
+						(NewSelection->Type == EObjectFieldType::Attribute || NewSelection->Type == EObjectFieldType::Relationship)
 					)
 					{
-						PropertyMetadataPanel->SetPrimPath(
+						PropertyMetadataPanel->SetObjectPath(
 							PropertiesList->GetUsdStage(),
-							*(FString{PropertiesList->GetPrimPath()} + "." + SelectedProperties[0])
+							*(FString{PropertiesList->GetObjectPath()} + "." + SelectedFields[0])
 						);
 						PropertyMetadataPanel->SetVisibility(EVisibility::Visible);
 					}
 					else
 					{
-						PropertyMetadataPanel->SetPrimPath({}, TEXT(""));
+						PropertyMetadataPanel->SetObjectPath({}, TEXT(""));
 						PropertyMetadataPanel->SetVisibility(EVisibility::Collapsed);
 					}
 				})
@@ -74,10 +74,10 @@ void SUsdPrimInfo::Construct( const FArguments& InArgs )
 			SNew( SBox )
 			.Content()
 			[
-				SAssignNew(PropertyMetadataPanel, SUsdPrimPropertiesList)
+				SAssignNew(PropertyMetadataPanel, SUsdObjectFieldList)
 				.NameColumnText_Lambda([this]() -> FText
 				{
-					FString PropertyName = FPaths::GetExtension(PropertyMetadataPanel->GetPrimPath());
+					FString PropertyName = FPaths::GetExtension(PropertyMetadataPanel->GetObjectPath());
 					return FText::FromString(FString::Printf(TEXT("%s metadata"), *PropertyName));
 				})
 				.Visibility(EVisibility::Collapsed)
@@ -116,13 +116,13 @@ void SUsdPrimInfo::Construct( const FArguments& InArgs )
 	];
 }
 
-void SUsdPrimInfo::SetPrimPath( const UE::FUsdStageWeak& UsdStage, const TCHAR* PrimPath )
+void SUsdPrimInfo::SetPrimPath(const UE::FUsdStageWeak& UsdStage, const TCHAR* PrimPath)
 {
 	if ( PropertiesList )
 	{
-		PropertiesList->SetPrimPath( UsdStage, PrimPath );
+		PropertiesList->SetObjectPath(UsdStage, PrimPath);
 
-		PropertyMetadataPanel->SetPrimPath({}, TEXT(""));
+		PropertyMetadataPanel->SetObjectPath({}, TEXT(""));
 		PropertyMetadataPanel->ClearSelection();
 	}
 
diff --git a/Engine/Plugins/Importers/USDImporter/Source/USDStageEditor/Private/SUSDPrimInfo.h b/Engine/Plugins/Importers/USDImporter/Source/USDStageEditor/Private/SUSDPrimInfo.h
--- a/Engine/Plugins/Importers/USDImporter/Source/USDStageEditor/Private/SUSDPrimInfo.h
+++ b/Engine/Plugins/Importers/USDImporter/Source/USDStageEditor/Private/SUSDPrimInfo.h
@@ -16,12 +16,12 @@ public:
 	SLATE_END_ARGS()
 
 	void Construct( const FArguments& InArgs );
-	void SetPrimPath( const UE::FUsdStageWeak& UsdStage, const TCHAR* PrimPath );
+	void SetPrimPath(const UE::FUsdStageWeak& UsdStage, const TCHAR* PrimPath);
 
 private:
 	friend class SUsdStage;
-	TSharedPtr< class SUsdPrimPropertiesList > PropertiesList;
-	TSharedPtr< class SUsdPrimPropertiesList > PropertyMetadataPanel;
+	TSharedPtr< class SUsdObjectFieldList > PropertiesList;
+	TSharedPtr< class SUsdObjectFieldList > PropertyMetadataPanel;
 	TSharedPtr< class SUsdIntegrationsPanel > IntegrationsPanel;
 	TSharedPtr< class SVariantsList > VariantsList;
 	TSharedPtr< class SUsdReferencesList > ReferencesList;
diff --git a/Engine/Plugins/Importers/USDImporter/Source/USDStageEditor/Private/SUSDStage.cpp b/Engine/Plugins/Importers/USDImporter/Source/USDStageEditor/Private/SUSDStage.cpp
--- a/Engine/Plugins/Importers/USDImporter/Source/USDStageEditor/Private/SUSDStage.cpp
+++ b/Engine/Plugins/Importers/USDImporter/Source/USDStageEditor/Private/SUSDStage.cpp
@@ -25,7 +25,7 @@
 #include "UsdWrappers/SdfLayer.h"
 #include "UsdWrappers/UsdAttribute.h"
 #include "UsdWrappers/UsdStage.h"
-#include "Widgets/SUSDPrimPropertiesList.h"
+#include "Widgets/SUSDObjectFieldList.h"
 
 #include "ActorTreeItem.h"
 #include "Async/Async.h"
@@ -449,7 +449,7 @@ TArray<FString> SUsdStage::GetSelectedPropertyNames() const
 {
 	if( UsdPrimInfoWidget && UsdPrimInfoWidget->PropertiesList )
 	{
-		return UsdPrimInfoWidget->PropertiesList->GetSelectedPropertyNames();
+		return UsdPrimInfoWidget->PropertiesList->GetSelectedFieldNames();
 	}
 
 	return {};
@@ -459,7 +459,7 @@ void SUsdStage::SetSelectedPropertyNames( const TArray<FString>& NewSelection )
 {
 	if ( UsdPrimInfoWidget && UsdPrimInfoWidget->PropertiesList )
 	{
-		UsdPrimInfoWidget->PropertiesList->SetSelectedPropertyNames( NewSelection );
+		UsdPrimInfoWidget->PropertiesList->SetSelectedFieldNames(NewSelection);
 	}
 }
 
@@ -467,7 +467,7 @@ TArray<FString> SUsdStage::GetSelectedPropertyMetadataNames() const
 {
 	if (UsdPrimInfoWidget && UsdPrimInfoWidget->PropertyMetadataPanel)
 	{
-		return UsdPrimInfoWidget->PropertyMetadataPanel->GetSelectedPropertyNames();
+		return UsdPrimInfoWidget->PropertyMetadataPanel->GetSelectedFieldNames();
 	}
 
 	return {};
@@ -477,7 +477,7 @@ void SUsdStage::SetSelectedPropertyMetadataNames(const TArray<FString>& NewSelec
 {
 	if (UsdPrimInfoWidget && UsdPrimInfoWidget->PropertyMetadataPanel)
 	{
-		UsdPrimInfoWidget->PropertyMetadataPanel->SetSelectedPropertyNames(NewSelection);
+		UsdPrimInfoWidget->PropertyMetadataPanel->SetSelectedFieldNames(NewSelection);
 	}
 }
 
diff --git a/Engine/Plugins/Importers/USDImporter/Source/USDStageEditor/Private/Widgets/SUSDObjectFieldList.cpp b/Engine/Plugins/Importers/USDImporter/Source/USDStageEditor/Private/Widgets/SUSDObjectFieldList.cpp
new file mode 100644
--- /dev/null
+++ b/Engine/Plugins/Importers/USDImporter/Source/USDStageEditor/Private/Widgets/SUSDObjectFieldList.cpp
@@ -0,0 +1,822 @@
+// Copyright Epic Games, Inc. All Rights Reserved.
+
+#include "SUSDObjectFieldList.h"
+
+#include "SUSDStageEditorStyle.h"
+#include "UnrealUSDWrapper.h"
+#include "USDLog.h"
+#include "USDMemory.h"
+#include "USDStageActor.h"
+#include "USDStageModule.h"
+#include "USDTypesConversion.h"
+
+#include "Styling/AppStyle.h"
+#include "Engine/World.h"
+#include "Modules/ModuleManager.h"
+#include "ScopedTransaction.h"
+#include "Widgets/Input/SCheckBox.h"
+#include "Widgets/Input/SComboBox.h"
+#include "Widgets/Input/SEditableTextBox.h"
+#include "Widgets/Input/SNumericEntryBox.h"
+
+#if USE_USD_SDK
+
+#define LOCTEXT_NAMESPACE "SUsdObjectFieldList"
+
+namespace USDObjectFieldListConstants
+{
+	const FMargin LeftRowPadding( 6.0f, 2.5f, 2.0f, 2.5f );
+	const FMargin RightRowPadding( 3.0f, 2.5f, 2.0f, 2.5f );
+	const FMargin ComboBoxItemPadding( 3.0f, 0.0f, 2.0f, 0.0f );
+	const FMargin NumericEntryBoxItemPadding( 0.0f, 0.0f, 2.0f, 0.0f );
+	const float DesiredNumericEntryBoxWidth = 80.0f;
+	const float DesiredTextEntryBoxWidth = 80.0f;
+
+	const TCHAR* NormalFont = TEXT("PropertyWindow.NormalFont");
+}
+
+class SUsdObjectFieldRow;
+
+namespace USDObjectFieldListImpl
+{
+	enum class EWidgetType : uint8
+	{
+		None,
+		Bool,
+		U8,
+		I32,
+		U32,
+		I64,
+		U64,
+		F32,
+		F64,
+		Text,
+		Dropdown,
+	};
+
+	enum class EWidgetLabelTypes
+	{
+		NoLabel,
+		RGBA,
+		XYZW,
+	};
+
+	static TMap<FString, TArray<TSharedPtr<FString>>> TokenDropdownOptions;
+
+	void ResetOptions( const FString& TokenName )
+	{
+		TokenDropdownOptions.Remove( TokenName );
+	}
+
+	TArray< TSharedPtr< FString > >* GetTokenDropdownOptions( const FUsdObjectFieldViewModel& ViewModel )
+	{
+		if ( TArray< TSharedPtr< FString> >* FoundOptions = TokenDropdownOptions.Find( ViewModel.Label ) )
+		{
+			return FoundOptions;
+		}
+
+		TArray< TSharedPtr< FString > > Options = ViewModel.GetDropdownOptions();
+		if ( Options.Num() == 0 )
+		{
+			// We don't know the options for this property, so return nullptr so that it can become a regular text input box
+			return nullptr;
+		}
+
+		return &USDObjectFieldListImpl::TokenDropdownOptions.Add( ViewModel.Label, Options );
+	}
+
+	EWidgetLabelTypes GetLabelType( const UsdUtils::FConvertedVtValue& FieldValue, const FString& ValueRole )
+	{
+		if ( FieldValue.Entries.Num() < 1 )
+		{
+			return EWidgetLabelTypes::NoLabel;
+		}
+
+		int32 NumComponents = FieldValue.Entries[0].Num();
+		EWidgetLabelTypes LabelType = EWidgetLabelTypes::NoLabel;
+		if ( ValueRole.StartsWith( TEXT( "color" ), ESearchCase::IgnoreCase ) )
+		{
+			LabelType = EWidgetLabelTypes::RGBA;
+		}
+		else if ( NumComponents > 1 && NumComponents <= 4 && FieldValue.SourceType != UsdUtils::EUsdBasicDataTypes::Matrix2d )
+		{
+			LabelType = EWidgetLabelTypes::XYZW;
+		}
+
+		return LabelType;
+	}
+
+	EWidgetType GetWidgetType( UsdUtils::EUsdBasicDataTypes SourceType )
+	{
+		using namespace UsdUtils;
+
+		switch ( SourceType )
+		{
+		case EUsdBasicDataTypes::Bool:
+			return EWidgetType::Bool;
+			break;
+		case EUsdBasicDataTypes::Uchar:
+			return EWidgetType::U8;
+			break;
+		case EUsdBasicDataTypes::Int:
+		case EUsdBasicDataTypes::Int2:
+		case EUsdBasicDataTypes::Int3:
+		case EUsdBasicDataTypes::Int4:
+			return EWidgetType::I32;
+			break;
+		case EUsdBasicDataTypes::Uint:
+			return EWidgetType::U32;
+			break;
+		case EUsdBasicDataTypes::Int64:
+			return EWidgetType::I64;
+			break;
+		case EUsdBasicDataTypes::Uint64:
+			return EWidgetType::U64;
+			break;
+		case EUsdBasicDataTypes::Half:
+		case EUsdBasicDataTypes::Half2:
+		case EUsdBasicDataTypes::Half3:
+		case EUsdBasicDataTypes::Half4:
+		case EUsdBasicDataTypes::Quath:
+		case EUsdBasicDataTypes::Float:
+		case EUsdBasicDataTypes::Float2:
+		case EUsdBasicDataTypes::Float3:
+		case EUsdBasicDataTypes::Float4:
+		case EUsdBasicDataTypes::Quatf:
+			return EWidgetType::F32;
+			break;
+		case EUsdBasicDataTypes::Double:
+		case EUsdBasicDataTypes::Double2:
+		case EUsdBasicDataTypes::Double3:
+		case EUsdBasicDataTypes::Double4:
+		case EUsdBasicDataTypes::Timecode:
+		case EUsdBasicDataTypes::Matrix2d:
+		case EUsdBasicDataTypes::Matrix3d:
+		case EUsdBasicDataTypes::Matrix4d:
+		case EUsdBasicDataTypes::Quatd:
+			return EWidgetType::F64;
+			break;
+		case EUsdBasicDataTypes::Token:
+			return EWidgetType::Dropdown;
+			break;
+		case EUsdBasicDataTypes::String:
+		case EUsdBasicDataTypes::Asset:
+			return EWidgetType::Text;
+			break;
+		default:
+			break;
+		}
+
+		return EWidgetType::None;
+	}
+
+	/** Always returns `NumLabels` widgets, regardless of LabelTypes. Those may be the NullWidget though */
+	TArray<TSharedRef<SWidget>> GetNumericEntryBoxLabels( int32 NumLabels, EWidgetLabelTypes LabelType )
+	{
+		const static TArray<const FLinearColor*> Colors = {
+			&SNumericEntryBox<int32>::RedLabelBackgroundColor,
+			&SNumericEntryBox<int32>::GreenLabelBackgroundColor,
+			&SNumericEntryBox<int32>::BlueLabelBackgroundColor,
+			&FLinearColor::White
+		};
+
+		if ( NumLabels > 4 )
+		{
+			LabelType = EWidgetLabelTypes::NoLabel;
+		}
+
+		TArray<TSharedRef<SWidget>> Labels;
+		Labels.Reserve( NumLabels );
+
+		switch ( LabelType )
+		{
+		case EWidgetLabelTypes::RGBA:
+			for ( int32 Index = 0; Index < NumLabels; ++Index )
+			{
+				Labels.Add( SNumericEntryBox<int32>::BuildNarrowColorLabel( *Colors[ Index ] ) );
+			}
+			break;
+		case EWidgetLabelTypes::XYZW:
+			for ( int32 Index = 0; Index < NumLabels; ++Index )
+			{
+				Labels.Add( SNumericEntryBox<int32>::BuildNarrowColorLabel( *Colors[ Index ] ) );
+			}
+			break;
+		case EWidgetLabelTypes::NoLabel:
+		default:
+			for ( int32 Index = 0; Index < NumLabels; ++Index )
+			{
+				Labels.Add( SNullWidget::NullWidget );
+			}
+			break;
+		}
+
+		return Labels;
+	}
+}
+
+class SUsdObjectFieldRow : public SMultiColumnTableRow<TSharedPtr<FUsdObjectFieldViewModel>>
+{
+	SLATE_BEGIN_ARGS( SUsdObjectFieldRow ) {}
+	SLATE_END_ARGS()
+
+public:
+	void Construct(const FArguments& InArgs, const TSharedPtr<FUsdObjectFieldViewModel>& InUsdObjectField, const TSharedRef< STableViewBase >& OwnerTable);
+	virtual TSharedRef< SWidget > GenerateWidgetForColumn( const FName& ColumnName ) override;
+
+	void SetUsdObjectField(const TSharedPtr< FUsdObjectFieldViewModel >& InUsdObjectField);
+
+protected:
+	FText GetLabel() const
+	{
+		return FText::FromString(Field->Label);
+	}
+	FText GetType() const
+	{
+		const static FText TypeTexts[] = {
+			FText::FromString(TEXT("M")),
+			FText::FromString(TEXT("A")),
+			FText::FromString(TEXT("R"))
+		};
+
+		const int32 Index = static_cast<int32>(Field->Type);
+		return Index >= 0 && Index < static_cast<int32>(EObjectFieldType::MAX)
+			? TypeTexts[Index]
+			: FText::GetEmpty();
+	}
+
+	FText GetTypeToolTip() const
+	{
+		const static FText ToolTipTexts[] = {
+			LOCTEXT("MetadataTooltip", "Metadata"),
+			LOCTEXT("AttributeTooltip", "Attribute"),
+			LOCTEXT("RelationshipTooltip", "Relationship")
+		};
+
+		const int32 Index = static_cast<int32>(Field->Type);
+		return Index >= 0 && Index < static_cast<int32>(EObjectFieldType::MAX)
+			? ToolTipTexts[Index]
+			: FText::GetEmpty();
+	}
+
+	// Optional here because SNumericEntryBox uses optional values as input
+	template<typename T>
+	TOptional<T> GetValue( int32 ComponentIndex ) const
+	{
+		if ( Field->Value.Entries.Num() == 1 ) // Ignore arrays for now
+		{
+			if ( Field->Value.Entries[0].IsValidIndex( ComponentIndex ) )
+			{
+				if ( T* Value = Field->Value.Entries[0][ ComponentIndex ].TryGet<T>() )
+				{
+					return *Value;
+				}
+			}
+		}
+
+		return {};
+	}
+
+	// Other overloads as checkbox/text widgets don't use optional values as input, and these should only ever have one component anyway
+	FText GetValueText() const
+	{
+		if ( Field->Value.Entries.Num() == 1 )
+		{
+			if ( Field->Value.Entries[ 0 ].Num() > 0 )
+			{
+				if ( FString* Value = Field->Value.Entries[ 0 ][ 0 ].TryGet<FString>() )
+				{
+					return FText::FromString( *Value );
+				}
+			}
+		}
+
+		return FText::GetEmpty();
+	}
+
+	ECheckBoxState GetValueBool() const
+	{
+		if ( Field->Value.Entries.Num() == 1 )
+		{
+			if ( Field->Value.Entries[ 0 ].Num() > 0 )
+			{
+				if ( bool* Value = Field->Value.Entries[ 0 ][ 0 ].TryGet<bool>() )
+				{
+					return *Value ? ECheckBoxState::Checked : ECheckBoxState::Unchecked;
+				}
+			}
+		}
+
+		return ECheckBoxState::Undetermined;
+	}
+
+private:
+	template<typename T>
+	TArray<TAttribute<TOptional<T>>> GetAttributeArray();
+
+	TSharedRef< SWidget > GenerateTextWidget(const TAttribute<FText>& Attribute, const TAttribute<FText>& ToolTipAttribute = {}, ETextJustify::Type TextJustify = ETextJustify::Left);
+	TSharedRef< SWidget > GenerateEditableTextWidget( const TAttribute<FText>& Attribute, bool bIsReadOnly );
+	TSharedRef< SWidget > GenerateCheckboxWidget( const TAttribute<ECheckBoxState>& Attribute );
+
+	template<typename T>
+	TSharedRef< SWidget > GenerateSpinboxWidgets(
+		const TArray<TAttribute<TOptional<T>>>& Attribute,
+		USDObjectFieldListImpl::EWidgetLabelTypes LabelType = USDObjectFieldListImpl::EWidgetLabelTypes::NoLabel
+	);
+
+private:
+	template<typename T>
+	void OnSpinboxValueChanged( T NewValue, int32 ComponentIndex );
+
+	template<typename T>
+	void OnSpinboxValueCommitted( T InNewValue, ETextCommit::Type CommitType, int32 ComponentIndex );
+
+	void OnComboBoxSelectionChanged( TSharedPtr<FString> ChosenOption, ESelectInfo::Type SelectInfo );
+	void OnTextBoxTextCommitted(const FText& NewText, ETextCommit::Type CommitType);
+	void OnCheckBoxCheckStateChanged( ECheckBoxState NewState );
+
+private:
+	TSharedPtr<FUsdObjectFieldViewModel> Field;
+};
+
+void SUsdObjectFieldRow::Construct(const FArguments& InArgs, const TSharedPtr<FUsdObjectFieldViewModel>& InUsdObjectField, const TSharedRef<STableViewBase>& OwnerTable)
+{
+	SetUsdObjectField(InUsdObjectField);
+
+	SMultiColumnTableRow<TSharedPtr<FUsdObjectFieldViewModel>>::Construct(SMultiColumnTableRow<TSharedPtr<FUsdObjectFieldViewModel>>::FArguments(), OwnerTable);
+}
+
+TSharedRef< SWidget > SUsdObjectFieldRow::GenerateWidgetForColumn( const FName& ColumnName )
+{
+	using namespace USDObjectFieldListImpl;
+
+	TSharedRef< SWidget > ColumnWidget = SNullWidget::NullWidget;
+
+	if (ColumnName == ObjectFieldColumnIds::TypeColumn)
+	{
+		ColumnWidget = GenerateTextWidget(
+			{this, &SUsdObjectFieldRow::GetType},
+			{this, &SUsdObjectFieldRow::GetTypeToolTip},
+			ETextJustify::Center
+		);
+	}
+	else if (ColumnName == ObjectFieldColumnIds::NameColumn)
+	{
+		ColumnWidget = GenerateTextWidget({this, &SUsdObjectFieldRow::GetLabel});
+	}
+	else
+	{
+		EWidgetType WidgetType = GetWidgetType( Field->Value.SourceType );
+		EWidgetLabelTypes LabelType = GetLabelType( Field->Value, Field->ValueRole );
+
+		switch ( WidgetType )
+		{
+		case EWidgetType::Bool:
+			ColumnWidget = GenerateCheckboxWidget( { this, &SUsdObjectFieldRow::GetValueBool } );
+			break;
+		case EWidgetType::U8:
+			ColumnWidget = GenerateSpinboxWidgets( GetAttributeArray<uint8>(), LabelType );
+			break;
+		case EWidgetType::I32:
+			ColumnWidget = GenerateSpinboxWidgets( GetAttributeArray<int32>(), LabelType );
+			break;
+		case EWidgetType::U32:
+			ColumnWidget = GenerateSpinboxWidgets( GetAttributeArray<uint32>(), LabelType );
+			break;
+		case EWidgetType::I64:
+			ColumnWidget = GenerateSpinboxWidgets( GetAttributeArray<int64>(), LabelType );
+			break;
+		case EWidgetType::U64:
+			ColumnWidget = GenerateSpinboxWidgets( GetAttributeArray<uint64>(), LabelType );
+			break;
+		case EWidgetType::F32:
+			ColumnWidget = GenerateSpinboxWidgets( GetAttributeArray<float>(), LabelType );
+			break;
+		case EWidgetType::F64:
+			ColumnWidget = GenerateSpinboxWidgets( GetAttributeArray<double>(), LabelType );
+			break;
+		case EWidgetType::Dropdown:
+		{
+			TArray< TSharedPtr< FString > >* Options = GetTokenDropdownOptions( *Field );
+
+			// Show a dropdown if we know the available options for that token
+			if ( Options )
+			{
+				SAssignNew( ColumnWidget, SBox )
+				.HeightOverride( FUsdStageEditorStyle::Get()->GetFloat( "UsdStageEditor.ListItemHeight" ) )
+				.VAlign( VAlign_Center )
+				.HAlign( HAlign_Left )
+				[
+					SNew( SComboBox< TSharedPtr< FString > > )
+					.OptionsSource( Options )
+					.OnGenerateWidget_Lambda( [ & ]( TSharedPtr<FString> Option )
+					{
+						return SUsdObjectFieldRow::GenerateTextWidget( FText::FromString( *Option ) );
+					})
+					.OnSelectionChanged( this, &SUsdObjectFieldRow::OnComboBoxSelectionChanged )
+					[
+						// Having an editable text box inside the combobox allows the user to pick through the most common ones but to
+						// also specify a custom kind/purpose/etc. if they want to
+						SNew( SEditableTextBox )
+						.Text( this, &SUsdObjectFieldRow::GetValueText )
+						.Font( FAppStyle::GetFontStyle( "PropertyWindow.NormalFont" ) )
+						.Padding( FMargin( 3.0f ) )
+						// Fixed foreground color or else it will flip when our row is selected, and we already have a background
+						.ForegroundColor( FAppStyle::GetSlateColor( TEXT( "Colors.ForegroundHover" ) ) )
+						.OnTextCommitted( this, &SUsdObjectFieldRow::OnTextBoxTextCommitted )
+					]
+				];
+
+				break;
+			}
+
+			// If we don't have any options we intentionally fall down into the 'Text' case
+		}
+		case EWidgetType::Text:
+		{
+			ColumnWidget = SUsdObjectFieldRow::GenerateEditableTextWidget( { this, &SUsdObjectFieldRow::GetValueText }, Field->bReadOnly );
+			break;
+		}
+		case EWidgetType::None:
+		default:
+			ColumnWidget = SNullWidget::NullWidget;
+			break;
+		}
+	}
+
+	return SNew(SBox)
+		.VAlign(VAlign_Center)
+		.HAlign(HAlign_Fill)
+		[
+			ColumnWidget
+		];
+}
+
+void SUsdObjectFieldRow::SetUsdObjectField(const TSharedPtr<FUsdObjectFieldViewModel>& InUsdObjectField)
+{
+	Field = InUsdObjectField;
+}
+
+template<typename T>
+TArray<TAttribute<TOptional<T>>> SUsdObjectFieldRow::GetAttributeArray()
+{
+	using AttrType = TAttribute<TOptional<T>>;
+
+	TArray<AttrType> Attributes;
+	if ( Field->Value.Entries.Num() == 1 )
+	{
+		for ( int32 ComponentIndex = 0; ComponentIndex < Field->Value.Entries[0].Num(); ++ComponentIndex )
+		{
+			Attributes.Add( AttrType::Create( AttrType::FGetter::CreateSP( this, &SUsdObjectFieldRow::template GetValue<T>, ComponentIndex ) ) );
+		}
+	}
+
+	return Attributes;
+}
+
+TSharedRef< SWidget > SUsdObjectFieldRow::GenerateTextWidget(
+	const TAttribute<FText>& Attribute,
+	const TAttribute<FText>& ToolTipAttribute,
+	ETextJustify::Type TextJustify
+)
+{
+	return SNew(STextBlock)
+		.Text(Attribute)
+		.Justification(TextJustify)
+		.ToolTipText(ToolTipAttribute)
+		.Font(FAppStyle::GetFontStyle("PropertyWindow.NormalFont"))
+		.Margin(USDObjectFieldListConstants::RightRowPadding);
+}
+
+TSharedRef< SWidget > SUsdObjectFieldRow::GenerateEditableTextWidget( const TAttribute<FText>& Attribute, bool bIsReadOnly )
+{
+	return SNew( SBox )
+		.HeightOverride( FUsdStageEditorStyle::Get()->GetFloat( "UsdStageEditor.ListItemHeight" ) )
+		.MinDesiredWidth( USDObjectFieldListConstants::DesiredTextEntryBoxWidth )
+		.VAlign(VAlign_Center)
+		.HAlign(HAlign_Left)
+		[
+			SNew( SEditableTextBox )
+			.MinDesiredWidth( USDObjectFieldListConstants::DesiredTextEntryBoxWidth )
+			.Text( Attribute )
+			.IsReadOnly( bIsReadOnly )
+			.Font( FAppStyle::GetFontStyle( "PropertyWindow.NormalFont" ) )
+			.OnTextCommitted( this, &SUsdObjectFieldRow::OnTextBoxTextCommitted )
+			// Fixed foreground color or else it will flip when our row is selected, and we already have a background
+			.ForegroundColor( FAppStyle::GetSlateColor( bIsReadOnly ? TEXT( "Colors.Foreground" ) : TEXT( "Colors.ForegroundHover" ) ) )
+		];
+}
+
+TSharedRef< SWidget > SUsdObjectFieldRow::GenerateCheckboxWidget( const TAttribute<ECheckBoxState>& Attribute )
+{
+	return SNew( SBox )
+		.HeightOverride( FUsdStageEditorStyle::Get()->GetFloat( "UsdStageEditor.ListItemHeight" ) )
+		.VAlign( VAlign_Center )
+		.HAlign(HAlign_Left)
+		[
+			SNew( SCheckBox )
+			.IsChecked( Attribute )
+			.OnCheckStateChanged( this, &SUsdObjectFieldRow::OnCheckBoxCheckStateChanged )
+		];
+}
+
+template<typename T>
+TSharedRef< SWidget > SUsdObjectFieldRow::GenerateSpinboxWidgets( const TArray<TAttribute<TOptional<T>>>& Attributes, USDObjectFieldListImpl::EWidgetLabelTypes LabelType )
+{
+	if ( Attributes.Num() == 0 )
+	{
+		return SNullWidget::NullWidget;
+	}
+
+	TArray<TSharedRef<SWidget>> Labels = USDObjectFieldListImpl::GetNumericEntryBoxLabels( Attributes.Num(), LabelType );
+
+	// If we're some type of matrix attribute, we'll want to display multiple rows
+	const int32 NumRows = LabelType != USDObjectFieldListImpl::EWidgetLabelTypes::NoLabel
+		? 1 // If we have labels we're definitely not a matrix type
+		: Attributes.Num() == 4
+			? 2
+			: Attributes.Num() == 9
+				? 3
+				: Attributes.Num() == 16
+					? 4
+					: 1;
+
+	const int32 NumColumns = Attributes.Num() / NumRows;
+
+	TSharedPtr<SVerticalBox> VertBox = SNew( SVerticalBox );
+	for ( int32 RowIndex = 0; RowIndex < NumRows; ++RowIndex )
+	{
+		TSharedPtr<SHorizontalBox> HorizBox = SNew( SHorizontalBox );
+		for ( int32 ColumnIndex = 0; ColumnIndex < NumColumns; ++ColumnIndex )
+		{
+			const int32 ComponentIndex = RowIndex * NumColumns + ColumnIndex;
+			const TAttribute<TOptional<T>>& Attribute = Attributes[ ComponentIndex ];
+
+			TSharedRef<SWidget> EntryBox = SNew( SNumericEntryBox<T> )
+				.AllowSpin( true )
+				.Font( FAppStyle::GetFontStyle( "PropertyWindow.NormalFont" ) )
+				.ShiftMouseMovePixelPerDelta( 1 )
+				.SupportDynamicSliderMaxValue( true )
+				.SupportDynamicSliderMinValue( true )
+				.OnValueChanged( this, &SUsdObjectFieldRow::OnSpinboxValueChanged<T>, ComponentIndex )
+				.OnValueCommitted( this, &SUsdObjectFieldRow::OnSpinboxValueCommitted<T>, ComponentIndex )
+				.Value( Attribute )
+				.MinValue( TOptional<T>() )
+				.MaxValue( TOptional<T>() )
+				.MinSliderValue( TOptional<T>() )
+				.MaxSliderValue( TOptional<T>() )
+				.SliderExponent( T( 1 ) )
+				.Delta( T( 0 ) )
+				.LabelPadding( FMargin(3) )
+				.LabelLocation( SNumericEntryBox<T>::ELabelLocation::Inside )
+				.Label()
+				[
+					Labels[ ComponentIndex ]
+				];
+
+			HorizBox->AddSlot()
+				.Padding( USDObjectFieldListConstants::NumericEntryBoxItemPadding )
+				[
+					SNew( SBox )
+					.MinDesiredWidth( USDObjectFieldListConstants::DesiredNumericEntryBoxWidth )
+					.MaxDesiredWidth( USDObjectFieldListConstants::DesiredNumericEntryBoxWidth )
+					[
+						EntryBox
+					]
+				];
+		}
+
+		VertBox->AddSlot()
+			[
+				SNew( SBox )
+				.HeightOverride( FUsdStageEditorStyle::Get()->GetFloat( "UsdStageEditor.ListItemHeight" ) )
+				.VAlign( VAlign_Center )
+				[
+					HorizBox.ToSharedRef()
+				]
+			];
+	}
+
+	return SNew(SBox)
+		.MinDesiredWidth(USDObjectFieldListConstants::DesiredNumericEntryBoxWidth)
+		.VAlign(VAlign_Center)
+		.HAlign(HAlign_Left)
+		[
+			VertBox.ToSharedRef()
+		];
+}
+
+template<typename T>
+void SUsdObjectFieldRow::OnSpinboxValueChanged( T NewValue, int32 ComponentIndex )
+{
+	Field->Value.Entries[ 0 ][ ComponentIndex ].Set<T>( NewValue );
+}
+
+template<typename T>
+void SUsdObjectFieldRow::OnSpinboxValueCommitted( T InNewValue, ETextCommit::Type CommitType, int32 ComponentIndex )
+{
+	if ( CommitType == ETextCommit::OnCleared || Field->bReadOnly )
+	{
+		return;
+	}
+
+	FScopedTransaction Transaction(
+		FText::Format(
+			LOCTEXT( "SetUsdAttributeTransaction", "Set attribute '{0}'" ),
+			FText::FromString( Field->Label )
+		)
+	);
+
+	UsdUtils::FConvertedVtValue NewValue = Field->Value;
+	NewValue.Entries[ 0 ][ ComponentIndex ].Set<T>( InNewValue );
+
+	Field->SetAttributeValue( NewValue );
+}
+
+void SUsdObjectFieldRow::OnComboBoxSelectionChanged( TSharedPtr<FString> ChosenOption, ESelectInfo::Type SelectInfo )
+{
+	if ( Field->bReadOnly )
+	{
+		return;
+	}
+
+	FScopedTransaction Transaction(
+		FText::Format(
+			LOCTEXT( "SetUsdAttributeTransaction", "Set attribute '{0}'" ),
+			FText::FromString( Field->Label )
+		)
+	);
+
+	UsdUtils::FConvertedVtValue NewValue = Field->Value;
+	NewValue.Entries[ 0 ][ 0 ] = UsdUtils::FConvertedVtValueComponent( TInPlaceType<FString>(), *ChosenOption );
+
+	Field->SetAttributeValue( NewValue );
+}
+
+void SUsdObjectFieldRow::OnTextBoxTextCommitted( const FText& NewText, ETextCommit::Type CommitType )
+{
+	if ( CommitType == ETextCommit::OnCleared || Field->bReadOnly )
+	{
+		return;
+	}
+
+	FScopedTransaction Transaction(
+		FText::Format(
+			LOCTEXT( "SetUsdAttributeTransaction", "Set attribute '{0}'" ),
+			FText::FromString( Field->Label )
+		)
+	);
+
+	UsdUtils::FConvertedVtValue NewValue = Field->Value;
+	NewValue.Entries[ 0 ][ 0 ] = UsdUtils::FConvertedVtValueComponent( TInPlaceType<FString>(), NewText.ToString() );
+
+	Field->SetAttributeValue( NewValue );
+}
+
+void SUsdObjectFieldRow::OnCheckBoxCheckStateChanged( ECheckBoxState NewState )
+{
+	if ( NewState == ECheckBoxState::Undetermined || Field->bReadOnly )
+	{
+		return;
+	}
+
+	FScopedTransaction Transaction(
+		FText::Format(
+			LOCTEXT( "SetUsdAttributeTransaction", "Set attribute '{0}'" ),
+			FText::FromString( Field->Label )
+		)
+	);
+
+	UsdUtils::FConvertedVtValue NewValue = Field->Value;
+	NewValue.Entries[ 0 ][ 0 ] = UsdUtils::FConvertedVtValueComponent( TInPlaceType<bool>(), NewState == ECheckBoxState::Checked );
+
+	Field->SetAttributeValue( NewValue );
+}
+
+void SUsdObjectFieldList::Construct( const FArguments& InArgs )
+{
+	// Clear map as usd file may have additional Kinds now
+	USDObjectFieldListImpl::ResetOptions(TEXT("Kind"));
+
+	SAssignNew( HeaderRowWidget, SHeaderRow )
+
+	+SHeaderRow::Column(ObjectFieldColumnIds::TypeColumn)
+	.DefaultLabel(FText::GetEmpty())
+	.FixedWidth(24.0f)
+	.SortMode(this, &SUsdObjectFieldList::GetColumnSortMode, ObjectFieldColumnIds::TypeColumn)
+	.OnSort(this, &SUsdObjectFieldList::Sort)
+
+	+SHeaderRow::Column(ObjectFieldColumnIds::NameColumn)
+	.DefaultLabel(InArgs._NameColumnText)
+	.FillWidth( 40.f )
+	.SortMode(this, &SUsdObjectFieldList::GetColumnSortMode, ObjectFieldColumnIds::NameColumn)
+	.OnSort(this, &SUsdObjectFieldList::Sort)
+
+	+SHeaderRow::Column(ObjectFieldColumnIds::ValueColumn)
+	.DefaultLabel( LOCTEXT( "ValueColumnText", "Value" ) )
+	.FillWidth( 60.f );
+
+	SListView::Construct
+	(
+		SListView::FArguments()
+		.ListItemsSource( &ViewModel.Fields )
+		.OnGenerateRow( this, &SUsdObjectFieldList::OnGenerateRow )
+		.HeaderRow( HeaderRowWidget )
+	);
+
+	OnSelectionChanged = InArgs._OnSelectionChanged;
+}
+
+TSharedRef< ITableRow > SUsdObjectFieldList::OnGenerateRow( TSharedPtr< FUsdObjectFieldViewModel > InDisplayNode, const TSharedRef< STableViewBase >& OwnerTable )
+{
+	return SNew( SUsdObjectFieldRow, InDisplayNode, OwnerTable );
+}
+
+void SUsdObjectFieldList::GenerateFieldList(const UE::FUsdStageWeak& UsdStage, const TCHAR* InObjectPath)
+{
+	const float TimeCode = 0.f;
+	ViewModel.Refresh(UsdStage, InObjectPath, TimeCode);
+}
+
+void SUsdObjectFieldList::Sort(
+	const EColumnSortPriority::Type SortPriority,
+	const FName& ColumnId,
+	const EColumnSortMode::Type NewSortMode
+)
+{
+	ViewModel.CurrentSortColumn = ColumnId;
+	ViewModel.CurrentSortMode = NewSortMode;
+	ViewModel.Sort();
+
+	RequestListRefresh();
+}
+
+EColumnSortMode::Type SUsdObjectFieldList::GetColumnSortMode(const FName ColumnId) const
+{
+	if (ViewModel.CurrentSortColumn != ColumnId)
+	{
+		return EColumnSortMode::None;
+	}
+
+	return ViewModel.CurrentSortMode;
+}
+
+void SUsdObjectFieldList::SetObjectPath(const UE::FUsdStageWeak& UsdStage, const TCHAR* InObjectPath)
+{
+	GenerateFieldList(UsdStage, InObjectPath);
+
+	RequestListRefresh();
+}
+
+TArray<FString> SUsdObjectFieldList::GetSelectedFieldNames() const
+{
+	TArray<FString> SelectedProperties;
+
+	TArray< TSharedPtr< FUsdObjectFieldViewModel > > SelectedViewModels = GetSelectedItems();
+	SelectedProperties.Reserve( SelectedViewModels.Num() );
+
+	for ( const TSharedPtr< FUsdObjectFieldViewModel >& SelectedItem : SelectedViewModels )
+	{
+		if ( SelectedItem )
+		{
+			SelectedProperties.Add( SelectedItem->Label );
+		}
+	}
+
+	return SelectedProperties;
+}
+
+void SUsdObjectFieldList::SetSelectedFieldNames( const TArray<FString>& NewSelection )
+{
+	TSet<FString> NewSelectionSet{ NewSelection };
+
+	Private_ClearSelection();
+
+	TArray< TSharedPtr< FUsdObjectFieldViewModel > > NewItemSelection;
+	NewItemSelection.Reserve( NewSelection.Num() );
+
+	for ( const TSharedPtr< FUsdObjectFieldViewModel >& Item : ViewModel.Fields )
+	{
+		if ( Item && NewSelectionSet.Contains( Item->Label ) )
+		{
+			NewItemSelection.Add( Item );
+		}
+	}
+
+	const bool bSelected = true;
+	SetItemSelection( NewItemSelection, bSelected );
+}
+
+UE::FUsdStageWeak SUsdObjectFieldList::GetUsdStage() const
+{
+	return ViewModel.GetUsdStage();
+}
+
+FString SUsdObjectFieldList::GetObjectPath() const
+{
+	return ViewModel.GetObjectPath();
+}
+
+#undef LOCTEXT_NAMESPACE
+
+#endif // #if USE_USD_SDK
diff --git a/Engine/Plugins/Importers/USDImporter/Source/USDStageEditor/Private/Widgets/SUSDObjectFieldList.h b/Engine/Plugins/Importers/USDImporter/Source/USDStageEditor/Private/Widgets/SUSDObjectFieldList.h
new file mode 100644
--- /dev/null
+++ b/Engine/Plugins/Importers/USDImporter/Source/USDStageEditor/Private/Widgets/SUSDObjectFieldList.h
@@ -0,0 +1,41 @@
+// Copyright Epic Games, Inc. All Rights Reserved.
+
+#pragma once
+
+#include "Widgets/Views/SListView.h"
+#include "USDObjectFieldViewModel.h"
+
+#if USE_USD_SDK
+
+
+class SUsdObjectFieldList : public SListView< TSharedPtr< FUsdObjectFieldViewModel > >
+{
+public:
+	SLATE_BEGIN_ARGS(SUsdObjectFieldList) {}
+		SLATE_EVENT(FOnSelectionChanged, OnSelectionChanged)
+		SLATE_ATTRIBUTE(FText, NameColumnText)
+	SLATE_END_ARGS()
+
+public:
+	void Construct( const FArguments& InArgs );
+	void SetObjectPath(const UE::FUsdStageWeak& UsdStage, const TCHAR* InObjectPath);
+
+	TArray<FString> GetSelectedFieldNames() const;
+	void SetSelectedFieldNames( const TArray<FString>& NewSelection );
+
+	UE::FUsdStageWeak GetUsdStage() const;
+	FString GetObjectPath() const;
+
+protected:
+	TSharedRef< ITableRow > OnGenerateRow( TSharedPtr< FUsdObjectFieldViewModel > InDisplayNode, const TSharedRef< STableViewBase >& OwnerTable );
+	void GenerateFieldList(const UE::FUsdStageWeak& UsdStage, const TCHAR* InObjectPath);
+
+	void Sort(const EColumnSortPriority::Type SortPriority, const FName& ColumnId, const EColumnSortMode::Type NewSortMode);
+	EColumnSortMode::Type GetColumnSortMode(const FName ColumnId) const;
+
+private:
+	FUsdObjectFieldsViewModel ViewModel;
+	TSharedPtr< SHeaderRow > HeaderRowWidget;
+};
+
+#endif // USE_USD_SDK
diff --git a/Engine/Plugins/Importers/USDImporter/Source/USDStageEditor/Private/Widgets/SUSDPrimPropertiesList.cpp b/Engine/Plugins/Importers/USDImporter/Source/USDStageEditor/Private/Widgets/SUSDPrimPropertiesList.cpp
deleted file mode 100644
diff --git a/Engine/Plugins/Importers/USDImporter/Source/USDStageEditor/Private/Widgets/SUSDPrimPropertiesList.h b/Engine/Plugins/Importers/USDImporter/Source/USDStageEditor/Private/Widgets/SUSDPrimPropertiesList.h
deleted file mode 100644
diff --git a/Engine/Plugins/Importers/USDImporter/Source/USDStageEditorViewModels/Private/USDObjectFieldViewModel.cpp b/Engine/Plugins/Importers/USDImporter/Source/USDStageEditorViewModels/Private/USDObjectFieldViewModel.cpp
new file mode 100644
--- /dev/null
+++ b/Engine/Plugins/Importers/USDImporter/Source/USDStageEditorViewModels/Private/USDObjectFieldViewModel.cpp
@@ -0,0 +1,633 @@
+// Copyright Epic Games, Inc. All Rights Reserved.
+
+#include "USDObjectFieldViewModel.h"
+
+#include "UnrealUSDWrapper.h"
+#include "USDAttributeUtils.h"
+#include "USDErrorUtils.h"
+#include "USDLog.h"
+#include "USDMemory.h"
+#include "USDTypesConversion.h"
+#include "USDValueConversion.h"
+
+#include "UsdWrappers/SdfLayer.h"
+#include "UsdWrappers/SdfPath.h"
+#include "UsdWrappers/UsdAttribute.h"
+#include "UsdWrappers/UsdPrim.h"
+#include "UsdWrappers/VtValue.h"
+
+#include "Framework/Notifications/NotificationManager.h"
+#include "ScopedTransaction.h"
+#include "Widgets/Notifications/SNotificationList.h"
+
+#if USE_USD_SDK
+
+#include "USDIncludesStart.h"
+	#include "pxr/base/tf/stringUtils.h"
+	#include "pxr/base/vt/value.h"
+	#include "pxr/usd/kind/registry.h"
+	#include "pxr/usd/sdf/types.h"
+	#include "pxr/usd/usd/attribute.h"
+	#include "pxr/usd/usd/common.h"
+	#include "pxr/usd/usd/editContext.h"
+	#include "pxr/usd/usd/prim.h"
+	#include "pxr/usd/usd/relationship.h"
+	#include "pxr/usd/usdGeom/tokens.h"
+#include "USDIncludesEnd.h"
+
+#endif // #if USE_USD_SDK
+
+#define LOCTEXT_NAMESPACE "USDObjectFieldViewModel"
+
+FUsdObjectFieldViewModel::FUsdObjectFieldViewModel( FUsdObjectFieldsViewModel* InOwner )
+	: Owner( InOwner )
+{
+}
+
+TArray< TSharedPtr< FString > > FUsdObjectFieldViewModel::GetDropdownOptions() const
+{
+#if USE_USD_SDK
+	if ( Label == TEXT("kind") )
+	{
+		TArray< TSharedPtr< FString > > Options;
+		{
+			FScopedUsdAllocs Allocs;
+
+			std::vector< pxr::TfToken > Kinds = pxr::KindRegistry::GetAllKinds();
+			Options.Reserve( Kinds.size() );
+
+			for ( const pxr::TfToken& Kind : Kinds )
+			{
+				Options.Add( MakeShared< FString >( UsdToUnreal::ConvertToken( Kind ) ) );
+			}
+
+			// They are supposed to be in an unspecified order, so let's make them consistent
+			Options.Sort( [](const TSharedPtr< FString >& A, const TSharedPtr< FString >& B )
+				{
+					return A.IsValid() && B.IsValid() && ( *A < *B );
+				}
+			);
+
+		}
+		return Options;
+	}
+	else if ( Label == UsdToUnreal::ConvertToken( pxr::UsdGeomTokens->purpose ) )
+	{
+		return TArray< TSharedPtr< FString > >
+		{
+			MakeShared< FString >( UsdToUnreal::ConvertToken( pxr::UsdGeomTokens->default_ ) ),
+			MakeShared< FString >( UsdToUnreal::ConvertToken( pxr::UsdGeomTokens->proxy ) ),
+			MakeShared< FString >( UsdToUnreal::ConvertToken( pxr::UsdGeomTokens->render) ),
+			MakeShared< FString >( UsdToUnreal::ConvertToken( pxr::UsdGeomTokens->guide ) ),
+		};
+	}
+	else if ( Label == UsdToUnreal::ConvertToken( pxr::UsdGeomTokens->upAxis ) )
+	{
+		return TArray< TSharedPtr< FString > >
+		{
+			MakeShared< FString >( UsdToUnreal::ConvertToken( pxr::UsdGeomTokens->y ) ),
+			MakeShared< FString >( UsdToUnreal::ConvertToken( pxr::UsdGeomTokens->z ) ),
+		};
+	}
+	else if ( Label == UsdToUnreal::ConvertToken( pxr::UsdGeomTokens->visibility ) )
+	{
+		return TArray< TSharedPtr< FString > >
+		{
+			MakeShared< FString >( UsdToUnreal::ConvertToken( pxr::UsdGeomTokens->inherited ) ),
+			MakeShared< FString >( UsdToUnreal::ConvertToken( pxr::UsdGeomTokens->invisible ) ),
+		};
+	}
+#endif // #if USE_USD_SDK
+
+	return {};
+}
+
+void FUsdObjectFieldViewModel::SetAttributeValue( const UsdUtils::FConvertedVtValue& InValue )
+{
+	Owner->SetFieldValue( Label, InValue );
+}
+
+template<typename T>
+void FUsdObjectFieldsViewModel::CreateField(
+	EObjectFieldType Type,
+	const FString& FieldName,
+	const T& Value,
+	UsdUtils::EUsdBasicDataTypes SourceType,
+	const FString& ValueRole,
+	bool bReadOnly
+)
+{
+	UsdUtils::FConvertedVtValue VtValue;
+	VtValue.Entries = { { UsdUtils::FConvertedVtValueComponent{ TInPlaceType<T>(), Value } } };
+	VtValue.SourceType = SourceType;
+
+	FUsdObjectFieldViewModel Property( this );
+	Property.Type = Type;
+	Property.Label = FieldName;
+	Property.Value = VtValue;
+	Property.ValueRole = ValueRole;
+	Property.bReadOnly = bReadOnly;
+
+	Fields.Add( MakeSharedUnreal< FUsdObjectFieldViewModel >( MoveTemp( Property ) ) );
+}
+
+void FUsdObjectFieldsViewModel::CreateField(
+	EObjectFieldType Type,
+	const FString& FieldName,
+	const UsdUtils::FConvertedVtValue& Value,
+	bool bReadOnly
+)
+{
+	FUsdObjectFieldViewModel Property( this );
+	Property.Type = Type;
+	Property.Label = FieldName;
+	Property.Value = Value;
+	Property.bReadOnly = bReadOnly;
+
+	Fields.Add( MakeSharedUnreal< FUsdObjectFieldViewModel >( MoveTemp( Property ) ) );
+}
+
+void FUsdObjectFieldsViewModel::SetFieldValue( const FString& FieldName, const UsdUtils::FConvertedVtValue& InValue )
+{
+	bool bSuccess = false;
+
+#if USE_USD_SDK
+	if ( !UsdStage )
+	{
+		return;
+	}
+
+	// Transact here as setting this attribute may trigger USD events that affect assets/components
+	FScopedTransaction Transaction(FText::Format(
+		LOCTEXT("SetFieldValue", "Set value for field '{0}' of prim '{1}'"),
+		FText::FromString(FieldName),
+		FText::FromString(ObjectPath)
+	));
+
+	const bool bIsStageAttribute = ObjectPath == TEXT("/") || ObjectPath.IsEmpty();
+	bool bIsPropertyPath = false;
+
+	UE::FVtValue VtValue;
+	if ( UnrealToUsd::ConvertValue( InValue, VtValue ) )
+	{
+		if ( bIsStageAttribute )
+		{
+			FScopedUsdAllocs UsdAllocs;
+
+			// To set stage metadata the edit target must be the root or session layer
+			pxr::UsdStageRefPtr Stage{ UsdStage };
+			pxr::UsdEditContext( Stage, Stage->GetRootLayer() );
+			bSuccess = UsdStage.SetMetadata( *FieldName, VtValue );
+		}
+		else if (UE::FUsdPrim UsdPrim = UsdStage.GetPrimAtPath(UE::FSdfPath(*ObjectPath)))
+		{
+			FScopedUsdAllocs UsdAllocs;
+
+			pxr::UsdPrim PxrUsdPrim{UsdPrim};
+			pxr::TfToken FieldNameToken = UnrealToUsd::ConvertToken(*FieldName).Get();
+
+			// Single value, single component of FString
+			if ( FieldName == TEXT( "kind" ) && InValue.Entries.Num() == 1 && InValue.Entries[ 0 ].Num() == 1 && InValue.Entries[ 0 ][ 0 ].IsType<FString>() )
+			{
+				bSuccess = IUsdPrim::SetKind(
+					UsdPrim,
+					UnrealToUsd::ConvertToken( *( InValue.Entries[ 0 ][ 0 ].Get<FString>() ) ).Get()
+				);
+			}
+			else if ( UE::FUsdAttribute Attribute = UsdPrim.GetAttribute( *FieldName ) )
+			{
+				bSuccess = Attribute.Set( VtValue );
+				UsdUtils::NotifyIfOverriddenOpinion( Attribute );
+			}
+			else if (pxr::UsdRelationship Relationship = PxrUsdPrim.GetRelationship(FieldNameToken))
+			{
+				if (InValue.Entries.Num() == 1 && InValue.Entries[0].Num() == 1 && InValue.Entries[0][0].IsType<FString>())
+				{
+					const FString& NewValue = InValue.Entries[0][0].Get<FString>();
+
+					// We only allow editing relationship attributes that have a single target anyway
+					bSuccess = Relationship.SetTargets(std::vector<pxr::SdfPath>{UnrealToUsd::ConvertPath(*NewValue).Get()});
+					UsdUtils::NotifyIfOverriddenOpinion(Relationship);
+				}
+			}
+			else if (PxrUsdPrim.HasMetadata(FieldNameToken))
+			{
+				bSuccess = PxrUsdPrim.SetMetadata(FieldNameToken, VtValue.GetUsdValue());
+			}
+		}
+		// Trying to set property metadata
+		else if (!ObjectPath.IsEmpty())
+		{
+			FScopedUsdAllocs Allocs;
+
+			pxr::SdfPath UsdPath = UnrealToUsd::ConvertPath(*ObjectPath).Get();
+			if (UsdPath.IsPropertyPath())
+			{
+				bIsPropertyPath = true;
+
+				pxr::SdfPath PxrPrimPath = UsdPath.GetPrimPath();
+				pxr::TfToken PropertyName = UsdPath.GetNameToken();
+
+				pxr::UsdStageRefPtr PxrUsdStage{UsdStage};
+				if (pxr::UsdPrim Prim = PxrUsdStage->GetPrimAtPath(PxrPrimPath))
+				{
+					if (pxr::UsdProperty Property = Prim.GetProperty(PropertyName))
+					{
+						pxr::TfToken FieldNameToken = UnrealToUsd::ConvertToken(*FieldName).Get();
+						if (Property.HasMetadata(FieldNameToken))
+						{
+							bSuccess = Property.SetMetadata(FieldNameToken, VtValue.GetUsdValue());
+						}
+					}
+				}
+			}
+		}
+	}
+
+	if ( !bSuccess )
+	{
+		const FText ErrorMessage = FText::Format(LOCTEXT("FailToSetFieldMessage", "Failed to set '{0}' on {1} '{2}'"),
+			FText::FromString(FieldName),
+			FText::FromString(bIsStageAttribute ? TEXT("stage") : bIsPropertyPath ? TEXT("property") : TEXT("prim")),
+			FText::FromString(bIsStageAttribute ? UsdStage.GetRootLayer().GetRealPath() : ObjectPath)
+		);
+
+		FNotificationInfo ErrorToast( ErrorMessage );
+		ErrorToast.ExpireDuration = 5.0f;
+		ErrorToast.bFireAndForget = true;
+		ErrorToast.Image = FCoreStyle::Get().GetBrush( TEXT( "MessageLog.Warning" ) );
+		FSlateNotificationManager::Get().AddNotification( ErrorToast );
+
+		FUsdLogManager::LogMessage( EMessageSeverity::Warning, ErrorMessage );
+	}
+#endif // #if USE_USD_SDK
+}
+
+void FUsdObjectFieldsViewModel::Refresh(const UE::FUsdStageWeak& InUsdStage, const TCHAR* InObjectPath, float TimeCode)
+{
+	FScopedUnrealAllocs UnrealAllocs;
+
+	UsdStage = InUsdStage;
+	ObjectPath = InObjectPath;
+
+	Fields.Reset();
+
+#if USE_USD_SDK
+	TFunction<void(const FString&, const pxr::VtValue&)> DisplayFieldAsString =
+		[this](const FString& FieldName, const pxr::VtValue& Value)
+		{
+			const bool bReadOnly = true;
+			FString Stringified = UsdToUnreal::ConvertString(pxr::TfStringify(Value));
+
+			// STextBlock can get very slow calculating its desired size for very long string so chop it if needed
+			const int32 MaxValueLength = 300;
+			if (Stringified.Len() > MaxValueLength)
+			{
+				Stringified.LeftInline(MaxValueLength);
+				Stringified.Append(TEXT("..."));
+			}
+
+			CreateField(
+				EObjectFieldType::Metadata,
+				FieldName,
+				Stringified,
+				UsdUtils::EUsdBasicDataTypes::String,
+				TEXT(""),
+				bReadOnly
+			);
+		};
+
+	// Lambda to add rows to the model based on USD object metadata. Extracted here because we can reuse it for
+	// properties and prims
+	TFunction<void(pxr::UsdObject)> AddMetadataFields = [this, &DisplayFieldAsString](pxr::UsdObject Object)
+	{
+		std::map<class pxr::TfToken, pxr::VtValue, pxr::TfDictionaryLessThan> MetadataMap = Object.GetAllMetadata();
+		for (std::map<class pxr::TfToken, pxr::VtValue, pxr::TfDictionaryLessThan>::iterator Iter = MetadataMap.begin();
+			 Iter != MetadataMap.end();
+			 ++Iter)
+		{
+			const pxr::TfToken& Key = Iter->first;
+			const pxr::VtValue& Value = Iter->second;
+
+			FString MetadataName = UsdToUnreal::ConvertToken(Key);
+
+			UsdUtils::FConvertedVtValue ConvertedValue;
+			if (UsdToUnreal::ConvertValue(UE::FVtValue{Value}, ConvertedValue))
+			{
+				if (Value.IsArrayValued())
+				{
+					DisplayFieldAsString(MetadataName, Value);
+				}
+				else
+				{
+					const bool bReadOnly = false;
+					CreateField(
+						EObjectFieldType::Metadata,
+						MetadataName,
+						ConvertedValue,
+						bReadOnly
+					);
+				}
+			}
+			// If we don't know how to display this type, try stringifying it
+			else
+			{
+				DisplayFieldAsString(MetadataName, Value);
+			}
+		}
+	};
+
+	if ( UsdStage )
+	{
+		// Show info about the stage
+		if ( ObjectPath.Equals( TEXT( "/" ) ) || ObjectPath.IsEmpty() )
+		{
+			const bool bReadOnly = true;
+			const FString Role = TEXT( "" );
+			CreateField(
+				EObjectFieldType::Metadata,
+				TEXT("path"),
+				UsdStage.GetRootLayer().GetRealPath(),
+				UsdUtils::EUsdBasicDataTypes::String,
+				Role,
+				bReadOnly
+			);
+
+			FScopedUsdAllocs UsdAllocs;
+
+			std::vector<pxr::TfToken> TokenVector = pxr::SdfSchema::GetInstance().GetMetadataFields( pxr::SdfSpecType::SdfSpecTypePseudoRoot );
+			for ( const pxr::TfToken& Token : TokenVector )
+			{
+				pxr::VtValue VtValue;
+				if ( pxr::UsdStageRefPtr( UsdStage )->GetMetadata( Token, &VtValue ) )
+				{
+					FString FieldName = UsdToUnreal::ConvertToken( Token );
+
+					UsdUtils::FConvertedVtValue ConvertedValue;
+					if ( !UsdToUnreal::ConvertValue( UE::FVtValue{ VtValue }, ConvertedValue ) )
+					{
+						continue;
+					}
+
+					const bool bAttrReadOnly = ConvertedValue.bIsArrayValued;
+					CreateField(
+						EObjectFieldType::Metadata,
+						FieldName,
+						ConvertedValue,
+						bAttrReadOnly
+					);
+				}
+			}
+
+		}
+		// Show info about a prim
+		else if (UE::FUsdPrim UsdPrim = UsdStage.GetPrimAtPath(UE::FSdfPath{*ObjectPath}))
+		{
+			// For now we can't rename/reparent prims through this
+			const bool bPrimReadOnly = true;
+			const FString Role = TEXT("");
+			CreateField(
+				EObjectFieldType::Metadata,
+				TEXT("name"),
+				UsdPrim.GetName().ToString(),
+				UsdUtils::EUsdBasicDataTypes::String,
+				Role,
+				bPrimReadOnly
+			);
+			CreateField(
+				EObjectFieldType::Metadata,
+				TEXT("path"),
+				ObjectPath,
+				UsdUtils::EUsdBasicDataTypes::String,
+				Role,
+				bPrimReadOnly
+			);
+			CreateField(
+				EObjectFieldType::Metadata,
+				TEXT("kind"),
+				UsdToUnreal::ConvertString(IUsdPrim::GetKind(UsdPrim).GetString()),
+				UsdUtils::EUsdBasicDataTypes::Token
+			);
+
+			FScopedUsdAllocs UsdAllocs;
+
+			pxr::UsdPrim PxrUsdPrim{UsdPrim};
+
+			for ( const pxr::UsdAttribute& PrimAttribute : PxrUsdPrim.GetAttributes() )
+			{
+				FString AttributeName = UsdToUnreal::ConvertString( PrimAttribute.GetName().GetString() );
+
+				pxr::VtValue VtValue;
+				PrimAttribute.Get( &VtValue, TimeCode );
+
+				// Just show arrays as readonly strings for now
+				if ( VtValue.IsArrayValued() )
+				{
+					const bool bAttrReadOnly = true;
+					FString Stringified = UsdToUnreal::ConvertString( pxr::TfStringify( VtValue ) );
+
+					// STextBlock can get very slow calculating its desired size for very long string so chop it if needed
+					const int32 MaxValueLength = 300;
+					if ( Stringified.Len() > MaxValueLength )
+					{
+						Stringified.LeftInline( MaxValueLength );
+						Stringified.Append( TEXT( "..." ) );
+					}
+
+					CreateField(
+						EObjectFieldType::Attribute,
+						AttributeName,
+						Stringified,
+						UsdUtils::EUsdBasicDataTypes::String,
+						TEXT(""),
+						bAttrReadOnly
+					);
+				}
+				else
+				{
+					UsdUtils::FConvertedVtValue ConvertedValue;
+					if ( UsdToUnreal::ConvertValue( UE::FVtValue{ VtValue }, ConvertedValue ) )
+					{
+						const bool bAttrReadOnly = false;
+						CreateField(
+							EObjectFieldType::Attribute,
+							AttributeName,
+							ConvertedValue,
+							bAttrReadOnly
+						);
+					}
+
+					if ( PrimAttribute.HasAuthoredConnections() )
+					{
+						const FString ConnectionAttributeName = AttributeName + TEXT(":connect");
+
+						pxr::SdfPathVector ConnectedSources;
+						PrimAttribute.GetConnections( &ConnectedSources );
+
+						for ( pxr::SdfPath& ConnectedSource : ConnectedSources )
+						{
+							UsdUtils::FConvertedVtValueEntry Entry;
+							Entry.Emplace( TInPlaceType< FString >(), UsdToUnreal::ConvertPath( ConnectedSource ) );
+
+							UsdUtils::FConvertedVtValue ConnectionPropertyValue;
+							ConnectionPropertyValue.SourceType = UsdUtils::EUsdBasicDataTypes::String;
+							ConnectionPropertyValue.Entries = { Entry };
+
+							const bool bConnectionValueReadOnly = true;
+							CreateField(
+								EObjectFieldType::Attribute,
+								ConnectionAttributeName,
+								ConnectionPropertyValue,
+								bConnectionValueReadOnly
+							);
+						}
+					}
+				}
+			}
+
+			for (const pxr::UsdRelationship& Relationship : PxrUsdPrim.GetRelationships())
+			{
+				FString RelationshipName = UsdToUnreal::ConvertString(Relationship.GetName().GetString());
+
+				std::vector<pxr::SdfPath> Targets;
+				if (Relationship.GetTargets(&Targets))
+				{
+					if (Targets.size() == 1)
+					{
+						FString UETarget = UsdToUnreal::ConvertPath(Targets[0]);
+						const bool bReadOnly = false;
+						CreateField(
+							EObjectFieldType::Relationship,
+							RelationshipName,
+							UETarget,
+							UsdUtils::EUsdBasicDataTypes::String,
+							TEXT(""),
+							bReadOnly
+						);
+					}
+					else if (Targets.size() > 1)
+					{
+						FString CombinedTargets;
+						for (const pxr::SdfPath& Target : Targets)
+						{
+							CombinedTargets += UsdToUnreal::ConvertPath(Target) + TEXT(", ");
+						}
+						CombinedTargets.RemoveFromEnd(TEXT(", "));
+
+						// STextBlock can get very slow calculating its desired size for very long string so chop it if needed
+						const int32 MaxValueLength = 300;
+						if (CombinedTargets.Len() > MaxValueLength)
+						{
+							CombinedTargets.LeftInline(MaxValueLength);
+							CombinedTargets.Append(TEXT("..."));
+						}
+
+						const bool bReadOnly = true;
+						CreateField(
+							EObjectFieldType::Relationship,
+							RelationshipName,
+							CombinedTargets,
+							UsdUtils::EUsdBasicDataTypes::String,
+							TEXT(""),
+							bReadOnly
+						);
+					}
+				}
+			}
+
+			AddMetadataFields(PxrUsdPrim);
+		}
+		// It's a property path?
+		else if (!ObjectPath.IsEmpty())
+		{
+			FScopedUsdAllocs Allocs;
+
+			pxr::SdfPath UsdPath = UnrealToUsd::ConvertPath(*ObjectPath).Get();
+			if (UsdPath.IsPropertyPath())
+			{
+				pxr::SdfPath PxrPrimPath = UsdPath.GetPrimPath();
+				pxr::TfToken PropertyName = UsdPath.GetNameToken();
+
+				pxr::UsdStageRefPtr PxrUsdStage{UsdStage};
+				if (pxr::UsdPrim Prim = PxrUsdStage->GetPrimAtPath(PxrPrimPath))
+				{
+					if (pxr::UsdProperty Property = Prim.GetProperty(PropertyName))
+					{
+						AddMetadataFields(Property);
+					}
+				}
+			}
+		}
+	}
+
+	Sort();
+#endif // #if USE_USD_SDK
+}
+
+void FUsdObjectFieldsViewModel::Sort()
+{
+	// When sorting by type, we'll sort according to the types first, but still sort alphabetically within each type
+	if (CurrentSortColumn == ObjectFieldColumnIds::TypeColumn)
+	{
+		if (CurrentSortMode == EColumnSortMode::Ascending)
+		{
+			Fields.Sort([](const TSharedPtr<FUsdObjectFieldViewModel>& A, const TSharedPtr<FUsdObjectFieldViewModel>& B)
+			{
+				if(A->Type == B->Type)
+				{
+					return A->Label < B->Label;
+				}
+				else
+				{
+					return A->Type < B->Type;
+				}
+			});
+		}
+		else
+		{
+			Fields.Sort([](const TSharedPtr<FUsdObjectFieldViewModel>& A, const TSharedPtr<FUsdObjectFieldViewModel>& B)
+			{
+				if(A->Type == B->Type)
+				{
+					return A->Label < B->Label;
+				}
+				else
+				{
+					return A->Type > B->Type;
+				}
+			});
+		}
+	}
+	// When sorting by name we'll sort alphabetically regardless of type
+	else if (CurrentSortColumn == ObjectFieldColumnIds::NameColumn)
+	{
+		if (CurrentSortMode == EColumnSortMode::Ascending)
+		{
+			Fields.Sort([](const TSharedPtr<FUsdObjectFieldViewModel>& A, const TSharedPtr<FUsdObjectFieldViewModel>& B)
+			{
+				return A->Label < B->Label;
+			});
+		}
+		else
+		{
+			Fields.Sort([](const TSharedPtr<FUsdObjectFieldViewModel>& A, const TSharedPtr<FUsdObjectFieldViewModel>& B)
+			{
+				return A->Label > B->Label;
+			});
+		}
+	}
+}
+
+UE::FUsdStageWeak FUsdObjectFieldsViewModel::GetUsdStage() const
+{
+	return UsdStage;
+}
+
+FString FUsdObjectFieldsViewModel::GetObjectPath() const
+{
+	return ObjectPath;
+}
+
+#undef LOCTEXT_NAMESPACE
+
+
diff --git a/Engine/Plugins/Importers/USDImporter/Source/USDStageEditorViewModels/Private/USDPrimAttributesViewModel.cpp b/Engine/Plugins/Importers/USDImporter/Source/USDStageEditorViewModels/Private/USDPrimAttributesViewModel.cpp
deleted file mode 100644
diff --git a/Engine/Plugins/Importers/USDImporter/Source/USDStageEditorViewModels/Public/USDObjectFieldViewModel.h b/Engine/Plugins/Importers/USDImporter/Source/USDStageEditorViewModels/Public/USDObjectFieldViewModel.h
new file mode 100644
--- /dev/null
+++ b/Engine/Plugins/Importers/USDImporter/Source/USDStageEditorViewModels/Public/USDObjectFieldViewModel.h
@@ -0,0 +1,85 @@
+// Copyright Epic Games, Inc. All Rights Reserved.
+
+#pragma once
+
+#include "USDValueConversion.h"
+
+#include "UsdWrappers/UsdStage.h"
+
+#include "Templates/SharedPointer.h"
+#include "Widgets/Views/SHeaderRow.h"
+
+class FUsdObjectFieldsViewModel;
+
+enum class EObjectFieldType : int32
+{
+	Metadata = 0,
+	Attribute = 1,
+	Relationship = 2,
+	MAX = 3
+};
+
+namespace ObjectFieldColumnIds
+{
+	inline const FName TypeColumn = TEXT("FieldType");
+	inline const FName NameColumn = TEXT("FieldName");
+	inline const FName ValueColumn = TEXT("FieldValue");
+};
+
+class USDSTAGEEDITORVIEWMODELS_API FUsdObjectFieldViewModel : public TSharedFromThis< FUsdObjectFieldViewModel >
+{
+public:
+	explicit FUsdObjectFieldViewModel( FUsdObjectFieldsViewModel* InOwner );
+
+	// This member function is necessary because the no-RTTI slate module can't query USD for the available token options
+	TArray< TSharedPtr< FString > > GetDropdownOptions() const;
+	void SetAttributeValue( const UsdUtils::FConvertedVtValue& InValue );
+
+public:
+	EObjectFieldType Type;
+	FString Label;
+	UsdUtils::FConvertedVtValue Value;
+	FString ValueRole;
+	bool bReadOnly = false;
+
+private:
+	FUsdObjectFieldsViewModel* Owner;
+};
+
+
+class USDSTAGEEDITORVIEWMODELS_API FUsdObjectFieldsViewModel
+{
+public:
+	template<typename T>
+	void CreateField(
+		EObjectFieldType Type,
+		const FString& FieldName,
+		const T& Value,
+		UsdUtils::EUsdBasicDataTypes UsdType,
+		const FString& ValueRole = FString(),
+		bool bReadOnly = false
+	);
+	void CreateField(
+		EObjectFieldType Type,
+		const FString& FieldName,
+		const UsdUtils::FConvertedVtValue& Value,
+		bool bReadOnly = false
+	);
+
+	void SetFieldValue(const FString& FieldName, const UsdUtils::FConvertedVtValue& Value);
+
+	void Refresh(const UE::FUsdStageWeak& UsdStage, const TCHAR* ObjectPath, float TimeCode);
+	void Sort();
+
+	UE::FUsdStageWeak GetUsdStage() const;
+	FString GetObjectPath() const;
+
+public:
+	EColumnSortMode::Type CurrentSortMode = EColumnSortMode::Ascending;
+	FName CurrentSortColumn = ObjectFieldColumnIds::TypeColumn;
+	TArray<TSharedPtr<FUsdObjectFieldViewModel>> Fields;
+
+private:
+	UE::FUsdStageWeak UsdStage;
+	FString ObjectPath;
+};
diff --git a/Engine/Plugins/Importers/USDImporter/Source/USDStageEditorViewModels/Public/USDPrimAttributesViewModel.h b/Engine/Plugins/Importers/USDImporter/Source/USDStageEditorViewModels/Public/USDPrimAttributesViewModel.h
deleted file mode 100644
diff --git a/Engine/Plugins/Importers/USDImporter/Source/USDUtilities/Private/USDAttributeUtils.cpp b/Engine/Plugins/Importers/USDImporter/Source/USDUtilities/Private/USDAttributeUtils.cpp
--- a/Engine/Plugins/Importers/USDImporter/Source/USDUtilities/Private/USDAttributeUtils.cpp
+++ b/Engine/Plugins/Importers/USDImporter/Source/USDUtilities/Private/USDAttributeUtils.cpp
@@ -175,21 +175,6 @@ bool UsdUtils::IsAttributeMuted( const UE::FUsdAttribute& Attribute, const UE::F
 	return false;
 }
 
-void UsdUtils::NotifyIfOverriddenOpinion( const UE::FUsdAttribute& Attribute )
-{
-#if USE_USD_SDK
-	FScopedUsdAllocs Allocs;
-
-	pxr::UsdAttribute UsdAttribute{Attribute};
-	if ( !UsdAttribute )
-	{
-		return;
-	}
-
-	NotifyIfOverriddenOpinion(UsdAttribute);
-#endif // #if USE_USD_SDK
-}
-
 #if USE_USD_SDK
 void UsdUtils::NotifyIfOverriddenOpinion(const pxr::UsdProperty& Property)
 {
diff --git a/Engine/Plugins/Importers/USDImporter/Source/USDUtilities/Private/USDGeomMeshConversion.cpp b/Engine/Plugins/Importers/USDImporter/Source/USDUtilities/Private/USDGeomMeshConversion.cpp
--- a/Engine/Plugins/Importers/USDImporter/Source/USDUtilities/Private/USDGeomMeshConversion.cpp
+++ b/Engine/Plugins/Importers/USDImporter/Source/USDUtilities/Private/USDGeomMeshConversion.cpp
@@ -2018,7 +2018,7 @@ bool UnrealToUsd::ConvertStaticMesh( const UStaticMesh* StaticMesh, pxr::UsdPrim
 		if ( pxr::UsdAttribute Attr = UsdPrim.CreateAttribute( UnrealIdentifiers::UnrealNaniteOverride, pxr::SdfValueTypeNames->Token ) )
 		{
 			Attr.Set( UnrealIdentifiers::UnrealNaniteOverrideEnable );
-			UsdUtils::NotifyIfOverriddenOpinion( UE::FUsdAttribute{ Attr } );
+			UsdUtils::NotifyIfOverriddenOpinion(Attr);
 		}
 	}
 #endif // WITH_EDITOR
diff --git a/Engine/Plugins/Importers/USDImporter/Source/USDUtilities/Private/USDIntegrationUtils.cpp b/Engine/Plugins/Importers/USDImporter/Source/USDUtilities/Private/USDIntegrationUtils.cpp
--- a/Engine/Plugins/Importers/USDImporter/Source/USDUtilities/Private/USDIntegrationUtils.cpp
+++ b/Engine/Plugins/Importers/USDImporter/Source/USDUtilities/Private/USDIntegrationUtils.cpp
@@ -171,14 +171,14 @@ void UnrealToUsd::ConvertLiveLinkProperties( const UActorComponent* InComponent,
 				: FString();
 
 			Attr.Set( UnrealToUsd::ConvertString( *AnimBPPath ).Get(), pxr::UsdTimeCode::Default() );
-			UsdUtils::NotifyIfOverriddenOpinion( UE::FUsdAttribute{ Attr } );
+			UsdUtils::NotifyIfOverriddenOpinion(Attr);
 		}
 
 		if ( pxr::UsdAttribute Attr = InOutPrim.CreateAttribute( UnrealIdentifiers::UnrealLiveLinkEnabled, pxr::SdfValueTypeNames->Bool ) )
 		{
 			const bool bEnabled = SkeletalComponent->GetAnimationMode() == EAnimationMode::AnimationBlueprint;
 			Attr.Set( bEnabled, pxr::UsdTimeCode::Default() );
-			UsdUtils::NotifyIfOverriddenOpinion( UE::FUsdAttribute{ Attr } );
+			UsdUtils::NotifyIfOverriddenOpinion(Attr);
 		}
 	}
 	// Non-skeletal LiveLink case
@@ -187,13 +187,13 @@ void UnrealToUsd::ConvertLiveLinkProperties( const UActorComponent* InComponent,
 		if ( pxr::UsdAttribute Attr = InOutPrim.CreateAttribute( UnrealIdentifiers::UnrealLiveLinkSubjectName, pxr::SdfValueTypeNames->String ) )
 		{
 			Attr.Set( UnrealToUsd::ConvertString( *InController->SubjectRepresentation.Subject.Name.ToString() ).Get(), pxr::UsdTimeCode::Default() );
-			UsdUtils::NotifyIfOverriddenOpinion( UE::FUsdAttribute{ Attr } );
+			UsdUtils::NotifyIfOverriddenOpinion(Attr);
 		}
 
 		if ( pxr::UsdAttribute Attr = InOutPrim.CreateAttribute( UnrealIdentifiers::UnrealLiveLinkEnabled, pxr::SdfValueTypeNames->Bool ) )
 		{
 			Attr.Set( InController->bEvaluateLiveLink, pxr::UsdTimeCode::Default() );
-			UsdUtils::NotifyIfOverriddenOpinion( UE::FUsdAttribute{ Attr } );
+			UsdUtils::NotifyIfOverriddenOpinion(Attr);
 		}
 	}
 #endif // WITH_EDITOR
diff --git a/Engine/Plugins/Importers/USDImporter/Source/USDUtilities/Private/USDLightConversion.cpp b/Engine/Plugins/Importers/USDImporter/Source/USDUtilities/Private/USDLightConversion.cpp
--- a/Engine/Plugins/Importers/USDImporter/Source/USDUtilities/Private/USDLightConversion.cpp
+++ b/Engine/Plugins/Importers/USDImporter/Source/USDUtilities/Private/USDLightConversion.cpp
@@ -392,7 +392,7 @@ bool UnrealToUsd::ConvertLightComponent( const ULightComponentBase& LightCompone
 	if ( pxr::UsdAttribute Attr = LightAPI.CreateIntensityAttr() )
 	{
 		Attr.Set<float>( LightComponent.Intensity, UsdTimeCode );
-		UsdUtils::NotifyIfOverriddenOpinion( UE::FUsdAttribute{ Attr } );
+		UsdUtils::NotifyIfOverriddenOpinion(Attr);
 	}
 
 	// When converting into UE we multiply intensity and exposure together, so when writing back we just
@@ -400,7 +400,7 @@ bool UnrealToUsd::ConvertLightComponent( const ULightComponentBase& LightCompone
 	if ( pxr::UsdAttribute Attr = LightAPI.CreateExposureAttr() )
 	{
 		Attr.Set<float>( 0.0f, UsdTimeCode );
-		UsdUtils::NotifyIfOverriddenOpinion( UE::FUsdAttribute{ Attr } );
+		UsdUtils::NotifyIfOverriddenOpinion(Attr);
 	}
 
 	if ( const ULightComponent* DerivedLightComponent = Cast<ULightComponent>( &LightComponent ) )
@@ -408,13 +408,13 @@ bool UnrealToUsd::ConvertLightComponent( const ULightComponentBase& LightCompone
 		if ( pxr::UsdAttribute Attr = LightAPI.CreateEnableColorTemperatureAttr() )
 		{
 			Attr.Set<bool>( DerivedLightComponent->bUseTemperature, UsdTimeCode );
-			UsdUtils::NotifyIfOverriddenOpinion( UE::FUsdAttribute{ Attr } );
+			UsdUtils::NotifyIfOverriddenOpinion(Attr);
 		}
 
 		if ( pxr::UsdAttribute Attr = LightAPI.CreateColorTemperatureAttr() )
 		{
 			Attr.Set<float>( DerivedLightComponent->Temperature, UsdTimeCode );
-			UsdUtils::NotifyIfOverriddenOpinion( UE::FUsdAttribute{ Attr } );
+			UsdUtils::NotifyIfOverriddenOpinion(Attr);
 		}
 	}
 
@@ -422,7 +422,7 @@ bool UnrealToUsd::ConvertLightComponent( const ULightComponentBase& LightCompone
 	{
 		pxr::GfVec4f LinearColor = UnrealToUsd::ConvertColor( LightComponent.LightColor );
 		Attr.Set<pxr::GfVec3f>( pxr::GfVec3f( LinearColor[0], LinearColor[1], LinearColor[2] ), UsdTimeCode );
-		UsdUtils::NotifyIfOverriddenOpinion( UE::FUsdAttribute{ Attr } );
+		UsdUtils::NotifyIfOverriddenOpinion(Attr);
 	}
 
 	return true;
@@ -446,7 +446,7 @@ bool UnrealToUsd::ConvertDirectionalLightComponent( const UDirectionalLightCompo
 	if ( pxr::UsdAttribute Attr = Light.CreateAngleAttr() )
 	{
 		Attr.Set<float>( LightComponent.LightSourceAngle, UsdTimeCode );
-		UsdUtils::NotifyIfOverriddenOpinion( UE::FUsdAttribute{ Attr } );
+		UsdUtils::NotifyIfOverriddenOpinion(Attr);
 	}
 
 	// USD intensity units should be in Nits == Lux / Steradian, but there is no
@@ -483,7 +483,7 @@ bool UnrealToUsd::ConvertRectLightComponent( const URectLightComponent& LightCom
 		if ( pxr::UsdAttribute Attr = DiskLight.CreateRadiusAttr() )
 		{
 			Attr.Set<float>( UnrealToUsd::ConvertDistance( StageInfo, Radius ), UsdTimeCode );
-			UsdUtils::NotifyIfOverriddenOpinion( UE::FUsdAttribute{ Attr } );
+			UsdUtils::NotifyIfOverriddenOpinion(Attr);
 		}
 	}
 	// Rect light
@@ -494,13 +494,13 @@ bool UnrealToUsd::ConvertRectLightComponent( const URectLightComponent& LightCom
 		if ( pxr::UsdAttribute Attr = RectLight.CreateWidthAttr() )
 		{
 			Attr.Set<float>( UnrealToUsd::ConvertDistance( StageInfo, LightComponent.SourceWidth ), UsdTimeCode );
-			UsdUtils::NotifyIfOverriddenOpinion( UE::FUsdAttribute{ Attr } );
+			UsdUtils::NotifyIfOverriddenOpinion(Attr);
 		}
 
 		if ( pxr::UsdAttribute Attr = RectLight.CreateHeightAttr() )
 		{
 			Attr.Set<float>( UnrealToUsd::ConvertDistance( StageInfo, LightComponent.SourceHeight ), UsdTimeCode );
-			UsdUtils::NotifyIfOverriddenOpinion( UE::FUsdAttribute{ Attr } );
+			UsdUtils::NotifyIfOverriddenOpinion(Attr);
 		}
 	}
 	else
@@ -550,7 +550,7 @@ bool UnrealToUsd::ConvertRectLightComponent( const URectLightComponent& LightCom
 		}
 
 		Attr.Set<float>( FinalIntensityNits, UsdTimeCode );
-		UsdUtils::NotifyIfOverriddenOpinion( UE::FUsdAttribute{ Attr } );
+		UsdUtils::NotifyIfOverriddenOpinion(Attr);
 	}
 
 	return true;
@@ -576,13 +576,13 @@ bool UnrealToUsd::ConvertPointLightComponent( const UPointLightComponent& LightC
 	if ( pxr::UsdAttribute Attr = Light.CreateRadiusAttr() )
 	{
 		Attr.Set<float>( UnrealToUsd::ConvertDistance( StageInfo, LightComponent.SourceRadius ), UsdTimeCode );
-		UsdUtils::NotifyIfOverriddenOpinion( UE::FUsdAttribute{ Attr } );
+		UsdUtils::NotifyIfOverriddenOpinion(Attr);
 	}
 
 	if ( pxr::UsdAttribute Attr = Light.CreateTreatAsPointAttr() )
 	{
 		Attr.Set<bool>( FMath::IsNearlyZero( LightComponent.SourceRadius ), UsdTimeCode );
-		UsdUtils::NotifyIfOverriddenOpinion( UE::FUsdAttribute{ Attr } );
+		UsdUtils::NotifyIfOverriddenOpinion(Attr);
 	}
 
 	float SolidAngle = 4.f * PI;
@@ -624,7 +624,7 @@ bool UnrealToUsd::ConvertPointLightComponent( const UPointLightComponent& LightC
 		}
 
 		Attr.Set<float>( FinalIntensityNits, UsdTimeCode );
-		UsdUtils::NotifyIfOverriddenOpinion( UE::FUsdAttribute{ Attr } );
+		UsdUtils::NotifyIfOverriddenOpinion(Attr);
 	}
 
 	return true;
@@ -666,7 +666,7 @@ bool UnrealToUsd::ConvertSkyLightComponent( const USkyLightComponent& LightCompo
 
 				UsdUtils::MakePathRelativeToLayer( UE::FSdfLayer{ Prim.GetStage()->GetEditTarget().GetLayer() }, FilePath );
 				Attr.Set<pxr::SdfAssetPath>( pxr::SdfAssetPath{ UnrealToUsd::ConvertString( *FilePath ).Get() }, UsdTimeCode );
-				UsdUtils::NotifyIfOverriddenOpinion( UE::FUsdAttribute{ Attr } );
+				UsdUtils::NotifyIfOverriddenOpinion(Attr);
 			}
 		}
 	}
@@ -688,7 +688,7 @@ bool UnrealToUsd::ConvertSpotLightComponent( const USpotLightComponent& LightCom
 	if ( pxr::UsdAttribute ConeAngleAttr = ShapingAPI.CreateShapingConeAngleAttr() )
 	{
 		ConeAngleAttr.Set<float>( LightComponent.OuterConeAngle, UsdTimeCode );
-		UsdUtils::NotifyIfOverriddenOpinion( UE::FUsdAttribute{ ConeAngleAttr } );
+		UsdUtils::NotifyIfOverriddenOpinion(ConeAngleAttr);
 	}
 
 	// As of March 2021 there doesn't seem to be a consensus on what softness means, according to https://groups.google.com/g/usd-interest/c/A6bc4OZjSB0
@@ -698,7 +698,7 @@ bool UnrealToUsd::ConvertSpotLightComponent( const USpotLightComponent& LightCom
 		// Keep in [0, 1] range, where 1 is maximum softness, i.e. inner cone angle is zero
 		const float Softness = FMath::IsNearlyZero( LightComponent.OuterConeAngle ) ? 0.0 : 1.0f - LightComponent.InnerConeAngle / LightComponent.OuterConeAngle;
 		SoftnessAttr.Set<float>( Softness, UsdTimeCode );
-		UsdUtils::NotifyIfOverriddenOpinion( UE::FUsdAttribute{ SoftnessAttr } );
+		UsdUtils::NotifyIfOverriddenOpinion(SoftnessAttr);
 	}
 
 	return true;
diff --git a/Engine/Plugins/Importers/USDImporter/Source/USDUtilities/Private/USDPrimConversion.cpp b/Engine/Plugins/Importers/USDImporter/Source/USDUtilities/Private/USDPrimConversion.cpp
--- a/Engine/Plugins/Importers/USDImporter/Source/USDUtilities/Private/USDPrimConversion.cpp
+++ b/Engine/Plugins/Importers/USDImporter/Source/USDUtilities/Private/USDPrimConversion.cpp
@@ -1408,31 +1408,31 @@ bool UnrealToUsd::ConvertCameraComponent( const UCineCameraComponent& CameraComp
 	if ( pxr::UsdAttribute Attr = GeomCamera.CreateFocalLengthAttr() )
 	{
 		Attr.Set<float>( UnrealToUsd::ConvertDistance( StageInfo, CameraComponent.CurrentFocalLength ), UsdTimeCode );
-		UsdUtils::NotifyIfOverriddenOpinion( UE::FUsdAttribute{ Attr } );
+		UsdUtils::NotifyIfOverriddenOpinion(Attr);
 	}
 
 	if ( pxr::UsdAttribute Attr = GeomCamera.CreateFocusDistanceAttr() )
 	{
 		Attr.Set<float>( UnrealToUsd::ConvertDistance( StageInfo, CameraComponent.FocusSettings.ManualFocusDistance ), UsdTimeCode );
-		UsdUtils::NotifyIfOverriddenOpinion( UE::FUsdAttribute{ Attr } );
+		UsdUtils::NotifyIfOverriddenOpinion(Attr);
 	}
 
 	if ( pxr::UsdAttribute Attr = GeomCamera.CreateFStopAttr() )
 	{
 		Attr.Set<float>( CameraComponent.CurrentAperture, UsdTimeCode );
-		UsdUtils::NotifyIfOverriddenOpinion( UE::FUsdAttribute{ Attr } );
+		UsdUtils::NotifyIfOverriddenOpinion(Attr);
 	}
 
 	if ( pxr::UsdAttribute Attr = GeomCamera.CreateHorizontalApertureAttr() )
 	{
 		Attr.Set<float>( UnrealToUsd::ConvertDistance( StageInfo, CameraComponent.Filmback.SensorWidth ), UsdTimeCode );
-		UsdUtils::NotifyIfOverriddenOpinion( UE::FUsdAttribute{ Attr } );
+		UsdUtils::NotifyIfOverriddenOpinion(Attr);
 	}
 
 	if ( pxr::UsdAttribute Attr = GeomCamera.CreateVerticalApertureAttr() )
 	{
 		Attr.Set<float>( UnrealToUsd::ConvertDistance( StageInfo, CameraComponent.Filmback.SensorHeight ), UsdTimeCode );
-		UsdUtils::NotifyIfOverriddenOpinion( UE::FUsdAttribute{ Attr } );
+		UsdUtils::NotifyIfOverriddenOpinion(Attr);
 	}
 
 	return true;
@@ -1983,7 +1983,7 @@ bool UnrealToUsd::ConvertSceneComponent( const pxr::UsdStageRefPtr& Stage, const
 		}
 
 		VisibilityAttr.Set<pxr::TfToken>( Value );
-		UsdUtils::NotifyIfOverriddenOpinion( UE::FUsdAttribute{ VisibilityAttr } );
+		UsdUtils::NotifyIfOverriddenOpinion(VisibilityAttr);
 	}
 
 	return true;
@@ -2215,25 +2215,25 @@ bool UnrealToUsd::ConvertHierarchicalInstancedStaticMeshComponent( const UHierar
 	if ( UsdAttribute Attr = PointInstancer.CreateProtoIndicesAttr() )
 	{
 		Attr.Set( ProtoIndices, UsdTimeCode );
-		UsdUtils::NotifyIfOverriddenOpinion( UE::FUsdAttribute{ Attr } );
+		UsdUtils::NotifyIfOverriddenOpinion(Attr);
 	}
 
 	if ( UsdAttribute Attr = PointInstancer.CreatePositionsAttr() )
 	{
 		Attr.Set( Positions, UsdTimeCode );
-		UsdUtils::NotifyIfOverriddenOpinion( UE::FUsdAttribute{ Attr } );
+		UsdUtils::NotifyIfOverriddenOpinion(Attr);
 	}
 
 	if ( UsdAttribute Attr = PointInstancer.CreateOrientationsAttr() )
 	{
 		Attr.Set( Orientations, UsdTimeCode );
-		UsdUtils::NotifyIfOverriddenOpinion( UE::FUsdAttribute{ Attr } );
+		UsdUtils::NotifyIfOverriddenOpinion(Attr);
 	}
 
 	if ( UsdAttribute Attr = PointInstancer.CreateScalesAttr() )
 	{
 		Attr.Set( Scales, UsdTimeCode );
-		UsdUtils::NotifyIfOverriddenOpinion( UE::FUsdAttribute{ Attr } );
+		UsdUtils::NotifyIfOverriddenOpinion(Attr);
 	}
 
 	return true;
@@ -2264,8 +2264,8 @@ bool UnrealToUsd::ConvertXformable( const FTransform& RelativeTransform, pxr::Us
 	{
 		MatrixXform.Set( UsdTransform, UsdTimeCode );
 
-		UsdUtils::NotifyIfOverriddenOpinion( UE::FUsdAttribute{ MatrixXform.GetAttr() } );
-		UsdUtils::NotifyIfOverriddenOpinion( UE::FUsdAttribute{ XForm.GetXformOpOrderAttr() } );
+		UsdUtils::NotifyIfOverriddenOpinion(MatrixXform.GetAttr());
+		UsdUtils::NotifyIfOverriddenOpinion(XForm.GetXformOpOrderAttr());
 	}
 
 	return true;
@@ -2427,25 +2427,25 @@ bool UnrealToUsd::ConvertInstancedFoliageActor( const AInstancedFoliageActor& Ac
 	if ( UsdAttribute Attr = PointInstancer.CreateProtoIndicesAttr() )
 	{
 		Attr.Set( ProtoIndices, UsdTimeCode );
-		UsdUtils::NotifyIfOverriddenOpinion( UE::FUsdAttribute{ Attr } );
+		UsdUtils::NotifyIfOverriddenOpinion(Attr);
 	}
 
 	if ( UsdAttribute Attr = PointInstancer.CreatePositionsAttr() )
 	{
 		Attr.Set( Positions, UsdTimeCode );
-		UsdUtils::NotifyIfOverriddenOpinion( UE::FUsdAttribute{ Attr } );
+		UsdUtils::NotifyIfOverriddenOpinion(Attr);
 	}
 
 	if ( UsdAttribute Attr = PointInstancer.CreateOrientationsAttr() )
 	{
 		Attr.Set( Orientations, UsdTimeCode );
-		UsdUtils::NotifyIfOverriddenOpinion( UE::FUsdAttribute{ Attr } );
+		UsdUtils::NotifyIfOverriddenOpinion(Attr);
 	}
 
 	if ( UsdAttribute Attr = PointInstancer.CreateScalesAttr() )
 	{
 		Attr.Set( Scales, UsdTimeCode );
-		UsdUtils::NotifyIfOverriddenOpinion( UE::FUsdAttribute{ Attr } );
+		UsdUtils::NotifyIfOverriddenOpinion(Attr);
 	}
 
 	return true;
@@ -3321,7 +3321,7 @@ UnrealToUsd::FPropertyTrackWriter UnrealToUsd::CreatePropertyTrackWriter( const
 		// "attribute creation" spec first. This because NotifyIfOverriddenOpinion will go through the attribute's
 		// spec stack and consider our attribute overriden if it finds a stronger opinion than the one on the edit
 		// target. Well if our own spec hasn't been created yet it will misfire when it runs into any other spec
-		UsdUtils::NotifyIfOverriddenOpinion( UE::FUsdAttribute{ Attr } );
+		UsdUtils::NotifyIfOverriddenOpinion(Attr);
 	}
 
 	return Result;
diff --git a/Engine/Plugins/Importers/USDImporter/Source/USDUtilities/Private/USDShadeConversion.cpp b/Engine/Plugins/Importers/USDImporter/Source/USDUtilities/Private/USDShadeConversion.cpp
--- a/Engine/Plugins/Importers/USDImporter/Source/USDUtilities/Private/USDShadeConversion.cpp
+++ b/Engine/Plugins/Importers/USDImporter/Source/USDUtilities/Private/USDShadeConversion.cpp
@@ -2674,7 +2674,7 @@ bool UsdUtils::MarkMaterialPrimWithWorldSpaceNormals( const UE::FUsdPrim& Materi
 	}
 
 	Attr.Set<bool>( true );
-	UsdUtils::NotifyIfOverriddenOpinion( UE::FUsdAttribute{ Attr } );
+	UsdUtils::NotifyIfOverriddenOpinion(Attr);
 	return true;
 }
 
diff --git a/Engine/Plugins/Importers/USDImporter/Source/USDUtilities/Private/USDSkeletalDataConversion.cpp b/Engine/Plugins/Importers/USDImporter/Source/USDUtilities/Private/USDSkeletalDataConversion.cpp
--- a/Engine/Plugins/Importers/USDImporter/Source/USDUtilities/Private/USDSkeletalDataConversion.cpp
+++ b/Engine/Plugins/Importers/USDImporter/Source/USDUtilities/Private/USDSkeletalDataConversion.cpp
@@ -3472,10 +3472,10 @@ bool UnrealToUsd::ConvertControlRigSection(
 		EndInclTickFrame = UpperBoundToUse.GetValue().GetValue() + ( UpperBoundToUse.GetValue().IsInclusive() ? 0 : -1 );
 	}
 
-	UsdUtils::NotifyIfOverriddenOpinion( UE::FUsdAttribute{ BlendShapeWeightsAttr } );
-	UsdUtils::NotifyIfOverriddenOpinion( UE::FUsdAttribute{ TranslationsAttr } );
-	UsdUtils::NotifyIfOverriddenOpinion( UE::FUsdAttribute{ RotationsAttr } );
-	UsdUtils::NotifyIfOverriddenOpinion( UE::FUsdAttribute{ ScalesAttr } );
+	UsdUtils::NotifyIfOverriddenOpinion(BlendShapeWeightsAttr);
+	UsdUtils::NotifyIfOverriddenOpinion(TranslationsAttr);
+	UsdUtils::NotifyIfOverriddenOpinion(RotationsAttr);
+	UsdUtils::NotifyIfOverriddenOpinion(ScalesAttr);
 
 	pxr::VtArray<pxr::TfToken> CurveNames;
 	pxr::VtArray<float> CurvesValuesForTime;
@@ -3609,7 +3609,7 @@ bool UnrealToUsd::ConvertControlRigSection(
 						if ( bRenamedAChannel )
 						{
 							ChannelsAttr.Set( BlendShapeChannels );
-							UsdUtils::NotifyIfOverriddenOpinion( UE::FUsdAttribute{ ChannelsAttr } );
+							UsdUtils::NotifyIfOverriddenOpinion(ChannelsAttr);
 						}
 					}
 				}
@@ -3622,7 +3622,7 @@ bool UnrealToUsd::ConvertControlRigSection(
 		BlendShapesAttr.Set( CurveNames );
 		CurvesValuesForTime.resize( CurveNames.size() );
 
-		UsdUtils::NotifyIfOverriddenOpinion( UE::FUsdAttribute{ BlendShapesAttr } );
+		UsdUtils::NotifyIfOverriddenOpinion(BlendShapesAttr);
 	}
 
 	FFrameTime TickIncr = FFrameRate::TransformTime( 1, DisplayRate, TickResolution );
diff --git a/Engine/Plugins/Importers/USDImporter/Source/USDUtilities/Public/USDAttributeUtils.h b/Engine/Plugins/Importers/USDImporter/Source/USDUtilities/Public/USDAttributeUtils.h
--- a/Engine/Plugins/Importers/USDImporter/Source/USDUtilities/Public/USDAttributeUtils.h
+++ b/Engine/Plugins/Importers/USDImporter/Source/USDUtilities/Public/USDAttributeUtils.h
@@ -36,7 +36,6 @@ namespace UsdUtils
 	 * Use this after setting any attribute: This function does nothing in case the Stage's current edit target has
 	 * the strongest opinion for the attribute already
 	 */
-	USDUTILITIES_API void NotifyIfOverriddenOpinion( const UE::FUsdAttribute& Attribute );
 #if USE_USD_SDK
 	USDUTILITIES_API void NotifyIfOverriddenOpinion( const pxr::UsdProperty& Property );
 #endif // USE_USD_SDK
diff --git a/Engine/Plugins/Importers/USDImporter/Source/UnrealUSDWrapper/Private/UsdWrappers/UsdAttribute.cpp b/Engine/Plugins/Importers/USDImporter/Source/UnrealUSDWrapper/Private/UsdWrappers/UsdAttribute.cpp
--- a/Engine/Plugins/Importers/USDImporter/Source/UnrealUSDWrapper/Private/UsdWrappers/UsdAttribute.cpp
+++ b/Engine/Plugins/Importers/USDImporter/Source/UnrealUSDWrapper/Private/UsdWrappers/UsdAttribute.cpp
@@ -142,6 +142,16 @@ namespace UE
 	{
 		return Impl->PxrUsdAttribute.Get();
 	}
+
+	FUsdAttribute::operator pxr::UsdProperty&()
+	{
+		return Impl->PxrUsdAttribute.Get();
+	}
+
+	FUsdAttribute::operator const pxr::UsdProperty&() const
+	{
+		return Impl->PxrUsdAttribute.Get();
+	}
 #endif // #if USE_USD_SDK
 
 	bool FUsdAttribute::GetMetadata( const TCHAR* Key, UE::FVtValue& Value ) const
diff --git a/Engine/Plugins/Importers/USDImporter/Source/UnrealUSDWrapper/Public/UsdWrappers/UsdAttribute.h b/Engine/Plugins/Importers/USDImporter/Source/UnrealUSDWrapper/Public/UsdWrappers/UsdAttribute.h
--- a/Engine/Plugins/Importers/USDImporter/Source/UnrealUSDWrapper/Public/UsdWrappers/UsdAttribute.h
+++ b/Engine/Plugins/Importers/USDImporter/Source/UnrealUSDWrapper/Public/UsdWrappers/UsdAttribute.h
@@ -13,6 +13,7 @@
 
 PXR_NAMESPACE_OPEN_SCOPE
 	class UsdAttribute;
+	class UsdProperty;
 PXR_NAMESPACE_CLOSE_SCOPE
 
 #endif // #if USE_USD_SDK
@@ -58,6 +59,12 @@ namespace UE
 
 		operator pxr::UsdAttribute&();
 		operator const pxr::UsdAttribute&() const;
+
+		// Also define the conversion operators directly to the base class as we can't "forward declare the class
+		// inheritance", but it is nice to be able to pass an UE::FUsdAttribute to a function that takes an
+		// pxr::UsdProperty directly
+		operator pxr::UsdProperty&();
+		operator const pxr::UsdProperty&() const;
 #endif // #if USE_USD_SDK
 
 	// Wrapped pxr::UsdObject functions, refer to the USD SDK documentation
