USD: Don't try stringifying array values with more than 30k entries when displaying on prim properties panel;

Always show element count when displaying array values.

diff --git a/Engine/Plugins/Importers/USDImporter/Source/USDStageEditorViewModels/Private/USDObjectFieldViewModel.cpp b/Engine/Plugins/Importers/USDImporter/Source/USDStageEditorViewModels/Private/USDObjectFieldViewModel.cpp
--- a/Engine/Plugins/Importers/USDImporter/Source/USDStageEditorViewModels/Private/USDObjectFieldViewModel.cpp
+++ b/Engine/Plugins/Importers/USDImporter/Source/USDStageEditorViewModels/Private/USDObjectFieldViewModel.cpp
@@ -276,24 +276,39 @@ void FUsdObjectFieldsViewModel::Refresh(const UE::FUsdStageWeak& InUsdStage, con
 	TFunction<void(const FString&, const pxr::VtValue&)> DisplayFieldAsString =
 		[this](const FString& FieldName, const pxr::VtValue& Value)
 		{
-			const bool bReadOnly = true;
-			FString Stringified = UsdToUnreal::ConvertString(pxr::TfStringify(Value));
+			FString Stringified;
+
+			if (Value.IsArrayValued())
+			{
+				Stringified = FString::Printf(TEXT("%d elements: "), Value.GetArraySize());
+			}
 
-			// STextBlock can get very slow calculating its desired size for very long string so chop it if needed
-			const int32 MaxValueLength = 300;
-			if (Stringified.Len() > MaxValueLength)
+			// This array it's too large to even stringify fast enough, so for now just show the element count
+			if (Value.IsArrayValued() && Value.GetArraySize() > 5000)
 			{
-				Stringified.LeftInline(MaxValueLength);
-				Stringified.Append(TEXT("..."));
+				Stringified += TEXT("[too many entries to expand]");
 			}
+			else
+			{
+				Stringified += UsdToUnreal::ConvertString(pxr::TfStringify(Value));
 
+				// STextBlock can get very slow calculating its desired size for very long string so chop it if needed
+				const int32 MaxValueLength = 300;
+				if (Stringified.Len() > MaxValueLength)
+				{
+					Stringified.LeftInline(MaxValueLength);
+					Stringified.Append(TEXT("...]"));
+				}
+			}
+
+			const bool bAttrReadOnly = true;
 			CreateField(
 				EObjectFieldType::Metadata,
 				FieldName,
 				Stringified,
 				UsdUtils::EUsdBasicDataTypes::String,
 				TEXT(""),
-				bReadOnly
+				bAttrReadOnly
 			);
 		};
 
@@ -423,17 +438,27 @@ void FUsdObjectFieldsViewModel::Refresh(const UE::FUsdStageWeak& InUsdStage, con
 				// Just show arrays as readonly strings for now
 				if ( VtValue.IsArrayValued() )
 				{
-					const bool bAttrReadOnly = true;
-					FString Stringified = UsdToUnreal::ConvertString( pxr::TfStringify( VtValue ) );
+					FString Stringified = FString::Printf(TEXT("%d elements: "), VtValue.GetArraySize());
 
-					// STextBlock can get very slow calculating its desired size for very long string so chop it if needed
-					const int32 MaxValueLength = 300;
-					if ( Stringified.Len() > MaxValueLength )
+					// This array it's too large to even stringify fast enough, so for now just show the element count
+					if (VtValue.GetArraySize() > 5000)
+					{
+						Stringified += TEXT("[too many entries to expand]");
+					}
+					else
 					{
-						Stringified.LeftInline( MaxValueLength );
-						Stringified.Append( TEXT( "..." ) );
+						Stringified += UsdToUnreal::ConvertString(pxr::TfStringify(VtValue));
+
+						// STextBlock can get very slow calculating its desired size for very long string so chop it if needed
+						const int32 MaxValueLength = 300;
+						if (Stringified.Len() > MaxValueLength)
+						{
+							Stringified.LeftInline(MaxValueLength);
+							Stringified.Append(TEXT("...]"));
+						}
 					}
 
+					const bool bAttrReadOnly = true;
 					CreateField(
 						EObjectFieldType::Attribute,
 						AttributeName,
@@ -507,7 +532,7 @@ void FUsdObjectFieldsViewModel::Refresh(const UE::FUsdStageWeak& InUsdStage, con
 					}
 					else if (Targets.size() > 1)
 					{
-						FString CombinedTargets;
+						FString CombinedTargets = FString::Printf(TEXT("%d elements: ["));
 						for (const pxr::SdfPath& Target : Targets)
 						{
 							CombinedTargets += UsdToUnreal::ConvertPath(Target) + TEXT(", ");
@@ -522,6 +547,8 @@ void FUsdObjectFieldsViewModel::Refresh(const UE::FUsdStageWeak& InUsdStage, con
 							CombinedTargets.Append(TEXT("..."));
 						}
 
+						CombinedTargets += TEXT("]");
+
 						const bool bReadOnly = true;
 						CreateField(
 							EObjectFieldType::Relationship,
