USD: Support MaterialX materials referenced from USD files with Unreal's MaterialX importer.

diff --git a/Engine/Plugins/Importers/USDImporter/Source/USDClasses/Private/USDAssetCache2.cpp b/Engine/Plugins/Importers/USDImporter/Source/USDClasses/Private/USDAssetCache2.cpp
--- a/Engine/Plugins/Importers/USDImporter/Source/USDClasses/Private/USDAssetCache2.cpp
+++ b/Engine/Plugins/Importers/USDImporter/Source/USDClasses/Private/USDAssetCache2.cpp
@@ -607,12 +607,6 @@ void UUsdAssetCache2::CacheAsset(const FString& Hash, UObject* Asset, const UObj
 		return;
 	}
 
-	UE_LOG(LogUsd, Verbose, TEXT("Caching asset '%s' with hash '%s' into the USD Asset Cache '%s'"),
-		*Asset->GetPathName(),
-		*Hash,
-		*GetPathName()
-	);
-
 	FCachedAssetInfo* FoundInfo = nullptr;
 	{
 		FReadScopeLock Lock(RWLock);
@@ -639,12 +633,18 @@ void UUsdAssetCache2::CacheAsset(const FString& Hash, UObject* Asset, const UObj
 	{
 		if (Asset == ExistingAsset)
 		{
+			UE_LOG(LogUsd, Verbose, TEXT("Recaching asset '%s' with hash '%s' into the USD Asset Cache '%s'"),
+				*Asset->GetPathName(),
+				*Hash,
+				*GetPathName()
+			);
+
 			// We're trying to cache an asset that's already in the cache with this same hash.
 			// Just record that we "touched" this asset and return
 			Modify();
 
 			FWriteScopeLock Lock(RWLock);
-			TouchAsset(Asset, ReferencerToUse);
+			TouchAssetInternal(Asset, ReferencerToUse);
 			return;
 		}
 		else
@@ -677,6 +677,12 @@ void UUsdAssetCache2::CacheAsset(const FString& Hash, UObject* Asset, const UObj
 		}
 	}
 
+	UE_LOG(LogUsd, Verbose, TEXT("Caching asset '%s' with hash '%s' into the USD Asset Cache '%s'"),
+		*Asset->GetPathName(),
+		*Hash,
+		*GetPathName()
+	);
+
 	Modify();
 
 	FWriteScopeLock Lock(RWLock);
@@ -804,7 +810,7 @@ UObject* UUsdAssetCache2::GetCachedAsset(const FString& Hash)
 	if (FoundAsset)
 	{
 		FWriteScopeLock Lock{RWLock};
-		TouchAsset(FoundAsset, CurrentScopedReferencer);
+		TouchAssetInternal(FoundAsset, CurrentScopedReferencer);
 		return FoundAsset;
 	}
 
@@ -919,7 +925,7 @@ UObject* UUsdAssetCache2::GetCachedAsset(const FString& Hash)
 				UObject* LoadedAsset = LoadAssetFromBulkData(FoundCachedInfo);
 				ensure(LoadedAsset);
 
-				TouchAsset(LoadedAsset, CurrentScopedReferencer);
+				TouchAssetInternal(LoadedAsset, CurrentScopedReferencer);
 				FoundAsset = LoadedAsset;  // The very last asset we loaded is our main asset, as it's pushed last into AssetsToLoad
 			}
 		}
@@ -928,6 +934,26 @@ UObject* UUsdAssetCache2::GetCachedAsset(const FString& Hash)
 	return FoundAsset;
 }
 
+bool UUsdAssetCache2::TouchAsset(const UObject* Asset, const UObject* Referencer)
+{
+	if (!Asset)
+	{
+		return false;
+	}
+
+	if (IsAssetOwnedByCache(Asset->GetPathName()))
+	{
+		Modify();
+
+		FWriteScopeLock Lock(RWLock);
+		TouchAssetInternal(Asset, Referencer);
+
+		return true;
+	}
+
+	return false;
+}
+
 bool UUsdAssetCache2::AddAssetReference(const UObject* Asset, const UObject* Referencer)
 {
 	FWriteScopeLock Lock(RWLock);
@@ -1508,14 +1534,14 @@ bool UUsdAssetCache2::AddAssetReferenceInternal(const UObject* Asset, const UObj
 	return false;
 }
 
-void UUsdAssetCache2::TouchAsset(UObject* Asset, const UObject* Referencer)
+void UUsdAssetCache2::TouchAssetInternal(const UObject* Asset, const UObject* Referencer)
 {
 	if (!Asset)
 	{
 		return;
 	}
 
-	ActiveAssets.Add(Asset);
+	ActiveAssets.Add(const_cast<UObject*>(Asset));
 	ensure(LRUCache.FindAndTouch(Asset));
 
 	if (Referencer)
diff --git a/Engine/Plugins/Importers/USDImporter/Source/USDClasses/Public/USDAssetCache2.h b/Engine/Plugins/Importers/USDImporter/Source/USDClasses/Public/USDAssetCache2.h
--- a/Engine/Plugins/Importers/USDImporter/Source/USDClasses/Public/USDAssetCache2.h
+++ b/Engine/Plugins/Importers/USDImporter/Source/USDClasses/Public/USDAssetCache2.h
@@ -70,6 +70,13 @@ public:
 	UFUNCTION(BlueprintCallable, Category = "Caching", meta = (CallInEditor = "true"))
 	UObject* GetCachedAsset(const FString& Hash);
 
+	/**
+	 * Marks the provided asset as being used at this point, optionally adding a specific referencer.
+	 * This is useful because the asset cache will always prioritize retaining the most recently used assets
+	 */
+	UFUNCTION(BlueprintCallable, Category = "Caching", meta = (CallInEditor = "true"))
+	bool TouchAsset(const UObject* Asset, const UObject* Referencer = nullptr);
+
 	/**
 	 * Adds a new UObject referencer to a particular asset, returning true if the operation succeeded.
 	 *
@@ -171,7 +178,7 @@ private:
 	bool TryUnloadAsset(FCachedAssetInfo& InOutInfo);
 	bool IsAssetOwnedByCacheInternal(const FString& AssetPath) const;
 	bool AddAssetReferenceInternal(const UObject* Asset, const UObject* Referencer);
-	void TouchAsset(UObject* Asset, const UObject* Referencer = nullptr);
+	void TouchAssetInternal(const UObject* Asset, const UObject* Referencer = nullptr);
 
 private:
 	friend struct FUsdScopedAssetCacheReferencer;
diff --git a/Engine/Plugins/Importers/USDImporter/Source/USDSchemas/Private/Custom/MDLUSDShadeMaterialTranslator.cpp b/Engine/Plugins/Importers/USDImporter/Source/USDSchemas/Private/Custom/MDLUSDShadeMaterialTranslator.cpp
--- a/Engine/Plugins/Importers/USDImporter/Source/USDSchemas/Private/Custom/MDLUSDShadeMaterialTranslator.cpp
+++ b/Engine/Plugins/Importers/USDImporter/Source/USDSchemas/Private/Custom/MDLUSDShadeMaterialTranslator.cpp
@@ -167,12 +167,6 @@ void FMdlUsdShadeMaterialTranslator::CreateAssets()
 
 			if ( MdlMaterial )
 			{
-				UUsdAssetImportData* ImportData = NewObject< UUsdAssetImportData >( MdlMaterial, TEXT( "USDAssetImportData" ) );
-				ImportData->PrimPath = PrimPath.GetString();
-				MdlMaterial->AssetImportData = ImportData;
-
-				Context->AssetCache->CacheAsset( MdlFullName, MdlMaterial );
-
 				UE::MDLShadeTranslatorImpl::Private::NotifyIfMaterialNeedsVirtualTextures( MdlMaterial );
 			}
 			else
@@ -187,11 +181,6 @@ void FMdlUsdShadeMaterialTranslator::CreateAssets()
 		if ( !MdlMaterialInstance && MdlMaterial )
 		{
 			MdlMaterialInstance = NewObject< UMaterialInstanceConstant >(GetTransientPackage(), NAME_None, Context->ObjectFlags | RF_Transient);
-
-			UUsdAssetImportData* ImportData = NewObject< UUsdAssetImportData >( MdlMaterialInstance, TEXT( "USDAssetImportData" ) );
-			ImportData->PrimPath = PrimPath.GetString();
-			MdlMaterialInstance->AssetImportData = ImportData;
-
 			MdlMaterialInstance->SetParentEditorOnly( MdlMaterial );
 
 			UsdToUnreal::ConvertShadeInputsToParameters( ShadeMaterial, *MdlMaterialInstance, Context->AssetCache.Get(), *Context->RenderContext.ToString() );
@@ -210,76 +199,11 @@ void FMdlUsdShadeMaterialTranslator::CreateAssets()
 			MdlMaterialInstance->PreEditChange( nullptr );
 			MdlMaterialInstance->PostEditChange();
 
-			Context->AssetCache->CacheAsset(MdlFullInstanceName, MdlMaterialInstance);
-
 			UE::MDLShadeTranslatorImpl::Private::NotifyIfMaterialNeedsVirtualTextures( MdlMaterialInstance );
 		}
 
-		if (Context->InfoCache)
-		{
-			Context->InfoCache->LinkAssetToPrim(PrimPath, MdlMaterial);
-			Context->InfoCache->LinkAssetToPrim(PrimPath, MdlMaterialInstance);
-
-			if (UMaterial* MdlReference = Cast<UMaterial>(MdlMaterial))
-			{
-				TArray<UTexture*> UsedTextures;
-				const bool bAllQualityLevels = true;
-				const bool bAllFeatureLevels = true;
-				MdlReference->GetUsedTextures(
-					UsedTextures,
-					EMaterialQualityLevel::High,
-					bAllQualityLevels,
-					ERHIFeatureLevel::SM5,
-					bAllFeatureLevels
-				);
-
-				for (UTexture* Texture : UsedTextures)
-				{
-					if (Texture->GetOutermost() == GetTransientPackage())
-					{
-						Context->InfoCache->LinkAssetToPrim(PrimPath, Texture);
-
-						const FString FilePath = Texture->AssetImportData ? Texture->AssetImportData->GetFirstFilename() : Texture->GetName();
-						const FString TextureHash = UsdUtils::GetTextureHash(
-							FilePath,
-							Texture->SRGB,
-							Texture->CompressionSettings,
-							Texture->GetTextureAddressX(),
-							Texture->GetTextureAddressY()
-						);
-
-						Texture->SetFlags(RF_Transient);
-						Context->AssetCache->CacheAsset(TextureHash, Texture);
-					}
-				}
-			}
-
-			if (MdlMaterialInstance)
-			{
-				for (const FTextureParameterValue& TextureValue : MdlMaterialInstance->TextureParameterValues)
-				{
-					if (UTexture* Texture = TextureValue.ParameterValue)
-					{
-						if (Texture->GetOutermost() == GetTransientPackage())
-						{
-							Context->InfoCache->LinkAssetToPrim(PrimPath, Texture);
-
-							const FString FilePath = Texture->AssetImportData ? Texture->AssetImportData->GetFirstFilename() : Texture->GetName();
-							const FString TextureHash = UsdUtils::GetTextureHash(
-								FilePath,
-								Texture->SRGB,
-								Texture->CompressionSettings,
-								Texture->GetTextureAddressX(),
-								Texture->GetTextureAddressY()
-							);
-
-							Texture->SetFlags(RF_Transient);
-							Context->AssetCache->CacheAsset(TextureHash, Texture);
-						}
-					}
-				}
-			}
-		}
+		PostImportMaterial(MdlFullName, MdlMaterial);
+		PostImportMaterial(MdlFullInstanceName, MdlMaterialInstance);
 	}
 	else
 	{
diff --git a/Engine/Plugins/Importers/USDImporter/Source/USDSchemas/Private/Custom/MDLUSDShadeMaterialTranslator.h b/Engine/Plugins/Importers/USDImporter/Source/USDSchemas/Private/Custom/MDLUSDShadeMaterialTranslator.h
--- a/Engine/Plugins/Importers/USDImporter/Source/USDSchemas/Private/Custom/MDLUSDShadeMaterialTranslator.h
+++ b/Engine/Plugins/Importers/USDImporter/Source/USDSchemas/Private/Custom/MDLUSDShadeMaterialTranslator.h
@@ -4,24 +4,26 @@
 
 #include "USDShadeMaterialTranslator.h"
 
+#include "Custom/MaterialXUSDShadeMaterialTranslator.h"
+
 #if USE_USD_SDK && WITH_EDITOR
 
 #include "USDIncludesStart.h"
 	#include "pxr/pxr.h"
 #include "USDIncludesEnd.h"
 
-class FMdlUsdShadeMaterialTranslator : public FUsdShadeMaterialTranslator
+class FMdlUsdShadeMaterialTranslator : public FMaterialXUsdShadeMaterialTranslator
 {
-	using Super = FUsdShadeMaterialTranslator;
+	using Super = FMaterialXUsdShadeMaterialTranslator;
 
 public:
 	static FName MdlRenderContext;
 
 public:
-	using FUsdShadeMaterialTranslator::FUsdShadeMaterialTranslator;
+	using FMaterialXUsdShadeMaterialTranslator::FMaterialXUsdShadeMaterialTranslator;
 
 	virtual void CreateAssets() override;
-	
+
 };
 
 #endif // #if USE_USD_SDK && WITH_EDITOR
diff --git a/Engine/Plugins/Importers/USDImporter/Source/USDSchemas/Private/Custom/MaterialXUSDShadeMaterialTranslator.cpp b/Engine/Plugins/Importers/USDImporter/Source/USDSchemas/Private/Custom/MaterialXUSDShadeMaterialTranslator.cpp
new file mode 100644
--- /dev/null
+++ b/Engine/Plugins/Importers/USDImporter/Source/USDSchemas/Private/Custom/MaterialXUSDShadeMaterialTranslator.cpp
@@ -0,0 +1,510 @@
+// Copyright Epic Games, Inc. All Rights Reserved.
+
+#include "Custom/MaterialXUSDShadeMaterialTranslator.h"
+#if USE_USD_SDK
+#include "USDAssetCache.h"
+#include "USDAssetImportData.h"
+#include "USDClassesModule.h"
+#include "USDConversionUtils.h"
+#include "USDErrorUtils.h"
+#include "USDLog.h"
+#include "USDShadeConversion.h"
+#include "USDTypesConversion.h"
+
+#include "UsdWrappers/SdfLayer.h"
+#include "UsdWrappers/SdfPath.h"
+#include "UsdWrappers/UsdPrim.h"
+
+#include "Engine/Level.h"
+#include "Engine/Texture.h"
+#include "InterchangeAssetImportData.h"
+#include "InterchangeGenericMaterialPipeline.h"
+#include "InterchangeManager.h"
+#include "InterchangeTexture2DFactoryNode.h"
+#include "InterchangeTexture2DNode.h"
+#include "Materials/Material.h"
+#include "Materials/MaterialInstance.h"
+#include "Materials/MaterialInstanceConstant.h"
+#include "MaterialShared.h"
+#include "Misc/PackageName.h"
+#include "UObject/StrongObjectPtr.h"
+
+#include "USDIncludesStart.h"
+	#include "pxr/usd/sdf/assetPath.h"
+	#include "pxr/usd/usd/primCompositionQuery.h"
+	#include "pxr/usd/usdShade/material.h"
+#include "USDIncludesEnd.h"
+
+#if WITH_EDITOR
+#include "MaterialXFormat/Util.h"
+#endif
+
+#define LOCTEXT_NAMESPACE "MaterialXUSDShadeMaterialTranslator"
+
+static bool GUseInterchangeMaterialXTranslator = true;
+static FAutoConsoleVariableRef CvarUseInterchangeMaterialXTranslator(
+	TEXT("USD.UseInterchangeMaterialXTranslator"),
+	GUseInterchangeMaterialXTranslator,
+	TEXT("Whether to translate MaterialX materials referenced by USD files with Unreal's own MaterialX importer. If instead this is false, we will try parsing the generated UsdShadeMaterial prims as generated by USD's usdMtlx plugin directly"));
+
+FName FMaterialXUsdShadeMaterialTranslator::MaterialXRenderContext = TEXT("mtlx");
+
+namespace UE::USDMaterialXTranslatorInternal::Private
+{
+	void HashMaterialXOptions(UInterchangePipelineBase* MaterialXOptions, FSHA1& SHA1)
+	{
+		// Reference: UInterchangePipelineBase::SaveSettingsInternal
+
+		if (!MaterialXOptions)
+		{
+			return;
+		}
+
+		int32 PortFlags = 0;
+		UClass* Class = MaterialXOptions->GetClass();
+		for (FProperty* Property = Class->PropertyLink; Property; Property = Property->PropertyLinkNext)
+		{
+			FString Key = Property->GetName();
+			const FName PropertyName = Property->GetFName();
+
+			const bool bIsPropertyInherited = Property->GetOwnerClass() != Class;
+			UObject* SuperClassDefaultObject = Class->GetSuperClass()->GetDefaultObject();
+			FObjectProperty* SubObject = CastField<FObjectProperty>(Property);
+			FArrayProperty* Array = CastField<FArrayProperty>(Property);
+			if (Array)
+			{
+				const bool bForce = true;
+				const bool bConst = false;
+
+				FScriptArrayHelper_InContainer ArrayHelper(Array, MaterialXOptions);
+				for (int32 i = 0; i < ArrayHelper.Num(); i++)
+				{
+					FString	Buffer;
+					Array->Inner->ExportTextItem_Direct(Buffer, ArrayHelper.GetRawPtr(i), ArrayHelper.GetRawPtr(i), MaterialXOptions, PortFlags);
+
+					SHA1.UpdateWithString(*Buffer, Buffer.Len());
+				}
+			}
+			else if (UInterchangePipelineBase* SubPipeline = SubObject ? Cast<UInterchangePipelineBase>(SubObject->GetObjectPropertyValue_InContainer(MaterialXOptions)) : nullptr)
+			{
+				// Save the settings if the referenced pipeline is a subobject of ours
+				if (SubPipeline->IsInOuter(MaterialXOptions))
+				{
+					HashMaterialXOptions(SubPipeline, SHA1);
+				}
+			}
+			else
+			{
+				TCHAR TempKey[MAX_SPRINTF] = TEXT("");
+				for (int32 Index = 0; Index < Property->ArrayDim; Index++)
+				{
+					if (Property->ArrayDim != 1)
+					{
+						FCString::Sprintf(TempKey, TEXT("%s[%i]"), *Property->GetName(), Index);
+						Key = TempKey;
+					}
+
+					FString	Value;
+					Property->ExportText_InContainer(Index, Value, MaterialXOptions, MaterialXOptions, MaterialXOptions, PortFlags);
+					SHA1.UpdateWithString(*Value, Value.Len());
+				}
+			}
+		}
+	}
+}
+
+void FMaterialXUsdShadeMaterialTranslator::CreateAssets()
+{
+#if WITH_EDITOR
+	// We handle MaterialX materials here by leveraging USD's own usdMtlx plugin. What it does is translate references
+	// to .mtlx files into hierarchies of Material/Shader/NodeGraph prims as the stage is opened. To actually use
+	// the materials contained in the .mtlx files, other prims from the stage will create material bindings to those
+	// generated Material prims. For example, a stage that wishes to use MaterialX materials may have this:
+	//
+	// def Mesh "SomeMesh"
+	// {
+	//     ...
+	//     rel material:binding = </MaterialX/Materials/TextureMaterial>
+	// }
+	//
+	// def Scope "MaterialX" (
+	//     references = [
+	//         @./textureTest.mtlx@</MaterialX>,
+	//     ]
+	// )
+	// {
+	// }
+	//
+	// Note how SomeMesh references the "TextureMaterial" material: That's the name of one of the surfacematerials
+	// inside the MaterialX file. This because when the stage is opened, the usdMtlx plugin will generate and
+	// compose a child "Materials" prim and additional child prims something like the below, which is what the external
+	// prims will be referencing:
+	//
+	// def Scope "MaterialX"
+	// {
+	//     def "Materials"
+	//     {
+	//         def Material "TextureMaterial"
+	//         {
+	//             float inputs:base = 1
+	//             color3f inputs:base_color
+	//             float inputs:coat
+	//             float inputs:coat_affect_color
+	//             ...
+	//             token outputs:mtlx:surface.connect = </MaterialX/Materials/TextureMaterial/ND_standard_surface_surfaceshader.outputs:surface>
+	//
+	//             def Shader "ND_standard_surface_surfaceshader"
+	//             {
+	//                 uniform token info:id = "ND_standard_surface_surfaceshader"
+	//                 float inputs:base.connect = </MaterialX/Materials/TextureMaterial.inputs:base>
+	//                 color3f inputs:base_color.connect = </MaterialX/Materials/TextureMaterial/NG_imagetex1.outputs:out_color_0>
+	//                 float inputs:coat.connect = </MaterialX/Materials/TextureMaterial.inputs:coat>
+	//                 float inputs:coat_affect_color.connect = </MaterialX/Materials/TextureMaterial.inputs:coat_affect_color>
+	//                 ...
+	//             }
+	//         }
+	//     }
+	// }
+	//
+	// Our goal in here is to translate those generated Materials (like "TextureMaterial" above). This because we want
+	// to generate UMaterialInterfaces via Unreal's MaterialX plugin instead, and link them to those generated Material
+	// prims within our info cache. This way, the rest or our USDImporter plugin doesn't need to care or know where
+	// this material came from: It will find an UMaterialInterface linked to that Material prim and it will happily use
+	// that as any other material.
+
+	if (Context->RenderContext != MaterialXRenderContext || !GUseInterchangeMaterialXTranslator)
+	{
+		Super::CreateAssets();
+		return;
+	}
+
+	if (!Context->AssetCache || !Context->InfoCache)
+	{
+		return;
+	}
+
+	if (Context->bTranslateOnlyUsedMaterials && !Context->InfoCache->IsMaterialUsed(PrimPath))
+	{
+		UE_LOG(LogUsd, Verbose, TEXT("Skipping creating assets for material prim '%s' as it is not currently bound by any prim."), *PrimPath.GetString());
+		return;
+	}
+
+	FScopedUsdAllocs UsdAllocs;
+
+	pxr::UsdPrim Prim = GetPrim();
+	pxr::UsdShadeMaterial ShadeMaterial(Prim);
+	if (!ShadeMaterial)
+	{
+		return;
+	}
+
+	const static pxr::TfToken RenderContextToken = UnrealToUsd::ConvertToken(*MaterialXRenderContext.ToString()).Get();
+	pxr::UsdShadeShader SurfaceShader = ShadeMaterial.ComputeSurfaceSource(RenderContextToken);
+	if (!SurfaceShader)
+	{
+		Super::CreateAssets();
+		return;
+	}
+
+	// This material prim has the mtlx render context, so maybe it is one of the ones generated by usdMtlx.
+	// Let's traverse upwards and try finding a .mtlx file reference in one of our parents.
+
+	TFunction<bool(const pxr::UsdPrim&, FString&)> GetMaterialXFilePath = [this](const pxr::UsdPrim& Prim, FString& OutMaterialXFilePath)
+	{
+		pxr::UsdPrimCompositionQuery PrimCompositionQuery = pxr::UsdPrimCompositionQuery::GetDirectReferences(Prim);
+		for (const pxr::UsdPrimCompositionQueryArc& CompositionArc : PrimCompositionQuery.GetCompositionArcs())
+		{
+			if (CompositionArc.GetArcType() == pxr::PcpArcTypeReference)
+			{
+				pxr::SdfReferenceEditorProxy ReferenceEditor;
+				pxr::SdfReference UsdReference;
+
+				if (CompositionArc.GetIntroducingListEditor(&ReferenceEditor, &UsdReference))
+				{
+					// The mtlx file will be the "target layer". It's useful to get its real path via the SdfLayer
+					// interface as it already makes sure it is absolute
+					pxr::SdfLayerHandle TargetLayer = CompositionArc.GetTargetLayer();
+					FString AbsoluteFilePath = UsdToUnreal::ConvertString(TargetLayer->GetRealPath());
+
+					FString Extension = FPaths::GetExtension(AbsoluteFilePath);
+					if (Extension == MaterialXRenderContext && FPaths::FileExists(AbsoluteFilePath))
+					{
+						OutMaterialXFilePath = AbsoluteFilePath;
+						return true;
+					}
+				}
+			}
+		}
+
+		return false;
+	};
+
+	FString MaterialXFilePath;
+	pxr::UsdPrim MaterialXReferencerPrim = Prim;
+	if (!GetMaterialXFilePath(Prim, MaterialXFilePath))
+	{
+		// We know the usdMtlx plugin always generates a "Materials" schemaless prim to contain all the generated
+		// Materials, so let's use that too.
+		pxr::UsdPrim ParentPrim = Prim.GetParent();
+		while (ParentPrim && ParentPrim.GetName() != "Materials")
+		{
+			ParentPrim = ParentPrim.GetParent();
+		}
+
+		if (ParentPrim)
+		{
+			// This prim likely holds the reference to the MaterialX file, but let's search upwards too
+			pxr::UsdPrim MaterialXReferencerCandidate = ParentPrim.GetParent();
+			while (MaterialXReferencerCandidate)
+			{
+				if (GetMaterialXFilePath(MaterialXReferencerCandidate, MaterialXFilePath))
+				{
+					break;
+				}
+
+				MaterialXReferencerCandidate = MaterialXReferencerCandidate.GetParent();
+			}
+			MaterialXReferencerPrim = MaterialXReferencerCandidate;
+		}
+	}
+
+	if (MaterialXFilePath.IsEmpty() || !MaterialXReferencerPrim)
+	{
+		UE_LOG(LogUsd, Warning, TEXT("Recognized potential MaterialX materials on prim '%s', but failed to find a valid referenced MaterialX file. Reverting to parsing the generated Material prims instead."), *PrimPath.GetString());
+		Super::CreateAssets();
+		return;
+	}
+
+	// Quick shallow parse of MaterialX documents ourselves so that we can find all the referenced MaterialX files and
+	// build a robust hash to use for caching our generated assets
+	FSHA1 SHA1;
+	{
+		TSet<FString> ReferencedMaterialXFiles;
+		TFunction<void(const FString&)> CollectReferencedMaterialXFilesRecursive;
+		CollectReferencedMaterialXFilesRecursive =
+			[&ReferencedMaterialXFiles, &CollectReferencedMaterialXFilesRecursive]
+			(const FString& ActiveDocumentFullPath)
+			{
+				namespace mx = MaterialX;
+
+				mx::DocumentPtr Document = mx::createDocument();
+				mx::readFromXmlFile(Document, TCHAR_TO_UTF8(*ActiveDocumentFullPath));
+
+				for (const std::string& ReferencedURI : Document->getReferencedSourceUris())
+				{
+					FString UEReferencedURI = UTF8_TO_TCHAR(ReferencedURI.c_str());
+
+					// This could be relative to the referencer .mtlx file
+					UEReferencedURI = FPaths::ConvertRelativePathToFull(
+						FPaths::GetPath(ActiveDocumentFullPath),
+						UEReferencedURI
+					);
+
+					if (ReferencedMaterialXFiles.Contains(UEReferencedURI))
+					{
+						continue;
+					}
+
+					if (FPaths::FileExists(UEReferencedURI))
+					{
+						ReferencedMaterialXFiles.Add(UEReferencedURI);
+						CollectReferencedMaterialXFilesRecursive(UEReferencedURI);
+					}
+				}
+			};
+
+		// Note we don't have to add MaterialXFilePath ourselves because getReferencedSourceUris
+		// always includes at least the active document itself anyway
+		CollectReferencedMaterialXFilesRecursive(MaterialXFilePath);
+
+		for (const FString& ReferencedMaterialXFile : ReferencedMaterialXFiles)
+		{
+			if (TUniquePtr<FArchive> Ar{IFileManager::Get().CreateFileReader(*ReferencedMaterialXFile)})
+			{
+				TArray<uint8> LocalScratch;
+				LocalScratch.SetNumUninitialized(1024 * 64);
+
+				const int64 Size = Ar->TotalSize();
+				int64 Position = 0;
+
+				// Read in BufferSize chunks
+				while (Position < Size)
+				{
+					const int64 ReadNum = FMath::Min(Size - Position, (int64)LocalScratch.Num());
+					Ar->Serialize(LocalScratch.GetData(), ReadNum);
+					SHA1.Update(LocalScratch.GetData(), ReadNum);
+
+					Position += ReadNum;
+				}
+			}
+		}
+	}
+
+	// Hash the options we'll use so that we regenerate assets if those change
+	if (Context->MaterialXOptions)
+	{
+		UE::USDMaterialXTranslatorInternal::Private::HashMaterialXOptions(Context->MaterialXOptions, SHA1);
+	}
+
+	FSHAHash Hash;
+	SHA1.Final();
+	SHA1.GetHash(&Hash.Hash[0]);
+	FString MaterialXHash = Hash.ToString();
+
+	FString TargetMaterialAssetHash = MaterialXHash + TEXT("/") + UsdToUnreal::ConvertString(Prim.GetName());
+	UMaterialInterface* ParsedMaterial = Cast< UMaterialInterface >(Context->AssetCache->GetCachedAsset(TargetMaterialAssetHash));
+	if (!ParsedMaterial)
+	{
+		// If we don't have this asset in the cache yet, we need to parse the MaterialX file
+
+		UInterchangeManager& Manager = UInterchangeManager::GetInterchangeManager();
+		UInterchangeSourceData* SourceData = Manager.CreateSourceData(MaterialXFilePath);
+
+		FString ReferencerPrimPath = *UsdToUnreal::ConvertPath(MaterialXReferencerPrim.GetPrimPath());
+
+		FImportAssetParameters InterchangeParameters;
+		InterchangeParameters.bIsAutomated = true;
+
+		if (Context->MaterialXOptions)
+		{
+			InterchangeParameters.OverridePipelines.Add(Context->MaterialXOptions);
+		}
+
+		// Get a temp content folder we can give interchange as it must have a destination
+		int32 PackageSuffix = FMath::RandRange(0, TNumericLimits<int32>::Max()-1);
+		FString TempPackagePath = FString::Printf(TEXT("/Engine/USDImporter/Transient/%d"), PackageSuffix++);
+
+		// Make sure that content folder doesn't have any packages already in it
+		TArray<FString> Filenames;
+		while (FPackageName::FindPackagesInDirectory(Filenames, TempPackagePath))
+		{
+			TempPackagePath = FString::Printf(TEXT("/Engine/USDImporter/Transient/%d"), PackageSuffix++);
+		}
+
+		// Currently Interchange cannot import directly into the transient package, so we'll import into a temporary folder and then
+		// move everything over
+		InterchangeParameters.OnAssetsImportDoneNative.BindLambda(
+			[this, &MaterialXFilePath, &MaterialXHash, &TargetMaterialAssetHash, &ParsedMaterial, &ReferencerPrimPath, &TempPackagePath]
+			(const TArray<UObject*>& ImportedObjects)
+			{
+				// Move all our assets to the transient package
+				for (UObject* ImportedObject : ImportedObjects)
+				{
+					const FName NewName = MakeUniqueObjectName(
+						GetTransientPackage(),
+						ImportedObject->GetClass(),
+						ImportedObject->GetFName()
+					);
+					bool bSuccess = ImportedObject->Rename(
+						*NewName.ToString(),
+						GetTransientPackage(),
+						REN_NonTransactional | REN_DontCreateRedirectors | REN_ForceNoResetLoaders
+					);
+					if (!bSuccess)
+					{
+						UE_LOG(LogUsd, Warning, TEXT("Failed to rename imported MaterialX material '%s', generated for referencer prim '%s'."),
+							*ImportedObject->GetPathName(),
+							*ReferencerPrimPath
+						);
+						continue;
+					}
+
+					ImportedObject->SetFlags(RF_Transient);
+
+					if (UMaterialInterface* Material = Cast<UMaterialInterface>(ImportedObject))
+					{
+						// MaterialX names are unique, and can only have alphanumeric and the "_" character, so we should
+						// always have a solid enough mapping to assume UAsset name == Prim name == MaterialX name
+						FString NewMaterialHash = MaterialXHash + TEXT("/") + Material->GetFName().GetPlainNameString();
+
+						// We generate all assets from the MaterialX file once, but it's possible we're just updating a
+						// single material prim here. If we were to cache all assets here, we'd potentially be trying
+						// to overwrite the existing assets that are being used by other prims and wouldn't otherwise
+						// be discarded, so make sure we don't do that
+						UMaterialInterface* ExistingMaterial = Cast<UMaterialInterface>(Context->AssetCache->GetCachedAsset(NewMaterialHash));
+						if (!ExistingMaterial)
+						{
+							Context->AssetCache->CacheAsset(NewMaterialHash, Material);
+						}
+
+						if (NewMaterialHash == TargetMaterialAssetHash)
+						{
+							ParsedMaterial = Material;
+						}
+					}
+					else if (UTexture* Texture = Cast<UTexture>(ImportedObject))
+					{
+						// Note: This filename will be the .mtlx file for now, which is why we append the texture name manually later
+						const FString FilePath = Texture->AssetImportData ? Texture->AssetImportData->GetFirstFilename() : TEXT("");
+
+						FString TextureHash = UsdUtils::GetTextureHash(
+							FilePath,
+							Texture->SRGB,
+							Texture->CompressionSettings,
+							Texture->GetTextureAddressX(),
+							Texture->GetTextureAddressY()
+						);
+						// TODO: Replace this hash hack with relying on PostImportMaterial whenever the interchange
+						// texture AssetImportData contains the correct filepaths, and not just the .mtlx file
+						TextureHash += TEXT("_") + Texture->GetFName().GetPlainNameString();
+
+						// See comment on the analogous part of the Material case above
+						UTexture* ExistingTexture = Cast<UTexture>(Context->AssetCache->GetCachedAsset(TextureHash));
+						if (!ExistingTexture)
+						{
+							UUsdAssetImportData* ImportData = NewObject<UUsdAssetImportData>(Texture, TEXT("USDAssetImportData"));
+							ImportData->PrimPath = ReferencerPrimPath;
+							Texture->AssetImportData = ImportData;
+
+							Context->AssetCache->CacheAsset(TextureHash, Texture);
+						}
+					}
+					else
+					{
+						ensureMsgf(false, TEXT("Asset type unsupported!"));
+					}
+				}
+
+				// Now that the asset cache took everything we imported we shouldn't have anything else
+				TArray<FString> Filenames;
+				if (FPackageName::FindPackagesInDirectory(Filenames, TempPackagePath))
+				{
+					ensureMsgf(false, TEXT("We should not have any leftover assets from MaterialX translation!"));
+				}
+			}
+		);
+
+		const bool bSuccess = Manager.ImportAsset(TempPackagePath, SourceData, InterchangeParameters);
+		if (!bSuccess)
+		{
+			UE_LOG(LogUsd, Warning, TEXT("Recognized potential MaterialX materials on prim '%s', but MaterialX parsing of file '%s' failed. Reverting to parsing the generated Material prims instead."),
+				*PrimPath.GetString(),
+				*MaterialXFilePath
+			);
+			Super::CreateAssets();
+			return;
+		}
+	}
+
+	if (ParsedMaterial)
+	{
+		PostImportMaterial(TargetMaterialAssetHash, ParsedMaterial);
+	}
+	else
+	{
+		UE_LOG(LogUsd, Warning, TEXT("Parsing of MaterialX file '%s' on prim '%s' succeeded, but we failed to find the target material '%s' within the file."),
+			*MaterialXFilePath,
+			*PrimPath.GetString(),
+			*UsdToUnreal::ConvertString(Prim.GetName())
+		);
+		Super::CreateAssets();
+		return;
+	}
+#else
+	Super::CreateAssets();
+#endif // WITH_EDITOR
+}
+
+#undef LOCTEXT_NAMESPACE
+
+#endif // #if USE_USD_SDK
diff --git a/Engine/Plugins/Importers/USDImporter/Source/USDSchemas/Private/Custom/MaterialXUSDShadeMaterialTranslator.h b/Engine/Plugins/Importers/USDImporter/Source/USDSchemas/Private/Custom/MaterialXUSDShadeMaterialTranslator.h
new file mode 100644
--- /dev/null
+++ b/Engine/Plugins/Importers/USDImporter/Source/USDSchemas/Private/Custom/MaterialXUSDShadeMaterialTranslator.h
@@ -0,0 +1,27 @@
+// Copyright Epic Games, Inc. All Rights Reserved.
+
+#pragma once
+
+#include "USDShadeMaterialTranslator.h"
+
+#if USE_USD_SDK
+
+#include "USDIncludesStart.h"
+	#include "pxr/pxr.h"
+#include "USDIncludesEnd.h"
+
+class FMaterialXUsdShadeMaterialTranslator : public FUsdShadeMaterialTranslator
+{
+	using Super = FUsdShadeMaterialTranslator;
+
+public:
+	static FName MaterialXRenderContext;
+
+public:
+	using FUsdShadeMaterialTranslator::FUsdShadeMaterialTranslator;
+
+	virtual void CreateAssets() override;
+
+};
+
+#endif // #if USE_USD_SDK
diff --git a/Engine/Plugins/Importers/USDImporter/Source/USDSchemas/Private/MeshTranslationImpl.cpp b/Engine/Plugins/Importers/USDImporter/Source/USDSchemas/Private/MeshTranslationImpl.cpp
--- a/Engine/Plugins/Importers/USDImporter/Source/USDSchemas/Private/MeshTranslationImpl.cpp
+++ b/Engine/Plugins/Importers/USDImporter/Source/USDSchemas/Private/MeshTranslationImpl.cpp
@@ -16,6 +16,7 @@
 #include "Materials/MaterialInstanceConstant.h"
 #include "Materials/MaterialInstanceDynamic.h"
 #include "Materials/MaterialInterface.h"
+#include "MaterialShared.h"
 #include "Misc/Paths.h"
 #include "UObject/Package.h"
 
@@ -30,6 +31,126 @@
 #include "pxr/usd/usdShade/material.h"
 #include "USDIncludesEnd.h"
 
+namespace UE::MeshTranslationImplInternal::Private
+{
+	UMaterialInterface* CreateTwoSidedVersionOfMaterial(UMaterialInterface* OneSidedMat)
+	{
+		if (!OneSidedMat)
+		{
+			return nullptr;
+		}
+
+		FString PrimPath;
+#if WITH_EDITOR
+		if (UUsdAssetImportData* ImportData = Cast<UUsdAssetImportData>(OneSidedMat->AssetImportData))
+		{
+			PrimPath = ImportData->PrimPath;
+		}
+#endif // WITH_EDITOR
+
+		UMaterialInterface* TwoSidedMat = nullptr;
+
+		UMaterialInstance* OneSidedMaterialInstance = Cast<UMaterialInstance>(OneSidedMat);
+
+		// Important to use Parent.Get() and not GetBaseMaterial() here because if our parent is the translucent we'll
+		// get the reference UsdPreviewSurface instead, as that is also *its* reference
+		UMaterialInterface* ReferenceMaterial = OneSidedMaterialInstance ? OneSidedMaterialInstance->Parent.Get() : nullptr;
+		UMaterialInterface* ReferenceMaterialTwoSided = nullptr;
+		if (ReferenceMaterial && ReferenceMaterial->GetPathName().StartsWith(TEXT("/USDImporter/Materials/UsdPreviewSurface")))
+		{
+			ReferenceMaterialTwoSided = MeshTranslationImpl::GetTwoSidedVersionOfReferencePreviewSurfaceMaterial(ReferenceMaterial);
+		}
+
+		const FName NewInstanceName = MakeUniqueObjectName(
+			GetTransientPackage(),
+			UMaterialInstance::StaticClass(),
+			*(OneSidedMat->GetName() + UnrealIdentifiers::TwoSidedMaterialSuffix)
+		);
+
+		UMaterialInstanceDynamic* MID = Cast<UMaterialInstanceDynamic>(OneSidedMat);
+#if WITH_EDITOR
+		UMaterialInstanceConstant* MIC = Cast<UMaterialInstanceConstant>(OneSidedMat);
+
+		// One-sided material is an instance of one of our USD reference materials.
+		// Just create an instance of the TwoSided version of the same reference material and copy parameter values.
+		if (GIsEditor && MIC && ReferenceMaterialTwoSided)
+		{
+			UMaterialInstanceConstant* TwoSidedMIC = NewObject<UMaterialInstanceConstant>(
+				GetTransientPackage(),
+				NewInstanceName,
+				OneSidedMat->GetFlags()
+			);
+			if (TwoSidedMIC)
+			{
+				UUsdAssetImportData* ImportData = NewObject< UUsdAssetImportData >(
+					TwoSidedMIC,
+					TEXT("USDAssetImportData")
+				);
+				ImportData->PrimPath = PrimPath;
+				TwoSidedMIC->AssetImportData = ImportData;
+			}
+
+			TwoSidedMIC->SetParentEditorOnly(ReferenceMaterialTwoSided);
+			TwoSidedMIC->CopyMaterialUniformParametersEditorOnly(OneSidedMat);
+
+			TwoSidedMat = TwoSidedMIC;
+		}
+
+		// One-sided material is some other material (e.g. MaterialX/MDL-generated material).
+		// Create a new material instance of it and set the override to two-sided.
+		else if (GIsEditor)
+		{
+			UMaterialInstanceConstant* TwoSidedMIC = NewObject<UMaterialInstanceConstant>(
+				GetTransientPackage(),
+				NewInstanceName,
+				OneSidedMat->GetFlags()
+			);
+			if (TwoSidedMIC)
+			{
+				UUsdAssetImportData* ImportData = NewObject< UUsdAssetImportData >(
+					TwoSidedMIC,
+					TEXT("USDAssetImportData")
+				);
+				ImportData->PrimPath = PrimPath;
+				TwoSidedMIC->AssetImportData = ImportData;
+			}
+
+			TwoSidedMIC->SetParentEditorOnly(OneSidedMat);
+			TwoSidedMIC->BasePropertyOverrides.bOverride_TwoSided = true;
+			TwoSidedMIC->BasePropertyOverrides.TwoSided = true;
+
+			FMaterialUpdateContext UpdateContext(FMaterialUpdateContext::EOptions::Default, GMaxRHIShaderPlatform);
+			UpdateContext.AddMaterialInstance(TwoSidedMIC);
+
+			TwoSidedMat = TwoSidedMIC;
+		}
+
+		else
+#endif // WITH_EDITOR
+
+		// At runtime all we can do is create another instance of our two-sided reference materials, we cannot set
+		// another override
+		if (MID && ReferenceMaterialTwoSided)
+		{
+			UMaterialInstanceDynamic* TwoSidedMID = UMaterialInstanceDynamic::Create(
+				ReferenceMaterialTwoSided,
+				GetTransientPackage(),
+				NewInstanceName
+			);
+			if (!ensure(TwoSidedMID))
+			{
+				return nullptr;
+			}
+
+			TwoSidedMID->CopyParameterOverrides(MID);
+
+			TwoSidedMat = TwoSidedMID;
+		}
+
+		return TwoSidedMat;
+	}
+};
+
 TMap<const UsdUtils::FUsdPrimMaterialSlot*, UMaterialInterface*> MeshTranslationImpl::ResolveMaterialAssignmentInfo(
 	const pxr::UsdPrim& UsdPrim,
 	const TArray<UsdUtils::FUsdPrimMaterialAssignmentInfo>& AssignmentInfo,
@@ -95,10 +216,10 @@ TMap<const UsdUtils::FUsdPrimMaterialSlot*, UMaterialInterface*> MeshTranslation
 			}
 			case UsdUtils::EPrimAssignmentType::MaterialPrim:
 			{
-				UMaterialInstance* OneSidedMat = nullptr;
+				UMaterialInterface* OneSidedMat = nullptr;
 
-				TSet<UMaterialInstance*> ExistingMaterials = InfoCache.GetAssetsForPrim<UMaterialInstance>(UE::FSdfPath{*Slot.MaterialSource});
-				for (UMaterialInstance* ExistingMaterial : ExistingMaterials)
+				TSet<UMaterialInterface*> ExistingMaterials = InfoCache.GetAssetsForPrim<UMaterialInterface>(UE::FSdfPath{*Slot.MaterialSource});
+				for (UMaterialInterface* ExistingMaterial : ExistingMaterials)
 				{
 					const bool bExistingIsTwoSided = ExistingMaterial->IsTwoSided();
 
@@ -131,68 +252,10 @@ TMap<const UsdUtils::FUsdPrimMaterialSlot*, UMaterialInterface*> MeshTranslation
 
 						// Check if for some reason we already have a two-sided material ready due to a complex scenario
 						// related to the global cache
-						UMaterialInstance* TwoSidedMat = Cast<UMaterialInstance>(AssetCache.GetCachedAsset(MaterialHash));
+						UMaterialInterface* TwoSidedMat = Cast<UMaterialInterface>(AssetCache.GetCachedAsset(MaterialHash));
 						if (!TwoSidedMat)
 						{
-							// Important to not use GetBaseMaterial() here because if our parent is the translucent we'll
-							// get the reference UsdPreviewSurface instead, as that is also *its* reference
-							UMaterialInterface* ReferenceMaterial = OneSidedMat->Parent.Get();
-							UMaterialInterface* ReferenceMaterialTwoSided =
-								MeshTranslationImpl::GetTwoSidedVersionOfReferencePreviewSurfaceMaterial(ReferenceMaterial);
-							if (!ensure(ReferenceMaterialTwoSided && ReferenceMaterialTwoSided != ReferenceMaterial))
-							{
-								continue;
-							}
-
-							const FName NewInstanceName = MakeUniqueObjectName(
-								GetTransientPackage(),
-								UMaterialInstance::StaticClass(),
-								*(FPaths::GetBaseFilename(Slot.MaterialSource) + UnrealIdentifiers::TwoSidedMaterialSuffix)
-							);
-
-#if WITH_EDITOR
-							UMaterialInstanceConstant* MIC = Cast<UMaterialInstanceConstant>(OneSidedMat);
-							if (GIsEditor && MIC)
-							{
-								UMaterialInstanceConstant* TwoSidedMIC = NewObject<UMaterialInstanceConstant>(
-									GetTransientPackage(),
-									NewInstanceName,
-									Flags
-								);
-								if (TwoSidedMIC)
-								{
-									UUsdAssetImportData* ImportData = NewObject< UUsdAssetImportData >(
-										TwoSidedMIC,
-										TEXT("USDAssetImportData")
-									);
-									ImportData->PrimPath = Slot.MaterialSource;
-									TwoSidedMIC->AssetImportData = ImportData;
-								}
-
-								TwoSidedMIC->SetParentEditorOnly(ReferenceMaterialTwoSided);
-								TwoSidedMIC->CopyMaterialUniformParametersEditorOnly(OneSidedMat);
-
-								TwoSidedMat = TwoSidedMIC;
-							}
-							else
-#endif // WITH_EDITOR
-							if (UMaterialInstanceDynamic* MID = Cast<UMaterialInstanceDynamic>(OneSidedMat))
-							{
-								UMaterialInstanceDynamic* TwoSidedMID = UMaterialInstanceDynamic::Create(
-									ReferenceMaterialTwoSided,
-									GetTransientPackage(),
-									NewInstanceName
-								);
-								if (!ensure(TwoSidedMID))
-								{
-									continue;
-								}
-
-								TwoSidedMID->CopyParameterOverrides(MID);
-
-								TwoSidedMat = TwoSidedMID;
-							}
-
+							TwoSidedMat = UE::MeshTranslationImplInternal::Private::CreateTwoSidedVersionOfMaterial(OneSidedMat);
 						}
 
 						if (TwoSidedMat)
@@ -200,13 +263,25 @@ TMap<const UsdUtils::FUsdPrimMaterialSlot*, UMaterialInterface*> MeshTranslation
 							TwoSidedMat->SetFlags(RF_Transient);
 							Material = TwoSidedMat;
 						}
+						else
+						{
+							UE_LOG(LogUsd, Warning, TEXT("Failed to generate a two-sided material from the material prim at path '%s'. Falling back to using the single-sided material '%s' instead."),
+								*Slot.MaterialSource,
+								*OneSidedMat->GetPathName()
+							);
+							Material = OneSidedMat;
+							MaterialHash = AssetCache.GetHashForAsset(OneSidedMat);
+						}
 					}
 				}
 
-				// Cache the material to "ping it" as active, but also register two sided materials for the
-				// first time
-				AssetCache.CacheAsset(MaterialHash, Material);
-				InfoCache.LinkAssetToPrim(UE::FSdfPath{*Slot.MaterialSource}, Material);
+				if (Material)
+				{
+					// Cache the material to "ping it" as active, but also register two sided materials for the
+					// first time
+					AssetCache.CacheAsset(MaterialHash, Material);
+					InfoCache.LinkAssetToPrim(UE::FSdfPath{*Slot.MaterialSource}, Material);
+				}
 
 				break;
 			}
diff --git a/Engine/Plugins/Importers/USDImporter/Source/USDSchemas/Private/USDSchemasModule.cpp b/Engine/Plugins/Importers/USDImporter/Source/USDSchemas/Private/USDSchemasModule.cpp
--- a/Engine/Plugins/Importers/USDImporter/Source/USDSchemas/Private/USDSchemasModule.cpp
+++ b/Engine/Plugins/Importers/USDImporter/Source/USDSchemas/Private/USDSchemasModule.cpp
@@ -9,6 +9,7 @@
 #endif // #if WITH_EDITOR
 
 #include "Modules/ModuleManager.h"
+#include "Misc/Paths.h"
 
 #if USE_USD_SDK
 #include "USDGeomCameraTranslator.h"
@@ -20,14 +21,10 @@
 #include "USDMemory.h"
 #include "USDShadeMaterialTranslator.h"
 #include "USDSkelRootTranslator.h"
+#include "Custom/MaterialXUSDShadeMaterialTranslator.h"
 #include "Custom/MDLUSDShadeMaterialTranslator.h"
 #endif // #if USE_USD_SDK
 
-namespace UE::UsdSchemasModule::Private
-{
-	FName MaterialXRenderContext = TEXT("mtlx");
-}
-
 class FUsdSchemasModule : public IUsdSchemasModule
 {
 public:
@@ -46,6 +43,9 @@ public:
 		UsdLuxNonboundableLightBaseTranslatorHandle = GetTranslatorRegistry().Register< FUsdLuxLightTranslator >( TEXT("UsdLuxNonboundableLightBase") );
 
 #if WITH_EDITOR
+		GetRenderContextRegistry().Register(FMaterialXUsdShadeMaterialTranslator::MaterialXRenderContext);
+		MaterialXUsdShadeMaterialTranslatorHandle = GetTranslatorRegistry().Register<FMaterialXUsdShadeMaterialTranslator>(TEXT("UsdShadeMaterial"));
+
 		// Creating skeletal meshes technically works in Standalone mode, but by checking for this we artificially block it
 		// to not confuse users as to why it doesn't work at runtime. Not registering the actual translators lets the inner meshes get parsed as static meshes, at least.
 		if ( GIsEditor )
@@ -55,23 +55,22 @@ public:
 
 			if ( IMDLImporterModule* MDLImporterModule = FModuleManager::Get().LoadModulePtr< IMDLImporterModule >( TEXT("MDLImporter") ) )
 			{
+				GetRenderContextRegistry().Register(FMdlUsdShadeMaterialTranslator::MdlRenderContext);
 				MdlUsdShadeMaterialTranslatorHandle = GetTranslatorRegistry().Register< FMdlUsdShadeMaterialTranslator >( TEXT("UsdShadeMaterial") );
-				GetRenderContextRegistry().Register( FMdlUsdShadeMaterialTranslator::MdlRenderContext );
 			}
 		}
 #endif // WITH_EDITOR
 
-		GetRenderContextRegistry().Register( UE::UsdSchemasModule::Private::MaterialXRenderContext );
 #endif // #if USE_USD_SDK
 	}
 
 	virtual void ShutdownModule() override
 	{
 #if USE_USD_SDK
-		GetRenderContextRegistry().Unregister( UE::UsdSchemasModule::Private::MaterialXRenderContext );
 		GetTranslatorRegistry().Unregister( UsdGeomCameraTranslatorHandle );
 		GetTranslatorRegistry().Unregister( UsdGeomMeshTranslatorHandle );
 		GetTranslatorRegistry().Unregister( UsdGeomPointInstancerTranslatorHandle );
+
 #if WITH_EDITOR
 		if ( GIsEditor )
 		{
@@ -79,9 +78,13 @@ public:
 			GetTranslatorRegistry().Unregister( UsdGroomTranslatorHandle );
 
 			GetTranslatorRegistry().Unregister( MdlUsdShadeMaterialTranslatorHandle );
-			GetRenderContextRegistry().Unregister( FMdlUsdShadeMaterialTranslator::MdlRenderContext );
+			GetRenderContextRegistry().Unregister(FMdlUsdShadeMaterialTranslator::MdlRenderContext);
 		}
+
+		GetTranslatorRegistry().Unregister(MaterialXUsdShadeMaterialTranslatorHandle);
+		GetRenderContextRegistry().Unregister(FMaterialXUsdShadeMaterialTranslator::MaterialXRenderContext);
 #endif // WITH_EDITOR
+
 		GetTranslatorRegistry().Unregister( UsdGeomXformableTranslatorHandle );
 		GetTranslatorRegistry().Unregister( UsdShadeMaterialTranslatorHandle );
 		GetTranslatorRegistry().Unregister( UsdLuxBoundableLightBaseTranslatorHandle );
@@ -115,6 +118,7 @@ protected:
 
 	// Custom schemas
 	FRegisteredSchemaTranslatorHandle MdlUsdShadeMaterialTranslatorHandle;
+	FRegisteredSchemaTranslatorHandle MaterialXUsdShadeMaterialTranslatorHandle;
 };
 
 IMPLEMENT_MODULE_USD( FUsdSchemasModule, USDSchemas );
diff --git a/Engine/Plugins/Importers/USDImporter/Source/USDSchemas/Private/USDShadeMaterialTranslator.cpp b/Engine/Plugins/Importers/USDImporter/Source/USDSchemas/Private/USDShadeMaterialTranslator.cpp
--- a/Engine/Plugins/Importers/USDImporter/Source/USDSchemas/Private/USDShadeMaterialTranslator.cpp
+++ b/Engine/Plugins/Importers/USDImporter/Source/USDSchemas/Private/USDShadeMaterialTranslator.cpp
@@ -2,17 +2,18 @@
 
 #include "USDShadeMaterialTranslator.h"
 
-#include "Engine/Level.h"
 #include "MeshTranslationImpl.h"
 #include "USDAssetCache.h"
 #include "USDAssetImportData.h"
 #include "USDClassesModule.h"
+#include "USDClassesModule.h"
 #include "USDLog.h"
 #include "USDShadeConversion.h"
 #include "USDTypesConversion.h"
 
 #include "UsdWrappers/SdfPath.h"
 
+#include "Engine/Level.h"
 #include "Engine/Texture.h"
 #include "Materials/Material.h"
 #include "Materials/MaterialInstanceConstant.h"
@@ -445,46 +446,122 @@ void FUsdShadeMaterialTranslator::CreateAssets()
 		}
 	}
 
-	// Note that this needs to run even if we do find the material from the asset cache, otherwise we won't
+	PostImportMaterial(MaterialHashString, ConvertedMaterial);
+}
+
+bool FUsdShadeMaterialTranslator::CollapsesChildren( ECollapsingType CollapsingType ) const
+{
+	return false;
+}
+
+bool FUsdShadeMaterialTranslator::CanBeCollapsed( ECollapsingType CollapsingType ) const
+{
+	return false;
+}
+
+void FUsdShadeMaterialTranslator::PostImportMaterial(const FString& MaterialHash, UMaterialInterface* ImportedMaterial)
+{
+	if (!ImportedMaterial || !Context->InfoCache || !Context->AssetCache)
+	{
+		return;
+	}
+
+#if WITH_EDITOR
+	UUsdAssetImportData* ImportData = Cast<UUsdAssetImportData>(ImportedMaterial->AssetImportData);
+	if (!ImportData)
+	{
+		ImportData = NewObject<UUsdAssetImportData>(ImportedMaterial, TEXT("USDAssetImportData"));
+		ImportedMaterial->AssetImportData = ImportData;
+	}
+	ImportData->PrimPath = PrimPath.GetString();
+#endif // WITH_EDITOR
+
+	// Note that this needs to run even if we found this material in the asset cache already, otherwise we won't
 	// re-register the prim asset links when we reload a stage
-	if (ConvertedMaterial)
+	Context->AssetCache->CacheAsset(MaterialHash, ImportedMaterial);
+	Context->InfoCache->LinkAssetToPrim(PrimPath, ImportedMaterial);
+
+	// Also link the textures to the same material prim.
+	// This is important because it lets the stage actor drop its references to old unused textures in the
+	// asset cache if they aren't being used by any other material
+	TSet<UObject*> Dependencies = IUsdClassesModule::GetAssetDependencies(ImportedMaterial);
+	for (UObject* Object : Dependencies)
 	{
-		if (Context->InfoCache)
+		if (UTexture* Texture = Cast<UTexture>(Object))
 		{
-			Context->InfoCache->LinkAssetToPrim(PrimPath, ConvertedMaterial);
+			// We don't use "GetOutermost()" here because it's also possible to be owned by an asset cache that
+			// itself lives in the transient package... bIsOwnedByTransientPackage should be true just for new textures
+			// dumped on the transient package
+			const bool bIsOwnedByTransientPackage = Texture->GetOuter()->GetPackage() == GetTransientPackage();
+			const bool bIsOwnedByCache = Context->AssetCache->IsAssetOwnedByCache(Texture->GetPathName());
+
+			// Texture is already owned by the cache: Just touch it without recomputing its hash as that's expensive
+			if (bIsOwnedByCache)
+			{
+				Context->AssetCache->TouchAsset(Texture);
+			}
+			// Texture is owned by the transient package, but not cached yet: Let's take it
+			else if (bIsOwnedByTransientPackage)
+			{
+				FString FilePath;
+#if WITH_EDITOR
+				FilePath = Texture->AssetImportData ? Texture->AssetImportData->GetFirstFilename() : Texture->GetName();
+#else
+				FilePath = Texture->GetName();
+#endif // WITH_EDITOR
+				const FString TextureHash = UsdUtils::GetTextureHash(
+					FilePath,
+					Texture->SRGB,
+					Texture->CompressionSettings,
+					Texture->GetTextureAddressX(),
+					Texture->GetTextureAddressY()
+				);
+
+				// Some translators like FMaterialXUsdShadeMaterialTranslator will import many materials and textures
+				// at once and preload them all in the cache. In some complex scenarios when some materials are updated
+				// it is possible to arrive at a situation where the we need to reparse the source file again and
+				// regenerate a bunch of materials and textures without access to the asset cache. In those cases we'll
+				// unfortunately recreate identical textures, and if we try to store them in here we'll run into a hash
+				// collision. These should be rare, however, and require repeatedly e.g. updating the Material prims
+				// generated by MaterialX or updating the MaterialX import options
+				UTexture* ExistingTexture = Cast<UTexture>(Context->AssetCache->GetCachedAsset(TextureHash));
+				if (!ExistingTexture)
+				{
+					Texture->SetFlags(RF_Transient);
+					Context->AssetCache->CacheAsset(TextureHash, Texture);
+				}
+			}
 
-			// Also link the textures to the same material prim.
-			// This is important because it lets the stage actor drop its references to old unused textures in the
-			// asset cache if they aren't being used by any other material
-			if (UMaterialInstance* MaterialInstance = Cast<UMaterialInstance>(ConvertedMaterial))
+			if (bIsOwnedByCache || bIsOwnedByTransientPackage)
 			{
-				for (const FTextureParameterValue& TextureValue : MaterialInstance->TextureParameterValues)
+#if WITH_EDITOR
+				UUsdAssetImportData* TextureImportData = Cast<UUsdAssetImportData>(Texture->AssetImportData);
+				if (!TextureImportData)
 				{
-					if (UTexture* Texture = TextureValue.ParameterValue)
-					{
-						Context->InfoCache->LinkAssetToPrim(PrimPath, Texture);
-					}
+					TextureImportData = NewObject<UUsdAssetImportData>(ImportedMaterial, TEXT("USDAssetImportData"));
+					Texture->AssetImportData = TextureImportData;
 				}
+				TextureImportData->PrimPath = PrimPath.GetString();
+#endif // WITH_EDITOR
+
+				Context->InfoCache->LinkAssetToPrim(PrimPath, Texture);
 			}
 		}
-
-		if (Context->AssetCache)
+		else if (UMaterialInterface* ReferenceMaterial = Cast<UMaterialInterface>(Object))
 		{
-			Context->AssetCache->CacheAsset(MaterialHashString, ConvertedMaterial);
+			// Some scenarios can generate reference/instance material pairs, and reference materials are dependencies.
+			// We won't handle these dependencies recursively though, the caller is responsible for calling this for
+			// all individual materials as they need to also provide the hash to use for each
+			UMaterialInstance* Instance = Cast<UMaterialInstance>(ImportedMaterial);
+			ensure(Instance && Instance->Parent.Get() == ReferenceMaterial);
+		}
+		else
+		{
+			ensureMsgf(false, TEXT("Asset type unsupported!"));
 		}
 	}
 }
 
-bool FUsdShadeMaterialTranslator::CollapsesChildren( ECollapsingType CollapsingType ) const
-{
-	return false;
-}
-
-bool FUsdShadeMaterialTranslator::CanBeCollapsed( ECollapsingType CollapsingType ) const
-{
-	return false;
-}
-
 TSet<UE::FSdfPath> FUsdShadeMaterialTranslator::CollectAuxiliaryPrims() const
 {
 	TSet<UE::FSdfPath> Result;
diff --git a/Engine/Plugins/Importers/USDImporter/Source/USDSchemas/Public/USDSchemaTranslator.h b/Engine/Plugins/Importers/USDImporter/Source/USDSchemas/Public/USDSchemaTranslator.h
--- a/Engine/Plugins/Importers/USDImporter/Source/USDSchemas/Public/USDSchemaTranslator.h
+++ b/Engine/Plugins/Importers/USDImporter/Source/USDSchemas/Public/USDSchemaTranslator.h
@@ -26,6 +26,7 @@
 class FRegisteredSchemaTranslator;
 class FUsdSchemaTranslator;
 class FUsdSchemaTranslatorTaskChain;
+class UInterchangeGenericMaterialPipeline;
 class ULevel;
 class UMaterialInterface;
 class USceneComponent;
@@ -160,6 +161,9 @@ struct USDSCHEMAS_API FUsdSchemaTranslationContext : public TSharedFromThis< FUs
 	/** Describes what to add to the root bone animation within generated AnimSequences, if anything */
 	EUsdRootMotionHandling RootMotionHandling = EUsdRootMotionHandling::NoAdditionalRootMotion;
 
+	/** Import options for parsing MaterialX through Interchange */
+	TObjectPtr<UInterchangeGenericMaterialPipeline> MaterialXOptions;
+
 	/** If a generated UStaticMesh has at least this many triangles we will attempt to enable Nanite */
 	int32 NaniteTriangleThreshold;
 
diff --git a/Engine/Plugins/Importers/USDImporter/Source/USDSchemas/Public/USDShadeMaterialTranslator.h b/Engine/Plugins/Importers/USDImporter/Source/USDSchemas/Public/USDShadeMaterialTranslator.h
--- a/Engine/Plugins/Importers/USDImporter/Source/USDSchemas/Public/USDShadeMaterialTranslator.h
+++ b/Engine/Plugins/Importers/USDImporter/Source/USDSchemas/Public/USDShadeMaterialTranslator.h
@@ -21,6 +21,9 @@ public:
 	virtual bool CanBeCollapsed( ECollapsingType CollapsingType ) const override;
 
 	virtual TSet<UE::FSdfPath> CollectAuxiliaryPrims() const override;
+
+protected:
+	virtual void PostImportMaterial(const FString& MaterialHash, UMaterialInterface* ImportedMaterial);
 };
 
 #endif // #if USE_USD_SDK
diff --git a/Engine/Plugins/Importers/USDImporter/Source/USDSchemas/USDSchemas.Build.cs b/Engine/Plugins/Importers/USDImporter/Source/USDSchemas/USDSchemas.Build.cs
--- a/Engine/Plugins/Importers/USDImporter/Source/USDSchemas/USDSchemas.Build.cs
+++ b/Engine/Plugins/Importers/USDImporter/Source/USDSchemas/USDSchemas.Build.cs
@@ -25,6 +25,11 @@ namespace UnrealBuildTool.Rules
 					"CoreUObject",
 					"Engine",
 					"GeometryCache",
+					"InterchangeCore",
+					"InterchangeEngine",
+					"InterchangeFactoryNodes",
+					"InterchangeNodes",
+					"InterchangePipelines",
 					"LiveLinkAnimationCore",
 					"LiveLinkComponents",
 					"LiveLinkInterface",
@@ -65,6 +70,8 @@ namespace UnrealBuildTool.Rules
 						"UnrealEd",
 					}
 				);
+
+				AddEngineThirdPartyPrivateStaticDependencies(Target, "MaterialX");
 			}
 		}
 	}
diff --git a/Engine/Plugins/Importers/USDImporter/Source/USDStage/Private/USDStageActor.cpp b/Engine/Plugins/Importers/USDImporter/Source/USDStage/Private/USDStageActor.cpp
--- a/Engine/Plugins/Importers/USDImporter/Source/USDStage/Private/USDStageActor.cpp
+++ b/Engine/Plugins/Importers/USDImporter/Source/USDStage/Private/USDStageActor.cpp
@@ -57,6 +57,7 @@
 #include "Framework/Notifications/NotificationManager.h"
 #include "HAL/FileManager.h"
 #include "HAL/IConsoleManager.h"
+#include "InterchangeGenericMaterialPipeline.h"
 #include "LevelSequence.h"
 #include "LiveLinkComponentController.h"
 #include "Materials/Material.h"
@@ -118,6 +119,7 @@ struct FUsdStageActorImpl
 		TranslationContext->RenderContext = StageActor->RenderContext;
 		TranslationContext->MaterialPurpose = StageActor->MaterialPurpose;
 		TranslationContext->RootMotionHandling = StageActor->RootMotionHandling;
+		TranslationContext->MaterialXOptions = StageActor->MaterialXOptions;
 		TranslationContext->MaterialToPrimvarToUVIndex = &StageActor->MaterialToPrimvarToUVIndex;
 		TranslationContext->BlendShapesByPath = &StageActor->BlendShapesByPath;
 		TranslationContext->InfoCache = StageActor->InfoCache;
@@ -643,6 +645,8 @@ AUsdStageActor::AUsdStageActor()
 
 	RootComponent = SceneComponent;
 
+	MaterialXOptions = CreateDefaultSubobject<UInterchangeGenericMaterialPipeline>("MaterialXOptions");
+
 	// Note: We can't construct our RootUsdTwin as a default subobject here, it needs to be built on-demand.
 	// Even if we NewObject'd one it will work as a subobject in some contexts (maybe because the CDO will have a dedicated root twin?).
 	// As far as the engine is concerned, our prim twins are static assets like meshes or textures. However, they live on the transient
@@ -1780,6 +1784,15 @@ void AUsdStageActor::SetRootMotionHandling( EUsdRootMotionHandling NewHandlingSt
 	LoadUsdStage();
 }
 
+void AUsdStageActor::SetMaterialXOptions(UInterchangeGenericMaterialPipeline* NewOptions)
+{
+	const bool bMarkDirty = false;
+	Modify(bMarkDirty);
+
+	MaterialXOptions = NewOptions;
+	LoadUsdStage();
+}
+
 float AUsdStageActor::GetTime() const
 {
 	return Time;
@@ -2950,6 +2963,11 @@ void AUsdStageActor::OnObjectPropertyChanged(UObject* ObjectBeingModified, FProp
 		HandlePropertyChangedEvent(PropertyChangedEvent);
 		return;
 	}
+	else if (ObjectBeingModified == MaterialXOptions.Get())
+	{
+		SetMaterialXOptions(MaterialXOptions);
+		return;
+	}
 
 	// Don't modify the stage if we're in PIE
 	if (!HasAuthorityOverStage())
@@ -3177,6 +3195,10 @@ void AUsdStageActor::HandlePropertyChangedEvent(FPropertyChangedEvent& PropertyC
 	{
 		SetRootMotionHandling(RootMotionHandling);
 	}
+	else if (PropertyName == GET_MEMBER_NAME_CHECKED(AUsdStageActor, MaterialXOptions))
+	{
+		SetMaterialXOptions(MaterialXOptions);
+	}
 	else if (PropertyName == GET_MEMBER_NAME_CHECKED(AUsdStageActor, UsdAssetCache))
 	{
 		SetAssetCache(UsdAssetCache);
diff --git a/Engine/Plugins/Importers/USDImporter/Source/USDStage/Public/USDStageActor.h b/Engine/Plugins/Importers/USDImporter/Source/USDStage/Public/USDStageActor.h
--- a/Engine/Plugins/Importers/USDImporter/Source/USDStage/Public/USDStageActor.h
+++ b/Engine/Plugins/Importers/USDImporter/Source/USDStage/Public/USDStageActor.h
@@ -14,9 +14,10 @@
 #include "USDStageActor.generated.h"
 
 class FUsdInfoCache;
-class UUsdAssetCache2;
+class UInterchangeGenericMaterialPipeline;
 class ULevelSequence;
 class UUsdAssetCache;
+class UUsdAssetCache2;
 class UUsdPrimTwin;
 class UUsdTransactor;
 struct FUsdSchemaTranslationContext;
@@ -90,6 +91,10 @@ public:
 	UPROPERTY( EditAnywhere, BlueprintReadOnly, Category = "USD", config )
 	EUsdRootMotionHandling RootMotionHandling = EUsdRootMotionHandling::NoAdditionalRootMotion;
 
+	// Settings to use when parsing MaterialX materials via Interchange
+	UPROPERTY( EditAnywhere, BlueprintReadOnly, Category = "USD", meta = (DisplayName = "MaterialX Options"))
+	TObjectPtr<UInterchangeGenericMaterialPipeline> MaterialXOptions;
+
 public:
 	DECLARE_EVENT_OneParam(AUsdStageActor, FOnActorLoaded, AUsdStageActor*);
 	USDSTAGE_API static FOnActorLoaded OnActorLoaded;
@@ -146,6 +151,9 @@ public:
 	UFUNCTION( BlueprintCallable, Category = "USD", meta = ( CallInEditor = "true" ) )
 	USDSTAGE_API void SetRootMotionHandling( EUsdRootMotionHandling NewHandlingStrategy );
 
+	UFUNCTION(BlueprintCallable, Category = "USD", meta = (CallInEditor = "true"))
+	USDSTAGE_API void SetMaterialXOptions(UInterchangeGenericMaterialPipeline* NewOptions);
+
 	UFUNCTION(BlueprintCallable, Category = "USD", meta = (CallInEditor = "true"))
 	USDSTAGE_API float GetTime() const;
 
diff --git a/Engine/Plugins/Importers/USDImporter/Source/USDStage/USDStage.Build.cs b/Engine/Plugins/Importers/USDImporter/Source/USDStage/USDStage.Build.cs
--- a/Engine/Plugins/Importers/USDImporter/Source/USDStage/USDStage.Build.cs
+++ b/Engine/Plugins/Importers/USDImporter/Source/USDStage/USDStage.Build.cs
@@ -30,6 +30,7 @@ namespace UnrealBuildTool.Rules
 					"GeometryCacheTracks",
 					"GeometryCacheUSD",
 					"HairStrandsCore",
+					"InterchangePipelines",
 					"LevelSequence",
 					"LiveLinkComponents", // For tracking edits to LiveLinkComponentController properties and writing to USD
 					"LiveLinkInterface",
diff --git a/Engine/Plugins/Importers/USDImporter/Source/USDStageEditor/Private/SUSDStage.cpp b/Engine/Plugins/Importers/USDImporter/Source/USDStageEditor/Private/SUSDStage.cpp
--- a/Engine/Plugins/Importers/USDImporter/Source/USDStageEditor/Private/SUSDStage.cpp
+++ b/Engine/Plugins/Importers/USDImporter/Source/USDStageEditor/Private/SUSDStage.cpp
@@ -39,16 +39,18 @@
 #include "Framework/MultiBox/MultiBoxBuilder.h"
 #include "ISceneOutliner.h"
 #include "ISceneOutlinerColumn.h"
+#include "ISinglePropertyView.h"
 #include "LevelEditor.h"
 #include "Modules/ModuleManager.h"
+#include "PropertyEditorModule.h"
 #include "SceneOutlinerModule.h"
 #include "ScopedTransaction.h"
 #include "Styling/AppStyle.h"
 #include "UObject/StrongObjectPtr.h"
 #include "Widgets/Input/SButton.h"
+#include "Widgets/Input/SButton.h"
 #include "Widgets/Input/SComboBox.h"
 #include "Widgets/Input/SComboButton.h"
-#include "Widgets/Input/SButton.h"
 #include "Widgets/Input/SEditableTextBox.h"
 #include "Widgets/Input/SSpinBox.h"
 #include "Widgets/Layout/SSplitter.h"
@@ -1111,6 +1113,11 @@ void SUsdStage::FillOptionsMenu(FMenuBuilder& MenuBuilder)
 			LOCTEXT( "RootMotionHandling_ToolTip", "Choose how to handle root motion for generated AnimSequences and LevelSequences" ),
 			FNewMenuDelegate::CreateSP( this, &SUsdStage::FillRootMotionSubMenu ) );
 
+		MenuBuilder.AddSubMenu(
+			LOCTEXT("MaterialXOptions", "MaterialX options"),
+			LOCTEXT("MaterialXOptions_ToolTip", "Choose the settings to use when parsing MaterialX materials via Interchange"),
+			FNewMenuDelegate::CreateSP(this, &SUsdStage::FillMaterialXOptionsSubMenu));
+
 		MenuBuilder.AddSubMenu(
 			LOCTEXT( "Collapsing", "Collapsing" ),
 			LOCTEXT( "Collapsing_ToolTip", "Whether to try to combine individual assets and components of the same type on a Kind-per-Kind basis, like multiple Mesh prims into a single Static Mesh" ),
@@ -1585,6 +1592,36 @@ void SUsdStage::FillRootMotionSubMenu( FMenuBuilder& MenuBuilder )
 	AddRootMotionEntry( EUsdRootMotionHandling::UseMotionFromSkeleton, LOCTEXT("UseMotionFromSkeletonText", "Use motion from Skeleton"));
 }
 
+void SUsdStage::FillMaterialXOptionsSubMenu(FMenuBuilder& MenuBuilder)
+{
+	if (AUsdStageActor* StageActor = GetStageActorOrCDO())
+	{
+		FPropertyEditorModule& PropertyEditorModule = FModuleManager::GetModuleChecked<FPropertyEditorModule>("PropertyEditor");
+
+		FSinglePropertyParams InitParams;
+		InitParams.NamePlacement = EPropertyNamePlacement::Hidden;
+
+		TSharedPtr<ISinglePropertyView> PropertyView = PropertyEditorModule.CreateSingleProperty(
+			StageActor,
+			GET_MEMBER_NAME_CHECKED(AUsdStageActor, MaterialXOptions),
+			InitParams
+		);
+
+		TSharedRef<SBox> Box = SNew(SBox)
+			.Padding(FMargin(8.0f, 0.0f))
+			.VAlign(VAlign_Center)
+			[
+				PropertyView.ToSharedRef()
+			];
+
+		if (PropertyView)
+		{
+			const bool bNoIndent = true;
+			MenuBuilder.AddWidget(Box, FText::GetEmpty(), bNoIndent, bNoIndent);
+		}
+	}
+}
+
 void SUsdStage::FillCollapsingSubMenu( FMenuBuilder& MenuBuilder )
 {
 	MenuBuilder.AddMenuEntry(
diff --git a/Engine/Plugins/Importers/USDImporter/Source/USDStageEditor/Private/SUSDStage.h b/Engine/Plugins/Importers/USDImporter/Source/USDStageEditor/Private/SUSDStage.h
--- a/Engine/Plugins/Importers/USDImporter/Source/USDStageEditor/Private/SUSDStage.h
+++ b/Engine/Plugins/Importers/USDImporter/Source/USDStageEditor/Private/SUSDStage.h
@@ -91,6 +91,7 @@ protected:
 	void FillRenderContextSubMenu( FMenuBuilder& MenuBuilder );
 	void FillMaterialPurposeSubMenu( FMenuBuilder& MenuBuilder );
 	void FillRootMotionSubMenu( FMenuBuilder& MenuBuilder );
+	void FillMaterialXOptionsSubMenu( FMenuBuilder& MenuBuilder );
 	void FillCollapsingSubMenu( FMenuBuilder& MenuBuilder );
 	void FillInterpolationTypeSubMenu( FMenuBuilder& MenuBuilder );
 	void FillSelectionSubMenu( FMenuBuilder& MenuBuilder );
diff --git a/Engine/Plugins/Importers/USDImporter/Source/USDStageEditorViewModels/Private/USDStageViewModel.cpp b/Engine/Plugins/Importers/USDImporter/Source/USDStageEditorViewModels/Private/USDStageViewModel.cpp
--- a/Engine/Plugins/Importers/USDImporter/Source/USDStageEditorViewModels/Private/USDStageViewModel.cpp
+++ b/Engine/Plugins/Importers/USDImporter/Source/USDStageEditorViewModels/Private/USDStageViewModel.cpp
@@ -124,10 +124,33 @@ void FUsdStageViewModel::ReloadStage()
 		UsdUtils::StartMonitoringErrors();
 		{
 			FScopedUsdAllocs Allocs;
+
+			TSet<FString> NativeLayerExtensions{UnrealUSDWrapper::GetNativeFileFormats()};
+
+			std::set<pxr::SdfLayerHandle> LayersToRefresh;
+
 			const std::vector<pxr::SdfLayerHandle>& HandleVec = UsdStage->GetUsedLayers();
 
+			// We should only reload the native layers, because things like .mtlx files can show up in this list, if
+			// USD plugins allow interpreting them as "layers". We won't show these on our UI though, and not
+			// necessarily these are equipped to be modified/reloaded from USD (like how the .mtlx ones aren't).
+			for (const pxr::SdfLayerHandle& Handle : HandleVec)
+			{
+				std::string LayerPath;
+				std::map<std::string, std::string> Arguments;
+				if (pxr::SdfLayer::SplitIdentifier(Handle->GetIdentifier(), &LayerPath, &Arguments))
+				{
+					FString UELayerPath = UsdToUnreal::ConvertString(LayerPath);
+					FString UELayerExtension = FPaths::GetExtension(UELayerPath);
+					if (NativeLayerExtensions.Contains(UELayerExtension))
+					{
+						LayersToRefresh.insert(Handle);
+					}
+				}
+			}
+
 			const bool bForce = true;
-			pxr::SdfLayer::ReloadLayers( { HandleVec.begin(), HandleVec.end() }, bForce );
+			pxr::SdfLayer::ReloadLayers(LayersToRefresh, bForce);
 
 			// When reloading our UEState layer is closed but there is nothing on the root layer
 			// that would automatically pull the UEState session layer and cause it to be reloaded, so we need to try
@@ -293,6 +316,7 @@ void FUsdStageViewModel::ImportStage( const TCHAR* TargetContentFolder, UUsdStag
 		ImportContext.ImportOptions->ImportTimeCode = StageActor->GetTime();
 		ImportContext.ImportOptions->NaniteTriangleThreshold = StageActor->NaniteTriangleThreshold;
 		ImportContext.ImportOptions->RootMotionHandling = StageActor->RootMotionHandling;
+		ImportContext.ImportOptions->MaterialXOptions = StageActor->MaterialXOptions;
 		ImportContext.ImportOptions->KindsToCollapse = StageActor->KindsToCollapse;
 		ImportContext.ImportOptions->bMergeIdenticalMaterialSlots = StageActor->bMergeIdenticalMaterialSlots;
 
diff --git a/Engine/Plugins/Importers/USDImporter/Source/USDStageImporter/Private/USDStageImportOptions.cpp b/Engine/Plugins/Importers/USDImporter/Source/USDStageImporter/Private/USDStageImportOptions.cpp
--- a/Engine/Plugins/Importers/USDImporter/Source/USDStageImporter/Private/USDStageImportOptions.cpp
+++ b/Engine/Plugins/Importers/USDImporter/Source/USDStageImporter/Private/USDStageImportOptions.cpp
@@ -6,6 +6,7 @@
 #include "USDSchemaTranslator.h"
 
 #include "AnalyticsEventAttribute.h"
+#include "InterchangeGenericMaterialPipeline.h"
 #include "Modules/ModuleManager.h"
 #include "UObject/UnrealType.h"
 
@@ -24,6 +25,7 @@ UUsdStageImportOptions::UUsdStageImportOptions(const FObjectInitializer& ObjectI
 	IUsdSchemasModule& UsdSchemasModule = FModuleManager::Get().LoadModuleChecked< IUsdSchemasModule >( TEXT("USDSchemas") );
 	RenderContextToImport = UsdSchemasModule.GetRenderContextRegistry().GetUnrealRenderContext();
 	MaterialPurpose = *UnrealIdentifiers::MaterialPreviewPurpose;
+	MaterialXOptions = CreateDefaultSubobject<UInterchangeGenericMaterialPipeline>("MaterialXOptions");
 	bOverrideStageOptions = false;
 	StageOptions.MetersPerUnit = 0.01f;
 	StageOptions.UpAxis = EUsdUpAxis::ZAxis;
diff --git a/Engine/Plugins/Importers/USDImporter/Source/USDStageImporter/Private/USDStageImporter.cpp b/Engine/Plugins/Importers/USDImporter/Source/USDStageImporter/Private/USDStageImporter.cpp
--- a/Engine/Plugins/Importers/USDImporter/Source/USDStageImporter/Private/USDStageImporter.cpp
+++ b/Engine/Plugins/Importers/USDImporter/Source/USDStageImporter/Private/USDStageImporter.cpp
@@ -695,7 +695,7 @@ namespace UsdStageImporterImpl
 
 			// We have to rename the existing asset away, because some objects manage subobjects (like
 			// UMaterialInterface and its EditorOnlyData), and try to rename them whenever they're duplicated/renamed.
-			// That can lead to issues if the target names conflict with the names of ExistingAsset's subobjects			 
+			// That can lead to issues if the target names conflict with the names of ExistingAsset's subobjects
 			if (ExistingAsset)
 			{
 				FName UniqueName = MakeUniqueObjectName(
@@ -704,8 +704,8 @@ namespace UsdStageImporterImpl
 					ExistingAsset->GetFName()
 				);
 				ExistingAsset->Rename(
-					*UniqueName.ToString(), 
-					GetTransientPackage(), 
+					*UniqueName.ToString(),
+					GetTransientPackage(),
 					REN_DontCreateRedirectors | REN_NonTransactional | REN_DoNotDirty
 				);
 			}
@@ -1715,6 +1715,7 @@ void UUsdStageImporter::ImportFromFile(FUsdStageImportContext& ImportContext)
 	TranslationContext->RenderContext = ImportContext.ImportOptions->RenderContextToImport;
 	TranslationContext->MaterialPurpose = ImportContext.ImportOptions->MaterialPurpose;
 	TranslationContext->RootMotionHandling = ImportContext.ImportOptions->RootMotionHandling;
+	TranslationContext->MaterialXOptions = ImportContext.ImportOptions->MaterialXOptions;
 	TranslationContext->ParentComponent = ImportContext.SceneActor ? ImportContext.SceneActor->GetRootComponent() : nullptr;
 	TranslationContext->KindsToCollapse = ( EUsdDefaultKind ) ImportContext.ImportOptions->KindsToCollapse;
 	TranslationContext->bMergeIdenticalMaterialSlots = ImportContext.ImportOptions->bMergeIdenticalMaterialSlots;
@@ -1832,6 +1833,7 @@ bool UUsdStageImporter::ReimportSingleAsset(FUsdStageImportContext& ImportContex
 	TranslationContext->RenderContext = ImportContext.ImportOptions->RenderContextToImport;
 	TranslationContext->MaterialPurpose = ImportContext.ImportOptions->MaterialPurpose;
 	TranslationContext->RootMotionHandling = ImportContext.ImportOptions->RootMotionHandling;
+	TranslationContext->MaterialXOptions = ImportContext.ImportOptions->MaterialXOptions;
 	TranslationContext->KindsToCollapse = ( EUsdDefaultKind ) ImportContext.ImportOptions->KindsToCollapse;
 	TranslationContext->bMergeIdenticalMaterialSlots = ImportContext.ImportOptions->bMergeIdenticalMaterialSlots;
 	TranslationContext->bAllowInterpretingLODs = ImportContext.ImportOptions->bInterpretLODs;
diff --git a/Engine/Plugins/Importers/USDImporter/Source/USDStageImporter/Public/USDStageImportOptions.h b/Engine/Plugins/Importers/USDImporter/Source/USDStageImporter/Public/USDStageImportOptions.h
--- a/Engine/Plugins/Importers/USDImporter/Source/USDStageImporter/Public/USDStageImportOptions.h
+++ b/Engine/Plugins/Importers/USDImporter/Source/USDStageImporter/Public/USDStageImportOptions.h
@@ -13,6 +13,7 @@
 
 #include "USDStageImportOptions.generated.h"
 
+class UInterchangeGenericMaterialPipeline;
 struct FAnalyticsEventAttribute;
 
 UENUM(BlueprintType)
@@ -101,6 +102,10 @@ public:
 	UPROPERTY( BlueprintReadWrite, config, EditAnywhere, Category = "USD options", meta = ( EditCondition = bImportSkeletalAnimations ) )
 	EUsdRootMotionHandling RootMotionHandling = EUsdRootMotionHandling::NoAdditionalRootMotion;
 
+	// Settings to use when parsing MaterialX materials via Interchange
+	UPROPERTY(EditAnywhere, BlueprintReadOnly, Category = "DataToImport", AdvancedDisplay, meta = (DisplayName = "MaterialX Options"))
+	TObjectPtr<UInterchangeGenericMaterialPipeline> MaterialXOptions;
+
 	/** Whether to use the specified StageOptions instead of the stage's own settings */
 	UPROPERTY( BlueprintReadWrite, config, EditAnywhere, Category = "USD options" )
 	bool bOverrideStageOptions;
diff --git a/Engine/Plugins/Importers/USDImporter/Source/USDStageImporter/USDStageImporter.Build.cs b/Engine/Plugins/Importers/USDImporter/Source/USDStageImporter/USDStageImporter.Build.cs
--- a/Engine/Plugins/Importers/USDImporter/Source/USDStageImporter/USDStageImporter.Build.cs
+++ b/Engine/Plugins/Importers/USDImporter/Source/USDStageImporter/USDStageImporter.Build.cs
@@ -20,6 +20,7 @@ namespace UnrealBuildTool.Rules
 					"Engine",
 					"GeometryCache",
 					"InputCore",
+					"InterchangePipelines",
 					"JsonUtilities",
 					"LevelSequence",
 					"MainFrame",
diff --git a/Engine/Plugins/Importers/USDImporter/USDImporter.uplugin b/Engine/Plugins/Importers/USDImporter/USDImporter.uplugin
--- a/Engine/Plugins/Importers/USDImporter/USDImporter.uplugin
+++ b/Engine/Plugins/Importers/USDImporter/USDImporter.uplugin
@@ -96,6 +96,10 @@
 			"Name": "GeometryCache",
 			"Enabled": true
 		},
+		{
+			"Name": "Interchange",
+			"Enabled": true
+		},
 		{
 			"Name": "LevelSequenceEditor",
 			"Enabled": true
