USD: Fix 'bCollapseTopLevelPointInstancers' is deprecated warnings due to implicitly generated special member functions.

Horde issue https://horde.devtools.epicgames.com/job/6425b2314e0449e83eeda953?step=7db0&issue=282182

diff --git a/Engine/Plugins/Importers/USDImporter/Source/USDSchemas/Public/USDSchemaTranslator.h b/Engine/Plugins/Importers/USDImporter/Source/USDSchemas/Public/USDSchemaTranslator.h
--- a/Engine/Plugins/Importers/USDImporter/Source/USDSchemas/Public/USDSchemaTranslator.h
+++ b/Engine/Plugins/Importers/USDImporter/Source/USDSchemas/Public/USDSchemaTranslator.h
@@ -129,6 +129,13 @@ protected:
 
 struct USDSCHEMAS_API FUsdSchemaTranslationContext : public TSharedFromThis< FUsdSchemaTranslationContext >
 {
+	// Explicitly declare these defaulted special functions or else the compiler will do it elsewhere and emit
+	// deprecated warnings due to usage of bCollapseTopLevelPointInstancers
+	PRAGMA_DISABLE_DEPRECATION_WARNINGS
+	FUsdSchemaTranslationContext(const FUsdSchemaTranslationContext& Other) = default;
+	FUsdSchemaTranslationContext& operator=(const FUsdSchemaTranslationContext& Other) = default;
+	PRAGMA_ENABLE_DEPRECATION_WARNINGS
+
 	explicit FUsdSchemaTranslationContext( const UE::FUsdStage& InStage, UUsdAssetCache2& InAssetCache );
 
 	/** True if we're a context created by the USDStageImporter to fully import to persistent assets and actors */
