USDImporter: add wrapping for token-based API schema functions on UsdPrim

diff --git a/Engine/Plugins/Importers/USDImporter/Source/UnrealUSDWrapper/Private/UsdWrappers/UsdPrim.cpp b/Engine/Plugins/Importers/USDImporter/Source/UnrealUSDWrapper/Private/UsdWrappers/UsdPrim.cpp
--- a/Engine/Plugins/Importers/USDImporter/Source/UnrealUSDWrapper/Private/UsdWrappers/UsdPrim.cpp
+++ b/Engine/Plugins/Importers/USDImporter/Source/UnrealUSDWrapper/Private/UsdWrappers/UsdPrim.cpp
@@ -292,6 +292,108 @@ namespace UE
 		return HasAPI( SchemaType );
 	}
 
+	bool FUsdPrim::CanApplyAPI( FName SchemaIdentifier, FString* OutWhyNot ) const
+	{
+		bool bResult = false;
+
+#if USE_USD_SDK
+		FScopedUsdAllocs Allocs;
+
+		const pxr::TfToken UsdSchemaIdentifier( TCHAR_TO_ANSI( *SchemaIdentifier.ToString() ) );
+
+		std::string WhyNot;
+		bResult = Impl->PxrUsdPrim.Get().CanApplyAPI(
+			UsdSchemaIdentifier,
+			OutWhyNot != nullptr ? &WhyNot : nullptr );
+
+		if ( !bResult && OutWhyNot != nullptr )
+		{
+			*OutWhyNot = ANSI_TO_TCHAR( WhyNot.c_str() );
+		}
+#endif // #if USE_USD_SDK
+
+		return bResult;
+	}
+
+	bool FUsdPrim::CanApplyAPI( FName SchemaIdentifier, FName InstanceName, FString* OutWhyNot ) const
+	{
+		bool bResult = false;
+
+#if USE_USD_SDK
+		FScopedUsdAllocs Allocs;
+
+		const pxr::TfToken UsdSchemaIdentifier( TCHAR_TO_ANSI( *SchemaIdentifier.ToString() ) );
+		const pxr::TfToken UsdInstanceName( TCHAR_TO_ANSI( *InstanceName.ToString() ) );
+
+		std::string WhyNot;
+		bResult = Impl->PxrUsdPrim.Get().CanApplyAPI(
+			UsdSchemaIdentifier,
+			UsdInstanceName,
+			OutWhyNot != nullptr ? &WhyNot : nullptr);
+
+		if ( !bResult && OutWhyNot != nullptr )
+		{
+			*OutWhyNot = ANSI_TO_TCHAR( WhyNot.c_str() );
+		}
+#endif // #if USE_USD_SDK
+
+		return bResult;
+	}
+
+	bool FUsdPrim::ApplyAPI( FName SchemaIdentifier ) const
+	{
+#if USE_USD_SDK
+		FScopedUsdAllocs Allocs;
+
+		const pxr::TfToken UsdSchemaIdentifier( TCHAR_TO_ANSI( *SchemaIdentifier.ToString() ) );
+
+		return Impl->PxrUsdPrim.Get().ApplyAPI( UsdSchemaIdentifier );
+#else
+		return false;
+#endif // #if USE_USD_SDK
+	}
+
+	bool FUsdPrim::ApplyAPI( FName SchemaIdentifier, FName InstanceName ) const
+	{
+#if USE_USD_SDK
+		FScopedUsdAllocs Allocs;
+
+		const pxr::TfToken UsdSchemaIdentifier( TCHAR_TO_ANSI( *SchemaIdentifier.ToString() ) );
+		const pxr::TfToken UsdInstanceName( TCHAR_TO_ANSI( *InstanceName.ToString() ) );
+
+		return Impl->PxrUsdPrim.Get().ApplyAPI( UsdSchemaIdentifier, UsdInstanceName );
+#else
+		return false;
+#endif // #if USE_USD_SDK
+	}
+
+	bool FUsdPrim::RemoveAPI( FName SchemaIdentifier ) const
+	{
+#if USE_USD_SDK
+		FScopedUsdAllocs Allocs;
+
+		const pxr::TfToken UsdSchemaIdentifier( TCHAR_TO_ANSI( *SchemaIdentifier.ToString() ) );
+
+		return Impl->PxrUsdPrim.Get().RemoveAPI( UsdSchemaIdentifier );
+#else
+		return false;
+#endif // #if USE_USD_SDK
+	}
+
+	bool FUsdPrim::RemoveAPI( FName SchemaIdentifier, FName InstanceName ) const
+	{
+#if USE_USD_SDK
+		FScopedUsdAllocs Allocs;
+
+		const pxr::TfToken UsdSchemaIdentifier( TCHAR_TO_ANSI( *SchemaIdentifier.ToString() ) );
+		const pxr::TfToken UsdInstanceName( TCHAR_TO_ANSI( *InstanceName.ToString() ) );
+
+		return Impl->PxrUsdPrim.Get().RemoveAPI( UsdSchemaIdentifier, UsdInstanceName );
+#else
+		return false;
+#endif // #if USE_USD_SDK
+	}
+
 	const FSdfPath FUsdPrim::GetPrimPath() const
 	{
 #if USE_USD_SDK
diff --git a/Engine/Plugins/Importers/USDImporter/Source/UnrealUSDWrapper/Public/UsdWrappers/UsdPrim.h b/Engine/Plugins/Importers/USDImporter/Source/UnrealUSDWrapper/Public/UsdWrappers/UsdPrim.h
--- a/Engine/Plugins/Importers/USDImporter/Source/UnrealUSDWrapper/Public/UsdWrappers/UsdPrim.h
+++ b/Engine/Plugins/Importers/USDImporter/Source/UnrealUSDWrapper/Public/UsdWrappers/UsdPrim.h
@@ -2,6 +2,7 @@
 
 #pragma once
 
+#include "Containers/UnrealString.h"
 #include "Misc/Optional.h"
 #include "Templates/UniquePtr.h"
 #include "UObject/NameTypes.h"
@@ -93,6 +94,15 @@ namespace UE
 		UE_DEPRECATED( 5.3, "Please use the one FName overload for single apply API schemas, or the two FName overload for multiple apply API schemas." )
 		bool HasAPI( FName SchemaType, TOptional<FName> InstanceName ) const;
 
+		bool CanApplyAPI( FName SchemaIdentifier, FString* OutWhyNot = nullptr ) const;
+		bool CanApplyAPI( FName SchemaIdentifier, FName InstanceName, FString* OutWhyNot = nullptr ) const;
+
+		bool ApplyAPI( FName SchemaIdentifier ) const;
+		bool ApplyAPI( FName SchemaIdentifier, FName InstanceName ) const;
+
+		bool RemoveAPI( FName SchemaIdentifier ) const;
+		bool RemoveAPI( FName SchemaIdentifier, FName InstanceName ) const;
+
 		const FSdfPath GetPrimPath() const;
 		FUsdStage GetStage() const;
 
