USDImporter: add a file format arguments parameter to find/create/export functions of SdfLayer wrapper

diff --git a/Engine/Plugins/Importers/USDImporter/Source/UnrealUSDWrapper/Private/UsdWrappers/SdfLayer.cpp b/Engine/Plugins/Importers/USDImporter/Source/UnrealUSDWrapper/Private/UsdWrappers/SdfLayer.cpp
--- a/Engine/Plugins/Importers/USDImporter/Source/UnrealUSDWrapper/Private/UsdWrappers/SdfLayer.cpp
+++ b/Engine/Plugins/Importers/USDImporter/Source/UnrealUSDWrapper/Private/UsdWrappers/SdfLayer.cpp
@@ -309,25 +309,62 @@ namespace UE
 	}
 
 	template<typename PtrType>
-	FSdfLayer FSdfLayerBase<PtrType>::FindOrOpen( const TCHAR* Identifier )
+	FSdfLayer FSdfLayerBase<PtrType>::FindOrOpen( const TCHAR* Identifier, const TMap< FString, FString >& FileFormatArguments )
 	{
 #if USE_USD_SDK
-		return FSdfLayer( pxr::SdfLayer::FindOrOpen( TCHAR_TO_ANSI( Identifier ) ) );
+		FScopedUsdAllocs Allocs;
+
+		std::map< std::string, std::string > UsdFileFormatArguments;
+		for ( const TPair< FString, FString >& Pair : FileFormatArguments )
+		{
+			UsdFileFormatArguments.insert( { TCHAR_TO_ANSI( *Pair.Key ), TCHAR_TO_ANSI( *Pair.Value ) } );
+		}
+
+		return FSdfLayer( pxr::SdfLayer::FindOrOpen( TCHAR_TO_ANSI( Identifier ), UsdFileFormatArguments ) );
 #else
 		return FSdfLayer();
 #endif // #if USE_USD_SDK
 	}
 
 	template<typename PtrType>
-	FSdfLayer FSdfLayerBase<PtrType>::CreateNew( const TCHAR* Identifier )
+	FSdfLayer FSdfLayerBase<PtrType>::CreateNew( const TCHAR* Identifier, const TMap< FString, FString >& FileFormatArguments )
 	{
 #if USE_USD_SDK
-		return FSdfLayer( pxr::SdfLayer::CreateNew( TCHAR_TO_ANSI( Identifier ) ) );
+		FScopedUsdAllocs Allocs;
+
+		std::map< std::string, std::string > UsdFileFormatArguments;
+		for ( const TPair< FString, FString >& Pair : FileFormatArguments )
+		{
+			UsdFileFormatArguments.insert( { TCHAR_TO_ANSI( *Pair.Key ), TCHAR_TO_ANSI( *Pair.Value ) } );
+		}
+
+		return FSdfLayer( pxr::SdfLayer::CreateNew( TCHAR_TO_ANSI( Identifier ), UsdFileFormatArguments ) );
 #else
 		return FSdfLayer();
 #endif // #if USE_USD_SDK
 	}
 
+	template<typename PtrType>
+	TMap< FString, FString > FSdfLayerBase<PtrType>::GetFileFormatArguments() const
+	{
+		TMap< FString, FString > FileFormatArguments;
+
+#if USE_USD_SDK
+		FScopedUsdAllocs Allocs;
+
+		if ( const PtrType& Ptr = Impl->GetInner() )
+		{
+			const std::map< std::string, std::string >& UsdFileFormatArguments = Ptr->GetFileFormatArguments();
+			for ( std::map< std::string, std::string >::const_iterator Iter = UsdFileFormatArguments.cbegin(); Iter != UsdFileFormatArguments.cend(); ++Iter )
+			{
+				FileFormatArguments.Add( ANSI_TO_TCHAR( Iter->first.c_str() ), ANSI_TO_TCHAR( Iter->second.c_str() ) );
+			}
+		}
+#endif // #if USE_USD_SDK
+
+		return FileFormatArguments;
+	}
+
 	template<typename PtrType>
 	bool FSdfLayerBase<PtrType>::Save( bool bForce ) const
 	{
@@ -785,12 +822,20 @@ namespace UE
 	}
 
 	template<typename PtrType>
-	bool FSdfLayerBase<PtrType>::Export( const TCHAR* Filename ) const
+	bool FSdfLayerBase<PtrType>::Export( const TCHAR* Filename, const FString& Comment, const TMap< FString, FString >& FileFormatArguments ) const
 	{
 #if USE_USD_SDK
 		if ( const PtrType& Ptr = Impl->GetInner() )
 		{
-			return Ptr->Export( TCHAR_TO_ANSI( Filename ) );
+			FScopedUsdAllocs Allocs;
+
+			std::map< std::string, std::string > UsdFileFormatArguments;
+			for ( const TPair< FString, FString >& Pair : FileFormatArguments )
+			{
+				UsdFileFormatArguments.insert( { TCHAR_TO_ANSI( *Pair.Key ), TCHAR_TO_ANSI( *Pair.Value ) } );
+			}
+
+			return Ptr->Export( TCHAR_TO_ANSI( Filename ), TCHAR_TO_ANSI ( *Comment ), UsdFileFormatArguments);
 		}
 #endif // #if USE_USD_SDK
 
diff --git a/Engine/Plugins/Importers/USDImporter/Source/UnrealUSDWrapper/Public/UsdWrappers/SdfLayer.h b/Engine/Plugins/Importers/USDImporter/Source/UnrealUSDWrapper/Public/UsdWrappers/SdfLayer.h
--- a/Engine/Plugins/Importers/USDImporter/Source/UnrealUSDWrapper/Public/UsdWrappers/SdfLayer.h
+++ b/Engine/Plugins/Importers/USDImporter/Source/UnrealUSDWrapper/Public/UsdWrappers/SdfLayer.h
@@ -3,6 +3,8 @@
 #pragma once
 
 #include "CoreMinimal.h"
+#include "Containers/Map.h"
+#include "Containers/UnrealString.h"
 
 #include "UsdWrappers/ForwardDeclarations.h"
 
@@ -91,8 +93,14 @@ namespace UE
 		void TransferContent( const FSdfLayer& SourceLayer );
 
 		static TSet<FSdfLayerWeak> GetLoadedLayers();
-		static FSdfLayer FindOrOpen( const TCHAR* Identifier );
-		static FSdfLayer CreateNew( const TCHAR* Identifier );
+		static FSdfLayer FindOrOpen(
+			const TCHAR* Identifier,
+			const TMap< FString, FString >& FileFormatArguments = {} );
+		static FSdfLayer CreateNew(
+			const TCHAR* Identifier,
+			const TMap< FString, FString >& FileFormatArguments = {} );
+
+		TMap< FString, FString > GetFileFormatArguments() const;
 
 		bool Save( bool bForce = false ) const;
 
@@ -112,7 +120,10 @@ namespace UE
 		bool IsEmpty() const;
 		bool IsAnonymous() const;
 
-		bool Export( const TCHAR* Filename ) const;
+		bool Export(
+			const TCHAR* Filename,
+			const FString& Comment = {},
+			const TMap< FString, FString >& FileFormatArguments = {} ) const;
 
 		void Clear();
 
