USDImporter: add wrapping for UsdReferences to enable directly authoring references on a prim

diff --git a/Engine/Plugins/Importers/USDImporter/Source/UnrealUSDWrapper/Private/UsdWrappers/UsdPrim.cpp b/Engine/Plugins/Importers/USDImporter/Source/UnrealUSDWrapper/Private/UsdWrappers/UsdPrim.cpp
--- a/Engine/Plugins/Importers/USDImporter/Source/UnrealUSDWrapper/Private/UsdWrappers/UsdPrim.cpp
+++ b/Engine/Plugins/Importers/USDImporter/Source/UnrealUSDWrapper/Private/UsdWrappers/UsdPrim.cpp
@@ -6,6 +6,7 @@
 
 #include "UsdWrappers/SdfPath.h"
 #include "UsdWrappers/UsdAttribute.h"
+#include "UsdWrappers/UsdReferences.h"
 #include "UsdWrappers/UsdStage.h"
 
 #if USE_USD_SDK
@@ -551,6 +552,15 @@ namespace UE
 		return Attributes;
 	}
 
+	FUsdReferences FUsdPrim::GetReferences() const
+	{
+#if USE_USD_SDK
+		return FUsdReferences( Impl->PxrUsdPrim.Get() );
+#else
+		return FUsdReferences{};
+#endif // #if USE_USD_SDK
+	}
+
 	FUsdAttribute FUsdPrim::GetAttribute(const TCHAR* AttrName) const
 	{
 #if USE_USD_SDK
diff --git a/Engine/Plugins/Importers/USDImporter/Source/UnrealUSDWrapper/Private/UsdWrappers/UsdReferences.cpp b/Engine/Plugins/Importers/USDImporter/Source/UnrealUSDWrapper/Private/UsdWrappers/UsdReferences.cpp
new file mode 100644
--- /dev/null
+++ b/Engine/Plugins/Importers/USDImporter/Source/UnrealUSDWrapper/Private/UsdWrappers/UsdReferences.cpp
@@ -0,0 +1,321 @@
+// Copyright Epic Games, Inc. All Rights Reserved.
+
+#include "UsdWrappers/UsdReferences.h"
+
+#include "USDMemory.h"
+
+#include "UsdWrappers/UsdPrim.h"
+
+#if USE_USD_SDK
+
+#include "USDIncludesStart.h"
+	#include "pxr/usd/usd/prim.h"
+	#include "pxr/usd/usd/references.h"
+#include "USDIncludesEnd.h"
+
+#endif // #if USE_USD_SDK
+
+namespace UE
+{
+	namespace Internal
+	{
+		class FUsdReferencesImpl
+		{
+		public:
+			FUsdReferencesImpl() = default;
+
+#if USE_USD_SDK
+			explicit FUsdReferencesImpl(const pxr::UsdReferences& InUsdReferences)
+				: PxrUsdPrim(InUsdReferences.GetPrim())
+			{
+			}
+
+			explicit FUsdReferencesImpl(pxr::UsdReferences&& InUsdReferences)
+				: PxrUsdPrim(InUsdReferences.GetPrim())
+			{
+			}
+
+			explicit FUsdReferencesImpl(const pxr::UsdPrim& InUsdPrim)
+				: PxrUsdPrim(InUsdPrim)
+			{
+			}
+
+			// pxr::UsdReferences has no default constructor, so we store the prim itself.
+			TUsdStore<pxr::UsdPrim> PxrUsdPrim;
+#endif // #if USE_USD_SDK
+		};
+	}
+
+	FUsdReferences::FUsdReferences()
+	{
+		FScopedUnrealAllocs UnrealAllocs;
+		Impl = MakeUnique<Internal::FUsdReferencesImpl>();
+	}
+
+	FUsdReferences::FUsdReferences(const FUsdReferences& Other)
+	{
+#if USE_USD_SDK
+		FScopedUnrealAllocs UnrealAllocs;
+		Impl = MakeUnique<Internal::FUsdReferencesImpl>(Other.Impl->PxrUsdPrim.Get());
+#endif // #if USE_USD_SDK
+	}
+
+	FUsdReferences::FUsdReferences(FUsdReferences&& Other) = default;
+
+	FUsdReferences::FUsdReferences(const FUsdPrim& InPrim)
+	{
+#if USE_USD_SDK
+		FScopedUnrealAllocs UnrealAllocs;
+		Impl = MakeUnique<Internal::FUsdReferencesImpl>(InPrim);
+#endif // #if USE_USD_SDK
+	}
+
+	FUsdReferences& FUsdReferences::operator=(const FUsdReferences& Other)
+	{
+#if USE_USD_SDK
+		FScopedUnrealAllocs UnrealAllocs;
+		Impl = MakeUnique<Internal::FUsdReferencesImpl>(Other.Impl->PxrUsdPrim.Get());
+#endif // #if USE_USD_SDK
+
+		return *this;
+	}
+
+	FUsdReferences& FUsdReferences::operator=(FUsdReferences&& Other)
+	{
+		FScopedUnrealAllocs UnrealAllocs;
+		Impl = MoveTemp(Other.Impl);
+
+		return *this;
+	}
+
+	FUsdReferences::~FUsdReferences()
+	{
+		FScopedUnrealAllocs UnrealAllocs;
+		Impl.Reset();
+	}
+
+	FUsdReferences::operator bool() const
+	{
+#if USE_USD_SDK
+		return (bool)Impl->PxrUsdPrim.Get();
+#else
+		return false;
+#endif // #if USE_USD_SDK
+	}
+
+#if USE_USD_SDK
+	FUsdReferences::FUsdReferences(const pxr::UsdReferences& InUsdReferences)
+	{
+		FScopedUnrealAllocs UnrealAllocs;
+		Impl = MakeUnique<Internal::FUsdReferencesImpl>(InUsdReferences);
+	}
+
+	FUsdReferences::FUsdReferences(pxr::UsdReferences&& InUsdReferences)
+	{
+		FScopedUnrealAllocs UnrealAllocs;
+		Impl = MakeUnique<Internal::FUsdReferencesImpl>(MoveTemp(InUsdReferences));
+	}
+
+	FUsdReferences::FUsdReferences(const pxr::UsdPrim& InPrim)
+	{
+		FScopedUnrealAllocs UnrealAllocs;
+		Impl = MakeUnique<Internal::FUsdReferencesImpl>(InPrim);
+	}
+
+#endif // #if USE_USD_SDK
+
+	bool FUsdReferences::AddReference(const FSdfReference& Reference, EUsdListPosition Position)
+	{
+#if USE_USD_SDK
+		static_assert((int32)EUsdListPosition::FrontOfPrependList == (int32)pxr::UsdListPositionFrontOfPrependList, "EUsdListPosition enum doesn't match USD!");
+		static_assert((int32)EUsdListPosition::BackOfPrependList == (int32)pxr::UsdListPositionBackOfPrependList, "EUsdListPosition enum doesn't match USD!");
+		static_assert((int32)EUsdListPosition::FrontOfAppendList == (int32)pxr::UsdListPositionFrontOfAppendList, "EUsdListPosition enum doesn't match USD!");
+		static_assert((int32)EUsdListPosition::BackOfAppendList == (int32)pxr::UsdListPositionBackOfAppendList, "EUsdListPosition enum doesn't match USD!");
+
+		if (Impl->PxrUsdPrim.Get())
+		{
+			FScopedUsdAllocs UsdAllocs;
+
+			const std::string AssetPath(TCHAR_TO_ANSI(*Reference.AssetPath));
+			const pxr::SdfPath UsdPrimPath(Reference.PrimPath);
+			const pxr::SdfLayerOffset UsdLayerOffset(Reference.LayerOffset.Offset, Reference.LayerOffset.Scale);
+
+			pxr::SdfReference UsdReference(AssetPath, UsdPrimPath, UsdLayerOffset);
+
+			return Impl->PxrUsdPrim.Get().GetReferences().AddReference(
+				UsdReference,
+				static_cast<pxr::UsdListPosition>(Position));
+		}
+		else
+		{
+			return false;
+		}
+#else
+		return false;
+#endif // #if USE_USD_SDK
+	}
+
+	bool FUsdReferences::AddReference(
+		const FString& Identifier,
+		const FSdfPath& PrimPath,
+		const FSdfLayerOffset& LayerOffset,
+		EUsdListPosition Position)
+	{
+#if USE_USD_SDK
+		if (Impl->PxrUsdPrim.Get())
+		{
+			FScopedUsdAllocs UsdAllocs;
+
+			const std::string UsdIdentifier(TCHAR_TO_ANSI(*Identifier));
+			const pxr::SdfPath UsdPrimPath(PrimPath);
+			const pxr::SdfLayerOffset UsdLayerOffset(LayerOffset.Offset, LayerOffset.Scale);
+
+			return Impl->PxrUsdPrim.Get().GetReferences().AddReference(
+				UsdIdentifier, UsdPrimPath, UsdLayerOffset, static_cast<pxr::UsdListPosition>(Position));
+		}
+		else
+		{
+			return false;
+		}
+#else
+		return false;
+#endif // #if USE_USD_SDK
+	}
+
+	bool FUsdReferences::AddReference(
+		const FString& Identifier,
+		const FSdfLayerOffset& LayerOffset,
+		EUsdListPosition Position)
+	{
+#if USE_USD_SDK
+		if (Impl->PxrUsdPrim.Get())
+		{
+			FScopedUsdAllocs UsdAllocs;
+
+			const std::string UsdIdentifier(TCHAR_TO_ANSI(*Identifier));
+			const pxr::SdfLayerOffset UsdLayerOffset(LayerOffset.Offset, LayerOffset.Scale);
+
+			return Impl->PxrUsdPrim.Get().GetReferences().AddReference(
+				UsdIdentifier, UsdLayerOffset, static_cast<pxr::UsdListPosition>(Position));
+		}
+		else
+		{
+			return false;
+		}
+#else
+		return false;
+#endif // #if USE_USD_SDK
+	}
+
+	bool FUsdReferences::AddInternalReference(
+		const FSdfPath& PrimPath,
+		const FSdfLayerOffset& LayerOffset,
+		EUsdListPosition Position)
+	{
+#if USE_USD_SDK
+		if (Impl->PxrUsdPrim.Get())
+		{
+			FScopedUsdAllocs UsdAllocs;
+
+			const pxr::SdfPath UsdPrimPath(PrimPath);
+			const pxr::SdfLayerOffset UsdLayerOffset(LayerOffset.Offset, LayerOffset.Scale);
+
+			return Impl->PxrUsdPrim.Get().GetReferences().AddInternalReference(
+				UsdPrimPath, UsdLayerOffset, static_cast<pxr::UsdListPosition>(Position));
+		}
+		else
+		{
+			return false;
+		}
+#else
+		return false;
+#endif // #if USE_USD_SDK
+	}
+
+	bool FUsdReferences::RemoveReference(const FSdfReference& Reference)
+	{
+#if USE_USD_SDK
+		if (Impl->PxrUsdPrim.Get())
+		{
+			FScopedUsdAllocs UsdAllocs;
+
+			const std::string AssetPath(TCHAR_TO_ANSI(*Reference.AssetPath));
+			const pxr::SdfPath UsdPrimPath(Reference.PrimPath);
+			const pxr::SdfLayerOffset UsdLayerOffset(Reference.LayerOffset.Offset, Reference.LayerOffset.Scale);
+
+			pxr::SdfReference UsdReference(AssetPath, UsdPrimPath, UsdLayerOffset);
+
+			return Impl->PxrUsdPrim.Get().GetReferences().RemoveReference(UsdReference);
+		}
+		else
+		{
+			return false;
+		}
+#else
+		return false;
+#endif // #if USE_USD_SDK
+	}
+
+	bool FUsdReferences::ClearReferences()
+	{
+#if USE_USD_SDK
+		if (Impl->PxrUsdPrim.Get())
+		{
+			FScopedUsdAllocs UsdAllocs;
+
+			return Impl->PxrUsdPrim.Get().GetReferences().ClearReferences();
+		}
+		else
+		{
+			return false;
+		}
+#else
+		return false;
+#endif // #if USE_USD_SDK
+	}
+
+	bool FUsdReferences::SetReferences(const TArray<FSdfReference>& Items)
+	{
+#if USE_USD_SDK
+		if (Impl->PxrUsdPrim.Get())
+		{
+			FScopedUsdAllocs UsdAllocs;
+
+			std::vector<pxr::SdfReference> UsdReferences;
+
+			for (const FSdfReference& Item : Items)
+			{
+				const std::string AssetPath(TCHAR_TO_ANSI(*Item.AssetPath));
+				const pxr::SdfPath UsdPrimPath(Item.PrimPath);
+				const pxr::SdfLayerOffset UsdLayerOffset(Item.LayerOffset.Offset, Item.LayerOffset.Scale);
+
+				UsdReferences.emplace_back(AssetPath, UsdPrimPath, UsdLayerOffset);
+			}
+
+			return Impl->PxrUsdPrim.Get().GetReferences().SetReferences(UsdReferences);
+		}
+		else
+		{
+			return false;
+		}
+#else
+		return false;
+#endif // #if USE_USD_SDK
+	}
+
+	FUsdPrim FUsdReferences::GetPrim() const
+	{
+#if USE_USD_SDK
+		if (Impl->PxrUsdPrim.Get())
+		{
+			return FUsdPrim(Impl->PxrUsdPrim.Get());
+		}
+		else
+		{
+			return FUsdPrim{};
+		}
+#else
+		return FUsdPrim{};
+#endif // #if USE_USD_SDK
+	}
+}
diff --git a/Engine/Plugins/Importers/USDImporter/Source/UnrealUSDWrapper/Public/UnrealUSDWrapper.h b/Engine/Plugins/Importers/USDImporter/Source/UnrealUSDWrapper/Public/UnrealUSDWrapper.h
--- a/Engine/Plugins/Importers/USDImporter/Source/UnrealUSDWrapper/Public/UnrealUSDWrapper.h
+++ b/Engine/Plugins/Importers/USDImporter/Source/UnrealUSDWrapper/Public/UnrealUSDWrapper.h
@@ -135,6 +135,35 @@ enum class EUsdRootMotionHandling
 	UseMotionFromSkeleton
 };
 
+/** Corresponds to pxr::UsdListPosition, refer to the USD SDK documentation */
+UENUM()
+enum class EUsdListPosition : uint8
+{
+	// The position at the front of the prepend list.
+	// An item added at this position will, after composition is applied,
+	// be stronger than other items prepended in this layer, and stronger
+	// than items added by weaker layers.
+	FrontOfPrependList,
+
+	// The position at the back of the prepend list.
+	// An item added at this position will, after composition is applied,
+	// be weaker than other items prepended in this layer, but stronger
+	// than items added by weaker layers.
+	BackOfPrependList,
+
+	// The position at the front of the append list.
+	// An item added at this position will, after composition is applied,
+	// be stronger than other items appended in this layer, and stronger
+	// than items added by weaker layers.
+	FrontOfAppendList,
+
+	// The position at the back of the append list.
+	// An item added at this position will, after composition is applied,
+	// be weaker than other items appended in this layer, but stronger
+	// than items added by weaker layers.
+	BackOfAppendList
+};
+
 class IUnrealUSDWrapperModule : public IModuleInterface
 {
 };
diff --git a/Engine/Plugins/Importers/USDImporter/Source/UnrealUSDWrapper/Public/UsdWrappers/SdfLayer.h b/Engine/Plugins/Importers/USDImporter/Source/UnrealUSDWrapper/Public/UsdWrappers/SdfLayer.h
--- a/Engine/Plugins/Importers/USDImporter/Source/UnrealUSDWrapper/Public/UsdWrappers/SdfLayer.h
+++ b/Engine/Plugins/Importers/USDImporter/Source/UnrealUSDWrapper/Public/UsdWrappers/SdfLayer.h
@@ -7,10 +7,10 @@
 #include "Containers/UnrealString.h"
 
 #include "UsdWrappers/ForwardDeclarations.h"
+#include "UsdWrappers/SdfPath.h"
 
 namespace UE
 {
-	class FSdfPath;
 	class FSdfPrimSpec;
 
 	namespace Internal
@@ -31,6 +31,23 @@ namespace UE
 		double Scale = 1.0;
 	};
 
+	struct UNREALUSDWRAPPER_API FSdfReference
+	{
+		FSdfReference(
+				const FString& InAssetPath = {},
+				const FSdfPath& InPrimPath = {},
+				const FSdfLayerOffset& InLayerOffset = {})
+			: AssetPath( InAssetPath )
+			, PrimPath( InPrimPath )
+			, LayerOffset ( InLayerOffset )
+		{
+		}
+
+		FString AssetPath;
+		FSdfPath PrimPath;
+		FSdfLayerOffset LayerOffset;
+	};
+
 	/**
 	 * Minimal pxr::SdfLayer pointer wrapper for Unreal that can be used from no-rtti modules.
 	 * Use the aliases FSdfLayer and FSdfLayerWeak instead (defined on ForwardDeclarations.h)
diff --git a/Engine/Plugins/Importers/USDImporter/Source/UnrealUSDWrapper/Public/UsdWrappers/UsdPrim.h b/Engine/Plugins/Importers/USDImporter/Source/UnrealUSDWrapper/Public/UsdWrappers/UsdPrim.h
--- a/Engine/Plugins/Importers/USDImporter/Source/UnrealUSDWrapper/Public/UsdWrappers/UsdPrim.h
+++ b/Engine/Plugins/Importers/USDImporter/Source/UnrealUSDWrapper/Public/UsdWrappers/UsdPrim.h
@@ -26,6 +26,7 @@ namespace UE
 {
 	class FSdfPath;
 	class FUsdAttribute;
+	class FUsdReferences;
 
 	namespace Internal
 	{
@@ -122,6 +123,8 @@ namespace UE
 		void Load( EUsdLoadPolicy Policy = EUsdLoadPolicy::UsdLoadWithDescendants );
 		void Unload();
 
+		FUsdReferences GetReferences() const;
+
 		bool RemoveProperty( FName PropName ) const;
 
 		FUsdAttribute CreateAttribute( const TCHAR* AttrName, FName TypeName ) const;
diff --git a/Engine/Plugins/Importers/USDImporter/Source/UnrealUSDWrapper/Public/UsdWrappers/UsdReferences.h b/Engine/Plugins/Importers/USDImporter/Source/UnrealUSDWrapper/Public/UsdWrappers/UsdReferences.h
new file mode 100644
--- /dev/null
+++ b/Engine/Plugins/Importers/USDImporter/Source/UnrealUSDWrapper/Public/UsdWrappers/UsdReferences.h
@@ -0,0 +1,96 @@
+// Copyright Epic Games, Inc. All Rights Reserved.
+
+#pragma once
+
+#include "Containers/Array.h"
+#include "Containers/UnrealString.h"
+#include "Templates/UniquePtr.h"
+
+#include "UnrealUSDWrapper.h"
+#include "UsdWrappers/SdfLayer.h"
+
+#if USE_USD_SDK
+
+#include "USDIncludesStart.h"
+	#include "pxr/pxr.h"
+#include "USDIncludesEnd.h"
+
+PXR_NAMESPACE_OPEN_SCOPE
+	class UsdPrim;
+	class UsdReferences;
+PXR_NAMESPACE_CLOSE_SCOPE
+
+#endif // #if USE_USD_SDK
+
+namespace UE
+{
+	class FSdfPath;
+	class FUsdPrim;
+
+	namespace Internal
+	{
+		class FUsdReferencesImpl;
+	}
+
+	/**
+	 * Minimal pxr::UsdReferences wrapper for Unreal that can be used from no-rtti modules.
+	 */
+	class UNREALUSDWRAPPER_API FUsdReferences
+	{
+	public:
+		FUsdReferences();
+		FUsdReferences(const FUsdReferences& Other);
+		FUsdReferences(FUsdReferences&& Other);
+
+		explicit FUsdReferences(const FUsdPrim& InPrim);
+
+		FUsdReferences& operator=(const FUsdReferences& Other);
+		FUsdReferences& operator=(FUsdReferences&& Other);
+		
+		~FUsdReferences();
+
+		explicit operator bool() const;
+
+	public:
+#if USE_USD_SDK
+		// Auto conversion from/to pxr::UsdReferences
+		explicit FUsdReferences(const pxr::UsdReferences& InUsdReferences);
+		explicit FUsdReferences(pxr::UsdReferences&& InUsdReferences);
+
+		explicit FUsdReferences(const pxr::UsdPrim& InPrim);
+#endif // #if USE_USD_SDK
+
+	// Wrapped pxr::UsdReferences functions, refer to the USD SDK documentation
+	public:
+		bool AddReference(
+			const FSdfReference& Reference,
+			EUsdListPosition Position = EUsdListPosition::BackOfPrependList);
+
+		bool AddReference(
+			const FString& Identifier,
+			const FSdfPath& PrimPath,
+			const FSdfLayerOffset& LayerOffset = {},
+			EUsdListPosition Position = EUsdListPosition::BackOfPrependList);
+
+		bool AddReference(
+			const FString& Identifier,
+			const FSdfLayerOffset& LayerOffset = {},
+			EUsdListPosition Position = EUsdListPosition::BackOfPrependList);
+
+		bool AddInternalReference(
+			const FSdfPath& PrimPath,
+			const FSdfLayerOffset& LayerOffset = {},
+			EUsdListPosition Position = EUsdListPosition::BackOfPrependList);
+
+		bool RemoveReference(const FSdfReference& Reference);
+
+		bool ClearReferences();
+
+		bool SetReferences(const TArray<FSdfReference>& Items);
+
+		FUsdPrim GetPrim() const;
+
+	private:
+		TUniquePtr<Internal::FUsdReferencesImpl> Impl;
+	};
+}
