USDImporter: add wrapping for UsdPayloads to enable directly authoring payloads on a prim

The FUsdPrim::HasPayload() function was also deprecated in favor of HasAuthoredPayloads()
to mirror the USD API.

diff --git a/Engine/Plugins/Importers/USDImporter/Source/USDStageEditor/Private/SUSDStageTreeView.cpp b/Engine/Plugins/Importers/USDImporter/Source/USDStageEditor/Private/SUSDStageTreeView.cpp
--- a/Engine/Plugins/Importers/USDImporter/Source/USDStageEditor/Private/SUSDStageTreeView.cpp
+++ b/Engine/Plugins/Importers/USDImporter/Source/USDStageEditor/Private/SUSDStageTreeView.cpp
@@ -1500,7 +1500,7 @@ void SUsdStageTreeView::OnToggleAllPayloads( EPayloadsTrigger PayloadsTrigger )
 			{
 				UE::FUsdPrim& UsdPrim = InSelectedItem->UsdPrim;
 
-				if ( UsdPrim.HasPayload() )
+				if ( UsdPrim.HasAuthoredPayloads() )
 				{
 					bool bPrimIsLoaded = UsdPrim.IsLoaded();
 
diff --git a/Engine/Plugins/Importers/USDImporter/Source/USDStageEditorViewModels/Private/USDPrimViewModel.cpp b/Engine/Plugins/Importers/USDImporter/Source/USDStageEditorViewModels/Private/USDPrimViewModel.cpp
--- a/Engine/Plugins/Importers/USDImporter/Source/USDStageEditorViewModels/Private/USDPrimViewModel.cpp
+++ b/Engine/Plugins/Importers/USDImporter/Source/USDStageEditorViewModels/Private/USDPrimViewModel.cpp
@@ -141,7 +141,7 @@ void FUsdPrimViewModel::RefreshData( bool bRefreshChildren )
 	RowData->bHasCompositionArcs = UsdUtils::HasCompositionArcs( UsdPrim );
 
 	RowData->Type = bIsPseudoRoot ? FText::GetEmpty() : FText::FromName( UsdPrim.GetTypeName() );
-	RowData->bHasPayload = UsdPrim.HasPayload();
+	RowData->bHasPayload = UsdPrim.HasAuthoredPayloads();
 	RowData->bIsLoaded = UsdPrim.IsLoaded();
 
 	bool bOldVisibility = RowData->bIsVisible;
@@ -210,7 +210,7 @@ void FUsdPrimViewModel::ToggleVisibility()
 
 void FUsdPrimViewModel::TogglePayload()
 {
-	if ( UsdPrim && UsdPrim.HasPayload() )
+	if ( UsdPrim && UsdPrim.HasAuthoredPayloads() )
 	{
 		if ( UsdPrim.IsLoaded() )
 		{
diff --git a/Engine/Plugins/Importers/USDImporter/Source/USDUtilities/Private/USDConversionUtils.cpp b/Engine/Plugins/Importers/USDImporter/Source/USDUtilities/Private/USDConversionUtils.cpp
--- a/Engine/Plugins/Importers/USDImporter/Source/USDUtilities/Private/USDConversionUtils.cpp
+++ b/Engine/Plugins/Importers/USDImporter/Source/USDUtilities/Private/USDConversionUtils.cpp
@@ -362,7 +362,7 @@ bool UsdUtils::HasCompositionArcs( const pxr::UsdPrim& Prim )
 		return false;
 	}
 
-	return Prim.HasAuthoredReferences() || Prim.HasPayload() || Prim.HasAuthoredInherits() || Prim.HasAuthoredSpecializes() || Prim.HasVariantSets();
+	return Prim.HasAuthoredReferences() || Prim.HasAuthoredPayloads() || Prim.HasAuthoredInherits() || Prim.HasAuthoredSpecializes() || Prim.HasVariantSets();
 }
 
 bool UsdUtils::HasCompositionArcs( const pxr::SdfPrimSpecHandle& PrimSpec )
diff --git a/Engine/Plugins/Importers/USDImporter/Source/UnrealUSDWrapper/Private/UsdWrappers/UsdPayloads.cpp b/Engine/Plugins/Importers/USDImporter/Source/UnrealUSDWrapper/Private/UsdWrappers/UsdPayloads.cpp
new file mode 100644
--- /dev/null
+++ b/Engine/Plugins/Importers/USDImporter/Source/UnrealUSDWrapper/Private/UsdWrappers/UsdPayloads.cpp
@@ -0,0 +1,281 @@
+// Copyright Epic Games, Inc. All Rights Reserved.
+
+#include "UsdWrappers/UsdPayloads.h"
+
+#include "USDMemory.h"
+
+#include "UsdWrappers/UsdPrim.h"
+
+#if USE_USD_SDK
+
+#include "USDIncludesStart.h"
+	#include "pxr/usd/usd/prim.h"
+	#include "pxr/usd/usd/payloads.h"
+#include "USDIncludesEnd.h"
+
+#endif // #if USE_USD_SDK
+
+namespace UE
+{
+	namespace Internal
+	{
+		class FUsdPayloadsImpl
+		{
+		public:
+			FUsdPayloadsImpl() = default;
+
+#if USE_USD_SDK
+			explicit FUsdPayloadsImpl(const pxr::UsdPayloads& InUsdPayloads)
+				: PxrUsdPrim(InUsdPayloads.GetPrim())
+			{
+			}
+
+			explicit FUsdPayloadsImpl(pxr::UsdPayloads&& InUsdPayloads)
+				: PxrUsdPrim(InUsdPayloads.GetPrim())
+			{
+			}
+
+			explicit FUsdPayloadsImpl(const pxr::UsdPrim& InUsdPrim)
+				: PxrUsdPrim(InUsdPrim)
+			{
+			}
+
+			// pxr::UsdPayloads has no default constructor, so we store the prim itself.
+			TUsdStore<pxr::UsdPrim> PxrUsdPrim;
+#endif // #if USE_USD_SDK
+		};
+	}
+
+	FUsdPayloads::FUsdPayloads()
+	{
+		Impl = MakeUnique<Internal::FUsdPayloadsImpl>();
+	}
+
+	FUsdPayloads::FUsdPayloads(const FUsdPayloads& Other)
+	{
+#if USE_USD_SDK
+		FScopedUnrealAllocs UnrealAllocs;
+		Impl = MakeUnique<Internal::FUsdPayloadsImpl>(Other.Impl->PxrUsdPrim.Get());
+#endif // #if USE_USD_SDK
+	}
+
+	FUsdPayloads::FUsdPayloads(FUsdPayloads&& Other) = default;
+
+	FUsdPayloads::FUsdPayloads(const FUsdPrim& InPrim)
+	{
+#if USE_USD_SDK
+		FScopedUnrealAllocs UnrealAllocs;
+		Impl = MakeUnique<Internal::FUsdPayloadsImpl>(InPrim);
+#endif // #if USE_USD_SDK
+	}
+
+	FUsdPayloads& FUsdPayloads::operator=(const FUsdPayloads& Other)
+	{
+#if USE_USD_SDK
+		FScopedUnrealAllocs UnrealAllocs;
+		Impl = MakeUnique<Internal::FUsdPayloadsImpl>(Other.Impl->PxrUsdPrim.Get());
+#endif // #if USE_USD_SDK
+
+		return *this;
+	}
+
+	FUsdPayloads& FUsdPayloads::operator=(FUsdPayloads&& Other)
+	{
+		Impl = MoveTemp(Other.Impl);
+
+		return *this;
+	}
+
+	FUsdPayloads::~FUsdPayloads()
+	{
+		Impl.Reset();
+	}
+
+	FUsdPayloads::operator bool() const
+	{
+#if USE_USD_SDK
+		return (bool)Impl->PxrUsdPrim.Get();
+#else
+		return false;
+#endif // #if USE_USD_SDK
+	}
+
+#if USE_USD_SDK
+	FUsdPayloads::FUsdPayloads(const pxr::UsdPayloads& InUsdPayloads)
+	{
+		FScopedUnrealAllocs UnrealAllocs;
+		Impl = MakeUnique<Internal::FUsdPayloadsImpl>(InUsdPayloads);
+	}
+
+	FUsdPayloads::FUsdPayloads(pxr::UsdPayloads&& InUsdPayloads)
+	{
+		FScopedUnrealAllocs UnrealAllocs;
+		Impl = MakeUnique<Internal::FUsdPayloadsImpl>(MoveTemp(InUsdPayloads));
+	}
+
+	FUsdPayloads::FUsdPayloads(const pxr::UsdPrim& InPrim)
+	{
+		FScopedUnrealAllocs UnrealAllocs;
+		Impl = MakeUnique<Internal::FUsdPayloadsImpl>(InPrim);
+	}
+#endif // #if USE_USD_SDK
+
+	bool FUsdPayloads::AddPayload(const FSdfPayload& Payload, EUsdListPosition Position)
+	{
+#if USE_USD_SDK
+		if (Impl->PxrUsdPrim.Get())
+		{
+			FScopedUsdAllocs UsdAllocs;
+
+			const std::string AssetPath(TCHAR_TO_ANSI(*Payload.AssetPath));
+			const pxr::SdfPath UsdPrimPath(Payload.PrimPath);
+			const pxr::SdfLayerOffset UsdLayerOffset(Payload.LayerOffset.Offset, Payload.LayerOffset.Scale);
+
+			const pxr::SdfPayload UsdPayload(AssetPath, UsdPrimPath, UsdLayerOffset);
+
+			return Impl->PxrUsdPrim.Get().GetPayloads().AddPayload(
+				UsdPayload,
+				static_cast<pxr::UsdListPosition>(Position));
+		}
+#endif // #if USE_USD_SDK
+
+		return false;
+	}
+
+	bool FUsdPayloads::AddPayload(
+		const FString& Identifier,
+		const FSdfPath& PrimPath,
+		const FSdfLayerOffset& LayerOffset,
+		EUsdListPosition Position)
+	{
+#if USE_USD_SDK
+		if (Impl->PxrUsdPrim.Get())
+		{
+			FScopedUsdAllocs UsdAllocs;
+
+			const std::string UsdIdentifier(TCHAR_TO_ANSI(*Identifier));
+			const pxr::SdfPath UsdPrimPath(PrimPath);
+			const pxr::SdfLayerOffset UsdLayerOffset(LayerOffset.Offset, LayerOffset.Scale);
+
+			return Impl->PxrUsdPrim.Get().GetPayloads().AddPayload(
+				UsdIdentifier, UsdPrimPath, UsdLayerOffset, static_cast<pxr::UsdListPosition>(Position));
+		}
+#endif // #if USE_USD_SDK
+
+		return false;
+	}
+
+	bool FUsdPayloads::AddPayload(
+		const FString& Identifier,
+		const FSdfLayerOffset& LayerOffset,
+		EUsdListPosition Position)
+	{
+#if USE_USD_SDK
+		if (Impl->PxrUsdPrim.Get())
+		{
+			FScopedUsdAllocs UsdAllocs;
+
+			const std::string UsdIdentifier(TCHAR_TO_ANSI(*Identifier));
+			const pxr::SdfLayerOffset UsdLayerOffset(LayerOffset.Offset, LayerOffset.Scale);
+
+			return Impl->PxrUsdPrim.Get().GetPayloads().AddPayload(
+				UsdIdentifier, UsdLayerOffset, static_cast<pxr::UsdListPosition>(Position));
+		}
+#endif // #if USE_USD_SDK
+
+		return false;
+	}
+
+	bool FUsdPayloads::AddInternalPayload(
+		const FSdfPath& PrimPath,
+		const FSdfLayerOffset& LayerOffset,
+		EUsdListPosition Position)
+	{
+#if USE_USD_SDK
+		if (Impl->PxrUsdPrim.Get())
+		{
+			FScopedUsdAllocs UsdAllocs;
+
+			const pxr::SdfPath UsdPrimPath(PrimPath);
+			const pxr::SdfLayerOffset UsdLayerOffset(LayerOffset.Offset, LayerOffset.Scale);
+
+			return Impl->PxrUsdPrim.Get().GetPayloads().AddInternalPayload(
+				UsdPrimPath, UsdLayerOffset, static_cast<pxr::UsdListPosition>(Position));
+		}
+#endif // #if USE_USD_SDK
+
+		return false;
+	}
+
+	bool FUsdPayloads::RemovePayload(const FSdfPayload& Payload)
+	{
+#if USE_USD_SDK
+		if (Impl->PxrUsdPrim.Get())
+		{
+			FScopedUsdAllocs UsdAllocs;
+
+			const std::string AssetPath(TCHAR_TO_ANSI(*Payload.AssetPath));
+			const pxr::SdfPath UsdPrimPath(Payload.PrimPath);
+			const pxr::SdfLayerOffset UsdLayerOffset(Payload.LayerOffset.Offset, Payload.LayerOffset.Scale);
+
+			const pxr::SdfPayload UsdPayload(AssetPath, UsdPrimPath, UsdLayerOffset);
+
+			return Impl->PxrUsdPrim.Get().GetPayloads().RemovePayload(UsdPayload);
+		}
+#endif // #if USE_USD_SDK
+
+		return false;
+	}
+
+	bool FUsdPayloads::ClearPayloads()
+	{
+#if USE_USD_SDK
+		if (Impl->PxrUsdPrim.Get())
+		{
+			FScopedUsdAllocs UsdAllocs;
+
+			return Impl->PxrUsdPrim.Get().GetPayloads().ClearPayloads();
+		}
+#endif // #if USE_USD_SDK
+
+		return false;
+	}
+
+	bool FUsdPayloads::SetPayloads(const TArray<FSdfPayload>& Items)
+	{
+#if USE_USD_SDK
+		if (Impl->PxrUsdPrim.Get())
+		{
+			FScopedUsdAllocs UsdAllocs;
+
+			std::vector<pxr::SdfPayload> UsdPayloads;
+			UsdPayloads.reserve(Items.Num());
+
+			for (const FSdfPayload& Item : Items)
+			{
+				const std::string AssetPath(TCHAR_TO_ANSI(*Item.AssetPath));
+				const pxr::SdfPath UsdPrimPath(Item.PrimPath);
+				const pxr::SdfLayerOffset UsdLayerOffset(Item.LayerOffset.Offset, Item.LayerOffset.Scale);
+
+				UsdPayloads.emplace_back(AssetPath, UsdPrimPath, UsdLayerOffset);
+			}
+
+			return Impl->PxrUsdPrim.Get().GetPayloads().SetPayloads(UsdPayloads);
+		}
+#endif // #if USE_USD_SDK
+
+		return false;
+	}
+
+	FUsdPrim FUsdPayloads::GetPrim() const
+	{
+#if USE_USD_SDK
+		if (Impl->PxrUsdPrim.Get())
+		{
+			return FUsdPrim(Impl->PxrUsdPrim.Get());
+		}
+#endif // #if USE_USD_SDK
+
+		return FUsdPrim{};
+	}
+}
diff --git a/Engine/Plugins/Importers/USDImporter/Source/UnrealUSDWrapper/Private/UsdWrappers/UsdPrim.cpp b/Engine/Plugins/Importers/USDImporter/Source/UnrealUSDWrapper/Private/UsdWrappers/UsdPrim.cpp
--- a/Engine/Plugins/Importers/USDImporter/Source/UnrealUSDWrapper/Private/UsdWrappers/UsdPrim.cpp
+++ b/Engine/Plugins/Importers/USDImporter/Source/UnrealUSDWrapper/Private/UsdWrappers/UsdPrim.cpp
@@ -6,6 +6,7 @@
 
 #include "UsdWrappers/SdfPath.h"
 #include "UsdWrappers/UsdAttribute.h"
+#include "UsdWrappers/UsdPayloads.h"
 #include "UsdWrappers/UsdReferences.h"
 #include "UsdWrappers/UsdStage.h"
 
@@ -492,24 +493,30 @@ namespace UE
 #endif // #if USE_USD_SDK
 	}
 
-	bool FUsdPrim::HasAuthoredReferences() const
+	FUsdPayloads FUsdPrim::GetPayloads() const
 	{
 #if USE_USD_SDK
-		return Impl->PxrUsdPrim.Get().HasAuthoredReferences();
+		return FUsdPayloads(Impl->PxrUsdPrim.Get());
 #else
-		return false;
+		return FUsdPayloads{};
 #endif // #if USE_USD_SDK
 	}
 
-	bool FUsdPrim::HasPayload() const
+	bool FUsdPrim::HasAuthoredPayloads() const
 	{
 #if USE_USD_SDK
-		return Impl->PxrUsdPrim.Get().HasPayload();
+		return Impl->PxrUsdPrim.Get().HasAuthoredPayloads();
 #else
 		return false;
 #endif // #if USE_USD_SDK
 	}
 
+	// Deprecated in 5.3
+	bool FUsdPrim::HasPayload() const
+	{
+		return HasAuthoredPayloads();
+	}
+
 	bool FUsdPrim::IsLoaded() const
 	{
 #if USE_USD_SDK
@@ -561,6 +568,15 @@ namespace UE
 #endif // #if USE_USD_SDK
 	}
 
+	bool FUsdPrim::HasAuthoredReferences() const
+	{
+#if USE_USD_SDK
+		return Impl->PxrUsdPrim.Get().HasAuthoredReferences();
+#else
+		return false;
+#endif // #if USE_USD_SDK
+	}
+
 	FUsdAttribute FUsdPrim::GetAttribute(const TCHAR* AttrName) const
 	{
 #if USE_USD_SDK
diff --git a/Engine/Plugins/Importers/USDImporter/Source/UnrealUSDWrapper/Private/UsdWrappers/UsdReferences.cpp b/Engine/Plugins/Importers/USDImporter/Source/UnrealUSDWrapper/Private/UsdWrappers/UsdReferences.cpp
--- a/Engine/Plugins/Importers/USDImporter/Source/UnrealUSDWrapper/Private/UsdWrappers/UsdReferences.cpp
+++ b/Engine/Plugins/Importers/USDImporter/Source/UnrealUSDWrapper/Private/UsdWrappers/UsdReferences.cpp
@@ -48,7 +48,6 @@ namespace UE
 
 	FUsdReferences::FUsdReferences()
 	{
-		FScopedUnrealAllocs UnrealAllocs;
 		Impl = MakeUnique<Internal::FUsdReferencesImpl>();
 	}
 
@@ -82,7 +81,6 @@ namespace UE
 
 	FUsdReferences& FUsdReferences::operator=(FUsdReferences&& Other)
 	{
-		FScopedUnrealAllocs UnrealAllocs;
 		Impl = MoveTemp(Other.Impl);
 
 		return *this;
@@ -90,7 +88,6 @@ namespace UE
 
 	FUsdReferences::~FUsdReferences()
 	{
-		FScopedUnrealAllocs UnrealAllocs;
 		Impl.Reset();
 	}
 
@@ -121,7 +118,6 @@ namespace UE
 		FScopedUnrealAllocs UnrealAllocs;
 		Impl = MakeUnique<Internal::FUsdReferencesImpl>(InPrim);
 	}
-
 #endif // #if USE_USD_SDK
 
 	bool FUsdReferences::AddReference(const FSdfReference& Reference, EUsdListPosition Position)
@@ -140,19 +136,15 @@ namespace UE
 			const pxr::SdfPath UsdPrimPath(Reference.PrimPath);
 			const pxr::SdfLayerOffset UsdLayerOffset(Reference.LayerOffset.Offset, Reference.LayerOffset.Scale);
 
-			pxr::SdfReference UsdReference(AssetPath, UsdPrimPath, UsdLayerOffset);
+			const pxr::SdfReference UsdReference(AssetPath, UsdPrimPath, UsdLayerOffset);
 
 			return Impl->PxrUsdPrim.Get().GetReferences().AddReference(
 				UsdReference,
 				static_cast<pxr::UsdListPosition>(Position));
 		}
-		else
-		{
-			return false;
-		}
-#else
-		return false;
 #endif // #if USE_USD_SDK
+
+		return false;
 	}
 
 	bool FUsdReferences::AddReference(
@@ -173,13 +165,9 @@ namespace UE
 			return Impl->PxrUsdPrim.Get().GetReferences().AddReference(
 				UsdIdentifier, UsdPrimPath, UsdLayerOffset, static_cast<pxr::UsdListPosition>(Position));
 		}
-		else
-		{
-			return false;
-		}
-#else
-		return false;
 #endif // #if USE_USD_SDK
+
+		return false;
 	}
 
 	bool FUsdReferences::AddReference(
@@ -198,13 +186,9 @@ namespace UE
 			return Impl->PxrUsdPrim.Get().GetReferences().AddReference(
 				UsdIdentifier, UsdLayerOffset, static_cast<pxr::UsdListPosition>(Position));
 		}
-		else
-		{
-			return false;
-		}
-#else
-		return false;
 #endif // #if USE_USD_SDK
+
+		return false;
 	}
 
 	bool FUsdReferences::AddInternalReference(
@@ -223,13 +207,9 @@ namespace UE
 			return Impl->PxrUsdPrim.Get().GetReferences().AddInternalReference(
 				UsdPrimPath, UsdLayerOffset, static_cast<pxr::UsdListPosition>(Position));
 		}
-		else
-		{
-			return false;
-		}
-#else
-		return false;
 #endif // #if USE_USD_SDK
+
+		return false;
 	}
 
 	bool FUsdReferences::RemoveReference(const FSdfReference& Reference)
@@ -243,17 +223,13 @@ namespace UE
 			const pxr::SdfPath UsdPrimPath(Reference.PrimPath);
 			const pxr::SdfLayerOffset UsdLayerOffset(Reference.LayerOffset.Offset, Reference.LayerOffset.Scale);
 
-			pxr::SdfReference UsdReference(AssetPath, UsdPrimPath, UsdLayerOffset);
+			const pxr::SdfReference UsdReference(AssetPath, UsdPrimPath, UsdLayerOffset);
 
 			return Impl->PxrUsdPrim.Get().GetReferences().RemoveReference(UsdReference);
 		}
-		else
-		{
-			return false;
-		}
-#else
-		return false;
 #endif // #if USE_USD_SDK
+
+		return false;
 	}
 
 	bool FUsdReferences::ClearReferences()
@@ -265,13 +241,9 @@ namespace UE
 
 			return Impl->PxrUsdPrim.Get().GetReferences().ClearReferences();
 		}
-		else
-		{
-			return false;
-		}
-#else
-		return false;
 #endif // #if USE_USD_SDK
+
+		return false;
 	}
 
 	bool FUsdReferences::SetReferences(const TArray<FSdfReference>& Items)
@@ -282,6 +254,7 @@ namespace UE
 			FScopedUsdAllocs UsdAllocs;
 
 			std::vector<pxr::SdfReference> UsdReferences;
+			UsdReferences.reserve(Items.Num());
 
 			for (const FSdfReference& Item : Items)
 			{
@@ -294,13 +267,9 @@ namespace UE
 
 			return Impl->PxrUsdPrim.Get().GetReferences().SetReferences(UsdReferences);
 		}
-		else
-		{
-			return false;
-		}
-#else
-		return false;
 #endif // #if USE_USD_SDK
+
+		return false;
 	}
 
 	FUsdPrim FUsdReferences::GetPrim() const
@@ -310,12 +279,8 @@ namespace UE
 		{
 			return FUsdPrim(Impl->PxrUsdPrim.Get());
 		}
-		else
-		{
-			return FUsdPrim{};
-		}
-#else
-		return FUsdPrim{};
 #endif // #if USE_USD_SDK
+
+		return FUsdPrim{};
 	}
 }
diff --git a/Engine/Plugins/Importers/USDImporter/Source/UnrealUSDWrapper/Public/UsdWrappers/SdfLayer.h b/Engine/Plugins/Importers/USDImporter/Source/UnrealUSDWrapper/Public/UsdWrappers/SdfLayer.h
--- a/Engine/Plugins/Importers/USDImporter/Source/UnrealUSDWrapper/Public/UsdWrappers/SdfLayer.h
+++ b/Engine/Plugins/Importers/USDImporter/Source/UnrealUSDWrapper/Public/UsdWrappers/SdfLayer.h
@@ -31,6 +31,30 @@ namespace UE
 		double Scale = 1.0;
 	};
 
+	struct UNREALUSDWRAPPER_API FSdfPayload
+	{
+		FSdfPayload(
+			const FString& InAssetPath = {},
+			const FSdfPath& InPrimPath = {},
+			const FSdfLayerOffset& InLayerOffset = {})
+			: AssetPath(InAssetPath)
+			, PrimPath(InPrimPath)
+			, LayerOffset(InLayerOffset)
+		{
+		}
+
+		FString AssetPath;
+		FSdfPath PrimPath;
+		FSdfLayerOffset LayerOffset;
+	};
+
+	// In the USD API, pxr::SdfPayload and pxr::SdfReference are very similar
+	// except that the latter includes support for a custom data dictionary.
+	// The UE wrapping of SdfReference does not currently include support for
+	// that custom data, leaving the FSdfReference wrapper identical to the
+	// FSdfPayload wrapper. We keep them as distinct types though to better
+	// mimic the USD API and in case support for custom data dictionaries on
+	// references is added to the FSdfReference wrapper in the future.
 	struct UNREALUSDWRAPPER_API FSdfReference
 	{
 		FSdfReference(
diff --git a/Engine/Plugins/Importers/USDImporter/Source/UnrealUSDWrapper/Public/UsdWrappers/UsdPayloads.h b/Engine/Plugins/Importers/USDImporter/Source/UnrealUSDWrapper/Public/UsdWrappers/UsdPayloads.h
new file mode 100644
--- /dev/null
+++ b/Engine/Plugins/Importers/USDImporter/Source/UnrealUSDWrapper/Public/UsdWrappers/UsdPayloads.h
@@ -0,0 +1,96 @@
+// Copyright Epic Games, Inc. All Rights Reserved.
+
+#pragma once
+
+#include "Containers/Array.h"
+#include "Containers/UnrealString.h"
+#include "Templates/UniquePtr.h"
+
+#include "UnrealUSDWrapper.h"
+#include "UsdWrappers/SdfLayer.h"
+
+#if USE_USD_SDK
+
+#include "USDIncludesStart.h"
+	#include "pxr/pxr.h"
+#include "USDIncludesEnd.h"
+
+PXR_NAMESPACE_OPEN_SCOPE
+	class UsdPayloads;
+	class UsdPrim;
+PXR_NAMESPACE_CLOSE_SCOPE
+
+#endif // #if USE_USD_SDK
+
+namespace UE
+{
+	class FSdfPath;
+	class FUsdPrim;
+
+	namespace Internal
+	{
+		class FUsdPayloadsImpl;
+	}
+
+	/**
+	 * Minimal pxr::UsdPayloads wrapper for Unreal that can be used from no-rtti modules.
+	 */
+	class UNREALUSDWRAPPER_API FUsdPayloads
+	{
+	public:
+		FUsdPayloads();
+		FUsdPayloads(const FUsdPayloads& Other);
+		FUsdPayloads(FUsdPayloads&& Other);
+
+		explicit FUsdPayloads(const FUsdPrim& InPrim);
+
+		FUsdPayloads& operator=(const FUsdPayloads& Other);
+		FUsdPayloads& operator=(FUsdPayloads&& Other);
+		
+		~FUsdPayloads();
+
+		explicit operator bool() const;
+
+	public:
+#if USE_USD_SDK
+		// Auto conversion from/to pxr::UsdPayloads
+		explicit FUsdPayloads(const pxr::UsdPayloads& InUsdPayloads);
+		explicit FUsdPayloads(pxr::UsdPayloads&& InUsdPayloads);
+
+		explicit FUsdPayloads(const pxr::UsdPrim& InPrim);
+#endif // #if USE_USD_SDK
+
+	// Wrapped pxr::UsdPayloads functions, refer to the USD SDK documentation
+	public:
+		bool AddPayload(
+			const FSdfPayload& Payload,
+			EUsdListPosition Position = EUsdListPosition::BackOfPrependList);
+
+		bool AddPayload(
+			const FString& Identifier,
+			const FSdfPath& PrimPath,
+			const FSdfLayerOffset& LayerOffset = {},
+			EUsdListPosition Position = EUsdListPosition::BackOfPrependList);
+
+		bool AddPayload(
+			const FString& Identifier,
+			const FSdfLayerOffset& LayerOffset = {},
+			EUsdListPosition Position = EUsdListPosition::BackOfPrependList);
+
+		bool AddInternalPayload(
+			const FSdfPath& PrimPath,
+			const FSdfLayerOffset& LayerOffset = {},
+			EUsdListPosition Position = EUsdListPosition::BackOfPrependList);
+
+		bool RemovePayload(const FSdfPayload& Payload);
+
+		bool ClearPayloads();
+
+		bool SetPayloads(const TArray<FSdfPayload>& Items);
+
+		FUsdPrim GetPrim() const;
+
+	private:
+		TUniquePtr<Internal::FUsdPayloadsImpl> Impl;
+	};
+}
diff --git a/Engine/Plugins/Importers/USDImporter/Source/UnrealUSDWrapper/Public/UsdWrappers/UsdPrim.h b/Engine/Plugins/Importers/USDImporter/Source/UnrealUSDWrapper/Public/UsdWrappers/UsdPrim.h
--- a/Engine/Plugins/Importers/USDImporter/Source/UnrealUSDWrapper/Public/UsdWrappers/UsdPrim.h
+++ b/Engine/Plugins/Importers/USDImporter/Source/UnrealUSDWrapper/Public/UsdWrappers/UsdPrim.h
@@ -26,6 +26,7 @@ namespace UE
 {
 	class FSdfPath;
 	class FUsdAttribute;
+	class FUsdPayloads;
 	class FUsdReferences;
 
 	namespace Internal
@@ -116,14 +117,17 @@ namespace UE
 		TArray< FUsdPrim > GetFilteredChildren( bool bTraverseInstanceProxies ) const;
 
 		bool HasVariantSets() const;
-		bool HasAuthoredReferences() const;
 
+		FUsdPayloads GetPayloads() const;
+		bool HasAuthoredPayloads() const;
+		UE_DEPRECATED(5.3, "Please use HasAuthoredPayloads() instead.")
 		bool HasPayload() const;
 		bool IsLoaded() const;
 		void Load( EUsdLoadPolicy Policy = EUsdLoadPolicy::UsdLoadWithDescendants );
 		void Unload();
 
 		FUsdReferences GetReferences() const;
+		bool HasAuthoredReferences() const;
 
 		bool RemoveProperty( FName PropName ) const;
 
