USD: Add support for materials that use multiple UV sets.

diff --git a/Engine/Plugins/Importers/USDImporter/Source/GeometryCacheUSD/Private/GeometryCacheTrackUSD.cpp b/Engine/Plugins/Importers/USDImporter/Source/GeometryCacheUSD/Private/GeometryCacheTrackUSD.cpp
--- a/Engine/Plugins/Importers/USDImporter/Source/GeometryCacheUSD/Private/GeometryCacheTrackUSD.cpp
+++ b/Engine/Plugins/Importers/USDImporter/Source/GeometryCacheUSD/Private/GeometryCacheTrackUSD.cpp
@@ -256,8 +256,6 @@ void UGeometryCacheTrackUsd::UnloadUsdStage()
 void UGeometryCacheTrackUsd::Initialize(
 	const UE::FUsdStage& InStage,
 	const FString& InPrimPath,
-	const FName& InRenderContext,
-	const TMap< FString, TMap< FString, int32 > >& InMaterialToPrimvarToUVIndex,
 	int32 InStartFrameIndex,
 	int32 InEndFrameIndex,
 	FReadUsdMeshFunction InReadFunc
@@ -268,8 +266,6 @@ void UGeometryCacheTrackUsd::Initialize(
 	StageRootLayerPath = CurrentStagePinned ? CurrentStagePinned.GetRootLayer().GetRealPath() : FString();
 
 	PrimPath = InPrimPath;
-	RenderContext = InRenderContext;
-	MaterialToPrimvarToUVIndex = InMaterialToPrimvarToUVIndex;
 	StartFrameIndex = InStartFrameIndex;
 	EndFrameIndex = InEndFrameIndex;
 	FramesPerSecond = CurrentStagePinned.GetTimeCodesPerSecond();
@@ -291,6 +287,19 @@ void UGeometryCacheTrackUsd::Initialize(
 	UsdStream->Prefetch(StartFrameIndex);
 }
 
+void UGeometryCacheTrackUsd::Initialize(
+	const UE::FUsdStage& InStage,
+	const FString& InPrimPath,
+	const FName& InRenderContext,
+	const TMap< FString, TMap< FString, int32 > >& InMaterialToPrimvarToUVIndex,
+	int32 InStartFrameIndex,
+	int32 InEndFrameIndex,
+	FReadUsdMeshFunction InReadFunc
+)
+{
+	Initialize(InStage, InPrimPath, InStartFrameIndex, InEndFrameIndex, InReadFunc);
+}
+
 void UGeometryCacheTrackUsd::UpdateTime(float Time, bool bLooping)
 {
 	if (UsdStream)
diff --git a/Engine/Plugins/Importers/USDImporter/Source/GeometryCacheUSD/Public/GeometryCacheTrackUSD.h b/Engine/Plugins/Importers/USDImporter/Source/GeometryCacheUSD/Public/GeometryCacheTrackUSD.h
--- a/Engine/Plugins/Importers/USDImporter/Source/GeometryCacheUSD/Public/GeometryCacheTrackUSD.h
+++ b/Engine/Plugins/Importers/USDImporter/Source/GeometryCacheUSD/Public/GeometryCacheTrackUSD.h
@@ -20,6 +20,15 @@ class GEOMETRYCACHEUSD_API UGeometryCacheTrackUsd : public UGeometryCacheTrack
 	UGeometryCacheTrackUsd();
 
 public:
+	void Initialize(
+		const UE::FUsdStage& InStage,
+		const FString& InPrimPath,
+		int32 InStartFrameIndex,
+		int32 InEndFrameIndex,
+		FReadUsdMeshFunction InReadFunc
+	);
+
+	UE_DEPRECATED(5.3, "The RenderContext and InMaterialToPrimvarToUVIndex parameters are no longer used.")
 	void Initialize(
 		const UE::FUsdStage& InStage,
 		const FString& InPrimPath,
@@ -56,11 +65,9 @@ public:
 	void UnloadUsdStage();
 
 public:
-	FName RenderContext;
 	double FramesPerSecond;
 	int32 StartFrameIndex;
 	int32 EndFrameIndex;
-	TMap< FString, TMap< FString, int32 > > MaterialToPrimvarToUVIndex;
 
 	FString PrimPath;
 
diff --git a/Engine/Plugins/Importers/USDImporter/Source/USDClasses/Public/USDAssetImportData.h b/Engine/Plugins/Importers/USDImporter/Source/USDClasses/Public/USDAssetImportData.h
--- a/Engine/Plugins/Importers/USDImporter/Source/USDClasses/Public/USDAssetImportData.h
+++ b/Engine/Plugins/Importers/USDImporter/Source/USDClasses/Public/USDAssetImportData.h
@@ -7,8 +7,8 @@
 
 #include "USDAssetImportData.generated.h"
 
-UCLASS(AutoExpandCategories = (Options), MinimalAPI)
-class UUsdAssetImportData : public UAssetImportData
+UCLASS()
+class USDCLASSES_API UUsdAssetImportData : public UAssetImportData
 {
 	GENERATED_BODY()
 
@@ -22,8 +22,8 @@ public:
 	TObjectPtr<class UObject> ImportOptions;
 };
 
-UCLASS(AutoExpandCategories = (Options), MinimalAPI)
-class UUsdAnimSequenceAssetImportData : public UUsdAssetImportData
+UCLASS()
+class USDCLASSES_API UUsdAnimSequenceAssetImportData : public UUsdAssetImportData
 {
 	GENERATED_BODY()
 
@@ -41,7 +41,7 @@ public:
 
 /** Simple wrapper because we're not allowed to have TMap properties with TArray<FString> as values */
 USTRUCT()
-struct FUsdPrimPathList
+struct USDCLASSES_API FUsdPrimPathList
 {
 	GENERATED_BODY()
 
@@ -49,9 +49,37 @@ struct FUsdPrimPathList
 	TArray<FString> PrimPaths;
 };
 
-/** We assign these to UStaticMeshes or USkeletalMeshes generated from USD */
-UCLASS( AutoExpandCategories = ( Options ), MinimalAPI )
-class UUsdMeshAssetImportData : public UUsdAssetImportData
+/** AssetImportData assigned to Unreal materials that are generated when parsing USD Material prims */
+UCLASS()
+class USDCLASSES_API UUsdMaterialAssetImportData : public UUsdAssetImportData
+{
+	GENERATED_BODY()
+
+public:
+	/**
+	 * In the context of our reference materials that just read a single texture for each material parameter, this
+	 * describes the primvar that the USD material is sampling for each texture.
+	 * e.g. {'BaseColor': 'st0', 'Metallic': 'st1'}
+	 */
+	UPROPERTY()
+	TMap<FString, FString> ParameterToPrimvar;
+
+	/**
+	 * In the context of our reference materials that just read a single texture for each material parameter, this
+	 * describes the Unreal UV index that will be used for sampling by each USD primvar. The idea is to use this in
+	 * combination with ParameterToPrimvar, in order to describe which UV index the material has currently assigned to
+	 * each parameter. When assigning the material to meshes later, we'll compare this member with the
+	 * UUsdMeshAssetImportData's own PrimvarToUVIndex to see if they are compatible or not, and if not spawn a new
+	 * material instance that is.
+	 * e.g. {'firstPrimvar': 0, 'st': 1, 'st1': 2}
+	 */
+	UPROPERTY()
+	TMap<FString, int32> PrimvarToUVIndex;
+};
+
+/** We assign these to UStaticMeshes, UGeometryCaches or USkeletalMeshes generated from USD */
+UCLASS()
+class USDCLASSES_API UUsdMeshAssetImportData : public UUsdAssetImportData
 {
 	GENERATED_BODY()
 
@@ -63,4 +91,8 @@ public:
 	 */
 	UPROPERTY()
 	TMap< int32, FUsdPrimPathList > MaterialSlotToPrimPaths;
+
+	/** Describes which primvars should be assigned to each UV index. */
+	UPROPERTY()
+	TMap<FString, int32> PrimvarToUVIndex;
 };
\ No newline at end of file
diff --git a/Engine/Plugins/Importers/USDImporter/Source/USDSchemas/Private/MeshTranslationImpl.cpp b/Engine/Plugins/Importers/USDImporter/Source/USDSchemas/Private/MeshTranslationImpl.cpp
--- a/Engine/Plugins/Importers/USDImporter/Source/USDSchemas/Private/MeshTranslationImpl.cpp
+++ b/Engine/Plugins/Importers/USDImporter/Source/USDSchemas/Private/MeshTranslationImpl.cpp
@@ -4,33 +4,39 @@
 
 #include "USDAssetCache2.h"
 #include "USDAssetImportData.h"
+#include "USDConversionUtils.h"
 #include "USDGeomMeshConversion.h"
 #include "USDInfoCache.h"
 #include "USDLog.h"
 #include "USDMemory.h"
 #include "USDProjectSettings.h"
+#include "USDShadeConversion.h"
 #include "USDTypesConversion.h"
 
 #include "Components/MeshComponent.h"
-#include "CoreMinimal.h"
+#include "Components/SkeletalMeshComponent.h"
+#include "Components/StaticMeshComponent.h"
+#include "Engine/SkeletalMesh.h"
+#include "Engine/StaticMesh.h"
+#include "GeometryCache.h"
+#include "GeometryCacheComponent.h"
 #include "Materials/MaterialInstanceConstant.h"
 #include "Materials/MaterialInstanceDynamic.h"
 #include "Materials/MaterialInterface.h"
 #include "MaterialShared.h"
-#include "Misc/Paths.h"
+#include "MeshDescription.h"
 #include "UObject/Package.h"
 
 #if USE_USD_SDK
 
 #include "USDIncludesStart.h"
-#include "pxr/usd/usd/editContext.h"
 #include "pxr/usd/usd/prim.h"
 #include "pxr/usd/usd/stage.h"
 #include "pxr/usd/usdGeom/mesh.h"
-#include "pxr/usd/usdGeom/xformable.h"
-#include "pxr/usd/usdShade/material.h"
 #include "USDIncludesEnd.h"
 
+#define UNUSED_UV_INDEX MAX_MESH_TEXTURE_COORDS_MD
+
 namespace UE::MeshTranslationImplInternal::Private
 {
 	UMaterialInterface* CreateTwoSidedVersionOfMaterial(UMaterialInterface* OneSidedMat)
@@ -42,10 +48,12 @@ namespace UE::MeshTranslationImplInternal::Private
 
 		FString PrimPath;
 #if WITH_EDITOR
-		if (UUsdAssetImportData* ImportData = Cast<UUsdAssetImportData>(OneSidedMat->AssetImportData))
+		UUsdMaterialAssetImportData* OneSidedImportData = Cast<UUsdMaterialAssetImportData>(OneSidedMat->AssetImportData);
+		if (!ensureMsgf(OneSidedImportData, TEXT("Expected one-sided material '%s' to have an UUsdMaterialAssetImportData at this point!"), *OneSidedMat->GetPathName()))
 		{
-			PrimPath = ImportData->PrimPath;
+			return nullptr;
 		}
+		PrimPath = OneSidedImportData->PrimPath;
 #endif // WITH_EDITOR
 
 		UMaterialInterface* TwoSidedMat = nullptr;
@@ -56,7 +64,7 @@ namespace UE::MeshTranslationImplInternal::Private
 		// get the reference UsdPreviewSurface instead, as that is also *its* reference
 		UMaterialInterface* ReferenceMaterial = OneSidedMaterialInstance ? OneSidedMaterialInstance->Parent.Get() : nullptr;
 		UMaterialInterface* ReferenceMaterialTwoSided = nullptr;
-		if (ReferenceMaterial && ReferenceMaterial->GetPathName().StartsWith(TEXT("/USDImporter/Materials/UsdPreviewSurface")))
+		if (ReferenceMaterial && MeshTranslationImpl::IsReferencePreviewSurfaceMaterial(ReferenceMaterial))
 		{
 			ReferenceMaterialTwoSided = MeshTranslationImpl::GetTwoSidedVersionOfReferencePreviewSurfaceMaterial(ReferenceMaterial);
 		}
@@ -82,12 +90,8 @@ namespace UE::MeshTranslationImplInternal::Private
 			);
 			if (TwoSidedMIC)
 			{
-				UUsdAssetImportData* ImportData = NewObject< UUsdAssetImportData >(
-					TwoSidedMIC,
-					TEXT("USDAssetImportData")
-				);
-				ImportData->PrimPath = PrimPath;
-				TwoSidedMIC->AssetImportData = ImportData;
+				TwoSidedMIC->AssetImportData =
+					DuplicateObject(OneSidedImportData, TwoSidedMIC, TEXT("USDAssetImportData"));
 			}
 
 			TwoSidedMIC->SetParentEditorOnly(ReferenceMaterialTwoSided);
@@ -107,12 +111,8 @@ namespace UE::MeshTranslationImplInternal::Private
 			);
 			if (TwoSidedMIC)
 			{
-				UUsdAssetImportData* ImportData = NewObject< UUsdAssetImportData >(
-					TwoSidedMIC,
-					TEXT("USDAssetImportData")
-				);
-				ImportData->PrimPath = PrimPath;
-				TwoSidedMIC->AssetImportData = ImportData;
+				TwoSidedMIC->AssetImportData =
+					DuplicateObject(OneSidedImportData, TwoSidedMIC, TEXT("USDAssetImportData"));
 			}
 
 			TwoSidedMIC->SetParentEditorOnly(OneSidedMat);
@@ -144,11 +144,281 @@ namespace UE::MeshTranslationImplInternal::Private
 
 			TwoSidedMID->CopyParameterOverrides(MID);
 
+			// TODO: Copy the PrimvarToUVIndex from the asset user data
+
 			TwoSidedMat = TwoSidedMID;
 		}
 
 		return TwoSidedMat;
 	}
+
+	/**
+	 * Returns Material in case it is already compatible with the provided MeshPrimvarToUVIndex, otherwise creates a
+	 * new material instance of it, and sets the UVIndex material parameters to match a UVIndex setup that is compatible
+	 * with the mesh.
+	 * This function will also already cache and link the generated material.
+	 */
+	UMaterialInterface* CreatePrimvarCompatibleVersionOfMaterial(
+		UMaterialInterface& Material,
+		const TMap<FString, int32>& MeshPrimvarToUVIndex,
+		UUsdAssetCache2* AssetCache,
+		FUsdInfoCache* InfoCache
+	)
+	{
+		// TODO: Implement this for runtime when asset user data support is in
+	#if WITH_EDITOR
+		UUsdMaterialAssetImportData* MaterialAssetImportData = Cast<UUsdMaterialAssetImportData>(Material.AssetImportData);
+		if (!ensureMsgf(MaterialAssetImportData, TEXT("Expected material '%s' to have an UUsdMaterialAssetImportData at this point!"), *Material.GetPathName()))
+		{
+			return nullptr;
+		}
+
+		UE_LOG(LogUsd, Verbose, TEXT("Getting compatible material based on Material '%s' (Parameter to primvar: %s, primvar to UV index: %s) matching mesh primvar to UV index mapping '%s'"),
+			*Material.GetPathName(),
+			*UsdUtils::StringifyMap(MaterialAssetImportData->ParameterToPrimvar),
+			*UsdUtils::StringifyMap(MaterialAssetImportData->PrimvarToUVIndex),
+			*UsdUtils::StringifyMap(MeshPrimvarToUVIndex)
+		);
+
+		TArray<TSet<FString>> UVIndexToMeshPrimvars;
+		UVIndexToMeshPrimvars.SetNum(MAX_MESH_TEXTURE_COORDS_MD);
+		for (const TPair<FString, int32>& MeshPair : MeshPrimvarToUVIndex)
+		{
+			if (UVIndexToMeshPrimvars.IsValidIndex(MeshPair.Value))
+			{
+				UVIndexToMeshPrimvars[MeshPair.Value].Add(MeshPair.Key);
+			}
+		}
+
+		// Check if mesh and material are compatible. Note that it's perfectly valid for the material to try reading
+		// an UVIndex the mesh doesn't provide at all, or trying to read a primvar that doesn't exist on the mesh.
+		bool bCompatible = true;
+		for (const TPair<FString, int32>& Pair : MaterialAssetImportData->PrimvarToUVIndex)
+		{
+			const FString& MaterialPrimvar = Pair.Key;
+			int32 MaterialUVIndex = Pair.Value;
+
+			// If the mesh has the same primvar the material wants, it should be at the same UVIndex the material
+			// will read from
+			if (const int32* MeshUVIndex = MeshPrimvarToUVIndex.Find(MaterialPrimvar))
+			{
+				if (*MeshUVIndex != MaterialUVIndex)
+				{
+					bCompatible = false;
+					// Don't break here so we can show all warnings
+				}
+			}
+			else
+			{
+				UE_LOG(LogUsd, Log, TEXT("Failed to find primvar '%s' needed by material '%s' on its assigned mesh. Available primvars and UV indices: %s"),
+					*MaterialPrimvar,
+					*Material.GetPathName(),
+					*UsdUtils::StringifyMap(MeshPrimvarToUVIndex)
+				);
+			}
+
+			// If the material is going to read from a given UVIndex that exists on the mesh, that UV set should
+			// contain the primvar data that the material expects to read
+			if (UVIndexToMeshPrimvars.IsValidIndex(MaterialUVIndex))
+			{
+				const TSet<FString>& CompatiblePrimvars = UVIndexToMeshPrimvars[MaterialUVIndex];
+				if (!CompatiblePrimvars.Contains(MaterialPrimvar))
+				{
+					bCompatible = false;
+					// Don't break here so we can show all warnings
+				}
+			}
+		}
+		if (bCompatible)
+		{
+			UE_LOG(LogUsd, Verbose, TEXT("Material '%s' is compatible with provided primvar to UV index mapping"),
+				*Material.GetPathName()
+			);
+			return &Material;
+		}
+
+		// We need to find or create another compatible material instead
+		UMaterialInterface* CompatibleMaterial = nullptr;
+
+		// First, let's create the target primvar UVIndex assignment that is compatible with this mesh.
+		// We use an array of TPairs here so that we can sort these into a deterministic order for hashing later.
+		TArray<TPair<FString, int32>> CompatiblePrimvarAndUVIndexPairs;
+		CompatiblePrimvarAndUVIndexPairs.Reserve(MaterialAssetImportData->PrimvarToUVIndex.Num());
+		for (const TPair<FString, int32>& Pair : MaterialAssetImportData->PrimvarToUVIndex)
+		{
+			const FString& MaterialPrimvar = Pair.Key;
+
+			bool bFoundUVIndex = false;
+
+			// Mesh has this primvar available at some UV index, point to it
+			if (const int32* FoundMeshUVIndex = MeshPrimvarToUVIndex.Find(MaterialPrimvar))
+			{
+				int32 MeshUVIndex = *FoundMeshUVIndex;
+				if (MeshUVIndex >= 0 && MeshUVIndex < MAX_MESH_TEXTURE_COORDS_MD)
+				{
+					CompatiblePrimvarAndUVIndexPairs.Add(TPair<FString, int32>{MaterialPrimvar, MeshUVIndex});
+					bFoundUVIndex = true;
+				}
+			}
+
+			if (!bFoundUVIndex)
+			{
+				// Point this primvar to read an unused UV index instead, since our mesh doesn't have this primvar
+				CompatiblePrimvarAndUVIndexPairs.Add(TPair<FString, int32>{MaterialPrimvar, UNUSED_UV_INDEX});
+			}
+		}
+
+		// Try finding a compatible material in the asset cache, if our original Material belongs to it
+		FString MaterialHash = AssetCache ? AssetCache->GetHashForAsset(&Material) : FString{};
+		const bool bMaterialBelongsToAssetCache = !MaterialHash.IsEmpty();
+		if (bMaterialBelongsToAssetCache)
+		{
+			// Generate a deterministic hash based on the original material hash and this primvar UVIndex assignment
+			CompatiblePrimvarAndUVIndexPairs.Sort(
+				[](const TPair<FString, int32>& LHS, const TPair<FString, int32>& RHS)
+				{
+					if (LHS.Key == RHS.Key)
+					{
+						return LHS.Value < RHS.Value;
+					}
+					else
+					{
+						return LHS.Key < RHS.Key;
+					}
+				}
+			);
+			FSHAHash Hash;
+			FSHA1 SHA1;
+			SHA1.UpdateWithString(*MaterialHash, MaterialHash.Len());
+			for (const TPair<FString, int32>& Pair : CompatiblePrimvarAndUVIndexPairs)
+			{
+				SHA1.UpdateWithString(*Pair.Key, Pair.Key.Len());
+				SHA1.Update((const uint8*)&Pair.Value, sizeof(Pair.Value));
+			}
+			SHA1.Final();
+			SHA1.GetHash(&Hash.Hash[0]);
+			MaterialHash = Hash.ToString();
+
+			if (UMaterialInterface* ExistingCompatibleMaterial = Cast<UMaterialInterface>(AssetCache->GetCachedAsset(MaterialHash)))
+			{
+				UE_LOG(LogUsd, Verbose, TEXT("Found existing compatible Material '%s' on the asset cache with hash '%s'"),
+					*ExistingCompatibleMaterial->GetPathName(),
+					*MaterialHash
+				);
+				CompatibleMaterial = ExistingCompatibleMaterial;
+			}
+		}
+
+		// We have to create a brand new compatible material instance
+		if (!CompatibleMaterial)
+		{
+			TMap<FString, int32> CompatiblePrimvarToUVIndex;
+			CompatiblePrimvarToUVIndex.Reserve(CompatiblePrimvarAndUVIndexPairs.Num());
+			for (const TPair<FString, int32>& Pair : CompatiblePrimvarAndUVIndexPairs)
+			{
+				CompatiblePrimvarToUVIndex.Add(Pair);
+			}
+
+			const FName NewInstanceName = MakeUniqueObjectName(
+				GetTransientPackage(),
+				UMaterialInstance::StaticClass(),
+				Material.GetFName()
+			);
+
+			UE_LOG(LogUsd, Verbose, TEXT("Generating compatible version of Material '%s' (Parameter to primvar: %s, primvar to UV index: %s) with hash '%s'"),
+				*Material.GetPathName(),
+				*UsdUtils::StringifyMap(MaterialAssetImportData->ParameterToPrimvar),
+				*UsdUtils::StringifyMap(CompatiblePrimvarToUVIndex),
+				*MaterialHash
+			);
+
+			if (GIsEditor)
+			{
+				UMaterialInstanceConstant* CompatibleMIC = NewObject<UMaterialInstanceConstant>(
+					GetTransientPackage(),
+					NewInstanceName,
+					Material.GetFlags()
+				);
+
+				UUsdMaterialAssetImportData* CompatibleImportData = DuplicateObject(
+					MaterialAssetImportData,
+					CompatibleMIC,
+					TEXT("USDAssetImportData")
+				);
+				CompatibleImportData->PrimvarToUVIndex = CompatiblePrimvarToUVIndex;
+
+				// Actually set the parameters with the target UV indices
+				for (const TPair<FString, FString>& ParameterPair : CompatibleImportData->ParameterToPrimvar)
+				{
+					const FString& Parameter = ParameterPair.Key;
+					const FString& Primvar = ParameterPair.Value;
+
+					if (int32* UVIndex = CompatibleImportData->PrimvarToUVIndex.Find(Primvar))
+					{
+						// Force-disable using the texture at all if the mesh doesn't provide the primvar that should be
+						// used to sample it with
+						if (*UVIndex == UNUSED_UV_INDEX)
+						{
+							UsdUtils::SetScalarParameterValue(
+								*CompatibleMIC,
+								*FString::Printf(TEXT("Use%sTexture"), *Parameter),
+								0.0f
+							);
+						}
+						else
+						{
+							UsdUtils::SetScalarParameterValue(
+								*CompatibleMIC,
+								*FString::Printf(TEXT("%sUVIndex"), *Parameter),
+								static_cast<float>(*UVIndex)
+							);
+						}
+					}
+				}
+
+				CompatibleMIC->AssetImportData = CompatibleImportData;
+				CompatibleMIC->SetParentEditorOnly(&Material);
+				CompatibleMIC->PostEditChange();
+
+				CompatibleMaterial = CompatibleMIC;
+			}
+			else
+			{
+				UMaterialInstanceDynamic* CompatibleMID = UMaterialInstanceDynamic::Create(
+					&Material,
+					GetTransientPackage(),
+					NewInstanceName
+				);
+
+				// TODO: Also copy our custom asset user data here when we add support for that
+
+				// TODO: Actually set the parameters with the target UV indices
+
+				CompatibleMaterial = CompatibleMID;
+			}
+		}
+
+		if (CompatibleMaterial && CompatibleMaterial != &Material)
+		{
+			if (bMaterialBelongsToAssetCache)
+			{
+				AssetCache->CacheAsset(MaterialHash, CompatibleMaterial);
+			}
+
+			if (InfoCache)
+			{
+				for (const UE::FSdfPath& Prim : InfoCache->GetPrimsForAsset(&Material))
+				{
+					InfoCache->LinkAssetToPrim(Prim, CompatibleMaterial);
+				}
+			}
+		}
+
+		return CompatibleMaterial;
+	#else
+		return &Material;
+	#endif // WITH_EDITOR
+	}
 };
 
 TMap<const UsdUtils::FUsdPrimMaterialSlot*, UMaterialInterface*> MeshTranslationImpl::ResolveMaterialAssignmentInfo(
@@ -162,6 +432,16 @@ TMap<const UsdUtils::FUsdPrimMaterialSlot*, UMaterialInterface*> MeshTranslation
 	FScopedUnrealAllocs Allocs;
 
 	TMap<const UsdUtils::FUsdPrimMaterialSlot*, UMaterialInterface*> ResolvedMaterials;
+	if(AssignmentInfo.Num() == 0)
+	{
+		return ResolvedMaterials;
+	}
+
+	// Generating compatible materials is somewhat elaborate, so we'll cache the generated ones in this call in case we
+	// have multiple material slots using the same material. The MeshPrimvarToUVIndex would always be the same for those
+	// anyway, so we know we can reuse these compatible materials
+	TMap<UMaterialInterface*, UMaterialInterface*> MaterialToCompatibleMaterial;
+	const TMap<FString, int32>& MeshPrimvarToUVIndex = AssignmentInfo[0].PrimvarToUVIndex;
 
 	uint32 GlobalResolvedMaterialIndex = 0;
 	for (int32 InfoIndex = 0; InfoIndex < AssignmentInfo.Num(); ++InfoIndex)
@@ -195,7 +475,7 @@ TMap<const UsdUtils::FUsdPrimMaterialSlot*, UMaterialInterface*> MeshTranslation
 							MaterialInstance = UsdUtils::CreateDisplayColorMaterialInstanceConstant(DisplayColorDesc.GetValue());
 #if WITH_EDITOR
 							// Leave PrimPath as empty as it likely will be reused by many prims
-							UUsdAssetImportData* ImportData = NewObject< UUsdAssetImportData >(MaterialInstance, TEXT("USDAssetImportData"));
+							UUsdMaterialAssetImportData* ImportData = NewObject< UUsdMaterialAssetImportData >(MaterialInstance, TEXT("USDAssetImportData"));
 							MaterialInstance->AssetImportData = ImportData;
 #endif // WITH_EDITOR
 						}
@@ -217,15 +497,36 @@ TMap<const UsdUtils::FUsdPrimMaterialSlot*, UMaterialInterface*> MeshTranslation
 			case UsdUtils::EPrimAssignmentType::MaterialPrim:
 			{
 				UMaterialInterface* OneSidedMat = nullptr;
+				bool bOneSidedMatIsInstanceOfReferencePreviewSurface = false;
+
+				UE::FSdfPath MaterialPrimPath{*Slot.MaterialSource};
+
+				TSet<UMaterialInterface*> ExistingMaterials = InfoCache.GetAssetsForPrim<UMaterialInterface>(MaterialPrimPath);
 
-				TSet<UMaterialInterface*> ExistingMaterials = InfoCache.GetAssetsForPrim<UMaterialInterface>(UE::FSdfPath{*Slot.MaterialSource});
 				for (UMaterialInterface* ExistingMaterial : ExistingMaterials)
 				{
 					const bool bExistingIsTwoSided = ExistingMaterial->IsTwoSided();
 
 					if (!bExistingIsTwoSided)
 					{
-						OneSidedMat = ExistingMaterial;
+						// Prefer sticking with a material instance that has as parent one of our reference materials.
+						// The idea here being that we have two approaches when making TwoSided and compatible
+						// materials: A) Make the material compatible first, and then a TwoSided version of the
+						// compatible; B) Make the material TwoSided first, and then a compatible version of the
+						// TwoSided; We're going to chose B), for the reason that at runtime we can only make a material
+						// TwoSided if it is an instance of our reference materials (as we can't manually change the
+						// material base property overrides at runtime)
+						UMaterialInstance* ExistingInstance = Cast<UMaterialInstance>(ExistingMaterial);
+						const bool bExistingIsInstanceOfReferencePreviewSurface =
+							ExistingInstance
+							&& MeshTranslationImpl::IsReferencePreviewSurfaceMaterial(ExistingInstance->Parent);
+						if (!OneSidedMat
+							|| (!bOneSidedMatIsInstanceOfReferencePreviewSurface
+								&& bExistingIsInstanceOfReferencePreviewSurface))
+						{
+							OneSidedMat = ExistingMaterial;
+							bOneSidedMatIsInstanceOfReferencePreviewSurface = bExistingIsInstanceOfReferencePreviewSurface;
+						}
 					}
 
 					if (Slot.bMeshIsDoubleSided == bExistingIsTwoSided)
@@ -234,44 +535,56 @@ TMap<const UsdUtils::FUsdPrimMaterialSlot*, UMaterialInterface*> MeshTranslation
 					}
 				}
 
-				FString MaterialHash = AssetCache.GetHashForAsset(OneSidedMat);
+				FString MaterialHash = Material ? AssetCache.GetHashForAsset(Material) : FString{};
 
-				if (Slot.bMeshIsDoubleSided)
+				// Need to create a two-sided material on-demand, *before* we make it compatible:
+				// This because at runtime we can't just set the base property overrides, and just instead create a new
+				// MIC based on the TwoSided reference material, and the compatible material should be a MIC of that MIC
+				if (Slot.bMeshIsDoubleSided && !Material)
 				{
-					MaterialHash = MaterialHash + UnrealIdentifiers::TwoSidedMaterialSuffix;
-
-					// Need to create a two-sided material on-demand
-					if (!Material)
+					// By now we parsed all materials so we must have the single-sided version of this material
+					if (!OneSidedMat)
 					{
-						// By now we parsed all materials so we must have the single-sided version of this material
-						if (!OneSidedMat)
-						{
-							UE_LOG(LogUsd, Warning, TEXT("Failed to generate a two-sided material from the material prim at path '%s' as no single-sided material was generated for it."), *Slot.MaterialSource);
-							continue;
-						}
+						UE_LOG(
+							LogUsd,
+							Warning,
+							TEXT("Failed to generate a two-sided material from the material prim at path '%s' as no "
+								 "single-sided material was generated for it."),
+							*Slot.MaterialSource
+						);
+						continue;
+					}
 
-						// Check if for some reason we already have a two-sided material ready due to a complex scenario
-						// related to the global cache
-						UMaterialInterface* TwoSidedMat = Cast<UMaterialInterface>(AssetCache.GetCachedAsset(MaterialHash));
-						if (!TwoSidedMat)
-						{
-							TwoSidedMat = UE::MeshTranslationImplInternal::Private::CreateTwoSidedVersionOfMaterial(OneSidedMat);
-						}
+					const FString OneSidedHash = AssetCache.GetHashForAsset(OneSidedMat);
+					const FString TwoSidedHash = OneSidedHash + UnrealIdentifiers::TwoSidedMaterialSuffix;
 
-						if (TwoSidedMat)
-						{
-							TwoSidedMat->SetFlags(RF_Transient);
-							Material = TwoSidedMat;
-						}
-						else
-						{
-							UE_LOG(LogUsd, Warning, TEXT("Failed to generate a two-sided material from the material prim at path '%s'. Falling back to using the single-sided material '%s' instead."),
-								*Slot.MaterialSource,
-								*OneSidedMat->GetPathName()
-							);
-							Material = OneSidedMat;
-							MaterialHash = AssetCache.GetHashForAsset(OneSidedMat);
-						}
+					// Check if for some reason we already have a two-sided material ready due to a complex scenario
+					// related to the global cache
+					UMaterialInterface* TwoSidedMat = Cast<UMaterialInterface>(AssetCache.GetCachedAsset(TwoSidedHash));
+					if (!TwoSidedMat)
+					{
+						TwoSidedMat =
+							UE::MeshTranslationImplInternal::Private::CreateTwoSidedVersionOfMaterial(OneSidedMat);
+					}
+
+					if (TwoSidedMat)
+					{
+						TwoSidedMat->SetFlags(RF_Transient);
+						Material = TwoSidedMat;
+						MaterialHash = TwoSidedHash;
+					}
+					else
+					{
+						UE_LOG(
+							LogUsd,
+							Warning,
+							TEXT("Failed to generate a two-sided material from the material prim at path '%s'. Falling "
+								 "back to using the single-sided material '%s' instead."),
+							*Slot.MaterialSource,
+							*OneSidedMat->GetPathName()
+						);
+						Material = OneSidedMat;
+						MaterialHash = OneSidedHash;
 					}
 				}
 
@@ -281,6 +594,34 @@ TMap<const UsdUtils::FUsdPrimMaterialSlot*, UMaterialInterface*> MeshTranslation
 					// first time
 					AssetCache.CacheAsset(MaterialHash, Material);
 					InfoCache.LinkAssetToPrim(UE::FSdfPath{*Slot.MaterialSource}, Material);
+
+					// Finally, try to make our generated material primvar-compatible. We do this last because this will
+					// create another instance with the non-compatible material as reference material, which means we also
+					// need that reference to be cached and linked for the asset cache to be able to handle dependencies
+					// properly
+					if (UMaterialInterface* AlreadyHandledMaterial = MaterialToCompatibleMaterial.FindRef(Material))
+					{
+						Material = AlreadyHandledMaterial;
+
+						AssetCache.TouchAsset(Material);
+						InfoCache.LinkAssetToPrim(UE::FSdfPath{*Slot.MaterialSource}, Material);
+					}
+					else
+					{
+						UMaterialInterface* CompatibleMaterial =
+							UE::MeshTranslationImplInternal::Private::CreatePrimvarCompatibleVersionOfMaterial(
+								*Material,
+								MeshPrimvarToUVIndex,
+								&AssetCache,
+								&InfoCache
+							);
+
+						if (CompatibleMaterial)
+						{
+							MaterialToCompatibleMaterial.Add(Material, CompatibleMaterial);
+							Material = CompatibleMaterial;
+						}
+					}
 				}
 
 				break;
@@ -428,13 +769,57 @@ void MeshTranslationImpl::SetMaterialOverrides(
 		}
 	}
 
-	TMap<const UsdUtils::FUsdPrimMaterialSlot*, UMaterialInterface*> ResolvedMaterials = MeshTranslationImpl::ResolveMaterialAssignmentInfo(
-		ValidPrim,
-		LODIndexToAssignments,
-		AssetCache,
-		InfoCache,
-		Flags
+#if WITH_EDITOR
+	UUsdMeshAssetImportData* ImportData = nullptr;
+	if (UStaticMeshComponent* StaticMeshComponent = Cast<UStaticMeshComponent>(&MeshComponent))
+	{
+		if (UStaticMesh* Mesh = StaticMeshComponent->GetStaticMesh())
+		{
+			ImportData = Cast<UUsdMeshAssetImportData>(Mesh->GetAssetImportData());
+		}
+	}
+	else if (USkeletalMeshComponent* SkeletalMeshComponent = Cast<USkeletalMeshComponent>(&MeshComponent))
+	{
+		if (USkeletalMesh* Mesh = SkeletalMeshComponent->GetSkeletalMeshAsset())
+		{
+			ImportData = Cast<UUsdMeshAssetImportData>(Mesh->GetAssetImportData());
+		}
+	}
+	else if (UGeometryCacheComponent* GeometryCacheComponent = Cast<UGeometryCacheComponent>(&MeshComponent))
+	{
+		if (UGeometryCache* Mesh = GeometryCacheComponent->GetGeometryCache())
+		{
+			ImportData = Cast<UUsdMeshAssetImportData>(Mesh->AssetImportData);
+		}
+	}
+	else
+	{
+		ensureMsgf(false, TEXT("Unexpected component class '%s' encountered when setting material overrides for prim '%s'!"),
+			*MeshComponent.GetClass()->GetName(),
+			*UsdToUnreal::ConvertPath(Prim.GetPrimPath())
+		);
+	}
+
+	ensureMsgf(ImportData, TEXT("Mesh assigned to component '%s' generated for prim '%s' should have an UUsdMeshAssetImportData at this point!"),
+		*MeshComponent.GetPathName(),
+		*UsdToUnreal::ConvertPath(Prim.GetPrimPath())
 	);
+#endif // WITH_EDITOR
+
+	TMap<const UsdUtils::FUsdPrimMaterialSlot*, UMaterialInterface*> ResolvedMaterials;
+	if(ImportData && LODIndexToAssignments.Num() > 0)
+	{
+		// Stash our PrimvarToUVIndex in here, as that's where ResolveMaterialAssignmentInfo will look for it
+		LODIndexToAssignments[0].PrimvarToUVIndex = ImportData->PrimvarToUVIndex;
+
+		ResolvedMaterials = MeshTranslationImpl::ResolveMaterialAssignmentInfo(
+			ValidPrim,
+			LODIndexToAssignments,
+			AssetCache,
+			InfoCache,
+			Flags
+		);
+	}
 
 	// Compare resolved materials with existing assignments, and create overrides if we need to
 	uint32 StaticMeshSlotIndex = 0;
@@ -443,6 +828,17 @@ void MeshTranslationImpl::SetMaterialOverrides(
 		const TArray< UsdUtils::FUsdPrimMaterialSlot >& LODSlots = LODIndexToAssignments[LODIndex].Slots;
 		for (int32 LODSlotIndex = 0; LODSlotIndex < LODSlots.Num(); ++LODSlotIndex, ++StaticMeshSlotIndex)
 		{
+			// If we don't even have as many existing assignments as we have overrides just stop here.
+			// This should happen often now because we'll always at least attempt at setting overrides on every
+			// component (but only ever set anything if we really need to).
+			// Previously we only attempted setting overrides in case the component didn't "own" the mesh prim,
+			// but now it is not feasible to do that given the global asset cache and how assets may have come
+			// from an entirely new stage/session.
+			if (!ExistingAssignments.IsValidIndex(StaticMeshSlotIndex))
+			{
+				break;
+			}
+
 			const UsdUtils::FUsdPrimMaterialSlot& Slot = LODSlots[LODSlotIndex];
 
 			UMaterialInterface* Material = nullptr;
@@ -456,17 +852,6 @@ void MeshTranslationImpl::SetMaterialOverrides(
 				continue;
 			}
 
-			// If we don't even have as many existing assignments as we have overrides just stop here.
-			// This should happen often now because we'll always at least attempt at setting overrides on every
-			// component (but only ever set anything if we really need to).
-			// Previously we only attempted setting overrides in case the component didn't "own" the mesh prim,
-			// but now it is not feasible to do that given the global asset cache and how assets may have come
-			// from an entirely new stage/session.
-			if (!ExistingAssignments.IsValidIndex(StaticMeshSlotIndex))
-			{
-				continue;
-			}
-
 			UMaterialInterface* ExistingMaterial = ExistingAssignments[StaticMeshSlotIndex];
 			if (ExistingMaterial == Material)
 			{
@@ -632,4 +1017,33 @@ UMaterialInterface* MeshTranslationImpl::GetTwoSidedVersionOfReferencePreviewSur
 	return nullptr;
 }
 
+bool MeshTranslationImpl::IsReferencePreviewSurfaceMaterial(UMaterialInterface* Material)
+{
+	if (!Material)
+	{
+		return false;
+	}
+
+	const FSoftObjectPath PathName = Material->GetPathName();
+
+	const UUsdProjectSettings* Settings = GetDefault<UUsdProjectSettings>();
+	if (!Settings)
+	{
+		return false;
+	}
+
+	TSet<FSoftObjectPath> ReferenceMaterials = {
+		Settings->ReferencePreviewSurfaceMaterial,
+		Settings->ReferencePreviewSurfaceTranslucentMaterial,
+		Settings->ReferencePreviewSurfaceTwoSidedMaterial,
+		Settings->ReferencePreviewSurfaceTranslucentTwoSidedMaterial,
+		Settings->ReferencePreviewSurfaceVTMaterial,
+		Settings->ReferencePreviewSurfaceTranslucentVTMaterial,
+		Settings->ReferencePreviewSurfaceTwoSidedVTMaterial,
+		Settings->ReferencePreviewSurfaceTranslucentTwoSidedVTMaterial
+	};
+
+	return ReferenceMaterials.Contains(PathName);
+}
+
 #endif // #if USE_USD_SDK
\ No newline at end of file
diff --git a/Engine/Plugins/Importers/USDImporter/Source/USDSchemas/Private/MeshTranslationImpl.h b/Engine/Plugins/Importers/USDImporter/Source/USDSchemas/Private/MeshTranslationImpl.h
--- a/Engine/Plugins/Importers/USDImporter/Source/USDSchemas/Private/MeshTranslationImpl.h
+++ b/Engine/Plugins/Importers/USDImporter/Source/USDSchemas/Private/MeshTranslationImpl.h
@@ -82,6 +82,10 @@ namespace MeshTranslationImpl
 	// Example: Receives UsdPreviewSurfaceTwoSidedTranslucent -> Returns UsdPreviewSurfaceTwoSidedTranslucent
 	// Example: Receives SomeOtherReferenceMaterial -> Returns nullptr
 	UMaterialInterface* GetTwoSidedVersionOfReferencePreviewSurfaceMaterial( UMaterialInterface* ReferenceMaterial );
+
+	// Returns whether the Material is one of the UsdPreviewSurface reference materials (which can be reassigned by the
+	// user on a per project basis)
+	bool IsReferencePreviewSurfaceMaterial(UMaterialInterface* Material);
 }
 
 #endif // #if USE_USD_SDK
\ No newline at end of file
diff --git a/Engine/Plugins/Importers/USDImporter/Source/USDSchemas/Private/USDGeomMeshTranslator.cpp b/Engine/Plugins/Importers/USDImporter/Source/USDSchemas/Private/USDGeomMeshTranslator.cpp
--- a/Engine/Plugins/Importers/USDImporter/Source/USDSchemas/Private/USDGeomMeshTranslator.cpp
+++ b/Engine/Plugins/Importers/USDImporter/Source/USDSchemas/Private/USDGeomMeshTranslator.cpp
@@ -233,7 +233,7 @@ namespace UsdGeomMeshTranslatorImpl
 	}
 
 	/** Returns true if material infos have changed on the StaticMesh */
-	bool ProcessMaterials(
+	bool ProcessStaticMeshMaterials(
 		const pxr::UsdPrim& UsdPrim,
 		const TArray<UsdUtils::FUsdPrimMaterialAssignmentInfo>& LODIndexToMaterialInfo,
 		UStaticMesh& StaticMesh,
@@ -257,7 +257,12 @@ namespace UsdGeomMeshTranslatorImpl
 		}
 
 #if WITH_EDITOR
+		// TODO: Swap with asset user data for runtime support
 		UUsdMeshAssetImportData* ImportData = Cast< UUsdMeshAssetImportData>(StaticMesh.AssetImportData.Get());
+		ensureMsgf(ImportData, TEXT("Static Mesh '%s' generated for prim '%s' should have an UUsdMeshAssetImportData at this point!"),
+			*StaticMesh.GetPathName(),
+			*UsdToUnreal::ConvertPath(UsdPrim.GetPrimPath())
+		);
 #endif // WITH_EDITOR
 
 		TMap<const UsdUtils::FUsdPrimMaterialSlot*, UMaterialInterface*> ResolvedMaterials = MeshTranslationImpl::ResolveMaterialAssignmentInfo(
@@ -341,8 +346,6 @@ namespace UsdGeomMeshTranslatorImpl
 		return bMaterialAssignementsHaveChanged;
 	}
 
-#if WITH_EDITOR
-	// #ueent_todo: Merge the code with ProcessMaterials
 	bool ProcessGeometryCacheMaterials(
 		const pxr::UsdPrim& UsdPrim,
 		const TArray< UsdUtils::FUsdPrimMaterialAssignmentInfo >& LODIndexToMaterialInfo,
@@ -357,160 +360,66 @@ namespace UsdGeomMeshTranslatorImpl
 
 #if WITH_EDITOR
 		UUsdMeshAssetImportData* MeshImportData = Cast< UUsdMeshAssetImportData>(GeometryCache.AssetImportData.Get());
+		ensureMsgf(MeshImportData, TEXT("Geometry Cache '%s' generated for prim '%s' should have an UUsdMeshAssetImportData at this point!"),
+			*GeometryCache.GetPathName(),
+			*UsdToUnreal::ConvertPath(UsdPrim.GetPrimPath())
+		);
 #endif // WITH_EDITOR
 
-		uint32 StaticMeshSlotIndex = 0;
+		TMap<const UsdUtils::FUsdPrimMaterialSlot*, UMaterialInterface*> ResolvedMaterials = MeshTranslationImpl::ResolveMaterialAssignmentInfo(
+			UsdPrim,
+			LODIndexToMaterialInfo,
+			AssetCache,
+			*InfoCache,
+			Flags
+		);
+
+		uint32 SlotIndex = 0;
 		for (int32 LODIndex = 0; LODIndex < LODIndexToMaterialInfo.Num(); ++LODIndex)
 		{
 			const TArray< UsdUtils::FUsdPrimMaterialSlot >& LODSlots = LODIndexToMaterialInfo[LODIndex].Slots;
 
-			for (int32 LODSlotIndex = 0; LODSlotIndex < LODSlots.Num(); ++LODSlotIndex, ++StaticMeshSlotIndex)
+			for (int32 LODSlotIndex = 0; LODSlotIndex < LODSlots.Num(); ++LODSlotIndex, ++SlotIndex)
 			{
 				const UsdUtils::FUsdPrimMaterialSlot& Slot = LODSlots[LODSlotIndex];
 				UMaterialInterface* Material = nullptr;
-
-				switch (Slot.AssignmentType)
-				{
-				case UsdUtils::EPrimAssignmentType::DisplayColor:
-				{
-					FScopedUsdAllocs Allocs;
-
-					// Try reusing an already created DisplayColor material
-					if (UMaterialInstanceConstant* ExistingMaterial = Cast< UMaterialInstanceConstant >(AssetCache.GetCachedAsset(Slot.MaterialSource)))
-					{
-						Material = ExistingMaterial;
-					}
-
-					// Need to create a new DisplayColor material
-					if (Material == nullptr)
-					{
-						if (TOptional< UsdUtils::FDisplayColorMaterial > DisplayColorDesc = UsdUtils::FDisplayColorMaterial::FromString(Slot.MaterialSource))
-						{
-							UMaterialInstance* MaterialInstance = nullptr;
-
-							if (GIsEditor)  // Editor, PIE => true; Standlone, packaged => false
-							{
-								MaterialInstance = UsdUtils::CreateDisplayColorMaterialInstanceConstant(DisplayColorDesc.GetValue());
-
-								// Leave PrimPath as empty as it likely will be reused by many prims
-								UUsdAssetImportData* ImportData = NewObject< UUsdAssetImportData >(MaterialInstance, TEXT("USDAssetImportData"));
-								MaterialInstance->AssetImportData = ImportData;
-							}
-							else
-							{
-								MaterialInstance = UsdUtils::CreateDisplayColorMaterialInstanceDynamic(DisplayColorDesc.GetValue());
-							}
-
-							if (MaterialInstance)
-							{
-								MaterialInstance->SetFlags(RF_Transient);
-								AssetCache.CacheAsset(Slot.MaterialSource, MaterialInstance);
-								Material = MaterialInstance;
-							}
-						}
-					}
-
-					break;
-				}
-				case UsdUtils::EPrimAssignmentType::MaterialPrim:
-				{
-					FScopedUsdAllocs Allocs;
-
-					// Check first or else we may get a warning
-					if (pxr::SdfPath::IsValidPathString(UnrealToUsd::ConvertString(*Slot.MaterialSource).Get()))
-					{
-						pxr::SdfPath MaterialPrimPath = UnrealToUsd::ConvertPath(*Slot.MaterialSource).Get();
-
-						// TODO: This may break if MaterialPrimPath targets a prim inside a LOD variant that is disabled...
-						TUsdStore< pxr::UsdPrim > MaterialPrim = UsdPrim.GetStage()->GetPrimAtPath(MaterialPrimPath);
-						if (MaterialPrim.Get() && InfoCache)
-						{
-							Material = InfoCache->GetSingleAssetForPrim<UMaterialInterface>(
-								UE::FSdfPath{MaterialPrim.Get().GetPrimPath()}
-							);
-						}
-					}
-
-					// TODO: This stuff is missing the two-sided material upgrade code...
-
-					break;
-				}
-				case UsdUtils::EPrimAssignmentType::UnrealMaterial:
-				{
-					// TODO: Try reusing the MeshTranslationImpl code for this as it has more warnings/feedback
-					Material = Cast< UMaterialInterface >(FSoftObjectPath(Slot.MaterialSource).TryLoad());
-					break;
-				}
-				case UsdUtils::EPrimAssignmentType::None:
-				default:
+				if (UMaterialInterface** FoundMaterial = ResolvedMaterials.Find(&Slot))
 				{
-					ensure(false);
-					break;
-				}
+					Material = *FoundMaterial;
 				}
-
-				// Fallback to this UsdGeomMesh DisplayColor material if present
-				// TODO: This code is almost an exact copy-paste of the DisplayColor case above...
-				if (Material == nullptr)
+				else
 				{
-					FScopedUsdAllocs Allocs;
-
-					// Try reusing an already created DisplayColor material
-					if (UMaterialInstanceConstant* ExistingMaterial = Cast< UMaterialInstanceConstant >(AssetCache.GetCachedAsset(Slot.MaterialSource)))
-					{
-						Material = ExistingMaterial;
-					}
-
-					// Need to create a new DisplayColor material
-					if (Material == nullptr)
-					{
-						if (TOptional< UsdUtils::FDisplayColorMaterial > DisplayColorDesc = UsdUtils::ExtractDisplayColorMaterial(pxr::UsdGeomMesh(UsdPrim)))
-						{
-							UMaterialInstance* MaterialInstance = nullptr;
-
-							if (GIsEditor)  // Editor, PIE => true; Standlone, packaged => false
-							{
-								MaterialInstance = UsdUtils::CreateDisplayColorMaterialInstanceConstant(DisplayColorDesc.GetValue());
-
-								// Leave PrimPath as empty as it likely will be reused by many prims
-								UUsdAssetImportData* ImportData = NewObject< UUsdAssetImportData >(MaterialInstance, TEXT("USDAssetImportData"));
-								MaterialInstance->AssetImportData = ImportData;
-							}
-							else
-							{
-								MaterialInstance = UsdUtils::CreateDisplayColorMaterialInstanceDynamic(DisplayColorDesc.GetValue());
-							}
-
-							if (MaterialInstance)
-							{
-								MaterialInstance->SetFlags(RF_Transient);
-								AssetCache.CacheAsset(Slot.MaterialSource, MaterialInstance);
-								Material = MaterialInstance;
-							}
-						}
-					}
+					UE_LOG(
+						LogUsd,
+						Error,
+						TEXT("Failed to resolve material '%s' for slot '%d' of LOD '%d' for geometry cache '%s'"),
+						*Slot.MaterialSource,
+						LODSlotIndex,
+						LODIndex,
+						*UsdToUnreal::ConvertPath(UsdPrim.GetPath())
+					);
+					continue;
 				}
 
-				if (!GeometryCache.Materials.IsValidIndex(StaticMeshSlotIndex))
+				if (!GeometryCache.Materials.IsValidIndex(SlotIndex))
 				{
 					GeometryCache.Materials.Add(Material);
 					bMaterialAssignementsHaveChanged = true;
 				}
-				else if (!(GeometryCache.Materials[StaticMeshSlotIndex] == Material))
+				else if (!(GeometryCache.Materials[SlotIndex] == Material))
 				{
-					GeometryCache.Materials[StaticMeshSlotIndex] = Material;
+					GeometryCache.Materials[SlotIndex] = Material;
 					bMaterialAssignementsHaveChanged = true;
 				}
 
 #if WITH_EDITOR
-				MeshImportData->MaterialSlotToPrimPaths.FindOrAdd(StaticMeshSlotIndex).PrimPaths = Slot.PrimPaths.Array();
+				MeshImportData->MaterialSlotToPrimPaths.FindOrAdd(SlotIndex).PrimPaths = Slot.PrimPaths.Array();
 #endif // WITH_EDITOR
 			}
 		}
 
 		return bMaterialAssignementsHaveChanged;
 	}
-#endif // WITH_EDITOR
 
 	// If UsdMesh is a LOD, will parse it and all of the other LODs, and and place them in OutLODIndexToMeshDescription and OutLODIndexToMaterialInfo.
 	// Note that these other LODs will be hidden in other variants, and won't show up on traversal unless we actively switch the variants (which we do here).
@@ -542,54 +451,100 @@ namespace UsdGeomMeshTranslatorImpl
 
 		UsdToUnreal::FUsdMeshConversionOptions OptionsCopy = Options;
 
-		TFunction<bool(const pxr::UsdGeomMesh&, int32)> ConvertLOD = [&](const pxr::UsdGeomMesh& LODMesh, int32 LODIndex)
-		{
-			FMeshDescription TempMeshDescription;
-			UsdUtils::FUsdPrimMaterialAssignmentInfo TempMaterialInfo;
+		TSet<FString> AllPrimvars;
+		TSet<FString> PreferredPrimvars;
 
-			FStaticMeshAttributes StaticMeshAttributes(TempMeshDescription);
-			StaticMeshAttributes.Register();
+		// Here we are choosing to preemptively traverse all LODs to fetch the combined primvars that we'll use for
+		// each UV index for the static mesh as a whole
+		TFunction<bool(const pxr::UsdGeomMesh&, int32)> CombinePrimvars =
+			[&AllPrimvars, &PreferredPrimvars, &LODIndexToMaterialInfoMap]
+			(const pxr::UsdGeomMesh& LODMesh, int32 LODIndex)
+			{
+				TArray<TUsdStore<pxr::UsdGeomPrimvar>> MeshPrimvars = UsdUtils::GetUVSetPrimvars(
+					LODMesh,
+					TNumericLimits<int32>::Max()
+				);
+
+				for (const TUsdStore<pxr::UsdGeomPrimvar>& MeshPrimvar : MeshPrimvars)
+				{
+					FString PrimvarName = UsdToUnreal::ConvertToken(MeshPrimvar.Get().GetName());
+					PrimvarName.RemoveFromStart(TEXT("primvars:"));
+
+					AllPrimvars.Add(PrimvarName);
 
-			bool bSuccess = true;
+					// Keep track of which primvars are texCoord2f as we always want to prefer these over other float2s
+					if (MeshPrimvar.Get().GetTypeName().GetRole() == pxr::SdfValueTypeNames->TexCoord2f.GetRole())
+					{
+						PreferredPrimvars.Add(PrimvarName);
+					}
+				}
 
-			// The user can't manually hide/unhide a particular LOD in the engine after it is imported, so we should probably bake
-			// the particular LOD visibility into the combined mesh. Note how we don't use computed visibility here, as we only really
-			// care if this mesh in particular has been marked as invisible
-			pxr::TfToken Visibility;
-			pxr::UsdAttribute VisibilityAttr = LODMesh.GetVisibilityAttr();
-			if (VisibilityAttr && VisibilityAttr.Get(&Visibility, Options.TimeCode) && Visibility == pxr::UsdGeomTokens->inherited)
+				return true;
+			};
+		if (!UsdUtils::IterateLODMeshes(ParentPrim, CombinePrimvars))
+		{
+			return false;
+		}
+		TMap<FString, int32> PrimvarToUVIndex = UsdUtils::CombinePrimvarsIntoUVSets(
+			AllPrimvars,
+			PreferredPrimvars
+		);
+
+		TFunction<bool(const pxr::UsdGeomMesh&, int32)> ConvertLOD =
+			[&OptionsCopy, &Options, &PrimvarToUVIndex, &LODIndexToMeshDescriptionMap, &LODIndexToMaterialInfoMap]
+			(const pxr::UsdGeomMesh& LODMesh, int32 LODIndex)
 			{
-				// If we're interpreting LODs we must bake the transform from each LOD Mesh into the vertices, because there's no guarantee
-				// all LODs have the same transform, so we can't just put the transforms directly on the component. If we are not interpreting
-				// LODs we can do that though
-				// TODO: Handle resetXformOp here
-				bool bResetXformStack = false;
-				FTransform MeshTransform = FTransform::Identity;
-				bSuccess &= UsdToUnreal::ConvertXformable(LODMesh.GetPrim().GetStage(), LODMesh, MeshTransform, Options.TimeCode.GetValue(), &bResetXformStack);
+				FMeshDescription TempMeshDescription;
 
-				if (bSuccess)
+				UsdUtils::FUsdPrimMaterialAssignmentInfo TempMaterialInfo;
+				TempMaterialInfo.PrimvarToUVIndex = PrimvarToUVIndex;
+
+				FStaticMeshAttributes StaticMeshAttributes(TempMeshDescription);
+				StaticMeshAttributes.Register();
+
+				bool bSuccess = true;
+
+				// The user can't manually hide/unhide a particular LOD in the engine after it is imported, so we should probably bake
+				// the particular LOD visibility into the combined mesh. Note how we don't use computed visibility here, as we only really
+				// care if this mesh in particular has been marked as invisible
+				pxr::TfToken Visibility;
+				pxr::UsdAttribute VisibilityAttr = LODMesh.GetVisibilityAttr();
+				if (VisibilityAttr && VisibilityAttr.Get(&Visibility, Options.TimeCode) && Visibility == pxr::UsdGeomTokens->inherited)
 				{
-					OptionsCopy.AdditionalTransform = MeshTransform * Options.AdditionalTransform;
-					OptionsCopy.bMergeIdenticalMaterialSlots = false;  // We only merge slots when collapsing, and we never collapse LODs
-
-					bSuccess &= UsdToUnreal::ConvertGeomMesh(
-						LODMesh,
-						TempMeshDescription,
-						TempMaterialInfo,
-						OptionsCopy
-					);
+					// If we're interpreting LODs we must bake the transform from each LOD Mesh into the vertices, because there's no guarantee
+					// all LODs have the same transform, so we can't just put the transforms directly on the component. If we are not interpreting
+					// LODs we can do that though
+					// TODO: Handle resetXformOp here
+					bool bResetXformStack = false;
+					FTransform MeshTransform = FTransform::Identity;
+					bSuccess &= UsdToUnreal::ConvertXformable(LODMesh.GetPrim().GetStage(), LODMesh, MeshTransform, Options.TimeCode.GetValue(), &bResetXformStack);
+
+					if (bSuccess)
+					{
+						OptionsCopy.AdditionalTransform = MeshTransform * Options.AdditionalTransform;
+						OptionsCopy.bMergeIdenticalMaterialSlots = false;  // We only merge slots when collapsing, and we never collapse LODs
+
+						bSuccess &= UsdToUnreal::ConvertGeomMesh(
+							LODMesh,
+							TempMeshDescription,
+							TempMaterialInfo,
+							OptionsCopy
+						);
+					}
 				}
-			}
 
-			if (bSuccess)
-			{
-				LODIndexToMeshDescriptionMap.Add(LODIndex, MoveTemp(TempMeshDescription));
-				LODIndexToMaterialInfoMap.Add(LODIndex, MoveTemp(TempMaterialInfo));
-			}
+				if (bSuccess)
+				{
+					LODIndexToMeshDescriptionMap.Add(LODIndex, MoveTemp(TempMeshDescription));
+					LODIndexToMaterialInfoMap.Add(LODIndex, MoveTemp(TempMaterialInfo));
+				}
 
-			return true;
-		};
-		bool bFoundLODs = UsdUtils::IterateLODMeshes(ParentPrim, ConvertLOD);
+				return true;
+			};
+		if (!UsdUtils::IterateLODMeshes(ParentPrim, ConvertLOD))
+		{
+			return false;
+		}
 
 		// Place them in order as we can't have e.g. LOD0 and LOD2 without LOD1, and there's no reason downstream code needs to care about this
 		OutLODIndexToMeshDescription.Reset(LODIndexToMeshDescriptionMap.Num());
@@ -602,7 +557,7 @@ namespace UsdGeomMeshTranslatorImpl
 			OutLODIndexToMaterialInfo.Add(MoveTemp(LODIndexToMaterialInfoMap[OldLODIndex]));
 		}
 
-		return bFoundLODs;
+		return true;
 	}
 
 	void LoadMeshDescriptions(
@@ -1010,9 +965,6 @@ namespace UsdGeomMeshTranslatorImpl
 
 			const UE::FUsdStage& Stage = Context->Stage;
 
-			TMap< FString, TMap< FString, int32 > > Unused;
-			TMap< FString, TMap< FString, int32 > >* MaterialToPrimvarToUVIndex = Context->MaterialToPrimvarToUVIndex ? Context->MaterialToPrimvarToUVIndex : &Unused;
-
 			// Fetch the animated mesh start/end frame as they may be different from just the stage's start and end time codes
 			int32 StartFrame = FMath::FloorToInt(Stage.GetStartTimeCode());
 			int32 EndFrame = FMath::CeilToInt(Stage.GetEndTimeCode());
@@ -1044,8 +996,6 @@ namespace UsdGeomMeshTranslatorImpl
 			UsdTrack->Initialize(
 				Stage,
 				InPrimPath,
-				Context->RenderContext,
-				*MaterialToPrimvarToUVIndex,
 				StartFrame,
 				EndFrame,
 				[InPrimPath, Options](const TWeakObjectPtr<UGeometryCacheTrackUsd> TrackPtr, float Time, FGeometryCacheMeshData& OutMeshData) mutable
@@ -1076,7 +1026,6 @@ namespace UsdGeomMeshTranslatorImpl
 					// Need a local copy of Options to set the TimeCode since this function is called from multiple worker threads
 					UsdToUnreal::FUsdMeshConversionOptions LocalOptions(Options);
 					LocalOptions.TimeCode = pxr::UsdTimeCode{Time};
-					LocalOptions.MaterialToPrimvarToUVIndex = &Track->MaterialToPrimvarToUVIndex;
 
 					UsdGeomMeshTranslatorImpl::LoadMeshDescriptions(
 						pxr::UsdTyped(Prim),
@@ -1295,6 +1244,8 @@ void FBuildStaticMeshTaskChain::SetupTasks()
 				{
 					UUsdMeshAssetImportData* ImportData = NewObject<UUsdMeshAssetImportData>(StaticMesh, TEXT("UUSDAssetImportData"));
 					ImportData->PrimPath = MeshName;
+					ImportData->PrimvarToUVIndex = LODIndexToMaterialInfo[0].PrimvarToUVIndex;  // We use the same primvar mapping for all LODs
+
 					StaticMesh->AssetImportData = ImportData;
 				}
 #endif // WITH_EDITOR
@@ -1304,7 +1255,7 @@ void FBuildStaticMeshTaskChain::SetupTasks()
 				// components
 				if (bIsNew)
 				{
-					UsdGeomMeshTranslatorImpl::ProcessMaterials(
+					UsdGeomMeshTranslatorImpl::ProcessStaticMeshMaterials(
 						GetPrim(),
 						LODIndexToMaterialInfo,
 						*StaticMesh,
@@ -1427,9 +1378,6 @@ void FGeomMeshCreateAssetsTaskChain::SetupTasks()
 	Do(LaunchPolicy,
 		[this]() -> bool
 		{
-			TMap< FString, TMap< FString, int32 > > Unused;
-			TMap< FString, TMap< FString, int32 > >* MaterialToPrimvarToUVIndex = Context->MaterialToPrimvarToUVIndex ? Context->MaterialToPrimvarToUVIndex : &Unused;
-
 			pxr::TfToken RenderContextToken = pxr::UsdShadeTokens->universalRenderContext;
 			if (!Context->RenderContext.IsNone())
 			{
@@ -1447,7 +1395,6 @@ void FGeomMeshCreateAssetsTaskChain::SetupTasks()
 			Options.PurposesToLoad = Context->PurposesToLoad;
 			Options.RenderContext = RenderContextToken;
 			Options.MaterialPurpose = MaterialPurposeToken;
-			Options.MaterialToPrimvarToUVIndex = MaterialToPrimvarToUVIndex;
 			Options.bMergeIdenticalMaterialSlots = Context->bMergeIdenticalMaterialSlots;
 			Options.AdditionalTransform = AdditionalTransform;
 
@@ -1499,8 +1446,6 @@ void FGeometryCacheCreateAssetsTaskChain::SetupTasks()
 			// otherwise we may think it's a new asset just because the context is at a different timecode (e.g. if we reload)
 			// TODO: Hash all timecodes, or else our mesh may change at t=5 and we never reload it because we only hash t=0
 			const bool bAllowInterpretingLODs = false;  // GeometryCaches don't have LODs
-			TMap< FString, TMap< FString, int32 > > Unused;
-			TMap< FString, TMap< FString, int32 > >* MaterialToPrimvarToUVIndex = Context->MaterialToPrimvarToUVIndex ? Context->MaterialToPrimvarToUVIndex : &Unused;
 
 			pxr::TfToken RenderContextToken = pxr::UsdShadeTokens->universalRenderContext;
 			if (!Context->RenderContext.IsNone())
@@ -1519,7 +1464,6 @@ void FGeometryCacheCreateAssetsTaskChain::SetupTasks()
 			Options.PurposesToLoad = Context->PurposesToLoad;
 			Options.RenderContext = RenderContextToken;
 			Options.MaterialPurpose = MaterialPurposeToken;
-			Options.MaterialToPrimvarToUVIndex = MaterialToPrimvarToUVIndex;
 			Options.bMergeIdenticalMaterialSlots = Context->bMergeIdenticalMaterialSlots;
 
 			UsdGeomMeshTranslatorImpl::LoadMeshDescriptions(
@@ -1553,6 +1497,7 @@ void FGeometryCacheCreateAssetsTaskChain::SetupTasks()
 			{
 				UUsdMeshAssetImportData* ImportData = NewObject< UUsdMeshAssetImportData >(GeometryCache, TEXT("UUSDAssetImportData"));
 				ImportData->PrimPath = PrimPathString;
+				ImportData->PrimvarToUVIndex = LODIndexToMaterialInfo[0].PrimvarToUVIndex;
 				GeometryCache->AssetImportData = ImportData;
 
 				// Only the original creator of the prim at creation time gets to set the material assignments
diff --git a/Engine/Plugins/Importers/USDImporter/Source/USDSchemas/Private/USDGeomPointInstancerTranslator.cpp b/Engine/Plugins/Importers/USDImporter/Source/USDSchemas/Private/USDGeomPointInstancerTranslator.cpp
--- a/Engine/Plugins/Importers/USDImporter/Source/USDSchemas/Private/USDGeomPointInstancerTranslator.cpp
+++ b/Engine/Plugins/Importers/USDImporter/Source/USDSchemas/Private/USDGeomPointInstancerTranslator.cpp
@@ -115,9 +115,6 @@ void FUsdGeomPointInstancerCreateAssetsTaskChain::SetupTasks()
 			FMeshDescription& AddedMeshDescription = LODIndexToMeshDescription.Emplace_GetRef();
 			UsdUtils::FUsdPrimMaterialAssignmentInfo& AssignmentInfo = LODIndexToMaterialInfo.Emplace_GetRef();
 
-			TMap< FString, TMap< FString, int32 > > Unused;
-			TMap< FString, TMap< FString, int32 > >* MaterialToPrimvarToUVIndex = Context->MaterialToPrimvarToUVIndex ? Context->MaterialToPrimvarToUVIndex : &Unused;
-
 			pxr::TfToken RenderContextToken = pxr::UsdShadeTokens->universalRenderContext;
 			if ( !Context->RenderContext.IsNone() )
 			{
@@ -135,7 +132,6 @@ void FUsdGeomPointInstancerCreateAssetsTaskChain::SetupTasks()
 			Options.PurposesToLoad = Context->PurposesToLoad;
 			Options.RenderContext = RenderContextToken;
 			Options.MaterialPurpose = MaterialPurposeToken;
-			Options.MaterialToPrimvarToUVIndex = MaterialToPrimvarToUVIndex;
 			Options.bMergeIdenticalMaterialSlots = Context->bMergeIdenticalMaterialSlots;
 
 			UsdToUnreal::ConvertGeomMeshHierarchy(
diff --git a/Engine/Plugins/Importers/USDImporter/Source/USDSchemas/Private/USDGeomXformableTranslator.cpp b/Engine/Plugins/Importers/USDImporter/Source/USDSchemas/Private/USDGeomXformableTranslator.cpp
--- a/Engine/Plugins/Importers/USDImporter/Source/USDSchemas/Private/USDGeomXformableTranslator.cpp
+++ b/Engine/Plugins/Importers/USDImporter/Source/USDSchemas/Private/USDGeomXformableTranslator.cpp
@@ -222,9 +222,6 @@ void FUsdGeomXformableCreateAssetsTaskChain::SetupTasks()
 			FMeshDescription& AddedMeshDescription = LODIndexToMeshDescription.Emplace_GetRef();
 			UsdUtils::FUsdPrimMaterialAssignmentInfo& AssignmentInfo = LODIndexToMaterialInfo.Emplace_GetRef();
 
-			TMap< FString, TMap< FString, int32 > > Unused;
-			TMap< FString, TMap< FString, int32 > >* MaterialToPrimvarToUVIndex = Context->MaterialToPrimvarToUVIndex ? Context->MaterialToPrimvarToUVIndex : &Unused;
-
 			pxr::TfToken RenderContextToken = pxr::UsdShadeTokens->universalRenderContext;
 			if ( !Context->RenderContext.IsNone() )
 			{
@@ -245,7 +242,6 @@ void FUsdGeomXformableCreateAssetsTaskChain::SetupTasks()
 			Options.PurposesToLoad = Context->PurposesToLoad;
 			Options.RenderContext = RenderContextToken;
 			Options.MaterialPurpose = MaterialPurposeToken;
-			Options.MaterialToPrimvarToUVIndex = MaterialToPrimvarToUVIndex;
 			Options.bMergeIdenticalMaterialSlots = Context->bMergeIdenticalMaterialSlots;
 
 			UsdToUnreal::ConvertGeomMeshHierarchy(
diff --git a/Engine/Plugins/Importers/USDImporter/Source/USDSchemas/Private/USDShadeMaterialTranslator.cpp b/Engine/Plugins/Importers/USDImporter/Source/USDSchemas/Private/USDShadeMaterialTranslator.cpp
--- a/Engine/Plugins/Importers/USDImporter/Source/USDSchemas/Private/USDShadeMaterialTranslator.cpp
+++ b/Engine/Plugins/Importers/USDImporter/Source/USDSchemas/Private/USDShadeMaterialTranslator.cpp
@@ -162,6 +162,8 @@ namespace UE::UsdShadeTranslator::Private
 								}
 								NewMID->CopyParameterOverrides( MID );
 
+								// TODO: Copy asset user data over to NewMID
+
 								if (Context->AssetCache->CanRemoveAsset(Hash) && Context->AssetCache->RemoveAsset(Hash))
 								{
 									Context->AssetCache->CacheAsset(Hash, NewMID);
@@ -266,23 +268,12 @@ void FUsdShadeMaterialTranslator::CreateAssets()
 			*IUsdClassesModule::SanitizeObjectName(FPaths::GetBaseFilename(PrimPathString))
 		);
 
-		// TODO: There is an issue here: Note how we're only ever going to write the material prim's primvars into
-		// this PrimvarToUVIndex map when first creating the material, and not if we find it in the asset cache.
-		// This because finding the primvars to use essentially involves parsing the entire material again, so we
-		// likely shouldn't do it every time.
-		// This means that a mesh with float2 UV sets and materials that use them will likely end up with no UVs once
-		// the stage reloads...
-		TMap<FString, int32> Unused;
-		TMap<FString, int32>& PrimvarToUVIndex = Context->MaterialToPrimvarToUVIndex
-			? Context->MaterialToPrimvarToUVIndex->FindOrAdd(PrimPathString)
-			: Unused;
-
 #if WITH_EDITOR
 		if ( GIsEditor ) // Also have to prevent Standalone game from going with MaterialInstanceConstants
 		{
 			if ( UMaterialInstanceConstant* NewMaterial = NewObject<UMaterialInstanceConstant>( GetTransientPackage(), InstanceName, Context->ObjectFlags | EObjectFlags::RF_Transient ) )
 			{
-				UUsdAssetImportData* ImportData = NewObject< UUsdAssetImportData >( NewMaterial, TEXT( "USDAssetImportData" ) );
+				UUsdMaterialAssetImportData* ImportData = NewObject<UUsdMaterialAssetImportData>(NewMaterial, TEXT("USDAssetImportData"));
 				ImportData->PrimPath = PrimPath.GetString();
 				NewMaterial->AssetImportData = ImportData;
 
@@ -290,7 +281,6 @@ void FUsdShadeMaterialTranslator::CreateAssets()
 					ShadeMaterial,
 					*NewMaterial,
 					Context->AssetCache.Get(),
-					PrimvarToUVIndex,
 					*Context->RenderContext.ToString()
 				);
 				if ( !bSuccess )
@@ -383,7 +373,9 @@ void FUsdShadeMaterialTranslator::CreateAssets()
 				{
 					NewMaterial->SetFlags(RF_Transient);
 
-					if ( UsdToUnreal::ConvertMaterial( ShadeMaterial, *NewMaterial, Context->AssetCache.Get(), PrimvarToUVIndex, *Context->RenderContext.ToString() ) )
+					// TODO: Create asset user data here so that ConvertMaterial can store the primvars used on each UV index
+
+					if (UsdToUnreal::ConvertMaterial(ShadeMaterial, *NewMaterial, Context->AssetCache.Get(), *Context->RenderContext.ToString()))
 					{
 						TSet<UTexture*> VTTextures;
 						TSet<UTexture*> NonVTTextures;
@@ -431,20 +423,6 @@ void FUsdShadeMaterialTranslator::CreateAssets()
 			}
 		}
 	}
-	else if ( Context->MaterialToPrimvarToUVIndex && Context->InfoCache )
-	{
-		const TSet<UE::FSdfPath> FoundPrimPaths = Context->InfoCache->GetPrimsForAsset(ConvertedMaterial);
-		if (FoundPrimPaths.Num() > 0)
-		{
-			const UE::FSdfPath& FoundPrimPath = *FoundPrimPaths.CreateConstIterator();
-
-			if (TMap<FString, int32>* PrimvarToUVIndex = Context->MaterialToPrimvarToUVIndex->Find(FoundPrimPath.GetString()))
-			{
-				// Copy the Material -> Primvar -> UV index mapping from the cached material prim path to this prim path
-				Context->MaterialToPrimvarToUVIndex->FindOrAdd(PrimPath.GetString()) = *PrimvarToUVIndex;
-			}
-		}
-	}
 
 	PostImportMaterial(MaterialHashString, ConvertedMaterial);
 }
@@ -467,10 +445,10 @@ void FUsdShadeMaterialTranslator::PostImportMaterial(const FString& MaterialHash
 	}
 
 #if WITH_EDITOR
-	UUsdAssetImportData* ImportData = Cast<UUsdAssetImportData>(ImportedMaterial->AssetImportData);
+	UUsdMaterialAssetImportData* ImportData = Cast<UUsdMaterialAssetImportData>(ImportedMaterial->AssetImportData);
 	if (!ImportData)
 	{
-		ImportData = NewObject<UUsdAssetImportData>(ImportedMaterial, TEXT("USDAssetImportData"));
+		ImportData = NewObject<UUsdMaterialAssetImportData>(ImportedMaterial, TEXT("USDAssetImportData"));
 		ImportedMaterial->AssetImportData = ImportData;
 	}
 	ImportData->PrimPath = PrimPath.GetString();
@@ -538,7 +516,7 @@ void FUsdShadeMaterialTranslator::PostImportMaterial(const FString& MaterialHash
 				UUsdAssetImportData* TextureImportData = Cast<UUsdAssetImportData>(Texture->AssetImportData);
 				if (!TextureImportData)
 				{
-					TextureImportData = NewObject<UUsdAssetImportData>(ImportedMaterial, TEXT("USDAssetImportData"));
+					TextureImportData = NewObject<UUsdAssetImportData>(Texture, TEXT("USDAssetImportData"));
 					Texture->AssetImportData = TextureImportData;
 				}
 				TextureImportData->PrimPath = PrimPath.GetString();
diff --git a/Engine/Plugins/Importers/USDImporter/Source/USDSchemas/Private/USDSkelRootTranslator.cpp b/Engine/Plugins/Importers/USDImporter/Source/USDSchemas/Private/USDSkelRootTranslator.cpp
--- a/Engine/Plugins/Importers/USDImporter/Source/USDSchemas/Private/USDSkelRootTranslator.cpp
+++ b/Engine/Plugins/Importers/USDImporter/Source/USDSchemas/Private/USDSkelRootTranslator.cpp
@@ -103,6 +103,10 @@ namespace UsdSkelRootTranslatorImpl
 
 #if WITH_EDITOR
 		UUsdMeshAssetImportData* ImportData = Cast< UUsdMeshAssetImportData>( SkeletalMesh->GetAssetImportData() );
+		ensureMsgf(ImportData, TEXT("Skeletal Mesh '%s' generated for prim '%s' should have an UUsdMeshAssetImportData at this point!"),
+			*SkeletalMesh->GetPathName(),
+			*UsdToUnreal::ConvertPath(UsdPrim.GetPrimPath())
+		);
 #endif // WITH_EDITOR
 
 		TMap<const UsdUtils::FUsdPrimMaterialSlot*, UMaterialInterface*> ResolvedMaterials = MeshTranslationImpl::ResolveMaterialAssignmentInfo(
@@ -425,12 +429,9 @@ namespace UsdSkelRootTranslatorImpl
 		FName& OutSkeletonName,
 		UsdUtils::FBlendShapeMap* OutBlendShapes,
 		TSet<FString>& InOutUsedMorphTargetNames,
-		const TMap< FString, TMap< FString, int32 > >& InMaterialToPrimvarsUVSetNames,
-		float InTime,
 		const UUsdAssetCache2& AssetCache,
 		bool bInInterpretLODs,
-		const pxr::TfToken& RenderContext,
-		const pxr::TfToken& MaterialPurpose
+		const UsdToUnreal::FUsdMeshConversionOptions& Options
 	)
 	{
 		if ( !InSkeletonRoot )
@@ -500,13 +501,10 @@ namespace UsdSkelRootTranslatorImpl
 			&LODIndexToSkeletalMeshImportDataMap,
 			&LODIndexToMaterialInfoMap,
 			&SkelQuery,
-			InTime,
-			&InMaterialToPrimvarsUVSetNames,
 			&InOutUsedMorphTargetNames,
 			OutBlendShapes,
 			&StageInfo,
-			RenderContext,
-			MaterialPurpose
+			Options
 		]
 		( const pxr::UsdGeomMesh& LODMesh, int32 LODIndex )
 		{
@@ -522,7 +520,7 @@ namespace UsdSkelRootTranslatorImpl
 			}
 
 			FSkeletalMeshImportData& LODImportData = LODIndexToSkeletalMeshImportDataMap.FindOrAdd( LODIndex );
-			TArray<UsdUtils::FUsdPrimMaterialSlot>& LODSlots = LODIndexToMaterialInfoMap.FindOrAdd( LODIndex ).Slots;
+			UsdUtils::FUsdPrimMaterialAssignmentInfo& LODMaterialInfo = LODIndexToMaterialInfoMap.FindOrAdd( LODIndex );
 
 			// BlendShape data is respective to point indices for each mesh in isolation, but we combine all points
 			// into one FSkeletalMeshImportData per LOD, so we need to remap the indices using this
@@ -532,10 +530,8 @@ namespace UsdSkelRootTranslatorImpl
 				SkinningQuery,
 				SkelQuery,
 				LODImportData,
-				LODSlots,
-				InMaterialToPrimvarsUVSetNames,
-				RenderContext,
-				MaterialPurpose
+				LODMaterialInfo,
+				Options
 			);
 			if ( !bSuccess )
 			{
@@ -1269,9 +1265,6 @@ namespace UsdSkelRootTranslatorImpl
 				// No point in importing blend shapes if the import context doesn't want them
 				UsdUtils::FBlendShapeMap* OutBlendShapes = Context->BlendShapesByPath ? &NewBlendShapes : nullptr;
 
-				TMap< FString, TMap< FString, int32 > > Unused;
-				const TMap< FString, TMap< FString, int32 > >* MaterialToPrimvarToUVIndex = Context->MaterialToPrimvarToUVIndex ? Context->MaterialToPrimvarToUVIndex : &Unused;
-
 				pxr::TfToken RenderContextToken = pxr::UsdShadeTokens->universalRenderContext;
 				if ( !Context->RenderContext.IsNone() )
 				{
@@ -1284,6 +1277,11 @@ namespace UsdSkelRootTranslatorImpl
 					MaterialPurposeToken = UnrealToUsd::ConvertToken( *Context->MaterialPurpose.ToString() ).Get();
 				}
 
+				UsdToUnreal::FUsdMeshConversionOptions Options;
+				Options.TimeCode = Context->Time;
+				Options.RenderContext = RenderContextToken;
+				Options.MaterialPurpose = MaterialPurposeToken;
+
 				const bool bContinueTaskChain = UsdSkelRootTranslatorImpl::LoadAllSkeletalData(
 					SkeletonCache.Get(),
 					pxr::UsdSkelRoot( GetPrim() ),
@@ -1293,12 +1291,9 @@ namespace UsdSkelRootTranslatorImpl
 					SkeletonName,
 					OutBlendShapes,
 					UsedMorphTargetNames,
-					*MaterialToPrimvarToUVIndex,
-					Context->Time,
 					*Context->AssetCache.Get(),
 					Context->bAllowInterpretingLODs,
-					RenderContextToken,
-					MaterialPurposeToken
+					Options
 				);
 
 				return bContinueTaskChain;
diff --git a/Engine/Plugins/Importers/USDImporter/Source/USDSchemas/Public/USDSchemaTranslator.h b/Engine/Plugins/Importers/USDImporter/Source/USDSchemas/Public/USDSchemaTranslator.h
--- a/Engine/Plugins/Importers/USDImporter/Source/USDSchemas/Public/USDSchemaTranslator.h
+++ b/Engine/Plugins/Importers/USDImporter/Source/USDSchemas/Public/USDSchemaTranslator.h
@@ -187,6 +187,7 @@ struct USDSCHEMAS_API FUsdSchemaTranslationContext : public TSharedFromThis< FUs
 	 * When parsing materials, we keep track of which primvar we mapped to which UV channel.
 	 * When parsing meshes later, we use this data to place the correct primvar values in each UV channel.
 	 */
+	UE_DEPRECATED( 5.3, "This mapping is now stored directly on Material AssetImportData/AssetUserData" )
 	TMap< FString, TMap< FString, int32 > >* MaterialToPrimvarToUVIndex = nullptr;
 
 	/**
diff --git a/Engine/Plugins/Importers/USDImporter/Source/USDStage/Private/USDStageActor.cpp b/Engine/Plugins/Importers/USDImporter/Source/USDStage/Private/USDStageActor.cpp
--- a/Engine/Plugins/Importers/USDImporter/Source/USDStage/Private/USDStageActor.cpp
+++ b/Engine/Plugins/Importers/USDImporter/Source/USDStage/Private/USDStageActor.cpp
@@ -120,7 +120,6 @@ struct FUsdStageActorImpl
 		TranslationContext->MaterialPurpose = StageActor->MaterialPurpose;
 		TranslationContext->RootMotionHandling = StageActor->RootMotionHandling;
 		TranslationContext->MaterialXOptions = StageActor->MaterialXOptions;
-		TranslationContext->MaterialToPrimvarToUVIndex = &StageActor->MaterialToPrimvarToUVIndex;
 		TranslationContext->BlendShapesByPath = &StageActor->BlendShapesByPath;
 		TranslationContext->InfoCache = StageActor->InfoCache;
 		TranslationContext->bTranslateOnlyUsedMaterials = true;  // There is no point in turning this off when just opening a stage
@@ -2057,7 +2056,6 @@ void AUsdStageActor::OnObjectsReplaced(const TMap<UObject*, UObject*>& ObjectRep
 			NewActor->LevelSequenceHelper = MoveTemp(LevelSequenceHelper);
 			NewActor->LevelSequence = LevelSequence;
 			NewActor->BlendShapesByPath = MoveTemp(BlendShapesByPath);
-			NewActor->MaterialToPrimvarToUVIndex = MoveTemp(MaterialToPrimvarToUVIndex);
 
 			NewActor->UsdListener.Register(NewActor->UsdStage);
 
diff --git a/Engine/Plugins/Importers/USDImporter/Source/USDStageEditorViewModels/Private/USDStageViewModel.cpp b/Engine/Plugins/Importers/USDImporter/Source/USDStageEditorViewModels/Private/USDStageViewModel.cpp
--- a/Engine/Plugins/Importers/USDImporter/Source/USDStageEditorViewModels/Private/USDStageViewModel.cpp
+++ b/Engine/Plugins/Importers/USDImporter/Source/USDStageEditorViewModels/Private/USDStageViewModel.cpp
@@ -354,7 +354,6 @@ void FUsdStageViewModel::ImportStage( const TCHAR* TargetContentFolder, UUsdStag
 			ImportContext.AssetCache = StageActor->UsdAssetCache;
 			ImportContext.InfoCache = StageActor->GetInfoCache();
 			ImportContext.LevelSequenceHelper.SetInfoCache(StageActor->GetInfoCache());
-			ImportContext.MaterialToPrimvarToUVIndex = StageActor->GetMaterialToPrimvarToUVIndex();
 
 			ImportContext.TargetSceneActorAttachParent = StageActor->GetRootComponent()->GetAttachParent();
 			ImportContext.TargetSceneActorTargetTransform = StageActor->GetActorTransform();
diff --git a/Engine/Plugins/Importers/USDImporter/Source/USDStageImporter/Private/USDStageImporter.cpp b/Engine/Plugins/Importers/USDImporter/Source/USDStageImporter/Private/USDStageImporter.cpp
--- a/Engine/Plugins/Importers/USDImporter/Source/USDStageImporter/Private/USDStageImporter.cpp
+++ b/Engine/Plugins/Importers/USDImporter/Source/USDStageImporter/Private/USDStageImporter.cpp
@@ -1722,7 +1722,6 @@ void UUsdStageImporter::ImportFromFile(FUsdStageImportContext& ImportContext)
 	TranslationContext->bAllowInterpretingLODs = ImportContext.ImportOptions->bInterpretLODs;
 	TranslationContext->bAllowParsingSkeletalAnimations = ImportContext.ImportOptions->bImportGeometry && ImportContext.ImportOptions->bImportSkeletalAnimations;
 	TranslationContext->bTranslateOnlyUsedMaterials = ImportContext.ImportOptions->bImportOnlyUsedMaterials;
-	TranslationContext->MaterialToPrimvarToUVIndex = &ImportContext.MaterialToPrimvarToUVIndex;
 	TranslationContext->InfoCache = ImportContext.InfoCache;
 	TranslationContext->BlendShapesByPath = &BlendShapesByPath;
 	TranslationContext->GroomInterpolationSettings = ImportContext.ImportOptions->GroomInterpolationSettings;
@@ -1839,7 +1838,6 @@ bool UUsdStageImporter::ReimportSingleAsset(FUsdStageImportContext& ImportContex
 	TranslationContext->bAllowInterpretingLODs = ImportContext.ImportOptions->bInterpretLODs;
 	TranslationContext->bAllowParsingSkeletalAnimations = ImportContext.ImportOptions->bImportGeometry && ImportContext.ImportOptions->bImportSkeletalAnimations;
 	TranslationContext->bTranslateOnlyUsedMaterials = ImportContext.ImportOptions->bImportOnlyUsedMaterials;
-	TranslationContext->MaterialToPrimvarToUVIndex = &ImportContext.MaterialToPrimvarToUVIndex;
 	TranslationContext->InfoCache = ImportContext.InfoCache;
 	TranslationContext->BlendShapesByPath = &BlendShapesByPath;
 	TranslationContext->GroomInterpolationSettings = ImportContext.ImportOptions->GroomInterpolationSettings;
diff --git a/Engine/Plugins/Importers/USDImporter/Source/USDUtilities/Private/USDConversionUtils.cpp b/Engine/Plugins/Importers/USDImporter/Source/USDUtilities/Private/USDConversionUtils.cpp
--- a/Engine/Plugins/Importers/USDImporter/Source/USDUtilities/Private/USDConversionUtils.cpp
+++ b/Engine/Plugins/Importers/USDImporter/Source/USDUtilities/Private/USDConversionUtils.cpp
@@ -87,6 +87,12 @@
 
 #define LOCTEXT_NAMESPACE "USDConversionUtils"
 
+static bool GParseUVSetsFromFloat2Primvars = true;
+static FAutoConsoleVariableRef CVarParseUVSetsFromFloat2Primvars(
+	TEXT("USD.ParseUVSetsFromFloat2Primvars"),
+	GParseUVSetsFromFloat2Primvars,
+	TEXT("Primvars with the 'texCoord2f' role will always be parsed when handling potential UV sets. If this cvar is enabled, we'll also handle primvars declared as just 'float2' however. You could disable this cvar if your pipeline emits many 'float2' primvars that you do not wish to be parsed as UV sets."));
+
 namespace USDConversionUtilsImpl
 {
 	/** Show some warnings if the UVSet primvars show some unsupported/problematic behavior */
@@ -658,152 +664,194 @@ int32 UsdUtils::GetPrimvarUVIndex( FString PrimvarName )
 	return 0;
 }
 
-TArray< TUsdStore< pxr::UsdGeomPrimvar > > UsdUtils::GetUVSetPrimvars( const pxr::UsdGeomMesh& UsdMesh )
-{
-	return UsdUtils::GetUVSetPrimvars(UsdMesh, {});
-}
-
-TArray< TUsdStore< pxr::UsdGeomPrimvar > > UsdUtils::GetUVSetPrimvars(
+TArray<TUsdStore<pxr::UsdGeomPrimvar>> UsdUtils::GetUVSetPrimvars(
 	const pxr::UsdGeomMesh& UsdMesh,
-	const TMap< FString, TMap< FString, int32 > >& MaterialToPrimvarsUVSetNames,
-	const pxr::TfToken& RenderContext,
-	const pxr::TfToken& MaterialPurpose
+	int32 MaxNumPrimvars
 )
 {
-	if ( !UsdMesh )
-	{
-		return {};
-	}
-
-	const bool bProvideMaterialIndices = false;
-	UsdUtils::FUsdPrimMaterialAssignmentInfo Info = UsdUtils::GetPrimMaterialAssignments(
-		UsdMesh.GetPrim(),
-		pxr::UsdTimeCode( 0.0 ),
-		bProvideMaterialIndices,
-		RenderContext,
-		MaterialPurpose
-	);
-
-	return UsdUtils::GetUVSetPrimvars( UsdMesh, MaterialToPrimvarsUVSetNames, Info );
-}
-
-TArray< TUsdStore< pxr::UsdGeomPrimvar > > UsdUtils::GetUVSetPrimvars( const pxr::UsdGeomMesh& UsdMesh, const TMap< FString, TMap< FString, int32 > >& MaterialToPrimvarsUVSetNames, const UsdUtils::FUsdPrimMaterialAssignmentInfo& UsdMeshMaterialAssignmentInfo )
-{
-	if ( !UsdMesh )
+	if (!UsdMesh)
 	{
 		return {};
 	}
 
 	FScopedUsdAllocs Allocs;
 
+	TArray<TUsdStore<pxr::UsdGeomPrimvar>> TexCoord2fPrimvars;
+	TArray<TUsdStore<pxr::UsdGeomPrimvar>> Float2Primvars;
+
 	// Collect all primvars that could be used as UV sets
-	TMap<FString, pxr::UsdGeomPrimvar> PrimvarsByName;
-	TMap<int32, TArray<pxr::UsdGeomPrimvar>> UsablePrimvarsByUVIndex;
-	pxr::UsdGeomPrimvarsAPI PrimvarsAPI{ UsdMesh };
-	for ( const pxr::UsdGeomPrimvar& Primvar : PrimvarsAPI.GetPrimvars() )
+	pxr::UsdGeomPrimvarsAPI PrimvarsAPI{UsdMesh};
+	for (const pxr::UsdGeomPrimvar& Primvar : PrimvarsAPI.GetPrimvars())
 	{
-		if ( !Primvar || !Primvar.HasValue() )
+		if (!Primvar || !Primvar.HasValue())
 		{
 			continue;
 		}
 
 		// We only care about primvars that can be used as float2[]. TexCoord2f is included
 		const pxr::SdfValueTypeName& TypeName = Primvar.GetTypeName();
-		if ( !TypeName.GetType().IsA( pxr::SdfValueTypeNames->Float2Array.GetType() ) )
+		if (!TypeName.GetType().IsA(pxr::SdfValueTypeNames->Float2Array.GetType()))
 		{
 			continue;
 		}
 
-		FString PrimvarName = UsdToUnreal::ConvertToken( Primvar.GetBaseName() );
-		int32 TargetUVIndex = UsdUtils::GetPrimvarUVIndex( PrimvarName );
+		if (Primvar.GetTypeName().GetRole() == pxr::SdfValueTypeNames->TexCoord2f.GetRole())
+		{
+			TexCoord2fPrimvars.Add(Primvar);
+		}
+		else if (GParseUVSetsFromFloat2Primvars)
+		{
+			Float2Primvars.Add(Primvar);
+		}
+	}
 
-		UsablePrimvarsByUVIndex.FindOrAdd( TargetUVIndex ).Add( Primvar );
-		PrimvarsByName.Add( PrimvarName, Primvar );
+	TexCoord2fPrimvars.Sort([](const TUsdStore<pxr::UsdGeomPrimvar>& A, const TUsdStore<pxr::UsdGeomPrimvar>& B)
+	{
+		return A.Get().GetName() < B.Get().GetName();
+	});
+	Float2Primvars.Sort([](const TUsdStore<pxr::UsdGeomPrimvar>& A, const TUsdStore<pxr::UsdGeomPrimvar>& B)
+	{
+		return A.Get().GetName() < B.Get().GetName();
+	});
+
+	TArray<TUsdStore<pxr::UsdGeomPrimvar>> Result;
+	Result.Reserve(FMath::Min(TexCoord2fPrimvars.Num() + Float2Primvars.Num(), MaxNumPrimvars));
+
+	int32 TexCoordPrimvarIndex = 0;
+	while (Result.Num() < MaxNumPrimvars && TexCoord2fPrimvars.IsValidIndex(TexCoordPrimvarIndex))
+	{
+		Result.Add(TexCoord2fPrimvars[TexCoordPrimvarIndex++]);
 	}
 
-	// Collect all primvars that are in fact used by the materials assigned to this mesh
-	TMap<int32, TArray<pxr::UsdGeomPrimvar>> PrimvarsUsedByAssignedMaterialsPerUVIndex;
+	int32 Float2PrimvarIndex = 0;
+	while (Result.Num() < MaxNumPrimvars && Float2Primvars.IsValidIndex(Float2PrimvarIndex))
 	{
-		const bool bProvideMaterialIndices = false;
-		for ( const FUsdPrimMaterialSlot& Slot : UsdMeshMaterialAssignmentInfo.Slots )
+		Result.Add(Float2Primvars[Float2PrimvarIndex++]);
+	}
+
+	return Result;
+}
+
+TArray<TUsdStore<pxr::UsdGeomPrimvar>> UsdUtils::GetUVSetPrimvars(
+	const pxr::UsdGeomMesh& UsdMesh,
+	const TMap<FString, TMap<FString, int32>>& MaterialToPrimvarsUVSetNames,
+	const pxr::TfToken& RenderContext,
+	const pxr::TfToken& MaterialPurpose
+)
+{
+	return UsdUtils::GetUVSetPrimvars(UsdMesh);
+}
+
+TArray<TUsdStore<pxr::UsdGeomPrimvar>> UsdUtils::GetUVSetPrimvars(
+	const pxr::UsdGeomMesh& UsdMesh,
+	const TMap<FString, TMap<FString, int32>>& MaterialToPrimvarsUVSetNames,
+	const UsdUtils::FUsdPrimMaterialAssignmentInfo& UsdMeshMaterialAssignmentInfo
+)
+{
+	return UsdUtils::GetUVSetPrimvars(UsdMesh);
+}
+
+TArray<TUsdStore<pxr::UsdGeomPrimvar>> UsdUtils::AssemblePrimvarsIntoUVSets(
+	const TArray<TUsdStore<pxr::UsdGeomPrimvar>>& AllMeshUVPrimvars,
+	const TMap<FString, int32>& AllowedPrimvarsToUVIndex
+)
+{
+	TArray<TUsdStore<pxr::UsdGeomPrimvar>> PrimvarsByUVIndex;
+
+	if (AllowedPrimvarsToUVIndex.Num() > 0)
+	{
+		for (const TUsdStore<pxr::UsdGeomPrimvar>& MeshUVPrimvar : AllMeshUVPrimvars)
 		{
-			if ( Slot.AssignmentType == EPrimAssignmentType::MaterialPrim )
+			FString PrimvarName = UsdToUnreal::ConvertToken(MeshUVPrimvar.Get().GetName());
+			PrimvarName.RemoveFromStart(TEXT("primvars:"));
+
+			if (const int32* FoundTargetUVIndex = AllowedPrimvarsToUVIndex.Find(PrimvarName))
 			{
-				const FString& MaterialPath = Slot.MaterialSource;
-				if ( const TMap< FString, int32 >* FoundMaterialPrimvars = MaterialToPrimvarsUVSetNames.Find( MaterialPath ) )
+				int32 TargetUVIndex = *FoundTargetUVIndex;
+				if (TargetUVIndex < 0)
 				{
-					for ( const TPair<FString, int32>& PrimvarAndUVIndex : *FoundMaterialPrimvars )
+					continue;
+				}
+
+				if (!PrimvarsByUVIndex.IsValidIndex(TargetUVIndex))
+				{
+					if (TargetUVIndex < MAX_MESH_TEXTURE_COORDS_MD)
 					{
-						if ( pxr::UsdGeomPrimvar* FoundPrimvar = PrimvarsByName.Find( PrimvarAndUVIndex.Key ) )
-						{
-							PrimvarsUsedByAssignedMaterialsPerUVIndex.FindOrAdd( PrimvarAndUVIndex.Value ).AddUnique( *FoundPrimvar );
-						}
+						PrimvarsByUVIndex.SetNum(TargetUVIndex + 1);
+					}
+					else
+					{
+						continue;
 					}
 				}
+
+				TUsdStore<pxr::UsdGeomPrimvar>& ExistingPrimvar = PrimvarsByUVIndex[TargetUVIndex];
+				if (!ExistingPrimvar.Get())
+				{
+					PrimvarsByUVIndex[TargetUVIndex] = MeshUVPrimvar;
+				}
 			}
 		}
 	}
 
-	// Sort all primvars we found by name, so we get consistent results
-	for ( TPair<int32, TArray<pxr::UsdGeomPrimvar>>& UVIndexToPrimvars : UsablePrimvarsByUVIndex )
-	{
-		TArray<pxr::UsdGeomPrimvar>& Primvars = UVIndexToPrimvars.Value;
-		Primvars.Sort( []( const pxr::UsdGeomPrimvar& A, const pxr::UsdGeomPrimvar& B )
-		{
-			return A.GetName() < B.GetName();
-		} );
-	}
-	for ( TPair<int32, TArray<pxr::UsdGeomPrimvar>>& UVIndexToPrimvars : PrimvarsUsedByAssignedMaterialsPerUVIndex )
+	return PrimvarsByUVIndex;
+}
+
+TMap<FString, int32> UsdUtils::AssemblePrimvarsIntoPrimvarToUVIndexMap(
+	const TArray<TUsdStore<pxr::UsdGeomPrimvar>>& AllMeshUVPrimvars
+)
+{
+	TMap<FString, int32> Result;
+	Result.Reserve(AllMeshUVPrimvars.Num());
+
+	for (int32 UVIndex = 0; UVIndex < AllMeshUVPrimvars.Num(); ++UVIndex)
 	{
-		TArray<pxr::UsdGeomPrimvar>& Primvars = UVIndexToPrimvars.Value;
-		Primvars.Sort( []( const pxr::UsdGeomPrimvar& A, const pxr::UsdGeomPrimvar& B )
-		{
-			return A.GetName() < B.GetName();
-		} );
+		const TUsdStore<pxr::UsdGeomPrimvar>& Primvar = AllMeshUVPrimvars[UVIndex];
+		FString PrimvarName = UsdToUnreal::ConvertToken(Primvar.Get().GetName());
+		PrimvarName.RemoveFromStart(TEXT("primvars:"));
+
+		Result.Add(PrimvarName, UVIndex);
 	}
 
-	// A lot of things can go wrong, so show some feedback in case they do
-	FString MeshPath = UsdToUnreal::ConvertPath( UsdMesh.GetPrim().GetPath() );
-	USDConversionUtilsImpl::CheckUVSetPrimvars( UsablePrimvarsByUVIndex, PrimvarsUsedByAssignedMaterialsPerUVIndex, MeshPath );
+	return Result;
+}
 
-	// Assemble our final results by picking the best primvar we can find for each UV index.
-	// Note that we should keep searching even if we don't find our ideal case, because we don't
-	// want to just discard potential UV sets if the material we happen to have bound doesn't use them, as the
-	// user may just want to assign another material that does.
-	TArray< TUsdStore< pxr::UsdGeomPrimvar > > Result;
-	Result.SetNum( MAX_STATIC_TEXCOORDS );
-	for ( int32 UVIndex = 0; UVIndex < MAX_STATIC_TEXCOORDS; ++UVIndex )
-	{
-		// Best case scenario: float2[]-like primvars that are actually being used by texture readers as texcoords, regardless of role
-		TArray<pxr::UsdGeomPrimvar>* FoundUsedPrimvars = PrimvarsUsedByAssignedMaterialsPerUVIndex.Find( UVIndex );
-		if ( FoundUsedPrimvars && FoundUsedPrimvars->Num() > 0 )
-		{
-			Result[ UVIndex ] = ( *FoundUsedPrimvars )[ 0 ];
-			continue;
-		}
+TMap<FString, int32> UsdUtils::CombinePrimvarsIntoUVSets(
+	const TSet<FString>& AllPrimvars,
+	const TSet<FString>& PreferredPrimvars
+)
+{
+	TArray<FString> SortedPrimvars = AllPrimvars.Array();
 
-		TArray<pxr::UsdGeomPrimvar>* FoundUsablePrimvars = UsablePrimvarsByUVIndex.Find( UVIndex );
-		if ( FoundUsablePrimvars && FoundUsablePrimvars->Num() > 0 )
+	// Promote a deterministic primvar-to-UV-index assignment preferring texCoord2f primvars
+	SortedPrimvars.Sort(
+		[&PreferredPrimvars](const FString& LHS, const FString& RHS)
 		{
-			pxr::UsdGeomPrimvar* FoundTexcoordPrimvar = FoundUsablePrimvars->FindByPredicate( []( const pxr::UsdGeomPrimvar& Primvar )
+			const bool bLHSPreferred = PreferredPrimvars.Contains(LHS);
+			const bool bRHSPreferred = PreferredPrimvars.Contains(RHS);
+			if (bLHSPreferred == bRHSPreferred)
 			{
-				return Primvar.GetTypeName().GetRole() == pxr::SdfValueTypeNames->TexCoord2f.GetRole();
-			} );
-
-			// Second-best case: Primvars with texcoord2f role
-			if ( FoundTexcoordPrimvar )
+				return LHS < RHS;
+			}
+			else
 			{
-				Result[ UVIndex ] = *FoundTexcoordPrimvar;
-				continue;
+				return bLHSPreferred < bRHSPreferred;
 			}
-
-			// Third-best case: Any valid primvar
-			// Disabled for now as these could be any other random data the user may have as float2[]
-			// Result[UVIndex] = FoundUsedPrimvars[0];
 		}
+	);
+
+	// We can only have up to MAX_MESH_TEXTURE_COORDS_MD UV sets
+	SortedPrimvars.SetNum(FMath::Min(SortedPrimvars.Num(), (int32)MAX_MESH_TEXTURE_COORDS_MD));
+
+	TMap<FString, int32> PrimvarToUVIndex;
+	PrimvarToUVIndex.Reserve(SortedPrimvars.Num());
+	int32 UVIndex = 0;
+	for (const FString& Primvar : SortedPrimvars)
+	{
+		PrimvarToUVIndex.Add(Primvar, UVIndex++);
 	}
 
-	return Result;
+	return PrimvarToUVIndex;
 }
 
 bool UsdUtils::IsAnimated( const pxr::UsdPrim& Prim )
diff --git a/Engine/Plugins/Importers/USDImporter/Source/USDUtilities/Private/USDGeomMeshConversion.cpp b/Engine/Plugins/Importers/USDImporter/Source/USDUtilities/Private/USDGeomMeshConversion.cpp
--- a/Engine/Plugins/Importers/USDImporter/Source/USDUtilities/Private/USDGeomMeshConversion.cpp
+++ b/Engine/Plugins/Importers/USDImporter/Source/USDUtilities/Private/USDGeomMeshConversion.cpp
@@ -617,6 +617,135 @@ namespace UE::UsdGeomMeshConversion::Private
 
 		return bSuccess;
 	}
+
+	void RecursivelyCollectPrimvars(
+		const pxr::UsdPrim& Prim,
+		const UsdToUnreal::FUsdMeshConversionOptions& Options,
+		TSet<FString>& InOutAllPrimvars,
+		TSet<FString>& InOutPreferredPrimvars,
+		bool bIsFirstPrim
+	)
+	{
+		FScopedUsdAllocs Allocs;
+
+		// This should always replicate the same traversal pattern of RecursivelyCollapseChildMeshes
+
+		if (!EnumHasAllFlags(Options.PurposesToLoad, IUsdPrim::GetPurpose(Prim)))
+		{
+			return;
+		}
+
+		if (!bIsFirstPrim)
+		{
+			if (pxr::UsdGeomImageable UsdGeomImageable = pxr::UsdGeomImageable(Prim))
+			{
+				if (pxr::UsdAttribute VisibilityAttr = UsdGeomImageable.GetVisibilityAttr())
+				{
+					pxr::TfToken VisibilityToken;
+					if (VisibilityAttr.Get(&VisibilityToken) && VisibilityToken == pxr::UsdGeomTokens->invisible)
+					{
+						return;
+					}
+				}
+			}
+		}
+
+		bool bTraverseChildren = true;
+
+		if (pxr::UsdGeomMesh Mesh = pxr::UsdGeomMesh(Prim))
+		{
+			TArray<TUsdStore<pxr::UsdGeomPrimvar>> MeshPrimvars = UsdUtils::GetUVSetPrimvars(
+				Mesh,
+				TNumericLimits<int32>::Max()
+			);
+
+			for (const TUsdStore<pxr::UsdGeomPrimvar>& MeshPrimvar : MeshPrimvars)
+			{
+				FString PrimvarName = UsdToUnreal::ConvertToken(MeshPrimvar.Get().GetName());
+				PrimvarName.RemoveFromStart(TEXT("primvars:"));
+
+				InOutAllPrimvars.Add(PrimvarName);
+
+				// Keep track of which primvars are texCoord2f as we always want to prefer these over other float2s
+				if (MeshPrimvar.Get().GetTypeName().GetRole() == pxr::SdfValueTypeNames->TexCoord2f.GetRole())
+				{
+					InOutPreferredPrimvars.Add(PrimvarName);
+				}
+			}
+		}
+		else if (pxr::UsdGeomPointInstancer PointInstancer = pxr::UsdGeomPointInstancer{Prim})
+		{
+			pxr::SdfPathVector PrototypePaths;
+			if (!PointInstancer.GetPrototypesRel().GetTargets(&PrototypePaths))
+			{
+				return;
+			}
+
+			pxr::UsdStageRefPtr Stage = Prim.GetStage();
+			for (const pxr::SdfPath& PrototypePath : PrototypePaths)
+			{
+				pxr::UsdPrim PrototypeUsdPrim = Stage->GetPrimAtPath(PrototypePath);
+				if (!PrototypeUsdPrim)
+				{
+					continue;
+				}
+
+				const bool bChildIsFirstPrim = false;
+				RecursivelyCollectPrimvars(
+					PrototypeUsdPrim,
+					Options,
+					InOutAllPrimvars,
+					InOutPreferredPrimvars,
+					bChildIsFirstPrim
+				);
+			}
+
+			// We never want to step into point instancers when fetching prims for drawing
+			bTraverseChildren = false;
+		}
+
+		if (bTraverseChildren)
+		{
+			for (const pxr::UsdPrim& ChildPrim : Prim.GetFilteredChildren(pxr::UsdTraverseInstanceProxies()))
+			{
+				const bool bChildIsFirstPrim = false;
+				RecursivelyCollectPrimvars(
+					ChildPrim,
+					Options,
+					InOutAllPrimvars,
+					InOutPreferredPrimvars,
+					bChildIsFirstPrim
+				);
+			}
+		}
+	}
+
+	/**
+	 * Returns the set of primvar names that can be used for each UV index for a mesh collapsed from the subtree
+	 * starting at RootPrim.
+	 */
+	TMap<FString, int32> CollectSubtreePrimvars(
+		const pxr::UsdPrim& RootPrim,
+		const UsdToUnreal::FUsdMeshConversionOptions& Options,
+		bool bIsFirstPrim
+	)
+	{
+		TSet<FString> AllPrimvars;
+		TSet<FString> PreferredPrimvars;
+
+		RecursivelyCollectPrimvars(
+			RootPrim,
+			Options,
+			AllPrimvars,
+			PreferredPrimvars,
+			bIsFirstPrim
+		);
+
+		return UsdUtils::CombinePrimvarsIntoUVSets(
+			AllPrimvars,
+			PreferredPrimvars
+		);
+	}
 }
 namespace UsdGeomMeshImpl = UE::UsdGeomMeshConversion::Private;
 
@@ -630,7 +759,6 @@ namespace UsdToUnreal
 		, RenderContext( pxr::UsdShadeTokens->universalRenderContext )
 		, MaterialPurpose( pxr::UsdShadeTokens->allPurpose )
 		, TimeCode( pxr::UsdTimeCode::EarliestTime() )
-		, MaterialToPrimvarToUVIndex( nullptr )
 		, bMergeIdenticalMaterialSlots( true )
 	{
 	}
@@ -869,20 +997,41 @@ bool UsdToUnreal::ConvertGeomMesh(
 
 		TArray< FUVSet > UVSets;
 
-		TArray< TUsdStore< pxr::UsdGeomPrimvar > > PrimvarsByUVIndex = UsdUtils::GetUVSetPrimvars( UsdMesh, *Options.MaterialToPrimvarToUVIndex, LocalInfo );
+		TOptional<int32> ProvidedNumUVSets;
 
-		int32 HighestAddedUVChannel = 0;
-		for ( int32 UVChannelIndex = 0; UVChannelIndex < PrimvarsByUVIndex.Num(); ++UVChannelIndex )
+		// If we already have a primvar to UV index assignment, let's just use that.
+		// When collapsing, we'll do a pre-pass on all meshes to translate and determine this beforehand.
+		TArray<TUsdStore<pxr::UsdGeomPrimvar>> PrimvarsByUVIndex;
+		if (OutMaterialAssignments.PrimvarToUVIndex.Num() > 0)
 		{
-			if ( !PrimvarsByUVIndex.IsValidIndex( UVChannelIndex ) )
+			int32 HighestProvidedUVIndex = 0;
+			for (const TPair<FString, int32>& Pair : OutMaterialAssignments.PrimvarToUVIndex)
 			{
-				break;
+				HighestProvidedUVIndex = FMath::Max(HighestProvidedUVIndex, Pair.Value);
 			}
+			ProvidedNumUVSets = HighestProvidedUVIndex + 1;
 
-			pxr::UsdGeomPrimvar& Primvar = PrimvarsByUVIndex[ UVChannelIndex ].Get();
-			if ( !Primvar )
+			TArray<TUsdStore<pxr::UsdGeomPrimvar>> AllMeshUVPrimvars =
+				UsdUtils::GetUVSetPrimvars(UsdMesh, TNumericLimits<int32>::Max());
+
+			PrimvarsByUVIndex =
+				UsdUtils::AssemblePrimvarsIntoUVSets(AllMeshUVPrimvars, OutMaterialAssignments.PrimvarToUVIndex);
+		}
+		// Let's use the best primvar assignment for this particular mesh instead
+		else
+		{
+			PrimvarsByUVIndex = UsdUtils::GetUVSetPrimvars(UsdMesh);
+
+			OutMaterialAssignments.PrimvarToUVIndex =
+				UsdUtils::AssemblePrimvarsIntoPrimvarToUVIndexMap(PrimvarsByUVIndex);
+		}
+
+		int32 HighestAddedUVChannel = 0;
+		for (int32 UVChannelIndex = 0; UVChannelIndex < PrimvarsByUVIndex.Num(); ++UVChannelIndex)
+		{
+			pxr::UsdGeomPrimvar& Primvar = PrimvarsByUVIndex[UVChannelIndex].Get();
+			if (!Primvar)
 			{
-				// The user may have name their UV sets 'uv4' and 'uv5', in which case we have no UV sets below 4, so just skip them
 				continue;
 			}
 
@@ -916,13 +1065,25 @@ bool UsdToUnreal::ConvertGeomMesh(
 			}
 		}
 
-		// When importing multiple mesh pieces to the same static mesh.  Ensure each mesh piece has the same number of Uv's
+		// When importing multiple mesh pieces to the same static mesh.  Ensure each mesh piece has the same number of UVs
 		{
 			int32 ExistingUVCount = MeshDescriptionUVs.GetNumChannels();
 			int32 NumUVs = FMath::Max( HighestAddedUVChannel + 1, ExistingUVCount );
+
+			// When we provide a PrimvarToUVIndex map to this function it means we'll end up combining this
+			// MeshDescription with others later (e.g. due to collapsing or multiple-LOD meshes).
+			// In that case we can get better results by making sure all of the individual MeshDescriptions have the
+			// same total number of UV sets, even if the unused ones are empty.
+			// Otherwise, if we e.g. have a material reading UVIndex3 when we only have a single UV set, UE seems to
+			// just read that one UV set anyway, which is somewhat unexpected and can be misleading
+			if (ProvidedNumUVSets.IsSet())
+			{
+				NumUVs = FMath::Max<int32>(ProvidedNumUVSets.GetValue(), NumUVs);
+			}
+
 			NumUVs = FMath::Min<int32>( MAX_MESH_TEXTURE_COORDS_MD, NumUVs );
 			// At least one UV set must exist.
-			NumUVs = FMath::Max<int32>( 1, NumUVs );
+			NumUVs = FMath::Max<int32>(1, NumUVs);
 
 			//Make sure all Vertex instance have the correct number of UVs
 			MeshDescriptionUVs.SetNumChannels( NumUVs );
@@ -1152,7 +1313,6 @@ bool UsdToUnreal::ConvertGeomMesh(
 {
 	FUsdMeshConversionOptions Options;
 	Options.AdditionalTransform = AdditionalTransform;
-	Options.MaterialToPrimvarToUVIndex = &MaterialToPrimvarsUVSetNames;
 	Options.TimeCode = TimeCode;
 	Options.RenderContext = RenderContext;
 	Options.bMergeIdenticalMaterialSlots = bMergeIdenticalMaterialSlots;
@@ -1367,6 +1527,16 @@ bool UsdToUnreal::ConvertGeomMeshHierarchy(
 	// provide the options object to ConvertGeomMesh and ConvertPointInstancerToMesh
 	FUsdMeshConversionOptions OptionsCopy = Options;
 
+	// Prepass to figure out the best primvars to use for the entire collapsed mesh UV sets
+	if (OutMaterialAssignments.PrimvarToUVIndex.Num() == 0)
+	{
+		OutMaterialAssignments.PrimvarToUVIndex = UsdGeomMeshImpl::CollectSubtreePrimvars(
+			Prim,
+			Options,
+			bSkipRootPrimTransformAndVisibility
+		);
+	}
+
 	return UsdGeomMeshImpl::RecursivelyCollapseChildMeshes(
 		Prim,
 		OutMeshDescription,
@@ -1393,7 +1563,6 @@ bool UsdToUnreal::ConvertGeomMeshHierarchy(
 	Options.TimeCode = TimeCode;
 	Options.PurposesToLoad = PurposesToLoad;
 	Options.RenderContext = RenderContext;
-	Options.MaterialToPrimvarToUVIndex = &MaterialToPrimvarToUVIndex;
 	Options.bMergeIdenticalMaterialSlots = bMergeIdenticalMaterialSlots;
 
 	return ConvertGeomMeshHierarchy(
diff --git a/Engine/Plugins/Importers/USDImporter/Source/USDUtilities/Private/USDShadeConversion.cpp b/Engine/Plugins/Importers/USDImporter/Source/USDUtilities/Private/USDShadeConversion.cpp
--- a/Engine/Plugins/Importers/USDImporter/Source/USDUtilities/Private/USDShadeConversion.cpp
+++ b/Engine/Plugins/Importers/USDImporter/Source/USDUtilities/Private/USDShadeConversion.cpp
@@ -431,29 +431,10 @@ namespace UE
 				return {};
 			}
 
-			void CheckForUVIndexConflicts( TMap<FString, int32>& PrimvarToUVIndex, const FString& ShadeMaterialPath )
-			{
-				TMap<int32, TArray<FString>> InvertedMap;
-				for ( const TPair< FString, int32 >& PrimvarAndUVIndex : PrimvarToUVIndex )
-				{
-					InvertedMap.FindOrAdd( PrimvarAndUVIndex.Value ).Add( PrimvarAndUVIndex.Key );
-				}
-
-				for ( const TPair< int32, TArray<FString> >& UVIndexAndPrimvars : InvertedMap )
-				{
-					const TArray<FString>& Primvars = UVIndexAndPrimvars.Value;
-					if ( Primvars.Num() > 1 )
-					{
-						FString PrimvarNames = FString::Join(Primvars, TEXT(", "));
-						UE_LOG(LogUsd, Warning, TEXT("For shade material '%s', multiple primvars (%s) target the same UV index '%d'"), *ShadeMaterialPath, *PrimvarNames, UVIndexAndPrimvars.Key);
-					}
-				}
-			}
-
 			struct FTextureParameterValue
 			{
 				UTexture* Texture;
-				int32 UVIndex;
+				FString Primvar;
 				int32 OutputIndex = 0;
 			};
 			struct FPrimvarReaderParameterValue
@@ -463,8 +444,13 @@ namespace UE
 			};
 			using FParameterValue = TVariant<float, FVector, FTextureParameterValue, FPrimvarReaderParameterValue, bool>;
 
-			bool GetTextureParameterValue( pxr::UsdShadeInput& ShadeInput, TextureGroup LODGroup, FParameterValue& OutValue,
-				UMaterialInterface* Material = nullptr, UUsdAssetCache2* TexturesCache = nullptr, TMap<FString, int32>* PrimvarToUVIndex = nullptr )
+			bool GetTextureParameterValue(
+				pxr::UsdShadeInput& ShadeInput,
+				TextureGroup LODGroup,
+				FParameterValue& OutValue,
+				UMaterialInterface* Material = nullptr,
+				UUsdAssetCache2* TexturesCache = nullptr
+			)
 			{
 				FScopedUsdAllocs UsdAllocs;
 
@@ -638,28 +624,13 @@ namespace UE
 
 							if ( Texture )
 							{
-								FString PrimvarName = GetPrimvarUsedAsST( UsdUVTextureSource );
-								int32 UVIndex = 0;
-								if ( !PrimvarName.IsEmpty() && PrimvarToUVIndex )
-								{
-									UVIndex = UsdUtils::GetPrimvarUVIndex(PrimvarName);
-									PrimvarToUVIndex->Add(PrimvarName, UVIndex);
-								}
-								else if ( PrimvarToUVIndex )
-								{
-									FUsdLogManager::LogMessage(
-										EMessageSeverity::Warning,
-										FText::Format( LOCTEXT( "FailedToParsePrimVar", "Failed to find primvar used as st input for texture '{0}' in material '{1}'. Will use UV index 0 instead" ),
-											FText::FromString( TexturePath ),
-											FText::FromString( Material ? *Material->GetName() : TEXT( "nullptr" ) )
-										)
-									);
-								}
-
 								FTextureParameterValue OutTextureValue;
 								OutTextureValue.Texture = Texture;
-								OutTextureValue.UVIndex = UVIndex;
+								OutTextureValue.Primvar = GetPrimvarUsedAsST(UsdUVTextureSource);
 
+								// The UsdUVTexture Shader itself prim has multiple standard outputs, but this is not
+								// full swizzle support (check
+								// [REDACTED_PROFILE_URL]/USD/blob/5c5ebddff35012461a2b0ba773c47f05730cbab4/pxr/usdImaging/plugin/usdShaders/shaders/shaderDefs.usda#L198)
 								if ( UsdUVTextureAttributeType == pxr::UsdShadeAttributeType::Output )
 								{
 									const FName OutputName( UsdToUnreal::ConvertToken( UsdUVTextureSourceName ) );
@@ -708,8 +679,14 @@ namespace UE
 				return OutValue.IsType<FTextureParameterValue>();
 			}
 
-			bool GetFloatParameterValue( pxr::UsdShadeConnectableAPI& Connectable, const pxr::TfToken& InputName, float DefaultValue, FParameterValue& OutValue,
-				UMaterialInterface* Material = nullptr, UUsdAssetCache2* TexturesCache = nullptr, TMap<FString, int32>* PrimvarToUVIndex = nullptr)
+			bool GetFloatParameterValue(
+				pxr::UsdShadeConnectableAPI& Connectable,
+				const pxr::TfToken& InputName,
+				float DefaultValue,
+				FParameterValue& OutValue,
+				UMaterialInterface* Material = nullptr,
+				UUsdAssetCache2* TexturesCache = nullptr
+			)
 			{
 				FScopedUsdAllocs Allocs;
 
@@ -725,7 +702,7 @@ namespace UE
 				pxr::UsdShadeAttributeType AttributeType;
 				if ( pxr::UsdShadeConnectableAPI::GetConnectedSource( Input.GetAttr(), &Source, &SourceName, &AttributeType ) )
 				{
-					if ( !GetTextureParameterValue( Input, TEXTUREGROUP_WorldSpecular, OutValue, Material, TexturesCache, PrimvarToUVIndex ) )
+					if (!GetTextureParameterValue(Input, TEXTUREGROUP_WorldSpecular, OutValue, Material, TexturesCache))
 					{
 						// Check if we have a fallback input that we can use instead, since we don't have a valid texture value
 						if ( const pxr::UsdShadeInput FallbackInput = Source.GetInput( UnrealIdentifiers::Fallback ) )
@@ -740,7 +717,7 @@ namespace UE
 
 						// Recurse because the attribute may just be pointing at some other attribute that has the data
 						// (e.g. when shader input is just "hoisted" and connected to the parent material input)
-						return GetFloatParameterValue( Source, SourceName, DefaultValue, OutValue, Material, TexturesCache, PrimvarToUVIndex );
+						return GetFloatParameterValue(Source, SourceName, DefaultValue, OutValue, Material, TexturesCache);
 					}
 				}
 				// No other node connected, so we must have some value
@@ -827,8 +804,15 @@ namespace UE
 				return false;
 			}
 
-			bool GetVec3ParameterValue( pxr::UsdShadeConnectableAPI& Connectable, const pxr::TfToken& InputName, const FLinearColor& DefaultValue, FParameterValue& OutValue,
-				bool bIsNormalMap = false, UMaterialInterface* Material = nullptr, UUsdAssetCache2* TexturesCache = nullptr, TMap<FString, int32>* PrimvarToUVIndex = nullptr )
+			bool GetVec3ParameterValue(
+				pxr::UsdShadeConnectableAPI& Connectable,
+				const pxr::TfToken& InputName,
+				const FLinearColor& DefaultValue,
+				FParameterValue& OutValue,
+				bool bIsNormalMap = false,
+				UMaterialInterface* Material = nullptr,
+				UUsdAssetCache2* TexturesCache = nullptr
+			)
 			{
 				FScopedUsdAllocs Allocs;
 
@@ -844,7 +828,7 @@ namespace UE
 				pxr::UsdShadeAttributeType AttributeType;
 				if ( pxr::UsdShadeConnectableAPI::GetConnectedSource( Input.GetAttr(), &Source, &SourceName, &AttributeType ) )
 				{
-					if ( !GetTextureParameterValue( Input, bIsNormalMap ? TEXTUREGROUP_WorldNormalMap : TEXTUREGROUP_World, OutValue, Material, TexturesCache, PrimvarToUVIndex ) )
+					if (!GetTextureParameterValue(Input, bIsNormalMap ? TEXTUREGROUP_WorldNormalMap : TEXTUREGROUP_World, OutValue, Material, TexturesCache))
 					{
 						// Check whether this input receives its value through a connection to a
 						// primvar reader shader.
@@ -879,7 +863,7 @@ namespace UE
 						}
 
 						// This shader doesn't have anything: Traverse into the input connectable itself
-						return GetVec3ParameterValue( Source, SourceName, DefaultValue, OutValue, bIsNormalMap, Material, TexturesCache, PrimvarToUVIndex );
+						return GetVec3ParameterValue(Source, SourceName, DefaultValue, OutValue, bIsNormalMap, Material, TexturesCache);
 					}
 				}
 				// No other node connected, so we must have some value
@@ -898,7 +882,14 @@ namespace UE
 				return true;
 			}
 
-			bool GetBoolParameterValue( pxr::UsdShadeConnectableAPI& Connectable, const pxr::TfToken& InputName, bool DefaultValue, FParameterValue& OutValue, UMaterialInterface* Material = nullptr, UUsdAssetCache2* TexturesCache = nullptr, TMap<FString, int32>* PrimvarToUVIndex = nullptr)
+			bool GetBoolParameterValue(
+				pxr::UsdShadeConnectableAPI& Connectable,
+				const pxr::TfToken& InputName,
+				bool DefaultValue,
+				FParameterValue& OutValue,
+				UMaterialInterface* Material = nullptr,
+				UUsdAssetCache2* TexturesCache = nullptr
+			)
 			{
 				FScopedUsdAllocs Allocs;
 
@@ -916,7 +907,7 @@ namespace UE
 				{
 					// Recurse because the attribute may just be pointing at some other attribute that has the data
 					// (e.g. when shader input is just "hoisted" and connected to the parent material input)
-					return GetBoolParameterValue( Source, SourceName, DefaultValue, OutValue, Material, TexturesCache, PrimvarToUVIndex );
+					return GetBoolParameterValue(Source, SourceName, DefaultValue, OutValue, Material, TexturesCache);
 				}
 				// No other node connected, so we must have some value
 				else
@@ -962,7 +953,6 @@ namespace UE
 					UMaterialExpressionTextureSample* Expression = Cast< UMaterialExpressionTextureSample >( UMaterialEditingLibrary::CreateMaterialExpression( &Material, UMaterialExpressionTextureSample::StaticClass() ) );
 					Expression->Texture = TextureValue.Texture;
 					Expression->SamplerType = UMaterialExpressionTextureBase::GetSamplerTypeForTexture( TextureValue.Texture );
-					Expression->ConstCoordinate = TextureValue.UVIndex;
 
 					return Expression;
 				}
@@ -1006,102 +996,6 @@ namespace UE
 #endif // WITH_EDITOR
 			}
 
-			void SetScalarParameterValue( UMaterialInstance& Material, const TCHAR* ParameterName, float ParameterValue )
-			{
-				FMaterialParameterInfo Info;
-				Info.Name = ParameterName;
-
-				if ( UMaterialInstanceDynamic* Dynamic = Cast<UMaterialInstanceDynamic>( &Material ) )
-				{
-					Dynamic->SetScalarParameterValueByInfo( Info, ParameterValue );
-				}
-#if WITH_EDITOR
-				else if ( UMaterialInstanceConstant* Constant = Cast<UMaterialInstanceConstant>( &Material ) )
-				{
-					Constant->SetScalarParameterValueEditorOnly( Info, ParameterValue );
-				}
-#endif // WITH_EDITOR
-			}
-
-			void SetVectorParameterValue( UMaterialInstance& Material, const TCHAR* ParameterName, FLinearColor ParameterValue )
-			{
-				FMaterialParameterInfo Info;
-				Info.Name = ParameterName;
-
-				if ( UMaterialInstanceDynamic* Dynamic = Cast<UMaterialInstanceDynamic>( &Material ) )
-				{
-					Dynamic->SetVectorParameterValueByInfo( Info, ParameterValue );
-				}
-#if WITH_EDITOR
-				else if ( UMaterialInstanceConstant* Constant = Cast<UMaterialInstanceConstant>( &Material ) )
-				{
-					Constant->SetVectorParameterValueEditorOnly( Info, ParameterValue );
-				}
-#endif // WITH_EDITOR
-			}
-
-			void SetTextureParameterValue( UMaterialInstance& Material, const TCHAR* ParameterName, UTexture* ParameterValue )
-			{
-				FMaterialParameterInfo Info;
-				Info.Name = ParameterName;
-
-				if ( UMaterialInstanceDynamic* Dynamic = Cast<UMaterialInstanceDynamic>( &Material ) )
-				{
-					Dynamic->SetTextureParameterValueByInfo( Info, ParameterValue );
-				}
-#if WITH_EDITOR
-				else if ( UMaterialInstanceConstant* Constant = Cast<UMaterialInstanceConstant>( &Material ) )
-				{
-					Constant->SetTextureParameterValueEditorOnly( Info, ParameterValue );
-				}
-#endif // WITH_EDITOR
-			}
-
-			void SetBoolParameterValue( UMaterialInstance& Material, const TCHAR* ParameterName, bool bParameterValue )
-			{
-				bool bFound = false;
-
-#if WITH_EDITOR
-				// Try the static parameters first
-				if ( UMaterialInstanceConstant* Constant = Cast<UMaterialInstanceConstant>( &Material ) )
-				{
-					bool bNeedsUpdatePermutations = false;
-
-					FStaticParameterSet StaticParameters;
-					Constant->GetStaticParameterValues( StaticParameters );
-
-					for ( FStaticSwitchParameter& StaticSwitchParameter : StaticParameters.StaticSwitchParameters )
-					{
-						if ( StaticSwitchParameter.ParameterInfo.Name == ParameterName )
-						{
-							bFound = true;
-
-							if ( StaticSwitchParameter.Value != bParameterValue )
-							{
-								StaticSwitchParameter.Value = bParameterValue;
-								StaticSwitchParameter.bOverride = true;
-								bNeedsUpdatePermutations = true;
-							}
-
-							break;
-						}
-					}
-
-					if ( bNeedsUpdatePermutations )
-					{
-						FlushRenderingCommands();
-						Constant->UpdateStaticPermutation( StaticParameters );
-					}
-				}
-#endif // WITH_EDITOR
-
-				// Try it as a scalar parameter
-				if ( !bFound )
-				{
-					SetScalarParameterValue( Material, ParameterName, bParameterValue ? 1.0f : 0.0f );
-				}
-			}
-
 			/**
 			 * Intended to be used with Visit(), this will set an FParameterValue into InMaterial using InParameterName,
 			 * depending on the variant active in FParameterValue
@@ -1115,27 +1009,27 @@ namespace UE
 
 				void operator()( const float FloatValue ) const
 				{
-					SetScalarParameterValue( Material, ParameterName, FloatValue );
+					UsdUtils::SetScalarParameterValue( Material, ParameterName, FloatValue );
 				}
 
 				void operator()( const FVector& VectorValue ) const
 				{
-					SetVectorParameterValue( Material, ParameterName, FLinearColor(VectorValue) );
+					UsdUtils::SetVectorParameterValue( Material, ParameterName, FLinearColor(VectorValue) );
 				}
 
 				void operator()( const FTextureParameterValue& TextureValue ) const
 				{
-					SetTextureParameterValue( Material, ParameterName, TextureValue.Texture );
+					UsdUtils::SetTextureParameterValue( Material, ParameterName, TextureValue.Texture );
 				}
 
 				void operator()( const FPrimvarReaderParameterValue& PrimvarReaderValue ) const
 				{
-					SetVectorParameterValue( Material, ParameterName, FLinearColor(PrimvarReaderValue.FallbackValue) );
+					UsdUtils::SetVectorParameterValue( Material, ParameterName, FLinearColor(PrimvarReaderValue.FallbackValue) );
 				}
 
 				void operator()( const bool BoolValue ) const
 				{
-					SetBoolParameterValue( Material, ParameterName, BoolValue );
+					UsdUtils::SetBoolParameterValue( Material, ParameterName, BoolValue );
 				}
 
 			protected:
@@ -1148,24 +1042,71 @@ namespace UE
 			 */
 			struct FSetPreviewSurfaceParameterValueVisitor : private FSetParameterValueVisitor
 			{
-				using FSetParameterValueVisitor::FSetParameterValueVisitor;
+				explicit FSetPreviewSurfaceParameterValueVisitor(
+					UMaterialInstance& InMaterial,
+					const TCHAR* InParameterName,
+					const TMap<FString, int32>& InPrimvarToUVIndex
+				)
+					: FSetParameterValueVisitor(InMaterial, InParameterName)
+					, PrimvarToUVIndex(InPrimvarToUVIndex)
+				{
+				}
 
 				void operator()( const float FloatValue ) const
 				{
 					FSetParameterValueVisitor::operator()( FloatValue );
-					SetScalarParameterValue( Material, *FString::Printf( TEXT( "Use%sTexture" ), ParameterName ), 0.0f );
+					UsdUtils::SetScalarParameterValue( Material, *FString::Printf( TEXT( "Use%sTexture" ), ParameterName ), 0.0f );
 				}
 
 				void operator()( const FVector& VectorValue ) const
 				{
 					FSetParameterValueVisitor::operator()( VectorValue );
-					SetScalarParameterValue( Material, *FString::Printf( TEXT( "Use%sTexture" ), ParameterName ), 0.0f );
+					UsdUtils::SetScalarParameterValue( Material, *FString::Printf( TEXT( "Use%sTexture" ), ParameterName ), 0.0f );
 				}
 
 				void operator()( const FTextureParameterValue& TextureValue ) const
 				{
-					SetTextureParameterValue( Material, *FString::Printf( TEXT( "%sTexture" ), ParameterName ), TextureValue.Texture );
-					SetScalarParameterValue( Material, *FString::Printf( TEXT( "Use%sTexture" ), ParameterName ), 1.0f );
+					UsdUtils::SetTextureParameterValue( Material, *FString::Printf( TEXT( "%sTexture" ), ParameterName ), TextureValue.Texture );
+					UsdUtils::SetScalarParameterValue( Material, *FString::Printf( TEXT( "Use%sTexture" ), ParameterName ), 1.0f );
+
+					if (const int32* FoundIndex = PrimvarToUVIndex.Find(TextureValue.Primvar))
+					{
+						UsdUtils::SetScalarParameterValue(Material, *FString::Printf(TEXT("%sUVIndex"), ParameterName), *FoundIndex);
+					}
+					else
+					{
+						ensureMsgf(false, TEXT("Failed to find primvar '%s' when setting material parameter '%s' on material '%s'. Available primvars and UV indices: %s"),
+							*TextureValue.Primvar,
+							ParameterName,
+							*Material.GetPathName(),
+							*UsdUtils::StringifyMap(PrimvarToUVIndex)
+						);
+					}
+
+					FLinearColor ComponentMask;
+					switch (TextureValue.OutputIndex)
+					{
+					case 0: // RGB
+						ComponentMask = FLinearColor{1.f, 1.f, 1.f, 0.f};
+						break;
+					case 1: // R
+						ComponentMask = FLinearColor{1.f, 0.f, 0.f, 0.f};
+						break;
+					case 2: // G
+						ComponentMask = FLinearColor{0.f, 1.f, 0.f, 0.f};
+						break;
+					case 3: // B
+						ComponentMask = FLinearColor{0.f, 0.f, 1.f, 0.f};
+						break;
+					case 4: // A
+						ComponentMask = FLinearColor{0.f, 0.f, 0.f, 1.f};
+						break;
+					}
+					UsdUtils::SetVectorParameterValue(
+						Material,
+						*FString::Printf(TEXT("%sTextureComponent"), ParameterName),
+						ComponentMask
+					);
 				}
 
 				void operator()( const FPrimvarReaderParameterValue& PrimvarReaderValue ) const
@@ -1176,21 +1117,30 @@ namespace UE
 					// on UE meshes, so that's the only primvar we can support here.
 					if ( PrimvarReaderValue.PrimvarName == "displayColor" )
 					{
-						SetScalarParameterValue( Material, TEXT( "UseVertexColorForBaseColor" ), 1.0f );
+						UsdUtils::SetScalarParameterValue( Material, TEXT( "UseVertexColorForBaseColor" ), 1.0f );
 					}
 				}
+
+			protected:
+				const TMap<FString, int32>& PrimvarToUVIndex;
 			};
 
-			void SetParameterValue( UMaterialInstance& Material, const TCHAR* ParameterName, const FParameterValue& ParameterValue, bool bForUsdPreviewSurface )
+			void SetParameterValue(
+				UMaterialInstance& Material,
+				const TCHAR* ParameterName,
+				const FParameterValue& ParameterValue,
+				bool bForUsdPreviewSurface,
+				const TMap<FString, int32>& PrimvarToUVIndex
+			)
 			{
 				if ( bForUsdPreviewSurface )
 				{
-					FSetPreviewSurfaceParameterValueVisitor Visitor{ Material, ParameterName };
+					FSetPreviewSurfaceParameterValueVisitor Visitor{Material, ParameterName, PrimvarToUVIndex};
 					Visit( Visitor, ParameterValue );
 				}
 				else
 				{
-					FSetParameterValueVisitor Visitor{ Material, ParameterName };
+					FSetParameterValueVisitor Visitor{Material, ParameterName};
 					Visit( Visitor, ParameterValue );
 				}
 			}
@@ -1963,337 +1913,451 @@ namespace UE
 }
 namespace UsdShadeConversionImpl = UE::UsdShadeConversion::Private;
 
-bool UsdToUnreal::ConvertMaterial( const pxr::UsdShadeMaterial& UsdShadeMaterial, UMaterialInstance& Material )
-{
-	TMap<FString, int32> PrimvarToUVIndex;
-	UUsdAssetCache2* NewCache = nullptr;
-	return ConvertMaterial( UsdShadeMaterial, Material, NewCache, PrimvarToUVIndex );
-}
-
-bool UsdToUnreal::ConvertMaterial( const pxr::UsdShadeMaterial& UsdShadeMaterial, UMaterialInstance& Material, UUsdAssetCache2* TexturesCache, TMap< FString, int32 >& PrimvarToUVIndex, const TCHAR* RenderContext )
+bool UsdToUnreal::ConvertMaterial(
+	const pxr::UsdShadeMaterial& UsdShadeMaterial,
+	UMaterialInstance& Material,
+	UUsdAssetCache2* TexturesCache,
+	const TCHAR* RenderContext
+)
 {
 	FScopedUsdAllocs UsdAllocs;
 
-	bool bHasMaterialInfo = false;
-
 	pxr::TfToken RenderContextToken = pxr::UsdShadeTokens->universalRenderContext;
-	if ( RenderContext )
+	if (RenderContext)
 	{
-		pxr::TfToken ProvidedRenderContextToken = UnrealToUsd::ConvertToken( RenderContext ).Get();
+		pxr::TfToken ProvidedRenderContextToken = UnrealToUsd::ConvertToken(RenderContext).Get();
 
 		// We won't ever create an asset for the Unreal render context material prim.
 		// Check the universal context here in case this material should also generate an asset for that context too
-		if ( ProvidedRenderContextToken != UnrealIdentifiers::Unreal )
+		if (ProvidedRenderContextToken != UnrealIdentifiers::Unreal)
 		{
 			RenderContextToken = ProvidedRenderContextToken;
 		}
 	}
 
-	pxr::UsdShadeShader SurfaceShader = UsdShadeMaterial.ComputeSurfaceSource( RenderContextToken );
-	if ( !SurfaceShader )
+	pxr::UsdShadeShader SurfaceShader = UsdShadeMaterial.ComputeSurfaceSource(RenderContextToken);
+	if (!SurfaceShader)
 	{
 		return false;
 	}
 
-	pxr::UsdShadeConnectableAPI Connectable{ SurfaceShader };
+	pxr::UsdShadeConnectableAPI Connectable{SurfaceShader};
 
 	UsdShadeConversionImpl::FParameterValue ParameterValue;
 
-	auto SetTextureComponentParamForScalarInput = [ &Material ]( const UsdShadeConversionImpl::FParameterValue& InParameterValue, const TCHAR* ParamName )
-	{
-		if ( const UsdShadeConversionImpl::FTextureParameterValue* TextureParameterValue = InParameterValue.TryGet< UsdShadeConversionImpl::FTextureParameterValue >() )
-		{
-			FLinearColor ComponentMask;
-			switch( TextureParameterValue->OutputIndex )
-			{
-			case 0: // RGB, since we are handling scalars here plug red when RGB is requested
-			case 1: // R
-				ComponentMask = FLinearColor{ 1.f, 0.f, 0.f, 0.f };
-				break;
-			case 2: // G
-				ComponentMask = FLinearColor{ 0.f, 1.f, 0.f, 0.f };
-				break;
-			case 3: // B
-				ComponentMask = FLinearColor{ 0.f, 0.f, 1.f, 0.f };
-				break;
-			case 4: // A
-				ComponentMask = FLinearColor{ 0.f, 0.f, 0.f, 1.f };
-				break;
-			}
-
-			UsdShadeConversionImpl::SetVectorParameterValue( Material, ParamName, ComponentMask );
-		}
-	};
+	TMap<FString, UsdShadeConversionImpl::FParameterValue> MaterialParameters;
+	MaterialParameters.Reserve(8);
 
-	const bool bForUsdPreviewSurface = true;
+	const bool bIsNormalMap = true;
 
 	// Base color
+	if (UsdShadeConversionImpl::GetVec3ParameterValue(Connectable, UnrealIdentifiers::DiffuseColor, FLinearColor(0, 0, 0), ParameterValue, !bIsNormalMap, &Material, TexturesCache))
 	{
-		const bool bIsNormalMap = false;
-		if ( UsdShadeConversionImpl::GetVec3ParameterValue( Connectable, UnrealIdentifiers::DiffuseColor, FLinearColor( 0, 0, 0 ), ParameterValue, bIsNormalMap, &Material, TexturesCache, &PrimvarToUVIndex ) )
-		{
-			UsdShadeConversionImpl::SetParameterValue( Material, TEXT( "BaseColor" ), ParameterValue, bForUsdPreviewSurface );
-			bHasMaterialInfo = true;
-		}
+		MaterialParameters.Add(TEXT("BaseColor"), ParameterValue);
 	}
 
 	// Emissive color
+	if (UsdShadeConversionImpl::GetVec3ParameterValue(Connectable, UnrealIdentifiers::EmissiveColor, FLinearColor(0, 0, 0), ParameterValue, !bIsNormalMap, &Material, TexturesCache))
 	{
-		const bool bIsNormalMap = false;
-		if ( UsdShadeConversionImpl::GetVec3ParameterValue( Connectable, UnrealIdentifiers::EmissiveColor, FLinearColor( 0, 0, 0 ), ParameterValue, bIsNormalMap, &Material, TexturesCache, &PrimvarToUVIndex ) )
-		{
-			UsdShadeConversionImpl::SetParameterValue( Material, TEXT( "EmissiveColor" ), ParameterValue, bForUsdPreviewSurface );
-			bHasMaterialInfo = true;
-		}
+		MaterialParameters.Add(TEXT("EmissiveColor"), ParameterValue);
 	}
 
 	// Metallic
+	if (UsdShadeConversionImpl::GetFloatParameterValue(Connectable, UnrealIdentifiers::Metallic, 0.f, ParameterValue, &Material, TexturesCache))
 	{
-		if ( UsdShadeConversionImpl::GetFloatParameterValue( Connectable, UnrealIdentifiers::Metallic, 0.f, ParameterValue, &Material, TexturesCache, &PrimvarToUVIndex ) )
-		{
-			UsdShadeConversionImpl::SetParameterValue( Material, TEXT( "Metallic" ), ParameterValue, bForUsdPreviewSurface );
-			SetTextureComponentParamForScalarInput( ParameterValue, TEXT( "MetallicTextureComponent" ) );
-			bHasMaterialInfo = true;
-		}
+		MaterialParameters.Add(TEXT("Metallic"), ParameterValue);
 	}
 
 	// Roughness
+	if (UsdShadeConversionImpl::GetFloatParameterValue(Connectable, UnrealIdentifiers::Roughness, 1.f, ParameterValue, &Material, TexturesCache))
 	{
-		if ( UsdShadeConversionImpl::GetFloatParameterValue( Connectable, UnrealIdentifiers::Roughness, 1.f, ParameterValue, &Material, TexturesCache, &PrimvarToUVIndex ) )
-		{
-			UsdShadeConversionImpl::SetParameterValue( Material, TEXT( "Roughness" ), ParameterValue, bForUsdPreviewSurface );
-			SetTextureComponentParamForScalarInput( ParameterValue, TEXT( "RoughnessTextureComponent" ) );
-			bHasMaterialInfo = true;
-		}
+		MaterialParameters.Add(TEXT("Roughness"), ParameterValue);
 	}
 
 	// Opacity
+	if (UsdShadeConversionImpl::GetFloatParameterValue(Connectable, UnrealIdentifiers::Opacity, 1.f, ParameterValue, &Material, TexturesCache))
 	{
-		if ( UsdShadeConversionImpl::GetFloatParameterValue( Connectable, UnrealIdentifiers::Opacity, 1.f, ParameterValue, &Material, TexturesCache, &PrimvarToUVIndex ) )
-		{
-			UsdShadeConversionImpl::SetParameterValue( Material, TEXT( "Opacity" ), ParameterValue, bForUsdPreviewSurface );
-			SetTextureComponentParamForScalarInput( ParameterValue, TEXT( "OpacityTextureComponent" ) );
-			bHasMaterialInfo = true;
-		}
+		MaterialParameters.Add(TEXT("Opacity"), ParameterValue);
 	}
 
 	// Normal
+	if (UsdShadeConversionImpl::GetVec3ParameterValue(Connectable, UnrealIdentifiers::Normal, FLinearColor(0, 0, 1), ParameterValue, bIsNormalMap, &Material, TexturesCache))
 	{
-		const bool bIsNormalMap = true;
-		if ( UsdShadeConversionImpl::GetVec3ParameterValue( Connectable, UnrealIdentifiers::Normal, FLinearColor( 0, 0, 1 ), ParameterValue, bIsNormalMap, &Material, TexturesCache, &PrimvarToUVIndex ) )
-		{
-			UsdShadeConversionImpl::SetParameterValue( Material, TEXT( "Normal" ), ParameterValue, bForUsdPreviewSurface );
-			bHasMaterialInfo = true;
-		}
+		MaterialParameters.Add(TEXT("Normal"), ParameterValue);
 	}
 
 	// Refraction
+	if (UsdShadeConversionImpl::GetFloatParameterValue(Connectable, UnrealIdentifiers::Refraction, 1.5f, ParameterValue, &Material, TexturesCache))
 	{
-		if ( UsdShadeConversionImpl::GetFloatParameterValue( Connectable, UnrealIdentifiers::Refraction, 1.5f, ParameterValue, &Material, TexturesCache, &PrimvarToUVIndex ) )
-		{
-			UsdShadeConversionImpl::SetParameterValue( Material, TEXT( "Refraction" ), ParameterValue, bForUsdPreviewSurface );
-			SetTextureComponentParamForScalarInput( ParameterValue, TEXT( "RefractionTextureComponent" ) );
-			bHasMaterialInfo = true;
-		}
+		MaterialParameters.Add(TEXT("Refraction"), ParameterValue);
 	}
 
 	// Occlusion
+	if (UsdShadeConversionImpl::GetFloatParameterValue(Connectable, UnrealIdentifiers::Occlusion, 1.0f, ParameterValue, &Material, TexturesCache))
+	{
+		MaterialParameters.Add(TEXT("AmbientOcclusion"), ParameterValue);
+	}
+
+	TMap<FString, FString> ParameterToPrimvar;
+
+	// Sort primvars
+	TSet<FString> UsedPrimvars;
+	for (const TPair<FString, UsdShadeConversionImpl::FParameterValue>& Parameter : MaterialParameters)
 	{
-		if ( UsdShadeConversionImpl::GetFloatParameterValue( Connectable, UnrealIdentifiers::Occlusion, 1.0f, ParameterValue, &Material, TexturesCache, &PrimvarToUVIndex ) )
+		if (const UsdShadeConversionImpl::FTextureParameterValue* TextureParameterValue = Parameter.Value.TryGet<UsdShadeConversionImpl::FTextureParameterValue>())
 		{
-			UsdShadeConversionImpl::SetParameterValue( Material, TEXT( "AmbientOcclusion" ), ParameterValue, bForUsdPreviewSurface );
-			SetTextureComponentParamForScalarInput( ParameterValue, TEXT( "AmbientOcclusionTextureComponent" ) );
-			bHasMaterialInfo = true;
+			UsedPrimvars.Add(TextureParameterValue->Primvar);
+			ParameterToPrimvar.Add(Parameter.Key, TextureParameterValue->Primvar);
 		}
 	}
+	UsedPrimvars.Remove(TEXT(""));
+	TArray<FString> SortedPrimvars = UsedPrimvars.Array();
+	SortedPrimvars.Sort();  // Try for some deterministic ordering (st0 should come before st1, etc.)
+	TMap<FString, int32> PrimvarToUVIndex;
+	PrimvarToUVIndex.Reserve(SortedPrimvars.Num());
+	int32 UVIndex = 0;
+	for (const FString& Primvar : SortedPrimvars)
+	{
+		PrimvarToUVIndex.Add(Primvar, UVIndex++);
+	}
+
+	// Set material parameters on the actual material instance
+	const bool bForUsdPreviewSurface = true;
+	for (const TPair<FString, UsdShadeConversionImpl::FParameterValue>& Parameter : MaterialParameters)
+	{
+		UsdShadeConversionImpl::SetParameterValue(
+			Material,
+			*Parameter.Key,
+			Parameter.Value,
+			bForUsdPreviewSurface,
+			PrimvarToUVIndex
+		);
+	}
 
 	// Handle world space normals
 	pxr::UsdPrim UsdPrim = UsdShadeMaterial.GetPrim();
-	if ( pxr::UsdAttribute Attr = UsdPrim.GetAttribute( UnrealIdentifiers::WorldSpaceNormals ) )
+	if (pxr::UsdAttribute Attr = UsdPrim.GetAttribute(UnrealIdentifiers::WorldSpaceNormals))
 	{
 		bool bValue = false;
-		if ( Attr.Get<bool>( &bValue ) && bValue )
+		if (Attr.Get<bool>(&bValue) && bValue)
 		{
-			UsdShadeConversionImpl::SetScalarParameterValue( Material, TEXT( "UseWorldSpaceNormals" ), 1.0f );
+			UsdUtils::SetScalarParameterValue(Material, TEXT("UseWorldSpaceNormals"), 1.0f);
 		}
 	}
 
-	FString ShadeMaterialPath = UsdToUnreal::ConvertPath( UsdPrim.GetPath() );
-	UsdShadeConversionImpl::CheckForUVIndexConflicts( PrimvarToUVIndex, ShadeMaterialPath );
-
-	return bHasMaterialInfo;
-}
+	// Record which primvars we used on each UV index. This is important as we'll match this with the analogous member
+	// on static/skeletal meshe import data, and create a new material instance with different UV index parameter
+	// values if we need to
+	if (UUsdMaterialAssetImportData* ImportData = Cast<UUsdMaterialAssetImportData>(Material.AssetImportData.Get()))
+	{
+		ImportData->PrimvarToUVIndex = PrimvarToUVIndex;
+		ImportData->ParameterToPrimvar = ParameterToPrimvar;
+	}
 
-bool UsdToUnreal::ConvertMaterial( const pxr::UsdShadeMaterial& UsdShadeMaterial, UMaterial& Material )
-{
-	TMap<FString, int32> PrimvarToUVIndex;
-	UUsdAssetCache2* NewCache = nullptr;
-	return ConvertMaterial( UsdShadeMaterial, Material, NewCache, PrimvarToUVIndex );
+	return MaterialParameters.Num() > 0;
 }
 
-bool UsdToUnreal::ConvertMaterial( const pxr::UsdShadeMaterial& UsdShadeMaterial, UMaterial& Material, UUsdAssetCache2* TexturesCache, TMap< FString, int32 >& PrimvarToUVIndex, const TCHAR* RenderContext )
+bool UsdToUnreal::ConvertMaterial(
+	const pxr::UsdShadeMaterial& UsdShadeMaterial,
+	UMaterial& Material,
+	UUsdAssetCache2* TexturesCache,
+	const TCHAR* RenderContext
+)
 {
-	bool bHasMaterialInfo = false;
-
 #if WITH_EDITOR
 	pxr::TfToken RenderContextToken = pxr::UsdShadeTokens->universalRenderContext;
-	if ( RenderContext )
+	if (RenderContext)
 	{
-		pxr::TfToken ProvidedRenderContextToken = UnrealToUsd::ConvertToken( RenderContext ).Get();
+		pxr::TfToken ProvidedRenderContextToken = UnrealToUsd::ConvertToken(RenderContext).Get();
 
 		// We won't ever create an asset for the Unreal render context material prim.
 		// Check the universal context here in case this material should also generate an asset for that context too
-		if ( ProvidedRenderContextToken != UnrealIdentifiers::Unreal )
+		if (ProvidedRenderContextToken != UnrealIdentifiers::Unreal)
 		{
 			RenderContextToken = ProvidedRenderContextToken;
 		}
 	}
 
-	pxr::UsdShadeShader SurfaceShader = UsdShadeMaterial.ComputeSurfaceSource( RenderContextToken );
-	if ( !SurfaceShader )
+	pxr::UsdShadeShader SurfaceShader = UsdShadeMaterial.ComputeSurfaceSource(RenderContextToken);
+	if (!SurfaceShader)
 	{
 		return false;
 	}
 
-	pxr::UsdShadeConnectableAPI Connectable{ SurfaceShader };
+	pxr::UsdShadeConnectableAPI Connectable{SurfaceShader};
 
-	UsdShadeConversionImpl::FParameterValue ParameterValue;
+	UsdShadeConversionImpl::FParameterValue TempParameterValue;
+	TOptional<UsdShadeConversionImpl::FParameterValue> BaseColorParameter;
+	TOptional<UsdShadeConversionImpl::FParameterValue> EmissiveParameter;
+	TOptional<UsdShadeConversionImpl::FParameterValue> MetallicParameter;
+	TOptional<UsdShadeConversionImpl::FParameterValue> RoughnessParameter;
+	TOptional<UsdShadeConversionImpl::FParameterValue> OpacityParameter;
+	TOptional<UsdShadeConversionImpl::FParameterValue> NormalParameter;
+	TOptional<UsdShadeConversionImpl::FParameterValue> RefractionParameter;
+	TOptional<UsdShadeConversionImpl::FParameterValue> AmbientOcclusionParameter;
 
-	auto SetOutputIndex = []( const UsdShadeConversionImpl::FParameterValue& InParameterValue, int32& OutputIndexToSet )
+	const bool bIsNormalMap = true;
+
+	// Base color
+	if (UsdShadeConversionImpl::GetVec3ParameterValue(Connectable, UnrealIdentifiers::DiffuseColor, FLinearColor(0, 0, 0), TempParameterValue, !bIsNormalMap, &Material, TexturesCache))
 	{
-		if ( const UsdShadeConversionImpl::FTextureParameterValue* TextureParameterValue = InParameterValue.TryGet< UsdShadeConversionImpl::FTextureParameterValue >() )
-		{
-			OutputIndexToSet = TextureParameterValue->OutputIndex;
-		}
-	};
+		BaseColorParameter = TempParameterValue;
+	}
 
-	UMaterialEditorOnlyData* EditorOnly = Material.GetEditorOnlyData();
+	// Emissive color
+	if (UsdShadeConversionImpl::GetVec3ParameterValue(Connectable, UnrealIdentifiers::EmissiveColor, FLinearColor(0, 0, 0), TempParameterValue, !bIsNormalMap, &Material, TexturesCache))
+	{
+		EmissiveParameter = TempParameterValue;
+	}
 
-	// Base color
+	// Metallic
+	if (UsdShadeConversionImpl::GetFloatParameterValue(Connectable, UnrealIdentifiers::Metallic, 0.f, TempParameterValue, &Material, TexturesCache))
+	{
+		MetallicParameter = TempParameterValue;
+	}
+
+	// Roughness
+	if (UsdShadeConversionImpl::GetFloatParameterValue(Connectable, UnrealIdentifiers::Roughness, 1.f, TempParameterValue, &Material, TexturesCache))
+	{
+		RoughnessParameter = TempParameterValue;
+	}
+
+	// Opacity
+	if (UsdShadeConversionImpl::GetFloatParameterValue(Connectable, UnrealIdentifiers::Opacity, 1.f, TempParameterValue, &Material, TexturesCache))
+	{
+		OpacityParameter = TempParameterValue;
+	}
+
+	// Normal
+	if (UsdShadeConversionImpl::GetVec3ParameterValue(Connectable, UnrealIdentifiers::Normal, FLinearColor(0, 0, 1), TempParameterValue, bIsNormalMap, &Material, TexturesCache))
+	{
+		NormalParameter = TempParameterValue;
+	}
+
+	// Refraction
+	const bool bHasRefractionValue = UsdShadeConversionImpl::GetFloatParameterValue(Connectable, UnrealIdentifiers::Refraction, 1.5f, TempParameterValue, &Material, TexturesCache);
+	if (bHasRefractionValue || Material.BlendMode == BLEND_Translucent) // Force a 1.5 IOR if USD didn't specify a value, as it's USD's fallback value
+	{
+		RefractionParameter = TempParameterValue;
+	}
+
+	// Ambient occlusion
+	if (UsdShadeConversionImpl::GetFloatParameterValue(Connectable, UnrealIdentifiers::Occlusion, 1.f, TempParameterValue, &Material, TexturesCache))
 	{
-		const bool bIsNormalMap = false;
-		if ( UsdShadeConversionImpl::GetVec3ParameterValue( Connectable, UnrealIdentifiers::DiffuseColor, FLinearColor( 0, 0, 0 ), ParameterValue, bIsNormalMap, &Material, TexturesCache, &PrimvarToUVIndex ) )
+		AmbientOcclusionParameter = TempParameterValue;
+	}
+
+	TMap<FString, FString> ParameterToPrimvar;
+
+	// Sort primvars
+	TSet<FString> UsedPrimvars;
+	TFunction<void(const FString& Parameter, const TOptional<UsdShadeConversionImpl::FParameterValue>&)> FetchPrimvar =
+		[&UsedPrimvars, &ParameterToPrimvar](const FString& Parameter, const TOptional<UsdShadeConversionImpl::FParameterValue>& ParameterValue)
 		{
-			if (UMaterialExpression* Expression = UsdShadeConversionImpl::GetExpressionForValue(Material, ParameterValue))
+			if (!ParameterValue.IsSet())
 			{
-				EditorOnly->BaseColor.Expression = Expression;
-				SetOutputIndex(ParameterValue, EditorOnly->BaseColor.OutputIndex);
-
-				bHasMaterialInfo = true;
+				return;
 			}
-		}
+			if (const UsdShadeConversionImpl::FTextureParameterValue* TextureParameterValue = ParameterValue.GetValue().TryGet<UsdShadeConversionImpl::FTextureParameterValue>())
+			{
+				UsedPrimvars.Add(TextureParameterValue->Primvar);
+				ParameterToPrimvar.Add(Parameter, TextureParameterValue->Primvar);
+			}
+		};
+	FetchPrimvar(TEXT("BaseColor"), BaseColorParameter);
+	FetchPrimvar(TEXT("EmissiveColor"), EmissiveParameter);
+	FetchPrimvar(TEXT("Metallic"), MetallicParameter);
+	FetchPrimvar(TEXT("Roughness"), RoughnessParameter);
+	FetchPrimvar(TEXT("Opacity"), OpacityParameter);
+	FetchPrimvar(TEXT("Normal"), NormalParameter);
+	FetchPrimvar(TEXT("Refraction"), RefractionParameter);
+	FetchPrimvar(TEXT("AmpientOcclusion"), AmbientOcclusionParameter);
+	UsedPrimvars.Remove(TEXT(""));
+	TArray<FString> SortedPrimvars = UsedPrimvars.Array();
+	SortedPrimvars.Sort();
+	TMap<FString, int32> PrimvarToUVIndex;
+	PrimvarToUVIndex.Reserve(SortedPrimvars.Num());
+	int32 UVIndex = 0;
+	for (const FString& Primvar : SortedPrimvars)
+	{
+		PrimvarToUVIndex.Add(Primvar, UVIndex++);
 	}
 
-	// Emissive color
-	{
-		const bool bIsNormalMap = false;
-		if ( UsdShadeConversionImpl::GetVec3ParameterValue( Connectable, UnrealIdentifiers::EmissiveColor, FLinearColor( 0, 0, 0 ), ParameterValue, bIsNormalMap, &Material, TexturesCache, &PrimvarToUVIndex ) )
+	bool bHasMaterialInfo = true;
+
+	// auto so we can use the MaterialInput generic parameter for FColorMaterialInput, FScalarMaterialInput and FVectorMaterialInput
+	auto ConnectMaterialInput =
+		[&PrimvarToUVIndex, &bHasMaterialInfo]
+		(UMaterial& Material, auto& MaterialInput, const UsdShadeConversionImpl::FParameterValue& InParameterValue) -> bool
 		{
-			if (UMaterialExpression* Expression = UsdShadeConversionImpl::GetExpressionForValue(Material, ParameterValue))
+			UMaterialExpression* Expression = UsdShadeConversionImpl::GetExpressionForValue(Material, InParameterValue);
+			if (!Expression)
 			{
-				EditorOnly->EmissiveColor.Expression = Expression;
-				SetOutputIndex(ParameterValue, EditorOnly->EmissiveColor.OutputIndex);
+				return false;
+			}
 
-				bHasMaterialInfo = true;
+			bHasMaterialInfo = true;
+
+			MaterialInput.Expression = Expression;
+
+			if (const UsdShadeConversionImpl::FTextureParameterValue* TextureParameterValue = InParameterValue.TryGet<UsdShadeConversionImpl::FTextureParameterValue>())
+			{
+				MaterialInput.OutputIndex = TextureParameterValue->OutputIndex;
+
+				if (UMaterialExpressionTextureSample* TextureExpression = Cast<UMaterialExpressionTextureSample>(Expression))
+				{
+					if (int32* FoundCoordinate = PrimvarToUVIndex.Find(TextureParameterValue->Primvar))
+					{
+						TextureExpression->ConstCoordinate = *FoundCoordinate;
+					}
+					else
+					{
+						ensureMsgf(false, TEXT("Failed to find primvar '%s' when setting material parameter. Available primvars and UV indices: %s"),
+							*TextureParameterValue->Primvar,
+							*UsdUtils::StringifyMap(PrimvarToUVIndex)
+						);
+					}
+				}
 			}
-		}
+
+			return true;
+		};
+
+	UMaterialEditorOnlyData* EditorOnly = Material.GetEditorOnlyData();
+
+	// Base color
+	if (UsdShadeConversionImpl::FParameterValue* BaseColorParameterValue = BaseColorParameter.GetPtrOrNull())
+	{
+		ConnectMaterialInput(
+			Material,
+			EditorOnly->BaseColor,
+			*BaseColorParameterValue
+		);
 	}
 
-	// Metallic
-	if ( UsdShadeConversionImpl::GetFloatParameterValue( Connectable, UnrealIdentifiers::Metallic, 0.f, ParameterValue, &Material, TexturesCache, &PrimvarToUVIndex ) )
+	// Emissive color
+	if (UsdShadeConversionImpl::FParameterValue* EmissiveParameterValue = EmissiveParameter.GetPtrOrNull())
 	{
-		if ( UMaterialExpression* Expression = UsdShadeConversionImpl::GetExpressionForValue( Material, ParameterValue ) )
-		{
-			EditorOnly->Metallic.Expression = Expression;
-			SetOutputIndex( ParameterValue, EditorOnly->Metallic.OutputIndex );
+		ConnectMaterialInput(
+			Material,
+			EditorOnly->EmissiveColor,
+			*EmissiveParameterValue
+		);
+	}
 
-			bHasMaterialInfo = true;
-		}
+	// Metallic
+	if (UsdShadeConversionImpl::FParameterValue* MetallicParameterValue = MetallicParameter.GetPtrOrNull())
+	{
+		ConnectMaterialInput(
+			Material,
+			EditorOnly->Metallic,
+			*MetallicParameterValue
+		);
 	}
 
 	// Roughness
-	if ( UsdShadeConversionImpl::GetFloatParameterValue( Connectable, UnrealIdentifiers::Roughness, 1.f, ParameterValue, &Material, TexturesCache, &PrimvarToUVIndex ) )
+	if (UsdShadeConversionImpl::FParameterValue* RoughnessParameterValue = RoughnessParameter.GetPtrOrNull())
 	{
-		if ( UMaterialExpression* Expression = UsdShadeConversionImpl::GetExpressionForValue( Material, ParameterValue ) )
-		{
-			EditorOnly->Roughness.Expression = Expression;
-			SetOutputIndex( ParameterValue, EditorOnly->Roughness.OutputIndex );
-
-			bHasMaterialInfo = true;
-		}
+		ConnectMaterialInput(
+			Material,
+			EditorOnly->Roughness,
+			*RoughnessParameterValue
+		);
 	}
 
 	// Opacity
-	if ( UsdShadeConversionImpl::GetFloatParameterValue( Connectable, UnrealIdentifiers::Opacity, 1.f, ParameterValue, &Material, TexturesCache, &PrimvarToUVIndex ) )
+	if (UsdShadeConversionImpl::FParameterValue* OpacityParameterValue = OpacityParameter.GetPtrOrNull())
 	{
-		if ( UMaterialExpression* Expression = UsdShadeConversionImpl::GetExpressionForValue( Material, ParameterValue ) )
-		{
-			EditorOnly->Opacity.Expression = Expression;
-			SetOutputIndex( ParameterValue, EditorOnly->Opacity.OutputIndex );
+		const bool bConnected = ConnectMaterialInput(
+			Material,
+			EditorOnly->Opacity,
+			*OpacityParameterValue
+		);
 
+		if (bConnected)
+		{
 			Material.BlendMode = BLEND_Translucent;
-			bHasMaterialInfo = true;
 		}
 	}
 
 	// Normal
+	if (UsdShadeConversionImpl::FParameterValue* NormalParameterValue = NormalParameter.GetPtrOrNull())
 	{
-		const bool bIsNormalMap = true;
-		if ( UsdShadeConversionImpl::GetVec3ParameterValue( Connectable, UnrealIdentifiers::Normal, FLinearColor( 0, 0, 1 ), ParameterValue, bIsNormalMap, &Material, TexturesCache, &PrimvarToUVIndex ) )
-		{
-			if ( UMaterialExpression* Expression = UsdShadeConversionImpl::GetExpressionForValue( Material, ParameterValue ) )
-			{
-				EditorOnly->Normal.Expression = Expression;
-				SetOutputIndex(ParameterValue, EditorOnly->Normal.OutputIndex);
-
-				bHasMaterialInfo = true;
-			}
-		}
+		ConnectMaterialInput(
+			Material,
+			EditorOnly->Normal,
+			*NormalParameterValue
+		);
 	}
 
 	// Refraction
-	const bool bHasRefractionValue = UsdShadeConversionImpl::GetFloatParameterValue( Connectable, UnrealIdentifiers::Refraction, 1.5f, ParameterValue, &Material, TexturesCache, &PrimvarToUVIndex );
-	if ( bHasRefractionValue || Material.BlendMode == BLEND_Translucent ) // Force a 1.5 IOR if USD didn't specify a value, as it's USD's fallback value
+	if (UsdShadeConversionImpl::FParameterValue* RefractionParameterValue = RefractionParameter.GetPtrOrNull())
 	{
-		if ( UMaterialExpression* Expression = UsdShadeConversionImpl::GetExpressionForValue( Material, ParameterValue ) )
-		{
-			EditorOnly->Refraction.Expression = Expression;
-			SetOutputIndex( ParameterValue, EditorOnly->Refraction.OutputIndex );
-
-			bHasMaterialInfo = true;
-		}
+		ConnectMaterialInput(
+			Material,
+			EditorOnly->Refraction,
+			*RefractionParameterValue
+		);
 	}
 
-	// Occlusion
-	if ( UsdShadeConversionImpl::GetFloatParameterValue( Connectable, UnrealIdentifiers::Occlusion, 1.f, ParameterValue, &Material, TexturesCache, &PrimvarToUVIndex ) )
+	// Ambient occlusion
+	if (UsdShadeConversionImpl::FParameterValue* OcclusionParameterValue = AmbientOcclusionParameter.GetPtrOrNull())
 	{
-		if ( UMaterialExpression* Expression = UsdShadeConversionImpl::GetExpressionForValue( Material, ParameterValue ) )
-		{
-			EditorOnly->AmbientOcclusion.Expression = Expression;
-			SetOutputIndex( ParameterValue, EditorOnly->AmbientOcclusion.OutputIndex );
-
-			bHasMaterialInfo = true;
-		}
+		ConnectMaterialInput(
+			Material,
+			EditorOnly->AmbientOcclusion,
+			*OcclusionParameterValue
+		);
 	}
 
 	// Handle world space normals
 	pxr::UsdPrim UsdPrim = UsdShadeMaterial.GetPrim();
-	if ( pxr::UsdAttribute Attr = UsdPrim.GetAttribute( UnrealIdentifiers::WorldSpaceNormals ) )
+	if (pxr::UsdAttribute Attr = UsdPrim.GetAttribute(UnrealIdentifiers::WorldSpaceNormals))
 	{
 		bool bValue = false;
-		if ( Attr.Get<bool>( &bValue ) && bValue )
+		if (Attr.Get<bool>(&bValue) && bValue)
 		{
 			Material.bTangentSpaceNormal = false;
 		}
 	}
 
-	FString ShadeMaterialPath = UsdToUnreal::ConvertPath( UsdPrim.GetPath() );
-	UsdShadeConversionImpl::CheckForUVIndexConflicts( PrimvarToUVIndex, ShadeMaterialPath );
+	// Record which primvars we used on each UV index. This is important as we'll match this with the analogous member
+	// on static/skeletal meshe import data, and create a new material instance with different UV index parameter
+	// values if we need to
+	if (UUsdMaterialAssetImportData* ImportData = Cast<UUsdMaterialAssetImportData>(Material.AssetImportData.Get()))
+	{
+		ImportData->PrimvarToUVIndex = PrimvarToUVIndex;
+		ImportData->ParameterToPrimvar = ParameterToPrimvar;
+	}
 
 #endif // WITH_EDITOR
 	return bHasMaterialInfo;
 }
 
+// Deprecated
+bool UsdToUnreal::ConvertMaterial(
+	const pxr::UsdShadeMaterial& UsdShadeMaterial,
+	UMaterialInstance& Material,
+	UUsdAssetCache2* TexturesCache,
+	TMap<FString, int32>& InPrimvarToUVIndex,
+	const TCHAR* RenderContext
+)
+{
+	return ConvertMaterial(UsdShadeMaterial, Material, TexturesCache, RenderContext);
+}
+
+// Deprecated
+bool UsdToUnreal::ConvertMaterial(
+	const pxr::UsdShadeMaterial& UsdShadeMaterial,
+	UMaterial& Material,
+	UUsdAssetCache2* TexturesCache,
+	TMap<FString, int32>& PrimvarToUVIndex,
+	const TCHAR* RenderContext
+)
+{
+	return ConvertMaterial(UsdShadeMaterial, Material, TexturesCache, RenderContext);
+}
+
 bool UsdToUnreal::ConvertShadeInputsToParameters( const pxr::UsdShadeMaterial& UsdShadeMaterial, UMaterialInstance& MaterialInstance, UUsdAssetCache2* TexturesCache, const TCHAR* RenderContext )
 {
 	FScopedUsdAllocs UsdAllocs;
@@ -2341,11 +2405,14 @@ bool UsdToUnreal::ConvertShadeInputsToParameters( const pxr::UsdShadeMaterial& U
 		const bool bForUsdPreviewSurface = false;
 		const bool bIsNormalMap = DisplayName.Contains( TEXT("normal") );
 
+		// For now it seems we don't set "primvar parameters" anyway, so don't bother building this up
+		TMap<FString, int32> Unused;
+
 		if ( ShadeInput.GetTypeName() == pxr::SdfValueTypeNames->Bool )
 		{
 			if ( UsdShadeConversionImpl::GetBoolParameterValue( Connectable, ShadeInput.GetBaseName(), false, ParameterValue, &MaterialInstance, TexturesCache) )
 			{
-				UsdShadeConversionImpl::SetParameterValue( MaterialInstance, *DisplayName, ParameterValue, bForUsdPreviewSurface );
+				UsdShadeConversionImpl::SetParameterValue(MaterialInstance, *DisplayName, ParameterValue, bForUsdPreviewSurface, Unused);
 				bHasMaterialInfo = true;
 			}
 		}
@@ -2355,13 +2422,13 @@ bool UsdToUnreal::ConvertShadeInputsToParameters( const pxr::UsdShadeMaterial& U
 		{
 			if ( UsdShadeConversionImpl::GetFloatParameterValue( Connectable, ShadeInput.GetBaseName(), 1.f, ParameterValue, &MaterialInstance, TexturesCache ) )
 			{
-				UsdShadeConversionImpl::SetParameterValue( MaterialInstance, *DisplayName, ParameterValue, bForUsdPreviewSurface );
+				UsdShadeConversionImpl::SetParameterValue(MaterialInstance, *DisplayName, ParameterValue, bForUsdPreviewSurface, Unused);
 				bHasMaterialInfo = true;
 			}
 		}
 		else if ( UsdShadeConversionImpl::GetVec3ParameterValue( Connectable, ShadeInput.GetBaseName(), FLinearColor( 0.f, 0.f, 0.f ), ParameterValue, bIsNormalMap, &MaterialInstance, TexturesCache ) )
 		{
-			UsdShadeConversionImpl::SetParameterValue( MaterialInstance, *DisplayName, ParameterValue, bForUsdPreviewSurface );
+			UsdShadeConversionImpl::SetParameterValue(MaterialInstance, *DisplayName, ParameterValue, bForUsdPreviewSurface, Unused);
 			bHasMaterialInfo = true;
 		}
 	}
@@ -2372,14 +2439,13 @@ bool UsdToUnreal::ConvertShadeInputsToParameters( const pxr::UsdShadeMaterial& U
 PRAGMA_DISABLE_DEPRECATION_WARNINGS
 bool UsdToUnreal::ConvertMaterial(const pxr::UsdShadeMaterial& UsdShadeMaterial, UMaterialInstance& Material, UUsdAssetCache* TexturesCache, TMap<FString, int32>& PrimvarToUVIndex, const TCHAR* RenderContext)
 {
-	UUsdAssetCache2* NewCache = nullptr;
-	return UsdToUnreal::ConvertMaterial(UsdShadeMaterial, Material, NewCache, PrimvarToUVIndex, RenderContext);
+	return UsdToUnreal::ConvertMaterial(UsdShadeMaterial, Material);
 }
 
 bool UsdToUnreal::ConvertMaterial(const pxr::UsdShadeMaterial& UsdShadeMaterial, UMaterial& Material, UUsdAssetCache* TexturesCache, TMap<FString, int32>& PrimvarToUVIndex, const TCHAR* RenderContext)
 {
 	UUsdAssetCache2* NewCache = nullptr;
-	return UsdToUnreal::ConvertMaterial(UsdShadeMaterial, Material, NewCache, PrimvarToUVIndex, RenderContext);
+	return UsdToUnreal::ConvertMaterial(UsdShadeMaterial, Material, NewCache, RenderContext);
 }
 
 bool UsdToUnreal::ConvertShadeInputsToParameters(const pxr::UsdShadeMaterial& UsdShadeMaterial, UMaterialInstance& MaterialInstance, UUsdAssetCache* TexturesCache, const TCHAR* RenderContext)
@@ -2678,6 +2744,126 @@ bool UsdUtils::MarkMaterialPrimWithWorldSpaceNormals( const UE::FUsdPrim& Materi
 	return true;
 }
 
+void UsdUtils::SetScalarParameterValue(UMaterialInstance& Material, const TCHAR* ParameterName, float ParameterValue)
+{
+	UE_LOG(LogUsd, Verbose, TEXT("Setting material '%s' scalar parameter '%s' with value '%f'"),
+		*Material.GetPathName(),
+		ParameterName ? ParameterName : TEXT("nullptr"),
+		ParameterValue
+	);
+
+	FMaterialParameterInfo Info;
+	Info.Name = ParameterName;
+
+	if (UMaterialInstanceDynamic* Dynamic = Cast<UMaterialInstanceDynamic>(&Material))
+	{
+		Dynamic->SetScalarParameterValueByInfo(Info, ParameterValue);
+	}
+#if WITH_EDITOR
+	else if (UMaterialInstanceConstant* Constant = Cast<UMaterialInstanceConstant>(&Material))
+	{
+		Constant->SetScalarParameterValueEditorOnly(Info, ParameterValue);
+	}
+#endif // WITH_EDITOR
+}
+
+void UsdUtils::SetVectorParameterValue(UMaterialInstance& Material, const TCHAR* ParameterName, FLinearColor ParameterValue)
+{
+	UE_LOG(LogUsd, Verbose, TEXT("Setting material '%s' vector parameter '%s' with value '%s'"),
+		*Material.GetPathName(),
+		ParameterName ? ParameterName : TEXT("nullptr"),
+		*ParameterValue.ToString()
+	);
+
+	FMaterialParameterInfo Info;
+	Info.Name = ParameterName;
+
+	if (UMaterialInstanceDynamic* Dynamic = Cast<UMaterialInstanceDynamic>(&Material))
+	{
+		Dynamic->SetVectorParameterValueByInfo(Info, ParameterValue);
+	}
+#if WITH_EDITOR
+	else if (UMaterialInstanceConstant* Constant = Cast<UMaterialInstanceConstant>(&Material))
+	{
+		Constant->SetVectorParameterValueEditorOnly(Info, ParameterValue);
+	}
+#endif // WITH_EDITOR
+}
+
+void UsdUtils::SetTextureParameterValue(UMaterialInstance& Material, const TCHAR* ParameterName, UTexture* ParameterValue)
+{
+	UE_LOG(LogUsd, Verbose, TEXT("Setting material '%s' texture parameter '%s' with value '%s'"),
+		*Material.GetPathName(),
+		ParameterName ? ParameterName : TEXT("nullptr"),
+		*ParameterValue->GetPathName()
+	);
+
+	FMaterialParameterInfo Info;
+	Info.Name = ParameterName;
+
+	if (UMaterialInstanceDynamic* Dynamic = Cast<UMaterialInstanceDynamic>(&Material))
+	{
+		Dynamic->SetTextureParameterValueByInfo(Info, ParameterValue);
+	}
+#if WITH_EDITOR
+	else if (UMaterialInstanceConstant* Constant = Cast<UMaterialInstanceConstant>(&Material))
+	{
+		Constant->SetTextureParameterValueEditorOnly(Info, ParameterValue);
+	}
+#endif // WITH_EDITOR
+}
+
+void UsdUtils::SetBoolParameterValue(UMaterialInstance& Material, const TCHAR* ParameterName, bool bParameterValue)
+{
+	UE_LOG(LogUsd, Verbose, TEXT("Setting material '%s' bool parameter '%s' with value '%d'"),
+		*Material.GetPathName(),
+		ParameterName ? ParameterName : TEXT("nullptr"),
+		bParameterValue
+	);
+
+	bool bFound = false;
+
+#if WITH_EDITOR
+	// Try the static parameters first
+	if (UMaterialInstanceConstant* Constant = Cast<UMaterialInstanceConstant>(&Material))
+	{
+		bool bNeedsUpdatePermutations = false;
+
+		FStaticParameterSet StaticParameters;
+		Constant->GetStaticParameterValues(StaticParameters);
+
+		for (FStaticSwitchParameter& StaticSwitchParameter : StaticParameters.StaticSwitchParameters)
+		{
+			if (StaticSwitchParameter.ParameterInfo.Name == ParameterName)
+			{
+				bFound = true;
+
+				if (StaticSwitchParameter.Value != bParameterValue)
+				{
+					StaticSwitchParameter.Value = bParameterValue;
+					StaticSwitchParameter.bOverride = true;
+					bNeedsUpdatePermutations = true;
+				}
+
+				break;
+			}
+		}
+
+		if (bNeedsUpdatePermutations)
+		{
+			FlushRenderingCommands();
+			Constant->UpdateStaticPermutation(StaticParameters);
+		}
+	}
+#endif // WITH_EDITOR
+
+	// Try it as a scalar parameter
+	if (!bFound)
+	{
+		SetScalarParameterValue(Material, ParameterName, bParameterValue ? 1.0f : 0.0f);
+	}
+}
+
 TOptional<FString> UsdUtils::GetUnrealSurfaceOutput( const pxr::UsdPrim& MaterialPrim )
 {
 	if ( !MaterialPrim )
diff --git a/Engine/Plugins/Importers/USDImporter/Source/USDUtilities/Private/USDSkeletalDataConversion.cpp b/Engine/Plugins/Importers/USDImporter/Source/USDUtilities/Private/USDSkeletalDataConversion.cpp
--- a/Engine/Plugins/Importers/USDImporter/Source/USDUtilities/Private/USDSkeletalDataConversion.cpp
+++ b/Engine/Plugins/Importers/USDImporter/Source/USDUtilities/Private/USDSkeletalDataConversion.cpp
@@ -1242,10 +1242,8 @@ bool UsdToUnreal::ConvertSkinnedMesh(
 	const pxr::UsdSkelSkinningQuery& SkinningQuery,
 	const pxr::UsdSkelSkeletonQuery& SkeletonQuery,
 	FSkeletalMeshImportData& SkelMeshImportData,
-	TArray< UsdUtils::FUsdPrimMaterialSlot >& MaterialAssignments,
-	const TMap< FString, TMap< FString, int32 > >& MaterialToPrimvarsUVSetNames,
-	const pxr::TfToken& RenderContext,
-	const pxr::TfToken& MaterialPurpose
+	UsdUtils::FUsdPrimMaterialAssignmentInfo& MaterialAssignments,
+	const FUsdMeshConversionOptions& CommonOptions
 )
 {
 	TRACE_CPUPROFILER_EVENT_SCOPE( UsdToUnreal::ConvertSkinnedMesh );
@@ -1400,8 +1398,8 @@ bool UsdToUnreal::ConvertSkinnedMesh(
 		SkinningPrim,
 		pxr::UsdTimeCode::EarliestTime(),
 		bProvideMaterialIndices,
-		RenderContext,
-		MaterialPurpose
+		CommonOptions.RenderContext,
+		CommonOptions.MaterialPurpose
 	);
 	TArray< UsdUtils::FUsdPrimMaterialSlot >& LocalMaterialSlots = LocalInfo.Slots;
 	TArray< int32 >& FaceMaterialIndices = LocalInfo.MaterialIndices;
@@ -1414,9 +1412,9 @@ bool UsdToUnreal::ConvertSkinnedMesh(
 	TArray<int32> LocalToCombinedMaterialIndex;
 	LocalToCombinedMaterialIndex.SetNumZeroed( LocalInfo.Slots.Num() );
 
-	for (int32 Index = 0; Index < MaterialAssignments.Num(); ++Index)
+	for (int32 Index = 0; Index < MaterialAssignments.Slots.Num(); ++Index)
 	{
-		const UsdUtils::FUsdPrimMaterialSlot& Slot = MaterialAssignments[ Index ];
+		const UsdUtils::FUsdPrimMaterialSlot& Slot = MaterialAssignments.Slots[Index];
 
 		// Combine entries in this way so that we can append PrimPaths
 		TMap< UsdUtils::FUsdPrimMaterialSlot, int32 >::TKeyIterator KeyIt = SlotToCombinedMaterialIndex.CreateKeyIterator( Slot );
@@ -1445,14 +1443,14 @@ bool UsdToUnreal::ConvertSkinnedMesh(
 		}
 		else
 		{
-			int32 NewIndex = MaterialAssignments.Add( LocalSlot );
-			SlotToCombinedMaterialIndex.Add( LocalSlot, NewIndex );
-			LocalToCombinedMaterialIndex[ LocalIndex ] = NewIndex;
+			int32 NewIndex = MaterialAssignments.Slots.Add(LocalSlot);
+			SlotToCombinedMaterialIndex.Add(LocalSlot, NewIndex);
+			LocalToCombinedMaterialIndex[LocalIndex] = NewIndex;
 		}
 	}
 	// Now that we merged all prim paths into they keys of CombinedMaterialSlotsToIndex, let's copy them back into
 	// our output
-	for ( UsdUtils::FUsdPrimMaterialSlot& Slot : MaterialAssignments )
+	for (UsdUtils::FUsdPrimMaterialSlot& Slot : MaterialAssignments.Slots)
 	{
 		TMap< UsdUtils::FUsdPrimMaterialSlot, int32 >::TKeyIterator KeyIt = SlotToCombinedMaterialIndex.CreateKeyIterator( Slot );
 		ensure( KeyIt );
@@ -1594,12 +1592,24 @@ bool UsdToUnreal::ConvertSkinnedMesh(
 
 	TArray< FUVSet > UVSets;
 
-	TArray< TUsdStore< UsdGeomPrimvar > > PrimvarsByUVIndex = UsdUtils::GetUVSetPrimvars(
-		UsdMesh,
-		MaterialToPrimvarsUVSetNames,
-		RenderContext,
-		MaterialPurpose
-	);
+	// If we already have a primvar to UV index assignment, let's just use that.
+	// When collapsing, we'll do a pre-pass on all meshes to translate and determine this beforehand.
+	TArray<TUsdStore<pxr::UsdGeomPrimvar>> PrimvarsByUVIndex;
+	if (MaterialAssignments.PrimvarToUVIndex.Num() > 0)
+	{
+		TArray<TUsdStore<pxr::UsdGeomPrimvar>> AllMeshUVPrimvars =
+			UsdUtils::GetUVSetPrimvars(UsdMesh, TNumericLimits<int32>::Max());
+
+		PrimvarsByUVIndex =
+			UsdUtils::AssemblePrimvarsIntoUVSets(AllMeshUVPrimvars, MaterialAssignments.PrimvarToUVIndex);
+	}
+	// Let's use the best primvar assignment for this particular mesh instead
+	else
+	{
+		PrimvarsByUVIndex = UsdUtils::GetUVSetPrimvars(UsdMesh);
+
+		MaterialAssignments.PrimvarToUVIndex = UsdUtils::AssemblePrimvarsIntoPrimvarToUVIndexMap(PrimvarsByUVIndex);
+	}
 
 	int32 UVChannelIndex = 0;
 	while ( true )
@@ -1879,6 +1889,35 @@ bool UsdToUnreal::ConvertSkinnedMesh(
 	return true;
 }
 
+bool UsdToUnreal::ConvertSkinnedMesh(
+	const pxr::UsdSkelSkinningQuery& SkinningQuery,
+	const pxr::UsdSkelSkeletonQuery& SkeletonQuery,
+	FSkeletalMeshImportData& SkelMeshImportData,
+	TArray< UsdUtils::FUsdPrimMaterialSlot >& MaterialAssignments,
+	const TMap< FString, TMap< FString, int32 > >& MaterialToPrimvarsUVSetNames,
+	const pxr::TfToken& RenderContext,
+	const pxr::TfToken& MaterialPurpose
+)
+{
+	FUsdMeshConversionOptions Options;
+	Options.RenderContext = RenderContext;
+	Options.MaterialPurpose = MaterialPurpose;
+
+	UsdUtils::FUsdPrimMaterialAssignmentInfo TempInfo;
+	TempInfo.Slots = MaterialAssignments;
+
+	const bool bResult = ConvertSkinnedMesh(
+		SkinningQuery,
+		SkeletonQuery,
+		SkelMeshImportData,
+		TempInfo,
+		Options
+	);
+
+	MaterialAssignments = TempInfo.Slots;
+	return bResult;
+}
+
 bool UsdToUnreal::ConvertSkinnedMesh(
 	const pxr::UsdSkelSkinningQuery& UsdSkinningQuery,
 	const FTransform& AdditionalTransform,
@@ -1952,16 +1991,23 @@ bool UsdToUnreal::ConvertSkinnedMesh(
 			break;
 		}
 
-		// We really only ever deal with the first binding, so we can return here
-		return ConvertSkinnedMesh(
+		UsdToUnreal::FUsdMeshConversionOptions Options;
+		Options.RenderContext = RenderContext;
+		Options.MaterialPurpose = MaterialPurpose;
+
+		UsdUtils::FUsdPrimMaterialAssignmentInfo Info;
+		Swap(Info.Slots, MaterialAssignments);
+		const bool bSuccess = ConvertSkinnedMesh(
 			UsdSkinningQuery,
 			SkelQuery,
 			SkelMeshImportData,
-			MaterialAssignments,
-			MaterialToPrimvarsUVSetNames,
-			RenderContext,
-			MaterialPurpose
+			Info,
+			Options
 		);
+		Swap(Info.Slots, MaterialAssignments);
+
+		// We really only ever deal with the first binding, so we can return here
+		return bSuccess;
 	}
 
 	return false;
diff --git a/Engine/Plugins/Importers/USDImporter/Source/USDUtilities/Public/USDConversionUtils.h b/Engine/Plugins/Importers/USDImporter/Source/USDUtilities/Public/USDConversionUtils.h
--- a/Engine/Plugins/Importers/USDImporter/Source/USDUtilities/Public/USDConversionUtils.h
+++ b/Engine/Plugins/Importers/USDImporter/Source/USDUtilities/Public/USDConversionUtils.h
@@ -3,6 +3,7 @@
 #pragma once
 
 #include "CoreMinimal.h"
+#include "MeshDescription.h"
 #include "UObject/ObjectMacros.h"
 #include "UObject/UObjectGlobals.h"
 
@@ -135,16 +136,25 @@ namespace UsdUtils
 	 * @param PrimvarName - Name of the primvar that should be used as UV set
 	 * @return UV index that should be used for this primvar
 	 */
+	UE_DEPRECATED(5.3, "This heuristic is no longer used. Check GetUVSetPrimvars for the new behavior.")
 	USDUTILITIES_API int32 GetPrimvarUVIndex( FString PrimvarName );
 
 	/**
-	 * Gets the names of the primvars that should be used as UV sets, per index, for this mesh.
-	 * (e.g. first item of array is primvar for UV set 0, second for UV set 1, etc).
-	 * This overload will only return primvars with 'texcoord2f' role.	 *
+	 * Gets the primvars that should be used as UV sets, per index, for this UsdMesh.
+	 * This will return between 0 and MaxNumPrimvars primvars, where the first item of array is the primvar that should be used for
+	 * UV index 0 (if any), the second the primvar that should be used for UV set 1, etc.
+	 *
+	 * In case there are more than MaxNumPrimvars available primvars, the returned list will contain all the 'texCoord2f' role
+	 * primvars (sorted lexicographically), followed by the regular float2 primvars (also separately sorted lexicographically).
+	 *
 	 * @param UsdMesh - Mesh that contains primvars that can be used as texture coordinates.
-	 * @return Array where each index gives the primvar that should be used for that UV index
+	 * @param MaxNumPrimvars - Maximum number of primvars to return from this function.
+	 * @return Array of up to MaxNumPrimvars primvars sorted by priority (most important come first)
 	 */
-	USDUTILITIES_API TArray< TUsdStore< pxr::UsdGeomPrimvar > > GetUVSetPrimvars( const pxr::UsdGeomMesh& UsdMesh );
+	USDUTILITIES_API TArray<TUsdStore<pxr::UsdGeomPrimvar>> GetUVSetPrimvars(
+		const pxr::UsdGeomMesh& UsdMesh,
+		int32 MaxNumPrimvars = MAX_MESH_TEXTURE_COORDS_MD
+	);
 
 	/**
 	 * Gets the names of the primvars that should be used as UV sets, per index, for this mesh.
@@ -157,13 +167,73 @@ namespace UsdUtils
 	 * @param MaterialPurpose - Specific material purpose to use when parsing the UsdMesh's material bindings
 	 * @return Array where each index gives the primvar that should be used for that UV index
 	 */
+	UE_DEPRECATED(5.3, "Use the signature that just receives an UsdMesh, as the returned primvars no longer depend on material assignment information.")
 	USDUTILITIES_API TArray< TUsdStore< pxr::UsdGeomPrimvar > > GetUVSetPrimvars(
 		const pxr::UsdGeomMesh& UsdMesh,
-		const TMap< FString, TMap< FString, int32 > >& MaterialToPrimvarsUVSetNames,
+		const TMap<FString, TMap<FString, int32>>& MaterialToPrimvarsUVSetNames,
 		const pxr::TfToken& RenderContext = pxr::UsdShadeTokens->universalRenderContext,
 		const pxr::TfToken& MaterialPurpose = pxr::UsdShadeTokens->allPurpose
 	);
-	USDUTILITIES_API TArray< TUsdStore< pxr::UsdGeomPrimvar > > GetUVSetPrimvars( const pxr::UsdGeomMesh& UsdMesh, const TMap< FString, TMap< FString, int32 > >& MaterialToPrimvarsUVSetNames, const UsdUtils::FUsdPrimMaterialAssignmentInfo& UsdMeshMaterialAssignmentInfo );
+	UE_DEPRECATED(5.3, "Use the signature that just receives an UsdMesh, as the returned primvars no longer depend on material assignment information.")
+	USDUTILITIES_API TArray< TUsdStore< pxr::UsdGeomPrimvar > > GetUVSetPrimvars(
+		const pxr::UsdGeomMesh& UsdMesh,
+		const TMap<FString, TMap<FString, int32>>& MaterialToPrimvarsUVSetNames,
+		const UsdUtils::FUsdPrimMaterialAssignmentInfo& UsdMeshMaterialAssignmentInfo
+	);
+
+	/**
+	 * Rearranges an array of primvars into another array of up to MAX_MESH_TEXTURE_COORDS_MD primvars, describing which
+	 * primvars should be used for each UV set according to the provided AllowedPrimvarsToUVIndex mapping.
+	 *
+	 * The intent is to use this when e.g. a Mesh prim is to be collapsed together with other Mesh prims into a single
+	 * StaticMesh, and they all need consistent UV sets across them. So even though AllMeshUVPrimvars contains [aaa,
+	 * bbb, eee], AllowedPrimvarsToUVIndex may instruct us to assign these to UV indices [0, 1, 4] (and not just [0, 1,
+	 * 2]) because our Mesh is to be merged with another Mesh prim that contains [aaa, bbb, ccc, ddd], so that ccc can
+	 * be put on index 2 and ddd on index 3.
+	 *
+	 * @param AllMeshUVPrimvars - Array of primvars extracted from some mesh sorted by priority (most important come
+	 * first)
+	 * @param AllowedPrimvarsToUVIndex - Mapping from primvar names to the UV index that they're allowed to be assigned
+	 * to
+	 * @return Array where each index gives the primvar that should be used for that UV index
+	 */
+	USDUTILITIES_API TArray<TUsdStore<pxr::UsdGeomPrimvar>> AssemblePrimvarsIntoUVSets(
+		const TArray<TUsdStore<pxr::UsdGeomPrimvar>>& AllMeshUVPrimvars,
+		const TMap<FString, int32>& AllowedPrimvarsToUVIndex
+	);
+
+	/**
+	 * Constructs a mapping from primvar name to their position in the provided AllMeshUVPrimvars array.
+	 * @param AllMeshUVPrimvars - Array of primvars extracted from some mesh sorted by priority (most important come
+	 * first)
+	 * @return Mapping from primvar names to the UV index
+	 */
+	USDUTILITIES_API TMap<FString, int32> AssemblePrimvarsIntoPrimvarToUVIndexMap(
+		const TArray<TUsdStore<pxr::UsdGeomPrimvar>>& AllMeshUVPrimvars
+	);
+
+	/**
+	 * Given an array of primvars, and an array of which of those are preferred (AllPrimvars should also include
+	 * PreferredPrimvars), this will sort lexicographically and return a prioritized PrimvarToUVIndex assignment map of
+	 * up to MAX_MESH_TEXTURE_COORDS_MD UV indices
+	 */
+	USDUTILITIES_API TMap<FString, int32> CombinePrimvarsIntoUVSets(
+		const TSet<FString>& AllPrimvars,
+		const TSet<FString>& PreferredPrimvars
+	);
+
+	template<typename T, typename U>
+	FString StringifyMap(const TMap<T, U>& Map)
+	{
+		FString Result = TEXT("{");
+		for (const TPair<T, U>& Pair : Map)
+		{
+			Result += FString::Printf(TEXT("'%s': %s, "), *LexToString(Pair.Key), *LexToString(Pair.Value));
+		}
+		Result.RemoveFromEnd(TEXT(", "));
+		Result += TEXT("}");
+		return Result;
+	}
 
 	USDUTILITIES_API bool IsAnimated( const pxr::UsdPrim& Prim );
 	USDUTILITIES_API bool HasAnimatedVisibility( const pxr::UsdPrim& Prim );
diff --git a/Engine/Plugins/Importers/USDImporter/Source/USDUtilities/Public/USDGeomMeshConversion.h b/Engine/Plugins/Importers/USDImporter/Source/USDUtilities/Public/USDGeomMeshConversion.h
--- a/Engine/Plugins/Importers/USDImporter/Source/USDUtilities/Public/USDGeomMeshConversion.h
+++ b/Engine/Plugins/Importers/USDImporter/Source/USDUtilities/Public/USDGeomMeshConversion.h
@@ -45,6 +45,13 @@ namespace UsdToUnreal
 	/** Common options to mesh conversion functions */
 	struct USDUTILITIES_API FUsdMeshConversionOptions
 	{
+		// These need to be explicitly defaulted surrounded by the deprecation macros or else the implicit versions will
+		// be generated elsewhere, use MaterialToPrimvarToUVIndex and trigger the deprecation warnings
+		PRAGMA_DISABLE_DEPRECATION_WARNINGS
+		FUsdMeshConversionOptions(const FUsdMeshConversionOptions& Other) = default;
+		FUsdMeshConversionOptions& operator=(const FUsdMeshConversionOptions& Other) = default;
+		PRAGMA_ENABLE_DEPRECATION_WARNINGS
+
 		const static FUsdMeshConversionOptions DefaultOptions;
 		FUsdMeshConversionOptions();
 
@@ -66,6 +73,7 @@ namespace UsdToUnreal
 		// Maps from a material prim path, to pairs indicating which primvar names are used as 'st' coordinates for
 		// this mesh, and which UVIndex materials will sample from( e.g.[ "st0", 0 ], [ "myUvSet2", 2 ], etc ).
 		// This is used to pick which primvars will become UV sets.
+		UE_DEPRECATED(5.3, "No longer used, now this information is stored directly on generated Material AssetImportData/AssetUserData")
 		const TMap< FString, TMap<FString, int32> >* MaterialToPrimvarToUVIndex;
 
 		// Whether to try reusing material slots (both local and the ones already in MaterialAssignments) when
@@ -74,18 +82,22 @@ namespace UsdToUnreal
 	};
 
 	/**
-	 * Converts an UsdGeomMesh and appends the converted data to OutMeshDescription and OutMaterialAssignments.
+	 * Converts an UsdGeomMesh and appends the converted data to InOutMeshDescription and InOutMaterialAssignments.
+	 *
+	 * If a PrimvarToUVIndex mapping is provided in InOutMaterialAssignments, it will be used to remap the read
+	 * primvars into UV sets. Otherwise, a new PrimvarToUVIndex map will be constructed on-demand and assigned
+	 * to InOutMaterialAssignments instead.
 	 *
 	 * @param Mesh - Mesh prim to convert
-	 * @param OutMeshDescription - Output parameter that will be filled with the converted data
-	 * @param OutMaterialAssignments - Output parameter that will be filled with the material data extracted from Mesh
+	 * @param InOutMeshDescription - Input/Output parameter that will be filled with the converted data
+	 * @param InOutMaterialAssignments - Input/Output parameter that will be filled with the material data extracted from Mesh
 	 * @param CommonOptions - Shared options used for the conversion
 	 * @return Whether the conversion was successful or not.
 	 */
 	USDUTILITIES_API bool ConvertGeomMesh(
 		const pxr::UsdGeomMesh& Mesh,
-		FMeshDescription& OutMeshDescription,
-		UsdUtils::FUsdPrimMaterialAssignmentInfo& OutMaterialAssignments,
+		FMeshDescription& InOutMeshDescription,
+		UsdUtils::FUsdPrimMaterialAssignmentInfo& InOutMaterialAssignments,
 		const FUsdMeshConversionOptions& CommonOptions = FUsdMeshConversionOptions::DefaultOptions
 	);
 	UE_DEPRECATED( 5.1, "Please use the overload that receives FUsdMeshConversionOptions." )
@@ -139,17 +151,21 @@ namespace UsdToUnreal
 	 * putting the resulting mesh data within OutMeshDescription and OutMaterialAssignments. It will not bake Prim's
 	 * own transform into the Mesh data.
 	 *
+	 * If a PrimvarToUVIndex mapping is provided in InOutMaterialAssignments, it will be used to remap the read
+	 * primvars into UV sets. Otherwise, a new PrimvarToUVIndex map will be constructed on-demand and assigned
+	 * to the InOutMaterialAssignments instead.
+	 *
 	 * @param Prim - Prim that will be converted (result will include its mesh and its children's)
-	 * @param OutMeshDescription - Output parameter that will be filled with the converted mesh descriptions
-	 * @param OutMaterialAssignments - Output parameter that will be filled with the material data extracted from
+	 * @param OutMeshDescription - Input/Output parameter that will be filled with the converted mesh descriptions
+	 * @param OutMaterialAssignments - Input/Output parameter that will be filled with the material data extracted from
 									   UsdSchema. Each item of the array matches the mesh description of same index
 	 * @param CommonOptions - Shared options used for the conversion
 	 * @return Whether the conversion was successful or not.
 	 */
 	USDUTILITIES_API bool ConvertGeomMeshHierarchy(
 		const pxr::UsdPrim& Prim,
-		FMeshDescription& OutMeshDescription,
-		UsdUtils::FUsdPrimMaterialAssignmentInfo& OutMaterialAssignments,
+		FMeshDescription& InOutMeshDescription,
+		UsdUtils::FUsdPrimMaterialAssignmentInfo& InOutMaterialAssignments,
 		const FUsdMeshConversionOptions& CommonOptions = FUsdMeshConversionOptions::DefaultOptions,
 		bool bSkipRootPrimTransformAndVisibility = false
 	);
@@ -266,8 +282,11 @@ namespace UsdUtils
 	{
 		TArray<FUsdPrimMaterialSlot> Slots;
 
-		/** Describes the index of the Slot that each polygon/face of a mesh uses. Matches the order Slots */
+		/** Describes the index of the Slot that each polygon/face of a mesh uses. Matches the order in Slots */
 		TArray<int32> MaterialIndices;
+
+		/** Describes which primvars should be assigned to each UV index for this mesh. */
+		TMap<FString, int32> PrimvarToUVIndex;
 	};
 }
 
diff --git a/Engine/Plugins/Importers/USDImporter/Source/USDUtilities/Public/USDShadeConversion.h b/Engine/Plugins/Importers/USDImporter/Source/USDUtilities/Public/USDShadeConversion.h
--- a/Engine/Plugins/Importers/USDImporter/Source/USDUtilities/Public/USDShadeConversion.h
+++ b/Engine/Plugins/Importers/USDImporter/Source/USDUtilities/Public/USDShadeConversion.h
@@ -47,16 +47,16 @@ namespace UsdToUnreal
 	/**
 	 * Extracts material data from UsdShadeMaterial and places the results in Material. Note that since this is used for UMaterialInstanceDynamics at runtime as well,
 	 * it will not set base property overrides (e.g. BlendMode) or the parent material, and will just assume that the caller handles that.
+	 * Note that in order to receive the primvar to UV index mapping calculated within this function, the provided Material should have an UUsdMaterialAssetImportData
+	 * object as its AssetImportData.
 	 * @param UsdShadeMaterial - Shade material with the data to convert
 	 * @param Material - Output parameter that will be filled with the converted data. Only the versions that receive a dynamic material instance will work at runtime
 	 * @param TexturesCache - Cache to prevent importing a texture more than once
-	 * @param PrimvarToUVIndex - Output parameter that will be filled the name of a primvar the material wants to use as UV set name, and the corresponding UV index it will sample texture coordinates from
+	 * @param RenderContext - Which render context output to read from the UsdShadeMaterial
 	 * @return Whether the conversion was successful or not.
 	 */
-	USDUTILITIES_API bool ConvertMaterial( const pxr::UsdShadeMaterial& UsdShadeMaterial, UMaterialInstance& Material );
-	USDUTILITIES_API bool ConvertMaterial( const pxr::UsdShadeMaterial& UsdShadeMaterial, UMaterialInstance& Material, UUsdAssetCache2* TexturesCache, TMap< FString, int32 >& PrimvarToUVIndex, const TCHAR* RenderContext = nullptr );
-	USDUTILITIES_API bool ConvertMaterial( const pxr::UsdShadeMaterial& UsdShadeMaterial, UMaterial& Material );
-	USDUTILITIES_API bool ConvertMaterial( const pxr::UsdShadeMaterial& UsdShadeMaterial, UMaterial& Material, UUsdAssetCache2* TexturesCache, TMap< FString, int32 >& PrimvarToUVIndex, const TCHAR* RenderContext = nullptr );
+	USDUTILITIES_API bool ConvertMaterial(const pxr::UsdShadeMaterial& UsdShadeMaterial, UMaterialInstance& Material, UUsdAssetCache2* TexturesCache = nullptr, const TCHAR* RenderContext = nullptr);
+	USDUTILITIES_API bool ConvertMaterial(const pxr::UsdShadeMaterial& UsdShadeMaterial, UMaterial& Material, UUsdAssetCache2* TexturesCache = nullptr, const TCHAR* RenderContext = nullptr);
 
 	/**
 	 * Attemps to assign the values of the surface shader inputs to the MaterialInstance parameters by matching the inputs display names to the parameters names.
@@ -70,6 +70,10 @@ namespace UsdToUnreal
 	USDUTILITIES_API bool ConvertShadeInputsToParameters( const pxr::UsdShadeMaterial& UsdShadeMaterial, UMaterialInstance& MaterialInstance, UUsdAssetCache2* TexturesCache, const TCHAR* RenderContext = nullptr );
 
 PRAGMA_DISABLE_DEPRECATION_WARNINGS
+	UE_DEPRECATED(5.3, "Please use the overload that doesn't use PrimvarToUVIndex anymore")
+	USDUTILITIES_API bool ConvertMaterial( const pxr::UsdShadeMaterial& UsdShadeMaterial, UMaterial& Material, UUsdAssetCache2* TexturesCache, TMap< FString, int32 >& PrimvarToUVIndex, const TCHAR* RenderContext = nullptr );
+	UE_DEPRECATED(5.3, "Use the other overload that receives an UUsdAssetCache2 object instead")
+	USDUTILITIES_API bool ConvertMaterial( const pxr::UsdShadeMaterial& UsdShadeMaterial, UMaterialInstance& Material, UUsdAssetCache2* TexturesCache, TMap< FString, int32 >& PrimvarToUVIndex, const TCHAR* RenderContext = nullptr );
 	UE_DEPRECATED(5.2, "Use the other overload that receives an UUsdAssetCache2 object instead")
 	USDUTILITIES_API bool ConvertMaterial( const pxr::UsdShadeMaterial& UsdShadeMaterial, UMaterialInstance& Material, UUsdAssetCache* TexturesCache, TMap< FString, int32 >& PrimvarToUVIndex, const TCHAR* RenderContext = nullptr );
 	UE_DEPRECATED(5.2, "Use the other overload that receives an UUsdAssetCache2 object instead")
@@ -193,6 +197,12 @@ namespace UsdUtils
 
 	/** Temporary function until UnrealWrappers can create attributes, just adds a custom bool attribute 'worldSpaceNormals' as true */
 	USDUTILITIES_API bool MarkMaterialPrimWithWorldSpaceNormals( const UE::FUsdPrim& MaterialPrim );
+
+	/** Sets material instance parameters whether Material is a MaterialInstanceConstant or a MaterialInstanceDynamic */
+	USDUTILITIES_API void SetScalarParameterValue(UMaterialInstance& Material, const TCHAR* ParameterName, float ParameterValue);
+	USDUTILITIES_API void SetVectorParameterValue(UMaterialInstance& Material, const TCHAR* ParameterName, FLinearColor ParameterValue);
+	USDUTILITIES_API void SetTextureParameterValue(UMaterialInstance& Material, const TCHAR* ParameterName, UTexture* ParameterValue);
+	USDUTILITIES_API void SetBoolParameterValue(UMaterialInstance& Material, const TCHAR* ParameterName, bool bParameterValue);
 }
 
 #endif // #if USE_USD_SDK
diff --git a/Engine/Plugins/Importers/USDImporter/Source/USDUtilities/Public/USDSkeletalDataConversion.h b/Engine/Plugins/Importers/USDImporter/Source/USDUtilities/Public/USDSkeletalDataConversion.h
--- a/Engine/Plugins/Importers/USDImporter/Source/USDUtilities/Public/USDSkeletalDataConversion.h
+++ b/Engine/Plugins/Importers/USDImporter/Source/USDUtilities/Public/USDSkeletalDataConversion.h
@@ -8,6 +8,7 @@
 #include "UObject/ObjectMacros.h"
 
 #include "USDConversionUtils.h"
+#include "USDGeomMeshConversion.h"
 #include "UsdWrappers/ForwardDeclarations.h"
 
 #if USE_USD_SDK
@@ -44,6 +45,7 @@ namespace SkeletalMeshImportData
 namespace UsdUtils
 {
 	struct FUsdPrimMaterialSlot;
+	struct FUsdPrimMaterialAssignmentInfo;
 }
 
 #endif // #if USE_USD_SDK
@@ -158,6 +160,11 @@ namespace UsdToUnreal
 
 	/**
 	 * Extracts skeletal mesh data fro UsdSkinningQuery, and places the results in SkelMeshImportData.
+	 *
+	 * If a PrimvarToUVIndex mapping is provided in MaterialAssignments, it will be used to remap the read
+	 * primvars into UV sets. Otherwise, a new PrimvarToUVIndex map will be constructed on-demand and assigned
+	 * to MaterialAssignments instead.
+	 *
 	 * @param UsdSkinningQuery - SkinningQuery with the data to convert
 	 * @param SkeletonQuery - Query object for the skeleton this mesh is bound to
 	 * @param AdditionalTransform - Additional transform to apply to the vertices of the mesh (Note: No longer used, even in the deprecated overload of this function)
@@ -168,6 +175,14 @@ namespace UsdToUnreal
 	 * @param MaterialPurpose - Material purpose to use when parsing the skinned Mesh prim's material bindings
 	 * @return Whether the conversion was successful or not.
 	 */
+	USDUTILITIES_API bool ConvertSkinnedMesh(
+		const pxr::UsdSkelSkinningQuery& UsdSkinningQuery,
+		const pxr::UsdSkelSkeletonQuery& SkeletonQuery,
+		FSkeletalMeshImportData& SkelMeshImportData,
+		UsdUtils::FUsdPrimMaterialAssignmentInfo& MaterialAssignments,
+		const FUsdMeshConversionOptions& CommonOptions = FUsdMeshConversionOptions::DefaultOptions
+	);
+	UE_DEPRECATED( 5.3, "Please use the overload with an FUsdPrimMaterialAssignmentInfo parameter." )
 	USDUTILITIES_API bool ConvertSkinnedMesh(
 		const pxr::UsdSkelSkinningQuery& UsdSkinningQuery,
 		const pxr::UsdSkelSkeletonQuery& SkeletonQuery,
