USDImporter: expose UsdStage::CreateClassPrim() on the UE FUsdStage wrapper

diff --git a/Engine/Plugins/Importers/USDImporter/Source/UnrealUSDWrapper/Private/UsdWrappers/UsdStage.cpp b/Engine/Plugins/Importers/USDImporter/Source/UnrealUSDWrapper/Private/UsdWrappers/UsdStage.cpp
--- a/Engine/Plugins/Importers/USDImporter/Source/UnrealUSDWrapper/Private/UsdWrappers/UsdStage.cpp
+++ b/Engine/Plugins/Importers/USDImporter/Source/UnrealUSDWrapper/Private/UsdWrappers/UsdStage.cpp
@@ -686,6 +686,19 @@ namespace UE
 		return FUsdPrim();
 	}
 
+	template<typename PtrType>
+	FUsdPrim FUsdStageBase<PtrType>::CreateClassPrim(const FSdfPath& RootPrimPath)
+	{
+#if USE_USD_SDK
+		if (PtrType& Ptr = Impl->GetInner())
+		{
+			return FUsdPrim(Ptr->CreateClassPrim(RootPrimPath));
+		}
+#endif // #if USE_USD_SDK
+
+		return FUsdPrim();
+	}
+
 	template<typename PtrType>
 	bool FUsdStageBase<PtrType>::RemovePrim( const FSdfPath & Path )
 	{
diff --git a/Engine/Plugins/Importers/USDImporter/Source/UnrealUSDWrapper/Public/UsdWrappers/UsdStage.h b/Engine/Plugins/Importers/USDImporter/Source/UnrealUSDWrapper/Public/UsdWrappers/UsdStage.h
--- a/Engine/Plugins/Importers/USDImporter/Source/UnrealUSDWrapper/Public/UsdWrappers/UsdStage.h
+++ b/Engine/Plugins/Importers/USDImporter/Source/UnrealUSDWrapper/Public/UsdWrappers/UsdStage.h
@@ -110,6 +110,7 @@ namespace UE
 
 		FUsdPrim OverridePrim( const FSdfPath& Path );
 		FUsdPrim DefinePrim( const FSdfPath& Path, const TCHAR* TypeName = TEXT("") );
+		FUsdPrim CreateClassPrim(const FSdfPath& RootPrimPath);
 		bool RemovePrim( const FSdfPath& Path );
 
 	private:
