USD: Fix broken build when packaging a game with USDImporter plugin.

diff --git a/Engine/Plugins/Importers/USDImporter/Source/USDSchemas/Private/MeshTranslationImpl.cpp b/Engine/Plugins/Importers/USDImporter/Source/USDSchemas/Private/MeshTranslationImpl.cpp
--- a/Engine/Plugins/Importers/USDImporter/Source/USDSchemas/Private/MeshTranslationImpl.cpp
+++ b/Engine/Plugins/Importers/USDImporter/Source/USDSchemas/Private/MeshTranslationImpl.cpp
@@ -769,6 +769,8 @@ void MeshTranslationImpl::SetMaterialOverrides(
 		}
 	}
 
+	TMap<const UsdUtils::FUsdPrimMaterialSlot*, UMaterialInterface*> ResolvedMaterials;
+
 #if WITH_EDITOR
 	UUsdMeshAssetImportData* ImportData = nullptr;
 	if (UStaticMeshComponent* StaticMeshComponent = Cast<UStaticMeshComponent>(&MeshComponent))
@@ -804,9 +806,7 @@ void MeshTranslationImpl::SetMaterialOverrides(
 		*MeshComponent.GetPathName(),
 		*UsdToUnreal::ConvertPath(Prim.GetPrimPath())
 	);
-#endif // WITH_EDITOR
 
-	TMap<const UsdUtils::FUsdPrimMaterialSlot*, UMaterialInterface*> ResolvedMaterials;
 	if(ImportData && LODIndexToAssignments.Num() > 0)
 	{
 		// Stash our PrimvarToUVIndex in here, as that's where ResolveMaterialAssignmentInfo will look for it
@@ -820,6 +820,7 @@ void MeshTranslationImpl::SetMaterialOverrides(
 			Flags
 		);
 	}
+#endif // WITH_EDITOR
 
 	// Compare resolved materials with existing assignments, and create overrides if we need to
 	uint32 StaticMeshSlotIndex = 0;
diff --git a/Engine/Plugins/Importers/USDImporter/Source/USDUtilities/Private/USDShadeConversion.cpp b/Engine/Plugins/Importers/USDImporter/Source/USDUtilities/Private/USDShadeConversion.cpp
--- a/Engine/Plugins/Importers/USDImporter/Source/USDUtilities/Private/USDShadeConversion.cpp
+++ b/Engine/Plugins/Importers/USDImporter/Source/USDUtilities/Private/USDShadeConversion.cpp
@@ -2048,11 +2048,13 @@ bool UsdToUnreal::ConvertMaterial(
 	// Record which primvars we used on each UV index. This is important as we'll match this with the analogous member
 	// on static/skeletal meshe import data, and create a new material instance with different UV index parameter
 	// values if we need to
+#if WITH_EDITOR
 	if (UUsdMaterialAssetImportData* ImportData = Cast<UUsdMaterialAssetImportData>(Material.AssetImportData.Get()))
 	{
 		ImportData->PrimvarToUVIndex = PrimvarToUVIndex;
 		ImportData->ParameterToPrimvar = ParameterToPrimvar;
 	}
+#endif // WITH_EDITOR
 
 	return MaterialParameters.Num() > 0;
 }
@@ -2064,6 +2066,7 @@ bool UsdToUnreal::ConvertMaterial(
 	const TCHAR* RenderContext
 )
 {
+	bool bHasMaterialInfo = false;
 #if WITH_EDITOR
 	pxr::TfToken RenderContextToken = pxr::UsdShadeTokens->universalRenderContext;
 	if (RenderContext)
@@ -2183,8 +2186,6 @@ bool UsdToUnreal::ConvertMaterial(
 		PrimvarToUVIndex.Add(Primvar, UVIndex++);
 	}
 
-	bool bHasMaterialInfo = true;
-
 	// auto so we can use the MaterialInput generic parameter for FColorMaterialInput, FScalarMaterialInput and FVectorMaterialInput
 	auto ConnectMaterialInput =
 		[&PrimvarToUVIndex, &bHasMaterialInfo]
