USD: Added support for importing baked mesh deformation as GeometryCache asset.

diff --git a/Engine/Plugins/Importers/USDImporter/Source/USDSchemas/Private/USDGeomMeshTranslator.cpp b/Engine/Plugins/Importers/USDImporter/Source/USDSchemas/Private/USDGeomMeshTranslator.cpp
--- a/Engine/Plugins/Importers/USDImporter/Source/USDSchemas/Private/USDGeomMeshTranslator.cpp
+++ b/Engine/Plugins/Importers/USDImporter/Source/USDSchemas/Private/USDGeomMeshTranslator.cpp
@@ -14,7 +14,6 @@
 #include "USDClassesModule.h"
 #include "USDConversionUtils.h"
 #include "USDGeomMeshConversion.h"
-#include "USDGroomTranslatorUtils.h"
 #include "USDInfoCache.h"
 #include "USDIntegrationUtils.h"
 #include "USDLog.h"
@@ -29,9 +28,6 @@
 #include "Components/StaticMeshComponent.h"
 #include "Engine/CollisionProfile.h"
 #include "Engine/StaticMesh.h"
-#include "GeometryCache.h"
-#include "GeometryCacheComponent.h"
-#include "GroomComponent.h"
 #include "HAL/IConsoleManager.h"
 #include "Interfaces/ITargetPlatform.h"
 #include "Interfaces/ITargetPlatformManagerModule.h"
@@ -49,8 +45,6 @@
 #include "UObject/SoftObjectPath.h"
 
 #if WITH_EDITOR
-#include "GeometryCacheTrackUSD.h"
-#include "GeometryCacheUSDComponent.h"
 #include "IMeshBuilderModule.h"
 #include "MeshBudgetProjectSettings.h"
 #endif // WITH_EDITOR
@@ -65,17 +59,6 @@
 #include "pxr/usd/usdShade/material.h"
 #include "USDIncludesEnd.h"
 
-#if WITH_EDITOR  // The GeometryCacheStreamer module is editor-only
-// Can toggle on/off to compare performance with StaticMesh instead of GeometryCache
-static bool GUseGeometryCacheUSD = true;
-
-static FAutoConsoleVariableRef CVarUsdUseGeometryCache(
-	TEXT("USD.UseGeometryCache"),
-	GUseGeometryCacheUSD,
-	TEXT("Use GeometryCache instead of static meshes for loading animated meshes"));
-
-#endif // WITH_EDITOR
-
 static float GMeshNormalRepairThreshold = 0.05f;
 static FAutoConsoleVariableRef CVarMeshNormalRepairThreshold(
 	TEXT("USD.MeshNormalRepairThreshold"),
@@ -204,34 +187,6 @@ namespace UsdGeomMeshTranslatorImpl
 #endif
 	}
 
-	bool IsAnimated(const pxr::UsdPrim& Prim)
-	{
-		FScopedUsdAllocs UsdAllocs;
-
-		bool bHasAttributesTimeSamples = false;
-		{
-			constexpr bool bIncludeInherited = false;
-			pxr::TfTokenVector GeomMeshAttributeNames = pxr::UsdGeomMesh::GetSchemaAttributeNames(bIncludeInherited);
-			pxr::TfTokenVector GeomPointBasedAttributeNames = pxr::UsdGeomPointBased::GetSchemaAttributeNames(bIncludeInherited);
-
-			GeomMeshAttributeNames.reserve(GeomMeshAttributeNames.size() + GeomPointBasedAttributeNames.size());
-			GeomMeshAttributeNames.insert(GeomMeshAttributeNames.end(), GeomPointBasedAttributeNames.begin(), GeomPointBasedAttributeNames.end());
-
-			for (const pxr::TfToken& AttributeName : GeomMeshAttributeNames)
-			{
-				const pxr::UsdAttribute& Attribute = Prim.GetAttribute(AttributeName);
-
-				if (Attribute && Attribute.ValueMightBeTimeVarying())
-				{
-					bHasAttributesTimeSamples = true;
-					break;
-				}
-			}
-		}
-
-		return bHasAttributesTimeSamples;
-	}
-
 	/** Returns true if material infos have changed on the StaticMesh */
 	bool ProcessStaticMeshMaterials(
 		const pxr::UsdPrim& UsdPrim,
@@ -346,81 +301,6 @@ namespace UsdGeomMeshTranslatorImpl
 		return bMaterialAssignementsHaveChanged;
 	}
 
-	bool ProcessGeometryCacheMaterials(
-		const pxr::UsdPrim& UsdPrim,
-		const TArray< UsdUtils::FUsdPrimMaterialAssignmentInfo >& LODIndexToMaterialInfo,
-		UGeometryCache& GeometryCache,
-		UUsdAssetCache2& AssetCache,
-		FUsdInfoCache* InfoCache,
-		float Time,
-		EObjectFlags Flags
-	)
-	{
-		bool bMaterialAssignementsHaveChanged = false;
-
-#if WITH_EDITOR
-		UUsdMeshAssetImportData* MeshImportData = Cast< UUsdMeshAssetImportData>(GeometryCache.AssetImportData.Get());
-		ensureMsgf(MeshImportData, TEXT("Geometry Cache '%s' generated for prim '%s' should have an UUsdMeshAssetImportData at this point!"),
-			*GeometryCache.GetPathName(),
-			*UsdToUnreal::ConvertPath(UsdPrim.GetPrimPath())
-		);
-#endif // WITH_EDITOR
-
-		TMap<const UsdUtils::FUsdPrimMaterialSlot*, UMaterialInterface*> ResolvedMaterials = MeshTranslationImpl::ResolveMaterialAssignmentInfo(
-			UsdPrim,
-			LODIndexToMaterialInfo,
-			AssetCache,
-			*InfoCache,
-			Flags
-		);
-
-		uint32 SlotIndex = 0;
-		for (int32 LODIndex = 0; LODIndex < LODIndexToMaterialInfo.Num(); ++LODIndex)
-		{
-			const TArray< UsdUtils::FUsdPrimMaterialSlot >& LODSlots = LODIndexToMaterialInfo[LODIndex].Slots;
-
-			for (int32 LODSlotIndex = 0; LODSlotIndex < LODSlots.Num(); ++LODSlotIndex, ++SlotIndex)
-			{
-				const UsdUtils::FUsdPrimMaterialSlot& Slot = LODSlots[LODSlotIndex];
-				UMaterialInterface* Material = nullptr;
-				if (UMaterialInterface** FoundMaterial = ResolvedMaterials.Find(&Slot))
-				{
-					Material = *FoundMaterial;
-				}
-				else
-				{
-					UE_LOG(
-						LogUsd,
-						Error,
-						TEXT("Failed to resolve material '%s' for slot '%d' of LOD '%d' for geometry cache '%s'"),
-						*Slot.MaterialSource,
-						LODSlotIndex,
-						LODIndex,
-						*UsdToUnreal::ConvertPath(UsdPrim.GetPath())
-					);
-					continue;
-				}
-
-				if (!GeometryCache.Materials.IsValidIndex(SlotIndex))
-				{
-					GeometryCache.Materials.Add(Material);
-					bMaterialAssignementsHaveChanged = true;
-				}
-				else if (!(GeometryCache.Materials[SlotIndex] == Material))
-				{
-					GeometryCache.Materials[SlotIndex] = Material;
-					bMaterialAssignementsHaveChanged = true;
-				}
-
-#if WITH_EDITOR
-				MeshImportData->MaterialSlotToPrimPaths.FindOrAdd(SlotIndex).PrimPaths = Slot.PrimPaths.Array();
-#endif // WITH_EDITOR
-			}
-		}
-
-		return bMaterialAssignementsHaveChanged;
-	}
-
 	// If UsdMesh is a LOD, will parse it and all of the other LODs, and and place them in OutLODIndexToMeshDescription and OutLODIndexToMaterialInfo.
 	// Note that these other LODs will be hidden in other variants, and won't show up on traversal unless we actively switch the variants (which we do here).
 	// We use a separate function for this because there is a very specific set of conditions where we successfully can do this, and we
@@ -493,54 +373,54 @@ namespace UsdGeomMeshTranslatorImpl
 		TFunction<bool(const pxr::UsdGeomMesh&, int32)> ConvertLOD =
 			[&OptionsCopy, &Options, &PrimvarToUVIndex, &LODIndexToMeshDescriptionMap, &LODIndexToMaterialInfoMap]
 			(const pxr::UsdGeomMesh& LODMesh, int32 LODIndex)
-			{
-				FMeshDescription TempMeshDescription;
+		{
+			FMeshDescription TempMeshDescription;
 
-				UsdUtils::FUsdPrimMaterialAssignmentInfo TempMaterialInfo;
+			UsdUtils::FUsdPrimMaterialAssignmentInfo TempMaterialInfo;
 				TempMaterialInfo.PrimvarToUVIndex = PrimvarToUVIndex;
 
-				FStaticMeshAttributes StaticMeshAttributes(TempMeshDescription);
-				StaticMeshAttributes.Register();
+			FStaticMeshAttributes StaticMeshAttributes(TempMeshDescription);
+			StaticMeshAttributes.Register();
 
-				bool bSuccess = true;
+			bool bSuccess = true;
 
-				// The user can't manually hide/unhide a particular LOD in the engine after it is imported, so we should probably bake
-				// the particular LOD visibility into the combined mesh. Note how we don't use computed visibility here, as we only really
-				// care if this mesh in particular has been marked as invisible
-				pxr::TfToken Visibility;
-				pxr::UsdAttribute VisibilityAttr = LODMesh.GetVisibilityAttr();
-				if (VisibilityAttr && VisibilityAttr.Get(&Visibility, Options.TimeCode) && Visibility == pxr::UsdGeomTokens->inherited)
-				{
-					// If we're interpreting LODs we must bake the transform from each LOD Mesh into the vertices, because there's no guarantee
-					// all LODs have the same transform, so we can't just put the transforms directly on the component. If we are not interpreting
-					// LODs we can do that though
-					// TODO: Handle resetXformOp here
-					bool bResetXformStack = false;
-					FTransform MeshTransform = FTransform::Identity;
-					bSuccess &= UsdToUnreal::ConvertXformable(LODMesh.GetPrim().GetStage(), LODMesh, MeshTransform, Options.TimeCode.GetValue(), &bResetXformStack);
-
-					if (bSuccess)
-					{
-						OptionsCopy.AdditionalTransform = MeshTransform * Options.AdditionalTransform;
-						OptionsCopy.bMergeIdenticalMaterialSlots = false;  // We only merge slots when collapsing, and we never collapse LODs
-
-						bSuccess &= UsdToUnreal::ConvertGeomMesh(
-							LODMesh,
-							TempMeshDescription,
-							TempMaterialInfo,
-							OptionsCopy
-						);
-					}
-				}
+			// The user can't manually hide/unhide a particular LOD in the engine after it is imported, so we should probably bake
+			// the particular LOD visibility into the combined mesh. Note how we don't use computed visibility here, as we only really
+			// care if this mesh in particular has been marked as invisible
+			pxr::TfToken Visibility;
+			pxr::UsdAttribute VisibilityAttr = LODMesh.GetVisibilityAttr();
+			if (VisibilityAttr && VisibilityAttr.Get(&Visibility, Options.TimeCode) && Visibility == pxr::UsdGeomTokens->inherited)
+			{
+				// If we're interpreting LODs we must bake the transform from each LOD Mesh into the vertices, because there's no guarantee
+				// all LODs have the same transform, so we can't just put the transforms directly on the component. If we are not interpreting
+				// LODs we can do that though
+				// TODO: Handle resetXformOp here
+				bool bResetXformStack = false;
+				FTransform MeshTransform = FTransform::Identity;
+				bSuccess &= UsdToUnreal::ConvertXformable(LODMesh.GetPrim().GetStage(), LODMesh, MeshTransform, Options.TimeCode.GetValue(), &bResetXformStack);
 
 				if (bSuccess)
 				{
-					LODIndexToMeshDescriptionMap.Add(LODIndex, MoveTemp(TempMeshDescription));
-					LODIndexToMaterialInfoMap.Add(LODIndex, MoveTemp(TempMaterialInfo));
+					OptionsCopy.AdditionalTransform = MeshTransform * Options.AdditionalTransform;
+					OptionsCopy.bMergeIdenticalMaterialSlots = false;  // We only merge slots when collapsing, and we never collapse LODs
+
+					bSuccess &= UsdToUnreal::ConvertGeomMesh(
+						LODMesh,
+						TempMeshDescription,
+						TempMaterialInfo,
+						OptionsCopy
+					);
 				}
+			}
 
-				return true;
-			};
+			if (bSuccess)
+			{
+				LODIndexToMeshDescriptionMap.Add(LODIndex, MoveTemp(TempMeshDescription));
+				LODIndexToMaterialInfoMap.Add(LODIndex, MoveTemp(TempMaterialInfo));
+			}
+
+			return true;
+		};
 		if (!UsdUtils::IterateLODMeshes(ParentPrim, ConvertLOD))
 		{
 			return false;
@@ -880,303 +760,6 @@ namespace UsdGeomMeshTranslatorImpl
 		StaticMesh.CalculateExtendedBounds();
 #endif // WITH_EDITOR
 	}
-
-#if WITH_EDITOR
-	void GeometryCacheDataForMeshDescription(FGeometryCacheMeshData& OutMeshData, FMeshDescription& MeshDescription);
-
-	void GetGeometryCacheDataTimeCodeRange(const UE::FUsdStage& Stage, const FString& PrimPath, int32& OutStartFrame, int32& OutEndFrame)
-	{
-		if (!Stage || PrimPath.IsEmpty())
-		{
-			return;
-		}
-
-		FScopedUsdAllocs Allocs;
-
-		pxr::UsdPrim UsdPrim = pxr::UsdPrim{Stage.GetPrimAtPath(UE::FSdfPath{ *PrimPath })};
-		if (!UsdPrim)
-		{
-			return;
-		}
-
-		constexpr bool bIncludeInherited = false;
-		pxr::TfTokenVector GeomMeshAttributeNames = pxr::UsdGeomMesh::GetSchemaAttributeNames(bIncludeInherited);
-		pxr::TfTokenVector GeomPointBasedAttributeNames = pxr::UsdGeomPointBased::GetSchemaAttributeNames(bIncludeInherited);
-
-		GeomMeshAttributeNames.reserve(GeomMeshAttributeNames.size() + GeomPointBasedAttributeNames.size());
-		GeomMeshAttributeNames.insert(GeomMeshAttributeNames.end(), GeomPointBasedAttributeNames.begin(), GeomPointBasedAttributeNames.end());
-
-		TOptional<double> MinStartTimeCode;
-		TOptional<double> MaxEndTimeCode;
-
-		for (const pxr::TfToken& AttributeName : GeomMeshAttributeNames)
-		{
-			if (const pxr::UsdAttribute& Attribute = UsdPrim.GetAttribute(AttributeName))
-			{
-				std::vector<double> TimeSamples;
-				if (Attribute.GetTimeSamples(&TimeSamples) && TimeSamples.size() > 0)
-				{
-					MinStartTimeCode = FMath::Min(MinStartTimeCode.Get(TNumericLimits<double>::Max()), TimeSamples[0]);
-					MaxEndTimeCode = FMath::Max(MaxEndTimeCode.Get(TNumericLimits<double>::Lowest()), TimeSamples[TimeSamples.size() - 1]);
-				}
-			}
-		}
-
-		if (MinStartTimeCode.IsSet() && MaxEndTimeCode.IsSet())
-		{
-			OutStartFrame = FMath::FloorToInt(MinStartTimeCode.GetValue());
-			OutEndFrame = FMath::CeilToInt(MaxEndTimeCode.GetValue());
-		}
-	}
-
-	UGeometryCache* CreateGeometryCache(const FString& InPrimPath, TArray< FMeshDescription >& LODIndexToMeshDescription, TSharedRef< FUsdSchemaTranslationContext> Context, bool& bOutIsNew)
-	{
-		UGeometryCache* GeometryCache = nullptr;
-
-		bool bHasValidMeshDescription = false;
-
-		FSHAHash AllLODHash;
-		FSHA1 SHA1;
-		for (const FMeshDescription& MeshDescription : LODIndexToMeshDescription)
-		{
-			FSHAHash LODHash = FStaticMeshOperations::ComputeSHAHash(MeshDescription);
-			SHA1.Update(&LODHash.Hash[0], sizeof(LODHash.Hash));
-
-			if (!MeshDescription.IsEmpty())
-			{
-				bHasValidMeshDescription = true;
-			}
-		}
-		SHA1.Final();
-		SHA1.GetHash(&AllLODHash.Hash[0]);
-
-		GeometryCache = Cast< UGeometryCache >(Context->AssetCache->GetCachedAsset(AllLODHash.ToString()));
-
-		if (!GeometryCache && bHasValidMeshDescription)
-		{
-			bOutIsNew = true;
-
-			const FName AssetName = MakeUniqueObjectName(
-				GetTransientPackage(),
-				UGeometryCache::StaticClass(),
-				*IUsdClassesModule::SanitizeObjectName(FPaths::GetBaseFilename(InPrimPath))
-			);
-			GeometryCache = NewObject< UGeometryCache >(GetTransientPackage(), AssetName, Context->ObjectFlags | EObjectFlags::RF_Public);
-
-			const UE::FUsdStage& Stage = Context->Stage;
-
-			// Fetch the animated mesh start/end frame as they may be different from just the stage's start and end time codes
-			int32 StartFrame = FMath::FloorToInt(Stage.GetStartTimeCode());
-			int32 EndFrame = FMath::CeilToInt(Stage.GetEndTimeCode());
-			UsdGeomMeshTranslatorImpl::GetGeometryCacheDataTimeCodeRange(Stage, InPrimPath, StartFrame, EndFrame);
-
-			// The GeometryCache module expects the end frame to be one past the last animation frame
-			EndFrame += 1;
-
-			pxr::TfToken RenderContextToken = pxr::UsdShadeTokens->universalRenderContext;
-			if (!Context->RenderContext.IsNone())
-			{
-				RenderContextToken = UnrealToUsd::ConvertToken(*Context->RenderContext.ToString()).Get();
-			}
-
-			pxr::TfToken MaterialPurposeToken = pxr::UsdShadeTokens->allPurpose;
-			if (!Context->MaterialPurpose.IsNone())
-			{
-				MaterialPurposeToken = UnrealToUsd::ConvertToken(*Context->MaterialPurpose.ToString()).Get();
-			}
-
-			UsdToUnreal::FUsdMeshConversionOptions Options;
-			Options.PurposesToLoad = Context->PurposesToLoad;
-			Options.bMergeIdenticalMaterialSlots = Context->bMergeIdenticalMaterialSlots;
-			Options.RenderContext = RenderContextToken;
-			Options.MaterialPurpose = MaterialPurposeToken;
-
-			// Create and configure a new USDTrack to be added to the GeometryCache
-			UGeometryCacheTrackUsd* UsdTrack = NewObject< UGeometryCacheTrackUsd >(GeometryCache);
-			UsdTrack->Initialize(
-				Stage,
-				InPrimPath,
-				StartFrame,
-				EndFrame,
-				[InPrimPath, Options](const TWeakObjectPtr<UGeometryCacheTrackUsd> TrackPtr, float Time, FGeometryCacheMeshData& OutMeshData) mutable
-				{
-					UGeometryCacheTrackUsd* Track = TrackPtr.Get();
-					if (!Track)
-					{
-						return false;
-					}
-
-					if (!Track->CurrentStagePinned)
-					{
-						return false;
-					}
-
-					UE::FUsdPrim Prim = Track->CurrentStagePinned.GetPrimAtPath(UE::FSdfPath{*Track->PrimPath});
-					if (!Prim)
-					{
-						return false;
-					}
-
-					// Get MeshDescription associated with the prim
-					// #ueent_todo: Replace MeshDescription with RawMesh to optimize
-					TArray< FMeshDescription > LODIndexToMeshDescription;
-					TArray< UsdUtils::FUsdPrimMaterialAssignmentInfo > LODIndexToMaterialInfo;
-					const bool bAllowInterpretingLODs = false;  // GeometryCaches don't have LODs, so we will never do this
-
-					// Need a local copy of Options to set the TimeCode since this function is called from multiple worker threads
-					UsdToUnreal::FUsdMeshConversionOptions LocalOptions(Options);
-					LocalOptions.TimeCode = pxr::UsdTimeCode{Time};
-
-					UsdGeomMeshTranslatorImpl::LoadMeshDescriptions(
-						pxr::UsdTyped(Prim),
-						LODIndexToMeshDescription,
-						LODIndexToMaterialInfo,
-						LocalOptions,
-						bAllowInterpretingLODs
-					);
-
-					// Convert the MeshDescription to MeshData, first LOD only
-					for (FMeshDescription& MeshDescription : LODIndexToMeshDescription)
-					{
-						if (!MeshDescription.IsEmpty())
-						{
-							// Compute the normals and tangents for the mesh
-							const float ComparisonThreshold = THRESH_POINTS_ARE_SAME;
-
-							// This function make sure the Polygon Normals Tangents Binormals are computed and also remove degenerated triangle from the render mesh description.
-							FStaticMeshOperations::ComputeTriangleTangentsAndNormals(MeshDescription, ComparisonThreshold);
-
-							// Compute any missing normals or tangents.
-							// Static meshes always blend normals of overlapping corners.
-							EComputeNTBsFlags ComputeNTBsOptions = EComputeNTBsFlags::BlendOverlappingNormals;
-							ComputeNTBsOptions |= EComputeNTBsFlags::IgnoreDegenerateTriangles;
-							ComputeNTBsOptions |= EComputeNTBsFlags::UseMikkTSpace;
-
-							FStaticMeshOperations::ComputeTangentsAndNormals(MeshDescription, ComputeNTBsOptions);
-
-							UsdGeomMeshTranslatorImpl::GeometryCacheDataForMeshDescription(OutMeshData, MeshDescription);
-
-							return true;
-						}
-					}
-					return false;
-				}
-			);
-
-			GeometryCache->AddTrack(UsdTrack);
-
-			TArray< FMatrix > Mats;
-			Mats.Add(FMatrix::Identity);
-			Mats.Add(FMatrix::Identity);
-
-			TArray< float > MatTimes;
-			MatTimes.Add(0.0f);
-			MatTimes.Add(0.0f);
-			UsdTrack->SetMatrixSamples(Mats, MatTimes);
-
-			Context->AssetCache->CacheAsset(AllLODHash.ToString(), GeometryCache);
-		}
-		else
-		{
-			bOutIsNew = false;
-		}
-
-		return GeometryCache;
-	}
-
-	// #ueent_todo: Replace MeshDescription with RawMesh and also make it work with StaticMesh
-	void GeometryCacheDataForMeshDescription(FGeometryCacheMeshData& OutMeshData, FMeshDescription& MeshDescription)
-	{
-		TRACE_CPUPROFILER_EVENT_SCOPE(GeometryCacheDataForMeshDescription);
-
-		OutMeshData.Positions.Reset();
-		OutMeshData.TextureCoordinates.Reset();
-		OutMeshData.TangentsX.Reset();
-		OutMeshData.TangentsZ.Reset();
-		OutMeshData.Colors.Reset();
-		OutMeshData.Indices.Reset();
-
-		OutMeshData.MotionVectors.Reset();
-		OutMeshData.BatchesInfo.Reset();
-		OutMeshData.BoundingBox.Init();
-
-		OutMeshData.VertexInfo.bHasColor0 = true;
-		OutMeshData.VertexInfo.bHasTangentX = true;
-		OutMeshData.VertexInfo.bHasTangentZ = true;
-		OutMeshData.VertexInfo.bHasUV0 = true;
-		OutMeshData.VertexInfo.bHasMotionVectors = false;
-
-		FStaticMeshAttributes MeshDescriptionAttributes(MeshDescription);
-
-		TVertexAttributesConstRef< FVector3f > VertexPositions = MeshDescriptionAttributes.GetVertexPositions();
-		TVertexInstanceAttributesConstRef< FVector3f > VertexInstanceNormals = MeshDescriptionAttributes.GetVertexInstanceNormals();
-		TVertexInstanceAttributesConstRef< FVector3f > VertexInstanceTangents = MeshDescriptionAttributes.GetVertexInstanceTangents();
-		TVertexInstanceAttributesConstRef< float > VertexInstanceBinormalSigns = MeshDescriptionAttributes.GetVertexInstanceBinormalSigns();
-		TVertexInstanceAttributesConstRef< FVector4f > VertexInstanceColors = MeshDescriptionAttributes.GetVertexInstanceColors();
-		TVertexInstanceAttributesConstRef< FVector2f > VertexInstanceUVs = MeshDescriptionAttributes.GetVertexInstanceUVs();
-
-		const int32 NumVertices = MeshDescription.Vertices().Num();
-		const int32 NumTriangles = MeshDescription.Triangles().Num();
-		const int32 NumMeshDataVertices = NumTriangles * 3;
-
-		OutMeshData.Positions.Reserve(NumVertices);
-		OutMeshData.Indices.Reserve(NumMeshDataVertices);
-		OutMeshData.TangentsZ.Reserve(NumMeshDataVertices);
-		OutMeshData.Colors.Reserve(NumMeshDataVertices);
-		OutMeshData.TextureCoordinates.Reserve(NumMeshDataVertices);
-
-		FBox BoundingBox(EForceInit::ForceInitToZero);
-		int32 MaterialIndex = 0;
-		int32 VertexIndex = 0;
-		for (FPolygonGroupID PolygonGroupID : MeshDescription.PolygonGroups().GetElementIDs())
-		{
-			// Skip empty polygon groups
-			if (MeshDescription.GetNumPolygonGroupPolygons(PolygonGroupID) == 0)
-			{
-				continue;
-			}
-
-			FGeometryCacheMeshBatchInfo BatchInfo;
-			BatchInfo.StartIndex = OutMeshData.Indices.Num();
-			BatchInfo.MaterialIndex = MaterialIndex++;
-
-			int32 TriangleCount = 0;
-			for (FPolygonID PolygonID : MeshDescription.GetPolygonGroupPolygonIDs(PolygonGroupID))
-			{
-				for (FTriangleID TriangleID : MeshDescription.GetPolygonTriangles(PolygonID))
-				{
-					for (FVertexInstanceID VertexInstanceID : MeshDescription.GetTriangleVertexInstances(TriangleID))
-					{
-						const FVector3f& Position = VertexPositions[MeshDescription.GetVertexInstanceVertex(VertexInstanceID)];
-						OutMeshData.Positions.Add(Position);
-						BoundingBox += (FVector)Position;
-
-						OutMeshData.Indices.Add(VertexIndex++);
-
-						FPackedNormal Normal = VertexInstanceNormals[VertexInstanceID];
-						Normal.Vector.W = VertexInstanceBinormalSigns[VertexInstanceID] < 0 ? -127 : 127;
-						OutMeshData.TangentsZ.Add(Normal);
-						OutMeshData.TangentsX.Add(VertexInstanceTangents[VertexInstanceID]);
-
-						const bool bSRGB = true;
-						OutMeshData.Colors.Add(FLinearColor(VertexInstanceColors[VertexInstanceID]).ToFColor(bSRGB));
-
-						// Supporting only one UV channel
-						const int32 UVIndex = 0;
-						OutMeshData.TextureCoordinates.Add(VertexInstanceUVs.Get(VertexInstanceID, UVIndex));
-					}
-
-					++TriangleCount;
-				}
-			}
-
-			OutMeshData.BoundingBox = (FBox3f)BoundingBox;
-
-			BatchInfo.NumTriangles = TriangleCount;
-			OutMeshData.BatchesInfo.Add(BatchInfo);
-		}
-	}
-#endif // WITH_EDITOR
 }
 
 FBuildStaticMeshTaskChain::FBuildStaticMeshTaskChain(
@@ -1420,149 +1003,21 @@ void FGeomMeshCreateAssetsTaskChain::SetupTasks()
 	FBuildStaticMeshTaskChain::SetupTasks();
 }
 
-#if WITH_EDITOR
-class FGeometryCacheCreateAssetsTaskChain : public FBuildStaticMeshTaskChain
-{
-public:
-	explicit FGeometryCacheCreateAssetsTaskChain(const TSharedRef< FUsdSchemaTranslationContext >& InContext, const UE::FSdfPath& InPrimPath)
-		: FBuildStaticMeshTaskChain(InContext, InPrimPath)
-	{
-		SetupTasks();
-	}
-
-protected:
-	virtual void SetupTasks() override;
-};
-
-void FGeometryCacheCreateAssetsTaskChain::SetupTasks()
-{
-	FScopedUnrealAllocs UnrealAllocs;
-
-	// Create mesh descriptions (Async or ExclusiveSync)
-	Do(ESchemaTranslationLaunchPolicy::Async,
-		[this]() -> bool
-		{
-			// Always hash the mesh at the same time because it may be animated, and
-			// otherwise we may think it's a new asset just because the context is at a different timecode (e.g. if we reload)
-			// TODO: Hash all timecodes, or else our mesh may change at t=5 and we never reload it because we only hash t=0
-			const bool bAllowInterpretingLODs = false;  // GeometryCaches don't have LODs
-
-			pxr::TfToken RenderContextToken = pxr::UsdShadeTokens->universalRenderContext;
-			if (!Context->RenderContext.IsNone())
-			{
-				RenderContextToken = UnrealToUsd::ConvertToken(*Context->RenderContext.ToString()).Get();
-			}
-
-			pxr::TfToken MaterialPurposeToken = pxr::UsdShadeTokens->allPurpose;
-			if (!Context->MaterialPurpose.IsNone())
-			{
-				MaterialPurposeToken = UnrealToUsd::ConvertToken(*Context->MaterialPurpose.ToString()).Get();
-			}
-
-			UsdToUnreal::FUsdMeshConversionOptions Options;
-			Options.TimeCode = UsdUtils::GetEarliestTimeCode();
-			Options.PurposesToLoad = Context->PurposesToLoad;
-			Options.RenderContext = RenderContextToken;
-			Options.MaterialPurpose = MaterialPurposeToken;
-			Options.bMergeIdenticalMaterialSlots = Context->bMergeIdenticalMaterialSlots;
-
-			UsdGeomMeshTranslatorImpl::LoadMeshDescriptions(
-				pxr::UsdTyped(GetPrim()),
-				LODIndexToMeshDescription,
-				LODIndexToMaterialInfo,
-				Options,
-				bAllowInterpretingLODs
-			);
-
-			// If we have at least one valid LOD, we should keep going
-			for (const FMeshDescription& MeshDescription : LODIndexToMeshDescription)
-			{
-				if (!MeshDescription.IsEmpty())
-				{
-					return true;
-				}
-			}
-			return false;
-		});
-
-	// Create the GeometryCache (Main thread)
-	Then(ESchemaTranslationLaunchPolicy::Sync,
-		[this]()
-		{
-			bool bIsNew = true;
-			const FString PrimPathString = PrimPath.GetString();
-			UGeometryCache* GeometryCache = UsdGeomMeshTranslatorImpl::CreateGeometryCache(PrimPathString, LODIndexToMeshDescription, Context, bIsNew);
-
-			if (bIsNew && GeometryCache)
-			{
-				UUsdMeshAssetImportData* ImportData = NewObject< UUsdMeshAssetImportData >(GeometryCache, TEXT("UUSDAssetImportData"));
-				ImportData->PrimPath = PrimPathString;
-				ImportData->PrimvarToUVIndex = LODIndexToMaterialInfo[0].PrimvarToUVIndex;
-				GeometryCache->AssetImportData = ImportData;
-
-				// Only the original creator of the prim at creation time gets to set the material assignments
-				// directly on the geomcache, all others prims ensure their materials via material overrides on the
-				// components
-				UsdGeomMeshTranslatorImpl::ProcessGeometryCacheMaterials(
-					GetPrim(),
-					LODIndexToMaterialInfo,
-					*GeometryCache,
-					*Context->AssetCache.Get(),
-					Context->InfoCache.Get(),
-					Context->Time,
-					Context->ObjectFlags
-				);
-			}
-
-			if (GeometryCache && Context->InfoCache)
-			{
-				const UE::FSdfPath& TargetPath = AlternativePrimToLinkAssetsTo.IsSet()
-					? AlternativePrimToLinkAssetsTo.GetValue()
-					: PrimPath;
-				Context->InfoCache->LinkAssetToPrim(TargetPath, GeometryCache);
-			}
-
-			// We don't have any additional step
-			return false;
-		});
-}
-#endif // WITH_EDITOR
-
 void FUsdGeomMeshTranslator::CreateAssets()
 {
 	TRACE_CPUPROFILER_EVENT_SCOPE(FUsdGeomMeshTranslator::CreateAssets);
 
 	RegisterAuxiliaryPrims();
 
-#if WITH_EDITOR
-	if (!Context->bIsImporting && GUseGeometryCacheUSD && UsdGeomMeshTranslatorImpl::IsAnimated(GetPrim()))
-	{
-		// Create the GeometryCache TaskChain
-		TSharedRef< FGeometryCacheCreateAssetsTaskChain > AssetsTaskChain = MakeShared< FGeometryCacheCreateAssetsTaskChain >(Context, PrimPath);
-
-		Context->TranslatorTasks.Add(MoveTemp(AssetsTaskChain));
-	}
-	else
-#endif // WITH_EDITOR
-	{
 		TSharedRef< FGeomMeshCreateAssetsTaskChain > AssetsTaskChain = MakeShared< FGeomMeshCreateAssetsTaskChain >(Context, PrimPath);
 
 		Context->TranslatorTasks.Add(MoveTemp(AssetsTaskChain));
 	}
-}
 
 USceneComponent* FUsdGeomMeshTranslator::CreateComponents()
 {
 	TOptional< TSubclassOf< USceneComponent > > ComponentType;
 
-#if WITH_EDITOR
-	// Force animated meshes as GeometryCache
-	if (!Context->bIsImporting && GUseGeometryCacheUSD && UsdGeomMeshTranslatorImpl::IsAnimated(GetPrim()))
-	{
-		ComponentType = UGeometryCacheUsdComponent::StaticClass();
-	}
-#endif // WITH_EDITOR
-
 	USceneComponent* SceneComponent = CreateComponentsEx(ComponentType, {});
 	UpdateComponents(SceneComponent);
 
@@ -1597,58 +1052,6 @@ USceneComponent* FUsdGeomMeshTranslator::CreateComponents()
 			}
 		}
 	}
-	else if (UGeometryCacheComponent* Component = Cast<UGeometryCacheComponent>(SceneComponent))
-	{
-		// We don't support importing geometry cache stuff yet so we should not be in here
-		ensure(!Context->bIsImporting);
-
-		if (Context->InfoCache && Context->AssetCache)
-		{
-			if (UGeometryCache* GeometryCache = Context->InfoCache->GetSingleAssetForPrim<UGeometryCache>(
-				PrimPath
-			))
-			{
-				// Geometry caches don't support LODs
-				const bool bAllowInterpretingLODs = false;
-
-				MeshTranslationImpl::SetMaterialOverrides(
-					GetPrim(),
-					GeometryCache->Materials,
-					*Component,
-					*Context->AssetCache,
-					*Context->InfoCache,
-					Context->Time,
-					Context->ObjectFlags,
-					bAllowInterpretingLODs,
-					Context->RenderContext,
-					Context->MaterialPurpose
-				);
-
-#if WITH_EDITOR
-				// Check if the prim has the GroomBinding schema and setup the component and assets necessary to bind the groom to the GeometryCache
-				if (UsdUtils::PrimHasSchema(GetPrim(), UnrealIdentifiers::GroomBindingAPI))
-				{
-					UsdGroomTranslatorUtils::CreateGroomBindingAsset(
-						GetPrim(),
-						*Context->AssetCache,
-						*Context->InfoCache,
-						Context->ObjectFlags
-					);
-
-					// For the groom binding to work, the GroomComponent must be a child of the SceneComponent
-					// so the Context ParentComponent is set to the SceneComponent temporarily
-					TGuardValue< USceneComponent* > ParentComponentGuard{Context->ParentComponent, SceneComponent};
-					const bool bNeedsActor = false;
-					UGroomComponent* GroomComponent = Cast< UGroomComponent >(CreateComponentsEx(TSubclassOf< USceneComponent >(UGroomComponent::StaticClass()), bNeedsActor));
-					if (GroomComponent)
-					{
-						UpdateComponents(SceneComponent);
-					}
-				}
-#endif // !WITH_EDITOR
-			}
-		}
-	}
 
 	return SceneComponent;
 }
@@ -1660,12 +1063,7 @@ void FUsdGeomMeshTranslator::UpdateComponents(USceneComponent* SceneComponent)
 		SceneComponent->Modify();
 	}
 
-	if (
-#if WITH_EDITOR
-		!GUseGeometryCacheUSD &&
-#endif // !WITH_EDITOR
-		UsdGeomMeshTranslatorImpl::IsAnimated(GetPrim())
-		)
+	if (UsdUtils::IsAnimatedMesh(GetPrim()))
 	{
 		// The assets might have changed since our attributes are animated
 		// Note that we must wait for these to complete as they make take a while and we want to
@@ -1674,73 +1072,6 @@ void FUsdGeomMeshTranslator::UpdateComponents(USceneComponent* SceneComponent)
 		Context->CompleteTasks();
 	}
 
-#if WITH_EDITOR
-	// Set the initial GeometryCache on the GeometryCacheUsdComponent
-	if (UGeometryCacheUsdComponent* GeometryCacheUsdComponent = Cast< UGeometryCacheUsdComponent >(SceneComponent))
-	{
-		ensure(!Context->bIsImporting);
-
-		UGeometryCache* GeometryCache = nullptr;
-		if (Context->InfoCache)
-		{
-			GeometryCache = Context->InfoCache->GetSingleAssetForPrim<UGeometryCache>(
-				PrimPath
-			);
-		}
-
-		bool bShouldRegister = false;
-		if (GeometryCache != GeometryCacheUsdComponent->GetGeometryCache())
-		{
-			bShouldRegister = true;
-
-			if (GeometryCacheUsdComponent->IsRegistered())
-			{
-				GeometryCacheUsdComponent->UnregisterComponent();
-			}
-
-			// Skip the extra handling in SetGeometryCache
-			GeometryCacheUsdComponent->GeometryCache = GeometryCache;
-		}
-
-		// Manually tick USD GeometryCache only when their tracks are disabled in Sequencer
-		static IConsoleVariable* CVar = IConsoleManager::Get().FindConsoleVariable(TEXT("USD.DisableGeoCacheTracks"));
-		bool bDisableGeoCacheTracks = CVar && CVar->GetBool();
-		if (bDisableGeoCacheTracks)
-		{
-			float TimeCode = Context->Time;
-			if (FMath::IsNaN(TimeCode))
-			{
-				int32 StartFrame = FMath::FloorToInt(Context->Stage.GetStartTimeCode());
-				int32 EndFrame = FMath::CeilToInt(Context->Stage.GetEndTimeCode());
-				UsdGeomMeshTranslatorImpl::GetGeometryCacheDataTimeCodeRange(Context->Stage, PrimPath.GetString(), StartFrame, EndFrame);
-
-				TimeCode = static_cast<float>(StartFrame);
-			}
-
-			// This is the main call responsible for animating the geometry cache.
-			// It needs to happen after setting the geometry cache and before registering, because we must force the
-			// geometry cache to register itself at Context->Time so that it will synchronously load that frame right away.
-			// Otherwise the geometry cache will start at t=0 regardless of Context->Time
-			GeometryCacheUsdComponent->SetManualTick(true);
-			GeometryCacheUsdComponent->TickAtThisTime(TimeCode, true, false, true);
-		}
-
-		// Note how we should only register if our geometry cache changed: If we did this every time we would
-		// register too early during the process of duplicating into PIE, and that would prevent a future RegisterComponent
-		// call from naturally creating the required render state
-		if (bShouldRegister && !GeometryCacheUsdComponent->IsRegistered())
-		{
-			GeometryCacheUsdComponent->RegisterComponent();
-		}
-
-		// If the prim has a GroomBinding schema, apply the target groom to its associated GroomComponent
-		if (UsdUtils::PrimHasSchema(GetPrim(), UnrealIdentifiers::GroomBindingAPI))
-		{
-			UsdGroomTranslatorUtils::SetGroomFromPrim(GetPrim(), *Context->InfoCache, SceneComponent);
-		}
-	}
-#endif // WITH_EDITOR
-
 	Super::UpdateComponents(SceneComponent);
 }
 
diff --git a/Engine/Plugins/Importers/USDImporter/Source/USDSchemas/Private/USDGeometryCacheTranslator.cpp b/Engine/Plugins/Importers/USDImporter/Source/USDSchemas/Private/USDGeometryCacheTranslator.cpp
new file mode 100644
--- /dev/null
+++ b/Engine/Plugins/Importers/USDImporter/Source/USDSchemas/Private/USDGeometryCacheTranslator.cpp
@@ -0,0 +1,1095 @@
+// Copyright Epic Games, Inc. All Rights Reserved.
+
+#include "USDGeometryCacheTranslator.h"
+
+#if USE_USD_SDK && WITH_EDITOR
+// The GeometryCacheStreamer module is editor-only, so is the translator
+
+#include "MeshTranslationImpl.h"
+#include "USDAssetCache.h"
+#include "USDAssetImportData.h"
+#include "USDClassesModule.h"
+#include "USDConversionUtils.h"
+#include "USDGroomTranslatorUtils.h"
+#include "USDInfoCache.h"
+#include "USDIntegrationUtils.h"
+#include "USDLog.h"
+#include "USDPrimConversion.h"
+#include "USDTypesConversion.h"
+
+#include "UsdWrappers/SdfPath.h"
+#include "UsdWrappers/UsdPrim.h"
+#include "UsdWrappers/UsdStage.h"
+
+#include "Async/ParallelFor.h"
+#include "GeometryCache.h"
+#include "GeometryCacheCodecV1.h"
+#include "GeometryCacheComponent.h"
+#include "GeometryCacheMeshData.h"
+#include "GeometryCacheTrackStreamable.h"
+#include "GeometryCacheTrackUSD.h"
+#include "GeometryCacheUSDComponent.h"
+#include "GroomComponent.h"
+#include "HAL/IConsoleManager.h"
+#include "StaticMeshAttributes.h"
+#include "StaticMeshOperations.h"
+#include "UObject/Package.h"
+
+#include "USDIncludesStart.h"
+	#include "pxr/usd/usd/typed.h"
+	#include "pxr/usd/usdGeom/imageable.h"
+	#include "pxr/usd/usdGeom/mesh.h"
+	#include "pxr/usd/usdGeom/subset.h"
+	#include "pxr/usd/usdShade/materialBindingAPI.h"
+#include "USDIncludesEnd.h"
+
+// Can toggle on/off to compare performance with StaticMesh instead of GeometryCache
+static bool GUseGeometryCacheUSD = true;
+static FAutoConsoleVariableRef CVarUsdUseGeometryCache(
+	TEXT("USD.GeometryCache.Enable"),
+	GUseGeometryCacheUSD,
+	TEXT("Use GeometryCache instead of static meshes for loading animated meshes"));
+
+static int32 GUsdGeometryCacheParallelMeshReads = 16;
+static FAutoConsoleVariableRef CVarUsdGeometryCacheParallelMeshReads(
+	TEXT("USD.GeometryCache.Import.ParallelMeshReads"),
+	GUsdGeometryCacheParallelMeshReads,
+	TEXT("Maximum number of mesh to process in parallel"));
+
+static int32 GUsdGeometryCacheParallelFrameReads = 16;
+static FAutoConsoleVariableRef CVarUsdGeometryCacheParallelFrameReads(
+	TEXT("USD.GeometryCache.Import.ParallelFrameReads"),
+	GUsdGeometryCacheParallelFrameReads,
+	TEXT("Maximum number of mesh frames to read in parallel"));
+
+namespace UsdGeometryCacheTranslatorImpl
+{
+	bool ProcessGeometryCacheMaterials(const pxr::UsdPrim& UsdPrim, const TArray<UsdUtils::FUsdPrimMaterialAssignmentInfo>& LODIndexToMaterialInfo,
+									   UGeometryCache& GeometryCache, UUsdAssetCache2& AssetCache, FUsdInfoCache* InfoCache, float Time, EObjectFlags Flags)
+	{
+		UUsdMeshAssetImportData* MeshImportData = Cast<UUsdMeshAssetImportData>(GeometryCache.AssetImportData.Get());
+		ensureMsgf(MeshImportData, TEXT("Geometry Cache '%s' generated for prim '%s' should have an UUsdMeshAssetImportData at this point!"),
+			*GeometryCache.GetPathName(), *UsdToUnreal::ConvertPath(UsdPrim.GetPrimPath()));
+
+		TMap<const UsdUtils::FUsdPrimMaterialSlot*, UMaterialInterface*> ResolvedMaterials = MeshTranslationImpl::ResolveMaterialAssignmentInfo(
+			UsdPrim, LODIndexToMaterialInfo, AssetCache, *InfoCache, Flags);
+
+		uint32 SlotIndex = 0;
+		bool bMaterialAssignementsHaveChanged = false;
+		for (int32 LODIndex = 0; LODIndex < LODIndexToMaterialInfo.Num(); ++LODIndex)
+		{
+			const TArray<UsdUtils::FUsdPrimMaterialSlot>& LODSlots = LODIndexToMaterialInfo[LODIndex].Slots;
+			for (int32 LODSlotIndex = 0; LODSlotIndex < LODSlots.Num(); ++LODSlotIndex, ++SlotIndex)
+			{
+				const UsdUtils::FUsdPrimMaterialSlot& Slot = LODSlots[LODSlotIndex];
+				UMaterialInterface* Material = nullptr;
+				if (UMaterialInterface** FoundMaterial = ResolvedMaterials.Find(&Slot))
+				{
+					Material = *FoundMaterial;
+				}
+				else
+				{
+					// Warn, but still add a material slot to preserve the materials order
+					UE_LOG(LogUsd, Warning, TEXT("Failed to resolve material '%s' for slot '%d' for geometry cache '%s'"), 
+						*Slot.MaterialSource, LODSlotIndex, *UsdToUnreal::ConvertPath(UsdPrim.GetPath()));
+				}
+
+				if (!GeometryCache.Materials.IsValidIndex(SlotIndex))
+				{
+					GeometryCache.Materials.Add(Material);
+					bMaterialAssignementsHaveChanged = true;
+				}
+				else if (!(GeometryCache.Materials[SlotIndex] == Material))
+				{
+					GeometryCache.Materials[SlotIndex] = Material;
+					bMaterialAssignementsHaveChanged = true;
+				}
+
+				MeshImportData->MaterialSlotToPrimPaths.FindOrAdd(SlotIndex).PrimPaths = Slot.PrimPaths.Array();
+			}
+		}
+
+		return bMaterialAssignementsHaveChanged;
+	}
+
+	void LoadMeshDescription(pxr::UsdTyped UsdMesh, FMeshDescription& OutMeshDescription, UsdUtils::FUsdPrimMaterialAssignmentInfo& OutMaterialInfo, const UsdToUnreal::FUsdMeshConversionOptions& Options)
+	{
+		if (!UsdMesh)
+		{
+			return;
+		}
+
+		FScopedUsdAllocs Allocs;
+
+		FMeshDescription TempMeshDescription;
+		UsdUtils::FUsdPrimMaterialAssignmentInfo TempMaterialInfo;
+
+		FStaticMeshAttributes StaticMeshAttributes(TempMeshDescription);
+		StaticMeshAttributes.Register();
+
+		bool bSuccess = UsdToUnreal::ConvertGeomMesh(pxr::UsdGeomMesh{UsdMesh}, TempMeshDescription, TempMaterialInfo, Options);
+		if (bSuccess)
+		{
+			OutMeshDescription = MoveTemp(TempMeshDescription);
+			OutMaterialInfo = MoveTemp(TempMaterialInfo);
+		}
+	}
+
+	void GeometryCacheDataForMeshDescription(FGeometryCacheMeshData& OutMeshData, FMeshDescription& MeshDescription, int32 MaterialIndex, const float SecondsPerFrame);
+
+	void GetGeometryCacheDataTimeCodeRange(const UE::FUsdStage& Stage, const FString& PrimPath, int32& OutStartFrame, int32& OutEndFrame)
+	{
+		if (!Stage || PrimPath.IsEmpty())
+		{
+			return;
+		}
+
+		FScopedUsdAllocs Allocs;
+
+		pxr::UsdPrim UsdPrim = pxr::UsdPrim(Stage.GetPrimAtPath(UE::FSdfPath(*PrimPath)));
+		if (!UsdPrim)
+		{
+			return;
+		}
+
+		constexpr bool bIncludeInherited = false;
+		pxr::TfTokenVector GeomMeshAttributeNames = pxr::UsdGeomMesh::GetSchemaAttributeNames(bIncludeInherited);
+		pxr::TfTokenVector GeomPointBasedAttributeNames = pxr::UsdGeomPointBased::GetSchemaAttributeNames(bIncludeInherited);
+
+		GeomMeshAttributeNames.reserve(GeomMeshAttributeNames.size() + GeomPointBasedAttributeNames.size());
+		GeomMeshAttributeNames.insert(GeomMeshAttributeNames.end(), GeomPointBasedAttributeNames.begin(), GeomPointBasedAttributeNames.end());
+
+		TOptional<double> MinStartTimeCode;
+		TOptional<double> MaxEndTimeCode;
+
+		for (const pxr::TfToken& AttributeName : GeomMeshAttributeNames)
+		{
+			if (const pxr::UsdAttribute& Attribute = UsdPrim.GetAttribute(AttributeName))
+			{
+				std::vector<double> TimeSamples;
+				if (Attribute.GetTimeSamples(&TimeSamples) && TimeSamples.size() > 0)
+				{
+					MinStartTimeCode = FMath::Min(MinStartTimeCode.Get(TNumericLimits<double>::Max()), TimeSamples[0]);
+					MaxEndTimeCode = FMath::Max(MaxEndTimeCode.Get(TNumericLimits<double>::Lowest()), TimeSamples[TimeSamples.size() - 1]);
+				}
+			}
+		}
+
+		if (MinStartTimeCode.IsSet() && MaxEndTimeCode.IsSet())
+		{
+			OutStartFrame = FMath::FloorToInt(MinStartTimeCode.GetValue());
+			OutEndFrame = FMath::CeilToInt(MaxEndTimeCode.GetValue());
+		}
+	}
+
+	struct FReadMeshDataArgs
+	{
+		FReadMeshDataArgs(const UE::FUsdStage& InStage, const UE::FUsdPrim& InRootPrim)
+			: Stage(InStage)
+			, RootPrim(InRootPrim)
+		{}
+
+		UE::FUsdStage Stage;
+		UE::FUsdPrim RootPrim;
+		UsdToUnreal::FUsdMeshConversionOptions Options;
+		int32 StartFrame = 0;
+		int32 EndFrame = 0;
+		float FramesPerSecond = 24.0f;
+		bool bPropagateTransform = false;
+	};
+
+	FReadMeshDataArgs GetReadMeshDataArgs(TSharedRef<FUsdSchemaTranslationContext> Context, const FString& RootPrimPath)
+	{
+		const UE::FUsdStage& Stage = Context->Stage;
+		UE::FUsdPrim RootPrim = Stage.GetPrimAtPath(UE::FSdfPath(*RootPrimPath));
+
+		FReadMeshDataArgs Args(Stage, RootPrim);
+
+		// Fetch the animated mesh start/end frame as they may be different from just the stage's start and end time codes
+		int32 StartFrame = FMath::FloorToInt(Stage.GetStartTimeCode());
+		int32 EndFrame = FMath::CeilToInt(Stage.GetEndTimeCode());
+		UsdGeometryCacheTranslatorImpl::GetGeometryCacheDataTimeCodeRange(Stage, RootPrimPath, StartFrame, EndFrame);
+
+		double FramesPerSecond = Stage.GetTimeCodesPerSecond();
+		if (FramesPerSecond == 0)
+		{
+			ensureMsgf(false, TEXT("Invalid USD GeometryCache FPS detected. Falling back to 1 FPS"));
+			FramesPerSecond = 1;
+		}
+
+		// The GeometryCache module expects the end frame to be one past the last animation frame
+		EndFrame += 1;
+
+		Args.StartFrame = StartFrame;
+		Args.EndFrame = EndFrame;
+		Args.FramesPerSecond = FramesPerSecond;
+
+		pxr::TfToken RenderContextToken = pxr::UsdShadeTokens->universalRenderContext;
+		if (!Context->RenderContext.IsNone())
+		{
+			RenderContextToken = UnrealToUsd::ConvertToken(*Context->RenderContext.ToString()).Get();
+		}
+
+		pxr::TfToken MaterialPurposeToken = pxr::UsdShadeTokens->allPurpose;
+		if (!Context->MaterialPurpose.IsNone())
+		{
+			MaterialPurposeToken = UnrealToUsd::ConvertToken(*Context->MaterialPurpose.ToString()).Get();
+		}
+
+		Args.Options.PurposesToLoad = Context->PurposesToLoad;
+		Args.Options.RenderContext = RenderContextToken;
+		Args.Options.MaterialPurpose = MaterialPurposeToken;
+		Args.Options.bMergeIdenticalMaterialSlots = false; // Don't merge because the GeometryCache is processed as unflattened (ie. one track per mesh)
+
+		return Args;
+	}
+
+	bool ReadMeshData(const FReadMeshDataArgs& Args, const UE::FUsdPrim& MeshPrim, int32 MaterialOffset, float Time, FGeometryCacheMeshData& OutMeshData)
+	{
+		FTransform PropatagedTransform = FTransform::Identity;
+		if (Args.bPropagateTransform)
+		{
+			UsdToUnreal::PropagateTransform(Args.Stage, Args.RootPrim, MeshPrim, Time, PropatagedTransform);
+		}
+
+		// Need a local copy of Options to set the TimeCode since this function is called from multiple worker threads
+		UsdToUnreal::FUsdMeshConversionOptions LocalOptions(Args.Options);
+		LocalOptions.TimeCode = pxr::UsdTimeCode(Time);
+		LocalOptions.AdditionalTransform = PropatagedTransform;
+
+		FMeshDescription MeshDescription;
+		UsdUtils::FUsdPrimMaterialAssignmentInfo MaterialInfo;
+		UsdGeometryCacheTranslatorImpl::LoadMeshDescription(pxr::UsdTyped(MeshPrim), MeshDescription, MaterialInfo, LocalOptions);
+
+		// Convert the MeshDescription to MeshData
+		if (!MeshDescription.IsEmpty())
+		{
+			// Compute the normals and tangents for the mesh
+			const float ComparisonThreshold = THRESH_POINTS_ARE_SAME;
+
+			// This function make sure the Polygon Normals Tangents Binormals are computed and also remove degenerated triangle from the render mesh description.
+			FStaticMeshOperations::ComputeTriangleTangentsAndNormals(MeshDescription, ComparisonThreshold);
+
+			// Compute any missing normals or tangents.
+			// Static meshes always blend normals of overlapping corners.
+			EComputeNTBsFlags ComputeNTBsOptions = EComputeNTBsFlags::BlendOverlappingNormals;
+			ComputeNTBsOptions |= EComputeNTBsFlags::IgnoreDegenerateTriangles;
+			ComputeNTBsOptions |= EComputeNTBsFlags::UseMikkTSpace;
+
+			FStaticMeshOperations::ComputeTangentsAndNormals(MeshDescription, ComputeNTBsOptions);
+
+			UsdGeometryCacheTranslatorImpl::GeometryCacheDataForMeshDescription(OutMeshData, MeshDescription, MaterialOffset, Args.FramesPerSecond);
+
+			return true;
+		}
+		return false;
+	}
+
+	UGeometryCacheTrackUsd* CreateUsdStreamTrack(UGeometryCache* GeometryCache, const FReadMeshDataArgs& Args, const FString& PrimPath, int32 MaterialOffset)
+	{
+		// Create and configure a new USDTrack to be added to the GeometryCache
+		UGeometryCacheTrackUsd* UsdTrack = NewObject<UGeometryCacheTrackUsd>(GeometryCache);
+		UsdTrack->Initialize(
+			Args.Stage,
+			PrimPath,
+			Args.StartFrame,
+			Args.EndFrame,
+			[Args, MaterialOffset](const TWeakObjectPtr<UGeometryCacheTrackUsd> TrackPtr, float Time, FGeometryCacheMeshData& OutMeshData) mutable
+			{
+				UGeometryCacheTrackUsd* Track = TrackPtr.Get();
+				if (!Track)
+				{
+					return false;
+				}
+
+				if (!Track->CurrentStagePinned)
+				{
+					return false;
+				}
+
+				UE::FUsdPrim Prim = Track->CurrentStagePinned.GetPrimAtPath(UE::FSdfPath(*Track->PrimPath));
+				if (!Prim)
+				{
+					return false;
+				}
+
+				return ReadMeshData(Args, Prim, MaterialOffset, Time, OutMeshData);
+			}
+		);
+		return UsdTrack;
+	}
+
+	UGeometryCacheTrackStreamable* CreateStreamableTrack(UGeometryCache* GeometryCache, const FString& PrimPath)
+	{
+		// Create and configure a new StreamableTrack to be added to the GeometryCache
+		FString ObjectName = IUsdClassesModule::SanitizeObjectName(FPaths::GetBaseFilename(PrimPath));
+
+		FName CodecName = MakeUniqueObjectName(GeometryCache, UGeometryCacheCodecV1::StaticClass(), FName(ObjectName + FString(TEXT("_Codec"))));
+		UGeometryCacheCodecV1* Codec = NewObject<UGeometryCacheCodecV1>(GeometryCache, CodecName , RF_Public);
+
+		// Compression settings for good quality
+		const float VertexQuantizationPrecision = 0.0005f;
+		const int32 UVBits = 16;
+		Codec->InitializeEncoder(VertexQuantizationPrecision, UVBits);
+
+		FName TrackName = MakeUniqueObjectName(GeometryCache, UGeometryCacheTrackStreamable::StaticClass(), FName(ObjectName));
+		UGeometryCacheTrackStreamable* StreamableTrack = NewObject<UGeometryCacheTrackStreamable>(GeometryCache, TrackName, RF_Public);
+
+		const bool bForceSingleOptimization = false;
+		const bool bCalculateMotionVectors = false;
+		const bool bOptimizeIndexBuffers = false;
+		StreamableTrack->BeginCoding(Codec, bForceSingleOptimization, bCalculateMotionVectors, bOptimizeIndexBuffers);
+		// EndCoding has to be called from the main thread once all the frame data have been added to the track
+
+		return StreamableTrack;
+	}
+
+	UGeometryCache* CreateGeometryCache(const FString& RootPrimPath, const FMeshDescription& MeshDescription, const TArray<UE::FSdfPath>& MeshPaths, const TArray<int32>& MaterialOffsets,
+										TSharedRef<FUsdSchemaTranslationContext> Context, bool& bOutIsNew)
+	{
+		// Compute the asset hash from the merged mesh description
+		FSHA1 SHA1;
+		FSHAHash MeshHash = FStaticMeshOperations::ComputeSHAHash(MeshDescription);
+		SHA1.Update(&MeshHash.Hash[0], sizeof(MeshHash.Hash));
+
+		const UE::FUsdStage& Stage = Context->Stage;
+		double FramesPerSecond = Stage.GetTimeCodesPerSecond();
+		if (FramesPerSecond == 0)
+		{
+			ensureMsgf(false, TEXT("Invalid USD GeometryCache FPS detected. Falling back to 1 FPS"));
+			FramesPerSecond = 1;
+		}
+
+		// Frame rate must be taken into account as well since different frame rates must produce different sampling in the tracks
+		SHA1.Update(reinterpret_cast<uint8*>(&FramesPerSecond), sizeof(FramesPerSecond));
+
+		// Track type depends on if it's importing or not. Import needs to generate a persistent asset with all the frames already sampled
+		SHA1.Update(reinterpret_cast<uint8*>(&Context->bIsImporting), sizeof(Context->bIsImporting));
+		SHA1.Final();
+
+		FSHAHash GeoCacheHash;
+		SHA1.GetHash(&GeoCacheHash.Hash[0]);
+		UGeometryCache* GeometryCache = Cast<UGeometryCache>(Context->AssetCache->GetCachedAsset(GeoCacheHash.ToString()));
+
+		if (!GeometryCache)
+		{
+			bOutIsNew = true;
+
+			const FName AssetName = MakeUniqueObjectName(GetTransientPackage(), UGeometryCache::StaticClass(), *IUsdClassesModule::SanitizeObjectName(FPaths::GetBaseFilename(RootPrimPath)));
+			GeometryCache = NewObject<UGeometryCache>(GetTransientPackage(), AssetName, Context->ObjectFlags | EObjectFlags::RF_Public | RF_Transient);
+
+			TOptional<FReadMeshDataArgs> Args;
+			if (!Context->bIsImporting)
+			{
+				Args = GetReadMeshDataArgs(Context, RootPrimPath);
+			}
+
+			// Create a track for each mesh to be processed and add it to the GeometryCache
+			for (int32 Index = 0; Index < MeshPaths.Num(); ++Index)
+			{
+				const FString& PrimPath = MeshPaths[Index].GetString();
+				UGeometryCacheTrack* Track = nullptr;
+				if (!Context->bIsImporting)
+				{
+					Track = CreateUsdStreamTrack(GeometryCache, *Args, PrimPath, MaterialOffsets[Index]);
+				}
+				else
+				{
+					Track = CreateStreamableTrack(GeometryCache, PrimPath);
+				}
+				GeometryCache->AddTrack(Track);
+
+				TArray<FMatrix> Mats;
+				Mats.Add(FMatrix::Identity);
+				Mats.Add(FMatrix::Identity);
+
+				TArray<float> MatTimes;
+				MatTimes.Add(0.0f);
+				MatTimes.Add(0.0f);
+				Track->SetMatrixSamples(Mats, MatTimes);
+			}
+
+			Context->AssetCache->CacheAsset(GeoCacheHash.ToString(), GeometryCache);
+		}
+		else
+		{
+			bOutIsNew = false;
+		}
+
+		return GeometryCache;
+	}
+
+	void FillGeometryCacheTracks(const FString& RootPrimPath, const TArray<UE::FSdfPath>& MeshPrims, const TArray<int32>& MaterialOffsets, TSharedRef<FUsdSchemaTranslationContext> Context, UGeometryCache* GeometryCache)
+	{
+		TRACE_CPUPROFILER_EVENT_SCOPE(UsdGeometryCacheTranslatorImpl::FillGeometryCacheTracks);
+
+		FReadMeshDataArgs Args = GetReadMeshDataArgs(Context, RootPrimPath);
+		Args.bPropagateTransform = true;
+
+		const bool bSingleThreaded = !FApp::ShouldUseThreadingForPerformance();
+		if (bSingleThreaded)
+		{
+			for (int32 Index = 0; Index < MeshPrims.Num(); ++Index)
+			{
+				UGeometryCacheTrackStreamable* StreamableTrack = Cast<UGeometryCacheTrackStreamable>(GeometryCache->Tracks[Index]);
+				UE::FUsdPrim MeshPrim = Args.Stage.GetPrimAtPath(MeshPrims[Index]);
+				const bool bConstantTopology = UsdUtils::GetMeshTopologyVariance(pxr::UsdGeomMesh(MeshPrim)) != UsdUtils::EMeshTopologyVariance::Heterogenous;
+
+				for (int32 FrameIndex = Args.StartFrame; FrameIndex < Args.EndFrame; ++FrameIndex)
+				{
+					// Read frame data
+					FGeometryCacheMeshData MeshData;
+					ReadMeshData(Args, MeshPrim, MaterialOffsets[Index], FrameIndex, MeshData);
+
+					// Add it to the track
+					StreamableTrack->AddMeshSample(MeshData, (FrameIndex - Args.StartFrame) / Args.FramesPerSecond, bConstantTopology);
+				}
+			}
+		}
+		else
+		{
+			// Balance the number of threads for mesh reads vs frame reads
+			const int32 NumMeshes = MeshPrims.Num();
+			const int32 MaxWorkerThreads = FTaskGraphInterface::Get().GetNumWorkerThreads();
+			int32 NumMeshThreads = FMath::Clamp(MaxWorkerThreads, 1, FMath::Min(NumMeshes, GUsdGeometryCacheParallelMeshReads));
+			int32 NumFrameThreads = FMath::Clamp(MaxWorkerThreads, 1, GUsdGeometryCacheParallelFrameReads);
+
+			int32 NumLoops = 0;
+			while (NumMeshThreads * NumFrameThreads > MaxWorkerThreads)
+			{
+				if (NumLoops % 2 == 0)
+				{
+					NumFrameThreads = FMath::Max(--NumFrameThreads, 1);
+				}
+				else
+				{
+					NumMeshThreads = FMath::Max(--NumMeshThreads, 1);
+				}
+				++NumLoops;
+			};
+
+			TArray<FEvent*> SyncEvents;
+			SyncEvents.SetNum(NumMeshThreads);
+			for (int32 Index = 0; Index < NumMeshThreads; ++Index)
+			{
+				SyncEvents[Index] = FPlatformProcess::GetSynchEventFromPool();
+			}
+
+			// Parallel mesh reads: Meshes can be read independently of each other
+			ParallelFor(NumMeshThreads, [&Args, &MeshPrims, &MaterialOffsets, &SyncEvents, NumMeshThreads, NumFrameThreads, GeometryCache, NumMeshes](int32 MeshThreadIndex)
+			{
+				int32 MeshIndex = MeshThreadIndex;
+
+				while (MeshIndex < NumMeshes)
+				{
+					UE::FUsdPrim MeshPrim = Args.Stage.GetPrimAtPath(MeshPrims[MeshIndex]);
+					int32 MaterialOffset = MaterialOffsets[MeshIndex];
+					UGeometryCacheTrackStreamable* StreamableTrack = Cast<UGeometryCacheTrackStreamable>(GeometryCache->Tracks[MeshIndex]);
+					FEvent* FrameWrittenEvent = SyncEvents[MeshThreadIndex];
+					const bool bConstantTopology = UsdUtils::GetMeshTopologyVariance(pxr::UsdGeomMesh(MeshPrim)) != UsdUtils::EMeshTopologyVariance::Heterogenous;
+
+					std::atomic<int32> WriteFrameIndex = Args.StartFrame;
+					FCriticalSection Mutex;
+
+					// Parallel frame read: frame data can be read concurrently but have to be processed in order for AddMeshSample
+					ParallelFor(NumFrameThreads, [&Args, &MeshPrim, &FrameWrittenEvent, &WriteFrameIndex, &Mutex, NumFrameThreads, MaterialOffset, StreamableTrack, bConstantTopology](int32 FrameThreadIndex)
+					{
+						int32 FrameIndex = Args.StartFrame + FrameThreadIndex;
+
+						while (FrameIndex < Args.EndFrame)
+						{
+							// Read frame data into memory
+							FGeometryCacheMeshData MeshData;
+							ReadMeshData(Args, MeshPrim, MaterialOffset, FrameIndex, MeshData);
+
+							// Wait until it's our turn to process this frame.
+							while (WriteFrameIndex < FrameIndex)
+							{
+								const uint32 WaitTimeInMs = 10;
+								FrameWrittenEvent->Wait(WaitTimeInMs);
+							}
+
+							{
+								FScopeLock WriteLock(&Mutex);
+
+								// Add it to the track
+								StreamableTrack->AddMeshSample(MeshData, (FrameIndex - Args.StartFrame) / Args.FramesPerSecond, bConstantTopology);
+
+								// Mark the next frame index as ready for processing.
+								++WriteFrameIndex;
+
+								FrameWrittenEvent->Trigger();
+							}
+
+							// Get new frame index to read for next run cycle
+							FrameIndex += NumFrameThreads;
+						};
+					});
+
+					MeshIndex += NumMeshThreads;
+				}
+			});
+
+			for (FEvent* SyncEvent : SyncEvents)
+			{
+				FPlatformProcess::ReturnSynchEventToPool(SyncEvent);
+			}
+		}
+	}
+
+	void FinalizeGeometryCache(UGeometryCache* GeometryCache)
+	{
+		for (int32 Index = 0; Index < GeometryCache->Tracks.Num(); ++Index)
+		{
+			if (UGeometryCacheTrackStreamable* StreamableTrack = Cast<UGeometryCacheTrackStreamable>(GeometryCache->Tracks[Index]))
+			{
+				StreamableTrack->EndCoding();
+			}
+		}
+	}
+
+	// #ueent_todo: Replace MeshDescription with RawMesh and also make it work with StaticMesh
+	void GeometryCacheDataForMeshDescription(FGeometryCacheMeshData& OutMeshData, FMeshDescription& MeshDescription, int32 MaterialIndex, float FramesPerSecond)
+	{
+		TRACE_CPUPROFILER_EVENT_SCOPE(GeometryCacheDataForMeshDescription);
+
+		OutMeshData.Positions.Reset();
+		OutMeshData.TextureCoordinates.Reset();
+		OutMeshData.TangentsX.Reset();
+		OutMeshData.TangentsZ.Reset();
+		OutMeshData.Colors.Reset();
+		OutMeshData.Indices.Reset();
+
+		OutMeshData.MotionVectors.Reset();
+		OutMeshData.BatchesInfo.Reset();
+		OutMeshData.BoundingBox.Init();
+
+		OutMeshData.VertexInfo.bHasColor0 = true;
+		OutMeshData.VertexInfo.bHasTangentX = true;
+		OutMeshData.VertexInfo.bHasTangentZ = true;
+		OutMeshData.VertexInfo.bHasUV0 = true;
+
+		FStaticMeshAttributes MeshDescriptionAttributes(MeshDescription);
+
+		TVertexAttributesConstRef< FVector3f > VertexPositions = MeshDescriptionAttributes.GetVertexPositions();
+		TVertexInstanceAttributesConstRef< FVector3f > VertexInstanceNormals = MeshDescriptionAttributes.GetVertexInstanceNormals();
+		TVertexInstanceAttributesConstRef< FVector3f > VertexInstanceTangents = MeshDescriptionAttributes.GetVertexInstanceTangents();
+		TVertexInstanceAttributesConstRef< float > VertexInstanceBinormalSigns = MeshDescriptionAttributes.GetVertexInstanceBinormalSigns();
+		TVertexInstanceAttributesConstRef< FVector4f > VertexInstanceColors = MeshDescriptionAttributes.GetVertexInstanceColors();
+		TVertexInstanceAttributesConstRef< FVector2f > VertexInstanceUVs = MeshDescriptionAttributes.GetVertexInstanceUVs();
+
+		TVertexInstanceAttributesConstRef< FVector3f > VertexInstanceVelocities = MeshDescription.VertexInstanceAttributes().GetAttributesRef<FVector3f>(MeshAttribute::VertexInstance::Velocity);
+
+		const bool bHasVelocities = VertexInstanceVelocities.IsValid();
+		OutMeshData.VertexInfo.bHasMotionVectors = bHasVelocities;
+
+		const int32 NumVertices = MeshDescription.Vertices().Num();
+		const int32 NumTriangles = MeshDescription.Triangles().Num();
+		const int32 NumMeshDataVertices = NumTriangles * 3;
+
+		OutMeshData.Positions.Reserve(NumVertices);
+		OutMeshData.Indices.Reserve(NumMeshDataVertices);
+		OutMeshData.TangentsZ.Reserve(NumMeshDataVertices);
+		OutMeshData.Colors.Reserve(NumMeshDataVertices);
+		OutMeshData.TextureCoordinates.Reserve(NumMeshDataVertices);
+		if (bHasVelocities)
+		{
+			OutMeshData.MotionVectors.Reserve(NumMeshDataVertices);
+		}
+
+		FBox BoundingBox(EForceInit::ForceInitToZero);
+		int32 VertexIndex = 0;
+		for (FPolygonGroupID PolygonGroupID : MeshDescription.PolygonGroups().GetElementIDs())
+		{
+			// Skip empty polygon groups
+			if (MeshDescription.GetNumPolygonGroupPolygons(PolygonGroupID) == 0)
+			{
+				continue;
+			}
+
+			FGeometryCacheMeshBatchInfo BatchInfo;
+			BatchInfo.StartIndex = OutMeshData.Indices.Num();
+			BatchInfo.MaterialIndex = MaterialIndex++;
+
+			int32 TriangleCount = 0;
+			for (FPolygonID PolygonID : MeshDescription.GetPolygonGroupPolygonIDs(PolygonGroupID))
+			{
+				for (FTriangleID TriangleID : MeshDescription.GetPolygonTriangles(PolygonID))
+				{
+					for (FVertexInstanceID VertexInstanceID : MeshDescription.GetTriangleVertexInstances(TriangleID))
+					{
+						const FVector3f& Position = VertexPositions[MeshDescription.GetVertexInstanceVertex(VertexInstanceID)];
+						OutMeshData.Positions.Add(Position);
+						BoundingBox += (FVector)Position;
+
+						OutMeshData.Indices.Add(VertexIndex++);
+
+						FPackedNormal Normal = VertexInstanceNormals[VertexInstanceID];
+						Normal.Vector.W = VertexInstanceBinormalSigns[VertexInstanceID] < 0 ? -127 : 127;
+						OutMeshData.TangentsZ.Add(Normal);
+						OutMeshData.TangentsX.Add(VertexInstanceTangents[VertexInstanceID]);
+
+						const bool bSRGB = true;
+						OutMeshData.Colors.Add(FLinearColor(VertexInstanceColors[VertexInstanceID]).ToFColor(bSRGB));
+
+						// Supporting only one UV channel
+						const int32 UVIndex = 0;
+						OutMeshData.TextureCoordinates.Add(VertexInstanceUVs.Get(VertexInstanceID, UVIndex));
+
+						if (bHasVelocities)
+						{
+							FVector3f MotionVector = VertexInstanceVelocities[VertexInstanceID];
+							MotionVector *= (-1.f / FramesPerSecond); // Velocity is per seconds but we need per frame for motion vectors
+
+							OutMeshData.MotionVectors.Add(MotionVector);
+						}
+					}
+
+					++TriangleCount;
+				}
+			}
+
+			OutMeshData.BoundingBox = (FBox3f)BoundingBox;
+
+			BatchInfo.NumTriangles = TriangleCount;
+			OutMeshData.BatchesInfo.Add(BatchInfo);
+		}
+	}
+
+	bool IsPurposefullyVisible(const pxr::UsdPrim& Prim, EUsdPurpose Purpose)
+	{
+		FScopedUsdAllocs UsdAllocs;
+
+		if (!EnumHasAllFlags(Purpose, IUsdPrim::GetPurpose(Prim)))
+		{
+			return false;
+		}
+
+		if (pxr::UsdGeomImageable UsdGeomImageable = pxr::UsdGeomImageable(Prim))
+		{
+			if (pxr::UsdAttribute VisibilityAttr = UsdGeomImageable.GetVisibilityAttr())
+			{
+				pxr::TfToken VisibilityToken;
+				if (VisibilityAttr.Get(&VisibilityToken) && VisibilityToken == pxr::UsdGeomTokens->invisible)
+				{
+					return false;
+				}
+			}
+		}
+		return true;
+	}
+}
+
+class FGeometryCacheCreateAssetsTaskChain : public FBuildStaticMeshTaskChain
+{
+public:
+	explicit FGeometryCacheCreateAssetsTaskChain(const TSharedRef< FUsdSchemaTranslationContext >& InContext, const UE::FSdfPath& InPrimPath)
+		: FBuildStaticMeshTaskChain(InContext, InPrimPath)
+	{
+		SetupTasks();
+	}
+
+protected:
+	virtual void SetupTasks() override;
+	TArray<UE::FSdfPath> MeshPrimPaths;
+	TArray<int32> MaterialOffsets;
+	TStrongObjectPtr<UGeometryCache> GeometryCache;
+};
+
+void FGeometryCacheCreateAssetsTaskChain::SetupTasks()
+{
+	// Create the mesh description (Async)
+	Do(ESchemaTranslationLaunchPolicy::Async,
+		[this]() -> bool
+		{
+			FScopedUnrealAllocs UnrealAllocs;
+
+			pxr::TfToken RenderContextToken = pxr::UsdShadeTokens->universalRenderContext;
+			if (!Context->RenderContext.IsNone())
+			{
+				RenderContextToken = UnrealToUsd::ConvertToken(*Context->RenderContext.ToString()).Get();
+			}
+
+			pxr::TfToken MaterialPurposeToken = pxr::UsdShadeTokens->allPurpose;
+			if (!Context->MaterialPurpose.IsNone())
+			{
+				MaterialPurposeToken = UnrealToUsd::ConvertToken(*Context->MaterialPurpose.ToString()).Get();
+			}
+
+			UsdToUnreal::FUsdMeshConversionOptions Options;
+			Options.TimeCode = UsdUtils::GetEarliestTimeCode();
+			Options.PurposesToLoad = Context->PurposesToLoad;
+			Options.RenderContext = RenderContextToken;
+			Options.MaterialPurpose = MaterialPurposeToken;
+			Options.bMergeIdenticalMaterialSlots = false; // Don't merge because the GeometryCache is processed as unflattened (ie. one track per mesh)
+
+			// GeometryCache has only one LOD so add just one MeshDescription and MaterialAssignmentInfo
+			FMeshDescription& AddedMeshDescription = LODIndexToMeshDescription.Emplace_GetRef();
+			UsdUtils::FUsdPrimMaterialAssignmentInfo& AssignmentInfo = LODIndexToMaterialInfo.Emplace_GetRef();
+
+			// The collapsed mesh description here will be used to cache the GeometryCache asset, but not to fill it since its content will be unflattened
+			// Bake the prim's transform into the mesh data
+			const bool bSkipRootPrimTransformAndVis = false;
+			UsdToUnreal::ConvertGeomMeshHierarchy(GetPrim(), AddedMeshDescription, AssignmentInfo, Options, bSkipRootPrimTransformAndVis);
+
+			// If we have at least one valid LOD, we should proceed to the next step
+			for (const FMeshDescription& MeshDescription : LODIndexToMeshDescription)
+			{
+				if (!MeshDescription.IsEmpty())
+				{
+					return true;
+				}
+			}
+			return false;
+		});
+
+	// Create the GeometryCache (Main thread)
+	Then(ESchemaTranslationLaunchPolicy::Sync,
+		[this]() -> bool
+		{
+			{
+				// Collect the visible meshes, animated or not, under Prim to be processed into the GeometryCache
+				FScopedUsdAllocs UsdAllocs;
+				TArray<TUsdStore<pxr::UsdPrim>> ChildPrims = UsdUtils::GetAllPrimsOfType(GetPrim(), pxr::TfType::Find<pxr::UsdGeomMesh>());
+
+				MeshPrimPaths.Reserve(ChildPrims.Num());
+				for (const TUsdStore<pxr::UsdPrim>& ChildPrim : ChildPrims)
+				{
+					if (UsdGeometryCacheTranslatorImpl::IsPurposefullyVisible(ChildPrim.Get(), Context->PurposesToLoad))
+					{
+						MeshPrimPaths.Add(UE::FSdfPath(ChildPrim.Get().GetPrimPath()));
+					}
+				}
+			}
+
+			if (MeshPrimPaths.Num() > 1)
+			{
+				FScopedUsdAllocs UsdAllocs;
+
+				// Compute the material offsets that will be needed to generate the GeometryCacheMeshData
+				// Each mesh will have its own material slots, but they are all appended into one GeometryCache
+				// so the offsets are just the number of material slots for each mesh added together in order
+				// of traversal of the meshes
+				MaterialOffsets.SetNum(MeshPrimPaths.Num());
+				int32 MaterialOffset = 0;
+				for (int32 Index = 0; Index < MeshPrimPaths.Num(); ++Index)
+				{
+					MaterialOffsets[Index] = MaterialOffset;
+
+					// A mesh has at least one material associated with it, but can have multiple material assignments based on its GeomSubsets
+					UE::FUsdPrim Prim = Context->Stage.GetPrimAtPath(MeshPrimPaths[Index]);
+
+					std::vector<pxr::UsdGeomSubset> GeomSubsets = pxr::UsdShadeMaterialBindingAPI(Prim).GetMaterialBindSubsets();
+					MaterialOffset += FMath::Max(1, static_cast<int32>(GeomSubsets.size()));
+				}
+			}
+			else
+			{
+				// Only one mesh, so there's no offset
+				MaterialOffsets.Add(0);
+			}
+
+			bool bIsNew = true;
+			const FString PrimPathString = PrimPath.GetString();
+			GeometryCache.Reset(UsdGeometryCacheTranslatorImpl::CreateGeometryCache(PrimPathString, LODIndexToMeshDescription[0], MeshPrimPaths, MaterialOffsets, Context, bIsNew));
+
+			if (GeometryCache && bIsNew)
+			{
+				UUsdMeshAssetImportData* ImportData = NewObject<UUsdMeshAssetImportData>(GeometryCache.Get(), TEXT("UUSDMeshAssetImportData"));
+				ImportData->PrimPath = PrimPathString;
+				GeometryCache->AssetImportData = ImportData;
+
+				// Only the original creator of the prim at creation time gets to set the material assignments
+				// directly on the geometry cache, all others prims ensure their materials via material overrides on the
+				// components
+				UsdGeometryCacheTranslatorImpl::ProcessGeometryCacheMaterials(
+					GetPrim(),
+					LODIndexToMaterialInfo,
+					*GeometryCache,
+					*Context->AssetCache.Get(),
+					Context->InfoCache.Get(),
+					Context->Time,
+					Context->ObjectFlags
+				);
+			}
+
+			if (GeometryCache && Context->InfoCache)
+			{
+				const UE::FSdfPath& TargetPath = AlternativePrimToLinkAssetsTo.IsSet() ? AlternativePrimToLinkAssetsTo.GetValue() : PrimPath;
+				Context->InfoCache->LinkAssetToPrim(TargetPath, GeometryCache.Get());
+			}
+
+			// Continue with the import steps
+			return Context->bIsImporting && GeometryCache;
+		});
+
+	// Fill the GeometryCache tracks with the frame data (Async)
+	Then(ESchemaTranslationLaunchPolicy::Async,
+		[this]() -> bool
+		{
+			UsdGeometryCacheTranslatorImpl::FillGeometryCacheTracks(PrimPath.GetString(), MeshPrimPaths, MaterialOffsets, Context, GeometryCache.Get());
+			return true;
+		});
+
+	// Finalize the GeometryCache (Main Thread)
+	Then(ESchemaTranslationLaunchPolicy::Sync,
+		[this]() -> bool
+		{
+			UsdGeometryCacheTranslatorImpl::FinalizeGeometryCache(GeometryCache.Get());
+			return false;
+		});
+}
+
+void FUsdGeometryCacheTranslator::CreateAssets()
+{
+	TRACE_CPUPROFILER_EVENT_SCOPE(FUsdGeometryCacheTranslator::CreateAssets);
+
+	if (!IsPotentialGeometryCacheRoot())
+	{
+		Super::CreateAssets();
+		return;
+	}
+
+	RegisterAuxiliaryPrims();
+
+	// Create the GeometryCache TaskChain
+	TSharedRef<FGeometryCacheCreateAssetsTaskChain> AssetsTaskChain = MakeShared<FGeometryCacheCreateAssetsTaskChain>(Context, PrimPath);
+
+	Context->TranslatorTasks.Add(MoveTemp(AssetsTaskChain));
+}
+
+USceneComponent* FUsdGeometryCacheTranslator::CreateComponents()
+{
+	TOptional<TSubclassOf<USceneComponent>> ComponentType;
+
+	if (!IsPotentialGeometryCacheRoot())
+	{
+		return Super::CreateComponents();
+	}
+
+	if (!Context->bIsImporting)
+	{
+		ComponentType = UGeometryCacheUsdComponent::StaticClass();
+	}
+	else
+	{
+		ComponentType = UGeometryCacheComponent::StaticClass();
+	}
+
+	USceneComponent* SceneComponent = CreateComponentsEx(ComponentType, {});
+	UpdateComponents(SceneComponent);
+
+	if (UGeometryCacheComponent* Component = Cast<UGeometryCacheComponent>(SceneComponent))
+	{
+		if (Context->InfoCache && Context->AssetCache)
+		{
+			if (UGeometryCache* GeometryCache = Context->InfoCache->GetSingleAssetForPrim<UGeometryCache>(PrimPath))
+			{
+				// Geometry caches don't support LODs
+				const bool bAllowInterpretingLODs = false;
+
+				MeshTranslationImpl::SetMaterialOverrides(
+					GetPrim(),
+					GeometryCache->Materials,
+					*Component,
+					*Context->AssetCache,
+					*Context->InfoCache,
+					Context->Time,
+					Context->ObjectFlags,
+					bAllowInterpretingLODs,
+					Context->RenderContext,
+					Context->MaterialPurpose
+				);
+
+				// Check if the prim has the GroomBinding schema and setup the component and assets necessary to bind the groom to the GeometryCache
+				if (UsdUtils::PrimHasSchema(GetPrim(), UnrealIdentifiers::GroomBindingAPI))
+				{
+					UsdGroomTranslatorUtils::CreateGroomBindingAsset(
+						GetPrim(),
+						*Context->AssetCache,
+						*Context->InfoCache,
+						Context->ObjectFlags
+					);
+
+					// For the groom binding to work, the GroomComponent must be a child of the SceneComponent
+					// so the Context ParentComponent is set to the SceneComponent temporarily
+					TGuardValue< USceneComponent* > ParentComponentGuard{Context->ParentComponent, SceneComponent};
+					const bool bNeedsActor = false;
+					UGroomComponent* GroomComponent = Cast< UGroomComponent >(CreateComponentsEx(TSubclassOf< USceneComponent >(UGroomComponent::StaticClass()), bNeedsActor));
+					if (GroomComponent)
+					{
+						UpdateComponents(SceneComponent);
+					}
+				}
+			}
+		}
+	}
+
+	return SceneComponent;
+}
+
+void FUsdGeometryCacheTranslator::UpdateComponents(USceneComponent* SceneComponent)
+{
+	if (!IsPotentialGeometryCacheRoot())
+	{
+		Super::UpdateComponents(SceneComponent);
+		return;
+	}
+
+	if (SceneComponent)
+	{
+		SceneComponent->Modify();
+	}
+
+	// Set the initial GeometryCache on the GeometryCacheComponent
+	if (UGeometryCacheComponent* GeometryCacheComponent = Cast<UGeometryCacheComponent>(SceneComponent))
+	{
+		UGeometryCache* GeometryCache = nullptr;
+		if (Context->InfoCache)
+		{
+			GeometryCache = Context->InfoCache->GetSingleAssetForPrim<UGeometryCache>(PrimPath);
+		}
+
+		bool bShouldRegister = false;
+		if (GeometryCache != GeometryCacheComponent->GetGeometryCache())
+		{
+			bShouldRegister = true;
+
+			if (GeometryCacheComponent->IsRegistered())
+			{
+				GeometryCacheComponent->UnregisterComponent();
+			}
+
+			// Skip the extra handling in SetGeometryCache
+			GeometryCacheComponent->GeometryCache = GeometryCache;
+		}
+
+		// Manually tick USD GeometryCache only when their tracks are disabled in Sequencer
+		static IConsoleVariable* CVar = IConsoleManager::Get().FindConsoleVariable(TEXT("USD.DisableGeoCacheTracks"));
+		bool bDisableGeoCacheTracks = CVar && CVar->GetBool();
+		if (bDisableGeoCacheTracks)
+		{
+			float TimeCode = Context->Time;
+			if (FMath::IsNaN(TimeCode))
+			{
+				int32 StartFrame = FMath::FloorToInt(Context->Stage.GetStartTimeCode());
+				int32 EndFrame = FMath::CeilToInt(Context->Stage.GetEndTimeCode());
+				UsdGeometryCacheTranslatorImpl::GetGeometryCacheDataTimeCodeRange(Context->Stage, PrimPath.GetString(), StartFrame, EndFrame);
+
+				TimeCode = static_cast<float>(StartFrame);
+			}
+
+			// This is the main call responsible for animating the geometry cache.
+			// It needs to happen after setting the geometry cache and before registering, because we must force the
+			// geometry cache to register itself at Context->Time so that it will synchronously load that frame right away.
+			// Otherwise the geometry cache will start at t=0 regardless of Context->Time
+			GeometryCacheComponent->SetManualTick(true);
+			GeometryCacheComponent->TickAtThisTime(TimeCode, true, false, true);
+		}
+
+		// Note how we should only register if our geometry cache changed: If we did this every time we would
+		// register too early during the process of duplicating into PIE, and that would prevent a future RegisterComponent
+		// call from naturally creating the required render state
+		if (bShouldRegister && !GeometryCacheComponent->IsRegistered())
+		{
+			GeometryCacheComponent->RegisterComponent();
+		}
+
+		// If the prim has a GroomBinding schema, apply the target groom to its associated GroomComponent
+		if (UsdUtils::PrimHasSchema(GetPrim(), UnrealIdentifiers::GroomBindingAPI))
+		{
+			UsdGroomTranslatorUtils::SetGroomFromPrim(GetPrim(), *Context->InfoCache, SceneComponent);
+		}
+	}
+}
+
+bool FUsdGeometryCacheTranslator::CollapsesChildren(ECollapsingType CollapsingType) const
+{
+	return IsPotentialGeometryCacheRoot() ? true : Super::CollapsesChildren(CollapsingType);
+}
+
+bool FUsdGeometryCacheTranslator::CanBeCollapsed(ECollapsingType CollapsingType) const
+{
+	return IsPotentialGeometryCacheRoot() ? false : Super::CanBeCollapsed(CollapsingType);
+}
+
+TSet<UE::FSdfPath> FUsdGeometryCacheTranslator::CollectAuxiliaryPrims() const
+{
+	if (!IsPotentialGeometryCacheRoot())
+	{
+		return Super::CollectAuxiliaryPrims();
+	}
+
+	TSet<UE::FSdfPath> Result;
+	{
+		FScopedUsdAllocs UsdAllocs;
+
+		// Here, we collect all meshes even non-animated ones since they'll be collapse into the cache
+		TArray<TUsdStore<pxr::UsdPrim>> ChildPrims = UsdUtils::GetAllPrimsOfType(GetPrim(), pxr::TfType::Find<pxr::UsdGeomMesh>());
+
+		Result.Reserve(ChildPrims.Num());
+		for (const TUsdStore<pxr::UsdPrim>& ChildPrim : ChildPrims)
+		{
+			if (UsdGeometryCacheTranslatorImpl::IsPurposefullyVisible(ChildPrim.Get(), Context->PurposesToLoad))
+			{
+				Result.Add(UE::FSdfPath(ChildPrim.Get().GetPrimPath()));
+			}	
+			
+			// Should probably collect the Xformable as well
+		}
+	}
+	return Result;
+}
+
+bool FUsdGeometryCacheTranslator::IsPotentialGeometryCacheRoot() const
+{
+	if (!GUseGeometryCacheUSD)
+	{
+		return false;
+	}
+
+	if (!Context->bIsImporting)
+	{
+		// For now, when loading on the stage, process only leaf animated mesh as geometry cache (without collapsing onto a root)
+		return UsdUtils::IsAnimatedMesh(GetPrim());
+	}
+
+	// Cache is not available when it is being rebuilt. This is called by RecursiveQueryCollapsesChildren
+	bool bIsPotentialGeoCacheRoot = false;
+	if (Context->InfoCache)
+	{
+		bIsPotentialGeoCacheRoot = Context->InfoCache->IsPotentialGeometryCacheRoot(PrimPath);
+	}
+	else
+	{
+		FScopedUsdAllocs UsdAllocs;
+		TArray<TUsdStore<pxr::UsdPrim>> ChildPrims = UsdUtils::GetAllPrimsOfType(GetPrim(), pxr::TfType::Find<pxr::UsdGeomMesh>());
+		for (const TUsdStore<pxr::UsdPrim>& ChildPrim : ChildPrims)
+		{
+			if (UsdUtils::IsAnimatedMesh(ChildPrim.Get()))
+			{
+				bIsPotentialGeoCacheRoot = true;
+				break;
+			}
+		}
+	}
+
+	// #uent_todo: Improve the collapsing rule for GeometryCache
+	if (bIsPotentialGeoCacheRoot)
+	{
+		FString FilePath;
+		if (UsdUtils::GetReferenceFilePath(GetPrim(), TEXT("abc"), FilePath))
+		{
+			return true;
+		}
+		else
+		{
+			// TBD where the collapsing should occur. Currently, if the stage has any animated meshes, it will collapse everything at the root
+			return true;
+		}
+	}
+
+	return false;
+}
+
+#endif // #if USE_USD_SDK
diff --git a/Engine/Plugins/Importers/USDImporter/Source/USDSchemas/Private/USDInfoCache.cpp b/Engine/Plugins/Importers/USDImporter/Source/USDSchemas/Private/USDInfoCache.cpp
--- a/Engine/Plugins/Importers/USDImporter/Source/USDSchemas/Private/USDInfoCache.cpp
+++ b/Engine/Plugins/Importers/USDImporter/Source/USDSchemas/Private/USDInfoCache.cpp
@@ -41,6 +41,7 @@ namespace UE::UsdInfoCache::Private
 		UE::FSdfPath ComponentCollapsedRoot;
 		uint64 ExpectedVertexCountForSubtree;
 		TArray<UsdUtils::FUsdPrimMaterialSlot> SubtreeMaterialSlots;
+		bool bIsPotentialGeoCacheRoot;
 	};
 }
 
@@ -50,6 +51,7 @@ FArchive& operator <<( FArchive& Ar, UE::UsdInfoCache::Private::FUsdPrimInfo& In
 	Ar << Info.ComponentCollapsedRoot;
 	Ar << Info.ExpectedVertexCountForSubtree;
 	Ar << Info.SubtreeMaterialSlots;
+	Ar << Info.bIsPotentialGeoCacheRoot;
 
 	return Ar;
 }
@@ -860,7 +862,62 @@ namespace UE::USDInfoCacheImpl::Private
 			Info.ComponentCollapsedRoot = *ComponentCollapsedRootOverride;
 		}
 	}
-#endif // USE_USD_SDK
+
+	bool RecursiveCheckForGeometryCache(const pxr::UsdPrim& UsdPrim, FUsdSchemaTranslationContext& Context, FUsdInfoCache::FUsdInfoCacheImpl& Impl)
+	{
+		TRACE_CPUPROFILER_EVENT_SCOPE(UE::USDInfoCacheImpl::Private::RecursiveCheckForGeometryCache);
+		FScopedUsdAllocs Allocs;
+
+		pxr::SdfPath UsdPrimPath = UsdPrim.GetPrimPath();
+		UE::FSdfPath PrimPath{UsdPrimPath};
+
+		pxr::UsdPrimSiblingRange PrimChildren = UsdPrim.GetFilteredChildren(pxr::UsdTraverseInstanceProxies(pxr::UsdPrimAllPrimsPredicate));
+
+		TArray<pxr::UsdPrim> Prims;
+		for (pxr::UsdPrim Child : PrimChildren)
+		{
+			Prims.Emplace(Child);
+		}
+
+		std::atomic<bool> bIsPotentialGeoCacheRoot = true;
+		const int32 MinBatchSize = 1;
+		ParallelFor(TEXT("RecursiveCheckForGeometryCache"), Prims.Num(), MinBatchSize,
+			[&](int32 Index)
+			{
+				bool bResult = RecursiveCheckForGeometryCache(Prims[Index], Context, Impl);
+				bIsPotentialGeoCacheRoot = bIsPotentialGeoCacheRoot && bResult;
+			}
+		);
+
+		bool bIsAnimatedMesh = UsdUtils::IsAnimatedMesh(UsdPrim);
+		{
+			FWriteScopeLock ScopeLock(Impl.InfoMapLock);
+
+			UE::UsdInfoCache::Private::FUsdPrimInfo& Info = Impl.InfoMap.FindOrAdd(UE::FSdfPath(UsdPrimPath));
+			Info.bIsPotentialGeoCacheRoot = bIsPotentialGeoCacheRoot || bIsAnimatedMesh;
+		}
+
+		return bIsPotentialGeoCacheRoot || bIsAnimatedMesh;
+	}
+#endif // USE_SD_SDK
+}
+
+bool FUsdInfoCache::IsPotentialGeometryCacheRoot(const UE::FSdfPath& Path) const
+{
+	if (FUsdInfoCacheImpl* ImplPtr = Impl.Get())
+	{
+		FReadScopeLock ScopeLock(ImplPtr->InfoMapLock);
+
+		if (const UE::UsdInfoCache::Private::FUsdPrimInfo* FoundInfo = ImplPtr->InfoMap.Find(Path))
+		{
+			return FoundInfo->bIsPotentialGeoCacheRoot;
+		}
+
+		// This should never happen: We should have cached the entire tree
+		ensureMsgf(false, TEXT("Prim path '%s' has not been cached!"), *Path.GetString());
+	}
+
+	return false;
 }
 
 TOptional<uint64> FUsdInfoCache::GetSubtreeVertexCount( const UE::FSdfPath& Path )
@@ -1091,6 +1148,12 @@ void FUsdInfoCache::RebuildCacheForSubtree( const UE::FUsdPrim& Prim, FUsdSchema
 			Context.bMergeIdenticalMaterialSlots
 		);
 
+		UE::USDInfoCacheImpl::Private::RecursiveCheckForGeometryCache(			
+			UsdPrim,
+			Context,
+			*ImplPtr
+			);
+
 		UE::USDInfoCacheImpl::Private::RecursiveQueryCollapsesChildren(
 			UsdPrim,
 			Context,
diff --git a/Engine/Plugins/Importers/USDImporter/Source/USDSchemas/Private/USDSchemasModule.cpp b/Engine/Plugins/Importers/USDImporter/Source/USDSchemas/Private/USDSchemasModule.cpp
--- a/Engine/Plugins/Importers/USDImporter/Source/USDSchemas/Private/USDSchemasModule.cpp
+++ b/Engine/Plugins/Importers/USDImporter/Source/USDSchemas/Private/USDSchemasModule.cpp
@@ -13,6 +13,7 @@
 
 #if USE_USD_SDK
 #include "USDGeomCameraTranslator.h"
+#include "USDGeometryCacheTranslator.h"
 #include "USDGeomMeshTranslator.h"
 #include "USDGeomPointInstancerTranslator.h"
 #include "USDGeomXformableTranslator.h"
@@ -52,6 +53,8 @@ public:
 		{
 			UsdSkelRootTranslatorHandle = GetTranslatorRegistry().Register< FUsdSkelRootTranslator >( TEXT("UsdSkelRoot") );
 			UsdGroomTranslatorHandle = GetTranslatorRegistry().Register< FUsdGroomTranslator >( TEXT("UsdGeomXformable") );
+			// The GeometryCacheTranslator also works on UsdGeomXformable through the GroomTranslator
+			UsdGeometryCacheTranslatorHandle = GetTranslatorRegistry().Register< FUsdGeometryCacheTranslator >( TEXT("UsdGeomMesh") );
 
 			if ( IMDLImporterModule* MDLImporterModule = FModuleManager::Get().LoadModulePtr< IMDLImporterModule >( TEXT("MDLImporter") ) )
 			{
@@ -76,6 +79,7 @@ public:
 		{
 			GetTranslatorRegistry().Unregister( UsdSkelRootTranslatorHandle );
 			GetTranslatorRegistry().Unregister( UsdGroomTranslatorHandle );
+			GetTranslatorRegistry().Unregister( UsdGeometryCacheTranslatorHandle );
 
 			GetTranslatorRegistry().Unregister( MdlUsdShadeMaterialTranslatorHandle );
 			GetRenderContextRegistry().Unregister(FMdlUsdShadeMaterialTranslator::MdlRenderContext);
@@ -107,6 +111,7 @@ protected:
 	FUsdRenderContextRegistry UsdRenderContextRegistry;
 
 	FRegisteredSchemaTranslatorHandle UsdGeomCameraTranslatorHandle;
+	FRegisteredSchemaTranslatorHandle UsdGeometryCacheTranslatorHandle;
 	FRegisteredSchemaTranslatorHandle UsdGeomMeshTranslatorHandle;
 	FRegisteredSchemaTranslatorHandle UsdGeomPointInstancerTranslatorHandle;
 	FRegisteredSchemaTranslatorHandle UsdSkelRootTranslatorHandle;
diff --git a/Engine/Plugins/Importers/USDImporter/Source/USDSchemas/Public/USDGeometryCacheTranslator.h b/Engine/Plugins/Importers/USDImporter/Source/USDSchemas/Public/USDGeometryCacheTranslator.h
new file mode 100644
--- /dev/null
+++ b/Engine/Plugins/Importers/USDImporter/Source/USDSchemas/Public/USDGeometryCacheTranslator.h
@@ -0,0 +1,32 @@
+// Copyright Epic Games, Inc. All Rights Reserved.
+
+#pragma once
+
+#if USE_USD_SDK && WITH_EDITOR
+
+#include "USDGeomMeshTranslator.h"
+
+class USDSCHEMAS_API FUsdGeometryCacheTranslator : public FUsdGeomMeshTranslator
+{
+public:
+	using Super = FUsdGeomMeshTranslator;
+
+	using FUsdGeomMeshTranslator::FUsdGeomMeshTranslator;
+
+	FUsdGeometryCacheTranslator(const FUsdGeometryCacheTranslator& Other) = delete;
+	FUsdGeometryCacheTranslator& operator=(const FUsdGeometryCacheTranslator& Other) = delete;
+
+	virtual void CreateAssets() override;
+	virtual USceneComponent* CreateComponents() override;
+	virtual void UpdateComponents(USceneComponent* SceneComponent) override;
+
+	virtual bool CollapsesChildren(ECollapsingType CollapsingType) const override;
+	virtual bool CanBeCollapsed(ECollapsingType CollapsingType) const override;
+
+	virtual TSet<UE::FSdfPath> CollectAuxiliaryPrims() const override;
+
+private:
+	bool IsPotentialGeometryCacheRoot() const;
+};
+
+#endif // #if USE_USD_SDK
diff --git a/Engine/Plugins/Importers/USDImporter/Source/USDSchemas/Public/USDGroomTranslator.h b/Engine/Plugins/Importers/USDImporter/Source/USDSchemas/Public/USDGroomTranslator.h
--- a/Engine/Plugins/Importers/USDImporter/Source/USDSchemas/Public/USDGroomTranslator.h
+++ b/Engine/Plugins/Importers/USDImporter/Source/USDSchemas/Public/USDGroomTranslator.h
@@ -4,13 +4,14 @@
 
 #if USE_USD_SDK && WITH_EDITOR
 
+#include "USDGeometryCacheTranslator.h"
 #include "USDGeomXformableTranslator.h"
 
-class USDSCHEMAS_API FUsdGroomTranslator : public FUsdGeomXformableTranslator
+class USDSCHEMAS_API FUsdGroomTranslator : public FUsdGeometryCacheTranslator
 {
 public:
-	using Super = FUsdGeomXformableTranslator;
-	using FUsdGeomXformableTranslator::FUsdGeomXformableTranslator;
+	using Super = FUsdGeometryCacheTranslator;
+	using FUsdGeometryCacheTranslator::FUsdGeometryCacheTranslator;
 
 	virtual void CreateAssets() override;
 
diff --git a/Engine/Plugins/Importers/USDImporter/Source/USDSchemas/Public/USDInfoCache.h b/Engine/Plugins/Importers/USDImporter/Source/USDSchemas/Public/USDInfoCache.h
--- a/Engine/Plugins/Importers/USDImporter/Source/USDSchemas/Public/USDInfoCache.h
+++ b/Engine/Plugins/Importers/USDImporter/Source/USDSchemas/Public/USDInfoCache.h
@@ -86,6 +86,9 @@ public:
 	TOptional<uint64> GetSubtreeMaterialSlotCount(const UE::FSdfPath& Path);
 	TOptional<TArray<UsdUtils::FUsdPrimMaterialSlot>> GetSubtreeMaterialSlots(const UE::FSdfPath& Path);
 
+	// Returns true if Path could potentially be collapsed as a Geometry Cache asset
+	bool IsPotentialGeometryCacheRoot(const UE::FSdfPath& Path) const;
+
 public:
 	void LinkAssetToPrim(const UE::FSdfPath& Path, UObject* Asset);
 
diff --git a/Engine/Plugins/Importers/USDImporter/Source/USDStageImporter/Private/USDStageImporter.cpp b/Engine/Plugins/Importers/USDImporter/Source/USDStageImporter/Private/USDStageImporter.cpp
--- a/Engine/Plugins/Importers/USDImporter/Source/USDStageImporter/Private/USDStageImporter.cpp
+++ b/Engine/Plugins/Importers/USDImporter/Source/USDStageImporter/Private/USDStageImporter.cpp
@@ -791,6 +791,7 @@ namespace UsdStageImporterImpl
 		TArray<UObject*> AnimSequences;
 		TArray<UObject*> AnimBlueprints;
 		TArray<UObject*> LevelSequences;
+		TArray<UObject*> GeometryCaches;
 
 		TSet<FString> UniqueAssetNames;
 		TMap<UObject*, FString> AssetToContentFolder;
@@ -807,14 +808,10 @@ namespace UsdStageImporterImpl
 			const static FString StaticMeshesFolder   = TEXT("StaticMeshes");
 			const static FString SkeletalMeshesFolder = TEXT("SkeletalMeshes");
 			const static FString LevelSequencesFolder = TEXT("LevelSequences");
+			const static FString GeometryCachesFolder = TEXT("GeometryCaches");
 
 			const FString* AssetTypeFolderPtr = nullptr;
-			if (Asset->IsA(UGeometryCache::StaticClass()))
-			{
-				UE_LOG(LogUsd, Warning, TEXT("Ignoring asset '%s': Importing GeometryCaches assets from USD is not supported at this time"), *Asset->GetName());
-				continue;
-			}
-			else if (Asset->IsA(UMaterialInterface::StaticClass()))
+			if (Asset->IsA(UMaterialInterface::StaticClass()))
 			{
 				AssetTypeFolderPtr = &MaterialsFolder;
 				Materials.Add(Asset);
@@ -859,6 +856,11 @@ namespace UsdStageImporterImpl
 				AssetTypeFolderPtr = &LevelSequencesFolder;
 				LevelSequences.Add(Asset);
 			}
+			else if (Asset->IsA(UGeometryCache::StaticClass()))
+			{
+				AssetTypeFolderPtr = &GeometryCachesFolder;
+				GeometryCaches.Add(Asset);
+			}
 			else
 			{
 				// We don't know what to do with this asset
@@ -911,6 +913,7 @@ namespace UsdStageImporterImpl
 		PublishAssetType(Skeletons);
 		PublishAssetType(PhysicsAssets);
 		PublishAssetType(StaticMeshes);
+		PublishAssetType(GeometryCaches);
 		PublishAssetType(Materials);
 		PublishAssetType(Textures);
 	}
@@ -1700,6 +1703,7 @@ void UUsdStageImporter::ImportFromFile(FUsdStageImportContext& ImportContext)
 	ImportContext.LevelSequenceHelper.SetInfoCache(ImportContext.InfoCache);
 	ImportContext.LevelSequenceHelper.Init( ImportContext.Stage );  // Must happen after the context gets an InfoCache!
 
+	{
 	// Shotgun approach to recreate all render states because we may want to reimport/delete/reassign a material/static/skeletalmesh while it is currently being drawn
 	FGlobalComponentRecreateRenderStateContext RecreateRenderStateContext;
 
@@ -1752,9 +1756,11 @@ void UUsdStageImporter::ImportFromFile(FUsdStageImportContext& ImportContext)
 	{
 		double ElapsedSeconds = FPlatformTime::ToSeconds64( FPlatformTime::Cycles64() - StartTime );
 		UsdStageImporterImpl::SendAnalytics( ImportContext, nullptr, TEXT("Import"), PublishedAssetsAndDependencies, ElapsedSeconds);
+		UE_LOG(LogUsd, Log, TEXT("Imported '%s' in %.3f seconds."), *ImportContext.FilePath, ElapsedSeconds);
 	}
 
 	UsdStageImporterImpl::CloseStageIfNeeded( ImportContext );
+	}
 
 	if ( ImportContext.bNeedsGarbageCollection )
 	{
@@ -1817,6 +1823,7 @@ bool UUsdStageImporter::ReimportSingleAsset(FUsdStageImportContext& ImportContex
 	ImportContext.LevelSequenceHelper.SetInfoCache(ImportContext.InfoCache);
 	ImportContext.LevelSequenceHelper.Init(ImportContext.Stage);  // Must happen after the context gets an InfoCache!
 
+	{
 	// Shotgun approach to recreate all render states because we may want to reimport/delete/reassign a material/static/skeletalmesh while it is currently being drawn
 	FGlobalComponentRecreateRenderStateContext RecreateRenderStateContext;
 
@@ -1899,9 +1906,11 @@ bool UUsdStageImporter::ReimportSingleAsset(FUsdStageImportContext& ImportContex
 	{
 		double ElapsedSeconds = FPlatformTime::ToSeconds64( FPlatformTime::Cycles64() - StartTime );
 		UsdStageImporterImpl::SendAnalytics( ImportContext, ReimportedObject, TEXT( "Reimport" ), { ReimportedObject }, ElapsedSeconds );
+		UE_LOG(LogUsd, Log, TEXT("Re-imported '%s' in %.3f seconds."), *ImportContext.FilePath, ElapsedSeconds);
 	}
 
 	UsdStageImporterImpl::CloseStageIfNeeded( ImportContext );
+	}
 
 	if ( ImportContext.bNeedsGarbageCollection )
 	{
diff --git a/Engine/Plugins/Importers/USDImporter/Source/USDUtilities/Private/USDConversionUtils.cpp b/Engine/Plugins/Importers/USDImporter/Source/USDUtilities/Private/USDConversionUtils.cpp
--- a/Engine/Plugins/Importers/USDImporter/Source/USDUtilities/Private/USDConversionUtils.cpp
+++ b/Engine/Plugins/Importers/USDImporter/Source/USDUtilities/Private/USDConversionUtils.cpp
@@ -854,22 +854,22 @@ TMap<FString, int32> UsdUtils::CombinePrimvarsIntoUVSets(
 	return PrimvarToUVIndex;
 }
 
-bool UsdUtils::IsAnimated( const pxr::UsdPrim& Prim )
+bool UsdUtils::IsAnimated(const pxr::UsdPrim& Prim)
 {
-	if ( !Prim || !Prim.IsActive() )
+	if (!Prim || !Prim.IsActive())
 	{
 		return false;
 	}
 
 	FScopedUsdAllocs UsdAllocs;
 
-	pxr::UsdGeomXformable Xformable( Prim );
-	if ( Xformable )
+	pxr::UsdGeomXformable Xformable(Prim);
+	if (Xformable)
 	{
 		std::vector< double > TimeSamples;
-		Xformable.GetTimeSamples( &TimeSamples );
+		Xformable.GetTimeSamples(&TimeSamples);
 
-		if ( TimeSamples.size() > 0 )
+		if (TimeSamples.size() > 0)
 		{
 			return true;
 		}
@@ -907,7 +907,7 @@ bool UsdUtils::IsAnimated( const pxr::UsdPrim& Prim )
 	}
 
 	const std::vector< pxr::UsdAttribute >& Attributes = Prim.GetAttributes();
-	for ( const pxr::UsdAttribute& Attribute : Attributes )
+	for (const pxr::UsdAttribute& Attribute : Attributes)
 	{
 		std::vector<double> TimeSamples;
 		if ( Attribute.GetTimeSamples( &TimeSamples ) && TimeSamples.size() > 0 )
@@ -916,20 +916,20 @@ bool UsdUtils::IsAnimated( const pxr::UsdPrim& Prim )
 		}
 	}
 
-	if ( pxr::UsdSkelRoot SkeletonRoot{ Prim } )
+	if (pxr::UsdSkelRoot SkeletonRoot{Prim})
 	{
 		pxr::UsdSkelCache SkeletonCache;
-		SkeletonCache.Populate( SkeletonRoot, pxr::UsdTraverseInstanceProxies() );
+		SkeletonCache.Populate(SkeletonRoot, pxr::UsdTraverseInstanceProxies());
 
 		std::vector< pxr::UsdSkelBinding > SkeletonBindings;
-		SkeletonCache.ComputeSkelBindings( SkeletonRoot, &SkeletonBindings, pxr::UsdTraverseInstanceProxies() );
+		SkeletonCache.ComputeSkelBindings(SkeletonRoot, &SkeletonBindings, pxr::UsdTraverseInstanceProxies());
 
-		for ( const pxr::UsdSkelBinding& Binding : SkeletonBindings )
+		for (const pxr::UsdSkelBinding& Binding : SkeletonBindings)
 		{
 			const pxr::UsdSkelSkeleton& Skeleton = Binding.GetSkeleton();
-			pxr::UsdSkelSkeletonQuery SkelQuery = SkeletonCache.GetSkelQuery( Skeleton );
+			pxr::UsdSkelSkeletonQuery SkelQuery = SkeletonCache.GetSkelQuery(Skeleton);
 			pxr::UsdSkelAnimQuery AnimQuery = SkelQuery.GetAnimQuery();
-			if ( !AnimQuery )
+			if (!AnimQuery)
 			{
 				continue;
 			}
@@ -1436,6 +1436,37 @@ void UsdUtils::AddReference( UE::FUsdPrim& Prim, const TCHAR* AbsoluteFilePath,
 #endif // #if USE_USD_SDK
 }
 
+bool UsdUtils::GetReferenceFilePath(const UE::FUsdPrim& Prim, const FString& FileExtension, FString& OutReferenceFilePath)
+{
+#if USE_USD_SDK
+	FScopedUsdAllocs UsdAllocs;
+
+	pxr::UsdPrimCompositionQuery PrimCompositionQuery = pxr::UsdPrimCompositionQuery::GetDirectReferences(Prim);
+	for (const pxr::UsdPrimCompositionQueryArc& CompositionArc : PrimCompositionQuery.GetCompositionArcs())
+	{
+		if (CompositionArc.GetArcType() == pxr::PcpArcTypeReference)
+		{
+			pxr::SdfReferenceEditorProxy ReferenceEditor;
+			pxr::SdfReference UsdReference;
+
+			if (CompositionArc.GetIntroducingListEditor(&ReferenceEditor, &UsdReference))
+			{
+				FString AbsoluteFilePath = UsdToUnreal::ConvertString(UsdReference.GetAssetPath());
+
+				FString Extension = FPaths::GetExtension(AbsoluteFilePath);
+				if (Extension == FileExtension && FPaths::FileExists(AbsoluteFilePath))
+				{
+					OutReferenceFilePath = AbsoluteFilePath;
+					return true;
+				}
+			}
+		}
+	}
+#endif // #if USE_USD_SDK
+
+	return false;
+}
+
 void UsdUtils::AddPayload( UE::FUsdPrim& Prim, const TCHAR* AbsoluteFilePath, const UE::FSdfPath& TargetPrimPath, double TimeCodeOffset, double TimeCodeScale )
 {
 #if USE_USD_SDK
diff --git a/Engine/Plugins/Importers/USDImporter/Source/USDUtilities/Private/USDGeomMeshConversion.cpp b/Engine/Plugins/Importers/USDImporter/Source/USDUtilities/Private/USDGeomMeshConversion.cpp
--- a/Engine/Plugins/Importers/USDImporter/Source/USDUtilities/Private/USDGeomMeshConversion.cpp
+++ b/Engine/Plugins/Importers/USDImporter/Source/USDUtilities/Private/USDGeomMeshConversion.cpp
@@ -65,6 +65,8 @@ static FAutoConsoleVariableRef CVarMaxInstancesPerPointInstancer(
 	GMaxInstancesPerPointInstancer,
 	TEXT( "We will only parse up to this many instances from any point instancer when reading from USD to UE. Set this to -1 to disable this limit." ) );
 
+const FName MeshAttribute::VertexInstance::Velocity("Velocity");
+
 namespace UE::UsdGeomMeshConversion::Private
 {
 	static const FString DisplayColorID = TEXT( "!DisplayColor" );
@@ -983,6 +985,29 @@ bool UsdToUnreal::ConvertGeomMesh(
 
 		pxr::TfToken NormalsInterpType = UsdMesh.GetNormalsInterpolation();
 
+		// Velocities
+		pxr::UsdAttribute VelocitiesAttribute = UsdMesh.GetVelocitiesAttr();
+		pxr::TfToken VelocitiesInterpType = pxr::UsdGeomTokens->vertex;
+		pxr::VtArray<pxr::GfVec3f> Velocities;
+		if (VelocitiesAttribute)
+		{
+			VelocitiesAttribute.Get(&Velocities, TimeCodeValue);
+
+			if (Velocities.size() > 0)
+			{
+				if (!OutMeshDescription.VertexInstanceAttributes().HasAttribute(MeshAttribute::VertexInstance::Velocity))
+				{
+					OutMeshDescription.VertexInstanceAttributes().RegisterAttribute<FVector3f>(MeshAttribute::VertexInstance::Velocity, 1, FVector3f::ZeroVector, EMeshAttributeFlags::Lerpable);
+				}
+
+				pxr::TfToken InterpType;
+				if (VelocitiesAttribute.GetMetadata(pxr::UsdGeomTokens->interpolation, &InterpType))
+				{
+					VelocitiesInterpType = InterpType;
+				}
+			}
+		}
+
 		// UVs
 		TVertexInstanceAttributesRef< FVector2f > MeshDescriptionUVs = StaticMeshAttributes.GetVertexInstanceUVs();
 
@@ -1090,6 +1115,7 @@ bool UsdToUnreal::ConvertGeomMesh(
 		}
 
 		TVertexInstanceAttributesRef< FVector3f > MeshDescriptionNormals = StaticMeshAttributes.GetVertexInstanceNormals();
+		TVertexInstanceAttributesRef<FVector3f> MeshDescriptionVelocities = OutMeshDescription.VertexInstanceAttributes().GetAttributesRef<FVector3f>(MeshAttribute::VertexInstance::Velocity);
 
 		OutMeshDescription.ReserveNewVertexInstances( FaceCounts.size() * 3 );
 		OutMeshDescription.ReserveNewPolygons( FaceCounts.size() );
@@ -1164,6 +1190,19 @@ bool UsdToUnreal::ConvertGeomMesh(
 					}
 				}
 
+				if (Velocities.size() > 0)
+				{
+					const int32 VelocityIndex = UsdGeomMeshImpl::GetPrimValueIndex(VelocitiesInterpType, ControlPointIndex, CurrentVertexInstanceIndex, PolygonIndex);
+
+					if (VelocityIndex < Velocities.size())
+					{
+						const pxr::GfVec3f& Velocity = Velocities[VelocityIndex];
+						FVector TransformedVelocity = Options.AdditionalTransform.TransformVector(UsdToUnreal::ConvertVector(StageInfo, Velocity));
+
+						MeshDescriptionVelocities[AddedVertexInstanceId] = (FVector3f) TransformedVelocity;
+					}
+				}
+
 				for ( const FUVSet& UVSet : UVSets )
 				{
 					const int32 ValueIndex = UsdGeomMeshImpl::GetPrimValueIndex( UVSet.InterpType, ControlPointIndex, CurrentVertexInstanceIndex, PolygonIndex );
@@ -3094,6 +3133,81 @@ bool UsdUtils::GetPointInstancerTransforms( const FUsdStageInfo& StageInfo, cons
 	return true;
 }
 
+bool UsdUtils::IsAnimatedMesh(const pxr::UsdPrim& Prim)
+{
+	FScopedUsdAllocs UsdAllocs;
+
+	pxr::UsdGeomMesh Mesh = pxr::UsdGeomMesh(Prim);
+	if (!Mesh)
+	{
+		return false;
+	}
+
+	bool bHasAttributesTimeSamples = false;
+	{
+		constexpr bool bIncludeInherited = false;
+		pxr::TfTokenVector GeomMeshAttributeNames = pxr::UsdGeomMesh::GetSchemaAttributeNames(bIncludeInherited);
+		pxr::TfTokenVector GeomPointBasedAttributeNames = pxr::UsdGeomPointBased::GetSchemaAttributeNames(bIncludeInherited);
+
+		GeomMeshAttributeNames.reserve(GeomMeshAttributeNames.size() + GeomPointBasedAttributeNames.size());
+		GeomMeshAttributeNames.insert(GeomMeshAttributeNames.end(), GeomPointBasedAttributeNames.begin(), GeomPointBasedAttributeNames.end());
+
+		for (const pxr::TfToken& AttributeName : GeomMeshAttributeNames)
+		{
+			const pxr::UsdAttribute& Attribute = Prim.GetAttribute(AttributeName);
+
+			if (Attribute && Attribute.ValueMightBeTimeVarying())
+			{
+				bHasAttributesTimeSamples = true;
+				break;
+			}
+		}
+	}
+
+	return bHasAttributesTimeSamples;
+}
+
+UsdUtils::EMeshTopologyVariance UsdUtils::GetMeshTopologyVariance(const pxr::UsdGeomMesh& UsdMesh)
+{
+	FScopedUsdAllocs UsdAllocs;
+
+	pxr::UsdAttribute Points = UsdMesh.GetPointsAttr();
+	if (!Points)
+	{
+		return EMeshTopologyVariance::Constant;
+	}
+
+	pxr::UsdAttribute FaceCountsAttribute = UsdMesh.GetFaceVertexCountsAttr();
+	if (!FaceCountsAttribute)
+	{
+		return EMeshTopologyVariance::Constant;
+	}
+
+	pxr::UsdAttribute FaceVertexIndicesAttribute = UsdMesh.GetFaceVertexIndicesAttr();
+	if (!FaceVertexIndicesAttribute)
+	{
+		return EMeshTopologyVariance::Constant;
+	}
+
+	if (!FaceVertexIndicesAttribute.ValueMightBeTimeVarying() && !FaceCountsAttribute.ValueMightBeTimeVarying())
+	{
+		if (!Points.ValueMightBeTimeVarying())
+		{
+			return EMeshTopologyVariance::Constant;
+		}
+		else
+		{
+			return EMeshTopologyVariance::Homogenous;
+		}
+	}
+	else
+	{
+		return EMeshTopologyVariance::Heterogenous;
+	}
+
+	return EMeshTopologyVariance::Constant;
+}
+
 #undef LOCTEXT_NAMESPACE
 
 #endif // #if USE_USD_SDK
diff --git a/Engine/Plugins/Importers/USDImporter/Source/USDUtilities/Private/USDPrimConversion.cpp b/Engine/Plugins/Importers/USDImporter/Source/USDUtilities/Private/USDPrimConversion.cpp
--- a/Engine/Plugins/Importers/USDImporter/Source/USDUtilities/Private/USDPrimConversion.cpp
+++ b/Engine/Plugins/Importers/USDImporter/Source/USDUtilities/Private/USDPrimConversion.cpp
@@ -385,6 +385,44 @@ bool UsdToUnreal::ConvertXformable( const pxr::UsdStageRefPtr& Stage, const pxr:
 	return true;
 }
 
+void UsdToUnreal::PropagateTransform(const pxr::UsdStageRefPtr& Stage, const pxr::UsdPrim& Root, const pxr::UsdPrim& Leaf, double EvalTime, FTransform& OutTransform)
+{
+	FScopedUsdAllocs UsdAllocs;
+
+	bool bResetXformStack = false;
+	FTransform CurrentTransform = FTransform::Identity;
+	if (ConvertXformable(Stage, pxr::UsdTyped{Leaf}, CurrentTransform, EvalTime, &bResetXformStack))
+	{
+		if (!bResetXformStack)
+		{
+			OutTransform *= CurrentTransform;
+			
+			if (Leaf != Root)
+			{
+				if (!Leaf.IsPseudoRoot())
+				{
+					PropagateTransform(Stage, Root, Leaf.GetParent(), EvalTime, OutTransform);
+				}
+				else
+				{
+					// Leaf was not even in Root's subtree
+					OutTransform = FTransform::Identity;
+				}
+			}
+		}
+		else
+		{
+			// The Xform stack was reset so that effectively stops the propagation
+			OutTransform = CurrentTransform;
+		}
+	}
+	else
+	{
+		// Leaf is not completely connected to Root by Xformables
+		OutTransform = FTransform::Identity;
+	}
+}
+
 bool UsdToUnreal::ConvertXformable( const pxr::UsdStageRefPtr& Stage, const pxr::UsdTyped& Schema, USceneComponent& SceneComponent, double EvalTime, bool bUsePrimTransform )
 {
 	pxr::UsdGeomXformable Xformable( Schema );
diff --git a/Engine/Plugins/Importers/USDImporter/Source/USDUtilities/Public/USDConversionUtils.h b/Engine/Plugins/Importers/USDImporter/Source/USDUtilities/Public/USDConversionUtils.h
--- a/Engine/Plugins/Importers/USDImporter/Source/USDUtilities/Public/USDConversionUtils.h
+++ b/Engine/Plugins/Importers/USDImporter/Source/USDUtilities/Public/USDConversionUtils.h
@@ -274,6 +274,9 @@ namespace UsdUtils
 	/** Adds a reference on Prim to the layer at AbsoluteFilePath */
 	USDUTILITIES_API void AddReference( UE::FUsdPrim& Prim, const TCHAR* AbsoluteFilePath, const UE::FSdfPath& TargetPrimPath = {}, double TimeCodeOffset = 0.0, double TimeCodeScale = 1.0 );
 
+	/** Gets the strongest direct reference on Prim with the given FileExtension. Returns true if there was one. */
+	USDUTILITIES_API bool GetReferenceFilePath(const UE::FUsdPrim& Prim, const FString& FileExtension, FString& OutReferenceFilePath);
+
 	/** Adds a payload on Prim pointing at the default prim of the layer at AbsoluteFilePath */
 	USDUTILITIES_API void AddPayload( UE::FUsdPrim& Prim, const TCHAR* AbsoluteFilePath, const UE::FSdfPath& TargetPrimPath = {}, double TimeCodeOffset = 0.0, double TimeCodeScale = 1.0 );
 
diff --git a/Engine/Plugins/Importers/USDImporter/Source/USDUtilities/Public/USDGeomMeshConversion.h b/Engine/Plugins/Importers/USDImporter/Source/USDUtilities/Public/USDGeomMeshConversion.h
--- a/Engine/Plugins/Importers/USDImporter/Source/USDUtilities/Public/USDGeomMeshConversion.h
+++ b/Engine/Plugins/Importers/USDImporter/Source/USDUtilities/Public/USDGeomMeshConversion.h
@@ -40,6 +40,11 @@ namespace UE
 	class FUsdPrim;
 }
 
+namespace MeshAttribute::VertexInstance
+{
+	extern USDUTILITIES_API const FName Velocity;
+}
+
 namespace UsdToUnreal
 {
 	/** Common options to mesh conversion functions */
@@ -382,6 +387,21 @@ namespace UsdUtils
 	 * Returns whether the transforms were successfully retrieved or not.
 	 */
 	USDUTILITIES_API bool GetPointInstancerTransforms( const FUsdStageInfo& StageInfo, const pxr::UsdGeomPointInstancer& PointInstancer, const int32 ProtoIndex, pxr::UsdTimeCode EvalTime, TArray<FTransform>& OutInstanceTransforms );
+
+	/** Returns true if Prim is a GeomMesh with animated attributes */
+	USDUTILITIES_API bool IsAnimatedMesh(const pxr::UsdPrim& Prim);
+
+	/** Mesh topology variance as defined by Alembic's MeshTopologyVariance */
+	enum class USDUTILITIES_API EMeshTopologyVariance : uint8
+	{
+		Constant,		// unchanging topology
+		Homogenous,		// connectivity is unchanging, but the positions may change
+		Heterogenous	// connectivity and positions may change
+	};
+
+	/** Returns the type of topology variance UsdMesh has */
+	USDUTILITIES_API EMeshTopologyVariance GetMeshTopologyVariance(const pxr::UsdGeomMesh& UsdMesh);
+
 }
 
 #endif // #if USE_USD_SDK
diff --git a/Engine/Plugins/Importers/USDImporter/Source/USDUtilities/Public/USDPrimConversion.h b/Engine/Plugins/Importers/USDImporter/Source/USDUtilities/Public/USDPrimConversion.h
--- a/Engine/Plugins/Importers/USDImporter/Source/USDUtilities/Public/USDPrimConversion.h
+++ b/Engine/Plugins/Importers/USDImporter/Source/USDUtilities/Public/USDPrimConversion.h
@@ -73,6 +73,16 @@ namespace UsdToUnreal
 	 */
 	USDUTILITIES_API bool ConvertXformable( const pxr::UsdStageRefPtr& Stage, const pxr::UsdTyped& Schema, FTransform& OutTransform, double EvalTime, bool* bOutResetTransformStack = nullptr );
 
+	/**
+	 * Propagates the transform from Root to Leaf prims inclusively into a UE-space FTransform if they are connected by Xformables
+	 * @param Stage - Stage that contains the prims to process
+	 * @param Root - Prim to start the propagation from
+	 * @param Leaf - Prim where to end the propagation
+	 * @param EvalTime - Time at which to sample the prim's attributes
+	 * @param OutTransform - The composed transform; identity if Leaf is not completely connected to Root by Xformables
+	 */
+	USDUTILITIES_API void PropagateTransform(const pxr::UsdStageRefPtr& Stage, const pxr::UsdPrim& Root, const pxr::UsdPrim& Leaf, double EvalTime, FTransform& OutTransform);
+
 	USDUTILITIES_API bool ConvertGeomCamera( const UE::FUsdPrim& Prim, UCineCameraComponent& CameraComponent, double UsdTimeCode = UsdUtils::GetDefaultTimeCode() );
 
 	/**
