USD: Use AssetUserData instead of AssetImportData for better runtime support;

Implement IInterface_AssetUserData for UGroomCache and UGeometryCache.

diff --git a/Engine/Plugins/Importers/USDImporter/Source/USDClasses/Public/USDAssetImportData.h b/Engine/Plugins/Importers/USDImporter/Source/USDClasses/Public/USDAssetImportData.h
--- a/Engine/Plugins/Importers/USDImporter/Source/USDClasses/Public/USDAssetImportData.h
+++ b/Engine/Plugins/Importers/USDImporter/Source/USDClasses/Public/USDAssetImportData.h
@@ -5,10 +5,12 @@
 #include "CoreMinimal.h"
 #include "EditorFramework/AssetImportData.h"
 
+#include "USDAssetUserData.h"
+
 #include "USDAssetImportData.generated.h"
 
 UCLASS()
-class USDCLASSES_API UUsdAssetImportData : public UAssetImportData
+class UE_DEPRECATED(5.3, "Prefer USD AssetUserData instead of AssetImportData") USDCLASSES_API UUsdAssetImportData : public UAssetImportData
 {
 	GENERATED_BODY()
 
@@ -22,8 +24,9 @@ public:
 	TObjectPtr<class UObject> ImportOptions;
 };
 
+PRAGMA_DISABLE_DEPRECATION_WARNINGS
 UCLASS()
-class USDCLASSES_API UUsdAnimSequenceAssetImportData : public UUsdAssetImportData
+class UE_DEPRECATED(5.3, "Prefer USD AssetUserData instead of AssetImportData") USDCLASSES_API UUsdAnimSequenceAssetImportData : public UUsdAssetImportData
 {
 	GENERATED_BODY()
 
@@ -39,47 +42,9 @@ public:
 	float LayerStartOffsetSeconds = 0.0f;
 };
 
-/** Simple wrapper because we're not allowed to have TMap properties with TArray<FString> as values */
-USTRUCT()
-struct USDCLASSES_API FUsdPrimPathList
-{
-	GENERATED_BODY()
-
-	UPROPERTY()
-	TArray<FString> PrimPaths;
-};
-
-/** AssetImportData assigned to Unreal materials that are generated when parsing USD Material prims */
-UCLASS()
-class USDCLASSES_API UUsdMaterialAssetImportData : public UUsdAssetImportData
-{
-	GENERATED_BODY()
-
-public:
-	/**
-	 * In the context of our reference materials that just read a single texture for each material parameter, this
-	 * describes the primvar that the USD material is sampling for each texture.
-	 * e.g. {'BaseColor': 'st0', 'Metallic': 'st1'}
-	 */
-	UPROPERTY()
-	TMap<FString, FString> ParameterToPrimvar;
-
-	/**
-	 * In the context of our reference materials that just read a single texture for each material parameter, this
-	 * describes the Unreal UV index that will be used for sampling by each USD primvar. The idea is to use this in
-	 * combination with ParameterToPrimvar, in order to describe which UV index the material has currently assigned to
-	 * each parameter. When assigning the material to meshes later, we'll compare this member with the
-	 * UUsdMeshAssetImportData's own PrimvarToUVIndex to see if they are compatible or not, and if not spawn a new
-	 * material instance that is.
-	 * e.g. {'firstPrimvar': 0, 'st': 1, 'st1': 2}
-	 */
-	UPROPERTY()
-	TMap<FString, int32> PrimvarToUVIndex;
-};
-
-/** We assign these to UStaticMeshes, UGeometryCaches or USkeletalMeshes generated from USD */
+/** We assign these to UStaticMeshes or USkeletalMeshes generated from USD */
 UCLASS()
-class USDCLASSES_API UUsdMeshAssetImportData : public UUsdAssetImportData
+class UE_DEPRECATED(5.3, "Prefer USD AssetUserData instead of AssetImportData") USDCLASSES_API UUsdMeshAssetImportData : public UUsdAssetImportData
 {
 	GENERATED_BODY()
 
@@ -91,8 +56,5 @@ public:
 	 */
 	UPROPERTY()
 	TMap< int32, FUsdPrimPathList > MaterialSlotToPrimPaths;
-
-	/** Describes which primvars should be assigned to each UV index. */
-	UPROPERTY()
-	TMap<FString, int32> PrimvarToUVIndex;
-};
\ No newline at end of file
+};
+PRAGMA_ENABLE_DEPRECATION_WARNINGS
\ No newline at end of file
diff --git a/Engine/Plugins/Importers/USDImporter/Source/USDClasses/Public/USDAssetUserData.h b/Engine/Plugins/Importers/USDImporter/Source/USDClasses/Public/USDAssetUserData.h
new file mode 100644
--- /dev/null
+++ b/Engine/Plugins/Importers/USDImporter/Source/USDClasses/Public/USDAssetUserData.h
@@ -0,0 +1,100 @@
+// Copyright Epic Games, Inc. All Rights Reserved.
+
+#pragma once
+
+#include "Engine/AssetUserData.h"
+
+#include "USDAssetUserData.generated.h"
+
+UCLASS()
+class USDCLASSES_API UUsdAssetUserData : public UAssetUserData
+{
+	GENERATED_BODY()
+
+public:
+	UPROPERTY()
+	FString FilePath;
+
+	UPROPERTY()
+	FString PrimPath;
+
+	// Likely a UUSDStageImportOptions, but we don't declare it here
+	// to prevent an unnecessary module dependency on USDStageImporter
+	UPROPERTY()
+	TObjectPtr<class UObject> ImportOptions;
+};
+
+UCLASS()
+class USDCLASSES_API UUsdAnimSequenceAssetUserData : public UUsdAssetUserData
+{
+	GENERATED_BODY()
+
+public:
+	// Offset into the stage in seconds to when this AnimSequence should start playing to match the
+	// skeletal animation in that stage. This is required because UAnimSequences just range from the
+	// first skeletal timeSample to the last, and the first sample is not necessarily the startTimeCode
+	// for the stage. Note that this is wrt. the startTimeCode of the *layer*, and not the composed stage.
+	// This is only used for animating USkeletalMeshComponents via the AUsdStageActor, and it is in
+	// seconds since it will need to drive the UAnimSequence with a time value in the [0, LenghtSeconds] range.
+	// This should be applied *after* any offset/scale conversions on the time coordinate.
+	UPROPERTY()
+	float LayerStartOffsetSeconds = 0.0f;
+};
+
+/** Simple wrapper because we're not allowed to have TMap properties with TArray<FString> as values */
+USTRUCT()
+struct FUsdPrimPathList
+{
+	GENERATED_BODY()
+
+	UPROPERTY()
+	TArray<FString> PrimPaths;
+};
+
+/** AssetImportData assigned to Unreal materials that are generated when parsing USD Material prims */
+UCLASS()
+class USDCLASSES_API UUsdMaterialAssetUserData : public UUsdAssetUserData
+{
+	GENERATED_BODY()
+
+public:
+	/**
+	 * In the context of our reference materials that just read a single texture for each material parameter, this
+	 * describes the primvar that the USD material is sampling for each texture.
+	 * e.g. {'BaseColor': 'st0', 'Metallic': 'st1'}
+	 */
+	UPROPERTY()
+	TMap<FString, FString> ParameterToPrimvar;
+
+	/**
+	 * In the context of our reference materials that just read a single texture for each material parameter, this
+	 * describes the Unreal UV index that will be used for sampling by each USD primvar. The idea is to use this in
+	 * combination with ParameterToPrimvar, in order to describe which UV index the material has currently assigned to
+	 * each parameter. When assigning the material to meshes later, we'll compare this member with the
+	 * UUsdMeshAssetImportData's own PrimvarToUVIndex to see if they are compatible or not, and if not spawn a new
+	 * material instance that is.
+	 * e.g. {'firstPrimvar': 0, 'st': 1, 'st1': 2}
+	 */
+	UPROPERTY()
+	TMap<FString, int32> PrimvarToUVIndex;
+};
+
+/** We assign these to UStaticMeshes or USkeletalMeshes generated from USD */
+UCLASS()
+class USDCLASSES_API UUsdMeshAssetUserData : public UUsdAssetUserData
+{
+	GENERATED_BODY()
+
+public:
+	/**
+	 * Maps from a material slot index of an UStaticMesh or USkeletalMesh to a list of source prims that contain this
+	 * assignment. It can contain multiple prims in case we combine material slots and/or collapse prims
+	 * (e.g. {0: ['/Root/mesh', '/Root/othermesh/geomsubset0', '/Root/othermesh/geomsubset1'] }).
+	 */
+	UPROPERTY()
+	TMap< int32, FUsdPrimPathList > MaterialSlotToPrimPaths;
+
+	/** Describes which primvars should be assigned to each UV index. */
+	UPROPERTY()
+	TMap<FString, int32> PrimvarToUVIndex;
+};
\ No newline at end of file
diff --git a/Engine/Plugins/Importers/USDImporter/Source/USDSchemas/Private/Custom/MaterialXUSDShadeMaterialTranslator.cpp b/Engine/Plugins/Importers/USDImporter/Source/USDSchemas/Private/Custom/MaterialXUSDShadeMaterialTranslator.cpp
--- a/Engine/Plugins/Importers/USDImporter/Source/USDSchemas/Private/Custom/MaterialXUSDShadeMaterialTranslator.cpp
+++ b/Engine/Plugins/Importers/USDImporter/Source/USDSchemas/Private/Custom/MaterialXUSDShadeMaterialTranslator.cpp
@@ -3,7 +3,7 @@
 #include "Custom/MaterialXUSDShadeMaterialTranslator.h"
 #if USE_USD_SDK
 #include "USDAssetCache.h"
-#include "USDAssetImportData.h"
+#include "USDAssetUserData.h"
 #include "USDClassesModule.h"
 #include "USDConversionUtils.h"
 #include "USDErrorUtils.h"
@@ -452,9 +452,9 @@ void FMaterialXUsdShadeMaterialTranslator::CreateAssets()
 						UTexture* ExistingTexture = Cast<UTexture>(Context->AssetCache->GetCachedAsset(TextureHash));
 						if (!ExistingTexture)
 						{
-							UUsdAssetImportData* ImportData = NewObject<UUsdAssetImportData>(Texture, TEXT("USDAssetImportData"));
-							ImportData->PrimPath = ReferencerPrimPath;
-							Texture->AssetImportData = ImportData;
+							UUsdAssetUserData* UserData = NewObject<UUsdAssetUserData>(Texture, TEXT("USDAssetUserData"));
+							UserData->PrimPath = ReferencerPrimPath;
+							Texture->AddAssetUserData(UserData);
 
 							Context->AssetCache->CacheAsset(TextureHash, Texture);
 						}
diff --git a/Engine/Plugins/Importers/USDImporter/Source/USDSchemas/Private/MeshTranslationImpl.cpp b/Engine/Plugins/Importers/USDImporter/Source/USDSchemas/Private/MeshTranslationImpl.cpp
--- a/Engine/Plugins/Importers/USDImporter/Source/USDSchemas/Private/MeshTranslationImpl.cpp
+++ b/Engine/Plugins/Importers/USDImporter/Source/USDSchemas/Private/MeshTranslationImpl.cpp
@@ -3,8 +3,8 @@
 #include "MeshTranslationImpl.h"
 
 #include "USDAssetCache2.h"
-#include "USDAssetImportData.h"
 #include "USDConversionUtils.h"
+#include "USDAssetUserData.h"
 #include "USDGeomMeshConversion.h"
 #include "USDInfoCache.h"
 #include "USDLog.h"
@@ -46,15 +46,15 @@ namespace UE::MeshTranslationImplInternal::Private
 			return nullptr;
 		}
 
-		FString PrimPath;
-#if WITH_EDITOR
-		UUsdMaterialAssetImportData* OneSidedImportData = Cast<UUsdMaterialAssetImportData>(OneSidedMat->AssetImportData);
-		if (!ensureMsgf(OneSidedImportData, TEXT("Expected one-sided material '%s' to have an UUsdMaterialAssetImportData at this point!"), *OneSidedMat->GetPathName()))
+		UUsdMaterialAssetUserData* OneSidedUserData = OneSidedMat->GetAssetUserData<UUsdMaterialAssetUserData>();
+		if (!ensureMsgf(
+				OneSidedUserData,
+				TEXT("Expected one-sided material '%s' to have an UUsdMaterialAssetUserData at this point!"),
+				*OneSidedMat->GetPathName()
+			))
 		{
 			return nullptr;
 		}
-		PrimPath = OneSidedImportData->PrimPath;
-#endif // WITH_EDITOR
 
 		UMaterialInterface* TwoSidedMat = nullptr;
 
@@ -83,16 +83,8 @@ namespace UE::MeshTranslationImplInternal::Private
 		// Just create an instance of the TwoSided version of the same reference material and copy parameter values.
 		if (GIsEditor && MIC && ReferenceMaterialTwoSided)
 		{
-			UMaterialInstanceConstant* TwoSidedMIC = NewObject<UMaterialInstanceConstant>(
-				GetTransientPackage(),
-				NewInstanceName,
-				OneSidedMat->GetFlags()
-			);
-			if (TwoSidedMIC)
-			{
-				TwoSidedMIC->AssetImportData =
-					DuplicateObject(OneSidedImportData, TwoSidedMIC, TEXT("USDAssetImportData"));
-			}
+			UMaterialInstanceConstant* TwoSidedMIC =
+				NewObject<UMaterialInstanceConstant>(GetTransientPackage(), NewInstanceName, OneSidedMat->GetFlags());
 
 			TwoSidedMIC->SetParentEditorOnly(ReferenceMaterialTwoSided);
 			TwoSidedMIC->CopyMaterialUniformParametersEditorOnly(OneSidedMat);
@@ -109,11 +101,6 @@ namespace UE::MeshTranslationImplInternal::Private
 				NewInstanceName,
 				OneSidedMat->GetFlags()
 			);
-			if (TwoSidedMIC)
-			{
-				TwoSidedMIC->AssetImportData =
-					DuplicateObject(OneSidedImportData, TwoSidedMIC, TEXT("USDAssetImportData"));
-			}
 
 			TwoSidedMIC->SetParentEditorOnly(OneSidedMat);
 			TwoSidedMIC->BasePropertyOverrides.bOverride_TwoSided = true;
@@ -144,11 +131,15 @@ namespace UE::MeshTranslationImplInternal::Private
 
 			TwoSidedMID->CopyParameterOverrides(MID);
 
-			// TODO: Copy the PrimvarToUVIndex from the asset user data
-
 			TwoSidedMat = TwoSidedMID;
 		}
 
+		if(TwoSidedMat)
+		{
+			UUsdMaterialAssetUserData* UserData = DuplicateObject(OneSidedUserData, TwoSidedMat, TEXT("USDAssetUserData"));
+			TwoSidedMat->AddAssetUserData(UserData);
+		}
+
 		return TwoSidedMat;
 	}
 
@@ -165,18 +156,16 @@ namespace UE::MeshTranslationImplInternal::Private
 		FUsdInfoCache* InfoCache
 	)
 	{
-		// TODO: Implement this for runtime when asset user data support is in
-	#if WITH_EDITOR
-		UUsdMaterialAssetImportData* MaterialAssetImportData = Cast<UUsdMaterialAssetImportData>(Material.AssetImportData);
-		if (!ensureMsgf(MaterialAssetImportData, TEXT("Expected material '%s' to have an UUsdMaterialAssetImportData at this point!"), *Material.GetPathName()))
+		UUsdMaterialAssetUserData* MaterialAssetUserData = Material.GetAssetUserData<UUsdMaterialAssetUserData>();
+		if (!ensureMsgf(MaterialAssetUserData, TEXT("Expected material '%s' to have an UUsdMaterialAssetUserData at this point!"), *Material.GetPathName()))
 		{
 			return nullptr;
 		}
 
 		UE_LOG(LogUsd, Verbose, TEXT("Getting compatible material based on Material '%s' (Parameter to primvar: %s, primvar to UV index: %s) matching mesh primvar to UV index mapping '%s'"),
 			*Material.GetPathName(),
-			*UsdUtils::StringifyMap(MaterialAssetImportData->ParameterToPrimvar),
-			*UsdUtils::StringifyMap(MaterialAssetImportData->PrimvarToUVIndex),
+			*UsdUtils::StringifyMap(MaterialAssetUserData->ParameterToPrimvar),
+			*UsdUtils::StringifyMap(MaterialAssetUserData->PrimvarToUVIndex),
 			*UsdUtils::StringifyMap(MeshPrimvarToUVIndex)
 		);
 
@@ -193,7 +182,7 @@ namespace UE::MeshTranslationImplInternal::Private
 		// Check if mesh and material are compatible. Note that it's perfectly valid for the material to try reading
 		// an UVIndex the mesh doesn't provide at all, or trying to read a primvar that doesn't exist on the mesh.
 		bool bCompatible = true;
-		for (const TPair<FString, int32>& Pair : MaterialAssetImportData->PrimvarToUVIndex)
+		for (const TPair<FString, int32>& Pair : MaterialAssetUserData->PrimvarToUVIndex)
 		{
 			const FString& MaterialPrimvar = Pair.Key;
 			int32 MaterialUVIndex = Pair.Value;
@@ -243,8 +232,8 @@ namespace UE::MeshTranslationImplInternal::Private
 		// First, let's create the target primvar UVIndex assignment that is compatible with this mesh.
 		// We use an array of TPairs here so that we can sort these into a deterministic order for hashing later.
 		TArray<TPair<FString, int32>> CompatiblePrimvarAndUVIndexPairs;
-		CompatiblePrimvarAndUVIndexPairs.Reserve(MaterialAssetImportData->PrimvarToUVIndex.Num());
-		for (const TPair<FString, int32>& Pair : MaterialAssetImportData->PrimvarToUVIndex)
+		CompatiblePrimvarAndUVIndexPairs.Reserve(MaterialAssetUserData->PrimvarToUVIndex.Num());
+		for (const TPair<FString, int32>& Pair : MaterialAssetUserData->PrimvarToUVIndex)
 		{
 			const FString& MaterialPrimvar = Pair.Key;
 
@@ -327,7 +316,7 @@ namespace UE::MeshTranslationImplInternal::Private
 
 			UE_LOG(LogUsd, Verbose, TEXT("Generating compatible version of Material '%s' (Parameter to primvar: %s, primvar to UV index: %s) with hash '%s'"),
 				*Material.GetPathName(),
-				*UsdUtils::StringifyMap(MaterialAssetImportData->ParameterToPrimvar),
+				*UsdUtils::StringifyMap(MaterialAssetUserData->ParameterToPrimvar),
 				*UsdUtils::StringifyMap(CompatiblePrimvarToUVIndex),
 				*MaterialHash
 			);
@@ -340,27 +329,45 @@ namespace UE::MeshTranslationImplInternal::Private
 					Material.GetFlags()
 				);
 
-				UUsdMaterialAssetImportData* CompatibleImportData = DuplicateObject(
-					MaterialAssetImportData,
-					CompatibleMIC,
-					TEXT("USDAssetImportData")
+				CompatibleMIC->SetParentEditorOnly(&Material);
+
+				CompatibleMaterial = CompatibleMIC;
+			}
+			else
+			{
+				UMaterialInstanceDynamic* CompatibleMID = UMaterialInstanceDynamic::Create(
+					&Material,
+					GetTransientPackage(),
+					NewInstanceName
 				);
-				CompatibleImportData->PrimvarToUVIndex = CompatiblePrimvarToUVIndex;
 
+				CompatibleMaterial = CompatibleMID;
+			}
+
+			UUsdMaterialAssetUserData* CompatibleUserData = DuplicateObject(
+				MaterialAssetUserData,
+				CompatibleMaterial,
+				TEXT("USDAssetUserData")
+			);
+			CompatibleUserData->PrimvarToUVIndex = CompatiblePrimvarToUVIndex;
+			CompatibleMaterial->AddAssetUserData(CompatibleUserData);
+
+			if (UMaterialInstance* CompatibleInstance = Cast<UMaterialInstance>(CompatibleMaterial))
+			{
 				// Actually set the parameters with the target UV indices
-				for (const TPair<FString, FString>& ParameterPair : CompatibleImportData->ParameterToPrimvar)
+				for (const TPair<FString, FString>& ParameterPair : CompatibleUserData->ParameterToPrimvar)
 				{
 					const FString& Parameter = ParameterPair.Key;
 					const FString& Primvar = ParameterPair.Value;
 
-					if (int32* UVIndex = CompatibleImportData->PrimvarToUVIndex.Find(Primvar))
+					if (int32* UVIndex = CompatibleUserData->PrimvarToUVIndex.Find(Primvar))
 					{
 						// Force-disable using the texture at all if the mesh doesn't provide the primvar that should be
 						// used to sample it with
 						if (*UVIndex == UNUSED_UV_INDEX)
 						{
 							UsdUtils::SetScalarParameterValue(
-								*CompatibleMIC,
+								*CompatibleInstance,
 								*FString::Printf(TEXT("Use%sTexture"), *Parameter),
 								0.0f
 							);
@@ -368,7 +375,7 @@ namespace UE::MeshTranslationImplInternal::Private
 						else
 						{
 							UsdUtils::SetScalarParameterValue(
-								*CompatibleMIC,
+								*CompatibleInstance,
 								*FString::Printf(TEXT("%sUVIndex"), *Parameter),
 								static_cast<float>(*UVIndex)
 							);
@@ -376,25 +383,7 @@ namespace UE::MeshTranslationImplInternal::Private
 					}
 				}
 
-				CompatibleMIC->AssetImportData = CompatibleImportData;
-				CompatibleMIC->SetParentEditorOnly(&Material);
-				CompatibleMIC->PostEditChange();
-
-				CompatibleMaterial = CompatibleMIC;
-			}
-			else
-			{
-				UMaterialInstanceDynamic* CompatibleMID = UMaterialInstanceDynamic::Create(
-					&Material,
-					GetTransientPackage(),
-					NewInstanceName
-				);
-
-				// TODO: Also copy our custom asset user data here when we add support for that
-
-				// TODO: Actually set the parameters with the target UV indices
-
-				CompatibleMaterial = CompatibleMID;
+				CompatibleInstance->PostEditChange();
 			}
 		}
 
@@ -415,9 +404,6 @@ namespace UE::MeshTranslationImplInternal::Private
 		}
 
 		return CompatibleMaterial;
-	#else
-		return &Material;
-	#endif // WITH_EDITOR
 	}
 };
 
@@ -473,17 +459,19 @@ TMap<const UsdUtils::FUsdPrimMaterialSlot*, UMaterialInterface*> MeshTranslation
 						if (GIsEditor)  // Editor, PIE => true; Standlone, packaged => false
 						{
 							MaterialInstance = UsdUtils::CreateDisplayColorMaterialInstanceConstant(DisplayColorDesc.GetValue());
-#if WITH_EDITOR
-							// Leave PrimPath as empty as it likely will be reused by many prims
-							UUsdMaterialAssetImportData* ImportData = NewObject< UUsdMaterialAssetImportData >(MaterialInstance, TEXT("USDAssetImportData"));
-							MaterialInstance->AssetImportData = ImportData;
-#endif // WITH_EDITOR
 						}
 						else
 						{
 							MaterialInstance = UsdUtils::CreateDisplayColorMaterialInstanceDynamic(DisplayColorDesc.GetValue());
 						}
 
+						if(MaterialInstance)
+						{
+							// Leave PrimPath as empty as it likely will be reused by many prims
+							UUsdAssetUserData* UserData = NewObject<UUsdAssetUserData>(MaterialInstance, TEXT("USDAssetUserData"));
+							MaterialInstance->AddAssetUserData(UserData);
+						}
+
 						// We can only cache transient assets
 						MaterialInstance->SetFlags(RF_Transient);
 
@@ -772,26 +760,26 @@ void MeshTranslationImpl::SetMaterialOverrides(
 	TMap<const UsdUtils::FUsdPrimMaterialSlot*, UMaterialInterface*> ResolvedMaterials;
 
 #if WITH_EDITOR
-	UUsdMeshAssetImportData* ImportData = nullptr;
+	UUsdMeshAssetUserData* UserData = nullptr;
 	if (UStaticMeshComponent* StaticMeshComponent = Cast<UStaticMeshComponent>(&MeshComponent))
 	{
 		if (UStaticMesh* Mesh = StaticMeshComponent->GetStaticMesh())
 		{
-			ImportData = Cast<UUsdMeshAssetImportData>(Mesh->GetAssetImportData());
+			UserData = Mesh->GetAssetUserData<UUsdMeshAssetUserData>();
 		}
 	}
 	else if (USkeletalMeshComponent* SkeletalMeshComponent = Cast<USkeletalMeshComponent>(&MeshComponent))
 	{
 		if (USkeletalMesh* Mesh = SkeletalMeshComponent->GetSkeletalMeshAsset())
 		{
-			ImportData = Cast<UUsdMeshAssetImportData>(Mesh->GetAssetImportData());
+			UserData = Mesh->GetAssetUserData<UUsdMeshAssetUserData>();
 		}
 	}
 	else if (UGeometryCacheComponent* GeometryCacheComponent = Cast<UGeometryCacheComponent>(&MeshComponent))
 	{
 		if (UGeometryCache* Mesh = GeometryCacheComponent->GetGeometryCache())
 		{
-			ImportData = Cast<UUsdMeshAssetImportData>(Mesh->AssetImportData);
+			UserData = Mesh->GetAssetUserData<UUsdMeshAssetUserData>();
 		}
 	}
 	else
@@ -802,15 +790,15 @@ void MeshTranslationImpl::SetMaterialOverrides(
 		);
 	}
 
-	ensureMsgf(ImportData, TEXT("Mesh assigned to component '%s' generated for prim '%s' should have an UUsdMeshAssetImportData at this point!"),
+	ensureMsgf(UserData, TEXT("Mesh assigned to component '%s' generated for prim '%s' should have an UUsdMeshAssetUserData at this point!"),
 		*MeshComponent.GetPathName(),
 		*UsdToUnreal::ConvertPath(Prim.GetPrimPath())
 	);
 
-	if(ImportData && LODIndexToAssignments.Num() > 0)
+	if(UserData && LODIndexToAssignments.Num() > 0)
 	{
 		// Stash our PrimvarToUVIndex in here, as that's where ResolveMaterialAssignmentInfo will look for it
-		LODIndexToAssignments[0].PrimvarToUVIndex = ImportData->PrimvarToUVIndex;
+		LODIndexToAssignments[0].PrimvarToUVIndex = UserData->PrimvarToUVIndex;
 
 		ResolvedMaterials = MeshTranslationImpl::ResolveMaterialAssignmentInfo(
 			ValidPrim,
diff --git a/Engine/Plugins/Importers/USDImporter/Source/USDSchemas/Private/USDGeomMeshTranslator.cpp b/Engine/Plugins/Importers/USDImporter/Source/USDSchemas/Private/USDGeomMeshTranslator.cpp
--- a/Engine/Plugins/Importers/USDImporter/Source/USDSchemas/Private/USDGeomMeshTranslator.cpp
+++ b/Engine/Plugins/Importers/USDImporter/Source/USDSchemas/Private/USDGeomMeshTranslator.cpp
@@ -1,16 +1,12 @@
 // Copyright Epic Games, Inc. All Rights Reserved.
 
 #include "USDGeomMeshTranslator.h"
-#include "Engine/Level.h"
-#include "Engine/World.h"
-#include "UObject/Package.h"
 
 #if USE_USD_SDK
 
 #include "MeshTranslationImpl.h"
 #include "UnrealUSDWrapper.h"
-#include "USDAssetCache.h"
-#include "USDAssetImportData.h"
+#include "USDAssetUserData.h"
 #include "USDClassesModule.h"
 #include "USDConversionUtils.h"
 #include "USDGeomMeshConversion.h"
@@ -24,24 +20,24 @@
 #include "UsdWrappers/UsdPrim.h"
 #include "UsdWrappers/UsdStage.h"
 
-#include "Async/Async.h"
 #include "Components/StaticMeshComponent.h"
 #include "Engine/CollisionProfile.h"
+#include "Engine/Level.h"
 #include "Engine/StaticMesh.h"
+#include "Engine/World.h"
 #include "HAL/IConsoleManager.h"
 #include "Interfaces/ITargetPlatform.h"
 #include "Interfaces/ITargetPlatformManagerModule.h"
-#include "Materials/Material.h"
 #include "Materials/MaterialInstanceConstant.h"
 #include "Materials/MaterialInstanceDynamic.h"
 #include "Materials/MaterialInterface.h"
 #include "Misc/SecureHash.h"
 #include "Modules/ModuleManager.h"
 #include "PhysicsEngine/BodySetup.h"
-#include "RHI.h"
 #include "StaticMeshAttributes.h"
 #include "StaticMeshOperations.h"
 #include "StaticMeshResources.h"
+#include "UObject/Package.h"
 #include "UObject/SoftObjectPath.h"
 
 #if WITH_EDITOR
@@ -50,13 +46,11 @@
 #endif // WITH_EDITOR
 
 #include "USDIncludesStart.h"
-#include "pxr/usd/usd/editContext.h"
 #include "pxr/usd/usd/prim.h"
 #include "pxr/usd/usd/stage.h"
 #include "pxr/usd/usd/typed.h"
 #include "pxr/usd/usdGeom/mesh.h"
-#include "pxr/usd/usdGeom/xformable.h"
-#include "pxr/usd/usdShade/material.h"
+#include "pxr/usd/usdGeom/subset.h"
 #include "USDIncludesEnd.h"
 
 static float GMeshNormalRepairThreshold = 0.05f;
@@ -211,14 +205,11 @@ namespace UsdGeomMeshTranslatorImpl
 			ExistingAssignments.Add(StaticMaterial.MaterialInterface);
 		}
 
-#if WITH_EDITOR
-		// TODO: Swap with asset user data for runtime support
-		UUsdMeshAssetImportData* ImportData = Cast< UUsdMeshAssetImportData>(StaticMesh.AssetImportData.Get());
-		ensureMsgf(ImportData, TEXT("Static Mesh '%s' generated for prim '%s' should have an UUsdMeshAssetImportData at this point!"),
+		UUsdMeshAssetUserData* UserData = StaticMesh.GetAssetUserData<UUsdMeshAssetUserData>();
+		ensureMsgf(UserData, TEXT("Static Mesh '%s' generated for prim '%s' should have an UUsdMeshAssetUserData at this point!"),
 			*StaticMesh.GetPathName(),
 			*UsdToUnreal::ConvertPath(UsdPrim.GetPrimPath())
 		);
-#endif // WITH_EDITOR
 
 		TMap<const UsdUtils::FUsdPrimMaterialSlot*, UMaterialInterface*> ResolvedMaterials = MeshTranslationImpl::ResolveMaterialAssignmentInfo(
 			UsdPrim,
@@ -261,10 +252,9 @@ namespace UsdGeomMeshTranslatorImpl
 					bMaterialAssignementsHaveChanged = true;
 				}
 
-#if WITH_EDITOR
-				if (ImportData)
+				if (UserData)
 				{
-					ImportData->MaterialSlotToPrimPaths.FindOrAdd(StaticMeshSlotIndex).PrimPaths = Slot.PrimPaths.Array();
+					UserData->MaterialSlotToPrimPaths.FindOrAdd(StaticMeshSlotIndex).PrimPaths = Slot.PrimPaths.Array();
 				}
 
 				// Setup the section map so that our LOD material index is properly mapped to the static mesh material index
@@ -290,7 +280,6 @@ namespace UsdGeomMeshTranslatorImpl
 
 					bMaterialAssignementsHaveChanged = true;
 				}
-#endif // WITH_EDITOR
 			}
 		}
 
@@ -822,16 +811,14 @@ void FBuildStaticMeshTaskChain::SetupTasks()
 
 #if WITH_EDITOR
 				StaticMesh->NaniteSettings.bEnabled = bShouldEnableNanite;
+#endif // WITH_EDITOR
 
 				if (bIsNew)
 				{
-					UUsdMeshAssetImportData* ImportData = NewObject<UUsdMeshAssetImportData>(StaticMesh, TEXT("UUSDAssetImportData"));
-					ImportData->PrimPath = MeshName;
-					ImportData->PrimvarToUVIndex = LODIndexToMaterialInfo[0].PrimvarToUVIndex;  // We use the same primvar mapping for all LODs
-
-					StaticMesh->AssetImportData = ImportData;
+					UUsdMeshAssetUserData* UserData = NewObject<UUsdMeshAssetUserData>(StaticMesh, TEXT("UUSDAssetUserData"));
+					UserData->PrimPath = MeshName;
+					StaticMesh->AddAssetUserData(UserData);
 				}
-#endif // WITH_EDITOR
 
 				// Only the original creator of the prim at creation time gets to set the material assignments
 				// directly on the mesh, all others prims ensure their materials via material overrides on the
diff --git a/Engine/Plugins/Importers/USDImporter/Source/USDSchemas/Private/USDGeomPointInstancerTranslator.cpp b/Engine/Plugins/Importers/USDImporter/Source/USDSchemas/Private/USDGeomPointInstancerTranslator.cpp
--- a/Engine/Plugins/Importers/USDImporter/Source/USDSchemas/Private/USDGeomPointInstancerTranslator.cpp
+++ b/Engine/Plugins/Importers/USDImporter/Source/USDSchemas/Private/USDGeomPointInstancerTranslator.cpp
@@ -4,7 +4,7 @@
 
 #include "MeshTranslationImpl.h"
 #include "USDAssetCache.h"
-#include "USDAssetImportData.h"
+#include "USDAssetUserData.h"
 #include "USDConversionUtils.h"
 #include "USDGeomMeshConversion.h"
 #include "USDGeomMeshTranslator.h"
@@ -405,17 +405,14 @@ void FUsdGeomPointInstancerTranslator::UpdateComponents(USceneComponent* PointIn
 	TSet<UStaticMesh*> PrototypeMeshSet = InfoCache.GetAssetsForPrim<UStaticMesh>(PrimPath);
 	std::unordered_map<pxr::SdfPath, UStaticMesh*, pxr::SdfPath::Hash> PrototypeMeshes;
 	PrototypeMeshes.reserve(PrototypeMeshSet.Num());
-	// TODO: Replace AssetImportData with AssetUserData, which is available at runtime
-#if WITH_EDITOR
 	for (UStaticMesh* PrototypeMesh : PrototypeMeshSet)
 	{
-		if (UUsdAssetImportData* ImportData = Cast<UUsdAssetImportData>(PrototypeMesh->GetAssetImportData()))
+		if (UUsdAssetUserData* UserData = PrototypeMesh->GetAssetUserData<UUsdAssetUserData>())
 		{
-			pxr::SdfPath PrototypePath = UnrealToUsd::ConvertPath(*ImportData->PrimPath).Get();
+			pxr::SdfPath PrototypePath = UnrealToUsd::ConvertPath(*UserData->PrimPath).Get();
 			PrototypeMeshes.insert({PrototypePath, PrototypeMesh});
 		}
 	}
-#endif // WITH_EDITOR
 
 	TArray<TFuture<TTuple<UHierarchicalInstancedStaticMeshComponent*, TArray<FTransform>>>> Tasks;
 	FScopedSlowTask PrototypePathsSlowTask( ( float ) PrototypePaths.size(), LOCTEXT( "GeomPointUpdateComponents", "Updating HierarchicalInstancedStaticMeshComponents for point instancers" ) );
diff --git a/Engine/Plugins/Importers/USDImporter/Source/USDSchemas/Private/USDGeometryCacheTranslator.cpp b/Engine/Plugins/Importers/USDImporter/Source/USDSchemas/Private/USDGeometryCacheTranslator.cpp
--- a/Engine/Plugins/Importers/USDImporter/Source/USDSchemas/Private/USDGeometryCacheTranslator.cpp
+++ b/Engine/Plugins/Importers/USDImporter/Source/USDSchemas/Private/USDGeometryCacheTranslator.cpp
@@ -6,8 +6,7 @@
 // The GeometryCacheStreamer module is editor-only, so is the translator
 
 #include "MeshTranslationImpl.h"
-#include "USDAssetCache.h"
-#include "USDAssetImportData.h"
+#include "USDAssetUserData.h"
 #include "USDClassesModule.h"
 #include "USDConversionUtils.h"
 #include "USDGroomTranslatorUtils.h"
@@ -67,9 +66,11 @@ namespace UsdGeometryCacheTranslatorImpl
 	bool ProcessGeometryCacheMaterials(const pxr::UsdPrim& UsdPrim, const TArray<UsdUtils::FUsdPrimMaterialAssignmentInfo>& LODIndexToMaterialInfo,
 									   UGeometryCache& GeometryCache, UUsdAssetCache2& AssetCache, FUsdInfoCache* InfoCache, float Time, EObjectFlags Flags)
 	{
-		UUsdMeshAssetImportData* MeshImportData = Cast<UUsdMeshAssetImportData>(GeometryCache.AssetImportData.Get());
-		ensureMsgf(MeshImportData, TEXT("Geometry Cache '%s' generated for prim '%s' should have an UUsdMeshAssetImportData at this point!"),
-			*GeometryCache.GetPathName(), *UsdToUnreal::ConvertPath(UsdPrim.GetPrimPath()));
+		UUsdMeshAssetUserData* MeshUserData = GeometryCache.GetAssetUserData<UUsdMeshAssetUserData>();
+		ensureMsgf(MeshUserData, TEXT("Geometry Cache '%s' generated for prim '%s' should have an UUsdMeshAssetUserData at this point!"),
+			*GeometryCache.GetPathName(),
+			*UsdToUnreal::ConvertPath(UsdPrim.GetPrimPath())
+		);
 
 		TMap<const UsdUtils::FUsdPrimMaterialSlot*, UMaterialInterface*> ResolvedMaterials = MeshTranslationImpl::ResolveMaterialAssignmentInfo(
 			UsdPrim, LODIndexToMaterialInfo, AssetCache, *InfoCache, Flags);
@@ -90,7 +91,7 @@ namespace UsdGeometryCacheTranslatorImpl
 				else
 				{
 					// Warn, but still add a material slot to preserve the materials order
-					UE_LOG(LogUsd, Warning, TEXT("Failed to resolve material '%s' for slot '%d' for geometry cache '%s'"), 
+					UE_LOG(LogUsd, Warning, TEXT("Failed to resolve material '%s' for slot '%d' for geometry cache '%s'"),
 						*Slot.MaterialSource, LODSlotIndex, *UsdToUnreal::ConvertPath(UsdPrim.GetPath()));
 				}
 
@@ -105,7 +106,7 @@ namespace UsdGeometryCacheTranslatorImpl
 					bMaterialAssignementsHaveChanged = true;
 				}
 
-				MeshImportData->MaterialSlotToPrimPaths.FindOrAdd(SlotIndex).PrimPaths = Slot.PrimPaths.Array();
+				MeshUserData->MaterialSlotToPrimPaths.FindOrAdd(SlotIndex).PrimPaths = Slot.PrimPaths.Array();
 			}
 		}
 
@@ -795,9 +796,9 @@ void FGeometryCacheCreateAssetsTaskChain::SetupTasks()
 
 			if (GeometryCache && bIsNew)
 			{
-				UUsdMeshAssetImportData* ImportData = NewObject<UUsdMeshAssetImportData>(GeometryCache.Get(), TEXT("UUSDMeshAssetImportData"));
-				ImportData->PrimPath = PrimPathString;
-				GeometryCache->AssetImportData = ImportData;
+				UUsdMeshAssetUserData* UserData = NewObject<UUsdMeshAssetUserData>(GeometryCache.Get(), TEXT("UUSDAssetUserData"));
+				UserData->PrimPath = PrimPathString;
+				GeometryCache->AddAssetUserData(UserData);
 
 				// Only the original creator of the prim at creation time gets to set the material assignments
 				// directly on the geometry cache, all others prims ensure their materials via material overrides on the
@@ -1033,8 +1034,8 @@ TSet<UE::FSdfPath> FUsdGeometryCacheTranslator::CollectAuxiliaryPrims() const
 			if (UsdGeometryCacheTranslatorImpl::IsPurposefullyVisible(ChildPrim.Get(), Context->PurposesToLoad))
 			{
 				Result.Add(UE::FSdfPath(ChildPrim.Get().GetPrimPath()));
-			}	
-			
+			}
+
 			// Should probably collect the Xformable as well
 		}
 	}
diff --git a/Engine/Plugins/Importers/USDImporter/Source/USDSchemas/Private/USDGroomTranslator.cpp b/Engine/Plugins/Importers/USDImporter/Source/USDSchemas/Private/USDGroomTranslator.cpp
--- a/Engine/Plugins/Importers/USDImporter/Source/USDSchemas/Private/USDGroomTranslator.cpp
+++ b/Engine/Plugins/Importers/USDImporter/Source/USDSchemas/Private/USDGroomTranslator.cpp
@@ -18,7 +18,7 @@
 #include "UObject/StrongObjectPtr.h"
 
 #include "UnrealUSDWrapper.h"
-#include "USDAssetImportData.h"
+#include "USDAssetUserData.h"
 #include "USDClassesModule.h"
 #include "USDGroomConversion.h"
 #include "USDGroomTranslatorUtils.h"
@@ -189,13 +189,12 @@ protected:
 					if (GroomAsset)
 					{
 						Context->AssetCache->CacheAsset(SHAHash.ToString(), GroomAsset);
-#if WITH_EDITOR
-						UUsdAssetImportData* ImportData = NewObject<UUsdAssetImportData>(GroomAsset, TEXT("UUSDAssetImportData"));
-						ImportData->PrimPath = PrimPathString;
-						ImportData->ImportOptions = ImportOptions.Get();
 
-						GroomAsset->AssetImportData = ImportData;
-#endif // WITH_EDITOR
+						UUsdAssetUserData* UserData = NewObject<UUsdAssetUserData>(GroomAsset, TEXT("UUSDAssetUserData"));
+						UserData->PrimPath = PrimPathString;
+						UserData->ImportOptions = ImportOptions.Get();
+
+						GroomAsset->AddAssetUserData(UserData);
 					}
 				}
 
@@ -335,13 +334,11 @@ protected:
 				UGroomCache* GroomCache = FGroomCacheImporter::ProcessToGroomCache(*GroomCacheProcessor, AnimInfo, HairImportContext, UniqueName.ToString());
 				if (GroomCache)
 				{
-#if WITH_EDITOR
-					UUsdAssetImportData* ImportData = NewObject<UUsdAssetImportData>(GroomCache, TEXT("UUSDAssetImportData"));
-					ImportData->PrimPath = PrimPath.GetString();
-					ImportData->ImportOptions = ImportOptions.Get();
+					UUsdAssetUserData* UserData = NewObject<UUsdAssetUserData>(GroomCache, TEXT("UUSDAssetUserData"));
+					UserData->PrimPath = PrimPath.GetString();
+					UserData->ImportOptions = ImportOptions.Get();
 
-					GroomCache->AssetImportData = ImportData;
-#endif // WITH_EDITOR
+					GroomCache->AddAssetUserData(UserData);
 
 					Context->AssetCache->CacheAsset(GroomCacheHash.ToString(), GroomCache);
 					Context->InfoCache->LinkAssetToPrim(PrimPath, GroomCache);
diff --git a/Engine/Plugins/Importers/USDImporter/Source/USDSchemas/Private/USDShadeMaterialTranslator.cpp b/Engine/Plugins/Importers/USDImporter/Source/USDSchemas/Private/USDShadeMaterialTranslator.cpp
--- a/Engine/Plugins/Importers/USDImporter/Source/USDSchemas/Private/USDShadeMaterialTranslator.cpp
+++ b/Engine/Plugins/Importers/USDImporter/Source/USDSchemas/Private/USDShadeMaterialTranslator.cpp
@@ -4,7 +4,7 @@
 
 #include "MeshTranslationImpl.h"
 #include "USDAssetCache.h"
-#include "USDAssetImportData.h"
+#include "USDAssetUserData.h"
 #include "USDClassesModule.h"
 #include "USDClassesModule.h"
 #include "USDLog.h"
@@ -160,9 +160,15 @@ namespace UE::UsdShadeTranslator::Private
 								{
 									continue;
 								}
+
 								NewMID->CopyParameterOverrides( MID );
 
-								// TODO: Copy asset user data over to NewMID
+								UUsdMaterialAssetUserData* OldUserData = UserMaterial->GetAssetUserData<UUsdMaterialAssetUserData>();
+								if(OldUserData)
+								{
+									UUsdMaterialAssetUserData* NewUserData = DuplicateObject(OldUserData, NewMID);
+									NewMID->AddAssetUserData(NewUserData);
+								}
 
 								if (Context->AssetCache->CanRemoveAsset(Hash) && Context->AssetCache->RemoveAsset(Hash))
 								{
@@ -273,9 +279,9 @@ void FUsdShadeMaterialTranslator::CreateAssets()
 		{
 			if ( UMaterialInstanceConstant* NewMaterial = NewObject<UMaterialInstanceConstant>( GetTransientPackage(), InstanceName, Context->ObjectFlags | EObjectFlags::RF_Transient ) )
 			{
-				UUsdMaterialAssetImportData* ImportData = NewObject<UUsdMaterialAssetImportData>(NewMaterial, TEXT("USDAssetImportData"));
-				ImportData->PrimPath = PrimPath.GetString();
-				NewMaterial->AssetImportData = ImportData;
+				UUsdMaterialAssetUserData* UserData = NewObject<UUsdMaterialAssetUserData>(NewMaterial, TEXT("USDAssetUserData"));
+				UserData->PrimPath = PrimPath.GetString();
+				NewMaterial->AddAssetUserData(UserData);
 
 				const bool bSuccess = UsdToUnreal::ConvertMaterial(
 					ShadeMaterial,
@@ -371,9 +377,11 @@ void FUsdShadeMaterialTranslator::CreateAssets()
 			{
 				if ( UMaterialInstanceDynamic* NewMaterial = UMaterialInstanceDynamic::Create( ReferenceMaterial, GetTransientPackage(), InstanceName ) )
 				{
-					NewMaterial->SetFlags(RF_Transient);
+					UUsdMaterialAssetUserData* UserData = NewObject<UUsdMaterialAssetUserData>(NewMaterial, TEXT("USDAssetUserData"));
+					UserData->PrimPath = PrimPath.GetString();
+					NewMaterial->AddAssetUserData(UserData);
 
-					// TODO: Create asset user data here so that ConvertMaterial can store the primvars used on each UV index
+					NewMaterial->SetFlags(RF_Transient);
 
 					if (UsdToUnreal::ConvertMaterial(ShadeMaterial, *NewMaterial, Context->AssetCache.Get(), *Context->RenderContext.ToString()))
 					{
@@ -444,15 +452,13 @@ void FUsdShadeMaterialTranslator::PostImportMaterial(const FString& MaterialHash
 		return;
 	}
 
-#if WITH_EDITOR
-	UUsdMaterialAssetImportData* ImportData = Cast<UUsdMaterialAssetImportData>(ImportedMaterial->AssetImportData);
-	if (!ImportData)
+	UUsdMaterialAssetUserData* UserData = ImportedMaterial->GetAssetUserData<UUsdMaterialAssetUserData>();
+	if (!UserData)
 	{
-		ImportData = NewObject<UUsdMaterialAssetImportData>(ImportedMaterial, TEXT("USDAssetImportData"));
-		ImportedMaterial->AssetImportData = ImportData;
+		UserData = NewObject<UUsdMaterialAssetUserData>(ImportedMaterial, TEXT("USDAssetUserData"));
+		ImportedMaterial->AddAssetUserData(UserData);
 	}
-	ImportData->PrimPath = PrimPath.GetString();
-#endif // WITH_EDITOR
+	UserData->PrimPath = PrimPath.GetString();
 
 	// Note that this needs to run even if we found this material in the asset cache already, otherwise we won't
 	// re-register the prim asset links when we reload a stage
@@ -482,11 +488,15 @@ void FUsdShadeMaterialTranslator::PostImportMaterial(const FString& MaterialHash
 			else if (bIsOwnedByTransientPackage)
 			{
 				FString FilePath;
-#if WITH_EDITOR
-				FilePath = Texture->AssetImportData ? Texture->AssetImportData->GetFirstFilename() : Texture->GetName();
-#else
-				FilePath = Texture->GetName();
-#endif // WITH_EDITOR
+				if(UUsdAssetUserData* TextureUserData = Texture->GetAssetUserData<UUsdAssetUserData>())
+				{
+					FilePath = TextureUserData->FilePath;
+				}
+				else
+				{
+					FilePath = Texture->GetName();
+				}
+
 				const FString TextureHash = UsdUtils::GetTextureHash(
 					FilePath,
 					Texture->SRGB,
@@ -512,15 +522,13 @@ void FUsdShadeMaterialTranslator::PostImportMaterial(const FString& MaterialHash
 
 			if (bIsOwnedByCache || bIsOwnedByTransientPackage)
 			{
-#if WITH_EDITOR
-				UUsdAssetImportData* TextureImportData = Cast<UUsdAssetImportData>(Texture->AssetImportData);
-				if (!TextureImportData)
+				UUsdAssetUserData* TextureUserData = Texture->GetAssetUserData<UUsdAssetUserData>();
+				if (!TextureUserData)
 				{
-					TextureImportData = NewObject<UUsdAssetImportData>(Texture, TEXT("USDAssetImportData"));
-					Texture->AssetImportData = TextureImportData;
+					TextureUserData = NewObject<UUsdAssetUserData>(ImportedMaterial, TEXT("USDAssetUserData"));
+					Texture->AddAssetUserData(TextureUserData);
 				}
-				TextureImportData->PrimPath = PrimPath.GetString();
-#endif // WITH_EDITOR
+				TextureUserData->PrimPath = PrimPath.GetString();
 
 				Context->InfoCache->LinkAssetToPrim(PrimPath, Texture);
 			}
diff --git a/Engine/Plugins/Importers/USDImporter/Source/USDSchemas/Private/USDSkelRootTranslator.cpp b/Engine/Plugins/Importers/USDImporter/Source/USDSchemas/Private/USDSkelRootTranslator.cpp
--- a/Engine/Plugins/Importers/USDImporter/Source/USDSchemas/Private/USDSkelRootTranslator.cpp
+++ b/Engine/Plugins/Importers/USDImporter/Source/USDSchemas/Private/USDSkelRootTranslator.cpp
@@ -6,7 +6,7 @@
 
 #include "MeshTranslationImpl.h"
 #include "USDAssetCache.h"
-#include "USDAssetImportData.h"
+#include "USDAssetUserData.h"
 #include "USDClassesModule.h"
 #include "USDConversionUtils.h"
 #include "USDErrorUtils.h"
@@ -101,13 +101,11 @@ namespace UsdSkelRootTranslatorImpl
 			ExistingAssignments.Add( SkeletalMaterial.MaterialInterface );
 		}
 
-#if WITH_EDITOR
-		UUsdMeshAssetImportData* ImportData = Cast< UUsdMeshAssetImportData>( SkeletalMesh->GetAssetImportData() );
-		ensureMsgf(ImportData, TEXT("Skeletal Mesh '%s' generated for prim '%s' should have an UUsdMeshAssetImportData at this point!"),
+		UUsdMeshAssetUserData* UserData = SkeletalMesh->GetAssetUserData<UUsdMeshAssetUserData>();
+		ensureMsgf(UserData, TEXT("Skeletal Mesh '%s' generated for prim '%s' should have an UUsdMeshAssetUserData at this point!"),
 			*SkeletalMesh->GetPathName(),
 			*UsdToUnreal::ConvertPath(UsdPrim.GetPrimPath())
 		);
-#endif // WITH_EDITOR
 
 		TMap<const UsdUtils::FUsdPrimMaterialSlot*, UMaterialInterface*> ResolvedMaterials = MeshTranslationImpl::ResolveMaterialAssignmentInfo(
 			UsdPrim,
@@ -186,12 +184,10 @@ namespace UsdSkelRootTranslatorImpl
 					bMaterialsHaveChanged = true;
 				}
 
-#if WITH_EDITOR
-				if ( ImportData )
+				if (UserData)
 				{
-					ImportData->MaterialSlotToPrimPaths.FindOrAdd( SkeletalMeshSlotIndex ).PrimPaths = Slot.PrimPaths.Array();
+					UserData->MaterialSlotToPrimPaths.FindOrAdd(SkeletalMeshSlotIndex).PrimPaths = Slot.PrimPaths.Array();
 				}
-#endif // WITH_EDITOR
 
 				// Already have a material at that LOD remap slot, need to reassign
 				if ( LODMaterialMap.IsValidIndex( LODSlotIndex ) )
@@ -1334,9 +1330,9 @@ namespace UsdSkelRootTranslatorImpl
 				{
 					if (bIsNew)
 					{
-						UUsdMeshAssetImportData* ImportData = NewObject< UUsdMeshAssetImportData >( SkeletalMesh, TEXT( "USDAssetImportData" ) );
-						ImportData->PrimPath = SkelRootPath;
-						SkeletalMesh->SetAssetImportData(ImportData);
+						UUsdMeshAssetUserData* UserData = NewObject<UUsdMeshAssetUserData>(SkeletalMesh, TEXT("USDAssetUserData"));
+						UserData->PrimPath = SkelRootPath;
+						SkeletalMesh->AddAssetUserData(UserData);
 
 						const bool bMaterialsHaveChanged = UsdSkelRootTranslatorImpl::ProcessMaterials(
 							GetPrim(),
@@ -1557,11 +1553,11 @@ namespace UsdSkelRootTranslatorImpl
 
 							if (AnimSequence->GetDataModel()->GetNumBoneTracks() != 0 || AnimSequence->GetDataModel()->GetNumberOfFloatCurves() != 0 )
 							{
-								UUsdAnimSequenceAssetImportData* ImportData = NewObject< UUsdAnimSequenceAssetImportData >( AnimSequence, TEXT( "USDAssetImportData" ) );
-								AnimSequence->AssetImportData = ImportData;
+								UUsdAnimSequenceAssetUserData* UserData = NewObject<UUsdAnimSequenceAssetUserData>(AnimSequence, TEXT("USDAssetUserData"));
+								UserData->PrimPath = SkelAnimationPrimPath;
+								UserData->LayerStartOffsetSeconds = LayerStartOffsetSeconds;
 
-								ImportData->PrimPath = SkelAnimationPrimPath;
-								ImportData->LayerStartOffsetSeconds = LayerStartOffsetSeconds;
+								AnimSequence->AddAssetUserData(UserData);
 
 								Context->AssetCache->CacheAsset( HashString, AnimSequence );
 							}
@@ -1782,9 +1778,9 @@ void FUsdSkelRootTranslator::UpdateComponents( USceneComponent* SceneComponent )
 			}
 
 			double LayerStartOffsetSeconds = 0.0f;
-			if ( UUsdAnimSequenceAssetImportData* ImportData = Cast<UUsdAnimSequenceAssetImportData>( AnimSequence->AssetImportData ) )
+			if (UUsdAnimSequenceAssetUserData* UserData = AnimSequence->GetAssetUserData<UUsdAnimSequenceAssetUserData>())
 			{
-				LayerStartOffsetSeconds = ImportData->LayerStartOffsetSeconds;
+				LayerStartOffsetSeconds = UserData->LayerStartOffsetSeconds;
 			}
 
 			// Always change the mode here because the sequencer will change it back to AnimationCustomMode when animating
diff --git a/Engine/Plugins/Importers/USDImporter/Source/USDStage/Private/USDLevelSequenceHelper.cpp b/Engine/Plugins/Importers/USDImporter/Source/USDStage/Private/USDLevelSequenceHelper.cpp
--- a/Engine/Plugins/Importers/USDImporter/Source/USDStage/Private/USDLevelSequenceHelper.cpp
+++ b/Engine/Plugins/Importers/USDImporter/Source/USDStage/Private/USDLevelSequenceHelper.cpp
@@ -3,7 +3,7 @@
 #include "USDLevelSequenceHelper.h"
 
 #include "GeometryCache.h"
-#include "USDAssetImportData.h"
+#include "USDAssetUserData.h"
 #include "USDAttributeUtils.h"
 #include "USDClassesModule.h"
 #include "USDConversionUtils.h"
@@ -1898,12 +1898,11 @@ void FUsdLevelSequenceHelperImpl::AddSkeletalTracks(const UUsdPrimTwin& PrimTwin
 	if (UMovieSceneSkeletalAnimationTrack* SkeletalTrack = AddTrack<UMovieSceneSkeletalAnimationTrack>(SkelAnimationPrim.GetName(), PrimTwin, *ComponentToBind, *SkelAnimationSequence, bIsMuted))
 	{
 		double LayerStartOffsetSeconds = 0.0f;
-#if WITH_EDITOR
-		if (UUsdAnimSequenceAssetImportData* ImportData = Cast<UUsdAnimSequenceAssetImportData>(Sequence->AssetImportData))
+		if (UUsdAnimSequenceAssetUserData* UserData = Sequence->GetAssetUserData<UUsdAnimSequenceAssetUserData>())
 		{
-			LayerStartOffsetSeconds = ImportData->LayerStartOffsetSeconds;
+			LayerStartOffsetSeconds = UserData->LayerStartOffsetSeconds;
 		}
-#endif // WITH_EDITOR
+
 		FFrameNumber StartOffsetTick = FFrameTime::FromDecimal(LayerStartOffsetSeconds * MovieScene->GetTickResolution().AsDecimal()).RoundToFrame();
 
 		SkeletalTrack->RemoveAllAnimationData();
diff --git a/Engine/Plugins/Importers/USDImporter/Source/USDStageImporter/Private/USDStageAssetImportFactory.cpp b/Engine/Plugins/Importers/USDImporter/Source/USDStageImporter/Private/USDStageAssetImportFactory.cpp
--- a/Engine/Plugins/Importers/USDImporter/Source/USDStageImporter/Private/USDStageAssetImportFactory.cpp
+++ b/Engine/Plugins/Importers/USDImporter/Source/USDStageImporter/Private/USDStageAssetImportFactory.cpp
@@ -3,6 +3,7 @@
 #include "USDStageAssetImportFactory.h"
 
 #include "USDAssetImportData.h"
+#include "USDAssetUserData.h"
 #include "USDErrorUtils.h"
 #include "USDLog.h"
 #include "USDStageImporter.h"
@@ -111,11 +112,23 @@ void UUsdStageAssetImportFactory::CleanUp()
 
 bool UUsdStageAssetImportFactory::CanReimport(UObject* Obj, TArray<FString>& OutFilenames)
 {
+	if (IInterface_AssetUserData* UserDataInterface = Cast<IInterface_AssetUserData>(Obj))
+	{
+		if(UUsdAssetUserData* UserData = UserDataInterface->GetAssetUserData<UUsdAssetUserData>())
+		{
+			OutFilenames.Add(UserData->FilePath);
+			return true;
+		}
+	}
+
+	// Fall back to also checking asset import data for now, in case it's an old asset
+	PRAGMA_DISABLE_DEPRECATION_WARNINGS
 	if (UUsdAssetImportData* ImportData = UsdUtils::GetAssetImportData(Obj))
 	{
 		OutFilenames.Add(ImportData->GetFirstFilename());
 		return true;
 	}
+	PRAGMA_ENABLE_DEPRECATION_WARNINGS
 
 	return false;
 }
@@ -127,42 +140,103 @@ void UUsdStageAssetImportFactory::SetReimportPaths(UObject* Obj, const TArray<FS
 		return;
 	}
 
+	if (IInterface_AssetUserData* UserDataInterface = Cast<IInterface_AssetUserData>(Obj))
+	{
+		if(UUsdAssetUserData* UserData = UserDataInterface->GetAssetUserData<UUsdAssetUserData>())
+		{
+			UserData->FilePath = NewReimportPaths[0];
+		}
+	}
+
+	PRAGMA_DISABLE_DEPRECATION_WARNINGS
 	if (UUsdAssetImportData* ImportData = UsdUtils::GetAssetImportData(Obj))
 	{
 		ImportData->UpdateFilenameOnly(NewReimportPaths[0]);
 	}
+	PRAGMA_ENABLE_DEPRECATION_WARNINGS
 }
 
 EReimportResult::Type UUsdStageAssetImportFactory::Reimport(UObject* Obj)
 {
 	if (!Obj)
 	{
-		FUsdLogManager::LogMessage( EMessageSeverity::Error, LOCTEXT( "ReimportErrorInvalidAsset", "Failed to reimport asset as it is invalid!" ) );
+		FUsdLogManager::LogMessage(
+			EMessageSeverity::Error,
+			LOCTEXT("ReimportErrorInvalidAsset", "Failed to reimport asset as it is invalid!")
+		);
 		return EReimportResult::Failed;
 	}
 
-	UUsdAssetImportData* ImportData = UsdUtils::GetAssetImportData(Obj);
-	if (!ImportData)
+	FString ReimportFilePath;
+	FString OriginalPrimPath;
+	UObject* ReimportOptions = nullptr;
+
+	if (IInterface_AssetUserData* UserDataInterface = Cast<IInterface_AssetUserData>(Obj))
+	{
+		if (UUsdAssetUserData* UserData = UserDataInterface->GetAssetUserData<UUsdAssetUserData>())
+		{
+			ReimportFilePath = UserData->FilePath;
+			OriginalPrimPath = UserData->PrimPath;
+			ReimportOptions = UserData->ImportOptions;
+		}
+	}
+
+	if (!ReimportOptions || ReimportFilePath.IsEmpty())
+	{
+		PRAGMA_DISABLE_DEPRECATION_WARNINGS
+		if (UUsdAssetImportData* ImportData = UsdUtils::GetAssetImportData(Obj))
+		{
+			ReimportFilePath = ImportData->GetFirstFilename();
+			OriginalPrimPath = ImportData->PrimPath;
+			ReimportOptions = ImportData->ImportOptions;
+		}
+		PRAGMA_ENABLE_DEPRECATION_WARNINGS
+	}
+
+	if (ReimportFilePath.IsEmpty() || OriginalPrimPath.IsEmpty())
 	{
-		FUsdLogManager::LogMessage( EMessageSeverity::Error, FText::Format( LOCTEXT( "ReimportErrorNoImportData", "Failed to reimport asset '{0}' as it doesn't seem to have import data!" ), FText::FromName( Obj->GetFName() ) ) );
+		FUsdLogManager::LogMessage(
+			EMessageSeverity::Error,
+			FText::Format(
+				LOCTEXT(
+					"ReimportErrorNoImportData",
+					"Failed to reimport asset '{0}' as it doesn't seem to have valid USD import data or user data!"
+				),
+				FText::FromName(Obj->GetFName())
+			)
+		);
 		return EReimportResult::Failed;
 	}
 
-	if ( UUsdStageImportOptions* ReimportOptions = Cast<UUsdStageImportOptions>( ImportData->ImportOptions ) )
+	if (ReimportOptions)
 	{
 		// Duplicate this as we may update these options on the Init call below, and if we just imported a scene and
-		// all assets are in memory (sharing the same import options object), that update would otherwise affect all the UUsdAssetImportData objects,
-		// which is not what we would expect
-		ImportContext.ImportOptions = DuplicateObject( ReimportOptions, ImportData );
+		// all assets are in memory (sharing the same import options object), that update would otherwise affect all the
+		// UUsdAssetImportData objects, which is not what we would expect
+		ImportContext.ImportOptions = Cast<UUsdStageImportOptions>(DuplicateObject(ReimportOptions, GetTransientPackage()));
 	}
 
 	ImportContext.bReadFromStageCache = false;
 
 	const bool bIsReimport = true;
 	const bool bAllowActorImport = false;
-	if (!ImportContext.Init(Obj->GetName(), ImportData->GetFirstFilename(), Obj->GetName(), Obj->GetFlags(), IsAutomatedImport(), bIsReimport, bAllowActorImport))
+	if (!ImportContext.Init(
+			Obj->GetName(),
+			ReimportFilePath,
+			Obj->GetName(),
+			Obj->GetFlags(),
+			IsAutomatedImport(),
+			bIsReimport,
+			bAllowActorImport
+		))
 	{
-		FUsdLogManager::LogMessage( EMessageSeverity::Error, FText::Format( LOCTEXT( "ReimportErrorNoContext", "Failed to initialize reimport context for asset '{0}'!" ), FText::FromName( Obj->GetFName() ) ) );
+		FUsdLogManager::LogMessage(
+			EMessageSeverity::Error,
+			FText::Format(
+				LOCTEXT("ReimportErrorNoContext", "Failed to initialize reimport context for asset '{0}'!"),
+				FText::FromName(Obj->GetFName())
+			)
+		);
 		return EReimportResult::Failed;
 	}
 
@@ -172,7 +246,7 @@ EReimportResult::Type UUsdStageAssetImportFactory::Reimport(UObject* Obj)
 
 	UUsdStageImporter* USDImporter = IUsdStageImporterModule::Get().GetImporter();
 	UObject* ReimportedAsset = nullptr;
-	bool bSuccess = USDImporter->ReimportSingleAsset(ImportContext, Obj, ImportData, ReimportedAsset);
+	bool bSuccess = USDImporter->ReimportSingleAsset(ImportContext, Obj, OriginalPrimPath, ReimportedAsset);
 
 	GEditor->GetEditorSubsystem<UImportSubsystem>()->BroadcastAssetReimport(Obj);
 
diff --git a/Engine/Plugins/Importers/USDImporter/Source/USDStageImporter/Private/USDStageImporter.cpp b/Engine/Plugins/Importers/USDImporter/Source/USDStageImporter/Private/USDStageImporter.cpp
--- a/Engine/Plugins/Importers/USDImporter/Source/USDStageImporter/Private/USDStageImporter.cpp
+++ b/Engine/Plugins/Importers/USDImporter/Source/USDStageImporter/Private/USDStageImporter.cpp
@@ -4,7 +4,7 @@
 
 #include "USDAssetCache.h"
 #include "USDAssetCache2.h"
-#include "USDAssetImportData.h"
+#include "USDAssetUserData.h"
 #include "USDClassesModule.h"
 #include "USDConversionUtils.h"
 #include "USDErrorUtils.h"
@@ -431,9 +431,15 @@ namespace UsdStageImporterImpl
 		FString AssetPrefix;
 
 		FString AssetPath = Asset->GetName();
-		if (UUsdAssetImportData* AssetImportData = UsdUtils::GetAssetImportData(Asset))
+
+		FString PrimPath;
+		if (IInterface_AssetUserData* UserDataInterface = Cast<IInterface_AssetUserData>(Asset))
 		{
-			AssetPath = AssetImportData->PrimPath;
+			if (UUsdAssetUserData* UserData = UserDataInterface->GetAssetUserData<UUsdAssetUserData>())
+			{
+				PrimPath = UserData->PrimPath;
+				AssetPath = PrimPath;
+			}
 		}
 
 		FString AssetName = FPaths::GetBaseFilename( AssetPath );
@@ -442,21 +448,19 @@ namespace UsdStageImporterImpl
 		{
 			AssetPrefix = TEXT("SM_");
 
-			if (UUsdAssetImportData* AssetImportData = Cast<UUsdAssetImportData>(Mesh->AssetImportData))
+			// If we have multiple LODs here we must have parsed the LOD variant set pattern. If our prims were named
+			// with the LOD pattern, go from e.g. '/Root/MyMesh/LOD0' to '/Root/MyMesh', or else every single LOD mesh
+			// will be named "SM_LOD0_X". We'll actually check though because if the user set a custom name for their
+			// prim other than LOD0 then we'll keep that
+			if (Mesh->GetNumLODs() > 1)
 			{
-				// If we have multiple LODs here we must have parsed the LOD variant set pattern. If our prims were named with the LOD
-				// pattern, go from e.g. '/Root/MyMesh/LOD0' to '/Root/MyMesh', or else every single LOD mesh will be named "SM_LOD0_X".
-				// We'll actually check though because if the user set a custom name for their prim other than LOD0 then we'll keep that
-				if ( Mesh->GetNumLODs() > 1 )
+				FString PrimName = FPaths::GetBaseFilename(AssetPath);
+				if (PrimName.RemoveFromStart(TEXT("LOD"), ESearchCase::CaseSensitive))
 				{
-					FString PrimName = FPaths::GetBaseFilename( AssetPath );
-					if ( PrimName.RemoveFromStart( TEXT( "LOD" ), ESearchCase::CaseSensitive ) )
+					if (PrimName.IsNumeric())
 					{
-						if ( PrimName.IsNumeric() )
-						{
-							AssetPath = FPaths::GetPath( AssetPath );
-							AssetName = FPaths::GetBaseFilename( AssetPath );
-						}
+						AssetPath = FPaths::GetPath(AssetPath);
+						AssetName = FPaths::GetBaseFilename(AssetPath);
 					}
 				}
 			}
@@ -504,25 +508,22 @@ namespace UsdStageImporterImpl
 				AssetPrefix = TEXT( "M_" );
 			}
 
-			if (UUsdAssetImportData* AssetImportData = Cast<UUsdAssetImportData>(Material->AssetImportData))
+			// The only materials with no prim path are our auto-generated displayColor materials
+			if (PrimPath.IsEmpty())
 			{
-				// The only materials with no prim path are our auto-generated displayColor materials
-				if ( AssetImportData->PrimPath.IsEmpty() )
-				{
-					AssetPath = TEXT("DisplayColor");
-					AssetName = FPaths::GetBaseFilename( AssetPath );
-				}
-				else
-				{
-					AssetPath = AssetImportData->PrimPath;
-					AssetName = FPaths::GetBaseFilename( AssetPath );
+				AssetPath = TEXT("DisplayColor");
+				AssetName = FPaths::GetBaseFilename(AssetPath);
+			}
+			else
+			{
+				AssetPath = PrimPath;
+				AssetName = FPaths::GetBaseFilename(AssetPath);
 
-					// If we have a preview surface two-sided material we'll also have a one-sided with the same name,
-					// so add a suffix here so we can clearly tell which is which
-					if ( Material->IsTwoSided() )
-					{
-						AssetName += UnrealIdentifiers::TwoSidedMaterialSuffix;
-					}
+				// If we have a preview surface two-sided material we'll also have a one-sided with the same name,
+				// so add a suffix here so we can clearly tell which is which
+				if (Material->IsTwoSided())
+				{
+					AssetName += UnrealIdentifiers::TwoSidedMaterialSuffix;
 				}
 			}
 		}
@@ -564,33 +565,41 @@ namespace UsdStageImporterImpl
 		return AssetName;
 	}
 
-	void UpdateAssetImportData( UObject* Asset, const FString& MainFilePath, UUsdStageImportOptions* ImportOptions )
+	void UpdateAssetUserData(UObject* Asset, const FString& MainFilePath, UUsdStageImportOptions* ImportOptions)
 	{
-		if ( !Asset )
+		if (!Asset)
 		{
 			return;
 		}
 
-		UUsdAssetImportData* ImportData = UsdUtils::GetAssetImportData( Asset );
-		if ( !ImportData )
+		UUsdAssetUserData* UserData = nullptr;
+		if (IInterface_AssetUserData* UserDataInterface = Cast<IInterface_AssetUserData>(Asset))
+		{
+			UserData = UserDataInterface->GetAssetUserData<UUsdAssetUserData>();
+		}
+		if (!UserData)
 		{
 			return;
 		}
 
 		// Don't force update as textures will already come with this preset to their actual texture path
-		if ( ImportData->SourceData.SourceFiles.Num() == 0 )
+		if (!Asset->IsA<UTexture>())
 		{
-			ImportData->UpdateFilenameOnly( MainFilePath );
+			UserData->FilePath = MainFilePath;
 		}
 
-		ImportData->ImportOptions = ImportOptions;
+		UserData->ImportOptions = ImportOptions;
 	}
 
-	void UpdateAssetImportData(const TSet<UObject*>& UsedAssetsAndDependencies, const FString& MainFilePath, UUsdStageImportOptions* ImportOptions)
+	void UpdateAssetUserData(
+		const TSet<UObject*>& UsedAssetsAndDependencies,
+		const FString& MainFilePath,
+		UUsdStageImportOptions* ImportOptions
+	)
 	{
-		for ( UObject* Asset : UsedAssetsAndDependencies )
+		for (UObject* Asset : UsedAssetsAndDependencies)
 		{
-			UpdateAssetImportData( Asset, MainFilePath, ImportOptions );
+			UpdateAssetUserData(Asset, MainFilePath, ImportOptions);
 		}
 	}
 
@@ -870,17 +879,20 @@ namespace UsdStageImporterImpl
 
 			if (ImportContext.ImportOptions->bPrimPathFolderStructure)
 			{
-				if (UUsdAssetImportData* ImportData = UsdUtils::GetAssetImportData(Asset))
+				if (IInterface_AssetUserData* UserDataInterface = Cast<IInterface_AssetUserData>(Asset))
 				{
-					// For skeletal stuff, the primpaths point to the SkelRoot, so it is useful to place the assets in there,
-					// as we'll always have at least the skeletal mesh and the skeleton
-					if (AssetTypeFolderPtr == &SkeletalMeshesFolder)
-					{
-						AssetToContentFolder.Add(Asset, ImportData->PrimPath);
-					}
-					else
+					if (UUsdAssetUserData* UserData = UserDataInterface->GetAssetUserData<UUsdAssetUserData>())
 					{
-						AssetToContentFolder.Add(Asset, FPaths::GetPath(ImportData->PrimPath));
+						// For skeletal stuff, the primpaths point to the SkelRoot, so it is useful to place the assets in there,
+						// as we'll always have at least the skeletal mesh and the skeleton
+						if (AssetTypeFolderPtr == &SkeletalMeshesFolder)
+						{
+							AssetToContentFolder.Add(Asset, UserData->PrimPath);
+						}
+						else
+						{
+							AssetToContentFolder.Add(Asset, FPaths::GetPath(UserData->PrimPath));
+						}
 					}
 				}
 			}
@@ -1739,7 +1751,7 @@ void UUsdStageImporter::ImportFromFile(FUsdStageImportContext& ImportContext)
 
 	UsdStageImporterImpl::CollectUsedAssetDependencies( ImportContext, UsedAssetsAndDependencies );
 	UsdStageImporterImpl::PruneUnwantedAssets( ImportContext, UsedAssetsAndDependencies, ObjectsToRemap, SoftObjectsToRemap );
-	UsdStageImporterImpl::UpdateAssetImportData( UsedAssetsAndDependencies, ImportContext.FilePath, ImportContext.ImportOptions );
+	UsdStageImporterImpl::UpdateAssetUserData( UsedAssetsAndDependencies, ImportContext.FilePath, ImportContext.ImportOptions );
 	UsdStageImporterImpl::PublishAssets( ImportContext, UsedAssetsAndDependencies, ObjectsToRemap, SoftObjectsToRemap );
 	UsdStageImporterImpl::ResolveActorConflicts( ImportContext, ExistingSceneActor, ObjectsToRemap, SoftObjectsToRemap );
 	UsdStageImporterImpl::RemapReferences( ImportContext, UsedAssetsAndDependencies, ObjectsToRemap );
@@ -1769,7 +1781,12 @@ void UUsdStageImporter::ImportFromFile(FUsdStageImportContext& ImportContext)
 #endif // #if USE_USD_SDK
 }
 
-bool UUsdStageImporter::ReimportSingleAsset(FUsdStageImportContext& ImportContext, UObject* OriginalAsset, UUsdAssetImportData* OriginalImportData, UObject*& OutReimportedAsset)
+bool UUsdStageImporter::ReimportSingleAsset(
+	FUsdStageImportContext& ImportContext,
+	UObject* OriginalAsset,
+	const FString& OriginalPrimPath,
+	UObject*& OutReimportedAsset
+)
 {
 	OutReimportedAsset = nullptr;
 	bool bSuccess = false;
@@ -1851,7 +1868,7 @@ bool UUsdStageImporter::ReimportSingleAsset(FUsdStageImportContext& ImportContex
 	{
 		UsdStageImporterImpl::CacheCollapsingState( TranslationContext.Get() );
 
-		UE::FUsdPrim TargetPrim = ImportContext.Stage.GetPrimAtPath( UE::FSdfPath( *OriginalImportData->PrimPath ) );
+		UE::FUsdPrim TargetPrim = ImportContext.Stage.GetPrimAtPath(UE::FSdfPath(*OriginalPrimPath));
 		if ( TargetPrim )
 		{
 			UsdStageImporterImpl::CreateAssetsForPrims({TargetPrim}, TranslationContext.Get(), LOCTEXT("CreateAssets", "Creating assets"));
@@ -1863,12 +1880,16 @@ bool UUsdStageImporter::ReimportSingleAsset(FUsdStageImportContext& ImportContex
 	UObject* ReimportedObject = nullptr;
 	for ( UObject* Asset : ImportContext.AssetCache->GetActiveAssets() )
 	{
-		UUsdAssetImportData* NewAssetImportData = UsdUtils::GetAssetImportData( Asset );
+		UUsdAssetUserData* UserData = nullptr;
+		if (IInterface_AssetUserData* UserDataInterface = Cast<IInterface_AssetUserData>(Asset))
+		{
+			UserData = UserDataInterface->GetAssetUserData<UUsdAssetUserData>();
+		}
 
 		if ( Asset &&
-			 NewAssetImportData &&
+			 UserData &&
 			 Asset->GetClass() == OriginalAsset->GetClass() &&
-			 NewAssetImportData->PrimPath.Equals( OriginalImportData->PrimPath, ESearchCase::CaseSensitive ) )
+			 UserData->PrimPath.Equals(OriginalPrimPath, ESearchCase::CaseSensitive))
 		{
 			ReimportedObject = Asset;
 			break;
@@ -1877,7 +1898,7 @@ bool UUsdStageImporter::ReimportSingleAsset(FUsdStageImportContext& ImportContex
 
 	if ( ReimportedObject )
 	{
-		UsdStageImporterImpl::UpdateAssetImportData( ReimportedObject, ImportContext.FilePath, ImportContext.ImportOptions);
+		UsdStageImporterImpl::UpdateAssetUserData(ReimportedObject, ImportContext.FilePath, ImportContext.ImportOptions);
 
 		// Assign things from the original assets before we publish the reimported asset, overwriting it
 		UsdStageImporterImpl::CopyOriginalMaterialAssignment(ImportContext, OriginalAsset, ReimportedObject );
diff --git a/Engine/Plugins/Importers/USDImporter/Source/USDStageImporter/Public/USDStageImporter.h b/Engine/Plugins/Importers/USDImporter/Source/USDStageImporter/Public/USDStageImporter.h
--- a/Engine/Plugins/Importers/USDImporter/Source/USDStageImporter/Public/USDStageImporter.h
+++ b/Engine/Plugins/Importers/USDImporter/Source/USDStageImporter/Public/USDStageImporter.h
@@ -14,5 +14,10 @@ class USDSTAGEIMPORTER_API UUsdStageImporter
 public:
 	void ImportFromFile(FUsdStageImportContext& ImportContext);
 
-	bool ReimportSingleAsset(FUsdStageImportContext& ImportContext, UObject* OriginalAsset, UUsdAssetImportData* OriginalImportData, UObject*& OutReimportedAsset);
+	bool ReimportSingleAsset(
+		FUsdStageImportContext& ImportContext,
+		UObject* OriginalAsset,
+		const FString& OriginalPrimPath,
+		UObject*& OutReimportedAsset
+	);
 };
diff --git a/Engine/Plugins/Importers/USDImporter/Source/USDUtilities/Private/USDConversionUtils.cpp b/Engine/Plugins/Importers/USDImporter/Source/USDUtilities/Private/USDConversionUtils.cpp
--- a/Engine/Plugins/Importers/USDImporter/Source/USDUtilities/Private/USDConversionUtils.cpp
+++ b/Engine/Plugins/Importers/USDImporter/Source/USDUtilities/Private/USDConversionUtils.cpp
@@ -1243,6 +1243,7 @@ double UsdUtils::GetEarliestTimeCode()
 #endif
 }
 
+PRAGMA_DISABLE_DEPRECATION_WARNINGS
 UUsdAssetImportData* UsdUtils::GetAssetImportData( UObject* Asset )
 {
 	UUsdAssetImportData* ImportData = nullptr;
@@ -1309,6 +1310,7 @@ UUsdAssetImportData* UsdUtils::GetAssetImportData( UObject* Asset )
 #endif
 	return ImportData;
 }
+PRAGMA_ENABLE_DEPRECATION_WARNINGS
 
 namespace UE::UsdConversionUtils::Private
 {
diff --git a/Engine/Plugins/Importers/USDImporter/Source/USDUtilities/Private/USDPrimConversion.cpp b/Engine/Plugins/Importers/USDImporter/Source/USDUtilities/Private/USDPrimConversion.cpp
--- a/Engine/Plugins/Importers/USDImporter/Source/USDUtilities/Private/USDPrimConversion.cpp
+++ b/Engine/Plugins/Importers/USDImporter/Source/USDUtilities/Private/USDPrimConversion.cpp
@@ -2,9 +2,8 @@
 
 #include "USDPrimConversion.h"
 
-#include "Engine/SkinnedAssetCommon.h"
 #include "UnrealUSDWrapper.h"
-#include "USDAssetImportData.h"
+#include "USDAssetUserData.h"
 #include "USDAttributeUtils.h"
 #include "USDConversionUtils.h"
 #include "USDLayerUtils.h"
@@ -37,6 +36,7 @@
 #include "Components/SkinnedMeshComponent.h"
 #include "Components/SpotLightComponent.h"
 #include "Engine/SkeletalMesh.h"
+#include "Engine/SkinnedAssetCommon.h"
 #include "Engine/StaticMesh.h"
 #include "GeometryCache.h"
 #include "GeometryCacheComponent.h"
@@ -2051,13 +2051,16 @@ bool UnrealToUsd::ConvertMeshComponent( const pxr::UsdStageRefPtr& Stage, const
 					continue;
 				}
 
-				UUsdMeshAssetImportData* ImportData = Cast<UUsdMeshAssetImportData>( GeometryCache->AssetImportData.Get() );
-				if ( !ImportData )
+				// const_cast as there's no const access to asset user data on IInterface_AssetUserData, but we won't
+				// modify anything
+				const UUsdMeshAssetUserData* UserData =
+					const_cast<UGeometryCache*>(GeometryCache)->GetAssetUserData<UUsdMeshAssetUserData>();
+				if (!UserData)
 				{
 					return false;
 				}
 
-				if ( const FUsdPrimPathList* SourcePrimPaths = ImportData->MaterialSlotToPrimPaths.Find( MatIndex ) )
+				if (const FUsdPrimPathList* SourcePrimPaths = UserData->MaterialSlotToPrimPaths.Find(MatIndex))
 				{
 					for ( const FString& PrimPath : SourcePrimPaths->PrimPaths )
 					{
@@ -2073,8 +2076,8 @@ bool UnrealToUsd::ConvertMeshComponent( const pxr::UsdStageRefPtr& Stage, const
 	{
 		if ( UStaticMesh* Mesh = StaticMeshComponent->GetStaticMesh() )
 		{
-			UUsdMeshAssetImportData* ImportData = Cast<UUsdMeshAssetImportData>( Mesh->AssetImportData.Get() );
-			if ( !ImportData )
+			const UUsdMeshAssetUserData* UserData = Mesh->GetAssetUserData<UUsdMeshAssetUserData>();
+			if (!UserData)
 			{
 				return false;
 			}
@@ -2101,7 +2104,7 @@ bool UnrealToUsd::ConvertMeshComponent( const pxr::UsdStageRefPtr& Stage, const
 							continue;
 						}
 
-						if ( const FUsdPrimPathList* SourcePrimPaths = ImportData->MaterialSlotToPrimPaths.Find( SectionMatIndex ) )
+						if (const FUsdPrimPathList* SourcePrimPaths = UserData->MaterialSlotToPrimPaths.Find(SectionMatIndex))
 						{
 							for ( const FString& PrimPath : SourcePrimPaths->PrimPaths )
 							{
@@ -2131,8 +2134,9 @@ bool UnrealToUsd::ConvertMeshComponent( const pxr::UsdStageRefPtr& Stage, const
 				return false;
 			}
 
-			UUsdMeshAssetImportData* ImportData = Cast<UUsdMeshAssetImportData>( SkeletalMesh->GetAssetImportData() );
-			if ( !ImportData )
+			const UUsdMeshAssetUserData* UserData =
+				const_cast<USkeletalMesh*>(SkeletalMesh)->GetAssetUserData<UUsdMeshAssetUserData>();
+			if (!UserData)
 			{
 				return false;
 			}
@@ -2180,7 +2184,7 @@ bool UnrealToUsd::ConvertMeshComponent( const pxr::UsdStageRefPtr& Stage, const
 							continue;
 						}
 
-						if ( const FUsdPrimPathList* SourcePrimPaths = ImportData->MaterialSlotToPrimPaths.Find( SectionMatIndex ) )
+						if (const FUsdPrimPathList* SourcePrimPaths = UserData->MaterialSlotToPrimPaths.Find(SectionMatIndex))
 						{
 							for ( const FString& PrimPath : SourcePrimPaths->PrimPaths )
 							{
diff --git a/Engine/Plugins/Importers/USDImporter/Source/USDUtilities/Private/USDShadeConversion.cpp b/Engine/Plugins/Importers/USDImporter/Source/USDUtilities/Private/USDShadeConversion.cpp
--- a/Engine/Plugins/Importers/USDImporter/Source/USDUtilities/Private/USDShadeConversion.cpp
+++ b/Engine/Plugins/Importers/USDImporter/Source/USDUtilities/Private/USDShadeConversion.cpp
@@ -1,13 +1,12 @@
 // Copyright Epic Games, Inc. All Rights Reserved.
 
 #include "USDShadeConversion.h"
-#include "TextureResource.h"
 
 #if USE_USD_SDK
 
 #include "USDAssetCache.h"
 #include "USDAssetCache2.h"
-#include "USDAssetImportData.h"
+#include "USDAssetUserData.h"
 #include "USDAttributeUtils.h"
 #include "USDClassesModule.h"
 #include "USDConversionUtils.h"
@@ -39,6 +38,7 @@
 #include "Modules/ModuleManager.h"
 #include "RenderUtils.h"
 #include "RHI.h"
+#include "TextureResource.h"
 
 #if WITH_EDITOR
 	#include "Factories/TextureFactory.h"
@@ -595,15 +595,13 @@ namespace UE
 							{
 								// Give the same prim path to the texture, so that it ends up imported right next to the material
 								FString MaterialPrimPath;
-#if WITH_EDITOR
 								if ( Material )
 								{
-									if ( UUsdAssetImportData* ImportData = Cast< UUsdAssetImportData >( Material->AssetImportData ) )
+									if(UUsdAssetUserData* UserData = Material->GetAssetUserData<UUsdAssetUserData>())
 									{
-										MaterialPrimPath = ImportData->PrimPath;
+										MaterialPrimPath = UserData->PrimPath;
 									}
 								}
-#endif // WITH_EDITOR
 
 								Texture = UsdUtils::CreateTexture( FileInput.GetAttr(), MaterialPrimPath, LODGroup, TexturesCache );
 
@@ -1211,10 +1209,10 @@ namespace UE
 
 						if ( Texture )
 						{
-							UUsdAssetImportData* ImportData = NewObject< UUsdAssetImportData >( Texture, TEXT( "USDAssetImportData" ) );
-							ImportData->PrimPath = PrimPath;
-							ImportData->UpdateFilenameOnly( ResolvedTexturePath );
-							Texture->AssetImportData = ImportData;
+							UUsdAssetUserData* UserData = NewObject<UUsdAssetUserData>(Texture, TEXT("USDAssetUserData"));
+							UserData->PrimPath = PrimPath;
+							UserData->FilePath = ResolvedTexturePath;
+							Texture->AddAssetUserData(UserData);
 						}
 					}
 				}
@@ -2048,13 +2046,11 @@ bool UsdToUnreal::ConvertMaterial(
 	// Record which primvars we used on each UV index. This is important as we'll match this with the analogous member
 	// on static/skeletal meshe import data, and create a new material instance with different UV index parameter
 	// values if we need to
-#if WITH_EDITOR
-	if (UUsdMaterialAssetImportData* ImportData = Cast<UUsdMaterialAssetImportData>(Material.AssetImportData.Get()))
+	if (UUsdMaterialAssetUserData* UserData = Material.GetAssetUserData<UUsdMaterialAssetUserData>())
 	{
-		ImportData->PrimvarToUVIndex = PrimvarToUVIndex;
-		ImportData->ParameterToPrimvar = ParameterToPrimvar;
+		UserData->PrimvarToUVIndex = PrimvarToUVIndex;
+		UserData->ParameterToPrimvar = ParameterToPrimvar;
 	}
-#endif // WITH_EDITOR
 
 	return MaterialParameters.Num() > 0;
 }
@@ -2325,10 +2321,10 @@ bool UsdToUnreal::ConvertMaterial(
 	// Record which primvars we used on each UV index. This is important as we'll match this with the analogous member
 	// on static/skeletal meshe import data, and create a new material instance with different UV index parameter
 	// values if we need to
-	if (UUsdMaterialAssetImportData* ImportData = Cast<UUsdMaterialAssetImportData>(Material.AssetImportData.Get()))
+	if (UUsdMaterialAssetUserData* UserData = Material.GetAssetUserData<UUsdMaterialAssetUserData>())
 	{
-		ImportData->PrimvarToUVIndex = PrimvarToUVIndex;
-		ImportData->ParameterToPrimvar = ParameterToPrimvar;
+		UserData->PrimvarToUVIndex = PrimvarToUVIndex;
+		UserData->ParameterToPrimvar = ParameterToPrimvar;
 	}
 
 #endif // WITH_EDITOR
@@ -2641,9 +2637,9 @@ void UsdUtils::NotifyIfVirtualTexturesNeeded( UTexture* Texture )
 	}
 
 	FString TexturePath = Texture->GetName();
-	if ( UUsdAssetImportData* AssetImportData = Cast<UUsdAssetImportData>( Texture ) )
+	if (UUsdAssetUserData* UserData = Texture->GetAssetUserData<UUsdAssetUserData>())
 	{
-		TexturePath = AssetImportData->PrimPath;
+		TexturePath = UserData->PrimPath;
 	}
 
 	if ( !UseVirtualTexturing( GMaxRHIFeatureLevel ) )
diff --git a/Engine/Plugins/Importers/USDImporter/Source/USDUtilities/Public/USDConversionUtils.h b/Engine/Plugins/Importers/USDImporter/Source/USDUtilities/Public/USDConversionUtils.h
--- a/Engine/Plugins/Importers/USDImporter/Source/USDUtilities/Public/USDConversionUtils.h
+++ b/Engine/Plugins/Importers/USDImporter/Source/USDUtilities/Public/USDConversionUtils.h
@@ -269,7 +269,10 @@ namespace UsdUtils
 	/** Returns the earliest possible timecode. Use it to always fetch the first frame of an animated attribute */
 	USDUTILITIES_API double GetEarliestTimeCode();
 
+	PRAGMA_DISABLE_DEPRECATION_WARNINGS
+	UE_DEPRECATED(5.3, "No longer used as the USDImporter plugin will now use AssetUserData over AssetImportData for better runtime compatibility")
 	USDUTILITIES_API UUsdAssetImportData* GetAssetImportData( UObject* Asset );
+	PRAGMA_ENABLE_DEPRECATION_WARNINGS
 
 	/** Adds a reference on Prim to the layer at AbsoluteFilePath */
 	USDUTILITIES_API void AddReference( UE::FUsdPrim& Prim, const TCHAR* AbsoluteFilePath, const UE::FSdfPath& TargetPrimPath = {}, double TimeCodeOffset = 0.0, double TimeCodeScale = 1.0 );
diff --git a/Engine/Plugins/Importers/USDImporter/Source/USDUtilities/Public/USDShadeConversion.h b/Engine/Plugins/Importers/USDImporter/Source/USDUtilities/Public/USDShadeConversion.h
--- a/Engine/Plugins/Importers/USDImporter/Source/USDUtilities/Public/USDShadeConversion.h
+++ b/Engine/Plugins/Importers/USDImporter/Source/USDUtilities/Public/USDShadeConversion.h
@@ -178,7 +178,7 @@ namespace UsdUtils
 		TextureAddress AddressY
 	);
 
-	/** Creates a texture from a pxr::SdfAssetPath attribute. PrimPath is optional, and should point to the source shadematerial prim path. It will be placed in its UUsdAssetImportData */
+	/** Creates a texture from a pxr::SdfAssetPath attribute. PrimPath is optional, and should point to the source shadematerial prim path. It will be placed in its UUsdAssetUserData */
 	USDUTILITIES_API UTexture* CreateTexture( const pxr::UsdAttribute& TextureAssetPathAttr, const FString& PrimPath = FString(), TextureGroup LODGroup = TEXTUREGROUP_World, UObject* Outer = GetTransientPackage() );
 
 	/** Checks if this texture needs virtual textures and emits a warning if it is disabled for the project */
diff --git a/Engine/Plugins/Runtime/GeometryCache/Source/GeometryCache/Classes/GeometryCache.h b/Engine/Plugins/Runtime/GeometryCache/Source/GeometryCache/Classes/GeometryCache.h
--- a/Engine/Plugins/Runtime/GeometryCache/Source/GeometryCache/Classes/GeometryCache.h
+++ b/Engine/Plugins/Runtime/GeometryCache/Source/GeometryCache/Classes/GeometryCache.h
@@ -86,6 +86,18 @@ public:
 	/** Get the hash of the meshes data of the GeometryCache */
 	FString GetHash() const;
 
+public:
+	/** Array of user data stored with the asset */
+	UPROPERTY(EditAnywhere, AdvancedDisplay, Instanced, Category = Hidden)
+	TArray<TObjectPtr<UAssetUserData>> AssetUserData;
+
+	//~ Begin IInterface_AssetUserData Interface
+	virtual void AddAssetUserData(UAssetUserData* InUserData) override;
+	virtual void RemoveUserDataOfClass(TSubclassOf<UAssetUserData> InUserDataClass) override;
+	virtual UAssetUserData* GetAssetUserDataOfClass(TSubclassOf<UAssetUserData> InUserDataClass) override;
+	virtual const TArray<UAssetUserData*>* GetAssetUserDataArray() const override;
+	//~ End IInterface_AssetUserData Interface
+
 private:
 	/** A fence which is used to keep track of the rendering thread releasing the geometry cache resources. */
 	FRenderCommandFence ReleaseResourcesFence;
diff --git a/Engine/Plugins/Runtime/GeometryCache/Source/GeometryCache/Private/GeometryCache.cpp b/Engine/Plugins/Runtime/GeometryCache/Source/GeometryCache/Private/GeometryCache.cpp
--- a/Engine/Plugins/Runtime/GeometryCache/Source/GeometryCache/Private/GeometryCache.cpp
+++ b/Engine/Plugins/Runtime/GeometryCache/Source/GeometryCache/Private/GeometryCache.cpp
@@ -2,6 +2,7 @@
 
 #include "GeometryCache.h"
 #include "EditorFramework/AssetImportData.h"
+#include "Engine/AssetUserData.h"
 #include "GeometryCacheMeshData.h"
 #include "GeometryCacheTrack.h"
 #include "UObject/FrameworkObjectVersion.h"
@@ -240,5 +241,49 @@ FString UGeometryCache::GetHash() const
 	return LexToString(TmpHash);
 }
 
+void UGeometryCache::AddAssetUserData(UAssetUserData* InUserData)
+{
+	if (InUserData != nullptr)
+	{
+		UAssetUserData* ExistingData = GetAssetUserDataOfClass(InUserData->GetClass());
+		if (ExistingData != nullptr)
+		{
+			AssetUserData.Remove(ExistingData);
+		}
+		AssetUserData.Add(InUserData);
+	}
+}
+
+UAssetUserData* UGeometryCache::GetAssetUserDataOfClass(TSubclassOf<UAssetUserData> InUserDataClass)
+{
+	for (int32 DataIdx = 0; DataIdx < AssetUserData.Num(); DataIdx++)
+	{
+		UAssetUserData* Datum = AssetUserData[DataIdx];
+		if (Datum != nullptr && Datum->IsA(InUserDataClass))
+		{
+			return Datum;
+		}
+	}
+	return nullptr;
+}
+
+void UGeometryCache::RemoveUserDataOfClass(TSubclassOf<UAssetUserData> InUserDataClass)
+{
+	for (int32 DataIdx = 0; DataIdx < AssetUserData.Num(); DataIdx++)
+	{
+		UAssetUserData* Datum = AssetUserData[DataIdx];
+		if (Datum != nullptr && Datum->IsA(InUserDataClass))
+		{
+			AssetUserData.RemoveAt(DataIdx);
+			return;
+		}
+	}
+}
+
+const TArray<UAssetUserData*>* UGeometryCache::GetAssetUserDataArray() const
+{
+	return &ToRawPtrTArrayUnsafe(AssetUserData);
+}
+
 #undef LOCTEXT_NAMESPACE // "GeometryCache"
 
diff --git a/Engine/Plugins/Runtime/HairStrands/Source/HairStrandsCore/Private/GroomCache.cpp b/Engine/Plugins/Runtime/HairStrands/Source/HairStrandsCore/Private/GroomCache.cpp
--- a/Engine/Plugins/Runtime/HairStrands/Source/HairStrandsCore/Private/GroomCache.cpp
+++ b/Engine/Plugins/Runtime/HairStrands/Source/HairStrandsCore/Private/GroomCache.cpp
@@ -3,6 +3,7 @@
 #include "GroomCache.h"
 
 #include "EditorFramework/AssetImportData.h"
+#include "Engine/AssetUserData.h"
 #include "GroomAsset.h"
 #include "Serialization/MemoryReader.h"
 #include "Serialization/MemoryWriter.h"
@@ -257,6 +258,50 @@ EGroomCacheType UGroomCache::GetType() const
 	return GroomCacheInfo.Type;
 }
 
+void UGroomCache::AddAssetUserData(UAssetUserData* InUserData)
+{
+	if (InUserData != nullptr)
+	{
+		UAssetUserData* ExistingData = GetAssetUserDataOfClass(InUserData->GetClass());
+		if (ExistingData != nullptr)
+		{
+			AssetUserData.Remove(ExistingData);
+		}
+		AssetUserData.Add(InUserData);
+	}
+}
+
+UAssetUserData* UGroomCache::GetAssetUserDataOfClass(TSubclassOf<UAssetUserData> InUserDataClass)
+{
+	for (int32 DataIdx = 0; DataIdx < AssetUserData.Num(); DataIdx++)
+	{
+		UAssetUserData* Datum = AssetUserData[DataIdx];
+		if (Datum != nullptr && Datum->IsA(InUserDataClass))
+		{
+			return Datum;
+		}
+	}
+	return nullptr;
+}
+
+void UGroomCache::RemoveUserDataOfClass(TSubclassOf<UAssetUserData> InUserDataClass)
+{
+	for (int32 DataIdx = 0; DataIdx < AssetUserData.Num(); DataIdx++)
+	{
+		UAssetUserData* Datum = AssetUserData[DataIdx];
+		if (Datum != nullptr && Datum->IsA(InUserDataClass))
+		{
+			AssetUserData.RemoveAt(DataIdx);
+			return;
+		}
+	}
+}
+
+const TArray<UAssetUserData*>* UGroomCache::GetAssetUserDataArray() const
+{
+	return &ToRawPtrTArrayUnsafe(AssetUserData);
+}
+
 void FGroomCacheChunk::Serialize(FArchive& Ar, UObject* Owner, int32 ChunkIndex)
 {
 	Ar << DataSize;
diff --git a/Engine/Plugins/Runtime/HairStrands/Source/HairStrandsCore/Public/GroomCache.h b/Engine/Plugins/Runtime/HairStrands/Source/HairStrandsCore/Public/GroomCache.h
--- a/Engine/Plugins/Runtime/HairStrands/Source/HairStrandsCore/Public/GroomCache.h
+++ b/Engine/Plugins/Runtime/HairStrands/Source/HairStrandsCore/Public/GroomCache.h
@@ -59,6 +59,18 @@ public:
 	TObjectPtr<class UAssetImportData> AssetImportData;	
 #endif
 
+public:
+	/** Array of user data stored with the asset */
+	UPROPERTY(EditAnywhere, AdvancedDisplay, Instanced, Category = Hidden)
+	TArray<TObjectPtr<UAssetUserData>> AssetUserData;
+
+	//~ Begin IInterface_AssetUserData Interface
+	virtual void AddAssetUserData(UAssetUserData* InUserData) override;
+	virtual void RemoveUserDataOfClass(TSubclassOf<UAssetUserData> InUserDataClass) override;
+	virtual UAssetUserData* GetAssetUserDataOfClass(TSubclassOf<UAssetUserData> InUserDataClass) override;
+	virtual const TArray<UAssetUserData*>* GetAssetUserDataArray() const override;
+	//~ End IInterface_AssetUserData Interface
+
 protected:
 	UPROPERTY(VisibleAnywhere, Category = GroomCache)
 	FGroomCacheInfo GroomCacheInfo;
