USD: Add support for UV transform shader nodes.

diff --git a/Engine/Plugins/Importers/USDImporter/Source/USDUtilities/Private/USDShadeConversion.cpp b/Engine/Plugins/Importers/USDImporter/Source/USDUtilities/Private/USDShadeConversion.cpp
--- a/Engine/Plugins/Importers/USDImporter/Source/USDUtilities/Private/USDShadeConversion.cpp
+++ b/Engine/Plugins/Importers/USDImporter/Source/USDUtilities/Private/USDShadeConversion.cpp
@@ -4,6 +4,7 @@
 
 #if USE_USD_SDK
 
+#include "UnrealUSDWrapper.h"
 #include "USDAssetCache.h"
 #include "USDAssetCache2.h"
 #include "USDAssetUserData.h"
@@ -34,6 +35,7 @@
 #include "Materials/MaterialInstanceConstant.h"
 #include "Materials/MaterialInstanceDynamic.h"
 #include "MaterialShared.h"
+#include "Math/TransformCalculus2D.h"
 #include "Misc/FileHelper.h"
 #include "Modules/ModuleManager.h"
 #include "RenderUtils.h"
@@ -59,6 +61,7 @@
 	#include "pxr/usd/sdf/layer.h"
 	#include "pxr/usd/sdf/layerUtils.h"
 	#include "pxr/usd/usdShade/material.h"
+	#include "pxr/usd/usdShade/nodeDefAPI.h"
 #include "USDIncludesEnd.h"
 
 // Required for packaging
@@ -402,48 +405,153 @@ namespace UE
 				return {};
 			}
 
-			/** Receives a UsdUVTexture shader, and returns the name of the primvar that it is using as 'st' */
-			FString GetPrimvarUsedAsST( pxr::UsdShadeConnectableAPI& UsdUVTexture )
+			struct FTextureParameterValue
+			{
+				UTexture* Texture;
+				FString Primvar;
+				FScale2f UVScale;
+				float UVRotation;
+				FVector2f UVTranslation;
+				int32 OutputIndex = 0;
+			};
+			struct FPrimvarReaderParameterValue
+			{
+				std::string PrimvarName;
+				FVector FallbackValue;
+			};
+			using FParameterValue = TVariant<float, FVector, FTextureParameterValue, FPrimvarReaderParameterValue, bool>;
+
+			/**
+			 * Receives a UsdUVTexture shader, and returns the name of the primvar that it is using as 'st',
+			 * plus the USD-space UV transforms that should be applied to that primvar when sampling this texture with
+			 * it.
+			 */
+			void GetSTPrimvarAndTransform(
+				const pxr::UsdShadeConnectableAPI& UsdUVTexture,
+				UE::UsdShadeConversion::Private::FTextureParameterValue& OutTextureValue
+			)
 			{
 				FScopedUsdAllocs Allocs;
 
-				pxr::UsdShadeInput StInput = UsdUVTexture.GetInput( UnrealIdentifiers::St );
-				if ( !StInput )
+				pxr::UsdShadeInput StInput = UsdUVTexture.GetInput(UnrealIdentifiers::St);
+				if (!StInput)
 				{
-					return {};
+					return;
 				}
 
-				pxr::UsdShadeConnectableAPI PrimvarReader;
-				pxr::TfToken PrimvarReaderId;
+				pxr::UsdShadeConnectableAPI Connectable;
+				pxr::TfToken ConnectableOutput;
 				pxr::UsdShadeAttributeType AttributeType;
 
-				// Connected to a PrimvarReader
-				if ( pxr::UsdShadeConnectableAPI::GetConnectedSource( StInput.GetAttr(), &PrimvarReader, &PrimvarReaderId, &AttributeType ) )
+				FTransform2f ConcatenatedTransform;
+				bool bFoundPrimvarReader = false;
+
+				// Traverse through potentially N UV transform nodes
+				while (pxr::UsdShadeConnectableAPI::GetConnectedSource(
+					StInput.GetAttr(),
+					&Connectable,
+					&ConnectableOutput,
+					&AttributeType
+				))
+				{
+					pxr::UsdPrim ConnectedPrim = Connectable.GetPrim();
+					pxr::UsdShadeNodeDefAPI ConnectedShadeNode{ConnectedPrim};
+					if (!ConnectedShadeNode)
+					{
+						// Deadend, should never really happen
+						return;
+					}
+
+					pxr::TfToken ConnectedNodeId;
+					ConnectedShadeNode.GetShaderId(&ConnectedNodeId);
+
+					// UV transform shader node
+					if (ConnectedNodeId == UnrealIdentifiers::UsdTransform2d)
+					{
+						FScale2f Scale;
+						if (pxr::UsdShadeInput ScaleInput = Connectable.GetInput(UnrealIdentifiers::Scale))
+						{
+							pxr::VtValue InputValue;
+							if (ScaleInput.Get(&InputValue) && InputValue.IsHolding<pxr::GfVec2f>())
+							{
+								pxr::GfVec2f VecValue = InputValue.Get<pxr::GfVec2f>();
+								Scale = FScale2f{VecValue[0], VecValue[1]};
+							}
+						}
+
+						float Rotation = 0.0f;
+						if (pxr::UsdShadeInput RotationInput = Connectable.GetInput(UnrealIdentifiers::Rotation))
+						{
+							pxr::VtValue InputValue;
+							if (RotationInput.Get(&InputValue) && InputValue.IsHolding<float>())
+							{
+								Rotation = InputValue.Get<float>();
+							}
+						}
+
+						FVector2f Translation;
+						if (pxr::UsdShadeInput TranslationInput = Connectable.GetInput(UnrealIdentifiers::Translation))
+						{
+							pxr::VtValue InputValue;
+							if (TranslationInput.Get(&InputValue) && InputValue.IsHolding<pxr::GfVec2f>())
+							{
+								pxr::GfVec2f VecValue = InputValue.Get<pxr::GfVec2f>();
+								Translation = FVector2f{VecValue[0], VecValue[1]};
+							}
+						}
+
+						// Concat transform (scale, then rotation, then translation)
+						FTransform2f NewTransform = FTransform2f{Scale}.Concatenate(
+							FTransform2f{FQuat2f{FMath::DegreesToRadians(Rotation)}, Translation}
+						);
+						ConcatenatedTransform = ConcatenatedTransform.Concatenate(NewTransform);
+
+						// Step through to whatever is *this* node's input
+						StInput = Connectable.GetInput(UnrealIdentifiers::In);
+					}
+					// Directly connected to primvar reader
+					else if (ConnectedNodeId == UnrealIdentifiers::UsdPrimvarReader_float2)
+					{
+						bFoundPrimvarReader = true;
+						break;
+					}
+					else
+					{
+						UE_LOG(
+							LogTemp,
+							Warning,
+							TEXT("Unexpected shader node id '%s' when traversing texture node '%s' for primvars!"),
+							*UsdToUnreal::ConvertToken(ConnectedNodeId),
+							*UsdToUnreal::ConvertPath(UsdUVTexture.GetPrim().GetPrimPath())
+						);
+						return;
+					}
+				}
+
+				// Ideally after running through the UsdTransform2d nodes we'd run into a primvar reader node
+				if (bFoundPrimvarReader)
 				{
-					if ( pxr::UsdShadeInput VarnameInput = PrimvarReader.GetInput( UnrealIdentifiers::Varname ) )
+					if (pxr::UsdShadeInput VarnameInput = Connectable.GetInput(UnrealIdentifiers::Varname))
 					{
 						// PrimvarReader may have a string literal with the primvar name, or it may be connected to
 						// e.g. an attribute defined elsewhere
-						return RecursivelySearchForStringValue( VarnameInput );
+						FString Primvar = RecursivelySearchForStringValue(VarnameInput);
+						if (!Primvar.IsEmpty())
+						{
+							// This stuff will end up as arguments for the UsdTransform2d UE material function, which
+							// will do a bunch of conversions inside. We could precompute some of that here, but instead
+							// we're choosing not to, because this means that these values will show up on the material
+							// instance exactly as they show up in the USD shader prims (e.g. if in USD the rotation is
+							// 30 (degrees) we'll see "30" on the material instance too)
+							OutTextureValue.Primvar = Primvar;
+							OutTextureValue.UVScale = ConcatenatedTransform.GetMatrix().GetScale();
+							OutTextureValue.UVRotation = -FMath::RadiansToDegrees(ConcatenatedTransform.GetMatrix().GetRotationAngle());
+							OutTextureValue.UVTranslation = FVector2f{ConcatenatedTransform.GetTranslation()};
+						}
 					}
 				}
-
-				return {};
 			}
 
-			struct FTextureParameterValue
-			{
-				UTexture* Texture;
-				FString Primvar;
-				int32 OutputIndex = 0;
-			};
-			struct FPrimvarReaderParameterValue
-			{
-				std::string PrimvarName;
-				FVector FallbackValue;
-			};
-			using FParameterValue = TVariant<float, FVector, FTextureParameterValue, FPrimvarReaderParameterValue, bool>;
-
 			bool GetTextureParameterValue(
 				pxr::UsdShadeInput& ShadeInput,
 				TextureGroup LODGroup,
@@ -624,7 +732,7 @@ namespace UE
 							{
 								FTextureParameterValue OutTextureValue;
 								OutTextureValue.Texture = Texture;
-								OutTextureValue.Primvar = GetPrimvarUsedAsST(UsdUVTextureSource);
+								GetSTPrimvarAndTransform(UsdUVTextureSource, OutTextureValue);
 
 								// The UsdUVTexture Shader itself prim has multiple standard outputs, but this is not
 								// full swizzle support (check
@@ -1067,6 +1175,23 @@ namespace UE
 					UsdUtils::SetTextureParameterValue( Material, *FString::Printf( TEXT( "%sTexture" ), ParameterName ), TextureValue.Texture );
 					UsdUtils::SetScalarParameterValue( Material, *FString::Printf( TEXT( "Use%sTexture" ), ParameterName ), 1.0f );
 
+					FLinearColor ScaleAndTranslation = FLinearColor{
+						TextureValue.UVScale.GetVector()[0],
+						TextureValue.UVScale.GetVector()[1],
+						TextureValue.UVTranslation[0],
+						TextureValue.UVTranslation[1]};
+					UsdUtils::SetVectorParameterValue(
+						Material,
+						*FString::Printf(TEXT("%sScaleTranslation"), ParameterName),
+						ScaleAndTranslation
+					);
+
+					UsdUtils::SetScalarParameterValue(
+						Material,
+						*FString::Printf(TEXT("%sRotation"), ParameterName),
+						TextureValue.UVRotation
+					);
+
 					if (const int32* FoundIndex = PrimvarToUVIndex.Find(TextureValue.Primvar))
 					{
 						UsdUtils::SetScalarParameterValue(Material, *FString::Printf(TEXT("%sUVIndex"), ParameterName), *FoundIndex);
diff --git a/Engine/Plugins/Importers/USDImporter/Source/UnrealUSDWrapper/Private/UnrealUSDWrapper.cpp b/Engine/Plugins/Importers/USDImporter/Source/UnrealUSDWrapper/Private/UnrealUSDWrapper.cpp
--- a/Engine/Plugins/Importers/USDImporter/Source/UnrealUSDWrapper/Private/UnrealUSDWrapper.cpp
+++ b/Engine/Plugins/Importers/USDImporter/Source/UnrealUSDWrapper/Private/UnrealUSDWrapper.cpp
@@ -119,6 +119,10 @@ namespace UnrealIdentifiers
 	const TfToken Surface = TfToken("surface");
 	const TfToken St = TfToken("st");
 	const TfToken Varname = TfToken("varname");
+	const TfToken Scale = TfToken("scale");
+	const TfToken Rotation = TfToken("rotation");
+	const TfToken Translation = TfToken("translation");
+	const TfToken In = TfToken("in");
 	const TfToken Result = TfToken("result");
 	const TfToken File = TfToken("file");
 	const TfToken WrapT = TfToken( "wrapT" );
@@ -134,6 +138,7 @@ namespace UnrealIdentifiers
 	const TfToken SourceColorSpaceToken = TfToken{ "sourceColorSpace" };
 
 	const TfToken UsdPreviewSurface = TfToken( "UsdPreviewSurface" );
+	const TfToken UsdTransform2d = TfToken( "UsdTransform2d" );
 	const TfToken UsdPrimvarReader_float2 = TfToken( "UsdPrimvarReader_float2" );
 	const TfToken UsdPrimvarReader_float3 = TfToken( "UsdPrimvarReader_float3" );
 	const TfToken UsdUVTexture = TfToken( "UsdUVTexture" );
diff --git a/Engine/Plugins/Importers/USDImporter/Source/UnrealUSDWrapper/Public/UnrealUSDWrapper.h b/Engine/Plugins/Importers/USDImporter/Source/UnrealUSDWrapper/Public/UnrealUSDWrapper.h
--- a/Engine/Plugins/Importers/USDImporter/Source/UnrealUSDWrapper/Public/UnrealUSDWrapper.h
+++ b/Engine/Plugins/Importers/USDImporter/Source/UnrealUSDWrapper/Public/UnrealUSDWrapper.h
@@ -344,6 +344,10 @@ namespace UnrealIdentifiers
 	extern UNREALUSDWRAPPER_API const pxr::TfToken Surface;
 	extern UNREALUSDWRAPPER_API const pxr::TfToken St;
 	extern UNREALUSDWRAPPER_API const pxr::TfToken Varname;
+	extern UNREALUSDWRAPPER_API const pxr::TfToken Scale;
+	extern UNREALUSDWRAPPER_API const pxr::TfToken Rotation;
+	extern UNREALUSDWRAPPER_API const pxr::TfToken Translation;
+	extern UNREALUSDWRAPPER_API const pxr::TfToken In;
 	extern UNREALUSDWRAPPER_API const pxr::TfToken Result;
 	extern UNREALUSDWRAPPER_API const pxr::TfToken File;
 	extern UNREALUSDWRAPPER_API const pxr::TfToken WrapS;
@@ -360,6 +364,7 @@ namespace UnrealIdentifiers
 
 	// Tokens copied from usdImaging, because at the moment it's all we need from it
 	extern UNREALUSDWRAPPER_API const pxr::TfToken UsdPreviewSurface;
+	extern UNREALUSDWRAPPER_API const pxr::TfToken UsdTransform2d;
 	extern UNREALUSDWRAPPER_API const pxr::TfToken UsdPrimvarReader_float2;
 	extern UNREALUSDWRAPPER_API const pxr::TfToken UsdPrimvarReader_float3;
 	extern UNREALUSDWRAPPER_API const pxr::TfToken UsdUVTexture;
