USD: Allow persistent LevelSequence bindings to our generated actors and components.

diff --git a/Engine/Plugins/Importers/USDImporter/Source/USDStage/Private/USDDynamicBindingResolverLibrary.cpp b/Engine/Plugins/Importers/USDImporter/Source/USDStage/Private/USDDynamicBindingResolverLibrary.cpp
new file mode 100644
--- /dev/null
+++ b/Engine/Plugins/Importers/USDImporter/Source/USDStage/Private/USDDynamicBindingResolverLibrary.cpp
@@ -0,0 +1,74 @@
+// Copyright Epic Games, Inc. All Rights Reserved.
+
+#pragma once
+
+#include "USDDynamicBindingResolverLibrary.h"
+
+#include "USDStageActor.h"
+
+#include "EngineUtils.h"
+#include "LevelSequence.h"
+#include "MovieScene.h"
+#include "MovieSceneDynamicBinding.h"
+
+FMovieSceneDynamicBindingResolveResult UUsdDynamicBindingResolverLibrary::ResolveWithStageActor(
+	UObject* WorldContextObject,
+	const FMovieSceneDynamicBindingResolveParams& Params,
+	const FString& StageActorNameFilter,
+	const FString& RootLayerFilter,
+	const FString& PrimPath
+)
+{
+	FMovieSceneDynamicBindingResolveResult Result;
+
+	bool bIsActorBinding = true;
+
+	ULevelSequence* LevelSequence = Cast<ULevelSequence>(Params.RootSequence.Get());
+	const FGuid& BindingGuid = Params.ObjectBindingID;
+	if (LevelSequence && BindingGuid.IsValid())
+	{
+		// Check if we're looking for an actor or component
+		// Note that this is currently not used as we only ever setup our dynamic bindings for actors
+		if (UMovieScene* MovieScene = LevelSequence->GetMovieScene())
+		{
+			if (FMovieScenePossessable* Possessable = MovieScene->FindPossessable(BindingGuid))
+			{
+				if (Possessable->GetParent().IsValid())
+				{
+					bIsActorBinding = false;
+				}
+			}
+		}
+
+		const bool bNoNameFilter = StageActorNameFilter.IsEmpty();
+		const bool bNoRootLayerFilter = RootLayerFilter.IsEmpty();
+
+		UWorld* World = WorldContextObject->GetWorld();
+		for (TActorIterator<AUsdStageActor> ActorIt(World, AUsdStageActor::StaticClass()); ActorIt; ++ActorIt)
+		{
+			if (AUsdStageActor* Actor = *ActorIt)
+			{
+				const bool bNameOK = bNoNameFilter || Actor->GetName() == StageActorNameFilter;
+				const bool bRootLayerOK = bNoRootLayerFilter || Actor->RootLayer.FilePath == RootLayerFilter;
+
+				if (bNameOK && bRootLayerOK)
+				{
+					if (USceneComponent* FoundComponent = Actor->GetGeneratedComponent(PrimPath))
+					{
+						if (bIsActorBinding)
+						{
+							Result.Object = FoundComponent->GetOwner();
+						}
+						else
+						{
+							Result.Object = FoundComponent;
+						}
+						break;
+					}
+				}
+			}
+		}
+	}
+
+	return Result;
+}
diff --git a/Engine/Plugins/Importers/USDImporter/Source/USDStage/Private/USDStageActor.cpp b/Engine/Plugins/Importers/USDImporter/Source/USDStage/Private/USDStageActor.cpp
--- a/Engine/Plugins/Importers/USDImporter/Source/USDStage/Private/USDStageActor.cpp
+++ b/Engine/Plugins/Importers/USDImporter/Source/USDStage/Private/USDStageActor.cpp
@@ -8,6 +8,7 @@
 #include "USDAssetImportData.h"
 #include "USDClassesModule.h"
 #include "USDConversionUtils.h"
+#include "USDDynamicBindingResolverLibrary.h"
 #include "USDErrorUtils.h"
 #include "USDGeomMeshConversion.h"
 #include "USDGeomMeshTranslator.h"
@@ -24,6 +25,7 @@
 #include "USDSchemasModule.h"
 #include "USDSchemaTranslator.h"
 #include "USDSkelRootTranslator.h"
+#include "USDStageModule.h"
 #include "USDTransactor.h"
 #include "USDTypesConversion.h"
 
@@ -67,6 +69,7 @@
 #include "Misc/Paths.h"
 #include "Misc/ScopedSlowTask.h"
 #include "Modules/ModuleManager.h"
+#include "MovieScene.h"
 #include "PhysicsEngine/BodySetup.h"
 #include "Rendering/SkeletalMeshLODImporterData.h"
 #include "Roles/LiveLinkTransformRole.h"
@@ -78,13 +81,23 @@
 #if WITH_EDITOR
 #include "USDClassesEditorModule.h"
 
+#include "BlueprintActionMenuItem.h"
+#include "BlueprintFunctionNodeSpawner.h"
+#include "EdGraphSchema_K2.h"
 #include "Editor.h"
 #include "Editor/TransBuffer.h"
 #include "Editor/UnrealEdEngine.h"
 #include "Engine/Selection.h"
 #include "ILevelSequenceEditorToolkit.h"
+#include "ISequencerModule.h"
+#include "K2Node_FunctionEntry.h"
+#include "K2Node_FunctionResult.h"
+#include "Kismet2/KismetEditorUtilities.h"
 #include "LevelEditor.h"
+#include "MovieSceneDirectorBlueprintUtils.h"
+#include "MovieSceneDynamicBindingUtils.h"
 #include "PropertyEditorModule.h"
+#include "ScopedTransaction.h"
 #include "Subsystems/AssetEditorSubsystem.h"
 #include "UnrealEdGlobals.h"
 #endif // WITH_EDITOR
@@ -554,6 +567,25 @@ struct FUsdStageActorImpl
 			IUsdClassesModule::SendAnalytics(MoveTemp(EventAttributes), TEXT("Open"), bAutomated, ElapsedSeconds, NumberOfFrames, Extension);
 		}
 	}
+
+	// If we have any Sequencer opened with a persistent LevelSequence, this will refresh them so that if their LevelSequences had a binding
+	// to one of our actors that was broken, it can be immediately repaired
+	static void RepairExternalSequencerBindings()
+	{
+#if WITH_EDITOR
+		IUsdStageModule& UsdStageModule = FModuleManager::Get().LoadModuleChecked<IUsdStageModule>(TEXT("UsdStage"));
+		for (const TWeakPtr<ISequencer>& ExistingSequencer : UsdStageModule.GetExistingSequencers())
+		{
+			if (TSharedPtr<ISequencer> PinnedSequencer = ExistingSequencer.Pin())
+			{
+				if (UMovieSceneSequence* FocusedSequence = PinnedSequencer->GetFocusedMovieSceneSequence())
+				{
+					PinnedSequencer->NotifyMovieSceneDataChanged(EMovieSceneDataChangeType::MovieSceneStructureItemAdded);
+				}
+			}
+		}
+#endif // WITH_EDITOR
+	}
 };
 
 /**
@@ -722,6 +754,23 @@ AUsdStageActor::AUsdStageActor()
 			);
 		}
 
+		// Make sure we are bound to the OnActorAddedToSequencer of every sequencer that exists
+		IUsdStageModule& UsdStageModule = FModuleManager::Get().LoadModuleChecked<IUsdStageModule>(TEXT("UsdStage"));
+		for (const TWeakPtr<ISequencer>& ExistingSequencer : UsdStageModule.GetExistingSequencers())
+		{
+			if (TSharedPtr<ISequencer> PinnedSequencer = ExistingSequencer.Pin())
+			{
+				PinnedSequencer->OnActorAddedToSequencer().AddUObject(this, &AUsdStageActor::OnActorAddedToSequencer, ExistingSequencer);
+			}
+		}
+		ISequencerModule& SequencerModule = FModuleManager::Get().LoadModuleChecked<ISequencerModule>(TEXT("Sequencer"));
+		OnSequencerCreatedHandle = SequencerModule.RegisterOnSequencerCreated(FOnSequencerCreated::FDelegate::CreateLambda(
+			[this](TSharedRef<ISequencer> NewSequencer)
+			{
+				NewSequencer->OnActorAddedToSequencer().AddUObject(this, &AUsdStageActor::OnActorAddedToSequencer, NewSequencer.ToWeakPtr());
+			}
+		));
+
 		FCoreUObjectDelegates::OnObjectPropertyChanged.AddUObject(this, &AUsdStageActor::OnObjectPropertyChanged);
 
 		// Also prevent standalone from doing this
@@ -1271,6 +1320,11 @@ void AUsdStageActor::OnUsdObjectsChanged(const UsdUtils::FObjectChangesByPath& I
 		PrimsToAnimate.Reset();
 	}
 
+	if (bHasResync)
+	{
+		FUsdStageActorImpl::RepairExternalSequencerBindings();
+	}
+
 	if (bHasLoadedOrAbandonedAssets && UsdAssetCache)
 	{
 		UsdAssetCache->RefreshStorage();
@@ -2218,6 +2272,8 @@ void AUsdStageActor::LoadUsdStage()
 		UsdAssetCache->RefreshStorage();
 	}
 
+	FUsdStageActorImpl::RepairExternalSequencerBindings();
+
 #if WITH_EDITOR
 	if (GIsEditor && GEditor && !IsGarbageCollecting())
 	{
@@ -2737,6 +2793,22 @@ void AUsdStageActor::BeginDestroy()
 			TransBuffer->OnRedo().Remove(OnRedoHandle);
 		}
 
+		if (ISequencerModule* SequencerModule = FModuleManager::Get().GetModulePtr<ISequencerModule>(TEXT("Sequencer")))
+		{
+			SequencerModule->UnregisterOnSequencerCreated(OnSequencerCreatedHandle);
+
+			if (IUsdStageModule* UsdStageModule = FModuleManager::Get().GetModulePtr<IUsdStageModule>(TEXT("UsdStage")))
+			{
+				for (const TWeakPtr<ISequencer>& ExistingSequencer : UsdStageModule->GetExistingSequencers())
+				{
+					if (TSharedPtr<ISequencer> PinnedSequencer = ExistingSequencer.Pin())
+					{
+						PinnedSequencer->OnActorAddedToSequencer().RemoveAll(this);
+					}
+				}
+			}
+		}
+
 		GEngine->OnLevelActorDeleted().RemoveAll(this);
 		FCoreUObjectDelegates::OnObjectsReplaced.RemoveAll(this);
 	}
@@ -2947,6 +3019,214 @@ void AUsdStageActor::UpdateSpawnedObjectsTransientFlag(bool bTransient)
 	GetRootPrimTwin()->Iterate(UpdateTransient, bRecursive);
 }
 
+void AUsdStageActor::OnActorAddedToSequencer(AActor* NewActor, const FGuid Guid, TWeakPtr<ISequencer> WeakSequencer)
+{
+#if WITH_EDITOR
+	if (!NewActor || !NewActor->HasAnyFlags(RF_Transient))
+	{
+		return;
+	}
+
+	UUsdPrimTwin* UsdPrimTwin = GetRootPrimTwin()->Find(NewActor->GetRootComponent());
+	if (!UsdPrimTwin)
+	{
+		return;
+	}
+	const FString& PrimPath = UsdPrimTwin->PrimPath;
+
+	TSharedPtr<ISequencer> PinnedSequencer = WeakSequencer.Pin();
+	if (!PinnedSequencer)
+	{
+		return;
+	}
+
+	ULevelSequence* Sequence = Cast<ULevelSequence>(PinnedSequencer->GetRootMovieSceneSequence());
+	if (!Sequence)
+	{
+		return;
+	}
+
+	UMovieScene* MovieScene = Sequence->GetMovieScene();
+	if (!MovieScene)
+	{
+		return;
+	}
+
+	FMovieScenePossessable* Possessable = MovieScene->FindPossessable(Guid);
+	if (!Possessable)
+	{
+		return;
+	}
+
+	// Make sure we trigger the creation of the DirectorBlueprint before we add our binding, or else
+	// FMovieSceneDynamicBindingInvoker::ResolveDynamicBinding may fail. Note: The regular code path for dynamic
+	// bindings via the UI ends up calling these exact same lines from
+	// FMovieSceneDirectorBlueprintEndpointCustomization::PopulateQuickBindSubMenu
+	FMovieSceneSequenceEditor* SequenceEditor = FMovieSceneSequenceEditor::Find(Sequence);
+	if (!SequenceEditor)
+	{
+		return;
+	}
+	UBlueprint* DirectorBlueprint = SequenceEditor->GetOrCreateDirectorBlueprint(Sequence);
+	if (!DirectorBlueprint)
+	{
+		return;
+	}
+
+	// The UI action to add the actor to the sequencer should have spawned it's transaction already, but just in case it
+	// hasn't let's do that, as we'll end up modifying UObjects in here
+	FScopedTransaction Transaction(
+		FText::Format(LOCTEXT("ActorAddedToSequencer", "Bound actor '{0}' to the Sequencer"), FText::FromString(NewActor->GetActorLabel()))
+	);
+	Sequence->Modify();
+	MovieScene->Modify();
+
+	// Setup the Dynamic Binding
+	{
+		Possessable->DynamicBinding = FMovieSceneDynamicBinding{};
+		Possessable->DynamicBinding.CompiledFunctionName = NAME_None;
+		Possessable->DynamicBinding.ResolveParamsPinName = TEXT("Params");
+		Possessable->DynamicBinding.Function = UUsdDynamicBindingResolverLibrary::StaticClass()->FindFunctionByName(
+			GET_FUNCTION_NAME_CHECKED(UUsdDynamicBindingResolverLibrary, ResolveWithStageActor)
+		);
+		if (Possessable->DynamicBinding.Function)
+		{
+			Possessable->DynamicBinding.ResolveParamsProperty = Possessable->DynamicBinding.Function->FindPropertyByName(
+				Possessable->DynamicBinding.ResolveParamsPinName
+			);
+		}
+
+		// Store a path to this very actor on the binding, so that it can find us later and ask how to resolve a particular
+		// prim path
+		FMovieSceneDynamicBindingPayloadVariable& ActorPathVariable = Possessable->DynamicBinding.PayloadVariables.FindOrAdd(
+			TEXT("StageActorNameFilter")
+		);
+		ActorPathVariable.Value = GetName();
+
+		FMovieSceneDynamicBindingPayloadVariable& RootLayerVariable = Possessable->DynamicBinding.PayloadVariables.FindOrAdd(
+			TEXT("RootLayerFilter")
+		);
+		RootLayerVariable.Value = FString{};	// No root layer filter by default for more flexibility
+
+		FMovieSceneDynamicBindingPayloadVariable& PrimPathVariable = Possessable->DynamicBinding.PayloadVariables.FindOrAdd(
+			TEXT("PrimPath")
+		);
+		PrimPathVariable.Value = PrimPath;
+	}
+
+	// Setup the Blueprint function ("endpoint") on the Sequence's director Blueprint.
+	// In the future we could also just setup a single endpoint for all dynamic bindings, as they're all identical here. However, even then
+	// I think we should allow the option of having this system emit a separatate endpoint for each binding like it does now, as the user
+	// may want to modify these dynamic binding functions for some reason (maybe add a prim path suffix, or retrieve the prim path from
+	// some other system, etc.).
+	// Reference: FMovieSceneDynamicBindingCustomization::GenerateEndpointDefinition
+	UK2Node_FunctionEntry* NewEndpoint = nullptr;
+	{
+		FMovieSceneDirectorBlueprintEndpointDefinition EndpointDefinition;
+		EndpointDefinition.EndpointType = EMovieSceneDirectorBlueprintEndpointType::Function;
+
+		// Use the same dummy function the regular UI-created Dynamic Bindings use to setup the endpoint signature, in case it changes
+		static const FName SampleResolveBindingFuncName{TEXT("SampleResolveBinding")};
+		UClass* EndpointUtilClass = UMovieSceneDynamicBindingEndpointUtil::StaticClass();
+		EndpointDefinition.EndpointSignature = EndpointUtilClass->FindFunctionByName(SampleResolveBindingFuncName);
+		if(!EndpointDefinition.EndpointSignature)
+		{
+			return;
+		}
+
+		EndpointDefinition.EndpointName = Possessable->GetName() + TEXT("_DynamicBinding");
+
+		DirectorBlueprint->Modify();
+
+		// Creates the new function on the LevelSequence's director blueprint
+		NewEndpoint = FMovieSceneDirectorBlueprintUtils::CreateFunctionEndpoint(DirectorBlueprint, EndpointDefinition);
+		NewEndpoint->MetaData.bCallInEditor = true;
+		if (!NewEndpoint)
+		{
+			return;
+		}
+	}
+
+	// Create a function call node within that endpoint that calls our UFunction
+	// Reference: FMovieSceneDirectorBlueprintEndpointCustomization::HandleQuickBindActionSelected
+	UBlueprintFunctionNodeSpawner* Spawner = UBlueprintFunctionNodeSpawner::Create(Possessable->DynamicBinding.Function);
+	FBlueprintActionMenuItem Action{Spawner};
+	UEdGraphPin* ThenPin = NewEndpoint->FindPin(UEdGraphSchema_K2::PN_Then, EGPD_Output);
+	FVector2D NodePosition(NewEndpoint->NodePosX + 400.f, NewEndpoint->NodePosY + 100.0f);
+	UEdGraphNode* NewNode = Action.PerformAction(NewEndpoint->GetGraph(), ThenPin, NodePosition);
+
+	// Connect the function entry point to our function call node if it has an exec pin
+	UEdGraphPin* NewNodeExecPin = NewNode->FindPin(UEdGraphSchema_K2::PN_Execute, EGPD_Input);
+	if (ThenPin && NewNodeExecPin)
+	{
+		ThenPin->MakeLinkTo(NewNodeExecPin);
+	}
+
+	// Connect the function call result pin to the endpoint's result node
+	TArray<UK2Node_FunctionResult*> ResultNodes;
+	NewNode->GetGraph()->GetNodesOfClass(ResultNodes);
+	if (ResultNodes.Num() > 0)
+	{
+		// If there is a result node, move it past the endpoint call
+		ResultNodes[0]->NodePosX = NodePosition.X + 400.f;
+
+		// Connect the function call node to the endpoint result node if it has an exec pin
+		UEdGraphPin* NewNodeThenPin = NewNode->FindPin(UEdGraphSchema_K2::PN_Execute, EGPD_Input);
+		UEdGraphPin* ResultExecPin = ResultNodes[0]->FindPin(UEdGraphSchema_K2::PN_Execute, EGPD_Input);
+		if (NewNodeThenPin && ResultExecPin)
+		{
+			NewNodeThenPin->MakeLinkTo(ResultExecPin);
+		}
+
+		// If the new node has a return value, and if the endpoint has one too, try to connect them together.
+		UEdGraphPin* OutputPin = ResultNodes[0]->FindPin(UEdGraphSchema_K2::PN_ReturnValue, EGPD_Input);
+		UEdGraphPin* NewNodeReturnValuePin = NewNode->FindPin(UEdGraphSchema_K2::PN_ReturnValue, EGPD_Output);
+		if (OutputPin && NewNodeReturnValuePin)
+		{
+			// Connect the nodes.
+			NewNodeReturnValuePin->MakeLinkTo(OutputPin);
+		}
+	}
+
+	// Connect additional parameter pins from the endpoint to the function call node
+	// Reference: FMovieSceneDirectorBlueprintEndpointCustomization::SetEndpoint
+	TSet<FName> NonPayloadPins;
+	const FString* WorldContextParamName = Possessable->DynamicBinding.Function->FindMetaData(FBlueprintMetadata::MD_WorldContext);
+	if (WorldContextParamName)
+	{
+		NonPayloadPins.Add(FName(*WorldContextParamName));
+	}
+	if (UK2Node_EditablePinBase* EditableNode = Cast<UK2Node_EditablePinBase>(NewEndpoint))
+	{
+		for (UEdGraphPin* PayloadPin : NewNode->Pins)
+		{
+			if (PayloadPin
+				&& PayloadPin->LinkedTo.Num() == 0
+				&& PayloadPin->Direction == EGPD_Input
+				&& PayloadPin->PinType.PinCategory != UEdGraphSchema_K2::PC_Exec
+				&& PayloadPin->PinName != UEdGraphSchema_K2::PN_Self
+				&& !NonPayloadPins.Contains(PayloadPin->PinName
+			))
+			{
+				// Pins for ref parameters for functions default to bIsReference but the payload cannot be by
+				// reference.
+				PayloadPin->PinType.bIsReference = false;
+
+				UEdGraphPin* NewPin = EditableNode->CreateUserDefinedPin(PayloadPin->PinName, PayloadPin->PinType, EGPD_Output);
+				if (NewNode != NewEndpoint && NewPin)
+				{
+					NewPin->MakeLinkTo(PayloadPin);
+				}
+			}
+		}
+	}
+
+	FMovieSceneDynamicBindingUtils::SetEndpoint(MovieScene, &Possessable->DynamicBinding, NewEndpoint);
+	FMovieSceneDynamicBindingUtils::EnsureBlueprintExtensionCreated(Sequence, DirectorBlueprint);
+	FKismetEditorUtilities::CompileBlueprint(DirectorBlueprint);
+#endif	  // WITH_EDITOR
+}
+
 void AUsdStageActor::OnUsdPrimTwinDestroyed(const UUsdPrimTwin& UsdPrimTwin)
 {
 	PrimsToAnimate.Remove(UsdPrimTwin.PrimPath);
diff --git a/Engine/Plugins/Importers/USDImporter/Source/USDStage/Private/USDStageModule.cpp b/Engine/Plugins/Importers/USDImporter/Source/USDStage/Private/USDStageModule.cpp
--- a/Engine/Plugins/Importers/USDImporter/Source/USDStage/Private/USDStageModule.cpp
+++ b/Engine/Plugins/Importers/USDImporter/Source/USDStage/Private/USDStageModule.cpp
@@ -12,6 +12,7 @@
 #include "Interfaces/IPluginManager.h"
 
 #if WITH_EDITOR
+#include "ISequencerModule.h"
 #include "PropertyEditorModule.h"
 #include "Settings/ProjectPackagingSettings.h"
 #endif // WITH_EDITOR
@@ -26,6 +27,27 @@ public:
 
 		FPropertyEditorModule& PropertyModule = FModuleManager::LoadModuleChecked<FPropertyEditorModule>( TEXT( "PropertyEditor" ) );
 		PropertyModule.RegisterCustomClassLayout( TEXT( "UsdStageActor" ), FOnGetDetailCustomizationInstance::CreateStatic( &FUsdStageActorCustomization::MakeInstance ) );
+
+		Sequencers.Reset();
+		ISequencerModule& SequencerModule = FModuleManager::Get().LoadModuleChecked<ISequencerModule>("Sequencer");
+		OnSequencerCreatedHandle = SequencerModule.RegisterOnSequencerCreated(FOnSequencerCreated::FDelegate::CreateLambda(
+			[this](TSharedRef<ISequencer> NewSequencer)
+			{
+				// Cleanup stale pointers
+				for (int32 Index = Sequencers.Num() - 1; Index >= 0; --Index)
+				{
+					if (!Sequencers[Index].IsValid())
+					{
+						const int32 Count = 1;
+						const bool bAllowShrinking = false;
+						Sequencers.RemoveAt(Index, Count, bAllowShrinking);
+					}
+				}
+
+				// Assuming the new one is valid given it's a TSharedRef
+				Sequencers.Add(NewSequencer.ToWeakPtr());
+			}
+		));
 #endif // WITH_EDITOR
 	}
 
@@ -34,9 +56,23 @@ public:
 #if WITH_EDITOR
 		FPropertyEditorModule& PropertyModule = FModuleManager::LoadModuleChecked< FPropertyEditorModule >( TEXT( "PropertyEditor" ) );
 		PropertyModule.UnregisterCustomClassLayout( TEXT( "UsdStageActor" ) );
+
+		if(ISequencerModule* SequencerModule = FModuleManager::Get().GetModulePtr<ISequencerModule>("Sequencer"))
+		{
+			SequencerModule->UnregisterOnSequencerCreated(OnSequencerCreatedHandle);
+			Sequencers.Reset();
+		}
 #endif // WITH_EDITOR
 	}
 
+#if WITH_EDITOR
+	// It would have been nice if the Sequencer module could provide this, but we can easily do this here too
+	const TArray<TWeakPtr<ISequencer>>& GetExistingSequencers() const override
+	{
+		return Sequencers;
+	}
+#endif // WITH_EDITOR
+
 	virtual AUsdStageActor& GetUsdStageActor( UWorld* World ) override
 	{
 		if ( AUsdStageActor* UsdStageActor = FindUsdStageActor( World ) )
@@ -61,6 +97,12 @@ public:
 
 		return nullptr;
 	}
+
+private:
+#if WITH_EDITOR
+	TArray<TWeakPtr<ISequencer>> Sequencers;
+	FDelegateHandle OnSequencerCreatedHandle;
+#endif // WITH_EDITOR
 };
 
 IMPLEMENT_MODULE_USD( FUsdStageModule, USDStage );
diff --git a/Engine/Plugins/Importers/USDImporter/Source/USDStage/Public/USDDynamicBindingResolverLibrary.h b/Engine/Plugins/Importers/USDImporter/Source/USDStage/Public/USDDynamicBindingResolverLibrary.h
new file mode 100644
--- /dev/null
+++ b/Engine/Plugins/Importers/USDImporter/Source/USDStage/Public/USDDynamicBindingResolverLibrary.h
@@ -0,0 +1,41 @@
+// Copyright Epic Games, Inc. All Rights Reserved.
+
+#pragma once
+
+#include "CoreTypes.h"
+#include "Kismet/BlueprintFunctionLibrary.h"
+#include "UObject/ObjectMacros.h"
+
+#include "USDDynamicBindingResolverLibrary.generated.h"
+
+struct FMovieSceneDynamicBindingResolveParams;
+
+UCLASS(meta = (SequencerBindingResolverLibrary))
+
+class USDSTAGE_API UUsdDynamicBindingResolverLibrary : public UBlueprintFunctionLibrary
+{
+	GENERATED_BODY()
+
+public:
+	/**
+	 * Resolves a Sequencer Dynamic Binding described on Params, returning the actor or component that the Sequencer
+	 * should bind to
+	 * @param WorldContextObject - Some UObject that lives in the UWorld that we're talking about
+	 * @param Params - The binding to resolve
+	 * @param StageActorNameFilter - The "ID Name"/FName of a AUsdStageActor to restrict our search for actor and components to. Can be
+	 * empty.
+	 * @param RootLayerFilter - The root layer file path to restrict our search for actor and components to. This should match what is shown
+	 * on the Stage Actor's RootLayer property. Can be empty.
+	 * @param PrimPath - The path to the prim that generates the actor or component that this binding should resolve to
+	 * (e.g. '/cube')
+	 * @return The result struct containing the resolved UObject
+	 */
+	UFUNCTION(BlueprintPure, Category = "USD|Dynamic Binding", meta = (WorldContext = "WorldContextObject"))
+	static FMovieSceneDynamicBindingResolveResult ResolveWithStageActor(
+		UObject* WorldContextObject,
+		const FMovieSceneDynamicBindingResolveParams& Params,
+		const FString& StageActorNameFilter,
+		const FString& RootLayerFilter,
+		const FString& PrimPath
+	);
+};
diff --git a/Engine/Plugins/Importers/USDImporter/Source/USDStage/Public/USDStageActor.h b/Engine/Plugins/Importers/USDImporter/Source/USDStage/Public/USDStageActor.h
--- a/Engine/Plugins/Importers/USDImporter/Source/USDStage/Public/USDStageActor.h
+++ b/Engine/Plugins/Importers/USDImporter/Source/USDStage/Public/USDStageActor.h
@@ -14,6 +14,7 @@
 #include "USDStageActor.generated.h"
 
 class FUsdInfoCache;
+class ISequencer;
 class UInterchangeGenericMaterialPipeline;
 class ULevelSequence;
 class UUsdAssetCache;
@@ -297,6 +298,8 @@ protected:
 
 	void UpdateSpawnedObjectsTransientFlag(bool bTransient);
 
+	void OnActorAddedToSequencer(AActor* NewActor, const FGuid Guid, TWeakPtr<ISequencer> WeakSequencer);
+
 #if WITH_EDITOR
 	void OnBeginPIE(bool bIsSimulating);
 	void OnPostPIEStarted(bool bIsSimulating);
@@ -380,6 +383,7 @@ protected:
 	bool bIsUndoRedoing;
 
 	FDelegateHandle OnRedoHandle;
+	FDelegateHandle OnSequencerCreatedHandle;
 
 	FThreadSafeCounter IsBlockedFromUsdNotices;
 
diff --git a/Engine/Plugins/Importers/USDImporter/Source/USDStage/Public/USDStageModule.h b/Engine/Plugins/Importers/USDImporter/Source/USDStage/Public/USDStageModule.h
--- a/Engine/Plugins/Importers/USDImporter/Source/USDStage/Public/USDStageModule.h
+++ b/Engine/Plugins/Importers/USDImporter/Source/USDStage/Public/USDStageModule.h
@@ -2,14 +2,21 @@
 
 #pragma once
 
+#include "Containers/Array.h"
 #include "Modules/ModuleInterface.h"
+#include "Templates/SharedPointerFwd.h"
 
 class AUsdStageActor;
 class UWorld;
+class ISequencer;
 
 class IUsdStageModule : public IModuleInterface
 {
 public:
-	virtual AUsdStageActor& GetUsdStageActor( UWorld* World ) = 0;
-	virtual AUsdStageActor* FindUsdStageActor( UWorld* World ) = 0;
+	virtual AUsdStageActor& GetUsdStageActor(UWorld* World) = 0;
+	virtual AUsdStageActor* FindUsdStageActor(UWorld* World) = 0;
+
+#if WITH_EDITOR
+	virtual const TArray<TWeakPtr<ISequencer>>& GetExistingSequencers() const = 0;
+#endif	  // WITH_EDITOR
 };
diff --git a/Engine/Plugins/Importers/USDImporter/Source/USDStage/USDStage.Build.cs b/Engine/Plugins/Importers/USDImporter/Source/USDStage/USDStage.Build.cs
--- a/Engine/Plugins/Importers/USDImporter/Source/USDStage/USDStage.Build.cs
+++ b/Engine/Plugins/Importers/USDImporter/Source/USDStage/USDStage.Build.cs
@@ -21,7 +21,6 @@ namespace UnrealBuildTool.Rules
 				{
 					"AnimationCore",
 					"CinematicCamera",
-					"RigVM",
 					"ControlRig",
 					"Core",
 					"CoreUObject",
@@ -38,6 +37,7 @@ namespace UnrealBuildTool.Rules
 					"MovieScene",
 					"MovieSceneTracks",
 					"Niagara",	// Needed by GroomComponent.h
+					"RigVM",
 					"Slate",
 					"SlateCore",
 					"StaticMeshDescription",
@@ -52,10 +52,12 @@ namespace UnrealBuildTool.Rules
 				PrivateDependencyModuleNames.AddRange(
 					new string[]
 					{
+						"BlueprintGraph", // For setting up the Sequencer Dynamic Binding blueprint graphs
 						"ControlRigDeveloper",
 						"DeveloperToolSettings",
 						"EditorStyle", // For the font style on the stage actor customization
 						"InputCore", // For keyboard control on the widget in the stage actor customization
+						"Kismet", // For setting up the Sequencer Dynamic Binding blueprint graphs
 						"LevelSequenceEditor",
 						"MovieSceneTools",
 						"PropertyEditor", // For the stage actor's details customization
diff --git a/Engine/Source/Editor/MovieSceneTools/Public/MovieSceneDynamicBindingUtils.h b/Engine/Source/Editor/MovieSceneTools/Public/MovieSceneDynamicBindingUtils.h
--- a/Engine/Source/Editor/MovieSceneTools/Public/MovieSceneDynamicBindingUtils.h
+++ b/Engine/Source/Editor/MovieSceneTools/Public/MovieSceneDynamicBindingUtils.h
@@ -62,7 +62,7 @@ struct MOVIESCENETOOLS_API FMovieSceneDynamicBindingUtils
  * Dummy class, used for easily getting a valid UFunction that helps prepare blueprint function graphs.
  */
 UCLASS()
-class UMovieSceneDynamicBindingEndpointUtil : public UObject
+class MOVIESCENETOOLS_API UMovieSceneDynamicBindingEndpointUtil : public UObject
 {
 	GENERATED_BODY()
 
