Fixed non-unity build: removed prama once in cpp file.

diff --git a/Engine/Plugins/Importers/USDImporter/Source/USDStage/Private/USDDynamicBindingResolverLibrary.cpp b/Engine/Plugins/Importers/USDImporter/Source/USDStage/Private/USDDynamicBindingResolverLibrary.cpp
--- a/Engine/Plugins/Importers/USDImporter/Source/USDStage/Private/USDDynamicBindingResolverLibrary.cpp
+++ b/Engine/Plugins/Importers/USDImporter/Source/USDStage/Private/USDDynamicBindingResolverLibrary.cpp
@@ -1,7 +1,5 @@
 // Copyright Epic Games, Inc. All Rights Reserved.
 
-#pragma once
-
 #include "USDDynamicBindingResolverLibrary.h"
 
 #include "USDStageActor.h"
