USDImporter: add wrapping for UsdVariantSet/UsdVariantSets to enable directly authoring variant sets on a prim

As part of this change, all of FUsdPayloads, FUsdReferences, and FUsdVariantSets are
only constructable through the appropriate Get...() function on FUsdPrim, since those
classes are all only meaningful in the context of a particular prim. This also aligns the
UE-side wrapping more closely with the USD API.

In addition, the instantiations of FScopedUnrealAllocs objects were removed from
all member function implementations. This avoids "forcing" any particular allocator
from being used when constructing objects, allowing the calling code to provide a
context using either of the UE or USD allocators.

Assignment and conversion operators were added for FUsdPayloads and FUsdReferences
as well, which were missing from their initial introduction.

diff --git a/Engine/Plugins/Importers/USDImporter/Source/UnrealUSDWrapper/Private/UsdWrappers/UsdPayloads.cpp b/Engine/Plugins/Importers/USDImporter/Source/UnrealUSDWrapper/Private/UsdWrappers/UsdPayloads.cpp
--- a/Engine/Plugins/Importers/USDImporter/Source/UnrealUSDWrapper/Private/UsdWrappers/UsdPayloads.cpp
+++ b/Engine/Plugins/Importers/USDImporter/Source/UnrealUSDWrapper/Private/UsdWrappers/UsdPayloads.cpp
@@ -40,6 +40,12 @@ namespace UE
 			{
 			}
 
+			pxr::UsdPayloads GetInner()
+			{
+				FScopedUsdAllocs UsdAllocs;
+				return PxrUsdPrim.Get().GetPayloads();
+			}
+
 			// pxr::UsdPayloads has no default constructor, so we store the prim itself.
 			TUsdStore<pxr::UsdPrim> PxrUsdPrim;
 #endif // #if USE_USD_SDK
@@ -54,25 +60,15 @@ namespace UE
 	FUsdPayloads::FUsdPayloads(const FUsdPayloads& Other)
 	{
 #if USE_USD_SDK
-		FScopedUnrealAllocs UnrealAllocs;
 		Impl = MakeUnique<Internal::FUsdPayloadsImpl>(Other.Impl->PxrUsdPrim.Get());
 #endif // #if USE_USD_SDK
 	}
 
 	FUsdPayloads::FUsdPayloads(FUsdPayloads&& Other) = default;
 
-	FUsdPayloads::FUsdPayloads(const FUsdPrim& InPrim)
-	{
-#if USE_USD_SDK
-		FScopedUnrealAllocs UnrealAllocs;
-		Impl = MakeUnique<Internal::FUsdPayloadsImpl>(InPrim);
-#endif // #if USE_USD_SDK
-	}
-
 	FUsdPayloads& FUsdPayloads::operator=(const FUsdPayloads& Other)
 	{
 #if USE_USD_SDK
-		FScopedUnrealAllocs UnrealAllocs;
 		Impl = MakeUnique<Internal::FUsdPayloadsImpl>(Other.Impl->PxrUsdPrim.Get());
 #endif // #if USE_USD_SDK
 
@@ -103,19 +99,39 @@ namespace UE
 #if USE_USD_SDK
 	FUsdPayloads::FUsdPayloads(const pxr::UsdPayloads& InUsdPayloads)
 	{
-		FScopedUnrealAllocs UnrealAllocs;
 		Impl = MakeUnique<Internal::FUsdPayloadsImpl>(InUsdPayloads);
 	}
 
 	FUsdPayloads::FUsdPayloads(pxr::UsdPayloads&& InUsdPayloads)
 	{
-		FScopedUnrealAllocs UnrealAllocs;
 		Impl = MakeUnique<Internal::FUsdPayloadsImpl>(MoveTemp(InUsdPayloads));
 	}
 
+	FUsdPayloads& FUsdPayloads::operator=(const pxr::UsdPayloads& InUsdPayloads)
+	{
+		Impl = MakeUnique<Internal::FUsdPayloadsImpl>(InUsdPayloads);
+		return *this;
+	}
+
+	FUsdPayloads& FUsdPayloads::operator=(pxr::UsdPayloads&& InUsdPayloads)
+	{
+		Impl = MakeUnique<Internal::FUsdPayloadsImpl>(MoveTemp(InUsdPayloads));
+		return *this;
+	}
+
+	FUsdPayloads::operator pxr::UsdPayloads()
+	{
+		return Impl->GetInner();
+	}
+
+	FUsdPayloads::operator const pxr::UsdPayloads() const
+	{
+		return Impl->GetInner();
+	}
+
+	// Private constructor accessible by FUsdPrim.
 	FUsdPayloads::FUsdPayloads(const pxr::UsdPrim& InPrim)
 	{
-		FScopedUnrealAllocs UnrealAllocs;
 		Impl = MakeUnique<Internal::FUsdPayloadsImpl>(InPrim);
 	}
 #endif // #if USE_USD_SDK
diff --git a/Engine/Plugins/Importers/USDImporter/Source/UnrealUSDWrapper/Private/UsdWrappers/UsdPrim.cpp b/Engine/Plugins/Importers/USDImporter/Source/UnrealUSDWrapper/Private/UsdWrappers/UsdPrim.cpp
--- a/Engine/Plugins/Importers/USDImporter/Source/UnrealUSDWrapper/Private/UsdWrappers/UsdPrim.cpp
+++ b/Engine/Plugins/Importers/USDImporter/Source/UnrealUSDWrapper/Private/UsdWrappers/UsdPrim.cpp
@@ -9,6 +9,7 @@
 #include "UsdWrappers/UsdPayloads.h"
 #include "UsdWrappers/UsdReferences.h"
 #include "UsdWrappers/UsdStage.h"
+#include "UsdWrappers/UsdVariantSets.h"
 
 #if USE_USD_SDK
 
@@ -484,6 +485,28 @@ namespace UE
 		return Children;
 	}
 
+	FUsdVariantSets FUsdPrim::GetVariantSets() const
+	{
+#if USE_USD_SDK
+		return FUsdVariantSets(Impl->PxrUsdPrim.Get());
+#else
+		return FUsdVariantSets{};
+#endif // #if USE_USD_SDK
+	}
+
+	FUsdVariantSet FUsdPrim::GetVariantSet(const FString& VariantSetName) const
+	{
+#if USE_USD_SDK
+		FScopedUsdAllocs UsdAllocs;
+
+		const std::string UsdVariantSetName(TCHAR_TO_ANSI(*VariantSetName));
+
+		return FUsdVariantSet(Impl->PxrUsdPrim.Get(), UsdVariantSetName);
+#else
+		return FUsdVariantSet{};
+#endif // #if USE_USD_SDK
+	}
+
 	bool FUsdPrim::HasVariantSets() const
 	{
 #if USE_USD_SDK
diff --git a/Engine/Plugins/Importers/USDImporter/Source/UnrealUSDWrapper/Private/UsdWrappers/UsdReferences.cpp b/Engine/Plugins/Importers/USDImporter/Source/UnrealUSDWrapper/Private/UsdWrappers/UsdReferences.cpp
--- a/Engine/Plugins/Importers/USDImporter/Source/UnrealUSDWrapper/Private/UsdWrappers/UsdReferences.cpp
+++ b/Engine/Plugins/Importers/USDImporter/Source/UnrealUSDWrapper/Private/UsdWrappers/UsdReferences.cpp
@@ -40,6 +40,12 @@ namespace UE
 			{
 			}
 
+			pxr::UsdReferences GetInner()
+			{
+				FScopedUsdAllocs UsdAllocs;
+				return PxrUsdPrim.Get().GetReferences();
+			}
+
 			// pxr::UsdReferences has no default constructor, so we store the prim itself.
 			TUsdStore<pxr::UsdPrim> PxrUsdPrim;
 #endif // #if USE_USD_SDK
@@ -54,25 +60,15 @@ namespace UE
 	FUsdReferences::FUsdReferences(const FUsdReferences& Other)
 	{
 #if USE_USD_SDK
-		FScopedUnrealAllocs UnrealAllocs;
 		Impl = MakeUnique<Internal::FUsdReferencesImpl>(Other.Impl->PxrUsdPrim.Get());
 #endif // #if USE_USD_SDK
 	}
 
 	FUsdReferences::FUsdReferences(FUsdReferences&& Other) = default;
 
-	FUsdReferences::FUsdReferences(const FUsdPrim& InPrim)
-	{
-#if USE_USD_SDK
-		FScopedUnrealAllocs UnrealAllocs;
-		Impl = MakeUnique<Internal::FUsdReferencesImpl>(InPrim);
-#endif // #if USE_USD_SDK
-	}
-
 	FUsdReferences& FUsdReferences::operator=(const FUsdReferences& Other)
 	{
 #if USE_USD_SDK
-		FScopedUnrealAllocs UnrealAllocs;
 		Impl = MakeUnique<Internal::FUsdReferencesImpl>(Other.Impl->PxrUsdPrim.Get());
 #endif // #if USE_USD_SDK
 
@@ -103,19 +99,39 @@ namespace UE
 #if USE_USD_SDK
 	FUsdReferences::FUsdReferences(const pxr::UsdReferences& InUsdReferences)
 	{
-		FScopedUnrealAllocs UnrealAllocs;
 		Impl = MakeUnique<Internal::FUsdReferencesImpl>(InUsdReferences);
 	}
 
 	FUsdReferences::FUsdReferences(pxr::UsdReferences&& InUsdReferences)
 	{
-		FScopedUnrealAllocs UnrealAllocs;
 		Impl = MakeUnique<Internal::FUsdReferencesImpl>(MoveTemp(InUsdReferences));
 	}
 
+	FUsdReferences& FUsdReferences::operator=(const pxr::UsdReferences& InUsdReferences)
+	{
+		Impl = MakeUnique<Internal::FUsdReferencesImpl>(InUsdReferences);
+		return *this;
+	}
+
+	FUsdReferences& FUsdReferences::operator=(pxr::UsdReferences&& InUsdReferences)
+	{
+		Impl = MakeUnique<Internal::FUsdReferencesImpl>(MoveTemp(InUsdReferences));
+		return *this;
+	}
+
+	FUsdReferences::operator pxr::UsdReferences()
+	{
+		return Impl->GetInner();
+	}
+
+	FUsdReferences::operator const pxr::UsdReferences() const
+	{
+		return Impl->GetInner();
+	}
+
+	// Private constructor accessible by FUsdPrim.
 	FUsdReferences::FUsdReferences(const pxr::UsdPrim& InPrim)
 	{
-		FScopedUnrealAllocs UnrealAllocs;
 		Impl = MakeUnique<Internal::FUsdReferencesImpl>(InPrim);
 	}
 #endif // #if USE_USD_SDK
diff --git a/Engine/Plugins/Importers/USDImporter/Source/UnrealUSDWrapper/Private/UsdWrappers/UsdVariantSets.cpp b/Engine/Plugins/Importers/USDImporter/Source/UnrealUSDWrapper/Private/UsdWrappers/UsdVariantSets.cpp
new file mode 100644
--- /dev/null
+++ b/Engine/Plugins/Importers/USDImporter/Source/UnrealUSDWrapper/Private/UsdWrappers/UsdVariantSets.cpp
@@ -0,0 +1,569 @@
+// Copyright Epic Games, Inc. All Rights Reserved.
+
+#include "UsdWrappers/UsdVariantSets.h"
+
+#include "USDMemory.h"
+
+#include "UsdWrappers/UsdPrim.h"
+
+#if USE_USD_SDK
+
+#include "USDIncludesStart.h"
+	#include "pxr/usd/usd/prim.h"
+	#include "pxr/usd/usd/variantSets.h"
+#include "USDIncludesEnd.h"
+
+#include <map>
+#include <utility>
+#include <vector>
+
+#endif // #if USE_USD_SDK
+
+namespace UE
+{
+	namespace Internal
+	{
+		class FUsdVariantSetImpl
+		{
+		public:
+			FUsdVariantSetImpl() = default;
+
+#if USE_USD_SDK
+			explicit FUsdVariantSetImpl(const pxr::UsdVariantSet& InUsdVariantSet)
+				: PxrUsdPrim(InUsdVariantSet.GetPrim())
+				, VariantSetName(InUsdVariantSet.GetName())
+			{
+			}
+
+			explicit FUsdVariantSetImpl(pxr::UsdVariantSet&& InUsdVariantSet)
+				: PxrUsdPrim(InUsdVariantSet.GetPrim())
+				, VariantSetName(InUsdVariantSet.GetName())
+			{
+			}
+
+			explicit FUsdVariantSetImpl(const pxr::UsdPrim& InUsdPrim, const std::string& InVariantSetName)
+				: PxrUsdPrim(InUsdPrim)
+				, VariantSetName(InVariantSetName)
+			{
+			}
+
+			pxr::UsdVariantSet GetInner()
+			{
+				FScopedUsdAllocs UsdAllocs;
+				return PxrUsdPrim.Get().GetVariantSet(VariantSetName);
+			}
+
+			// pxr::UsdVariantSet has no default constructor, so we store the prim
+			// and the variant set name independently and fetch the object on demand.
+			TUsdStore<pxr::UsdPrim> PxrUsdPrim;
+			std::string VariantSetName;
+#endif // #if USE_USD_SDK
+		};
+
+		class FUsdVariantSetsImpl
+		{
+		public:
+			FUsdVariantSetsImpl() = default;
+
+#if USE_USD_SDK
+			explicit FUsdVariantSetsImpl(const pxr::UsdPrim& InUsdPrim)
+				: PxrUsdPrim(InUsdPrim)
+			{
+			}
+
+			pxr::UsdVariantSets GetInner()
+			{
+				FScopedUsdAllocs UsdAllocs;
+				return PxrUsdPrim.Get().GetVariantSets();
+			}
+
+			// pxr::UsdVariantSets has no default constructor, so we store the
+			// prim itself. It also does not provide an accessor for its prim,
+			// so copy and move constructing from a pxr::UsdVariantSets object
+			// are not supported.
+			TUsdStore<pxr::UsdPrim> PxrUsdPrim;
+#endif // #if USE_USD_SDK
+		};
+	}
+
+
+	// FUsdVariantSet
+
+	FUsdVariantSet::FUsdVariantSet()
+	{
+		Impl = MakeUnique<Internal::FUsdVariantSetImpl>();
+	}
+
+	FUsdVariantSet::FUsdVariantSet(const FUsdVariantSet& Other)
+	{
+#if USE_USD_SDK
+		Impl = MakeUnique<Internal::FUsdVariantSetImpl>(
+			Other.Impl->PxrUsdPrim.Get(),
+			Other.Impl->VariantSetName);
+#endif // #if USE_USD_SDK
+	}
+
+	FUsdVariantSet::FUsdVariantSet(FUsdVariantSet&& Other) = default;
+
+	FUsdVariantSet& FUsdVariantSet::operator=(const FUsdVariantSet& Other)
+	{
+#if USE_USD_SDK
+		Impl = MakeUnique<Internal::FUsdVariantSetImpl>(
+			Other.Impl->PxrUsdPrim.Get(),
+			Other.Impl->VariantSetName);
+#endif // #if USE_USD_SDK
+
+		return *this;
+	}
+
+	FUsdVariantSet& FUsdVariantSet::operator=(FUsdVariantSet&& Other)
+	{
+		Impl = MoveTemp(Other.Impl);
+
+		return *this;
+	}
+
+	FUsdVariantSet::~FUsdVariantSet()
+	{
+		Impl.Reset();
+	}
+
+	FUsdVariantSet::operator bool() const
+	{
+		return IsValid();
+	}
+
+#if USE_USD_SDK
+	FUsdVariantSet::FUsdVariantSet(const pxr::UsdVariantSet& InUsdVariantSet)
+	{
+		Impl = MakeUnique<Internal::FUsdVariantSetImpl>(InUsdVariantSet);
+	}
+
+	FUsdVariantSet::FUsdVariantSet(pxr::UsdVariantSet&& InUsdVariantSet)
+	{
+		Impl = MakeUnique<Internal::FUsdVariantSetImpl>(MoveTemp(InUsdVariantSet));
+	}
+
+	FUsdVariantSet& FUsdVariantSet::operator=(const pxr::UsdVariantSet& InUsdVariantSet)
+	{
+		Impl = MakeUnique<Internal::FUsdVariantSetImpl>(InUsdVariantSet);
+		return *this;
+	}
+
+	FUsdVariantSet& FUsdVariantSet::operator=(pxr::UsdVariantSet&& InUsdVariantSet)
+	{
+		Impl = MakeUnique<Internal::FUsdVariantSetImpl>(MoveTemp(InUsdVariantSet));
+		return *this;
+	}
+
+	FUsdVariantSet::operator pxr::UsdVariantSet()
+	{
+		return Impl->GetInner();
+	}
+
+	FUsdVariantSet::operator const pxr::UsdVariantSet() const
+	{
+		return Impl->GetInner();
+	}
+
+	// Private constructor accessible by FUsdPrim and FUsdVariantSets.
+	FUsdVariantSet::FUsdVariantSet(const pxr::UsdPrim& InPrim, const std::string& InVariantSetName)
+	{
+		Impl = MakeUnique<Internal::FUsdVariantSetImpl>(InPrim, InVariantSetName);
+	}
+#endif // #if USE_USD_SDK
+
+	bool FUsdVariantSet::AddVariant(const FString& VariantName, EUsdListPosition Position)
+	{
+#if USE_USD_SDK
+		FScopedUsdAllocs UsdAllocs;
+
+		pxr::UsdVariantSet UsdVariantSet = Impl->GetInner();
+		if (UsdVariantSet.IsValid())
+		{
+			const std::string UsdVariantName(TCHAR_TO_ANSI(*VariantName));
+
+			return UsdVariantSet.AddVariant(
+				UsdVariantName,
+				static_cast<pxr::UsdListPosition>(Position));
+		}
+#endif // #if USE_USD_SDK
+
+		return false;
+	}
+
+	TArray<FString> FUsdVariantSet::GetVariantNames() const
+	{
+		TArray<FString> VariantNames;
+
+#if USE_USD_SDK
+		FScopedUsdAllocs UsdAllocs;
+
+		pxr::UsdVariantSet UsdVariantSet = Impl->GetInner();
+		if (UsdVariantSet.IsValid())
+		{
+			const std::vector<std::string> UsdVariantNames = UsdVariantSet.GetVariantNames();
+			if (UsdVariantNames.empty())
+			{
+				return VariantNames;
+			}
+
+			VariantNames.Reserve(UsdVariantNames.size());
+
+			for (const std::string& UsdVariantName : UsdVariantNames)
+			{
+				VariantNames.Emplace(ANSI_TO_TCHAR(UsdVariantName.c_str()));
+			}
+		}
+#endif // #if USE_USD_SDK
+
+		return VariantNames;
+	}
+
+	bool FUsdVariantSet::HasAuthoredVariant(const FString& VariantName) const
+	{
+#if USE_USD_SDK
+		FScopedUsdAllocs UsdAllocs;
+
+		pxr::UsdVariantSet UsdVariantSet = Impl->GetInner();
+		if (UsdVariantSet.IsValid())
+		{
+			const std::string UsdVariantName(TCHAR_TO_ANSI(*VariantName));
+
+			return UsdVariantSet.HasAuthoredVariant(UsdVariantName);
+		}
+#endif // #if USE_USD_SDK
+
+		return false;
+	}
+
+	FString FUsdVariantSet::GetVariantSelection() const
+	{
+#if USE_USD_SDK
+		FScopedUsdAllocs UsdAllocs;
+
+		pxr::UsdVariantSet UsdVariantSet = Impl->GetInner();
+		if (UsdVariantSet.IsValid())
+		{
+			const std::string UsdVariantSelection = UsdVariantSet.GetVariantSelection();
+
+			return FString(ANSI_TO_TCHAR(UsdVariantSelection.c_str()));
+		}
+#endif // #if USE_USD_SDK
+
+		return FString{};
+	}
+
+	bool FUsdVariantSet::HasAuthoredVariantSelection(FString* Value) const
+	{
+#if USE_USD_SDK
+		FScopedUsdAllocs UsdAllocs;
+
+		pxr::UsdVariantSet UsdVariantSet = Impl->GetInner();
+		if (UsdVariantSet.IsValid())
+		{
+			std::string UsdValue;
+			bool bResult = UsdVariantSet.HasAuthoredVariantSelection(&UsdValue);
+
+			if (Value != nullptr)
+			{
+				*Value = ANSI_TO_TCHAR(UsdValue.c_str());
+			}
+
+			return bResult;
+		}
+#endif // #if USE_USD_SDK
+
+		return false;
+	}
+
+	bool FUsdVariantSet::SetVariantSelection(const FString& VariantName)
+	{
+#if USE_USD_SDK
+		FScopedUsdAllocs UsdAllocs;
+
+		pxr::UsdVariantSet UsdVariantSet = Impl->GetInner();
+		if (UsdVariantSet.IsValid())
+		{
+			const std::string UsdVariantName(TCHAR_TO_ANSI(*VariantName));
+
+			return UsdVariantSet.SetVariantSelection(UsdVariantName);
+		}
+#endif // #if USE_USD_SDK
+
+		return false;
+	}
+
+	bool FUsdVariantSet::ClearVariantSelection()
+	{
+#if USE_USD_SDK
+		FScopedUsdAllocs UsdAllocs;
+
+		pxr::UsdVariantSet UsdVariantSet = Impl->GetInner();
+		if (UsdVariantSet.IsValid())
+		{
+			return UsdVariantSet.ClearVariantSelection();
+		}
+#endif // #if USE_USD_SDK
+
+		return false;
+	}
+
+	bool FUsdVariantSet::BlockVariantSelection()
+	{
+#if USE_USD_SDK
+		FScopedUsdAllocs UsdAllocs;
+
+		pxr::UsdVariantSet UsdVariantSet = Impl->GetInner();
+		if (UsdVariantSet.IsValid())
+		{
+			return UsdVariantSet.BlockVariantSelection();
+		}
+#endif // #if USE_USD_SDK
+
+		return false;
+	}
+
+	FUsdPrim FUsdVariantSet::GetPrim() const
+	{
+#if USE_USD_SDK
+		FScopedUsdAllocs UsdAllocs;
+
+		pxr::UsdVariantSet UsdVariantSet = Impl->GetInner();
+		if (UsdVariantSet.IsValid())
+		{
+			return FUsdPrim(UsdVariantSet.GetPrim());
+		}
+#endif // #if USE_USD_SDK
+
+		return FUsdPrim{};
+	}
+
+	FString FUsdVariantSet::GetName() const
+	{
+#if USE_USD_SDK
+		FScopedUsdAllocs UsdAllocs;
+
+		pxr::UsdVariantSet UsdVariantSet = Impl->GetInner();
+		if (UsdVariantSet.IsValid())
+		{
+			const std::string UsdVariantName = UsdVariantSet.GetName();
+
+			return FString(ANSI_TO_TCHAR(UsdVariantName.c_str()));
+		}
+#endif // #if USE_USD_SDK
+
+		return FString{};
+	}
+
+	bool FUsdVariantSet::IsValid() const
+	{
+#if USE_USD_SDK
+		return Impl->GetInner().IsValid();
+#endif // #if USE_USD_SDK
+
+		return false;
+	}
+
+
+	// FUsdVariantSets
+
+	FUsdVariantSets::FUsdVariantSets()
+	{
+		Impl = MakeUnique<Internal::FUsdVariantSetsImpl>();
+	}
+
+	FUsdVariantSets::FUsdVariantSets(const FUsdVariantSets& Other)
+	{
+#if USE_USD_SDK
+		Impl = MakeUnique<Internal::FUsdVariantSetsImpl>(
+			Other.Impl->PxrUsdPrim.Get());
+#endif // #if USE_USD_SDK
+	}
+
+	FUsdVariantSets::FUsdVariantSets(FUsdVariantSets&& Other) = default;
+
+	FUsdVariantSets& FUsdVariantSets::operator=(const FUsdVariantSets& Other)
+	{
+#if USE_USD_SDK
+		Impl = MakeUnique<Internal::FUsdVariantSetsImpl>(
+			Other.Impl->PxrUsdPrim.Get());
+#endif // #if USE_USD_SDK
+
+		return *this;
+	}
+
+	FUsdVariantSets& FUsdVariantSets::operator=(FUsdVariantSets&& Other)
+	{
+		Impl = MoveTemp(Other.Impl);
+
+		return *this;
+	}
+
+	FUsdVariantSets::~FUsdVariantSets()
+	{
+		Impl.Reset();
+	}
+
+	FUsdVariantSets::operator bool() const
+	{
+#if USE_USD_SDK
+		return (bool)Impl->PxrUsdPrim.Get();
+#else
+		return false;
+#endif // #if USE_USD_SDK
+	}
+
+#if USE_USD_SDK
+	FUsdVariantSets::operator pxr::UsdVariantSets()
+	{
+		return Impl->GetInner();
+	}
+
+	FUsdVariantSets::operator const pxr::UsdVariantSets() const
+	{
+		return Impl->GetInner();
+	}
+
+	// Private constructor accessible by FUsdPrim.
+	FUsdVariantSets::FUsdVariantSets(const pxr::UsdPrim& InPrim)
+	{
+		Impl = MakeUnique<Internal::FUsdVariantSetsImpl>(InPrim);
+	}
+#endif // #if USE_USD_SDK
+
+	FUsdVariantSet FUsdVariantSets::AddVariantSet(const FString& VariantName, EUsdListPosition Position)
+	{
+#if USE_USD_SDK
+		FScopedUsdAllocs UsdAllocs;
+
+		pxr::UsdVariantSets UsdVariantSets = Impl->GetInner();
+
+		const std::string UsdVariantName(TCHAR_TO_ANSI(*VariantName));
+
+		return FUsdVariantSet(
+			UsdVariantSets.AddVariantSet(
+				UsdVariantName,
+				static_cast<pxr::UsdListPosition>(Position)));
+#else
+		return FUsdVariantSet{};
+#endif // #if USE_USD_SDK
+	}
+
+	TArray<FString> FUsdVariantSets::GetNames() const
+	{
+		TArray<FString> VariantSetNames;
+
+#if USE_USD_SDK
+		FScopedUsdAllocs UsdAllocs;
+
+		pxr::UsdVariantSets UsdVariantSets = Impl->GetInner();
+		
+		const std::vector<std::string> UsdVariantSetNames = UsdVariantSets.GetNames();
+		if (UsdVariantSetNames.empty())
+		{
+			return VariantSetNames;
+		}
+
+		VariantSetNames.Reserve(UsdVariantSetNames.size());
+
+		for (const std::string& UsdVariantSetName : UsdVariantSetNames)
+		{
+			VariantSetNames.Emplace(ANSI_TO_TCHAR(UsdVariantSetName.c_str()));
+		}
+#endif // #if USE_USD_SDK
+
+		return VariantSetNames;
+	}
+
+	FUsdVariantSet FUsdVariantSets::GetVariantSet(const FString& VariantSetName) const
+	{
+#if USE_USD_SDK
+		FScopedUsdAllocs UsdAllocs;
+
+		pxr::UsdVariantSets UsdVariantSets = Impl->GetInner();
+
+		const std::string UsdVariantSetName(TCHAR_TO_ANSI(*VariantSetName));
+
+		return FUsdVariantSet(UsdVariantSets.GetVariantSet(UsdVariantSetName));
+#else
+		return FUsdVariantSet{};
+#endif // #if USE_USD_SDK
+	}
+
+	bool FUsdVariantSets::HasVariantSet(const FString& VariantSetName) const
+	{
+#if USE_USD_SDK
+		FScopedUsdAllocs UsdAllocs;
+
+		pxr::UsdVariantSets UsdVariantSets = Impl->GetInner();
+		
+		const std::string UsdVariantSetName(TCHAR_TO_ANSI(*VariantSetName));
+
+		return UsdVariantSets.HasVariantSet(UsdVariantSetName);
+#else
+		return false;
+#endif // #if USE_USD_SDK
+	}
+
+	FString FUsdVariantSets::GetVariantSelection(const FString& VariantSetName) const
+	{
+#if USE_USD_SDK
+		FScopedUsdAllocs UsdAllocs;
+
+		pxr::UsdVariantSets UsdVariantSets = Impl->GetInner();
+
+		const std::string UsdVariantSetName(TCHAR_TO_ANSI(*VariantSetName));
+
+		const std::string UsdVariantSelection = UsdVariantSets.GetVariantSelection(UsdVariantSetName);
+
+		return FString(ANSI_TO_TCHAR(UsdVariantSelection.c_str()));
+#else
+		return FString{};
+#endif // #if USE_USD_SDK
+	}
+
+	bool FUsdVariantSets::SetSelection(const FString& VariantSetName, const FString& VariantName)
+	{
+#if USE_USD_SDK
+		FScopedUsdAllocs UsdAllocs;
+
+		pxr::UsdVariantSets UsdVariantSets = Impl->GetInner();
+
+		const std::string UsdVariantSetName(TCHAR_TO_ANSI(*VariantSetName));
+		const std::string UsdVariantName(TCHAR_TO_ANSI(*VariantName));
+
+		return UsdVariantSets.SetSelection(UsdVariantSetName, UsdVariantName);
+#else
+		return false;
+#endif // #if USE_USD_SDK
+	}
+
+	TMap<FString, FString> FUsdVariantSets::GetAllVariantSelections() const
+	{
+		TMap<FString, FString> VariantSelections;
+
+#if USE_USD_SDK
+		FScopedUsdAllocs UsdAllocs;
+
+		pxr::UsdVariantSets UsdVariantSets = Impl->GetInner();
+
+		const std::map<std::string, std::string> UsdVariantSelections = UsdVariantSets.GetAllVariantSelections();
+		if (UsdVariantSelections.empty())
+		{
+			return VariantSelections;
+		}
+
+		VariantSelections.Reserve(UsdVariantSelections.size());
+
+		for (const std::pair<const std::string, std::string>& UsdVariantSelection : UsdVariantSelections)
+		{
+			VariantSelections.Emplace(
+				ANSI_TO_TCHAR(UsdVariantSelection.first.c_str()),
+				ANSI_TO_TCHAR(UsdVariantSelection.second.c_str()));
+		}
+#endif // #if USE_USD_SDK
+
+		return VariantSelections;
+	}
+}
diff --git a/Engine/Plugins/Importers/USDImporter/Source/UnrealUSDWrapper/Public/UsdWrappers/UsdPayloads.h b/Engine/Plugins/Importers/USDImporter/Source/UnrealUSDWrapper/Public/UsdWrappers/UsdPayloads.h
--- a/Engine/Plugins/Importers/USDImporter/Source/UnrealUSDWrapper/Public/UsdWrappers/UsdPayloads.h
+++ b/Engine/Plugins/Importers/USDImporter/Source/UnrealUSDWrapper/Public/UsdWrappers/UsdPayloads.h
@@ -42,25 +42,31 @@ namespace UE
 		FUsdPayloads(const FUsdPayloads& Other);
 		FUsdPayloads(FUsdPayloads&& Other);
 
-		explicit FUsdPayloads(const FUsdPrim& InPrim);
-
 		FUsdPayloads& operator=(const FUsdPayloads& Other);
 		FUsdPayloads& operator=(FUsdPayloads&& Other);
-		
+
 		~FUsdPayloads();
 
 		explicit operator bool() const;
 
-	public:
 #if USE_USD_SDK
 		// Auto conversion from/to pxr::UsdPayloads
 		explicit FUsdPayloads(const pxr::UsdPayloads& InUsdPayloads);
 		explicit FUsdPayloads(pxr::UsdPayloads&& InUsdPayloads);
+		FUsdPayloads& operator=(const pxr::UsdPayloads& InUsdPayloads);
+		FUsdPayloads& operator=(pxr::UsdPayloads&& InUsdPayloads);
 
+		operator pxr::UsdPayloads();
+		operator const pxr::UsdPayloads() const;
+
+	private:
+		// FUsdPayloads instances are only constructable by instances of
+		// FUsdPrim.
+		// Use FUsdPrim::GetPayloads() to construct an instance.
 		explicit FUsdPayloads(const pxr::UsdPrim& InPrim);
 #endif // #if USE_USD_SDK
 
-	// Wrapped pxr::UsdPayloads functions, refer to the USD SDK documentation
+		// Wrapped pxr::UsdPayloads functions, refer to the USD SDK documentation
 	public:
 		bool AddPayload(
 			const FSdfPayload& Payload,
@@ -92,5 +98,7 @@ namespace UE
 
 	private:
 		TUniquePtr<Internal::FUsdPayloadsImpl> Impl;
+
+		friend class FUsdPrim;
 	};
 }
diff --git a/Engine/Plugins/Importers/USDImporter/Source/UnrealUSDWrapper/Public/UsdWrappers/UsdPrim.h b/Engine/Plugins/Importers/USDImporter/Source/UnrealUSDWrapper/Public/UsdWrappers/UsdPrim.h
--- a/Engine/Plugins/Importers/USDImporter/Source/UnrealUSDWrapper/Public/UsdWrappers/UsdPrim.h
+++ b/Engine/Plugins/Importers/USDImporter/Source/UnrealUSDWrapper/Public/UsdWrappers/UsdPrim.h
@@ -28,6 +28,8 @@ namespace UE
 	class FUsdAttribute;
 	class FUsdPayloads;
 	class FUsdReferences;
+	class FUsdVariantSet;
+	class FUsdVariantSets;
 
 	namespace Internal
 	{
@@ -116,6 +118,8 @@ namespace UE
 		TArray< FUsdPrim > GetChildren() const;
 		TArray< FUsdPrim > GetFilteredChildren( bool bTraverseInstanceProxies ) const;
 
+		FUsdVariantSets GetVariantSets() const;
+		FUsdVariantSet GetVariantSet(const FString& VariantSetName) const;
 		bool HasVariantSets() const;
 
 		FUsdPayloads GetPayloads() const;
diff --git a/Engine/Plugins/Importers/USDImporter/Source/UnrealUSDWrapper/Public/UsdWrappers/UsdReferences.h b/Engine/Plugins/Importers/USDImporter/Source/UnrealUSDWrapper/Public/UsdWrappers/UsdReferences.h
--- a/Engine/Plugins/Importers/USDImporter/Source/UnrealUSDWrapper/Public/UsdWrappers/UsdReferences.h
+++ b/Engine/Plugins/Importers/USDImporter/Source/UnrealUSDWrapper/Public/UsdWrappers/UsdReferences.h
@@ -42,25 +42,31 @@ namespace UE
 		FUsdReferences(const FUsdReferences& Other);
 		FUsdReferences(FUsdReferences&& Other);
 
-		explicit FUsdReferences(const FUsdPrim& InPrim);
-
 		FUsdReferences& operator=(const FUsdReferences& Other);
 		FUsdReferences& operator=(FUsdReferences&& Other);
-		
+
 		~FUsdReferences();
 
 		explicit operator bool() const;
 
-	public:
 #if USE_USD_SDK
 		// Auto conversion from/to pxr::UsdReferences
 		explicit FUsdReferences(const pxr::UsdReferences& InUsdReferences);
 		explicit FUsdReferences(pxr::UsdReferences&& InUsdReferences);
+		FUsdReferences& operator=(const pxr::UsdReferences& InUsdReferences);
+		FUsdReferences& operator=(pxr::UsdReferences&& InUsdReferences);
 
+		operator pxr::UsdReferences();
+		operator const pxr::UsdReferences() const;
+
+	private:
+		// FUsdReferences instances are only constructable by instances of
+		// FUsdPrim.
+		// Use FUsdPrim::GetReferences() to construct an instance.
 		explicit FUsdReferences(const pxr::UsdPrim& InPrim);
 #endif // #if USE_USD_SDK
 
-	// Wrapped pxr::UsdReferences functions, refer to the USD SDK documentation
+		// Wrapped pxr::UsdReferences functions, refer to the USD SDK documentation
 	public:
 		bool AddReference(
 			const FSdfReference& Reference,
@@ -92,5 +98,7 @@ namespace UE
 
 	private:
 		TUniquePtr<Internal::FUsdReferencesImpl> Impl;
+
+		friend class FUsdPrim;
 	};
 }
diff --git a/Engine/Plugins/Importers/USDImporter/Source/UnrealUSDWrapper/Public/UsdWrappers/UsdVariantSets.h b/Engine/Plugins/Importers/USDImporter/Source/UnrealUSDWrapper/Public/UsdWrappers/UsdVariantSets.h
new file mode 100644
--- /dev/null
+++ b/Engine/Plugins/Importers/USDImporter/Source/UnrealUSDWrapper/Public/UsdWrappers/UsdVariantSets.h
@@ -0,0 +1,174 @@
+// Copyright Epic Games, Inc. All Rights Reserved.
+
+#pragma once
+
+#include "Containers/Array.h"
+#include "Containers/Map.h"
+#include "Containers/UnrealString.h"
+#include "Templates/UniquePtr.h"
+
+#include "UnrealUSDWrapper.h"
+
+#if USE_USD_SDK
+
+#include "USDIncludesStart.h"
+	#include "pxr/pxr.h"
+#include "USDIncludesEnd.h"
+
+#include <string>
+
+PXR_NAMESPACE_OPEN_SCOPE
+	class UsdPrim;
+	class UsdVariantSet;
+	class UsdVariantSets;
+PXR_NAMESPACE_CLOSE_SCOPE
+
+#endif // #if USE_USD_SDK
+
+namespace UE
+{
+	class FUsdPrim;
+
+	namespace Internal
+	{
+		class FUsdVariantSetImpl;
+		class FUsdVariantSetsImpl;
+	}
+
+	/**
+	 * Minimal pxr::UsdVariantSet wrapper for Unreal that can be used from no-rtti modules.
+	 */
+	class UNREALUSDWRAPPER_API FUsdVariantSet
+	{
+	public:
+		FUsdVariantSet();
+		FUsdVariantSet(const FUsdVariantSet& Other);
+		FUsdVariantSet(FUsdVariantSet&& Other);
+
+		FUsdVariantSet& operator=(const FUsdVariantSet& Other);
+		FUsdVariantSet& operator=(FUsdVariantSet&& Other);
+
+		~FUsdVariantSet();
+
+		explicit operator bool() const;
+
+#if USE_USD_SDK
+		// Auto conversion from/to pxr::UsdVariantSet
+		explicit FUsdVariantSet(const pxr::UsdVariantSet& InUsdVariantSet);
+		explicit FUsdVariantSet(pxr::UsdVariantSet&& InUsdVariantSet);
+		FUsdVariantSet& operator=(const pxr::UsdVariantSet& InUsdVariantSet);
+		FUsdVariantSet& operator=(pxr::UsdVariantSet&& InUsdVariantSet);
+
+		operator pxr::UsdVariantSet();
+		operator const pxr::UsdVariantSet() const;
+
+	private:
+		// FUsdVariantSet instances are only constructable by instances of
+		// FUsdPrim or FUsdVariantSets.
+		// Use FUsdPrim::GetVariantSet() or FUsdVariantSets::GetVariantSet()
+		// to construct an instance.
+		explicit FUsdVariantSet(const pxr::UsdPrim& InPrim, const std::string& InVariantSetName);
+#endif // #if USE_USD_SDK
+
+		// Wrapped pxr::UsdVariantSet functions, refer to the USD SDK documentation
+	public:
+		bool AddVariant(
+			const FString& VariantName,
+			EUsdListPosition Position = EUsdListPosition::BackOfPrependList);
+
+		TArray<FString> GetVariantNames() const;
+
+		bool HasAuthoredVariant(const FString& VariantName) const;
+
+		FString GetVariantSelection() const;
+
+		bool HasAuthoredVariantSelection(FString* Value = nullptr) const;
+
+		bool SetVariantSelection(const FString& VariantName);
+
+		bool ClearVariantSelection();
+
+		bool BlockVariantSelection();
+
+		// No support currently for accessing the edit target/context for
+		// an FUsdVariantSet.
+
+		FUsdPrim GetPrim() const;
+
+		FString GetName() const;
+
+		bool IsValid() const;
+
+	private:
+		TUniquePtr<Internal::FUsdVariantSetImpl> Impl;
+
+		friend class FUsdPrim;
+		friend class FUsdVariantSets;
+	};
+
+
+	/**
+	 * Minimal pxr::UsdVariantSets wrapper for Unreal that can be used from no-rtti modules.
+	 */
+	class UNREALUSDWRAPPER_API FUsdVariantSets
+	{
+	public:
+		FUsdVariantSets();
+		FUsdVariantSets(const FUsdVariantSets& Other);
+		FUsdVariantSets(FUsdVariantSets&& Other);
+
+		FUsdVariantSets& operator=(const FUsdVariantSets& Other);
+		FUsdVariantSets& operator=(FUsdVariantSets&& Other);
+
+		~FUsdVariantSets();
+
+		explicit operator bool() const;
+
+#if USE_USD_SDK
+		// The pxr::UsdPrim is not accessible from an instance of
+		// pxr::UsdVariantSets, so there are no facilities for
+		// constructing or assigning to an FUsdVariantSets from
+		// the USD native type.
+
+		// Auto conversion to pxr::UsdVariantSets
+		operator pxr::UsdVariantSets();
+		operator const pxr::UsdVariantSets() const;
+
+	private:
+		// FUsdVariantSets instances are only constructable by instances of
+		// FUsdPrim.
+		// Use FUsdPrim::GetVariantSets() to construct an instance.
+		explicit FUsdVariantSets(const pxr::UsdPrim& InPrim);
+#endif // #if USE_USD_SDK
+
+		// Wrapped pxr::UsdVariantSets functions, refer to the USD SDK documentation
+	public:
+		FUsdVariantSet AddVariantSet(
+			const FString& VariantSetName,
+			EUsdListPosition Position = EUsdListPosition::BackOfPrependList);
+
+		TArray<FString> GetNames() const;
+
+		FUsdVariantSet operator[](const FString& VariantSetName) const
+		{
+			return GetVariantSet(VariantSetName);
+		}
+
+		FUsdVariantSet GetVariantSet(const FString& VariantSetName) const;
+
+		bool HasVariantSet(const FString& VariantSetName) const;
+
+		FString GetVariantSelection(const FString& VariantSetName) const;
+
+		bool SetSelection(
+			const FString& VariantSetName,
+			const FString& VariantName);
+
+		TMap<FString, FString> GetAllVariantSelections() const;
+
+	private:
+		TUniquePtr<Internal::FUsdVariantSetsImpl> Impl;
+
+		friend class FUsdPrim;
+	};
+}
