USD: Fixed crash when the GeometryCache translator is trying to process the same GeometryCache twice.

diff --git a/Engine/Plugins/Importers/USDImporter/Source/USDSchemas/Private/USDGeometryCacheTranslator.cpp b/Engine/Plugins/Importers/USDImporter/Source/USDSchemas/Private/USDGeometryCacheTranslator.cpp
--- a/Engine/Plugins/Importers/USDImporter/Source/USDSchemas/Private/USDGeometryCacheTranslator.cpp
+++ b/Engine/Plugins/Importers/USDImporter/Source/USDSchemas/Private/USDGeometryCacheTranslator.cpp
@@ -821,7 +821,7 @@ void FGeometryCacheCreateAssetsTaskChain::SetupTasks()
 			}
 
 			// Continue with the import steps
-			return Context->bIsImporting && GeometryCache;
+			return Context->bIsImporting && GeometryCache && bIsNew;
 		});
 
 	// Fill the GeometryCache tracks with the frame data (Async)
