USD: Fix GeometryCacheTranslator not deferring to GeomXFormableTranslator::CollapsesChildren when spawned for Xforms.

Update test_component_prim_correspondence test for new conventions.

diff --git a/Engine/Plugins/Importers/USDImporter/Source/USDSchemas/Private/USDGeomMeshTranslator.cpp b/Engine/Plugins/Importers/USDImporter/Source/USDSchemas/Private/USDGeomMeshTranslator.cpp
--- a/Engine/Plugins/Importers/USDImporter/Source/USDSchemas/Private/USDGeomMeshTranslator.cpp
+++ b/Engine/Plugins/Importers/USDImporter/Source/USDSchemas/Private/USDGeomMeshTranslator.cpp
@@ -994,6 +994,11 @@ void FUsdGeomMeshTranslator::CreateAssets()
 {
 	TRACE_CPUPROFILER_EVENT_SCOPE(FUsdGeomMeshTranslator::CreateAssets);
 
+	if (!IsMeshPrim())
+	{
+		return Super::CreateAssets();
+	}
+
 	RegisterAuxiliaryPrims();
 
 		TSharedRef< FGeomMeshCreateAssetsTaskChain > AssetsTaskChain = MakeShared< FGeomMeshCreateAssetsTaskChain >(Context, PrimPath);
@@ -1003,6 +1008,11 @@ void FUsdGeomMeshTranslator::CreateAssets()
 
 USceneComponent* FUsdGeomMeshTranslator::CreateComponents()
 {
+	if (!IsMeshPrim())
+	{
+		return Super::CreateComponents();
+	}
+
 	TOptional< TSubclassOf< USceneComponent > > ComponentType;
 
 	USceneComponent* SceneComponent = CreateComponentsEx(ComponentType, {});
@@ -1045,6 +1055,11 @@ USceneComponent* FUsdGeomMeshTranslator::CreateComponents()
 
 void FUsdGeomMeshTranslator::UpdateComponents(USceneComponent* SceneComponent)
 {
+	if (!IsMeshPrim())
+	{
+		return Super::UpdateComponents(SceneComponent);
+	}
+
 	if (SceneComponent)
 	{
 		SceneComponent->Modify();
@@ -1064,6 +1079,11 @@ void FUsdGeomMeshTranslator::UpdateComponents(USceneComponent* SceneComponent)
 
 bool FUsdGeomMeshTranslator::CollapsesChildren(ECollapsingType CollapsingType) const
 {
+	if (!IsMeshPrim())
+	{
+		return Super::CollapsesChildren(CollapsingType);
+	}
+
 	// We can't claim we collapse anything here since we'll just parse the mesh for this prim and that's it,
 	// otherwise the translation context wouldn't spawn translators for our child prims.
 	// Another approach would be to actually recursively collapse our child mesh prims, but that leads to a few
@@ -1076,6 +1096,11 @@ bool FUsdGeomMeshTranslator::CollapsesChildren(ECollapsingType CollapsingType) c
 
 bool FUsdGeomMeshTranslator::CanBeCollapsed(ECollapsingType CollapsingType) const
 {
+	if (!IsMeshPrim())
+	{
+		return Super::CanBeCollapsed(CollapsingType);
+	}
+
 	UE::FUsdPrim Prim = GetPrim();
 
 	// Don't collapse if our final UStaticMesh would have multiple LODs
@@ -1091,6 +1116,11 @@ bool FUsdGeomMeshTranslator::CanBeCollapsed(ECollapsingType CollapsingType) cons
 
 TSet<UE::FSdfPath> FUsdGeomMeshTranslator::CollectAuxiliaryPrims() const
 {
+	if (!IsMeshPrim())
+	{
+		return Super::CollectAuxiliaryPrims();
+	}
+
 	TSet<UE::FSdfPath> Result;
 	{
 		FScopedUsdAllocs UsdAllocs;
@@ -1113,4 +1143,15 @@ TSet<UE::FSdfPath> FUsdGeomMeshTranslator::CollectAuxiliaryPrims() const
 	return Result;
 }
 
+bool FUsdGeomMeshTranslator::IsMeshPrim() const
+{
+	UE::FUsdPrim Prim = GetPrim();
+	if(Prim && Prim.IsA(TEXT("Mesh")))
+	{
+		return true;
+	}
+
+	return false;
+}
+
 #endif // #if USE_USD_SDK
diff --git a/Engine/Plugins/Importers/USDImporter/Source/USDSchemas/Public/USDGeomMeshTranslator.h b/Engine/Plugins/Importers/USDImporter/Source/USDSchemas/Public/USDGeomMeshTranslator.h
--- a/Engine/Plugins/Importers/USDImporter/Source/USDSchemas/Public/USDGeomMeshTranslator.h
+++ b/Engine/Plugins/Importers/USDImporter/Source/USDSchemas/Public/USDGeomMeshTranslator.h
@@ -93,6 +93,11 @@ public:
 	virtual bool CanBeCollapsed( ECollapsingType CollapsingType ) const override;
 
 	virtual TSet<UE::FSdfPath> CollectAuxiliaryPrims() const override;
+
+protected:
+	// Because of how the GroomTranslator derives the GeometryCacheTranslator that derives this, we may end up in situations where this
+	// translator is not invoked on Mesh prims directly, and so we just yield back to the Xformable translator
+	bool IsMeshPrim() const;
 };
 
 #endif // #if USE_USD_SDK
