USD: Fix GeometryCacheTranslator returning some prims can collapse when they can't.

diff --git a/Engine/Plugins/Importers/USDImporter/Source/USDSchemas/Private/USDInfoCache.cpp b/Engine/Plugins/Importers/USDImporter/Source/USDSchemas/Private/USDInfoCache.cpp
--- a/Engine/Plugins/Importers/USDImporter/Source/USDSchemas/Private/USDInfoCache.cpp
+++ b/Engine/Plugins/Importers/USDImporter/Source/USDSchemas/Private/USDInfoCache.cpp
@@ -879,13 +879,15 @@ namespace UE::USDInfoCacheImpl::Private
 			Prims.Emplace(Child);
 		}
 
-		std::atomic<bool> bIsPotentialGeoCacheRoot = true;
+		// TODO: This can probably be optimized further as we're not going to be a geometry cache root if we have
+		// child prims that can't be collapsed, etc.
+		std::atomic<bool> bIsPotentialGeoCacheRoot = false;
 		const int32 MinBatchSize = 1;
 		ParallelFor(TEXT("RecursiveCheckForGeometryCache"), Prims.Num(), MinBatchSize,
-			[&](int32 Index)
+			[&bIsPotentialGeoCacheRoot, &Prims, &Context, &Impl](int32 Index)
 			{
 				bool bResult = RecursiveCheckForGeometryCache(Prims[Index], Context, Impl);
-				bIsPotentialGeoCacheRoot = bIsPotentialGeoCacheRoot && bResult;
+				bIsPotentialGeoCacheRoot = bIsPotentialGeoCacheRoot || bResult;
 			}
 		);
 
