USD: Fix PrimvarToUVIndex mapping not being set for component overrides;

Update automated test for new asset types capable of being persisted;
Use SkelAnimation prim names for AnimSequences now that they're visible in the open stage workflow.

diff --git a/Engine/Plugins/Importers/USDImporter/Source/USDSchemas/Private/MeshTranslationImpl.cpp b/Engine/Plugins/Importers/USDImporter/Source/USDSchemas/Private/MeshTranslationImpl.cpp
--- a/Engine/Plugins/Importers/USDImporter/Source/USDSchemas/Private/MeshTranslationImpl.cpp
+++ b/Engine/Plugins/Importers/USDImporter/Source/USDSchemas/Private/MeshTranslationImpl.cpp
@@ -759,7 +759,6 @@ void MeshTranslationImpl::SetMaterialOverrides(
 
 	TMap<const UsdUtils::FUsdPrimMaterialSlot*, UMaterialInterface*> ResolvedMaterials;
 
-#if WITH_EDITOR
 	UUsdMeshAssetUserData* UserData = nullptr;
 	if (UStaticMeshComponent* StaticMeshComponent = Cast<UStaticMeshComponent>(&MeshComponent))
 	{
@@ -808,7 +807,6 @@ void MeshTranslationImpl::SetMaterialOverrides(
 			Flags
 		);
 	}
-#endif // WITH_EDITOR
 
 	// Compare resolved materials with existing assignments, and create overrides if we need to
 	uint32 StaticMeshSlotIndex = 0;
diff --git a/Engine/Plugins/Importers/USDImporter/Source/USDSchemas/Private/USDGeomMeshTranslator.cpp b/Engine/Plugins/Importers/USDImporter/Source/USDSchemas/Private/USDGeomMeshTranslator.cpp
--- a/Engine/Plugins/Importers/USDImporter/Source/USDSchemas/Private/USDGeomMeshTranslator.cpp
+++ b/Engine/Plugins/Importers/USDImporter/Source/USDSchemas/Private/USDGeomMeshTranslator.cpp
@@ -817,6 +817,7 @@ void FBuildStaticMeshTaskChain::SetupTasks()
 				{
 					UUsdMeshAssetUserData* UserData = NewObject<UUsdMeshAssetUserData>(StaticMesh, TEXT("UUSDAssetUserData"));
 					UserData->PrimPath = MeshName;
+					UserData->PrimvarToUVIndex = LODIndexToMaterialInfo[0].PrimvarToUVIndex;  // We use the same primvar mapping for all LODs
 					StaticMesh->AddAssetUserData(UserData);
 				}
 
diff --git a/Engine/Plugins/Importers/USDImporter/Source/USDSchemas/Private/USDGeometryCacheTranslator.cpp b/Engine/Plugins/Importers/USDImporter/Source/USDSchemas/Private/USDGeometryCacheTranslator.cpp
--- a/Engine/Plugins/Importers/USDImporter/Source/USDSchemas/Private/USDGeometryCacheTranslator.cpp
+++ b/Engine/Plugins/Importers/USDImporter/Source/USDSchemas/Private/USDGeometryCacheTranslator.cpp
@@ -798,6 +798,7 @@ void FGeometryCacheCreateAssetsTaskChain::SetupTasks()
 			{
 				UUsdMeshAssetUserData* UserData = NewObject<UUsdMeshAssetUserData>(GeometryCache.Get(), TEXT("UUSDAssetUserData"));
 				UserData->PrimPath = PrimPathString;
+				UserData->PrimvarToUVIndex = LODIndexToMaterialInfo[0].PrimvarToUVIndex;  // We use the same primvar mapping for all LODs
 				GeometryCache->AddAssetUserData(UserData);
 
 				// Only the original creator of the prim at creation time gets to set the material assignments
diff --git a/Engine/Plugins/Importers/USDImporter/Source/USDSchemas/Private/USDSkelRootTranslator.cpp b/Engine/Plugins/Importers/USDImporter/Source/USDSchemas/Private/USDSkelRootTranslator.cpp
--- a/Engine/Plugins/Importers/USDImporter/Source/USDSchemas/Private/USDSkelRootTranslator.cpp
+++ b/Engine/Plugins/Importers/USDImporter/Source/USDSchemas/Private/USDSkelRootTranslator.cpp
@@ -1332,6 +1332,7 @@ namespace UsdSkelRootTranslatorImpl
 					{
 						UUsdMeshAssetUserData* UserData = NewObject<UUsdMeshAssetUserData>(SkeletalMesh, TEXT("USDAssetUserData"));
 						UserData->PrimPath = SkelRootPath;
+						UserData->PrimvarToUVIndex = LODIndexToMaterialInfo[0].PrimvarToUVIndex;  // We use the same primvar mapping for all LODs
 						SkeletalMesh->AddAssetUserData(UserData);
 
 						const bool bMaterialsHaveChanged = UsdSkelRootTranslatorImpl::ProcessMaterials(
@@ -1523,15 +1524,19 @@ namespace UsdSkelRootTranslatorImpl
 						{
 							FScopedUnrealAllocs UEAllocs;
 
+							FName AnimSequenceName = MakeUniqueObjectName(
+								GetTransientPackage(),
+								UAnimSequence::StaticClass(),
+								*IUsdClassesModule::SanitizeObjectName(UsdToUnreal::ConvertToken(SkelAnimationPrim.GetName()))
+							);
+
 							// The UAnimSequence can't be created with the RF_Transactional flag, or else it will be serialized without
 							// Bone/CurveCompressionSettings. Undoing that transaction would call UAnimSequence::Serialize with nullptr values for both, which crashes.
 							// Besides, this particular asset type is only ever created when we import to content folder assets (so never for realtime), and
 							// in that case we don't need it to be transactional anyway
-							// TODO: I think this changed and now we generate AnimSequences when opening the stage too,
-							// so this may lead to issues
 							AnimSequence = NewObject<UAnimSequence>(
 								GetTransientPackage(),
-								NAME_None,
+								AnimSequenceName,
 								(Context->ObjectFlags & ~EObjectFlags::RF_Transactional) | EObjectFlags::RF_Transient
 							);
 							AnimSequence->SetSkeleton( SkeletalMesh->GetSkeleton() );
