USD: Fixed possible deadlocks when importing GeometryCaches.

diff --git a/Engine/Plugins/Importers/USDImporter/Source/USDSchemas/Private/USDGeometryCacheTranslator.cpp b/Engine/Plugins/Importers/USDImporter/Source/USDSchemas/Private/USDGeometryCacheTranslator.cpp
--- a/Engine/Plugins/Importers/USDImporter/Source/USDSchemas/Private/USDGeometryCacheTranslator.cpp
+++ b/Engine/Plugins/Importers/USDImporter/Source/USDSchemas/Private/USDGeometryCacheTranslator.cpp
@@ -460,7 +460,7 @@ namespace UsdGeometryCacheTranslatorImpl
 			{
 				if (NumLoops % 2 == 0)
 				{
-					NumFrameThreads = FMath::Max(--NumFrameThreads, 1);
+					NumFrameThreads = FMath::Max(FMath::RoundToInt(static_cast<float>(NumFrameThreads) * 0.8f), 1);
 				}
 				else
 				{
@@ -494,7 +494,7 @@ namespace UsdGeometryCacheTranslatorImpl
 
 					// Parallel frame read: frame data can be read concurrently but have to be processed in order for AddMeshSample
 					ParallelFor(NumFrameThreads, [&Args, &MeshPrim, &FrameWrittenEvent, &WriteFrameIndex, &Mutex, NumFrameThreads, MaterialOffset, StreamableTrack, bConstantTopology](int32 FrameThreadIndex)
-					{
+							{
 						int32 FrameIndex = Args.StartFrame + FrameThreadIndex;
 
 						while (FrameIndex < Args.EndFrame)
@@ -825,8 +825,10 @@ void FGeometryCacheCreateAssetsTaskChain::SetupTasks()
 			return Context->bIsImporting && GeometryCache && bIsNew;
 		});
 
-	// Fill the GeometryCache tracks with the frame data (Async)
-	Then(ESchemaTranslationLaunchPolicy::Async,
+	// Fill the GeometryCache tracks with the frame data
+	// It is done Sync to avoid starvation issue because FillGeometryCacheTracks is highly parallelized based on the number of meshes and frames to read
+	// Filling GeometryCaches in parallel could cause deadlocks
+	Then(ESchemaTranslationLaunchPolicy::Sync,
 		[this]() -> bool
 		{
 			UsdGeometryCacheTranslatorImpl::FillGeometryCacheTracks(PrimPath.GetString(), MeshPrimPaths, MaterialOffsets, Context, GeometryCache.Get());
