USD: Work around how FKControlRigs only spawn morph target controls for morph targets with curve metadata on the skeletons.

diff --git a/Engine/Plugins/Importers/USDImporter/Source/USDUtilities/Private/USDSkeletalDataConversion.cpp b/Engine/Plugins/Importers/USDImporter/Source/USDUtilities/Private/USDSkeletalDataConversion.cpp
--- a/Engine/Plugins/Importers/USDImporter/Source/USDUtilities/Private/USDSkeletalDataConversion.cpp
+++ b/Engine/Plugins/Importers/USDImporter/Source/USDUtilities/Private/USDSkeletalDataConversion.cpp
@@ -82,7 +82,9 @@
 
 #define LOCTEXT_NAMESPACE "UsdSkeletalDataConversion"
 
-static bool bAddCurveMetadataToSkeleton = false;
+// Default to true because as of Apr 2023 baking animations with FKControlRigs will only consider morph targets if the curve metadata is on the
+// skeleton for some reason
+static bool bAddCurveMetadataToSkeleton = true;
 static FAutoConsoleVariableRef CVarAddCurveMetadataToSkeleton(
 	TEXT("USD.AddCurveMetadataToSkeleton"),
 	bAddCurveMetadataToSkeleton,
