USD: Fix geometry cache translator permanently holding stages opened.

diff --git a/Engine/Plugins/Importers/USDImporter/Source/USDSchemas/Private/USDGeometryCacheTranslator.cpp b/Engine/Plugins/Importers/USDImporter/Source/USDSchemas/Private/USDGeometryCacheTranslator.cpp
--- a/Engine/Plugins/Importers/USDImporter/Source/USDSchemas/Private/USDGeometryCacheTranslator.cpp
+++ b/Engine/Plugins/Importers/USDImporter/Source/USDSchemas/Private/USDGeometryCacheTranslator.cpp
@@ -190,7 +190,7 @@ namespace UsdGeometryCacheTranslatorImpl
 			, RootPrim(InRootPrim)
 		{}
 
-		UE::FUsdStage Stage;
+		UE::FUsdStageWeak Stage;
 		UE::FUsdPrim RootPrim;
 		UsdToUnreal::FUsdMeshConversionOptions Options;
 		int32 StartFrame = 0;
