USD: Fix inconsistent landscape bakes when exporting levels with sublevels.

diff --git a/Engine/Plugins/Importers/USDImporter/Source/USDExporter/Private/USDConversionBlueprintLibrary.cpp b/Engine/Plugins/Importers/USDImporter/Source/USDExporter/Private/USDConversionBlueprintLibrary.cpp
--- a/Engine/Plugins/Importers/USDImporter/Source/USDExporter/Private/USDConversionBlueprintLibrary.cpp
+++ b/Engine/Plugins/Importers/USDImporter/Source/USDExporter/Private/USDConversionBlueprintLibrary.cpp
@@ -18,11 +18,14 @@
 
 #include "AnalyticsBlueprintLibrary.h"
 #include "AnalyticsEventAttribute.h"
+#include "AssetCompilingManager.h"
+#include "ContentStreaming.h"
 #include "CoreMinimal.h"
 #include "Editor.h"
 #include "Engine/Engine.h"
 #include "Engine/Level.h"
 #include "Engine/LevelStreaming.h"
+#include "Engine/Texture.h"
 #include "Engine/World.h"
 #include "InstancedFoliageActor.h"
 #include "ISequencer.h"
@@ -35,6 +38,26 @@ namespace UE
 	{
 		namespace Private
 		{
+			void WaitForAllAsyncAndSteamingTasks(UWorld* World)
+			{
+				FlushAsyncLoading();
+
+				if (World)
+				{
+					World->BlockTillLevelStreamingCompleted();
+
+					if (!FPlatformProperties::RequiresCookedData())
+					{
+						UMaterialInterface::SubmitRemainingJobsForWorld(World);
+						FAssetCompilingManager::Get().FinishAllCompilation();
+					}
+				}
+
+				UTexture::ForceUpdateTextureStreaming();
+
+				IStreamingManager::Get().StreamAllResources(0.0f);
+			}
+
 			void StreamInLevels( ULevel* Level, const TSet<FString>& LevelsToIgnore )
 			{
 				UWorld* InnerWorld = Level->GetTypedOuter<UWorld>();
@@ -134,6 +157,8 @@ namespace UE
 				// Synchronously show levels right now
 				InnerWorld->FlushLevelStreaming( EFlushLevelStreamingType::Visibility );
 
+				WaitForAllAsyncAndSteamingTasks(InnerWorld);
+
 				if ( bCreatedContext )
 				{
 					GEngine->DestroyWorldContext( InnerWorld );
@@ -250,6 +275,8 @@ void UUsdConversionBlueprintLibrary::StreamOutLevels( UWorld* OwningWorld, const
 		}
 	}
 
+	UE::UsdConversionBlueprintLibraryImpl::Private::WaitForAllAsyncAndSteamingTasks(OwningWorld);
+
 	if ( bCreatedContext )
 	{
 		GEngine->DestroyWorldContext( OwningWorld );
diff --git a/Engine/Plugins/Importers/USDImporter/Source/USDUtilities/Private/USDShadeConversion.cpp b/Engine/Plugins/Importers/USDImporter/Source/USDUtilities/Private/USDShadeConversion.cpp
--- a/Engine/Plugins/Importers/USDImporter/Source/USDUtilities/Private/USDShadeConversion.cpp
+++ b/Engine/Plugins/Importers/USDImporter/Source/USDUtilities/Private/USDShadeConversion.cpp
@@ -1738,7 +1738,7 @@ namespace UE
 
 					if ( ( NumSamples == 0 || !SampleSize ) && !UserConstantValue )
 					{
-						UE_LOG( LogUsd, Warning, TEXT( "Skipping material property %d as we have no valid data to use." ), Property );
+						UE_LOG( LogUsd, Log, TEXT( "Skipping material property %d as we have no valid data to use." ), Property );
 						continue;
 					}
 
