USD: Fix crash when importing via an automated task with no import options.

diff --git a/Engine/Plugins/Importers/USDImporter/Source/USDStageImporter/Private/USDStageImporter.cpp b/Engine/Plugins/Importers/USDImporter/Source/USDStageImporter/Private/USDStageImporter.cpp
--- a/Engine/Plugins/Importers/USDImporter/Source/USDStageImporter/Private/USDStageImporter.cpp
+++ b/Engine/Plugins/Importers/USDImporter/Source/USDStageImporter/Private/USDStageImporter.cpp
@@ -1659,6 +1659,14 @@ void UUsdStageImporter::ImportFromFile(FUsdStageImportContext& ImportContext)
 
 	double StartTime = FPlatformTime::Cycles64();
 
+	// Load some default options in case we don't have an options object, so that we don't have to check for this every time.
+	// We could also GetMutableDefault<UUsdStageImportOptions>() instead, but I think in this case it's expected that these
+	// should be the default value and not the config
+	if (!ImportContext.ImportOptions)
+	{
+		ImportContext.ImportOptions = NewObject<UUsdStageImportOptions>();
+	}
+
 	if ( ImportContext.ImportOptions->PrimsToImport.Num() == 0 )
 	{
 		return;
@@ -1794,6 +1802,11 @@ bool UUsdStageImporter::ReimportSingleAsset(
 #if USE_USD_SDK
 	double StartTime = FPlatformTime::Cycles64();
 
+	if (!ImportContext.ImportOptions)
+	{
+		ImportContext.ImportOptions = NewObject<UUsdStageImportOptions>();
+	}
+
 	// TODO: Maybe change this whole reimporting approach to just taking advantage of a population mask instead?
 	const bool bNeedsMasking =
 		ImportContext.ImportOptions->PrimsToImport != TArray<FString>{ UE::FSdfPath::AbsoluteRootPath().GetString() };
