USD: Fix spawned GeometryCache components never receiving any non-animated transform.

diff --git a/Engine/Plugins/Importers/USDImporter/Source/USDSchemas/Private/USDGeometryCacheTranslator.cpp b/Engine/Plugins/Importers/USDImporter/Source/USDSchemas/Private/USDGeometryCacheTranslator.cpp
--- a/Engine/Plugins/Importers/USDImporter/Source/USDSchemas/Private/USDGeometryCacheTranslator.cpp
+++ b/Engine/Plugins/Importers/USDImporter/Source/USDSchemas/Private/USDGeometryCacheTranslator.cpp
@@ -991,6 +991,20 @@ void FUsdGeometryCacheTranslator::UpdateComponents(USceneComponent* SceneCompone
 			GeometryCacheComponent->TickAtThisTime(TimeCode, true, false, true);
 		}
 
+		// If the prim has a GroomBinding schema, apply the target groom to its associated GroomComponent
+		if (UsdUtils::PrimHasSchema(GetPrim(), UnrealIdentifiers::GroomBindingAPI))
+		{
+			UsdGroomTranslatorUtils::SetGroomFromPrim(GetPrim(), *Context->InfoCache, SceneComponent);
+		}
+
+		// Defer to xformable translator to set our transforms, visibility, etc. but only when opening the stage: This will be baked in for import.
+		// Don't go through FUsdGeomMeshTranslator::UpdateComponents as it will want to create a static mesh if PrimPath is an animated mesh prim
+		// (which is likely, given that we're running this FUsdGeometryCacheTranslator for it)
+		if (!Context->bIsImporting)
+		{
+			FUsdGeomXformableTranslator::UpdateComponents(GeometryCacheComponent);
+		}
+
 		// Note how we should only register if our geometry cache changed: If we did this every time we would
 		// register too early during the process of duplicating into PIE, and that would prevent a future RegisterComponent
 		// call from naturally creating the required render state
@@ -998,12 +1012,6 @@ void FUsdGeometryCacheTranslator::UpdateComponents(USceneComponent* SceneCompone
 		{
 			GeometryCacheComponent->RegisterComponent();
 		}
-
-		// If the prim has a GroomBinding schema, apply the target groom to its associated GroomComponent
-		if (UsdUtils::PrimHasSchema(GetPrim(), UnrealIdentifiers::GroomBindingAPI))
-		{
-			UsdGroomTranslatorUtils::SetGroomFromPrim(GetPrim(), *Context->InfoCache, SceneComponent);
-		}
 	}
 }
 
