USD: Revert exposing the MaterialX pipeline options as those are interchange project settings now.

diff --git a/Engine/Plugins/Importers/USDImporter/Source/USDSchemas/Private/Custom/MaterialXUSDShadeMaterialTranslator.cpp b/Engine/Plugins/Importers/USDImporter/Source/USDSchemas/Private/Custom/MaterialXUSDShadeMaterialTranslator.cpp
--- a/Engine/Plugins/Importers/USDImporter/Source/USDSchemas/Private/Custom/MaterialXUSDShadeMaterialTranslator.cpp
+++ b/Engine/Plugins/Importers/USDImporter/Source/USDSchemas/Private/Custom/MaterialXUSDShadeMaterialTranslator.cpp
@@ -49,70 +49,6 @@ static FAutoConsoleVariableRef CvarUseInterchangeMaterialXTranslator(
 
 FName FMaterialXUsdShadeMaterialTranslator::MaterialXRenderContext = TEXT("mtlx");
 
-namespace UE::USDMaterialXTranslatorInternal::Private
-{
-	void HashMaterialXOptions(UInterchangePipelineBase* MaterialXOptions, FSHA1& SHA1)
-	{
-		// Reference: UInterchangePipelineBase::SaveSettingsInternal
-
-		if (!MaterialXOptions)
-		{
-			return;
-		}
-
-		int32 PortFlags = 0;
-		UClass* Class = MaterialXOptions->GetClass();
-		for (FProperty* Property = Class->PropertyLink; Property; Property = Property->PropertyLinkNext)
-		{
-			FString Key = Property->GetName();
-			const FName PropertyName = Property->GetFName();
-
-			const bool bIsPropertyInherited = Property->GetOwnerClass() != Class;
-			UObject* SuperClassDefaultObject = Class->GetSuperClass()->GetDefaultObject();
-			FObjectProperty* SubObject = CastField<FObjectProperty>(Property);
-			FArrayProperty* Array = CastField<FArrayProperty>(Property);
-			if (Array)
-			{
-				const bool bForce = true;
-				const bool bConst = false;
-
-				FScriptArrayHelper_InContainer ArrayHelper(Array, MaterialXOptions);
-				for (int32 i = 0; i < ArrayHelper.Num(); i++)
-				{
-					FString	Buffer;
-					Array->Inner->ExportTextItem_Direct(Buffer, ArrayHelper.GetRawPtr(i), ArrayHelper.GetRawPtr(i), MaterialXOptions, PortFlags);
-
-					SHA1.UpdateWithString(*Buffer, Buffer.Len());
-				}
-			}
-			else if (UInterchangePipelineBase* SubPipeline = SubObject ? Cast<UInterchangePipelineBase>(SubObject->GetObjectPropertyValue_InContainer(MaterialXOptions)) : nullptr)
-			{
-				// Save the settings if the referenced pipeline is a subobject of ours
-				if (SubPipeline->IsInOuter(MaterialXOptions))
-				{
-					HashMaterialXOptions(SubPipeline, SHA1);
-				}
-			}
-			else
-			{
-				TCHAR TempKey[MAX_SPRINTF] = TEXT("");
-				for (int32 Index = 0; Index < Property->ArrayDim; Index++)
-				{
-					if (Property->ArrayDim != 1)
-					{
-						FCString::Sprintf(TempKey, TEXT("%s[%i]"), *Property->GetName(), Index);
-						Key = TempKey;
-					}
-
-					FString	Value;
-					Property->ExportText_InContainer(Index, Value, MaterialXOptions, MaterialXOptions, MaterialXOptions, PortFlags);
-					SHA1.UpdateWithString(*Value, Value.Len());
-				}
-			}
-		}
-	}
-}
-
 void FMaterialXUsdShadeMaterialTranslator::CreateAssets()
 {
 #if WITH_EDITOR
@@ -209,7 +145,7 @@ void FMaterialXUsdShadeMaterialTranslator::CreateAssets()
 	// This material prim has the mtlx render context, so maybe it is one of the ones generated by usdMtlx.
 	// Let's traverse upwards and try finding a .mtlx file reference in one of our parents.
 
-	TFunction<bool(const pxr::UsdPrim&, FString&)> GetMaterialXFilePath = [this](const pxr::UsdPrim& Prim, FString& OutMaterialXFilePath)
+	TFunction<bool(const pxr::UsdPrim&, FString&)> GetMaterialXFilePath = [](const pxr::UsdPrim& Prim, FString& OutMaterialXFilePath)
 	{
 		pxr::UsdPrimCompositionQuery PrimCompositionQuery = pxr::UsdPrimCompositionQuery::GetDirectReferences(Prim);
 		for (const pxr::UsdPrimCompositionQueryArc& CompositionArc : PrimCompositionQuery.GetCompositionArcs())
@@ -340,12 +276,6 @@ void FMaterialXUsdShadeMaterialTranslator::CreateAssets()
 		}
 	}
 
-	// Hash the options we'll use so that we regenerate assets if those change
-	if (Context->MaterialXOptions)
-	{
-		UE::USDMaterialXTranslatorInternal::Private::HashMaterialXOptions(Context->MaterialXOptions, SHA1);
-	}
-
 	FSHAHash Hash;
 	SHA1.Final();
 	SHA1.GetHash(&Hash.Hash[0]);
@@ -365,11 +295,6 @@ void FMaterialXUsdShadeMaterialTranslator::CreateAssets()
 		FImportAssetParameters InterchangeParameters;
 		InterchangeParameters.bIsAutomated = true;
 
-		if (Context->MaterialXOptions)
-		{
-			InterchangeParameters.OverridePipelines.Add(Context->MaterialXOptions);
-		}
-
 		// Get a temp content folder we can give interchange as it must have a destination
 		int32 PackageSuffix = FMath::RandRange(0, TNumericLimits<int32>::Max()-1);
 		FString TempPackagePath = FString::Printf(TEXT("/Engine/USDImporter/Transient/%d"), PackageSuffix++);
@@ -384,7 +309,7 @@ void FMaterialXUsdShadeMaterialTranslator::CreateAssets()
 		// Currently Interchange cannot import directly into the transient package, so we'll import into a temporary folder and then
 		// move everything over
 		InterchangeParameters.OnAssetsImportDoneNative.BindLambda(
-			[this, &MaterialXFilePath, &MaterialXHash, &TargetMaterialAssetHash, &ParsedMaterial, &ReferencerPrimPath, &TempPackagePath]
+			[this, &MaterialXHash, &TargetMaterialAssetHash, &ParsedMaterial, &ReferencerPrimPath, &TempPackagePath]
 			(const TArray<UObject*>& ImportedObjects)
 			{
 				// Move all our assets to the transient package
diff --git a/Engine/Plugins/Importers/USDImporter/Source/USDSchemas/Public/USDSchemaTranslator.h b/Engine/Plugins/Importers/USDImporter/Source/USDSchemas/Public/USDSchemaTranslator.h
--- a/Engine/Plugins/Importers/USDImporter/Source/USDSchemas/Public/USDSchemaTranslator.h
+++ b/Engine/Plugins/Importers/USDImporter/Source/USDSchemas/Public/USDSchemaTranslator.h
@@ -168,9 +168,6 @@ struct USDSCHEMAS_API FUsdSchemaTranslationContext : public TSharedFromThis< FUs
 	/** Describes what to add to the root bone animation within generated AnimSequences, if anything */
 	EUsdRootMotionHandling RootMotionHandling = EUsdRootMotionHandling::NoAdditionalRootMotion;
 
-	/** Import options for parsing MaterialX through Interchange */
-	TObjectPtr<UInterchangeGenericMaterialPipeline> MaterialXOptions;
-
 	/** If a generated UStaticMesh has at least this many triangles we will attempt to enable Nanite */
 	int32 NaniteTriangleThreshold;
 
diff --git a/Engine/Plugins/Importers/USDImporter/Source/USDStage/Private/USDStageActor.cpp b/Engine/Plugins/Importers/USDImporter/Source/USDStage/Private/USDStageActor.cpp
--- a/Engine/Plugins/Importers/USDImporter/Source/USDStage/Private/USDStageActor.cpp
+++ b/Engine/Plugins/Importers/USDImporter/Source/USDStage/Private/USDStageActor.cpp
@@ -132,7 +132,6 @@ struct FUsdStageActorImpl
 		TranslationContext->RenderContext = StageActor->RenderContext;
 		TranslationContext->MaterialPurpose = StageActor->MaterialPurpose;
 		TranslationContext->RootMotionHandling = StageActor->RootMotionHandling;
-		TranslationContext->MaterialXOptions = StageActor->MaterialXOptions;
 		TranslationContext->BlendShapesByPath = &StageActor->BlendShapesByPath;
 		TranslationContext->InfoCache = StageActor->InfoCache;
 		TranslationContext->bTranslateOnlyUsedMaterials = true;  // There is no point in turning this off when just opening a stage
@@ -676,8 +675,6 @@ AUsdStageActor::AUsdStageActor()
 
 	RootComponent = SceneComponent;
 
-	MaterialXOptions = CreateDefaultSubobject<UInterchangeGenericMaterialPipeline>("MaterialXOptions");
-
 	// Note: We can't construct our RootUsdTwin as a default subobject here, it needs to be built on-demand.
 	// Even if we NewObject'd one it will work as a subobject in some contexts (maybe because the CDO will have a dedicated root twin?).
 	// As far as the engine is concerned, our prim twins are static assets like meshes or textures. However, they live on the transient
@@ -1837,15 +1834,6 @@ void AUsdStageActor::SetRootMotionHandling( EUsdRootMotionHandling NewHandlingSt
 	LoadUsdStage();
 }
 
-void AUsdStageActor::SetMaterialXOptions(UInterchangeGenericMaterialPipeline* NewOptions)
-{
-	const bool bMarkDirty = false;
-	Modify(bMarkDirty);
-
-	MaterialXOptions = NewOptions;
-	LoadUsdStage();
-}
-
 float AUsdStageActor::GetTime() const
 {
 	return Time;
@@ -3241,11 +3229,6 @@ void AUsdStageActor::OnObjectPropertyChanged(UObject* ObjectBeingModified, FProp
 		HandlePropertyChangedEvent(PropertyChangedEvent);
 		return;
 	}
-	else if (ObjectBeingModified == MaterialXOptions.Get())
-	{
-		SetMaterialXOptions(MaterialXOptions);
-		return;
-	}
 
 	// Don't modify the stage if we're in PIE
 	if (!HasAuthorityOverStage())
@@ -3473,10 +3456,6 @@ void AUsdStageActor::HandlePropertyChangedEvent(FPropertyChangedEvent& PropertyC
 	{
 		SetRootMotionHandling(RootMotionHandling);
 	}
-	else if (PropertyName == GET_MEMBER_NAME_CHECKED(AUsdStageActor, MaterialXOptions))
-	{
-		SetMaterialXOptions(MaterialXOptions);
-	}
 	else if (PropertyName == GET_MEMBER_NAME_CHECKED(AUsdStageActor, UsdAssetCache))
 	{
 		SetAssetCache(UsdAssetCache);
diff --git a/Engine/Plugins/Importers/USDImporter/Source/USDStage/Public/USDStageActor.h b/Engine/Plugins/Importers/USDImporter/Source/USDStage/Public/USDStageActor.h
--- a/Engine/Plugins/Importers/USDImporter/Source/USDStage/Public/USDStageActor.h
+++ b/Engine/Plugins/Importers/USDImporter/Source/USDStage/Public/USDStageActor.h
@@ -15,7 +15,6 @@
 
 class FUsdInfoCache;
 class ISequencer;
-class UInterchangeGenericMaterialPipeline;
 class ULevelSequence;
 class UUsdAssetCache;
 class UUsdAssetCache2;
@@ -92,10 +91,6 @@ public:
 	UPROPERTY( EditAnywhere, BlueprintReadOnly, Category = "USD", config )
 	EUsdRootMotionHandling RootMotionHandling = EUsdRootMotionHandling::NoAdditionalRootMotion;
 
-	// Settings to use when parsing MaterialX materials via Interchange
-	UPROPERTY( EditAnywhere, BlueprintReadOnly, Category = "USD", meta = (DisplayName = "MaterialX Options"))
-	TObjectPtr<UInterchangeGenericMaterialPipeline> MaterialXOptions;
-
 public:
 	DECLARE_EVENT_OneParam(AUsdStageActor, FOnActorLoaded, AUsdStageActor*);
 	USDSTAGE_API static FOnActorLoaded OnActorLoaded;
@@ -152,9 +147,6 @@ public:
 	UFUNCTION( BlueprintCallable, Category = "USD", meta = ( CallInEditor = "true" ) )
 	USDSTAGE_API void SetRootMotionHandling( EUsdRootMotionHandling NewHandlingStrategy );
 
-	UFUNCTION(BlueprintCallable, Category = "USD", meta = (CallInEditor = "true"))
-	USDSTAGE_API void SetMaterialXOptions(UInterchangeGenericMaterialPipeline* NewOptions);
-
 	UFUNCTION(BlueprintCallable, Category = "USD", meta = (CallInEditor = "true"))
 	USDSTAGE_API float GetTime() const;
 
diff --git a/Engine/Plugins/Importers/USDImporter/Source/USDStageEditor/Private/SUSDStage.cpp b/Engine/Plugins/Importers/USDImporter/Source/USDStageEditor/Private/SUSDStage.cpp
--- a/Engine/Plugins/Importers/USDImporter/Source/USDStageEditor/Private/SUSDStage.cpp
+++ b/Engine/Plugins/Importers/USDImporter/Source/USDStageEditor/Private/SUSDStage.cpp
@@ -1113,11 +1113,6 @@ void SUsdStage::FillOptionsMenu(FMenuBuilder& MenuBuilder)
 			LOCTEXT( "RootMotionHandling_ToolTip", "Choose how to handle root motion for generated AnimSequences and LevelSequences" ),
 			FNewMenuDelegate::CreateSP( this, &SUsdStage::FillRootMotionSubMenu ) );
 
-		MenuBuilder.AddSubMenu(
-			LOCTEXT("MaterialXOptions", "MaterialX options"),
-			LOCTEXT("MaterialXOptions_ToolTip", "Choose the settings to use when parsing MaterialX materials via Interchange"),
-			FNewMenuDelegate::CreateSP(this, &SUsdStage::FillMaterialXOptionsSubMenu));
-
 		MenuBuilder.AddSubMenu(
 			LOCTEXT( "Collapsing", "Collapsing" ),
 			LOCTEXT( "Collapsing_ToolTip", "Whether to try to combine individual assets and components of the same type on a Kind-per-Kind basis, like multiple Mesh prims into a single Static Mesh" ),
@@ -1592,36 +1587,6 @@ void SUsdStage::FillRootMotionSubMenu( FMenuBuilder& MenuBuilder )
 	AddRootMotionEntry( EUsdRootMotionHandling::UseMotionFromSkeleton, LOCTEXT("UseMotionFromSkeletonText", "Use motion from Skeleton"));
 }
 
-void SUsdStage::FillMaterialXOptionsSubMenu(FMenuBuilder& MenuBuilder)
-{
-	if (AUsdStageActor* StageActor = GetStageActorOrCDO())
-	{
-		FPropertyEditorModule& PropertyEditorModule = FModuleManager::GetModuleChecked<FPropertyEditorModule>("PropertyEditor");
-
-		FSinglePropertyParams InitParams;
-		InitParams.NamePlacement = EPropertyNamePlacement::Hidden;
-
-		TSharedPtr<ISinglePropertyView> PropertyView = PropertyEditorModule.CreateSingleProperty(
-			StageActor,
-			GET_MEMBER_NAME_CHECKED(AUsdStageActor, MaterialXOptions),
-			InitParams
-		);
-
-		TSharedRef<SBox> Box = SNew(SBox)
-			.Padding(FMargin(8.0f, 0.0f))
-			.VAlign(VAlign_Center)
-			[
-				PropertyView.ToSharedRef()
-			];
-
-		if (PropertyView)
-		{
-			const bool bNoIndent = true;
-			MenuBuilder.AddWidget(Box, FText::GetEmpty(), bNoIndent, bNoIndent);
-		}
-	}
-}
-
 void SUsdStage::FillCollapsingSubMenu( FMenuBuilder& MenuBuilder )
 {
 	MenuBuilder.AddMenuEntry(
diff --git a/Engine/Plugins/Importers/USDImporter/Source/USDStageEditor/Private/SUSDStage.h b/Engine/Plugins/Importers/USDImporter/Source/USDStageEditor/Private/SUSDStage.h
--- a/Engine/Plugins/Importers/USDImporter/Source/USDStageEditor/Private/SUSDStage.h
+++ b/Engine/Plugins/Importers/USDImporter/Source/USDStageEditor/Private/SUSDStage.h
@@ -91,7 +91,6 @@ protected:
 	void FillRenderContextSubMenu( FMenuBuilder& MenuBuilder );
 	void FillMaterialPurposeSubMenu( FMenuBuilder& MenuBuilder );
 	void FillRootMotionSubMenu( FMenuBuilder& MenuBuilder );
-	void FillMaterialXOptionsSubMenu( FMenuBuilder& MenuBuilder );
 	void FillCollapsingSubMenu( FMenuBuilder& MenuBuilder );
 	void FillInterpolationTypeSubMenu( FMenuBuilder& MenuBuilder );
 	void FillSelectionSubMenu( FMenuBuilder& MenuBuilder );
diff --git a/Engine/Plugins/Importers/USDImporter/Source/USDStageEditorViewModels/Private/USDStageViewModel.cpp b/Engine/Plugins/Importers/USDImporter/Source/USDStageEditorViewModels/Private/USDStageViewModel.cpp
--- a/Engine/Plugins/Importers/USDImporter/Source/USDStageEditorViewModels/Private/USDStageViewModel.cpp
+++ b/Engine/Plugins/Importers/USDImporter/Source/USDStageEditorViewModels/Private/USDStageViewModel.cpp
@@ -316,7 +316,6 @@ void FUsdStageViewModel::ImportStage( const TCHAR* TargetContentFolder, UUsdStag
 		ImportContext.ImportOptions->ImportTimeCode = StageActor->GetTime();
 		ImportContext.ImportOptions->NaniteTriangleThreshold = StageActor->NaniteTriangleThreshold;
 		ImportContext.ImportOptions->RootMotionHandling = StageActor->RootMotionHandling;
-		ImportContext.ImportOptions->MaterialXOptions = StageActor->MaterialXOptions;
 		ImportContext.ImportOptions->KindsToCollapse = StageActor->KindsToCollapse;
 		ImportContext.ImportOptions->bMergeIdenticalMaterialSlots = StageActor->bMergeIdenticalMaterialSlots;
 
diff --git a/Engine/Plugins/Importers/USDImporter/Source/USDStageImporter/Private/USDStageImporter.cpp b/Engine/Plugins/Importers/USDImporter/Source/USDStageImporter/Private/USDStageImporter.cpp
--- a/Engine/Plugins/Importers/USDImporter/Source/USDStageImporter/Private/USDStageImporter.cpp
+++ b/Engine/Plugins/Importers/USDImporter/Source/USDStageImporter/Private/USDStageImporter.cpp
@@ -1739,7 +1739,6 @@ void UUsdStageImporter::ImportFromFile(FUsdStageImportContext& ImportContext)
 	TranslationContext->RenderContext = ImportContext.ImportOptions->RenderContextToImport;
 	TranslationContext->MaterialPurpose = ImportContext.ImportOptions->MaterialPurpose;
 	TranslationContext->RootMotionHandling = ImportContext.ImportOptions->RootMotionHandling;
-	TranslationContext->MaterialXOptions = ImportContext.ImportOptions->MaterialXOptions;
 	TranslationContext->ParentComponent = ImportContext.SceneActor ? ImportContext.SceneActor->GetRootComponent() : nullptr;
 	TranslationContext->KindsToCollapse = ( EUsdDefaultKind ) ImportContext.ImportOptions->KindsToCollapse;
 	TranslationContext->bMergeIdenticalMaterialSlots = ImportContext.ImportOptions->bMergeIdenticalMaterialSlots;
@@ -1869,7 +1868,6 @@ bool UUsdStageImporter::ReimportSingleAsset(
 	TranslationContext->RenderContext = ImportContext.ImportOptions->RenderContextToImport;
 	TranslationContext->MaterialPurpose = ImportContext.ImportOptions->MaterialPurpose;
 	TranslationContext->RootMotionHandling = ImportContext.ImportOptions->RootMotionHandling;
-	TranslationContext->MaterialXOptions = ImportContext.ImportOptions->MaterialXOptions;
 	TranslationContext->KindsToCollapse = ( EUsdDefaultKind ) ImportContext.ImportOptions->KindsToCollapse;
 	TranslationContext->bMergeIdenticalMaterialSlots = ImportContext.ImportOptions->bMergeIdenticalMaterialSlots;
 	TranslationContext->bAllowInterpretingLODs = ImportContext.ImportOptions->bInterpretLODs;
