USD: Fix build on non-editor targets.

diff --git a/Engine/Plugins/Importers/USDImporter/Source/USDSchemas/Private/MeshTranslationImpl.cpp b/Engine/Plugins/Importers/USDImporter/Source/USDSchemas/Private/MeshTranslationImpl.cpp
--- a/Engine/Plugins/Importers/USDImporter/Source/USDSchemas/Private/MeshTranslationImpl.cpp
+++ b/Engine/Plugins/Importers/USDImporter/Source/USDSchemas/Private/MeshTranslationImpl.cpp
@@ -321,6 +321,7 @@ namespace UE::MeshTranslationImplInternal::Private
 				*MaterialHash
 			);
 
+#if WITH_EDITOR
 			if (GIsEditor)
 			{
 				UMaterialInstanceConstant* CompatibleMIC = NewObject<UMaterialInstanceConstant>(
@@ -334,6 +335,7 @@ namespace UE::MeshTranslationImplInternal::Private
 				CompatibleMaterial = CompatibleMIC;
 			}
 			else
+#endif // WITH_EDITOR
 			{
 				UMaterialInstanceDynamic* CompatibleMID = UMaterialInstanceDynamic::Create(
 					&Material,
@@ -383,7 +385,9 @@ namespace UE::MeshTranslationImplInternal::Private
 					}
 				}
 
+#if WITH_EDITOR
 				CompatibleInstance->PostEditChange();
+#endif // WITH_EDITOR
 			}
 		}
 
diff --git a/Engine/Plugins/Importers/USDImporter/Source/USDSchemas/Private/USDGeomMeshTranslator.cpp b/Engine/Plugins/Importers/USDImporter/Source/USDSchemas/Private/USDGeomMeshTranslator.cpp
--- a/Engine/Plugins/Importers/USDImporter/Source/USDSchemas/Private/USDGeomMeshTranslator.cpp
+++ b/Engine/Plugins/Importers/USDImporter/Source/USDSchemas/Private/USDGeomMeshTranslator.cpp
@@ -257,6 +257,7 @@ namespace UsdGeomMeshTranslatorImpl
 					UserData->MaterialSlotToPrimPaths.FindOrAdd(StaticMeshSlotIndex).PrimPaths = Slot.PrimPaths.Array();
 				}
 
+#if WITH_EDITOR
 				// Setup the section map so that our LOD material index is properly mapped to the static mesh material index
 				// At runtime we don't ever parse these variants as LODs so we don't need this
 				if (StaticMesh.GetSectionInfoMap().IsValidSection(LODIndex, LODSlotIndex))
@@ -280,6 +281,7 @@ namespace UsdGeomMeshTranslatorImpl
 
 					bMaterialAssignementsHaveChanged = true;
 				}
+#endif // WITH_EDITOR
 			}
 		}
 
