USD: Fix some prim geometry updates not causing assets to refresh on the UE level;

Fix some automated tests leaving the "CollapseTopLevelPointInstancer" cvar enabled and breaking other tests;
Update test_update_refreshes_assets to AnimSequences having prim names now.

diff --git a/Engine/Plugins/Importers/USDImporter/Source/USDSchemas/Private/USDInfoCache.cpp b/Engine/Plugins/Importers/USDImporter/Source/USDSchemas/Private/USDInfoCache.cpp
--- a/Engine/Plugins/Importers/USDImporter/Source/USDSchemas/Private/USDInfoCache.cpp
+++ b/Engine/Plugins/Importers/USDImporter/Source/USDSchemas/Private/USDInfoCache.cpp
@@ -724,25 +724,39 @@ namespace UE::USDInfoCacheImpl::Private
 	void CollectMaterialSlotCounts(
 		FUsdInfoCache::FUsdInfoCacheImpl& Impl,
 		const TMap< UE::FSdfPath, TArray<UsdUtils::FUsdPrimMaterialSlot>>& SubtreeMaterialSlots,
-		bool bMergeIdenticalSlots
+		bool bContextMergeIdenticalSlots
 	)
 	{
 		FWriteScopeLock Lock{Impl.InfoMapLock};
 
-		if ( bMergeIdenticalSlots )
+		for (const TPair<UE::FSdfPath, TArray<UsdUtils::FUsdPrimMaterialSlot>>& Pair : SubtreeMaterialSlots)
 		{
-			for ( const TPair< UE::FSdfPath, TArray<UsdUtils::FUsdPrimMaterialSlot>>& Pair : SubtreeMaterialSlots )
+			const UE::FSdfPath& PrimPath = Pair.Key;
+
+			bool bCanMergeSlotsForThisPrim = false;
+			if (const UE::UsdInfoCache::Private::FUsdPrimInfo* FoundInfo = Impl.InfoMap.Find(PrimPath))
 			{
-				TSet<UsdUtils::FUsdPrimMaterialSlot> SlotsSet{ Pair.Value };
-				UE::UsdInfoCache::Private::FUsdPrimInfo& Info = Impl.InfoMap.FindOrAdd( Pair.Key );
-				Info.SubtreeMaterialSlots = SlotsSet.Array();
+				const UE::FSdfPath* CollapsedRoot = &FoundInfo->AssetCollapsedRoot;
+
+				// We only merge slots in the context of collapsing
+				bool bPrimIsCollapsedOrCollapseRoot = !CollapsedRoot->IsEmpty() || PrimPath.IsAbsoluteRootPath();
+
+				// TODO: This is not perfect, because we may have bPrimIsPotentialGeometryCacheRoot but the prim hasn't
+				// actually become a geometry cache due to another reason
+				bool bPrimIsPotentialGeometryCacheRoot = FoundInfo->bIsPotentialGeoCacheRoot;
+
+				bCanMergeSlotsForThisPrim = bPrimIsCollapsedOrCollapseRoot && !bPrimIsPotentialGeometryCacheRoot;
 			}
-		}
-		else
-		{
-			for ( const TPair< UE::FSdfPath, TArray<UsdUtils::FUsdPrimMaterialSlot>>& Pair : SubtreeMaterialSlots )
+
+			UE::UsdInfoCache::Private::FUsdPrimInfo& Info = Impl.InfoMap.FindOrAdd(Pair.Key);
+
+			// For now we only ever merge material slots when collapsing
+			if (bCanMergeSlotsForThisPrim && bContextMergeIdenticalSlots)
+			{
+				Info.SubtreeMaterialSlots = TSet<UsdUtils::FUsdPrimMaterialSlot>{Pair.Value}.Array();
+			}
+			else
 			{
-				UE::UsdInfoCache::Private::FUsdPrimInfo& Info = Impl.InfoMap.FindOrAdd( Pair.Key );
 				Info.SubtreeMaterialSlots = Pair.Value;
 			}
 		}
@@ -1144,17 +1158,11 @@ void FUsdInfoCache::RebuildCacheForSubtree( const UE::FUsdPrim& Prim, FUsdSchema
 			TempSubtreeSlots
 		);
 
-		UE::USDInfoCacheImpl::Private::CollectMaterialSlotCounts(
-			*ImplPtr,
-			TempSubtreeSlots,
-			Context.bMergeIdenticalMaterialSlots
-		);
-
-		UE::USDInfoCacheImpl::Private::RecursiveCheckForGeometryCache(			
+		UE::USDInfoCacheImpl::Private::RecursiveCheckForGeometryCache(
 			UsdPrim,
 			Context,
 			*ImplPtr
-			);
+		);
 
 		UE::USDInfoCacheImpl::Private::RecursiveQueryCollapsesChildren(
 			UsdPrim,
@@ -1162,6 +1170,12 @@ void FUsdInfoCache::RebuildCacheForSubtree( const UE::FUsdPrim& Prim, FUsdSchema
 			*ImplPtr,
 			Registry
 		);
+
+		UE::USDInfoCacheImpl::Private::CollectMaterialSlotCounts(
+			*ImplPtr,
+			TempSubtreeSlots,
+			Context.bMergeIdenticalMaterialSlots
+		);
 	}
 #endif // USE_USD_SDK
 }
diff --git a/Engine/Plugins/Importers/USDImporter/Source/USDStage/Private/USDStageActor.cpp b/Engine/Plugins/Importers/USDImporter/Source/USDStage/Private/USDStageActor.cpp
--- a/Engine/Plugins/Importers/USDImporter/Source/USDStage/Private/USDStageActor.cpp
+++ b/Engine/Plugins/Importers/USDImporter/Source/USDStage/Private/USDStageActor.cpp
@@ -1045,7 +1045,13 @@ void AUsdStageActor::OnUsdObjectsChanged(const UsdUtils::FObjectChangesByPath& I
 					*UsdToUnreal::ConvertToken(UnrealIdentifiers::UnrealControlRigPath),
 					*UsdToUnreal::ConvertToken(UnrealIdentifiers::UnrealUseFKControlRig),
 					*UsdToUnreal::ConvertToken(UnrealIdentifiers::UnrealControlRigReduceKeys),
-					*UsdToUnreal::ConvertToken(UnrealIdentifiers::UnrealControlRigReductionTolerance)
+					*UsdToUnreal::ConvertToken(UnrealIdentifiers::UnrealControlRigReductionTolerance),
+					// For now we need to do this as we need to refresh the material slot info on the info cache if these
+					// update... this is of course way too aggressive, although it's unlikely people will be manually editing these.
+					// TODO: More nuanced info cache updates
+					UnrealIdentifiers::PrimvarsDisplayColor,
+					UnrealIdentifiers::PrimvarsDisplayOpacity,
+					UnrealIdentifiers::DoubleSided
 				};
 
 				// Some stage info should trigger some resyncs because they should trigger reparsing of geometry
diff --git a/Engine/Plugins/Importers/USDImporter/Source/UnrealUSDWrapper/Private/UnrealUSDWrapper.cpp b/Engine/Plugins/Importers/USDImporter/Source/UnrealUSDWrapper/Private/UnrealUSDWrapper.cpp
--- a/Engine/Plugins/Importers/USDImporter/Source/UnrealUSDWrapper/Private/UnrealUSDWrapper.cpp
+++ b/Engine/Plugins/Importers/USDImporter/Source/UnrealUSDWrapper/Private/UnrealUSDWrapper.cpp
@@ -440,10 +440,16 @@ FString UnrealIdentifiers::MaterialAllPurposeText = TEXT( "allPurpose" );
 FString UnrealIdentifiers::MaterialAllPurpose = ANSI_TO_TCHAR( pxr::UsdShadeTokens->allPurpose.GetString().c_str() );
 FString UnrealIdentifiers::MaterialPreviewPurpose = ANSI_TO_TCHAR( pxr::UsdShadeTokens->preview.GetString().c_str() );
 FString UnrealIdentifiers::MaterialFullPurpose = ANSI_TO_TCHAR( pxr::UsdShadeTokens->full.GetString().c_str() );
+FString UnrealIdentifiers::PrimvarsDisplayColor = ANSI_TO_TCHAR(pxr::UsdGeomTokens->primvarsDisplayColor.GetString().c_str());
+FString UnrealIdentifiers::PrimvarsDisplayOpacity = ANSI_TO_TCHAR(pxr::UsdGeomTokens->primvarsDisplayOpacity.GetString().c_str());
+FString UnrealIdentifiers::DoubleSided = ANSI_TO_TCHAR(pxr::UsdGeomTokens->doubleSided.GetString().c_str());
 #else
 FString UnrealIdentifiers::MaterialAllPurpose = TEXT( "" );
 FString UnrealIdentifiers::MaterialPreviewPurpose = TEXT( "preview" );
 FString UnrealIdentifiers::MaterialFullPurpose = TEXT( "full" );
+FString UnrealIdentifiers::PrimvarsDisplayColor = TEXT("primvars:displayColor");
+FString UnrealIdentifiers::PrimvarsDisplayOpacity = TEXT("primvars:displayOpacity");
+FString UnrealIdentifiers::DoubleSided = TEXT("doubleSided");
 #endif // USE_USD_SDK
 
 FUsdDelegates::FUsdImportDelegate FUsdDelegates::OnPreUsdImport;
diff --git a/Engine/Plugins/Importers/USDImporter/Source/UnrealUSDWrapper/Public/UnrealUSDWrapper.h b/Engine/Plugins/Importers/USDImporter/Source/UnrealUSDWrapper/Public/UnrealUSDWrapper.h
--- a/Engine/Plugins/Importers/USDImporter/Source/UnrealUSDWrapper/Public/UnrealUSDWrapper.h
+++ b/Engine/Plugins/Importers/USDImporter/Source/UnrealUSDWrapper/Public/UnrealUSDWrapper.h
@@ -429,6 +429,10 @@ namespace UnrealIdentifiers
 	extern UNREALUSDWRAPPER_API FString MaterialAllPurposeText; // Text to show on UI for "allPurpose", as its value is actually the empty string
 	extern UNREALUSDWRAPPER_API FString MaterialPreviewPurpose;
 	extern UNREALUSDWRAPPER_API FString MaterialFullPurpose;
+
+	extern UNREALUSDWRAPPER_API FString PrimvarsDisplayColor;
+	extern UNREALUSDWRAPPER_API FString PrimvarsDisplayOpacity;
+	extern UNREALUSDWRAPPER_API FString DoubleSided;
 }
 
 struct UNREALUSDWRAPPER_API FUsdDelegates
