USD: Fix USD MaterialX import tests.

Remove temp hack now that textures coming out of interchange have the source filenames on their asset import data.

diff --git a/Engine/Plugins/Importers/USDImporter/Source/USDSchemas/Private/Custom/MaterialXUSDShadeMaterialTranslator.cpp b/Engine/Plugins/Importers/USDImporter/Source/USDSchemas/Private/Custom/MaterialXUSDShadeMaterialTranslator.cpp
--- a/Engine/Plugins/Importers/USDImporter/Source/USDSchemas/Private/Custom/MaterialXUSDShadeMaterialTranslator.cpp
+++ b/Engine/Plugins/Importers/USDImporter/Source/USDSchemas/Private/Custom/MaterialXUSDShadeMaterialTranslator.cpp
@@ -369,9 +369,6 @@ void FMaterialXUsdShadeMaterialTranslator::CreateAssets()
 							Texture->GetTextureAddressX(),
 							Texture->GetTextureAddressY()
 						);
-						// TODO: Replace this hash hack with relying on PostImportMaterial whenever the interchange
-						// texture AssetImportData contains the correct filepaths, and not just the .mtlx file
-						TextureHash += TEXT("_") + Texture->GetFName().GetPlainNameString();
 
 						// See comment on the analogous part of the Material case above
 						UTexture* ExistingTexture = Cast<UTexture>(Context->AssetCache->GetCachedAsset(TextureHash));
