USD: Fix issues with asset conflict resolution when importing.

Remove call to DuplicateObject as we can almost always just rename Asset to the target package, as we're renaming ExistingAsset away anyway;
Remove forced MarkAsGarbage() followed by GC call on almost all existing assets as that is no longer needed and can lead to many issues;
Fix errors on SkeletalMeshComponents due to AnimScriptInstance trying to use data from old replaced SkeletalMesh assets;
Fix LevelSequences being duplicated unnecessarily due to not being present in the AssetCache.

diff --git a/Engine/Plugins/Importers/USDImporter/Source/USDStageImporter/Private/USDStageImporter.cpp b/Engine/Plugins/Importers/USDImporter/Source/USDStageImporter/Private/USDStageImporter.cpp
--- a/Engine/Plugins/Importers/USDImporter/Source/USDStageImporter/Private/USDStageImporter.cpp
+++ b/Engine/Plugins/Importers/USDImporter/Source/USDStageImporter/Private/USDStageImporter.cpp
@@ -605,7 +605,14 @@ namespace UsdStageImporterImpl
 
 	// Moves Asset from its folder to the package at DestFullContentPath and sets up its flags.
 	// Depending on ReplacePolicy it may replace the existing actor (if it finds one) or just abort
-	UObject* PublishAsset(FUsdStageImportContext& ImportContext, UObject* Asset, const FString& DestFullPackagePath, TMap<UObject*, UObject*>& ObjectsToRemap, TMap<FSoftObjectPath, FSoftObjectPath>& SoftObjectsToRemap )
+	UObject* PublishAsset(
+		FUsdStageImportContext& ImportContext,
+		UObject* Asset,
+		const FString& DestFullPackagePath,
+		TMap<UObject*, UObject*>& ObjectsToRemap,
+		TMap<FSoftObjectPath, FSoftObjectPath>& SoftObjectsToRemap,
+		TSet<UObject*>& OutAssetsToFinalize
+	)
 	{
 		if (!Asset)
 		{
@@ -672,16 +679,31 @@ namespace UsdStageImporterImpl
 
 		// If the asset is currently being referenced on its asset cache we can't just take it, and must duplicate it
 		// below with the DuplicateObject call
-		TStrongObjectPtr<UObject> TakenAsset = nullptr;
+		bool bForceDuplicate = false;
 		if (UUsdAssetCache2* Cache = ImportContext.AssetCache.Get())
 		{
 			FString Hash = Cache->GetHashForAsset(Asset);
-			if (!Hash.IsEmpty() && Cache->CanRemoveAsset(Hash))
+			const bool bIsAssetOwnedByCache = !Hash.IsEmpty();
+			if (bIsAssetOwnedByCache)
 			{
-				TakenAsset.Reset(Cache->RemoveAsset(Hash));
-				if (!TakenAsset.IsValid())
+				bool bRemovedFromCache = Cache->CanRemoveAsset(Hash);
+				if (bRemovedFromCache)
 				{
-					UE_LOG(LogUsd, Log, TEXT("About to duplicate asset '%s' on import as it can't be taken directly from the asset cache"), *Asset->GetPathName());
+					UObject* RemovedAsset = Cache->RemoveAsset(Hash);
+					ensure(!RemovedAsset || RemovedAsset == Asset);
+
+					bRemovedFromCache = RemovedAsset != nullptr;
+				}
+
+				if (!bRemovedFromCache)
+				{
+					bForceDuplicate = true;
+					UE_LOG(
+						LogUsd,
+						Log,
+						TEXT("About to duplicate asset '%s' on import as it can't be taken directly from the asset cache"),
+						*Asset->GetPathName()
+					);
 				}
 			}
 		}
@@ -689,19 +711,8 @@ namespace UsdStageImporterImpl
 		// Strategy copied from FDatasmithImporterImpl::PublicizeAsset
 		// Replace existing asset (reimport or conflict) with new asset
 		UObject* MovedAsset = ExistingAsset;
-		if (!TakenAsset || (ExistingAsset != nullptr && ExistingAsset != Asset && ReplacePolicy == EReplaceAssetPolicy::Replace))
+		if (bForceDuplicate || (ExistingAsset != nullptr && ExistingAsset != Asset && ReplacePolicy == EReplaceAssetPolicy::Replace))
 		{
-			// We must do this here or else on the next engine tick the USkeletalMeshComponent::ParallelAnimationEvaluation
-			// will try to use the old anim instance, which will try using the old Skeleton's SmartNameContainer and crash.
-			// We do not need to recreate the anim instance afterwards as InitAnim is called on component registration anyway.
-			for ( TObjectIterator<USkeletalMeshComponent> It; It; ++It )
-			{
-				if ( It->GetSkeletalMeshAsset() == ExistingAsset)
-				{
-					It->ClearAnimScriptInstance();
-				}
-			}
-
 			// We have to rename the existing asset away, because some objects manage subobjects (like
 			// UMaterialInterface and its EditorOnlyData), and try to rename them whenever they're duplicated/renamed.
 			// That can lead to issues if the target names conflict with the names of ExistingAsset's subobjects
@@ -719,33 +730,36 @@ namespace UsdStageImporterImpl
 				);
 			}
 
-			MovedAsset = DuplicateObject<UObject>( Asset, Package, *TargetAssetName );
-
-			// If our DuplicateObject didn't stomp the old asset because TargetAssetName != ExistingAsset->GetName(),
-			// mark the existing asset as garbage so that we collect it when we run GC after the import
-			if ( ExistingAsset && ExistingAsset->GetName() != TargetAssetName )
+			if (bForceDuplicate)
 			{
-				ExistingAsset->MarkAsGarbage();
-				ImportContext.bNeedsGarbageCollection = true;
+				MovedAsset = DuplicateObject<UObject>(Asset, Package, *TargetAssetName);
 			}
-
-			if (UStaticMesh* DestinationMesh = Cast< UStaticMesh >(MovedAsset))
+			// We used to always duplicate here, but now since we're renaming ExistingAsset away, most of the time we can just
+			// rename our asset to its target package too
+			else
 			{
-				// This is done during the mesh build process but we need to redo it after the DuplicateObject since the links are now valid
-				for (TObjectIterator< UStaticMeshComponent > It; It; ++It)
-				{
-					if (It->GetStaticMesh() == DestinationMesh)
-					{
-						It->FixupOverrideColorsIfNecessary(true);
-						It->InvalidateLightingCache();
-					}
-				}
+				ensure(Asset->Rename(
+					*TargetAssetName,
+					Package,
+					REN_DontCreateRedirectors | REN_NonTransactional | REN_DoNotDirty | REN_ForceNoResetLoaders
+				));
+				MovedAsset = Asset;
 			}
+
+			OutAssetsToFinalize.Add(ExistingAsset);
+			OutAssetsToFinalize.Add(MovedAsset);
+
+			// If we're replacing ExistingAsset, we must update all references we can find from ExistingAsset to the new
+			// asset, otherwise they'll be left pointing at transient or GC'd/stomped assets
+			SoftObjectsToRemap.Add(ExistingAsset, MovedAsset);
+			ObjectsToRemap.Add(ExistingAsset, MovedAsset);
 		}
 		else
 		{
 			// We can't dirty the package here. Read the comment around MarkPackageDirty, below
-			Asset->Rename(*TargetAssetName, Package, REN_DontCreateRedirectors | REN_NonTransactional | REN_DoNotDirty );
+			ensure(
+				Asset->Rename(*TargetAssetName, Package, REN_DontCreateRedirectors | REN_NonTransactional | REN_DoNotDirty | REN_ForceNoResetLoaders)
+			);
 			MovedAsset = Asset;
 		}
 
@@ -789,7 +803,13 @@ namespace UsdStageImporterImpl
 	}
 
 	// Move imported assets from transient folder to their final package, updating AssetCache to point to the moved assets
-	void PublishAssets(FUsdStageImportContext& ImportContext, const TSet<UObject*>& AssetsToPublish, TMap<UObject*, UObject*>& ObjectsToRemap, TMap<FSoftObjectPath, FSoftObjectPath>& SoftObjectsToRemap )
+	void PublishAssets(
+		FUsdStageImportContext& ImportContext,
+		const TSet<UObject*>& AssetsToPublish,
+		TMap<UObject*, UObject*>& ObjectsToRemap,
+		TMap<FSoftObjectPath, FSoftObjectPath>& SoftObjectsToRemap,
+		TSet<UObject*>& OutAssetsToFinalize
+	)
 	{
 		TArray<UObject*> Textures;
 		TArray<UObject*> Materials;
@@ -910,7 +930,14 @@ namespace UsdStageImporterImpl
 				{
 					FString TargetAssetName = GetUserFriendlyName(Asset, UniqueAssetNames);
 					FString DestPackagePath = FPaths::Combine(ImportContext.PackagePath, *AssetTypeFolder, TargetAssetName);
-					PublishAsset(ImportContext, Asset, DestPackagePath, ObjectsToRemap, SoftObjectsToRemap);
+					PublishAsset(
+						ImportContext,
+						Asset,
+						DestPackagePath,
+						ObjectsToRemap,
+						SoftObjectsToRemap,
+						OutAssetsToFinalize
+					);
 				}
 			}
 		};
@@ -1247,8 +1274,6 @@ namespace UsdStageImporterImpl
 			constexpr EArchiveReplaceObjectFlags ReplaceFlags = (EArchiveReplaceObjectFlags::IgnoreOuterRef | EArchiveReplaceObjectFlags::IgnoreArchetypeRef);
 			FArchiveReplaceObjectRef< UObject > ArchiveReplaceObjectRefInner(Referencer, ObjectsToRemap, ReplaceFlags);
 		}
-
-
 	}
 
 	void CallAssetsPostEditChange(const TSet<UObject*>& PublishedObjects)
@@ -1268,7 +1293,12 @@ namespace UsdStageImporterImpl
 		}
 	}
 
-	void Cleanup(AActor* NewSceneActor, AActor* ExistingSceneActor, EReplaceActorPolicy ReplacePolicy)
+	void Cleanup(
+		AActor* NewSceneActor,
+		AActor* ExistingSceneActor,
+		EReplaceActorPolicy ReplacePolicy,
+		const TSet<UObject*>& AssetsToFinalize
+	)
 	{
 		if ( !NewSceneActor )
 		{
@@ -1281,6 +1311,25 @@ namespace UsdStageImporterImpl
 		{
 			RecursiveDestroyActor(ExistingSceneActor);
 		}
+
+		for (TObjectIterator<UMeshComponent> It; It; ++It)
+		{
+			if(UStaticMeshComponent* StaticMeshComponent = Cast<UStaticMeshComponent>(*It))
+			{
+				if(AssetsToFinalize.Contains(StaticMeshComponent->GetStaticMesh()))
+				{
+					StaticMeshComponent->FixupOverrideColorsIfNecessary(true);
+					StaticMeshComponent->InvalidateLightingCache();
+				}
+			}
+			else if (USkeletalMeshComponent* SkeletalMeshComponent = Cast<USkeletalMeshComponent>(*It))
+			{
+				if(AssetsToFinalize.Contains(SkeletalMeshComponent->GetSkeletalMeshAsset()))
+				{
+					SkeletalMeshComponent->ClearAnimScriptInstance();
+				}
+			}
+		}
 	}
 
 	void CloseStageIfNeeded(FUsdStageImportContext& ImportContext)
@@ -1702,6 +1751,7 @@ void UUsdStageImporter::ImportFromFile(FUsdStageImportContext& ImportContext)
 	TMap<UObject*, UObject*> ObjectsToRemap;
 	TSet<UObject*> UsedAssetsAndDependencies;
 	TSet<UObject*> PublishedAssetsAndDependencies;
+	TSet<UObject*> AssetsToFinalize;
 	UsdUtils::FBlendShapeMap BlendShapesByPath;
 
 	// Ensure a valid asset cache
@@ -1723,7 +1773,6 @@ void UUsdStageImporter::ImportFromFile(FUsdStageImportContext& ImportContext)
 	ImportContext.LevelSequenceHelper.SetInfoCache(ImportContext.InfoCache);
 	ImportContext.LevelSequenceHelper.Init( ImportContext.Stage );  // Must happen after the context gets an InfoCache!
 
-	{
 	// Shotgun approach to recreate all render states because we may want to reimport/delete/reassign a material/static/skeletalmesh while it is currently being drawn
 	FGlobalComponentRecreateRenderStateContext RecreateRenderStateContext;
 
@@ -1759,13 +1808,13 @@ void UUsdStageImporter::ImportFromFile(FUsdStageImportContext& ImportContext)
 	UsdStageImporterImpl::CollectUsedAssetDependencies( ImportContext, UsedAssetsAndDependencies );
 	UsdStageImporterImpl::PruneUnwantedAssets( ImportContext, UsedAssetsAndDependencies, ObjectsToRemap, SoftObjectsToRemap );
 	UsdStageImporterImpl::UpdateAssetUserData( UsedAssetsAndDependencies, ImportContext.FilePath, ImportContext.ImportOptions );
-	UsdStageImporterImpl::PublishAssets( ImportContext, UsedAssetsAndDependencies, ObjectsToRemap, SoftObjectsToRemap );
 	UsdStageImporterImpl::ResolveActorConflicts( ImportContext, ExistingSceneActor, ObjectsToRemap, SoftObjectsToRemap );
+	UsdStageImporterImpl::PublishAssets( ImportContext, UsedAssetsAndDependencies, ObjectsToRemap, SoftObjectsToRemap, AssetsToFinalize );
 	UsdStageImporterImpl::RemapReferences( ImportContext, UsedAssetsAndDependencies, ObjectsToRemap );
 	UsdStageImporterImpl::RemapSoftReferences( ImportContext, UsedAssetsAndDependencies, SoftObjectsToRemap );
+	UsdStageImporterImpl::Cleanup( ImportContext.SceneActor, ExistingSceneActor, ImportContext.ImportOptions->ExistingActorPolicy, AssetsToFinalize );
 	UsdStageImporterImpl::GetPublishedAssetsAndDependencies(UsedAssetsAndDependencies, ObjectsToRemap, SoftObjectsToRemap, PublishedAssetsAndDependencies);
 	UsdStageImporterImpl::CallAssetsPostEditChange(PublishedAssetsAndDependencies);
-	UsdStageImporterImpl::Cleanup( ImportContext.SceneActor, ExistingSceneActor, ImportContext.ImportOptions->ExistingActorPolicy );
 	UsdStageImporterImpl::NotifyAssetRegistry(PublishedAssetsAndDependencies);
 	UsdStageImporterImpl::RefreshComponents( ImportContext.SceneActor, ImportContext.ImportOptions->bImportAtSpecificTimeCode );
 
@@ -1779,12 +1828,6 @@ void UUsdStageImporter::ImportFromFile(FUsdStageImportContext& ImportContext)
 	}
 
 	UsdStageImporterImpl::CloseStageIfNeeded( ImportContext );
-	}
-
-	if ( ImportContext.bNeedsGarbageCollection )
-	{
-		CollectGarbage( GARBAGE_COLLECTION_KEEPFLAGS );
-	}
 #endif // #if USE_USD_SDK
 }
 
@@ -1831,6 +1874,7 @@ bool UUsdStageImporter::ReimportSingleAsset(
 
 	TMap<FSoftObjectPath, FSoftObjectPath> SoftObjectsToRemap;
 	TMap<UObject*, UObject*> ObjectsToRemap;
+	TSet<UObject*> AssetsToFinalize;
 	UsdUtils::FBlendShapeMap BlendShapesByPath;
 
 	// Ensure a valid asset cache
@@ -1852,7 +1896,6 @@ bool UUsdStageImporter::ReimportSingleAsset(
 	ImportContext.LevelSequenceHelper.SetInfoCache(ImportContext.InfoCache);
 	ImportContext.LevelSequenceHelper.Init(ImportContext.Stage);  // Must happen after the context gets an InfoCache!
 
-	{
 	// Shotgun approach to recreate all render states because we may want to reimport/delete/reassign a material/static/skeletalmesh while it is currently being drawn
 	FGlobalComponentRecreateRenderStateContext RecreateRenderStateContext;
 
@@ -1916,7 +1959,14 @@ bool UUsdStageImporter::ReimportSingleAsset(
 		UsdStageImporterImpl::CopySkeletonAssignment(ImportContext, OriginalAsset, ReimportedObject );
 
 		// Just publish the one asset we wanted to reimport. Note that we may have other assets here too, but we'll ignore those e.g. a displayColor material or a skeleton
-		OutReimportedAsset = UsdStageImporterImpl::PublishAsset(ImportContext, ReimportedObject, OriginalAsset->GetOutermost()->GetPathName(), ObjectsToRemap, SoftObjectsToRemap);
+		OutReimportedAsset = UsdStageImporterImpl::PublishAsset(
+			ImportContext,
+			ReimportedObject,
+			OriginalAsset->GetOutermost()->GetPathName(),
+			ObjectsToRemap,
+			SoftObjectsToRemap,
+			AssetsToFinalize
+		);
 
 		TSet<UObject*> UsedAssetsAndDependencies = ImportContext.AssetCache->GetActiveAssets();
 		TSet<UObject*> PublishedAssetsAndDependencies;
@@ -1928,7 +1978,12 @@ bool UUsdStageImporter::ReimportSingleAsset(
 		bSuccess = OutReimportedAsset != nullptr && ImportContext.AssetCache->GetActiveAssets().Contains( ReimportedObject );
 	}
 
-	UsdStageImporterImpl::Cleanup( ImportContext.SceneActor, nullptr, ImportContext.ImportOptions->ExistingActorPolicy );
+	UsdStageImporterImpl::Cleanup(
+		ImportContext.SceneActor,
+		nullptr,
+		ImportContext.ImportOptions->ExistingActorPolicy,
+		AssetsToFinalize
+	);
 	UsdStageImporterImpl::NotifyAssetRegistry( { ReimportedObject } );
 	UsdStageImporterImpl::RefreshComponents( ImportContext.SceneActor, ImportContext.ImportOptions->bImportAtSpecificTimeCode );
 
@@ -1942,12 +1997,6 @@ bool UUsdStageImporter::ReimportSingleAsset(
 	}
 
 	UsdStageImporterImpl::CloseStageIfNeeded( ImportContext );
-	}
-
-	if ( ImportContext.bNeedsGarbageCollection )
-	{
-		CollectGarbage( GARBAGE_COLLECTION_KEEPFLAGS );
-	}
 
 #endif // #if USE_USD_SDK
 	return bSuccess;
