USD: Use the default material on mesh assets instead of leaving material slots with nothing.

diff --git a/Engine/Plugins/Importers/USDImporter/Source/USDSchemas/Private/USDGeomMeshTranslator.cpp b/Engine/Plugins/Importers/USDImporter/Source/USDSchemas/Private/USDGeomMeshTranslator.cpp
--- a/Engine/Plugins/Importers/USDImporter/Source/USDSchemas/Private/USDGeomMeshTranslator.cpp
+++ b/Engine/Plugins/Importers/USDImporter/Source/USDSchemas/Private/USDGeomMeshTranslator.cpp
@@ -28,6 +28,8 @@
 #include "HAL/IConsoleManager.h"
 #include "Interfaces/ITargetPlatform.h"
 #include "Interfaces/ITargetPlatformManagerModule.h"
+#include "MaterialDomain.h"
+#include "Materials/Material.h"
 #include "Materials/MaterialInstanceConstant.h"
 #include "Materials/MaterialInstanceDynamic.h"
 #include "Materials/MaterialInterface.h"
@@ -228,7 +230,7 @@ namespace UsdGeomMeshTranslatorImpl
 			{
 				const UsdUtils::FUsdPrimMaterialSlot& Slot = LODSlots[LODSlotIndex];
 
-				UMaterialInterface* Material = nullptr;
+				UMaterialInterface* Material = UMaterial::GetDefaultMaterial(MD_Surface);
 				if (UMaterialInterface** FoundMaterial = ResolvedMaterials.Find(&Slot))
 				{
 					Material = *FoundMaterial;
diff --git a/Engine/Plugins/Importers/USDImporter/Source/USDSchemas/Private/USDGeometryCacheTranslator.cpp b/Engine/Plugins/Importers/USDImporter/Source/USDSchemas/Private/USDGeometryCacheTranslator.cpp
--- a/Engine/Plugins/Importers/USDImporter/Source/USDSchemas/Private/USDGeometryCacheTranslator.cpp
+++ b/Engine/Plugins/Importers/USDImporter/Source/USDSchemas/Private/USDGeometryCacheTranslator.cpp
@@ -30,6 +30,8 @@
 #include "GeometryCacheUSDComponent.h"
 #include "GroomComponent.h"
 #include "HAL/IConsoleManager.h"
+#include "MaterialDomain.h"
+#include "Materials/Material.h"
 #include "StaticMeshAttributes.h"
 #include "StaticMeshOperations.h"
 #include "UObject/Package.h"
@@ -83,7 +85,7 @@ namespace UsdGeometryCacheTranslatorImpl
 			for (int32 LODSlotIndex = 0; LODSlotIndex < LODSlots.Num(); ++LODSlotIndex, ++SlotIndex)
 			{
 				const UsdUtils::FUsdPrimMaterialSlot& Slot = LODSlots[LODSlotIndex];
-				UMaterialInterface* Material = nullptr;
+				UMaterialInterface* Material = UMaterial::GetDefaultMaterial(MD_Surface);
 				if (UMaterialInterface** FoundMaterial = ResolvedMaterials.Find(&Slot))
 				{
 					Material = *FoundMaterial;
diff --git a/Engine/Plugins/Importers/USDImporter/Source/USDSchemas/Private/USDSkelRootTranslator.cpp b/Engine/Plugins/Importers/USDImporter/Source/USDSchemas/Private/USDSkelRootTranslator.cpp
--- a/Engine/Plugins/Importers/USDImporter/Source/USDSchemas/Private/USDSkelRootTranslator.cpp
+++ b/Engine/Plugins/Importers/USDImporter/Source/USDSchemas/Private/USDSkelRootTranslator.cpp
@@ -35,6 +35,7 @@
 #include "Engine/SkeletalMesh.h"
 #include "Engine/SkinnedAssetCommon.h"
 #include "GroomComponent.h"
+#include "MaterialDomain.h"
 #include "Materials/Material.h"
 #include "Materials/MaterialInstanceConstant.h"
 #include "Rendering/SkeletalMeshLODImporterData.h"
@@ -136,7 +137,7 @@ namespace UsdSkelRootTranslatorImpl
 			{
 				const UsdUtils::FUsdPrimMaterialSlot& Slot = LODSlots[ LODSlotIndex ];
 
-				UMaterialInterface* Material = nullptr;
+				UMaterialInterface* Material = UMaterial::GetDefaultMaterial(MD_Surface);
 
 				if ( UMaterialInterface** FoundMaterial = ResolvedMaterials.Find( &Slot ) )
 				{
