USD: Fix conditional always returning true when checking for subtree vertex count limits.

diff --git a/Engine/Plugins/Importers/USDImporter/Source/USDSchemas/Private/USDInfoCache.cpp b/Engine/Plugins/Importers/USDImporter/Source/USDSchemas/Private/USDInfoCache.cpp
--- a/Engine/Plugins/Importers/USDImporter/Source/USDSchemas/Private/USDInfoCache.cpp
+++ b/Engine/Plugins/Importers/USDImporter/Source/USDSchemas/Private/USDInfoCache.cpp
@@ -21,9 +21,10 @@
 	#include "pxr/usd/sdf/path.h"
 	#include "pxr/usd/usd/prim.h"
 	#include "pxr/usd/usdGeom/mesh.h"
-	#include "pxr/usd/usdShade/materialBindingAPI.h"
 	#include "pxr/usd/usdGeom/pointInstancer.h"
 	#include "pxr/usd/usdGeom/subset.h"
+	#include "pxr/usd/usdShade/materialBindingAPI.h"
+	#include "pxr/usd/usdSkel/root.h"
 #include "USDIncludesEnd.h"
 #endif // USE_USD_SDK
 
@@ -776,10 +777,9 @@ namespace UE::USDInfoCacheImpl::Private
 
 		// Only care about collapsing into a StaticMesh: We should always collapse into a SkeletalMesh as we have
 		// no real alternative for handling them
-		TSharedPtr<FUsdGeomXformableTranslator> XformableTranslator = StaticCastSharedPtr<FUsdGeomXformableTranslator>(Translator);
-		if (!XformableTranslator.IsValid())
+		if (UsdPrim.IsA<pxr::UsdSkelRoot>())
 		{
-			return false;
+			return true;
 		}
 
 		pxr::SdfPath UsdPrimPath = UsdPrim.GetPrimPath();
