USD: Limit UsdPreviewSurface reference materials to at most 4 TextureCoordinates.

Revert removal of error message done at cl 25304312.

diff --git a/Engine/Plugins/Importers/USDImporter/Source/USDSchemas/Private/MeshTranslationImpl.cpp b/Engine/Plugins/Importers/USDImporter/Source/USDSchemas/Private/MeshTranslationImpl.cpp
--- a/Engine/Plugins/Importers/USDImporter/Source/USDSchemas/Private/MeshTranslationImpl.cpp
+++ b/Engine/Plugins/Importers/USDImporter/Source/USDSchemas/Private/MeshTranslationImpl.cpp
@@ -35,7 +35,9 @@
 #include "pxr/usd/usdGeom/mesh.h"
 #include "USDIncludesEnd.h"
 
-#define UNUSED_UV_INDEX MAX_MESH_TEXTURE_COORDS_MD
+#define UNUSED_UV_INDEX USD_PREVIEW_SURFACE_MAX_UV_SETS
+
+static_assert(USD_PREVIEW_SURFACE_MAX_UV_SETS <= MAX_MESH_TEXTURE_COORDS_MD, "UsdPreviewSurface materials can only have up to as many UV sets as MeshDescription supports!");
 
 namespace UE::MeshTranslationImplInternal::Private
 {
@@ -170,7 +172,7 @@ namespace UE::MeshTranslationImplInternal::Private
 		);
 
 		TArray<TSet<FString>> UVIndexToMeshPrimvars;
-		UVIndexToMeshPrimvars.SetNum(MAX_MESH_TEXTURE_COORDS_MD);
+		UVIndexToMeshPrimvars.SetNum(USD_PREVIEW_SURFACE_MAX_UV_SETS);
 		for (const TPair<FString, int32>& MeshPair : MeshPrimvarToUVIndex)
 		{
 			if (UVIndexToMeshPrimvars.IsValidIndex(MeshPair.Value))
@@ -243,7 +245,7 @@ namespace UE::MeshTranslationImplInternal::Private
 			if (const int32* FoundMeshUVIndex = MeshPrimvarToUVIndex.Find(MaterialPrimvar))
 			{
 				int32 MeshUVIndex = *FoundMeshUVIndex;
-				if (MeshUVIndex >= 0 && MeshUVIndex < MAX_MESH_TEXTURE_COORDS_MD)
+				if (MeshUVIndex >= 0 && MeshUVIndex < USD_PREVIEW_SURFACE_MAX_UV_SETS)
 				{
 					CompatiblePrimvarAndUVIndexPairs.Add(TPair<FString, int32>{MaterialPrimvar, MeshUVIndex});
 					bFoundUVIndex = true;
diff --git a/Engine/Plugins/Importers/USDImporter/Source/USDUtilities/Private/USDConversionUtils.cpp b/Engine/Plugins/Importers/USDImporter/Source/USDUtilities/Private/USDConversionUtils.cpp
--- a/Engine/Plugins/Importers/USDImporter/Source/USDUtilities/Private/USDConversionUtils.cpp
+++ b/Engine/Plugins/Importers/USDImporter/Source/USDUtilities/Private/USDConversionUtils.cpp
@@ -775,7 +775,7 @@ TArray<TUsdStore<pxr::UsdGeomPrimvar>> UsdUtils::AssemblePrimvarsIntoUVSets(
 
 				if (!PrimvarsByUVIndex.IsValidIndex(TargetUVIndex))
 				{
-					if (TargetUVIndex < MAX_MESH_TEXTURE_COORDS_MD)
+					if (TargetUVIndex < USD_PREVIEW_SURFACE_MAX_UV_SETS)
 					{
 						PrimvarsByUVIndex.SetNum(TargetUVIndex + 1);
 					}
@@ -840,8 +840,8 @@ TMap<FString, int32> UsdUtils::CombinePrimvarsIntoUVSets(
 		}
 	);
 
-	// We can only have up to MAX_MESH_TEXTURE_COORDS_MD UV sets
-	SortedPrimvars.SetNum(FMath::Min(SortedPrimvars.Num(), (int32)MAX_MESH_TEXTURE_COORDS_MD));
+	// We can only have up to USD_PREVIEW_SURFACE_MAX_UV_SETS UV sets
+	SortedPrimvars.SetNum(FMath::Min(SortedPrimvars.Num(), (int32)USD_PREVIEW_SURFACE_MAX_UV_SETS));
 
 	TMap<FString, int32> PrimvarToUVIndex;
 	PrimvarToUVIndex.Reserve(SortedPrimvars.Num());
diff --git a/Engine/Plugins/Importers/USDImporter/Source/USDUtilities/Private/USDGeomMeshConversion.cpp b/Engine/Plugins/Importers/USDImporter/Source/USDUtilities/Private/USDGeomMeshConversion.cpp
--- a/Engine/Plugins/Importers/USDImporter/Source/USDUtilities/Private/USDGeomMeshConversion.cpp
+++ b/Engine/Plugins/Importers/USDImporter/Source/USDUtilities/Private/USDGeomMeshConversion.cpp
@@ -1106,7 +1106,7 @@ bool UsdToUnreal::ConvertGeomMesh(
 				NumUVs = FMath::Max<int32>(ProvidedNumUVSets.GetValue(), NumUVs);
 			}
 
-			NumUVs = FMath::Min<int32>( MAX_MESH_TEXTURE_COORDS_MD, NumUVs );
+			NumUVs = FMath::Min<int32>( USD_PREVIEW_SURFACE_MAX_UV_SETS, NumUVs );
 			// At least one UV set must exist.
 			NumUVs = FMath::Max<int32>(1, NumUVs);
 
diff --git a/Engine/Plugins/Importers/USDImporter/Source/USDUtilities/Public/USDConversionUtils.h b/Engine/Plugins/Importers/USDImporter/Source/USDUtilities/Public/USDConversionUtils.h
--- a/Engine/Plugins/Importers/USDImporter/Source/USDUtilities/Public/USDConversionUtils.h
+++ b/Engine/Plugins/Importers/USDImporter/Source/USDUtilities/Public/USDConversionUtils.h
@@ -153,7 +153,7 @@ namespace UsdUtils
 	 */
 	USDUTILITIES_API TArray<TUsdStore<pxr::UsdGeomPrimvar>> GetUVSetPrimvars(
 		const pxr::UsdGeomMesh& UsdMesh,
-		int32 MaxNumPrimvars = MAX_MESH_TEXTURE_COORDS_MD
+		int32 MaxNumPrimvars = USD_PREVIEW_SURFACE_MAX_UV_SETS
 	);
 
 	/**
@@ -182,7 +182,7 @@ namespace UsdUtils
 	);
 
 	/**
-	 * Rearranges an array of primvars into another array of up to MAX_MESH_TEXTURE_COORDS_MD primvars, describing which
+	 * Rearranges an array of primvars into another array of up to USD_PREVIEW_SURFACE_MAX_UV_SETS primvars, describing which
 	 * primvars should be used for each UV set according to the provided AllowedPrimvarsToUVIndex mapping.
 	 *
 	 * The intent is to use this when e.g. a Mesh prim is to be collapsed together with other Mesh prims into a single
@@ -215,7 +215,7 @@ namespace UsdUtils
 	/**
 	 * Given an array of primvars, and an array of which of those are preferred (AllPrimvars should also include
 	 * PreferredPrimvars), this will sort lexicographically and return a prioritized PrimvarToUVIndex assignment map of
-	 * up to MAX_MESH_TEXTURE_COORDS_MD UV indices
+	 * up to USD_PREVIEW_SURFACE_MAX_UV_SETS UV indices
 	 */
 	USDUTILITIES_API TMap<FString, int32> CombinePrimvarsIntoUVSets(
 		const TSet<FString>& AllPrimvars,
diff --git a/Engine/Plugins/Importers/USDImporter/Source/UnrealUSDWrapper/Public/UnrealUSDWrapper.h b/Engine/Plugins/Importers/USDImporter/Source/UnrealUSDWrapper/Public/UnrealUSDWrapper.h
--- a/Engine/Plugins/Importers/USDImporter/Source/UnrealUSDWrapper/Public/UnrealUSDWrapper.h
+++ b/Engine/Plugins/Importers/USDImporter/Source/UnrealUSDWrapper/Public/UnrealUSDWrapper.h
@@ -24,6 +24,10 @@
 
 #include "UnrealUSDWrapper.generated.h"
 
+// We are capped to 4 because the geometry cache and skeletal mesh + morph target shaders have a hard-coded limit of 4 texture coordinates, so our
+// material instances can't cross that limit or else they wouldn't compile when assigned to those meshes
+#define USD_PREVIEW_SURFACE_MAX_UV_SETS 4
+
 #if USE_USD_SDK
 PXR_NAMESPACE_OPEN_SCOPE
 	class GfMatrix4d;
diff --git a/Engine/Shaders/Private/GeometryCacheVertexFactory.ush b/Engine/Shaders/Private/GeometryCacheVertexFactory.ush
--- a/Engine/Shaders/Private/GeometryCacheVertexFactory.ush
+++ b/Engine/Shaders/Private/GeometryCacheVertexFactory.ush
@@ -64,6 +64,10 @@ struct FVertexFactoryInput
 	#if NUM_MATERIAL_TEXCOORDS_VERTEX > 3
 		float2	TexCoords3 : ATTRIBUTE8;
 	#endif
+
+	#if NUM_MATERIAL_TEXCOORDS_VERTEX > 4
+		#error Too many texture coordinate sets defined on Geometry cache vertex input. Max: 4.
+	#endif
 #endif
 
 	VF_GPUSCENE_DECLARE_INPUT_BLOCK(13)
