USD: Add events for when stage is loaded or unloaded with the AUsdStageActor.

diff --git a/Engine/Plugins/Importers/USDImporter/Source/USDStage/Private/USDStageActor.cpp b/Engine/Plugins/Importers/USDImporter/Source/USDStage/Private/USDStageActor.cpp
--- a/Engine/Plugins/Importers/USDImporter/Source/USDStage/Private/USDStageActor.cpp
+++ b/Engine/Plugins/Importers/USDImporter/Source/USDStage/Private/USDStageActor.cpp
@@ -2275,6 +2275,8 @@ void AUsdStageActor::LoadUsdStage()
 	}
 #endif // WITH_EDITOR
 
+	OnStageLoaded.Broadcast();
+
 	// Log time spent to load the stage
 	double ElapsedSeconds = FPlatformTime::ToSeconds64(FPlatformTime::Cycles64() - StartTime);
 
@@ -2340,15 +2342,15 @@ void AUsdStageActor::UnloadUsdStage()
 	}
 #endif // WITH_EDITOR
 
-	CloseUsdStage();
-
 	if (UsdAssetCache)
 	{
 		UsdAssetCache->RemoveAllAssetReferences(this);
 		UsdAssetCache->RefreshStorage();
 	}
 
-	OnStageChanged.Broadcast();
+	OnStageUnloaded.Broadcast();
+
+	CloseUsdStage();
 }
 
 void AUsdStageActor::EnsureAssetCache()
diff --git a/Engine/Plugins/Importers/USDImporter/Source/USDStage/Public/USDStageActor.h b/Engine/Plugins/Importers/USDImporter/Source/USDStage/Public/USDStageActor.h
--- a/Engine/Plugins/Importers/USDImporter/Source/USDStage/Public/USDStageActor.h
+++ b/Engine/Plugins/Importers/USDImporter/Source/USDStage/Public/USDStageActor.h
@@ -99,6 +99,8 @@ public:
 	FOnStageActorEvent OnPreStageChanged;
 	FOnStageActorEvent OnStageChanged;
 	FOnStageActorEvent OnActorDestroyed;
+	FOnStageActorEvent OnStageLoaded;
+	FOnStageActorEvent OnStageUnloaded;
 
 	DECLARE_EVENT_TwoParams(AUsdStageActor, FOnPrimChanged, const FString&, bool);
 	FOnPrimChanged OnPrimChanged;
