USD: Fixed random skeletal animation test failure due to uninitialized variable.

diff --git a/Engine/Plugins/Importers/USDImporter/Source/USDUtilities/Private/USDSkeletalDataConversion.cpp b/Engine/Plugins/Importers/USDImporter/Source/USDUtilities/Private/USDSkeletalDataConversion.cpp
--- a/Engine/Plugins/Importers/USDImporter/Source/USDUtilities/Private/USDSkeletalDataConversion.cpp
+++ b/Engine/Plugins/Importers/USDImporter/Source/USDUtilities/Private/USDSkeletalDataConversion.cpp
@@ -2727,6 +2727,7 @@ USkeletalMesh* UsdToUnreal::GetSkeletalMeshFromImportData(
 
 void UsdUtils::ResolveWeightsForBlendShape( const UsdUtils::FUsdBlendShape& InBlendShape, float InWeight, float& OutMainWeight, TArray<float>& OutInbetweenWeights )
 {
+	OutMainWeight = 0.0f;
 	int32 NumInbetweens = InBlendShape.Inbetweens.Num();
 	if ( NumInbetweens == 0 )
 	{
