USD: Fix asset cache serialization behavior for transactions and duplication.

diff --git a/Engine/Plugins/Importers/USDImporter/Source/USDClasses/Private/USDAssetCache2.cpp b/Engine/Plugins/Importers/USDImporter/Source/USDClasses/Private/USDAssetCache2.cpp
--- a/Engine/Plugins/Importers/USDImporter/Source/USDClasses/Private/USDAssetCache2.cpp
+++ b/Engine/Plugins/Importers/USDImporter/Source/USDClasses/Private/USDAssetCache2.cpp
@@ -30,7 +30,7 @@ static FAutoConsoleVariableRef COnDemandCachedAssetLoading(
 	TEXT("When true will cause the USD Asset Cache to only load cached UObjects from disk once they are needed, regardless of when the actual Asset Cache asset itself is loaded."));
 
 static int32 GCurrentCacheVersion = 1;
-static int32 GCurrentAssetInfoVersion = 1;
+static int32 GCurrentAssetInfoVersion = 2;
 
 namespace UE::AssetCache::Private
 {
@@ -549,17 +549,37 @@ FArchive& operator<<(FArchive& Ar, UUsdAssetCache2::ECacheStorageType& Type)
 
 void UUsdAssetCache2::FCachedAssetInfo::Serialize(FArchive& Ar, UObject* Owner)
 {
+	int32 SavedAssetInfoVersion = INDEX_NONE;
 	if (Ar.IsSaving())
 	{
 		Ar << GCurrentAssetInfoVersion;
+		SavedAssetInfoVersion = GCurrentAssetInfoVersion;
 	}
 	else if (Ar.IsLoading())
 	{
-		int32 SavedAssetInfoVersion = INDEX_NONE;
 		Ar << SavedAssetInfoVersion;
 		ensure(SavedAssetInfoVersion <= GCurrentAssetInfoVersion);
 	}
 
+	UE_LOG(
+		LogUsd,
+		Verbose,
+		TEXT("Serializing Asset Info '%s' (%s) with version %d, storage type: %d, disk bytes: %d, memory bytes: %d, bulkdata size: %d, saving? %d, persistent? %d, "
+			 "transacting? %d, duplicating? %d, cooking? %d"),
+		*AssetName,
+		*Hash,
+		SavedAssetInfoVersion,
+		CurrentStorageType,
+		SizeOnDiskInBytes,
+		SizeOnMemoryInBytes,
+		BulkData.GetBulkDataSize(),
+		Ar.IsSaving(),
+		Ar.IsPersistent(),
+		Ar.IsTransacting(),
+		(Ar.GetPortFlags() & (PPF_DuplicateForPIE | PPF_Duplicate)),
+		Ar.IsCooking()
+	);
+
 	Ar << Hash;
 	Ar << AssetClassName;
 	Ar << AssetName;
@@ -570,10 +590,20 @@ void UUsdAssetCache2::FCachedAssetInfo::Serialize(FArchive& Ar, UObject* Owner)
 	Ar << Dependencies;
 	Ar << Consumers;
 
-	// Very important: This flag ensures that when we call Serialize below we won't actually
-	// load the full bulkdata right away, so that we can do it on-demand
-	BulkData.SetBulkDataFlags(BULKDATA_Force_NOT_InlinePayload);
-	BulkData.Serialize(Ar, Owner);
+	// We do need to serialize this in case we do too many undo/redo operations without a cache resync between them
+	if (SavedAssetInfoVersion > 1)
+	{
+		Ar << CurrentStorageType;
+	}
+
+	// We only need to read/write to our bulkdata when saving to disk
+	if (Ar.IsPersistent())
+	{
+		// Very important: This flag ensures that when we call Serialize below we won't actually
+		// load the full bulkdata right away, so that we can do it on-demand
+		BulkData.SetBulkDataFlags(BULKDATA_Force_NOT_InlinePayload);
+		BulkData.Serialize(Ar, Owner);
+	}
 }
 
 UUsdAssetCache2::UUsdAssetCache2()
@@ -1588,112 +1618,134 @@ void UUsdAssetCache2::Serialize(FArchive& Ar)
 
 			FWriteScopeLock Lock(RWLock);
 
-			TSet<FSoftObjectPath> PersistentAssets;
-			for (TLruCache<FSoftObjectPath, FCachedAssetInfo>::TIterator Iter{ LRUCache }; Iter; ++Iter)
+			TSet<FSoftObjectPath> PersistentAssetsWithInfo;
+			if (Ar.IsPersistent())
 			{
-				const FCachedAssetInfo& Info = Iter.Value();
-				if (Info.CurrentStorageType == ECacheStorageType::Persistent)
+				for (TLruCache<FSoftObjectPath, FCachedAssetInfo>::TIterator Iter{LRUCache}; Iter; ++Iter)
 				{
-					PersistentAssets.Add(Iter.Key());
+					const FCachedAssetInfo& Info = Iter.Value();
+					if (Info.CurrentStorageType == ECacheStorageType::Persistent)
+					{
+						PersistentAssetsWithInfo.Add(Iter.Key());
+					}
 				}
 			}
-			int32 NumPersistent = PersistentAssets.Num();
-			Ar << NumPersistent;
 
-			UE_LOG(LogUsd, Verbose, TEXT("Serializing USD Asset Cache '%s'"),
-				*GetPathName()
+			int32 NumInfos = Ar.IsPersistent() ? PersistentAssetsWithInfo.Num() : LRUCache.Num();
+			Ar << NumInfos;
+
+			UE_LOG(LogUsd, Verbose, TEXT("Serializing USD Asset Cache '%s' with %d asset infos. Persistent? %d, transacting? %d, duplicating? %d, cooking? %d"),
+				*GetPathName(),
+				NumInfos,
+				Ar.IsPersistent(),
+				Ar.IsTransacting(),
+				(Ar.GetPortFlags() & (PPF_DuplicateForPIE | PPF_Duplicate)),
+				Ar.IsCooking()
 			);
 
 			// By using the iterators we'll traverse the most recently used object first
-			for (TLruCache<FSoftObjectPath, FCachedAssetInfo>::TIterator Iter{ LRUCache }; Iter; ++Iter)
+			for (TLruCache<FSoftObjectPath, FCachedAssetInfo>::TIterator Iter{LRUCache}; Iter; ++Iter)
 			{
 				FCachedAssetInfo& Info = Iter.Value();
 				FSoftObjectPath& AssetPath = Iter.Key();
 				ensure(!AssetPath.IsNull());
 
-				if (Info.CurrentStorageType != ECacheStorageType::Persistent)
+				if (Ar.IsPersistent() && Info.CurrentStorageType != ECacheStorageType::Persistent)
 				{
+					UE_LOG(
+						LogUsd,
+						Verbose,
+						TEXT("Skipping serialization of info on asset '%s' (%s) as it will not be persisted"),
+						*Info.AssetName,
+						*Info.Hash
+					);
 					continue;
 				}
 
-				// This is the main way in which assets get serialized into the bulkdata
-				if (UObject* Asset = AssetStorage.FindRef(Info.Hash))
+				// This is the main way in which assets get serialized into the memory bulkdata, which we only ever need to do
+				// when serializing to disk (which will happen within the Info.Serialize call below)
+				if (Ar.IsPersistent())
 				{
-					// Serialize the referenced asset into a byte buffer
-					TArray<uint8> Buffer;
-					UE::AssetCache::Private::SerializeObjectAndSubObjects(Asset, Buffer, Ar);
-					SIZE_T NumBytes = Buffer.Num();
-
-					// Copy the byte buffer into the Info's bulkdata in memory
-					// TODO: Maybe it's possible to avoid this copy by having a custom object writer
-					// that can write directly to the realloc'd bulkdata bytes
-					Info.BulkData.Lock(LOCK_READ_WRITE);
+					UObject* Asset = AssetStorage.FindRef(Info.Hash);
+					if (ensure(Asset))
 					{
-						void* BulkDataBytes = Info.BulkData.Realloc(Buffer.Num());
-						FMemory::Memcpy(BulkDataBytes, Buffer.GetData(), Buffer.Num());
+						// Serialize the referenced asset into a byte buffer
+						TArray<uint8> Buffer;
+						UE::AssetCache::Private::SerializeObjectAndSubObjects(Asset, Buffer, Ar);
+
+						// Copy the byte buffer into the Info's bulkdata in memory
+						// TODO: Maybe it's possible to avoid this copy by having a custom object writer
+						// that can write directly to the realloc'd bulkdata bytes
+						Info.BulkData.Lock(LOCK_READ_WRITE);
+						{
+							void* BulkDataBytes = Info.BulkData.Realloc(Buffer.Num());
+							FMemory::Memcpy(BulkDataBytes, Buffer.GetData(), Buffer.Num());
+						}
+						Info.BulkData.Unlock();
+
+						// Now we know *exactly* how big the asset is on disk
+						Info.SizeOnDiskInBytes = Buffer.Num();
+
+						UE_LOG(
+							LogUsd,
+							Verbose,
+							TEXT("Serializing asset '%s' (%s) to %d bulk data bytes"),
+							*Info.AssetName,
+							*Info.Hash,
+							Info.SizeOnDiskInBytes
+						);
 					}
-					Info.BulkData.Unlock();
-
-					// Now we know *exactly* how big the asset is on disk now
-					Info.SizeOnDiskInBytes = Buffer.Num();
 				}
 
-				UE_LOG(LogUsd, Verbose, TEXT("\tPersistent asset '%s' ('%s') with hash '%s' (%u disk bytes)"),
-					*Info.AssetName,
-					*Info.AssetClassName.ToString(),
-					*Info.Hash,
-					Info.SizeOnDiskInBytes
-				);
-
 				Ar << AssetPath;
 
-				// We only want to serialize dependencies on persistent assets, or else when we deserialize this
+				// When persisting, we only want to serialize dependencies on persistent assets, or else when we deserialize this
 				// we'll have a bunch of dependencies on assets that don't exist
-				TSet<FSoftObjectPath> PersistentDependencies = Info.Dependencies.Intersect(PersistentAssets);
-				TSet<FSoftObjectPath> PersistentConsumers = Info.Consumers.Intersect(PersistentAssets);
-				Swap(Info.Dependencies, PersistentDependencies);
-				Swap(Info.Consumers, PersistentConsumers);
+				TSet<FSoftObjectPath> PersistentDependencies;
+				TSet<FSoftObjectPath> PersistentConsumers;
+				if (Ar.IsPersistent())
+				{
+					PersistentDependencies = Info.Dependencies.Intersect(PersistentAssetsWithInfo);
+					PersistentConsumers = Info.Consumers.Intersect(PersistentAssetsWithInfo);
+					Swap(Info.Dependencies, PersistentDependencies);
+					Swap(Info.Consumers, PersistentConsumers);
+				}
+
+				Info.Serialize(Ar, this);
+
+				if (Ar.IsPersistent())
 				{
-					Info.CurrentStorageType = ECacheStorageType::Persistent;
-					Info.Serialize(Ar, this);
+					Swap(Info.Dependencies, PersistentDependencies);
+					Swap(Info.Consumers, PersistentConsumers);
 				}
-				Swap(Info.Dependencies, PersistentDependencies);
-				Swap(Info.Consumers, PersistentConsumers);
 			}
 		}
 		else if (Ar.IsLoading())
 		{
 			FWriteScopeLock Lock(RWLock);
 
-			int32 NumPersistent = INDEX_NONE;
-			Ar << NumPersistent;
-			ensure(NumPersistent >= 0);
+			int32 NumInfos = INDEX_NONE;
+			Ar << NumInfos;
+			ensure(NumInfos >= 0);
 
 			TArray<FCachedAssetInfo> MostToLeastRecent;
-			MostToLeastRecent.Reserve(NumPersistent);
+			MostToLeastRecent.Reserve(NumInfos);
 
 			TArray<FSoftObjectPath> AssetPaths;
-			AssetPaths.Reserve(NumPersistent);
+			AssetPaths.Reserve(NumInfos);
 
-			UE_LOG(LogUsd, Verbose, TEXT("Deserializing USD Asset Cache '%s'"),
-				*GetPathName()
+			UE_LOG(LogUsd, Verbose, TEXT("Deserializing USD Asset Cache '%s' with %d asset infos"),
+				*GetPathName(),
+				NumInfos
 			);
 
-			for (int32 Index = 0; Index < NumPersistent; ++Index)
+			for (int32 Index = 0; Index < NumInfos; ++Index)
 			{
 				FSoftObjectPath& AssetPath = AssetPaths.Emplace_GetRef();
 				Ar << AssetPath;
 
 				FCachedAssetInfo& NewEl = MostToLeastRecent.Emplace_GetRef();
 				NewEl.Serialize(Ar, this);
-				NewEl.CurrentStorageType = ECacheStorageType::Persistent;
-
-				UE_LOG(LogUsd, Verbose, TEXT("\tPersistent asset '%s' ('%s') with hash '%s' (%u disk bytes)"),
-					*NewEl.AssetName,
-					*NewEl.AssetClassName.ToString(),
-					*NewEl.Hash,
-					NewEl.SizeOnDiskInBytes
-				);
 
 				PendingPersistentStorage.Add(NewEl.Hash, AssetPath);
 			}
@@ -1703,7 +1755,7 @@ void UUsdAssetCache2::Serialize(FArchive& Ar)
 			LRUCache.Empty(MaxElements);
 
 			// Push them back-to-front into the LRU cache to make sure the most recent one is added last
-			for (int32 Index = NumPersistent - 1; Index >= 0; --Index)
+			for (int32 Index = NumInfos - 1; Index >= 0; --Index)
 			{
 				LRUCache.Add(MoveTemp(AssetPaths[Index]), MoveTemp(MostToLeastRecent[Index]));
 			}
