USD: Fix material override assignments for SkeletalMesh using materials without any textures.

diff --git a/Engine/Plugins/Importers/USDImporter/Source/USDSchemas/Private/USDSkelRootTranslator.cpp b/Engine/Plugins/Importers/USDImporter/Source/USDSchemas/Private/USDSkelRootTranslator.cpp
--- a/Engine/Plugins/Importers/USDImporter/Source/USDSchemas/Private/USDSkelRootTranslator.cpp
+++ b/Engine/Plugins/Importers/USDImporter/Source/USDSchemas/Private/USDSkelRootTranslator.cpp
@@ -771,6 +771,28 @@ namespace UsdSkelRootTranslatorImpl
 			LODIndexToAssignments.Add( MoveTemp( Entry.Value ) );
 		}
 
+		// Stash our mesh PrimvarToUVIndex into the assignment info, as that's where ResolveMaterialAssignmentInfo will look for it
+		UUsdMeshAssetUserData* UserData = nullptr;
+		if (USkeletalMeshComponent* SkeletalMeshComponent = Cast<USkeletalMeshComponent>(&MeshComponent))
+		{
+			if (USkeletalMesh* Mesh = SkeletalMeshComponent->GetSkeletalMeshAsset())
+			{
+				UserData = Mesh->GetAssetUserData<UUsdMeshAssetUserData>();
+			}
+		}
+		if (UserData)
+		{
+			LODIndexToAssignments[0].PrimvarToUVIndex = UserData->PrimvarToUVIndex;
+		}
+		else
+		{
+			ensureMsgf(
+				UserData,
+				TEXT("Expected component '%s''s SkeletalMesh to have an instance of UUsdMeshAssetUserData at this point!"),
+				*MeshComponent.GetPathName()
+			);
+		}
+
 		TMap<const UsdUtils::FUsdPrimMaterialSlot*, UMaterialInterface*> ResolvedMaterials = MeshTranslationImpl::ResolveMaterialAssignmentInfo(
 			ValidSkelRootPrim,
 			LODIndexToAssignments,
