USD: Remove stuff deprecated in 5.1

diff --git a/Engine/Plugins/Importers/USDImporter/Source/USDExporter/Private/USDExporterBlueprintLibrary.cpp b/Engine/Plugins/Importers/USDImporter/Source/USDExporter/Private/USDExporterBlueprintLibrary.cpp
deleted file mode 100644
diff --git a/Engine/Plugins/Importers/USDImporter/Source/USDExporter/Public/USDExporterBlueprintLibrary.h b/Engine/Plugins/Importers/USDImporter/Source/USDExporter/Public/USDExporterBlueprintLibrary.h
deleted file mode 100644
diff --git a/Engine/Plugins/Importers/USDImporter/Source/USDUtilities/Private/USDGeomMeshConversion.cpp b/Engine/Plugins/Importers/USDImporter/Source/USDUtilities/Private/USDGeomMeshConversion.cpp
--- a/Engine/Plugins/Importers/USDImporter/Source/USDUtilities/Private/USDGeomMeshConversion.cpp
+++ b/Engine/Plugins/Importers/USDImporter/Source/USDUtilities/Private/USDGeomMeshConversion.cpp
@@ -1304,61 +1304,6 @@ bool UsdToUnreal::ConvertGeomMesh(
 	return true;
 }
 
-// Deprecated
-bool UsdToUnreal::ConvertGeomMesh(
-	const pxr::UsdTyped& UsdSchema,
-	FMeshDescription& MeshDescription,
-	UsdUtils::FUsdPrimMaterialAssignmentInfo& MaterialAssignments,
-	const pxr::UsdTimeCode TimeCode,
-	const pxr::TfToken& RenderContext
-)
-{
-	FUsdMeshConversionOptions Options;
-	Options.TimeCode = TimeCode;
-	Options.RenderContext = RenderContext;
-
-	return ConvertGeomMesh( pxr::UsdGeomMesh{ UsdSchema }, MeshDescription, MaterialAssignments, Options );
-}
-
-// Deprecated
-bool UsdToUnreal::ConvertGeomMesh(
-	const pxr::UsdTyped& UsdSchema,
-	FMeshDescription& MeshDescription,
-	UsdUtils::FUsdPrimMaterialAssignmentInfo& MaterialAssignments,
-	const FTransform& AdditionalTransform,
-	const pxr::UsdTimeCode TimeCode,
-	const pxr::TfToken& RenderContext
-)
-{
-	FUsdMeshConversionOptions Options;
-	Options.AdditionalTransform = AdditionalTransform;
-	Options.TimeCode = TimeCode;
-	Options.RenderContext = RenderContext;
-
-	return ConvertGeomMesh( pxr::UsdGeomMesh{ UsdSchema }, MeshDescription, MaterialAssignments, Options );
-}
-
-// Deprecated
-bool UsdToUnreal::ConvertGeomMesh(
-	const pxr::UsdTyped& UsdSchema,
-	FMeshDescription& MeshDescription,
-	UsdUtils::FUsdPrimMaterialAssignmentInfo& MaterialAssignments,
-	const FTransform& AdditionalTransform,
-	const TMap< FString, TMap< FString, int32 > >& MaterialToPrimvarsUVSetNames,
-	const pxr::UsdTimeCode TimeCode,
-	const pxr::TfToken& RenderContext,
-	bool bMergeIdenticalMaterialSlots
-)
-{
-	FUsdMeshConversionOptions Options;
-	Options.AdditionalTransform = AdditionalTransform;
-	Options.TimeCode = TimeCode;
-	Options.RenderContext = RenderContext;
-	Options.bMergeIdenticalMaterialSlots = bMergeIdenticalMaterialSlots;
-
-	return ConvertGeomMesh( pxr::UsdGeomMesh{ UsdSchema }, MeshDescription, MaterialAssignments, Options );
-}
-
 bool UsdToUnreal::ConvertPointInstancerToMesh(
 	const pxr::UsdGeomPointInstancer& PointInstancer,
 	FMeshDescription& OutMeshDescription,
@@ -1585,34 +1530,6 @@ bool UsdToUnreal::ConvertGeomMeshHierarchy(
 	);
 }
 
-// Deprecated
-bool UsdToUnreal::ConvertGeomMeshHierarchy(
-	const pxr::UsdPrim& Prim,
-	const pxr::UsdTimeCode& TimeCode,
-	const EUsdPurpose PurposesToLoad,
-	const pxr::TfToken& RenderContext,
-	const TMap< FString, TMap<FString, int32> >& MaterialToPrimvarToUVIndex,
-	FMeshDescription& OutMeshDescription,
-	UsdUtils::FUsdPrimMaterialAssignmentInfo& OutMaterialAssignments,
-	bool bSkipRootPrimTransformAndVisibility,
-	bool bMergeIdenticalMaterialSlots
-)
-{
-	FUsdMeshConversionOptions Options;
-	Options.TimeCode = TimeCode;
-	Options.PurposesToLoad = PurposesToLoad;
-	Options.RenderContext = RenderContext;
-	Options.bMergeIdenticalMaterialSlots = bMergeIdenticalMaterialSlots;
-
-	return ConvertGeomMeshHierarchy(
-		Prim,
-		OutMeshDescription,
-		OutMaterialAssignments,
-		Options,
-		bSkipRootPrimTransformAndVisibility
-	);
-}
-
 UMaterialInstanceDynamic* UsdUtils::CreateDisplayColorMaterialInstanceDynamic( const UsdUtils::FDisplayColorMaterial& DisplayColorDescription )
 {
 	const UUsdProjectSettings* Settings = GetDefault<UUsdProjectSettings>();
diff --git a/Engine/Plugins/Importers/USDImporter/Source/USDUtilities/Private/USDShadeConversion.cpp b/Engine/Plugins/Importers/USDImporter/Source/USDUtilities/Private/USDShadeConversion.cpp
--- a/Engine/Plugins/Importers/USDImporter/Source/USDUtilities/Private/USDShadeConversion.cpp
+++ b/Engine/Plugins/Importers/USDImporter/Source/USDUtilities/Private/USDShadeConversion.cpp
@@ -3103,55 +3103,6 @@ bool UsdUtils::IsMaterialTranslucent( const pxr::UsdShadeMaterial& UsdShadeMater
 	return bHasOpacityConnection && ( bIsTranslucentFloat || bHasBoundTexture );
 }
 
-bool UsdUtils::IsMaterialUsingUDIMs( const pxr::UsdShadeMaterial& UsdShadeMaterial )
-{
-	FScopedUsdAllocs UsdAllocs;
-
-	pxr::UsdShadeShader SurfaceShader = UsdShadeMaterial.ComputeSurfaceSource();
-	if ( !SurfaceShader )
-	{
-		return false;
-	}
-
-	for ( const pxr::UsdShadeInput& ShadeInput : SurfaceShader.GetInputs() )
-	{
-		pxr::UsdShadeConnectableAPI Source;
-		pxr::TfToken SourceName;
-		pxr::UsdShadeAttributeType AttributeType;
-
-		if ( pxr::UsdShadeConnectableAPI::GetConnectedSource( ShadeInput.GetAttr(), &Source, &SourceName, &AttributeType ) )
-		{
-			pxr::UsdShadeInput FileInput;
-
-			// UsdUVTexture: Get its file input
-			if ( AttributeType == pxr::UsdShadeAttributeType::Output )
-			{
-				FileInput = Source.GetInput( UnrealIdentifiers::File );
-			}
-			// Check if we are being directly passed an asset
-			else
-			{
-				FileInput = Source.GetInput( SourceName );
-			}
-
-			if ( FileInput && FileInput.GetTypeName() == pxr::SdfValueTypeNames->Asset ) // Check that FileInput is of type Asset
-			{
-				pxr::SdfAssetPath TextureAssetPath;
-				FileInput.GetAttr().Get< pxr::SdfAssetPath >( &TextureAssetPath );
-
-				FString TexturePath = UsdToUnreal::ConvertString( TextureAssetPath.GetAssetPath() );
-
-				if ( TexturePath.Contains( TEXT("<UDIM>") ) )
-				{
-					return true;
-				}
-			}
-		}
-	}
-
-	return false;
-}
-
 FSHAHash UsdUtils::HashShadeMaterial( const pxr::UsdShadeMaterial& UsdShadeMaterial, const pxr::TfToken& RenderContext )
 {
 	FScopedUsdAllocs UsdAllocs;
diff --git a/Engine/Plugins/Importers/USDImporter/Source/USDUtilities/Private/USDSkeletalDataConversion.cpp b/Engine/Plugins/Importers/USDImporter/Source/USDUtilities/Private/USDSkeletalDataConversion.cpp
--- a/Engine/Plugins/Importers/USDImporter/Source/USDUtilities/Private/USDSkeletalDataConversion.cpp
+++ b/Engine/Plugins/Importers/USDImporter/Source/USDUtilities/Private/USDSkeletalDataConversion.cpp
@@ -1920,101 +1920,6 @@ bool UsdToUnreal::ConvertSkinnedMesh(
 	return bResult;
 }
 
-bool UsdToUnreal::ConvertSkinnedMesh(
-	const pxr::UsdSkelSkinningQuery& UsdSkinningQuery,
-	const FTransform& AdditionalTransform,
-	FSkeletalMeshImportData& SkelMeshImportData,
-	TArray< UsdUtils::FUsdPrimMaterialSlot >& MaterialAssignments,
-	const TMap< FString, TMap< FString, int32 > >& MaterialToPrimvarsUVSetNames,
-	const pxr::TfToken& RenderContext,
-	const pxr::TfToken& MaterialPurpose
-)
-{
-	FScopedUsdAllocs UEAllocs;
-
-	// We must use a SkeletonQuery to properly convert the skinned mesh to the target pose. If we weren't provided one,
-	// we must fetch it manually
-
-	pxr::UsdPrim SkinnedPrim = UsdSkinningQuery.GetPrim();
-	if ( !SkinnedPrim )
-	{
-		return false;
-	}
-
-	bool bFoundSkelRoot = false;
-	pxr::UsdPrim ParentPrim = SkinnedPrim.GetParent();
-	while ( ParentPrim && !ParentPrim.IsPseudoRoot() )
-	{
-		if ( ParentPrim.IsA<pxr::UsdSkelRoot>() )
-		{
-			bFoundSkelRoot = true;
-			break;
-		}
-
-		ParentPrim = ParentPrim.GetParent();
-	}
-	if ( !bFoundSkelRoot )
-	{
-		return false;
-	}
-
-	pxr::UsdSkelRoot SkelRoot{ ParentPrim };
-
-	pxr::UsdSkelCache SkelCache;
-	if ( !SkelCache.Populate( SkelRoot, pxr::UsdTraverseInstanceProxies() ) )
-	{
-		return false;
-	}
-
-	std::vector< pxr::UsdSkelBinding > SkeletonBindings;
-	if ( !SkelCache.ComputeSkelBindings( SkelRoot, &SkeletonBindings, pxr::UsdTraverseInstanceProxies() ) )
-	{
-		return false;
-	}
-
-	for ( const pxr::UsdSkelBinding& Binding : SkeletonBindings )
-	{
-		const pxr::UsdSkelSkeleton& Skeleton = Binding.GetSkeleton();
-		pxr::UsdSkelSkeletonQuery SkelQuery = SkelCache.GetSkelQuery( Skeleton );
-
-		// Double-check that we're talking about the same mesh, at least. Each mesh can only be bound to one
-		// skeleton so this is probably guaranteed to be the same SkelBinding
-		bool bFoundSkinningQuery = false;
-		for ( const pxr::UsdSkelSkinningQuery& SomeSkinningQuery : Binding.GetSkinningTargets() )
-		{
-			if ( SomeSkinningQuery.GetPrim() == SkinnedPrim )
-			{
-				bFoundSkinningQuery = true;
-				break;
-			}
-		}
-		if ( !bFoundSkinningQuery )
-		{
-			break;
-		}
-
-		UsdToUnreal::FUsdMeshConversionOptions Options;
-		Options.RenderContext = RenderContext;
-		Options.MaterialPurpose = MaterialPurpose;
-
-		UsdUtils::FUsdPrimMaterialAssignmentInfo Info;
-		Swap(Info.Slots, MaterialAssignments);
-		const bool bSuccess = ConvertSkinnedMesh(
-			UsdSkinningQuery,
-			SkelQuery,
-			SkelMeshImportData,
-			Info,
-			Options
-		);
-		Swap(Info.Slots, MaterialAssignments);
-
-		// We really only ever deal with the first binding, so we can return here
-		return bSuccess;
-	}
-
-	return false;
-}
-
 // Using UsdSkelSkeletonQuery instead of UsdSkelAnimQuery as it automatically does the joint remapping when we ask it to compute joint transforms.
 // It also initializes the joint transforms with the rest pose, if available, in case the animation doesn't provide data for all joints.
 bool UsdToUnreal::ConvertSkelAnim(
@@ -2551,16 +2456,6 @@ bool UsdToUnreal::ConvertBlendShape( const pxr::UsdSkelBlendShape& UsdBlendShape
 	return true;
 }
 
-bool UsdToUnreal::ConvertBlendShape( const pxr::UsdSkelBlendShape& UsdBlendShape, const FUsdStageInfo& StageInfo, const FTransform& AdditionalTransform, uint32 PointIndexOffset, TSet<FString>& UsedMorphTargetNames, UsdUtils::FBlendShapeMap& OutBlendShapes )
-{
-	return UsdToUnreal::ConvertBlendShape( UsdBlendShape, StageInfo, 0, PointIndexOffset, UsedMorphTargetNames, OutBlendShapes );
-}
-
-bool UsdToUnreal::ConvertBlendShape( const pxr::UsdSkelBlendShape& UsdBlendShape, const FUsdStageInfo& StageInfo, int32 LODIndex, const FTransform& AdditionalTransform, uint32 PointIndexOffset, TSet<FString>& UsedMorphTargetNames, UsdUtils::FBlendShapeMap& OutBlendShapes )
-{
-	return UsdToUnreal::ConvertBlendShape( UsdBlendShape, StageInfo, LODIndex, PointIndexOffset, UsedMorphTargetNames, OutBlendShapes );
-}
-
 USkeletalMesh* UsdToUnreal::GetSkeletalMeshFromImportData(
 	TArray<FSkeletalMeshImportData>& LODIndexToSkeletalMeshImportData,
 	const TArray<SkeletalMeshImportData::FBone>& InSkeletonBones,
@@ -2832,58 +2727,6 @@ void UsdUtils::BindAnimationSource( pxr::UsdPrim& Prim, const pxr::UsdPrim& Anim
 	SkelBindingAPI.CreateAnimationSourceRel().SetTargets( pxr::SdfPathVector( { AnimationSource.GetPath() } ) );
 }
 
-UE::FUsdPrim UsdUtils::FindAnimationSource( const UE::FUsdPrim& Prim )
-{
-	if ( !Prim )
-	{
-		return {};
-	}
-
-	FScopedUsdAllocs UsdAllocs;
-
-	pxr::UsdPrim UsdPrim{ Prim };
-	if ( !UsdPrim.HasAPI<pxr::UsdSkelBindingAPI>() )
-	{
-		UE_LOG( LogUsd, Warning, TEXT( "Failed to find animation source for prim '%s' because it doesn't have a UsdSkelBindingAPI!" ), *Prim.GetPrimPath().GetString() );
-		return {};
-	}
-
-	// According to https://graphics.pixar.com/usd/release/api/_usd_skel__schemas.html#UsdSkel_BindingAPI_Skeletons,
-	// "Since a UsdSkelRoot primitive provides encapsulation of skeletal data, bindings defined on any primitives
-	// that do not have a UsdSkelRoot primitive as one of their ancestors have no meaning, and should be ignored."
-	// So here we make sure that Prim has at least one SkelRoot ancestor
-	bool bFoundSkelRoot = false;
-	pxr::UsdPrim Parent = UsdPrim.GetParent();
-	while ( Parent && !Parent.IsPseudoRoot() )
-	{
-		if ( Parent.IsA<pxr::UsdSkelRoot>() )
-		{
-			bFoundSkelRoot = true;
-			break;
-		}
-
-		Parent = Parent.GetParent();
-	}
-	if ( !bFoundSkelRoot )
-	{
-		return {};
-	}
-
-	pxr::UsdSkelBindingAPI SkelBindingAPI{ Prim };
-
-	pxr::SdfPathVector AnimationSources;
-	const bool bSuccess = SkelBindingAPI.GetAnimationSourceRel().GetTargets( &AnimationSources );
-	if ( !bSuccess || AnimationSources.size() < 1 )
-	{
-		return {};
-	}
-
-	const pxr::SdfPath& FirstSource = AnimationSources[ 0 ];
-	pxr::UsdStageRefPtr UsdStage = Prim.GetStage();
-
-	return UE::FUsdPrim{ UsdStage->GetPrimAtPath( FirstSource ) };
-}
-
 UE::FUsdPrim UsdUtils::FindFirstAnimationSource( const UE::FUsdPrim& SkelRootPrim )
 {
 	if ( !SkelRootPrim )
diff --git a/Engine/Plugins/Importers/USDImporter/Source/USDUtilities/Public/USDGeomMeshConversion.h b/Engine/Plugins/Importers/USDImporter/Source/USDUtilities/Public/USDGeomMeshConversion.h
--- a/Engine/Plugins/Importers/USDImporter/Source/USDUtilities/Public/USDGeomMeshConversion.h
+++ b/Engine/Plugins/Importers/USDImporter/Source/USDUtilities/Public/USDGeomMeshConversion.h
@@ -105,34 +105,6 @@ namespace UsdToUnreal
 		UsdUtils::FUsdPrimMaterialAssignmentInfo& InOutMaterialAssignments,
 		const FUsdMeshConversionOptions& CommonOptions = FUsdMeshConversionOptions::DefaultOptions
 	);
-	UE_DEPRECATED( 5.1, "Please use the overload that receives FUsdMeshConversionOptions." )
-	USDUTILITIES_API bool ConvertGeomMesh(
-		const pxr::UsdTyped& UsdSchema,
-		FMeshDescription& MeshDescription,
-		UsdUtils::FUsdPrimMaterialAssignmentInfo& MaterialAssignments,
-		const pxr::UsdTimeCode TimeCode = pxr::UsdTimeCode::EarliestTime(),
-		const pxr::TfToken& RenderContext = pxr::UsdShadeTokens->universalRenderContext
-	);
-	UE_DEPRECATED( 5.1, "Please use the overload that receives FUsdMeshConversionOptions." )
-	USDUTILITIES_API bool ConvertGeomMesh(
-		const pxr::UsdTyped& UsdSchema,
-		FMeshDescription& MeshDescription,
-		UsdUtils::FUsdPrimMaterialAssignmentInfo& MaterialAssignments,
-		const FTransform& AdditionalTransform,
-		const pxr::UsdTimeCode TimeCode = pxr::UsdTimeCode::EarliestTime(),
-		const pxr::TfToken& RenderContext = pxr::UsdShadeTokens->universalRenderContext
-	);
-	UE_DEPRECATED( 5.1, "Please use the overload that receives FUsdMeshConversionOptions." )
-	USDUTILITIES_API bool ConvertGeomMesh(
-		const pxr::UsdTyped& UsdSchema,
-		FMeshDescription& MeshDescription,
-		UsdUtils::FUsdPrimMaterialAssignmentInfo& MaterialAssignments,
-		const FTransform& AdditionalTransform,
-		const TMap< FString, TMap< FString, int32 > >& MaterialToPrimvarsUVSetNames,
-		const pxr::UsdTimeCode TimeCode = pxr::UsdTimeCode::EarliestTime(),
-		const pxr::TfToken& RenderContext = pxr::UsdShadeTokens->universalRenderContext,
-		bool bMergeIdenticalMaterialSlots = true
-	);
 
 	/** Bake a point instancer prim into a single FMeshDescription and MaterialAssignments, manually "instancing" all
 	 *  the prototype instances into a single combined mesh.
@@ -174,19 +146,6 @@ namespace UsdToUnreal
 		const FUsdMeshConversionOptions& CommonOptions = FUsdMeshConversionOptions::DefaultOptions,
 		bool bSkipRootPrimTransformAndVisibility = false
 	);
-
-	UE_DEPRECATED( 5.1, "Please use the overload that receives FUsdMeshConversionOptions." )
-	USDUTILITIES_API bool ConvertGeomMeshHierarchy(
-		const pxr::UsdPrim& Prim,
-		const pxr::UsdTimeCode& TimeCode,
-		const EUsdPurpose PurposesToLoad,
-		const pxr::TfToken& RenderContext,
-		const TMap< FString, TMap<FString, int32> >& MaterialToPrimvarToUVIndex,
-		FMeshDescription& OutMeshDescription,
-		UsdUtils::FUsdPrimMaterialAssignmentInfo& OutMaterialAssignments,
-		bool bSkipRootPrimTransformAndVisibility = false,
-		bool bMergeIdenticalMaterialSlots = true
-	);
 }
 
 namespace UnrealToUsd
diff --git a/Engine/Plugins/Importers/USDImporter/Source/USDUtilities/Public/USDShadeConversion.h b/Engine/Plugins/Importers/USDImporter/Source/USDUtilities/Public/USDShadeConversion.h
--- a/Engine/Plugins/Importers/USDImporter/Source/USDUtilities/Public/USDShadeConversion.h
+++ b/Engine/Plugins/Importers/USDImporter/Source/USDUtilities/Public/USDShadeConversion.h
@@ -153,13 +153,6 @@ namespace UsdUtils
 	 */
 	USDUTILITIES_API bool IsMaterialTranslucent( const pxr::UsdShadeMaterial& UsdShadeMaterial );
 
-	/**
-	* Returns whether the material uses UDIMs textures.
-	* This function exists because we need this information *before* we pick the right parent for a material instance and properly convert it.
-	*/
-	UE_DEPRECATED( 5.1, "No longer used: Textures and materials will be upgraded to VT-based on-demand" )
-	USDUTILITIES_API bool IsMaterialUsingUDIMs( const pxr::UsdShadeMaterial& UsdShadeMaterial );
-
 	USDUTILITIES_API FSHAHash HashShadeMaterial( const pxr::UsdShadeMaterial& UsdShadeMaterial, const pxr::TfToken& RenderContext = pxr::UsdShadeTokens->universalRenderContext );
 
 	/** Returns the resolved path from a pxr::SdfAssetPath attribute. For UDIMs path, returns the path to the 1001 tile. */
diff --git a/Engine/Plugins/Importers/USDImporter/Source/USDUtilities/Public/USDSkeletalDataConversion.h b/Engine/Plugins/Importers/USDImporter/Source/USDUtilities/Public/USDSkeletalDataConversion.h
--- a/Engine/Plugins/Importers/USDImporter/Source/USDUtilities/Public/USDSkeletalDataConversion.h
+++ b/Engine/Plugins/Importers/USDImporter/Source/USDUtilities/Public/USDSkeletalDataConversion.h
@@ -116,10 +116,6 @@ namespace UsdUtils
 	 */
 	USDUTILITIES_API void BindAnimationSource( pxr::UsdPrim& Prim, const pxr::UsdPrim& AnimationSource );
 
-	// Finds the SkelAnimation prim that is bound to Prim as its animationSource, if Prim has the UsdSkelBindingAPI. Returns an invalid prim otherwise.
-	UE_DEPRECATED( 5.1, "Prefer UsdUtils::FindFirstAnimationSource instead" )
-	USDUTILITIES_API UE::FUsdPrim FindAnimationSource( const UE::FUsdPrim& Prim );
-
 	/**
 	 * Returns the SkelAnimation prim that is resolved for the first skeletal binding of SkelRootPrim, if it is a SkelRoot.
 	 * We use this as we currently only parse a single Skeleton per SkelRoot (and so a single SkelAnimation), but in the future we may
@@ -153,10 +149,6 @@ namespace UsdToUnreal
 	 */
 	USDUTILITIES_API bool ConvertBlendShape( const pxr::UsdSkelBlendShape& UsdBlendShape, const FUsdStageInfo& StageInfo, uint32 PointIndexOffset, TSet<FString>& UsedMorphTargetNames, UsdUtils::FBlendShapeMap& OutBlendShapes );
 	USDUTILITIES_API bool ConvertBlendShape( const pxr::UsdSkelBlendShape& UsdBlendShape, const FUsdStageInfo& StageInfo, int32 LODIndex, uint32 PointIndexOffset, TSet<FString>& UsedMorphTargetNames, UsdUtils::FBlendShapeMap& OutBlendShapes );
-	UE_DEPRECATED( 5.1, "AdditionalTransform is no longer used, please use the other overloads that don't receive it as a parameter." )
-	USDUTILITIES_API bool ConvertBlendShape( const pxr::UsdSkelBlendShape& UsdBlendShape, const FUsdStageInfo& StageInfo, const FTransform& AdditionalTransform, uint32 PointIndexOffset, TSet<FString>& UsedMorphTargetNames, UsdUtils::FBlendShapeMap& OutBlendShapes );
-	UE_DEPRECATED( 5.1, "AdditionalTransform is no longer used, please use the other overloads that don't receive it as a parameter." )
-	USDUTILITIES_API bool ConvertBlendShape( const pxr::UsdSkelBlendShape& UsdBlendShape, const FUsdStageInfo& StageInfo, int32 LODIndex, const FTransform& AdditionalTransform, uint32 PointIndexOffset, TSet<FString>& UsedMorphTargetNames, UsdUtils::FBlendShapeMap& OutBlendShapes );
 
 	/**
 	 * Extracts skeletal mesh data fro UsdSkinningQuery, and places the results in SkelMeshImportData.
@@ -192,16 +184,6 @@ namespace UsdToUnreal
 		const pxr::TfToken& RenderContext = pxr::UsdShadeTokens->universalRenderContext,
 		const pxr::TfToken& MaterialPurpose = pxr::UsdShadeTokens->allPurpose
 	);
-	UE_DEPRECATED( 5.1, "Please use the other overload that also receives the relevant UsdSkelSkeletonQuery object." )
-	USDUTILITIES_API bool ConvertSkinnedMesh(
-		const pxr::UsdSkelSkinningQuery& UsdSkinningQuery,
-		const FTransform& AdditionalTransform,
-		FSkeletalMeshImportData& SkelMeshImportData,
-		TArray< UsdUtils::FUsdPrimMaterialSlot >& MaterialAssignments,
-		const TMap< FString, TMap< FString, int32 > >& MaterialToPrimvarsUVSetNames,
-		const pxr::TfToken& RenderContext = pxr::UsdShadeTokens->universalRenderContext,
-		const pxr::TfToken& MaterialPurpose = pxr::UsdShadeTokens->allPurpose
-	);
 
 	/**
 	 * Will extract animation data from the animation source of InUsdSkeletonQuery's skeleton, and populate OutSkeletalAnimationAsset with the data.
