USD: Fix MaterialX-related warnings when handling any Material prim on the mtlx render context.

diff --git a/Engine/Plugins/Importers/USDImporter/Source/USDSchemas/Private/Custom/MaterialXUSDShadeMaterialTranslator.cpp b/Engine/Plugins/Importers/USDImporter/Source/USDSchemas/Private/Custom/MaterialXUSDShadeMaterialTranslator.cpp
--- a/Engine/Plugins/Importers/USDImporter/Source/USDSchemas/Private/Custom/MaterialXUSDShadeMaterialTranslator.cpp
+++ b/Engine/Plugins/Importers/USDImporter/Source/USDSchemas/Private/Custom/MaterialXUSDShadeMaterialTranslator.cpp
@@ -134,10 +134,20 @@ void FMaterialXUsdShadeMaterialTranslator::CreateAssets()
 		return;
 	}
 
+	// We check for the mtlx surface output directly, because ComputeSurfaceSource will return a valid SurfaceShader
+	// in case the material just has a regular universal render context output.
+	// This is just for checking though: We will defer back to USD to let it ComputeSurfaceSource with whatever logic it has
 	const static pxr::TfToken RenderContextToken = UnrealToUsd::ConvertToken(*MaterialXRenderContext.ToString()).Get();
+	pxr::UsdShadeOutput MtlxOutput = ShadeMaterial.GetSurfaceOutput(RenderContextToken);
+	if (!MtlxOutput)
+	{
+		Super::CreateAssets();
+		return;
+	}
 	pxr::UsdShadeShader SurfaceShader = ShadeMaterial.ComputeSurfaceSource(RenderContextToken);
 	if (!SurfaceShader)
 	{
+		// This really shouldn't ever happen if we have an actual 'mtlx' output, but anyway
 		Super::CreateAssets();
 		return;
 	}
