USD: Fix crash when exiting the engine when saving USD changes;

Fix clang-tidy warnings.

diff --git a/Engine/Plugins/Importers/USDImporter/Source/USDStageEditor/Private/SUSDStage.cpp b/Engine/Plugins/Importers/USDImporter/Source/USDStageEditor/Private/SUSDStage.cpp
--- a/Engine/Plugins/Importers/USDImporter/Source/USDStageEditor/Private/SUSDStage.cpp
+++ b/Engine/Plugins/Importers/USDImporter/Source/USDStageEditor/Private/SUSDStage.cpp
@@ -9,44 +9,33 @@
 #include "UnrealUSDWrapper.h"
 #include "USDClassesModule.h"
 #include "USDConversionUtils.h"
-#include "USDErrorUtils.h"
 #include "USDLayerUtils.h"
 #include "USDProjectSettings.h"
 #include "USDSchemasModule.h"
 #include "USDSchemaTranslator.h"
 #include "USDStageActor.h"
 #include "USDStageEditorSettings.h"
-#include "USDStageImportContext.h"
-#include "USDStageImporter.h"
-#include "USDStageImporterModule.h"
 #include "USDStageImportOptions.h"
 #include "USDStageModule.h"
 #include "USDTypesConversion.h"
 #include "UsdWrappers/SdfLayer.h"
-#include "UsdWrappers/UsdAttribute.h"
 #include "UsdWrappers/UsdStage.h"
 #include "Widgets/SUSDObjectFieldList.h"
 
 #include "ActorTreeItem.h"
-#include "Async/Async.h"
 #include "DesktopPlatformModule.h"
-#include "Dialogs/DlgPickPath.h"
 #include "Editor.h"
 #include "Editor/Transactor.h"
-#include "Engine/Selection.h"
 #include "Engine/World.h"
 #include "EngineAnalytics.h"
 #include "Framework/MultiBox/MultiBoxBuilder.h"
 #include "ISceneOutliner.h"
 #include "ISceneOutlinerColumn.h"
-#include "ISinglePropertyView.h"
-#include "LevelEditor.h"
 #include "Modules/ModuleManager.h"
-#include "PropertyEditorModule.h"
 #include "SceneOutlinerModule.h"
 #include "ScopedTransaction.h"
+#include "Selection.h"
 #include "Styling/AppStyle.h"
-#include "UObject/StrongObjectPtr.h"
 #include "Widgets/Input/SButton.h"
 #include "Widgets/Input/SButton.h"
 #include "Widgets/Input/SComboBox.h"
@@ -248,7 +237,7 @@ namespace SUSDStageImpl
 		{
 			return true;
 		}
-		virtual void SortItems( TArray<FSceneOutlinerTreeItemPtr>& OutItems, const EColumnSortMode::Type SortMode ) const
+		virtual void SortItems( TArray<FSceneOutlinerTreeItemPtr>& OutItems, const EColumnSortMode::Type SortMode ) const override
 		{
 			if ( SortMode == EColumnSortMode::Ascending )
 			{
@@ -299,6 +288,13 @@ void SUsdStage::Construct( const FArguments& InArgs )
 
 	UE::FUsdStage UsdStage;
 
+	OnEditorCloseHandle = GEditor->OnEditorClose().AddLambda(
+		[this]()
+		{
+			bEditorIsShuttingDown = true;
+		}
+	);
+
 	const AUsdStageActor* StageActor = ViewModel.UsdStageActor.Get();
 	if ( StageActor )
 	{
@@ -531,9 +527,14 @@ void SUsdStage::SetupStageActorDelegates()
 				// The USD notices may come from a background USD TBB thread, but we should only update slate from the main/slate threads.
 				// We can't retrieve the FSlateApplication singleton here (because that can also only be used from the main/slate threads),
 				// so we must use core tickers here
-				FTSTicker::GetCoreTicker().AddTicker(
-					FTickerDelegate::CreateLambda( [this, PrimPath, bResync]( float Time )
+				FTSTicker::GetCoreTicker().AddTicker(FTickerDelegate::CreateLambda(
+					[this, PrimPath, bResync](float Time)
 					{
+						if (bEditorIsShuttingDown)
+						{
+							return false;
+						}
+
 						if ( this->UsdStageTreeView )
 						{
 							this->UsdStageTreeView->RefreshPrim( PrimPath, bResync );
@@ -560,31 +561,36 @@ void SUsdStage::SetupStageActorDelegates()
 
 						// Returning false means this is a one-off, and won't repeat
 						return false;
-					})
-				);
+					}
+				));
 			}
 		);
 
 		// Fired when we switch which is the currently opened stage
 		OnStageChangedHandle = ViewModel.UsdStageActor->OnStageChanged.AddLambda(
-			[ this ]()
+			[this]()
 			{
-				FTSTicker::GetCoreTicker().AddTicker(
-					FTickerDelegate::CreateLambda( [this]( float Time )
+				FTSTicker::GetCoreTicker().AddTicker(FTickerDelegate::CreateLambda(
+					[this](float Time)
 					{
+						if (bEditorIsShuttingDown)
+						{
+							return false;
+						}
+
 						// So we can reset even if our actor is being destroyed right now
 						const bool bEvenIfPendingKill = true;
-						if ( ViewModel.UsdStageActor.IsValid( bEvenIfPendingKill ) )
+						if (ViewModel.UsdStageActor.IsValid(bEvenIfPendingKill))
 						{
 							// Reset our selection to the stage root
-							SelectedPrimPath = TEXT( "/" );
+							SelectedPrimPath = TEXT("/");
 
-							if ( this->UsdPrimInfoWidget )
+							if (this->UsdPrimInfoWidget)
 							{
-								this->UsdPrimInfoWidget->SetPrimPath( GetCurrentStage(), TEXT( "/" ) );
+								this->UsdPrimInfoWidget->SetPrimPath(GetCurrentStage(), TEXT("/"));
 							}
 
-							if ( this->UsdStageTreeView )
+							if (this->UsdStageTreeView)
 							{
 								this->UsdStageTreeView->ClearSelection();
 								this->UsdStageTreeView->RequestTreeRefresh();
@@ -595,8 +601,8 @@ void SUsdStage::SetupStageActorDelegates()
 
 						// Returning false means this is a one-off, and won't repeat
 						return false;
-					})
-				);
+					}
+				));
 			}
 		);
 
@@ -611,15 +617,20 @@ void SUsdStage::SetupStageActorDelegates()
 				ClearStageActorDelegates();
 				this->ViewModel.CloseStage();
 
-				FTSTicker::GetCoreTicker().AddTicker(
-					FTickerDelegate::CreateLambda( [this]( float Time )
+				FTSTicker::GetCoreTicker().AddTicker(FTickerDelegate::CreateLambda(
+					[this](float Time)
 					{
+						if (bEditorIsShuttingDown)
+						{
+							return false;
+						}
+
 						this->Refresh();
 
 						// Returning false means this is a one-off, and won't repeat
 						return false;
-					})
-				);
+					}
+				));
 			}
 		);
 
@@ -629,8 +640,13 @@ void SUsdStage::SetupStageActorDelegates()
 				FTSTicker::GetCoreTicker().AddTicker(
 					FTickerDelegate::CreateLambda( [this]( float Time )
 					{
+						if (bEditorIsShuttingDown)
+						{
+							return false;
+						}
+
 						AUsdStageActor* StageActor = ViewModel.UsdStageActor.Get();
-						if ( this->UsdLayersTreeView && StageActor )
+						if (this->UsdLayersTreeView && StageActor)
 						{
 							constexpr bool bResync = false;
 							UsdLayersTreeView->Refresh(
@@ -648,13 +664,18 @@ void SUsdStage::SetupStageActorDelegates()
 		);
 
 		OnLayersChangedHandle = ViewModel.UsdStageActor->GetUsdListener().GetOnLayersChanged().AddLambda(
-			[ this ]( const TArray< FString >& LayersNames )
+			[this](const TArray<FString>& LayersNames)
 			{
-				FTSTicker::GetCoreTicker().AddTicker(
-					FTickerDelegate::CreateLambda( [this]( float Time )
+				FTSTicker::GetCoreTicker().AddTicker(FTickerDelegate::CreateLambda(
+					[this](float Time)
 					{
+						if (bEditorIsShuttingDown)
+						{
+							return false;
+						}
+
 						AUsdStageActor* StageActor = ViewModel.UsdStageActor.Get();
-						if ( this->UsdLayersTreeView && StageActor )
+						if (this->UsdLayersTreeView && StageActor)
 						{
 							constexpr bool bResync = false;
 							UsdLayersTreeView->Refresh(
@@ -666,8 +687,8 @@ void SUsdStage::SetupStageActorDelegates()
 
 						// Returning false means this is a one-off, and won't repeat
 						return false;
-					})
-				);
+					}
+				));
 			}
 		);
 	}
@@ -692,6 +713,8 @@ SUsdStage::~SUsdStage()
 	FEditorDelegates::PostPIEStarted.Remove( PostPIEStartedHandle );
 	FEditorDelegates::EndPIE.Remove( EndPIEHandle );
 
+	GEditor->OnEditorClose().Remove(OnEditorCloseHandle);
+
 	FCoreUObjectDelegates::OnObjectPropertyChanged.Remove( OnStageActorPropertyChangedHandle );
 	AUsdStageActor::OnActorLoaded.Remove( OnActorLoadedHandle );
 	USelection::SelectionChangedEvent.Remove( OnViewportSelectionChangedHandle );
@@ -1864,7 +1887,7 @@ void SUsdStage::FillSelectionSubMenu( FMenuBuilder& MenuBuilder )
 			{
 				return true;
 			}),
-			FIsActionChecked::CreateLambda([this]()
+			FIsActionChecked::CreateLambda([]()
 			{
 				const UUsdStageEditorSettings* Settings = GetDefault<UUsdStageEditorSettings>();
 				return Settings && Settings->bSelectionSynced;
diff --git a/Engine/Plugins/Importers/USDImporter/Source/USDStageEditor/Private/SUSDStage.h b/Engine/Plugins/Importers/USDImporter/Source/USDStageEditor/Private/SUSDStage.h
--- a/Engine/Plugins/Importers/USDImporter/Source/USDStageEditor/Private/SUSDStage.h
+++ b/Engine/Plugins/Importers/USDImporter/Source/USDStageEditor/Private/SUSDStage.h
@@ -6,10 +6,6 @@
 #include "UsdWrappers/ForwardDeclarations.h"
 
 #include "Animation/CurveSequence.h"
-#include "CoreMinimal.h"
-#include "Input/Reply.h"
-#include "Layout/Visibility.h"
-#include "Misc/Optional.h"
 #include "Templates/SharedPointer.h"
 #include "Widgets/DeclarativeSyntaxSupport.h"
 #include "Widgets/SCompoundWidget.h"
@@ -153,6 +149,16 @@ protected:
 	int32 CurrentNaniteThreshold = INT32_MAX;
 
 	TArray<TSharedPtr<FString>> MaterialPurposes;
+
+	// We use our own bool to track engine exit because we use some tickers to delay work onto the next frame,
+	// and during engine shutdown its possible for this delayed work to run *before* IsEngineExitRequested() actually
+	// turns true, so we can't check it.
+	// If that work continues (and triggers its slate updates and so on), we can run into crashes as we're
+	// trying to render slate as the engine shuts down.
+	// This bool is set on GEditor->OnEditorClose(), which happens earlier in the shutdown process, and can be used
+	// to let our ticker lambdas gracefully early out during those shutdown scenarios.
+	bool bEditorIsShuttingDown = false;
+	FDelegateHandle OnEditorCloseHandle;
 };
 
 #endif // #if USE_USD_SDK
