USD: Add dynamic level sequence bindings also to components, so that the bindings can be automatically repaired at runtime.

diff --git a/Engine/Plugins/Importers/USDImporter/Source/USDClasses/Public/USDAssetUserData.h b/Engine/Plugins/Importers/USDImporter/Source/USDClasses/Public/USDAssetUserData.h
--- a/Engine/Plugins/Importers/USDImporter/Source/USDClasses/Public/USDAssetUserData.h
+++ b/Engine/Plugins/Importers/USDImporter/Source/USDClasses/Public/USDAssetUserData.h
@@ -97,4 +97,29 @@ public:
 	/** Describes which primvars should be assigned to each UV index. */
 	UPROPERTY()
 	TMap<FString, int32> PrimvarToUVIndex;
+};
+
+/**
+ * We assign these to persistent LevelSequences that bind to one of the actors/components that the stage actor spawns.
+ * We need this as part of a mechanism to automatically repair those bindings when they break if we close/reload the stage.
+ */
+UCLASS()
+class USDCLASSES_API UUsdLevelSequenceAssetUserData : public UAssetUserData
+{
+	GENERATED_BODY()
+
+public:
+	/**
+	 * The LevelSequence has a Guid that is changed every time its state is modified.
+	 * We pay attention to that so that we can avoid reprocessing a LevelSequence that hasn't changed
+	 */
+	UPROPERTY()
+	FGuid LastCheckedSignature;
+
+	/**
+	 * Set of binding GUIDs that we already handled in the past.
+	 * We use this so that we won't try and overwrite the changes in case the user manually clears/modifies a binding we previously setup.
+	 */
+	UPROPERTY()
+	TSet<FGuid> HandledBindingGuids;
 };
\ No newline at end of file
diff --git a/Engine/Plugins/Importers/USDImporter/Source/USDStage/Private/USDStageActor.cpp b/Engine/Plugins/Importers/USDImporter/Source/USDStage/Private/USDStageActor.cpp
--- a/Engine/Plugins/Importers/USDImporter/Source/USDStage/Private/USDStageActor.cpp
+++ b/Engine/Plugins/Importers/USDImporter/Source/USDStage/Private/USDStageActor.cpp
@@ -6,6 +6,7 @@
 #include "UObject/Package.h"
 #include "USDAssetCache.h"
 #include "USDAssetImportData.h"
+#include "USDAssetUserData.h"
 #include "USDClassesModule.h"
 #include "USDConversionUtils.h"
 #include "USDDynamicBindingResolverLibrary.h"
@@ -585,6 +586,177 @@ struct FUsdStageActorImpl
 		}
 #endif // WITH_EDITOR
 	}
+
+	static void SetupDynamicBinding(
+		const FString& PrimPath,
+		FMovieScenePossessable& Possessable,
+		ULevelSequence* Sequence,
+		const FString& DefaultActorFilter
+	)
+	{
+#if WITH_EDITOR
+		UMovieScene* MovieScene = Sequence->GetMovieScene();
+		if (!MovieScene)
+		{
+			return;
+		}
+
+		// Make sure we trigger the creation of the DirectorBlueprint before we add our binding, or else
+		// FMovieSceneDynamicBindingInvoker::ResolveDynamicBinding may fail. Note: The regular code path for dynamic
+		// bindings via the UI ends up calling these exact same lines from
+		// FMovieSceneDirectorBlueprintEndpointCustomization::PopulateQuickBindSubMenu
+		FMovieSceneSequenceEditor* SequenceEditor = FMovieSceneSequenceEditor::Find(Sequence);
+		if (!SequenceEditor)
+		{
+			return;
+		}
+		UBlueprint* DirectorBlueprint = SequenceEditor->GetOrCreateDirectorBlueprint(Sequence);
+		if (!DirectorBlueprint)
+		{
+			return;
+		}
+
+		// Setup the Dynamic Binding
+		{
+			Possessable.DynamicBinding = FMovieSceneDynamicBinding{};
+			Possessable.DynamicBinding.CompiledFunctionName = NAME_None;
+			Possessable.DynamicBinding.ResolveParamsPinName = TEXT("Params");
+			Possessable.DynamicBinding.Function = UUsdDynamicBindingResolverLibrary::StaticClass()->FindFunctionByName(
+				GET_FUNCTION_NAME_CHECKED(UUsdDynamicBindingResolverLibrary, ResolveWithStageActor)
+			);
+			if (Possessable.DynamicBinding.Function)
+			{
+				Possessable.DynamicBinding.ResolveParamsProperty = Possessable.DynamicBinding.Function->FindPropertyByName(
+					Possessable.DynamicBinding.ResolveParamsPinName
+				);
+			}
+
+			// Store a path to this very actor on the binding, so that it can find us later and ask how to resolve a particular
+			// prim path
+			FMovieSceneDynamicBindingPayloadVariable& ActorPathVariable = Possessable.DynamicBinding.PayloadVariables.FindOrAdd(
+				TEXT("StageActorNameFilter")
+			);
+			ActorPathVariable.Value = DefaultActorFilter;
+
+			FMovieSceneDynamicBindingPayloadVariable& RootLayerVariable = Possessable.DynamicBinding.PayloadVariables.FindOrAdd(
+				TEXT("RootLayerFilter")
+			);
+			RootLayerVariable.Value = FString{};	// No root layer filter by default for more flexibility
+
+			FMovieSceneDynamicBindingPayloadVariable& PrimPathVariable = Possessable.DynamicBinding.PayloadVariables.FindOrAdd(
+				TEXT("PrimPath")
+			);
+			PrimPathVariable.Value = PrimPath;
+		}
+
+		// Setup the Blueprint function ("endpoint") on the Sequence's director Blueprint.
+		// In the future we could also just setup a single endpoint for all dynamic bindings, as they're all identical here. However, even then
+		// I think we should allow the option of having this system emit a separatate endpoint for each binding like it does now, as the user
+		// may want to modify these dynamic binding functions for some reason (maybe add a prim path suffix, or retrieve the prim path from
+		// some other system, etc.).
+		// Reference: FMovieSceneDynamicBindingCustomization::GenerateEndpointDefinition
+		UK2Node_FunctionEntry* NewEndpoint = nullptr;
+		{
+			FMovieSceneDirectorBlueprintEndpointDefinition EndpointDefinition;
+			EndpointDefinition.EndpointType = EMovieSceneDirectorBlueprintEndpointType::Function;
+
+			// Use the same dummy function the regular UI-created Dynamic Bindings use to setup the endpoint signature, in case it changes
+			static const FName SampleResolveBindingFuncName{TEXT("SampleResolveBinding")};
+			UClass* EndpointUtilClass = UMovieSceneDynamicBindingEndpointUtil::StaticClass();
+			EndpointDefinition.EndpointSignature = EndpointUtilClass->FindFunctionByName(SampleResolveBindingFuncName);
+			if (!EndpointDefinition.EndpointSignature)
+			{
+				return;
+			}
+
+			EndpointDefinition.EndpointName = Possessable.GetName() + TEXT("_DynamicBinding");
+
+			DirectorBlueprint->Modify();
+
+			// Creates the new function on the LevelSequence's director blueprint
+			NewEndpoint = FMovieSceneDirectorBlueprintUtils::CreateFunctionEndpoint(DirectorBlueprint, EndpointDefinition);
+			NewEndpoint->MetaData.bCallInEditor = true;
+			if (!NewEndpoint)
+			{
+				return;
+			}
+		}
+
+		// Create a function call node within that endpoint that calls our UFunction
+		// Reference: FMovieSceneDirectorBlueprintEndpointCustomization::HandleQuickBindActionSelected
+		UBlueprintFunctionNodeSpawner* Spawner = UBlueprintFunctionNodeSpawner::Create(Possessable.DynamicBinding.Function);
+		FBlueprintActionMenuItem Action{Spawner};
+		UEdGraphPin* ThenPin = NewEndpoint->FindPin(UEdGraphSchema_K2::PN_Then, EGPD_Output);
+		FVector2D NodePosition(NewEndpoint->NodePosX + 400.f, NewEndpoint->NodePosY + 100.0f);
+		UEdGraphNode* NewNode = Action.PerformAction(NewEndpoint->GetGraph(), ThenPin, NodePosition);
+
+		// Connect the function entry point to our function call node if it has an exec pin
+		UEdGraphPin* NewNodeExecPin = NewNode->FindPin(UEdGraphSchema_K2::PN_Execute, EGPD_Input);
+		if (ThenPin && NewNodeExecPin)
+		{
+			ThenPin->MakeLinkTo(NewNodeExecPin);
+		}
+
+		// Connect the function call result pin to the endpoint's result node
+		TArray<UK2Node_FunctionResult*> ResultNodes;
+		NewNode->GetGraph()->GetNodesOfClass(ResultNodes);
+		if (ResultNodes.Num() > 0)
+		{
+			// If there is a result node, move it past the endpoint call
+			ResultNodes[0]->NodePosX = NodePosition.X + 400.f;
+
+			// Connect the function call node to the endpoint result node if it has an exec pin
+			UEdGraphPin* NewNodeThenPin = NewNode->FindPin(UEdGraphSchema_K2::PN_Execute, EGPD_Input);
+			UEdGraphPin* ResultExecPin = ResultNodes[0]->FindPin(UEdGraphSchema_K2::PN_Execute, EGPD_Input);
+			if (NewNodeThenPin && ResultExecPin)
+			{
+				NewNodeThenPin->MakeLinkTo(ResultExecPin);
+			}
+
+			// If the new node has a return value, and if the endpoint has one too, try to connect them together.
+			UEdGraphPin* OutputPin = ResultNodes[0]->FindPin(UEdGraphSchema_K2::PN_ReturnValue, EGPD_Input);
+			UEdGraphPin* NewNodeReturnValuePin = NewNode->FindPin(UEdGraphSchema_K2::PN_ReturnValue, EGPD_Output);
+			if (OutputPin && NewNodeReturnValuePin)
+			{
+				// Connect the nodes.
+				NewNodeReturnValuePin->MakeLinkTo(OutputPin);
+			}
+		}
+
+		// Connect additional parameter pins from the endpoint to the function call node
+		// Reference: FMovieSceneDirectorBlueprintEndpointCustomization::SetEndpoint
+		TSet<FName> NonPayloadPins;
+		const FString* WorldContextParamName = Possessable.DynamicBinding.Function->FindMetaData(FBlueprintMetadata::MD_WorldContext);
+		if (WorldContextParamName)
+		{
+			NonPayloadPins.Add(FName(*WorldContextParamName));
+		}
+		if (UK2Node_EditablePinBase* EditableNode = Cast<UK2Node_EditablePinBase>(NewEndpoint))
+		{
+			for (UEdGraphPin* PayloadPin : NewNode->Pins)
+			{
+				if (PayloadPin && PayloadPin->LinkedTo.Num() == 0 && PayloadPin->Direction == EGPD_Input
+					&& PayloadPin->PinType.PinCategory != UEdGraphSchema_K2::PC_Exec && PayloadPin->PinName != UEdGraphSchema_K2::PN_Self
+					&& !NonPayloadPins.Contains(PayloadPin->PinName))
+				{
+					// Pins for ref parameters for functions default to bIsReference but the payload cannot be by
+					// reference.
+					PayloadPin->PinType.bIsReference = false;
+
+					UEdGraphPin* NewPin = EditableNode->CreateUserDefinedPin(PayloadPin->PinName, PayloadPin->PinType, EGPD_Output);
+					if (NewNode != NewEndpoint && NewPin)
+					{
+						NewPin->MakeLinkTo(PayloadPin);
+					}
+				}
+			}
+		}
+
+		FMovieSceneDynamicBindingUtils::SetEndpoint(MovieScene, &Possessable.DynamicBinding, NewEndpoint);
+		FMovieSceneDynamicBindingUtils::EnsureBlueprintExtensionCreated(Sequence, DirectorBlueprint);
+		FKismetEditorUtilities::CompileBlueprint(DirectorBlueprint);
+#endif	  // WITH_EDITOR
+	}
 };
 
 /**
@@ -751,22 +923,25 @@ AUsdStageActor::AUsdStageActor()
 			);
 		}
 
+		TFunction<void(TSharedRef<ISequencer>)> SubscribeToSequencerEvents = [this](TSharedRef<ISequencer> Sequencer)
+		{
+			Sequencer->OnActorAddedToSequencer().AddUObject(this, &AUsdStageActor::OnActorAddedToSequencer, Sequencer.ToWeakPtr());
+			Sequencer->OnMovieSceneDataChanged().AddUObject(this, &AUsdStageActor::OnMovieSceneDataChanged, Sequencer.ToWeakPtr());
+		};
+
 		// Make sure we are bound to the OnActorAddedToSequencer of every sequencer that exists
 		IUsdStageModule& UsdStageModule = FModuleManager::Get().LoadModuleChecked<IUsdStageModule>(TEXT("UsdStage"));
 		for (const TWeakPtr<ISequencer>& ExistingSequencer : UsdStageModule.GetExistingSequencers())
 		{
 			if (TSharedPtr<ISequencer> PinnedSequencer = ExistingSequencer.Pin())
 			{
-				PinnedSequencer->OnActorAddedToSequencer().AddUObject(this, &AUsdStageActor::OnActorAddedToSequencer, ExistingSequencer);
+				SubscribeToSequencerEvents(PinnedSequencer.ToSharedRef());
 			}
 		}
 		ISequencerModule& SequencerModule = FModuleManager::Get().LoadModuleChecked<ISequencerModule>(TEXT("Sequencer"));
-		OnSequencerCreatedHandle = SequencerModule.RegisterOnSequencerCreated(FOnSequencerCreated::FDelegate::CreateLambda(
-			[this](TSharedRef<ISequencer> NewSequencer)
-			{
-				NewSequencer->OnActorAddedToSequencer().AddUObject(this, &AUsdStageActor::OnActorAddedToSequencer, NewSequencer.ToWeakPtr());
-			}
-		));
+		OnSequencerCreatedHandle = SequencerModule.RegisterOnSequencerCreated(
+			FOnSequencerCreated::FDelegate::CreateLambda(SubscribeToSequencerEvents)
+		);
 
 		FCoreUObjectDelegates::OnObjectPropertyChanged.AddUObject(this, &AUsdStageActor::OnObjectPropertyChanged);
 
@@ -2800,6 +2975,7 @@ void AUsdStageActor::BeginDestroy()
 					if (TSharedPtr<ISequencer> PinnedSequencer = ExistingSequencer.Pin())
 					{
 						PinnedSequencer->OnActorAddedToSequencer().RemoveAll(this);
+						PinnedSequencer->OnMovieSceneDataChanged().RemoveAll(this);
 					}
 				}
 			}
@@ -3054,17 +3230,55 @@ void AUsdStageActor::OnActorAddedToSequencer(AActor* NewActor, const FGuid Guid,
 		return;
 	}
 
-	// Make sure we trigger the creation of the DirectorBlueprint before we add our binding, or else
-	// FMovieSceneDynamicBindingInvoker::ResolveDynamicBinding may fail. Note: The regular code path for dynamic
-	// bindings via the UI ends up calling these exact same lines from
-	// FMovieSceneDirectorBlueprintEndpointCustomization::PopulateQuickBindSubMenu
-	FMovieSceneSequenceEditor* SequenceEditor = FMovieSceneSequenceEditor::Find(Sequence);
-	if (!SequenceEditor)
+	// The UI action to add the actor to the sequencer should have spawned it's transaction already, but just in case it
+	// hasn't let's do that, as we'll end up modifying UObjects in here
+	FScopedTransaction Transaction(
+		FText::Format(LOCTEXT("ActorAddedToSequencer", "Bound USD actor '{0}' to the Sequencer"), FText::FromString(NewActor->GetActorLabel()))
+	);
+	Sequence->Modify();
+	MovieScene->Modify();
+
+	FUsdStageActorImpl::SetupDynamicBinding(PrimPath, *Possessable, Sequence, *GetName());
+#endif	  // WITH_EDITOR
+}
+
+void AUsdStageActor::OnMovieSceneDataChanged(EMovieSceneDataChangeType ChangeType, TWeakPtr<ISequencer> WeakSequencer)
+{
+#if WITH_EDITOR
+	if (ChangeType != EMovieSceneDataChangeType::MovieSceneStructureItemAdded
+		&& ChangeType != EMovieSceneDataChangeType::MovieSceneStructureItemsChanged)
+	{
+		return;
+	}
+
+	TSharedPtr<ISequencer> PinnedSequencer = WeakSequencer.Pin();
+	if (!PinnedSequencer)
 	{
 		return;
 	}
-	UBlueprint* DirectorBlueprint = SequenceEditor->GetOrCreateDirectorBlueprint(Sequence);
-	if (!DirectorBlueprint)
+
+	ULevelSequence* Sequence = Cast<ULevelSequence>(PinnedSequencer->GetRootMovieSceneSequence());
+	if (!Sequence)
+	{
+		return;
+	}
+
+	UUsdLevelSequenceAssetUserData* AssetUserData = Sequence->GetAssetUserData<UUsdLevelSequenceAssetUserData>();
+	if(!AssetUserData)
+	{
+		AssetUserData = NewObject<UUsdLevelSequenceAssetUserData>(Sequence);
+		Sequence->AddAssetUserData(AssetUserData);
+	}
+
+	const FGuid& CurrentSignature = Sequence->GetSignature();
+	if (CurrentSignature == AssetUserData->LastCheckedSignature)
+	{
+		return;
+	}
+	AssetUserData->LastCheckedSignature = CurrentSignature;
+
+	UMovieScene* MovieScene = Sequence->GetMovieScene();
+	if (!MovieScene)
 	{
 		return;
 	}
@@ -3072,154 +3286,77 @@ void AUsdStageActor::OnActorAddedToSequencer(AActor* NewActor, const FGuid Guid,
 	// The UI action to add the actor to the sequencer should have spawned it's transaction already, but just in case it
 	// hasn't let's do that, as we'll end up modifying UObjects in here
 	FScopedTransaction Transaction(
-		FText::Format(LOCTEXT("ActorAddedToSequencer", "Bound actor '{0}' to the Sequencer"), FText::FromString(NewActor->GetActorLabel()))
+		LOCTEXT("ComponentAddedToSequencer", "Added USD component bindings to the Sequencer")
 	);
 	Sequence->Modify();
 	MovieScene->Modify();
 
-	// Setup the Dynamic Binding
+	TFunction<UObject*(const FMovieScenePossessable&)> LocateBoundObject = nullptr;
+	LocateBoundObject = [&LocateBoundObject, Sequence, MovieScene](const FMovieScenePossessable& Possessable) -> UObject*
 	{
-		Possessable->DynamicBinding = FMovieSceneDynamicBinding{};
-		Possessable->DynamicBinding.CompiledFunctionName = NAME_None;
-		Possessable->DynamicBinding.ResolveParamsPinName = TEXT("Params");
-		Possessable->DynamicBinding.Function = UUsdDynamicBindingResolverLibrary::StaticClass()->FindFunctionByName(
-			GET_FUNCTION_NAME_CHECKED(UUsdDynamicBindingResolverLibrary, ResolveWithStageActor)
-		);
-		if (Possessable->DynamicBinding.Function)
+		const FGuid& Guid = Possessable.GetGuid();
+		const FGuid& ParentGuid = Possessable.GetParent();
+
+		UObject* ParentContext = nullptr;
+		if (ParentGuid.IsValid())
 		{
-			Possessable->DynamicBinding.ResolveParamsProperty = Possessable->DynamicBinding.Function->FindPropertyByName(
-				Possessable->DynamicBinding.ResolveParamsPinName
-			);
+			if (FMovieScenePossessable* ParentPossessable = MovieScene->FindPossessable(ParentGuid))
+			{
+				ParentContext = LocateBoundObject(*ParentPossessable);
+			}
 		}
 
-		// Store a path to this very actor on the binding, so that it can find us later and ask how to resolve a particular
-		// prim path
-		FMovieSceneDynamicBindingPayloadVariable& ActorPathVariable = Possessable->DynamicBinding.PayloadVariables.FindOrAdd(
-			TEXT("StageActorNameFilter")
-		);
-		ActorPathVariable.Value = GetName();
-
-		FMovieSceneDynamicBindingPayloadVariable& RootLayerVariable = Possessable->DynamicBinding.PayloadVariables.FindOrAdd(
-			TEXT("RootLayerFilter")
-		);
-		RootLayerVariable.Value = FString{};	// No root layer filter by default for more flexibility
+		TArray<UObject*, TInlineAllocator<1>> Objects = Cast<UMovieSceneSequence>(Sequence)->LocateBoundObjects(Guid, ParentContext);
+		if (Objects.Num() > 0)
+		{
+			return Objects[0];
+		}
 
-		FMovieSceneDynamicBindingPayloadVariable& PrimPathVariable = Possessable->DynamicBinding.PayloadVariables.FindOrAdd(
-			TEXT("PrimPath")
-		);
-		PrimPathVariable.Value = PrimPath;
-	}
+		return nullptr;
+	};
 
-	// Setup the Blueprint function ("endpoint") on the Sequence's director Blueprint.
-	// In the future we could also just setup a single endpoint for all dynamic bindings, as they're all identical here. However, even then
-	// I think we should allow the option of having this system emit a separatate endpoint for each binding like it does now, as the user
-	// may want to modify these dynamic binding functions for some reason (maybe add a prim path suffix, or retrieve the prim path from
-	// some other system, etc.).
-	// Reference: FMovieSceneDynamicBindingCustomization::GenerateEndpointDefinition
-	UK2Node_FunctionEntry* NewEndpoint = nullptr;
+	int32 PossessableCount = MovieScene->GetPossessableCount();
+	for(int32 Index = 0; Index < PossessableCount; ++Index)
 	{
-		FMovieSceneDirectorBlueprintEndpointDefinition EndpointDefinition;
-		EndpointDefinition.EndpointType = EMovieSceneDirectorBlueprintEndpointType::Function;
+		FMovieScenePossessable& Possessable = MovieScene->GetPossessable(Index);
 
-		// Use the same dummy function the regular UI-created Dynamic Bindings use to setup the endpoint signature, in case it changes
-		static const FName SampleResolveBindingFuncName{TEXT("SampleResolveBinding")};
-		UClass* EndpointUtilClass = UMovieSceneDynamicBindingEndpointUtil::StaticClass();
-		EndpointDefinition.EndpointSignature = EndpointUtilClass->FindFunctionByName(SampleResolveBindingFuncName);
-		if(!EndpointDefinition.EndpointSignature)
+		// We only handle components on this code path, and component bindings always have a parent actor binding
+		if (!Possessable.GetParent().IsValid())
 		{
-			return;
+			continue;
 		}
 
-		EndpointDefinition.EndpointName = Possessable->GetName() + TEXT("_DynamicBinding");
-
-		DirectorBlueprint->Modify();
-
-		// Creates the new function on the LevelSequence's director blueprint
-		NewEndpoint = FMovieSceneDirectorBlueprintUtils::CreateFunctionEndpoint(DirectorBlueprint, EndpointDefinition);
-		NewEndpoint->MetaData.bCallInEditor = true;
-		if (!NewEndpoint)
+		// If the binding already has another dynamic binding let's not touch it,
+		// regardless of whether we set that dynamic binding up or the user did
+		if (Possessable.DynamicBinding.Function != nullptr)
 		{
-			return;
+			continue;
 		}
-	}
-
-	// Create a function call node within that endpoint that calls our UFunction
-	// Reference: FMovieSceneDirectorBlueprintEndpointCustomization::HandleQuickBindActionSelected
-	UBlueprintFunctionNodeSpawner* Spawner = UBlueprintFunctionNodeSpawner::Create(Possessable->DynamicBinding.Function);
-	FBlueprintActionMenuItem Action{Spawner};
-	UEdGraphPin* ThenPin = NewEndpoint->FindPin(UEdGraphSchema_K2::PN_Then, EGPD_Output);
-	FVector2D NodePosition(NewEndpoint->NodePosX + 400.f, NewEndpoint->NodePosY + 100.0f);
-	UEdGraphNode* NewNode = Action.PerformAction(NewEndpoint->GetGraph(), ThenPin, NodePosition);
 
-	// Connect the function entry point to our function call node if it has an exec pin
-	UEdGraphPin* NewNodeExecPin = NewNode->FindPin(UEdGraphSchema_K2::PN_Execute, EGPD_Input);
-	if (ThenPin && NewNodeExecPin)
-	{
-		ThenPin->MakeLinkTo(NewNodeExecPin);
-	}
-
-	// Connect the function call result pin to the endpoint's result node
-	TArray<UK2Node_FunctionResult*> ResultNodes;
-	NewNode->GetGraph()->GetNodesOfClass(ResultNodes);
-	if (ResultNodes.Num() > 0)
-	{
-		// If there is a result node, move it past the endpoint call
-		ResultNodes[0]->NodePosX = NodePosition.X + 400.f;
-
-		// Connect the function call node to the endpoint result node if it has an exec pin
-		UEdGraphPin* NewNodeThenPin = NewNode->FindPin(UEdGraphSchema_K2::PN_Execute, EGPD_Input);
-		UEdGraphPin* ResultExecPin = ResultNodes[0]->FindPin(UEdGraphSchema_K2::PN_Execute, EGPD_Input);
-		if (NewNodeThenPin && ResultExecPin)
+		const FGuid& BindingGuid = Possessable.GetGuid();
+		if (AssetUserData->HandledBindingGuids.Contains(BindingGuid))
 		{
-			NewNodeThenPin->MakeLinkTo(ResultExecPin);
+			continue;
 		}
 
-		// If the new node has a return value, and if the endpoint has one too, try to connect them together.
-		UEdGraphPin* OutputPin = ResultNodes[0]->FindPin(UEdGraphSchema_K2::PN_ReturnValue, EGPD_Input);
-		UEdGraphPin* NewNodeReturnValuePin = NewNode->FindPin(UEdGraphSchema_K2::PN_ReturnValue, EGPD_Output);
-		if (OutputPin && NewNodeReturnValuePin)
+		UActorComponent* BoundComponent = Cast<UActorComponent>(LocateBoundObject(Possessable));
+		if (!BoundComponent)
 		{
-			// Connect the nodes.
-			NewNodeReturnValuePin->MakeLinkTo(OutputPin);
+			continue;
 		}
-	}
 
-	// Connect additional parameter pins from the endpoint to the function call node
-	// Reference: FMovieSceneDirectorBlueprintEndpointCustomization::SetEndpoint
-	TSet<FName> NonPayloadPins;
-	const FString* WorldContextParamName = Possessable->DynamicBinding.Function->FindMetaData(FBlueprintMetadata::MD_WorldContext);
-	if (WorldContextParamName)
-	{
-		NonPayloadPins.Add(FName(*WorldContextParamName));
-	}
-	if (UK2Node_EditablePinBase* EditableNode = Cast<UK2Node_EditablePinBase>(NewEndpoint))
-	{
-		for (UEdGraphPin* PayloadPin : NewNode->Pins)
+		FString SourcePrimPath = GetSourcePrimPath(BoundComponent);
+		if (SourcePrimPath.IsEmpty())
 		{
-			if (PayloadPin
-				&& PayloadPin->LinkedTo.Num() == 0
-				&& PayloadPin->Direction == EGPD_Input
-				&& PayloadPin->PinType.PinCategory != UEdGraphSchema_K2::PC_Exec
-				&& PayloadPin->PinName != UEdGraphSchema_K2::PN_Self
-				&& !NonPayloadPins.Contains(PayloadPin->PinName
-			))
-			{
-				// Pins for ref parameters for functions default to bIsReference but the payload cannot be by
-				// reference.
-				PayloadPin->PinType.bIsReference = false;
-
-				UEdGraphPin* NewPin = EditableNode->CreateUserDefinedPin(PayloadPin->PinName, PayloadPin->PinType, EGPD_Output);
-				if (NewNode != NewEndpoint && NewPin)
-				{
-					NewPin->MakeLinkTo(PayloadPin);
-				}
-			}
+			continue;
 		}
-	}
 
-	FMovieSceneDynamicBindingUtils::SetEndpoint(MovieScene, &Possessable->DynamicBinding, NewEndpoint);
-	FMovieSceneDynamicBindingUtils::EnsureBlueprintExtensionCreated(Sequence, DirectorBlueprint);
-	FKismetEditorUtilities::CompileBlueprint(DirectorBlueprint);
+		// Only mark it as handled if we found the right stage actor for this component. Other stage actors that don't
+		// own this component will all return empty strings for SourcePrimPath
+		AssetUserData->HandledBindingGuids.Add(BindingGuid);
+
+		FUsdStageActorImpl::SetupDynamicBinding(SourcePrimPath, Possessable, Sequence, *GetName());
+	}
 #endif	  // WITH_EDITOR
 }
 
diff --git a/Engine/Plugins/Importers/USDImporter/Source/USDStage/Public/USDStageActor.h b/Engine/Plugins/Importers/USDImporter/Source/USDStage/Public/USDStageActor.h
--- a/Engine/Plugins/Importers/USDImporter/Source/USDStage/Public/USDStageActor.h
+++ b/Engine/Plugins/Importers/USDImporter/Source/USDStage/Public/USDStageActor.h
@@ -20,6 +20,7 @@ class UUsdAssetCache;
 class UUsdAssetCache2;
 class UUsdPrimTwin;
 class UUsdTransactor;
+enum class EMovieSceneDataChangeType;
 struct FUsdSchemaTranslationContext;
 namespace UE
 {
@@ -293,6 +294,7 @@ protected:
 	void UpdateSpawnedObjectsTransientFlag(bool bTransient);
 
 	void OnActorAddedToSequencer(AActor* NewActor, const FGuid Guid, TWeakPtr<ISequencer> WeakSequencer);
+	void OnMovieSceneDataChanged(EMovieSceneDataChangeType ChangeType, TWeakPtr<ISequencer> WeakSequencer);
 
 #if WITH_EDITOR
 	void OnBeginPIE(bool bIsSimulating);
diff --git a/Engine/Source/Runtime/MovieScene/Private/Evaluation/MovieSceneEvaluationState.cpp b/Engine/Source/Runtime/MovieScene/Private/Evaluation/MovieSceneEvaluationState.cpp
--- a/Engine/Source/Runtime/MovieScene/Private/Evaluation/MovieSceneEvaluationState.cpp
+++ b/Engine/Source/Runtime/MovieScene/Private/Evaluation/MovieSceneEvaluationState.cpp
@@ -365,7 +365,16 @@ void FMovieSceneObjectCache::UpdateBindings(const FGuid& InGuid, IMovieScenePlay
 					}
 					else
 					{
-						Player.ResolveBoundObjects(InGuid, SequenceID, *Sequence, ResolutionContext, FoundObjects);
+						FMovieSceneDynamicBindingResolveResult ResolveResult = FMovieSceneDynamicBindingInvoker::ResolveDynamicBinding(Player, Sequence, SequenceID, InGuid, Possessable->DynamicBinding);
+						if (ResolveResult.Object)
+						{
+							ensureMsgf(ResolveResult.bIsPossessedObject, TEXT("Possessables don't support spawnable-type ownership"));
+							FoundObjects.Add(ResolveResult.Object);
+						}
+						else
+						{
+							Player.ResolveBoundObjects(InGuid, SequenceID, *Sequence, ResolutionContext, FoundObjects);
+						}
 					}
 					
 					Bindings = BoundObjects.Find(InGuid);
