USD: Fix log message using LogTemp insterad of LogUsd.

diff --git a/Engine/Plugins/Importers/USDImporter/Source/USDUtilities/Private/USDShadeConversion.cpp b/Engine/Plugins/Importers/USDImporter/Source/USDUtilities/Private/USDShadeConversion.cpp
--- a/Engine/Plugins/Importers/USDImporter/Source/USDUtilities/Private/USDShadeConversion.cpp
+++ b/Engine/Plugins/Importers/USDImporter/Source/USDUtilities/Private/USDShadeConversion.cpp
@@ -518,7 +518,7 @@ namespace UE
 					else
 					{
 						UE_LOG(
-							LogTemp,
+							LogUsd,
 							Warning,
 							TEXT("Unexpected shader node id '%s' when traversing texture node '%s' for primvars!"),
 							*UsdToUnreal::ConvertToken(ConnectedNodeId),
