USD: Fix mistake on 24632828 causing inaccurate triangle count when deciding whether to enable Nanite or not.

diff --git a/Engine/Plugins/Importers/USDImporter/Source/USDSchemas/Private/USDGeomMeshTranslator.cpp b/Engine/Plugins/Importers/USDImporter/Source/USDSchemas/Private/USDGeomMeshTranslator.cpp
--- a/Engine/Plugins/Importers/USDImporter/Source/USDSchemas/Private/USDGeomMeshTranslator.cpp
+++ b/Engine/Plugins/Importers/USDImporter/Source/USDSchemas/Private/USDGeomMeshTranslator.cpp
@@ -124,7 +124,7 @@ namespace UsdGeomMeshTranslatorImpl
 		// We want Nanite because the mesh is large enough for the threshold, which is set to something valid
 		if (!bHasNaniteOverrideEnabled && Context.InfoCache.IsValid())
 		{
-			uint64 NumTriangles = Context.InfoCache->GetSubtreeVertexCount(PrimPath).Get(0);
+			const int32 NumTriangles = LODIndexToMeshDescription[0].Triangles().Num();
 			if (NumTriangles >= Context.NaniteTriangleThreshold)
 			{
 				UE_LOG(LogUsd, Verbose, TEXT("Trying to enable Nanite for mesh generated for prim '%s' as it has '%d' triangles, and the threshold is '%d'"),
