USD: Improved heuristics for finding the collapse root for an imported GeometryCache.

diff --git a/Engine/Plugins/Importers/USDImporter/Source/USDSchemas/Private/USDGeometryCacheTranslator.cpp b/Engine/Plugins/Importers/USDImporter/Source/USDSchemas/Private/USDGeometryCacheTranslator.cpp
--- a/Engine/Plugins/Importers/USDImporter/Source/USDSchemas/Private/USDGeometryCacheTranslator.cpp
+++ b/Engine/Plugins/Importers/USDImporter/Source/USDSchemas/Private/USDGeometryCacheTranslator.cpp
@@ -44,13 +44,6 @@
 	#include "pxr/usd/usdShade/materialBindingAPI.h"
 #include "USDIncludesEnd.h"
 
-// Can toggle on/off to compare performance with StaticMesh instead of GeometryCache
-static bool GUseGeometryCacheUSD = true;
-static FAutoConsoleVariableRef CVarUsdUseGeometryCache(
-	TEXT("USD.GeometryCache.Enable"),
-	GUseGeometryCacheUSD,
-	TEXT("Use GeometryCache instead of static meshes for loading animated meshes"));
-
 static int32 GUsdGeometryCacheParallelMeshReads = 16;
 static FAutoConsoleVariableRef CVarUsdGeometryCacheParallelMeshReads(
 	TEXT("USD.GeometryCache.Import.ParallelMeshReads"),
@@ -1019,11 +1012,19 @@ void FUsdGeometryCacheTranslator::UpdateComponents(USceneComponent* SceneCompone
 
 bool FUsdGeometryCacheTranslator::CollapsesChildren(ECollapsingType CollapsingType) const
 {
+	if (!Context->InfoCache.IsValid())
+	{
+		return Super::CollapsesChildren(CollapsingType);
+	}
 	return IsPotentialGeometryCacheRoot() ? true : Super::CollapsesChildren(CollapsingType);
 }
 
 bool FUsdGeometryCacheTranslator::CanBeCollapsed(ECollapsingType CollapsingType) const
 {
+	if (!Context->InfoCache.IsValid())
+	{
+		return Super::CanBeCollapsed(CollapsingType);
+	}
 	return IsPotentialGeometryCacheRoot() ? false : Super::CanBeCollapsed(CollapsingType);
 }
 
@@ -1057,53 +1058,10 @@ TSet<UE::FSdfPath> FUsdGeometryCacheTranslator::CollectAuxiliaryPrims() const
 
 bool FUsdGeometryCacheTranslator::IsPotentialGeometryCacheRoot() const
 {
-	if (!GUseGeometryCacheUSD)
-	{
-		return false;
-	}
-
-	if (!Context->bIsImporting)
-	{
-		// For now, when loading on the stage, process only leaf animated mesh as geometry cache (without collapsing onto a root)
-		return UsdUtils::IsAnimatedMesh(GetPrim());
-	}
-
-	// Cache is not available when it is being rebuilt. This is called by RecursiveQueryCollapsesChildren
-	bool bIsPotentialGeoCacheRoot = false;
-	if (Context->InfoCache)
-	{
-		bIsPotentialGeoCacheRoot = Context->InfoCache->IsPotentialGeometryCacheRoot(PrimPath);
-	}
-	else
-	{
-		FScopedUsdAllocs UsdAllocs;
-		TArray<TUsdStore<pxr::UsdPrim>> ChildPrims = UsdUtils::GetAllPrimsOfType(GetPrim(), pxr::TfType::Find<pxr::UsdGeomMesh>());
-		for (const TUsdStore<pxr::UsdPrim>& ChildPrim : ChildPrims)
-		{
-			if (UsdUtils::IsAnimatedMesh(ChildPrim.Get()))
-			{
-				bIsPotentialGeoCacheRoot = true;
-				break;
-			}
-		}
-	}
-
-	// #uent_todo: Improve the collapsing rule for GeometryCache
-	if (bIsPotentialGeoCacheRoot)
-	{
-		FString FilePath;
-		if (UsdUtils::GetReferenceFilePath(GetPrim(), TEXT("abc"), FilePath))
-		{
-			return true;
-		}
-		else
-		{
-			// TBD where the collapsing should occur. Currently, if the stage has any animated meshes, it will collapse everything at the root
-			return true;
-		}
-	}
+	check(Context->InfoCache.IsValid());
 
-	return false;
+	// The logic to check for GeometryCache is completely in the InfoCache
+	return Context->InfoCache->IsPotentialGeometryCacheRoot(PrimPath);
 }
 
 #endif // #if USE_USD_SDK
diff --git a/Engine/Plugins/Importers/USDImporter/Source/USDSchemas/Private/USDInfoCache.cpp b/Engine/Plugins/Importers/USDImporter/Source/USDSchemas/Private/USDInfoCache.cpp
--- a/Engine/Plugins/Importers/USDImporter/Source/USDSchemas/Private/USDInfoCache.cpp
+++ b/Engine/Plugins/Importers/USDImporter/Source/USDSchemas/Private/USDInfoCache.cpp
@@ -14,15 +14,20 @@
 #include "UsdWrappers/UsdPrim.h"
 
 #include "Async/ParallelFor.h"
+#include "Math/NumericLimits.h"
+#include "Misc/Paths.h"
 #include "Modules/ModuleManager.h"
 
 #if USE_USD_SDK
 #include "USDIncludesStart.h"
 	#include "pxr/usd/sdf/path.h"
 	#include "pxr/usd/usd/prim.h"
+	#include "pxr/usd/usd/primCompositionQuery.h"
 	#include "pxr/usd/usdGeom/mesh.h"
 	#include "pxr/usd/usdGeom/pointInstancer.h"
+	#include "pxr/usd/usdGeom/scope.h"
 	#include "pxr/usd/usdGeom/subset.h"
+	#include "pxr/usd/usdGeom/xform.h"
 	#include "pxr/usd/usdShade/materialBindingAPI.h"
 	#include "pxr/usd/usdSkel/root.h"
 #include "USDIncludesEnd.h"
@@ -34,15 +39,41 @@ static FAutoConsoleVariableRef CVarMaxNumVerticesCollapsedMesh(
 	GMaxNumVerticesCollapsedMesh,
 	TEXT("Maximum number of vertices that a combined Mesh can have for us to collapse it into a single StaticMesh"));
 
+// Can toggle on/off to compare performance with StaticMesh instead of GeometryCache
+static bool GUseGeometryCacheUSD = true;
+static FAutoConsoleVariableRef CVarUsdUseGeometryCache(
+	TEXT("USD.GeometryCache.Enable"),
+	GUseGeometryCacheUSD,
+	TEXT("Use GeometryCache instead of static meshes for loading animated meshes"));
+
+static int32 GGeometryCacheMaxDepth = 15;
+static FAutoConsoleVariableRef CVarGeometryCacheMaxDepth(
+	TEXT("USD.GeometryCache.MaxDepth"),
+	GGeometryCacheMaxDepth,
+	TEXT("Maximum distance between an animated mesh prim to its collapsed geometry cache root"));
+
 namespace UE::UsdInfoCache::Private
 {
+	// Flags to hint at the state of a prim for the purpose of geometry cache
+	enum class EGeometryCachePrimState : uint8
+	{
+		None = 0x00,
+		Uncollapsible = 0x01,		// prim cannot be collapsed as part of a geometry cache
+		Mesh = 0x02,				// prim is a mesh, animated or not
+		Xform = 0x04,				// prim is a xform, animated or not
+		Collapsible = Mesh | Xform,	// only meshes and xforms can be collapsed into a geometry cache
+		ValidRoot = 0x08			// prim can collapse itself and its children into a geometry cache
+	};
+	ENUM_CLASS_FLAGS(EGeometryCachePrimState)
+
 	struct FUsdPrimInfo
 	{
 		UE::FSdfPath AssetCollapsedRoot;
 		UE::FSdfPath ComponentCollapsedRoot;
-		uint64 ExpectedVertexCountForSubtree;
+		uint64 ExpectedVertexCountForSubtree = 0;
 		TArray<UsdUtils::FUsdPrimMaterialSlot> SubtreeMaterialSlots;
-		bool bIsPotentialGeoCacheRoot;
+		int32 GeometryCacheDepth = -1;
+		EGeometryCachePrimState GeometryCacheState = EGeometryCachePrimState::None;
 	};
 }
 
@@ -52,13 +83,20 @@ FArchive& operator <<( FArchive& Ar, UE::UsdInfoCache::Private::FUsdPrimInfo& In
 	Ar << Info.ComponentCollapsedRoot;
 	Ar << Info.ExpectedVertexCountForSubtree;
 	Ar << Info.SubtreeMaterialSlots;
-	Ar << Info.bIsPotentialGeoCacheRoot;
+	Ar << Info.GeometryCacheDepth;
+	Ar << Info.GeometryCacheState;
 
 	return Ar;
 }
 
 struct FUsdInfoCache::FUsdInfoCacheImpl
 {
+	FUsdInfoCacheImpl()
+	: AllowedExtensionsForGeometryCacheSource(UnrealUSDWrapper::GetNativeFileFormats())
+	{
+		AllowedExtensionsForGeometryCacheSource.Add(TEXT("abc"));
+	}
+
 	// Information we must have about all prims on the stage
 	TMap< UE::FSdfPath, UE::UsdInfoCache::Private::FUsdPrimInfo > InfoMap;
 	mutable FRWLock InfoMapLock;
@@ -78,6 +116,9 @@ struct FUsdInfoCache::FUsdInfoCacheImpl
 	TMap<UE::FSdfPath, TSet<UE::FSdfPath>> AuxToMainPrims;
 	TMap<UE::FSdfPath, TSet<UE::FSdfPath>> MainToAuxPrims;
 	mutable FRWLock AuxiliaryPrimsLock;
+
+	// Geometry cache can come from a reference or payload of these file types
+	TArray<FString> AllowedExtensionsForGeometryCacheSource;
 };
 
 FUsdInfoCache::FUsdInfoCache()
@@ -740,12 +781,9 @@ namespace UE::USDInfoCacheImpl::Private
 				const UE::FSdfPath* CollapsedRoot = &FoundInfo->AssetCollapsedRoot;
 
 				// We only merge slots in the context of collapsing
-				bool bPrimIsCollapsedOrCollapseRoot = !CollapsedRoot->IsEmpty() || PrimPath.IsAbsoluteRootPath();
-
-				// TODO: This is not perfect, because we may have bPrimIsPotentialGeometryCacheRoot but the prim hasn't
-				// actually become a geometry cache due to another reason
-				bool bPrimIsPotentialGeometryCacheRoot = FoundInfo->bIsPotentialGeoCacheRoot;
+				const bool bPrimIsCollapsedOrCollapseRoot = !CollapsedRoot->IsEmpty() || PrimPath.IsAbsoluteRootPath();
 
+				const bool bPrimIsPotentialGeometryCacheRoot = FoundInfo->GeometryCacheState == UE::UsdInfoCache::Private::EGeometryCachePrimState::ValidRoot;
 				bCanMergeSlotsForThisPrim = bPrimIsCollapsedOrCollapseRoot && !bPrimIsPotentialGeometryCacheRoot;
 			}
 
@@ -822,9 +860,13 @@ namespace UE::USDInfoCacheImpl::Private
 			{
 				const bool bCanMeshSubtreeBeCollapsed = CanMeshSubtreeBeCollapsed(UsdPrim, Context, Impl, SchemaTranslator);
 
+				const UE::UsdInfoCache::Private::FUsdPrimInfo& Info = Impl.InfoMap.FindChecked(UE::FSdfPath{UsdPrimPath});
+				const bool bIsPotentialGeometryCacheRoot = Info.GeometryCacheState ==  UE::UsdInfoCache::Private::EGeometryCachePrimState::ValidRoot;
+
 				if (!bIsAssetCollapsed)
 				{
-					if (SchemaTranslator->CollapsesChildren(ECollapsingType::Assets) && bCanMeshSubtreeBeCollapsed)
+					// The potential geometry cache root is checked first since the FUsdGeometryCacheTranslator::CollapsesChildren has no logic of its own
+					if (bIsPotentialGeometryCacheRoot || (SchemaTranslator->CollapsesChildren(ECollapsingType::Assets) && bCanMeshSubtreeBeCollapsed))
 					{
 						AssetCollapsedRootOverride = &UsdPrimPath;
 					}
@@ -832,7 +874,7 @@ namespace UE::USDInfoCacheImpl::Private
 
 				if (!bIsComponentCollapsed)
 				{
-					if (SchemaTranslator->CollapsesChildren(ECollapsingType::Components) && bCanMeshSubtreeBeCollapsed)
+					if (bIsPotentialGeometryCacheRoot || (SchemaTranslator->CollapsesChildren(ECollapsingType::Components) && bCanMeshSubtreeBeCollapsed))
 					{
 						ComponentCollapsedRootOverride = &UsdPrimPath;
 					}
@@ -877,11 +919,69 @@ namespace UE::USDInfoCacheImpl::Private
 		}
 	}
 
-	bool RecursiveCheckForGeometryCache(const pxr::UsdPrim& UsdPrim, FUsdSchemaTranslationContext& Context, FUsdInfoCache::FUsdInfoCacheImpl& Impl)
+	void FindValidGeometryCacheRoot(const pxr::UsdPrim& UsdPrim, FUsdSchemaTranslationContext& Context, FUsdInfoCache::FUsdInfoCacheImpl& Impl, UE::UsdInfoCache::Private::EGeometryCachePrimState& OutState)
 	{
-		TRACE_CPUPROFILER_EVENT_SCOPE(UE::USDInfoCacheImpl::Private::RecursiveCheckForGeometryCache);
+		TRACE_CPUPROFILER_EVENT_SCOPE(FindValidGeometryCacheRoot);
+
+		using namespace UE::UsdInfoCache::Private;
+
 		FScopedUsdAllocs Allocs;
 
+		pxr::SdfPath UsdPrimPath = UsdPrim.GetPrimPath();
+		UE::FSdfPath PrimPath{UsdPrimPath};
+		{
+			FReadScopeLock ScopeLock(Impl.InfoMapLock);
+
+			UE::UsdInfoCache::Private::FUsdPrimInfo& Info = Impl.InfoMap.FindChecked(UE::FSdfPath(UsdPrim.GetPrimPath()));
+
+			// A prim is considered a valid root if its subtree has no uncollapsible branch and a valid depth.
+			// A valid depth is positive, meaning it has an animated mesh, and doesn't exceed the limit.
+			bool bIsValidDepth = Info.GeometryCacheDepth > -1 && Info.GeometryCacheDepth <= GGeometryCacheMaxDepth;
+			if (!EnumHasAnyFlags(Info.GeometryCacheState, EGeometryCachePrimState::Uncollapsible) && bIsValidDepth)
+			{
+				OutState = EGeometryCachePrimState::ValidRoot;
+				Info.GeometryCacheState = EGeometryCachePrimState::ValidRoot;
+				return;
+			}
+			// The prim is not a valid root so it's flagged as uncollapsible since the root will be among its children
+			// and the eventual geometry cache cannot be collapsed.
+			else
+			{
+				OutState = EGeometryCachePrimState::Uncollapsible;
+				Info.GeometryCacheState = EGeometryCachePrimState::Uncollapsible;
+			}
+		}
+
+		pxr::UsdPrimSiblingRange PrimChildren = UsdPrim.GetFilteredChildren(pxr::UsdTraverseInstanceProxies(pxr::UsdPrimAllPrimsPredicate));
+
+		// Continue the search for a valid root among the children
+		TArray<pxr::UsdPrim> Prims;
+		for (pxr::UsdPrim Child : PrimChildren)
+		{
+			FReadScopeLock ScopeLock(Impl.InfoMapLock);
+			const UE::UsdInfoCache::Private::FUsdPrimInfo& Info = Impl.InfoMap.FindChecked(UE::FSdfPath(Child.GetPrimPath()));
+
+			// A subtree is considered only if it has anything collapsible in the first place
+			if (EnumHasAnyFlags(Info.GeometryCacheState, EGeometryCachePrimState::Collapsible))
+			{
+				FindValidGeometryCacheRoot(Child, Context, Impl, OutState);
+			}
+		}
+
+		OutState = EGeometryCachePrimState::Uncollapsible;
+	}
+
+	void RecursiveCheckForGeometryCache(const pxr::UsdPrim& UsdPrim, FUsdSchemaTranslationContext& Context, FUsdInfoCache::FUsdInfoCacheImpl& Impl, int32& OutDepth, UE::UsdInfoCache::Private::EGeometryCachePrimState& OutState)
+	{
+		TRACE_CPUPROFILER_EVENT_SCOPE(RecursiveCheckForGeometryCache);
+
+		using namespace UE::UsdInfoCache::Private;
+
+		FScopedUsdAllocs Allocs;
+
+		// With this recursive check for geometry cache, we want to find branches with an animated mesh at the leaf and find the root where they can meet.
+		// This root prim will collapses the static and animated meshes under it into a single geometry cache.
+
 		pxr::SdfPath UsdPrimPath = UsdPrim.GetPrimPath();
 		UE::FSdfPath PrimPath{UsdPrimPath};
 
@@ -893,27 +993,222 @@ namespace UE::USDInfoCacheImpl::Private
 			Prims.Emplace(Child);
 		}
 
-		// TODO: This can probably be optimized further as we're not going to be a geometry cache root if we have
-		// child prims that can't be collapsed, etc.
-		std::atomic<bool> bIsPotentialGeoCacheRoot = false;
+		TArray<int32> Depths;
+		Depths.SetNum(Prims.Num());
+
+		TArray<EGeometryCachePrimState> States;
+		States.SetNum(Prims.Num());
+
 		const int32 MinBatchSize = 1;
 		ParallelFor(TEXT("RecursiveCheckForGeometryCache"), Prims.Num(), MinBatchSize,
-			[&bIsPotentialGeoCacheRoot, &Prims, &Context, &Impl](int32 Index)
+			[&Prims, &Context, &Impl, &Depths, &States](int32 Index)
 			{
-				bool bResult = RecursiveCheckForGeometryCache(Prims[Index], Context, Impl);
-				bIsPotentialGeoCacheRoot = bIsPotentialGeoCacheRoot || bResult;
+				RecursiveCheckForGeometryCache(Prims[Index], Context, Impl, Depths[Index], States[Index]);
 			}
 		);
 
 		bool bIsAnimatedMesh = UsdUtils::IsAnimatedMesh(UsdPrim);
+		if (!Context.bIsImporting)
 		{
 			FWriteScopeLock ScopeLock(Impl.InfoMapLock);
 
-			UE::UsdInfoCache::Private::FUsdPrimInfo& Info = Impl.InfoMap.FindOrAdd(UE::FSdfPath(UsdPrimPath));
-			Info.bIsPotentialGeoCacheRoot = bIsPotentialGeoCacheRoot || bIsAnimatedMesh;
+			UE::UsdInfoCache::Private::FUsdPrimInfo& Info = Impl.InfoMap.FindChecked(UE::FSdfPath(UsdPrimPath));
+
+			// When loading on the stage, the GeometryCache root can only be the animated mesh prim itself
+			// and there's no collapsing involved since each animated mesh will become a GeometryCache.
+			// The depth is irrelevant here.
+			Info.GeometryCacheDepth = -1;
+			Info.GeometryCacheState = bIsAnimatedMesh ? EGeometryCachePrimState::ValidRoot : EGeometryCachePrimState::Uncollapsible;
+
+			return;
+		}
+
+		// A geometry cache "branch" starts from an animated mesh prim for which we assign a depth of 0
+		// Other branches, without any animated mesh, we don't care about and will remain at -1
+		int32 Depth = -1;
+		if (bIsAnimatedMesh)
+		{
+			Depth = 0;
+		}
+		else
+		{
+			// The depth is propagated from children to parent, incremented by 1 at each level,
+			// with the parent depth being the deepest of its children depth
+			int32 ChildDepth = -1;
+			for (int32 Index = 0; Index < Depths.Num(); ++Index)
+			{
+				if (Depths[Index] > -1)
+				{
+					ChildDepth = FMath::Max(ChildDepth, Depths[Index] + 1);
+				}
+			}
+			Depth = ChildDepth;
+		}
+
+		// Along with the depth, we want some hints on the content of the subtree of the prim as this will tell us
+		// if the prim can serve as a root and collapse its children into a GeometryCache. The sole condition for
+		// being a valid root is that all the branches of the subtree are collapsible.
+		EGeometryCachePrimState ChildrenState = EGeometryCachePrimState::None;
+		for (EGeometryCachePrimState ChildState : States) 
+		{
+			ChildrenState |= ChildState;
 		}
 
-		return bIsPotentialGeoCacheRoot || bIsAnimatedMesh;
+		EGeometryCachePrimState PrimState = EGeometryCachePrimState::None;
+		const bool bIsMesh = !!pxr::UsdGeomMesh(UsdPrim);
+		const bool bIsXform = !!pxr::UsdGeomXform(UsdPrim);
+		if (bIsMesh)
+		{
+			// Animated or static mesh. Static meshes could potentially be animated by transforms in their hierarchy.
+			// A mesh prim should be a leaf, but it can have GeomSubset prims as children, but those don't
+			// affect the collapsibility status.
+			PrimState = EGeometryCachePrimState::Mesh;
+		}
+		else if (bIsXform)
+		{
+			// An xform prim is considered collapsible since it could have a mesh prim under it. It has to bubble up its children state.
+			PrimState = ChildrenState != EGeometryCachePrimState::None ? ChildrenState | EGeometryCachePrimState::Xform : EGeometryCachePrimState::Xform;
+		}
+		else
+		{
+			// This prim is not considered collapsible with some exception
+			// Like a Scope could have some meshes under it, so it has to bubble up its children state
+			const bool bIsException = !!pxr::UsdGeomScope(UsdPrim);
+			if (bIsException && EnumHasAnyFlags(ChildrenState, EGeometryCachePrimState::Mesh))
+			{
+				PrimState = ChildrenState;
+			}
+			else
+			{
+				PrimState = EGeometryCachePrimState::Uncollapsible;
+			}
+		}
+
+		// A prim could be a potential root if it has a reference or payload to an allowed file type for GeometryCache
+		bool bIsPotentialRoot = false;
+		{
+			pxr::UsdPrimCompositionQuery PrimCompositionQuery = pxr::UsdPrimCompositionQuery::GetDirectReferences(UsdPrim);
+			for (const pxr::UsdPrimCompositionQueryArc& CompositionArc : PrimCompositionQuery.GetCompositionArcs())
+			{
+				if (CompositionArc.GetArcType() == pxr::PcpArcTypeReference)
+				{
+					pxr::SdfReferenceEditorProxy ReferenceEditor;
+					pxr::SdfReference UsdReference;
+
+					if (CompositionArc.GetIntroducingListEditor(&ReferenceEditor, &UsdReference))
+					{
+						FString FilePath = UsdToUnreal::ConvertString(UsdReference.GetAssetPath());
+						FString Extension = FPaths::GetExtension(FilePath);
+
+						if (Impl.AllowedExtensionsForGeometryCacheSource.Contains(Extension))
+						{
+							bIsPotentialRoot = true;
+							break;
+						}
+					}
+				}
+				else if (CompositionArc.GetArcType() == pxr::PcpArcTypePayload)
+				{
+					pxr::SdfPayloadEditorProxy PayloadEditor;
+					pxr::SdfPayload UsdPayload;
+
+					if (CompositionArc.GetIntroducingListEditor(&PayloadEditor, &UsdPayload))
+					{
+						FString FilePath = UsdToUnreal::ConvertString(UsdPayload.GetAssetPath());
+						FString Extension = FPaths::GetExtension(FilePath);
+
+						if (Impl.AllowedExtensionsForGeometryCacheSource.Contains(Extension))
+						{
+							bIsPotentialRoot = true;
+							break;
+						}
+					}
+				}
+			}
+		}
+
+		{
+			FWriteScopeLock ScopeLock(Impl.InfoMapLock);
+
+			UE::UsdInfoCache::Private::FUsdPrimInfo& Info = Impl.InfoMap.FindChecked(UE::FSdfPath(UsdPrimPath));
+			Info.GeometryCacheDepth = Depth;
+			Info.GeometryCacheState = PrimState;
+		}
+
+		// We've encountered a potential root and the subtree has a geometry cache branch, so find its root
+		if (bIsPotentialRoot && Depth > -1)
+		{
+			if (Depth > GGeometryCacheMaxDepth)
+			{
+				UE_LOG(LogUsd, Warning, TEXT("Prim '%s' is potentially a geometry cache %d levels deep, which exceeds the limit of %d. "
+					"This could affect its imported animation. The limit can be increased with the cvar USD.GeometryCache.MaxDepth if needed."),
+					*PrimPath.GetString(), Depth, GGeometryCacheMaxDepth);
+			}
+			FindValidGeometryCacheRoot(UsdPrim, Context, Impl, PrimState);
+			Depth = -1;
+		}
+
+		OutDepth = Depth;
+		OutState = PrimState;
+	}
+
+	void CheckForGeometryCache(const pxr::UsdPrim& UsdPrim, FUsdSchemaTranslationContext& Context, FUsdInfoCache::FUsdInfoCacheImpl& Impl)
+	{
+		TRACE_CPUPROFILER_EVENT_SCOPE(CheckForGeometryCache);
+
+		using namespace UE::UsdInfoCache::Private;
+
+		if (!GUseGeometryCacheUSD)
+		{
+			return;
+		}
+
+		// If the stage doesn't contain any animated mesh prims, then don't bother doing a full check
+		bool bHasAnimatedMesh = false;
+		{
+			FScopedUsdAllocs UsdAllocs;
+			TArray<TUsdStore<pxr::UsdPrim>> ChildPrims = UsdUtils::GetAllPrimsOfType(UsdPrim, pxr::TfType::Find<pxr::UsdGeomMesh>());
+			for (const TUsdStore<pxr::UsdPrim>& ChildPrim : ChildPrims)
+			{
+				if (UsdUtils::IsAnimatedMesh(ChildPrim.Get()))
+				{
+					bHasAnimatedMesh = true;
+					break;
+				}
+			}
+		}
+
+		if (!bHasAnimatedMesh)
+		{
+			return;
+		}
+
+		int32 Depth = -1;
+		EGeometryCachePrimState State = EGeometryCachePrimState::None;
+		RecursiveCheckForGeometryCache(UsdPrim, Context, Impl, Depth, State);
+
+		// If we end up with a positive depth, it means the check found an animated mesh somewhere
+		// but no potential root before reaching the pseudoroot, so find one
+		if (Depth > -1)
+		{
+			if (Depth > GGeometryCacheMaxDepth)
+			{
+				UE_LOG(LogUsd, Warning, TEXT("The stage has a geometry cache %d levels deep, which exceeds the limit of %d. "
+					"This could affect its imported animation. The limit can be increased with the cvar USD.GeometryCache.MaxDepth if needed."),
+					Depth, GGeometryCacheMaxDepth);
+			}
+
+			FScopedUsdAllocs UsdAllocs;
+
+			pxr::UsdPrimSiblingRange PrimChildren = UsdPrim.GetFilteredChildren(pxr::UsdTraverseInstanceProxies(pxr::UsdPrimAllPrimsPredicate));
+
+			// The pseudoroot itself cannot be a root for the geometry cache so start from its children
+			TArray<pxr::UsdPrim> Prims;
+			for (pxr::UsdPrim Child : PrimChildren)
+			{
+				FindValidGeometryCacheRoot(Child, Context, Impl, State);
+			}
+		}
 	}
 #endif // USE_SD_SDK
 }
@@ -926,7 +1221,7 @@ bool FUsdInfoCache::IsPotentialGeometryCacheRoot(const UE::FSdfPath& Path) const
 
 		if (const UE::UsdInfoCache::Private::FUsdPrimInfo* FoundInfo = ImplPtr->InfoMap.Find(Path))
 		{
-			return FoundInfo->bIsPotentialGeoCacheRoot;
+			return FoundInfo->GeometryCacheState == UE::UsdInfoCache::Private::EGeometryCachePrimState::ValidRoot;
 		}
 
 		// This should never happen: We should have cached the entire tree
@@ -1158,11 +1453,7 @@ void FUsdInfoCache::RebuildCacheForSubtree( const UE::FUsdPrim& Prim, FUsdSchema
 			TempSubtreeSlots
 		);
 
-		UE::USDInfoCacheImpl::Private::RecursiveCheckForGeometryCache(
-			UsdPrim,
-			Context,
-			*ImplPtr
-		);
+		UE::USDInfoCacheImpl::Private::CheckForGeometryCache(UsdPrim, Context, *ImplPtr);
 
 		UE::USDInfoCacheImpl::Private::RecursiveQueryCollapsesChildren(
 			UsdPrim,
@@ -1214,4 +1505,4 @@ bool FUsdInfoCache::IsEmpty()
 	}
 
 	return true;
-}
\ No newline at end of file
+}
