USD: Remove remaining MaterialX options from USDStageImportOptions missed from cl 25304674

diff --git a/Engine/Plugins/Importers/USDImporter/Source/USDSchemas/Public/USDSchemaTranslator.h b/Engine/Plugins/Importers/USDImporter/Source/USDSchemas/Public/USDSchemaTranslator.h
--- a/Engine/Plugins/Importers/USDImporter/Source/USDSchemas/Public/USDSchemaTranslator.h
+++ b/Engine/Plugins/Importers/USDImporter/Source/USDSchemas/Public/USDSchemaTranslator.h
@@ -26,7 +26,6 @@
 class FRegisteredSchemaTranslator;
 class FUsdSchemaTranslator;
 class FUsdSchemaTranslatorTaskChain;
-class UInterchangeGenericMaterialPipeline;
 class ULevel;
 class UMaterialInterface;
 class USceneComponent;
diff --git a/Engine/Plugins/Importers/USDImporter/Source/USDStageImporter/Private/USDStageImportOptions.cpp b/Engine/Plugins/Importers/USDImporter/Source/USDStageImporter/Private/USDStageImportOptions.cpp
--- a/Engine/Plugins/Importers/USDImporter/Source/USDStageImporter/Private/USDStageImportOptions.cpp
+++ b/Engine/Plugins/Importers/USDImporter/Source/USDStageImporter/Private/USDStageImportOptions.cpp
@@ -25,7 +25,6 @@ UUsdStageImportOptions::UUsdStageImportOptions(const FObjectInitializer& ObjectI
 	IUsdSchemasModule& UsdSchemasModule = FModuleManager::Get().LoadModuleChecked< IUsdSchemasModule >( TEXT("USDSchemas") );
 	RenderContextToImport = UsdSchemasModule.GetRenderContextRegistry().GetUnrealRenderContext();
 	MaterialPurpose = *UnrealIdentifiers::MaterialPreviewPurpose;
-	MaterialXOptions = CreateDefaultSubobject<UInterchangeGenericMaterialPipeline>("MaterialXOptions");
 	bOverrideStageOptions = false;
 	StageOptions.MetersPerUnit = 0.01f;
 	StageOptions.UpAxis = EUsdUpAxis::ZAxis;
diff --git a/Engine/Plugins/Importers/USDImporter/Source/USDStageImporter/Public/USDStageImportOptions.h b/Engine/Plugins/Importers/USDImporter/Source/USDStageImporter/Public/USDStageImportOptions.h
--- a/Engine/Plugins/Importers/USDImporter/Source/USDStageImporter/Public/USDStageImportOptions.h
+++ b/Engine/Plugins/Importers/USDImporter/Source/USDStageImporter/Public/USDStageImportOptions.h
@@ -13,7 +13,6 @@
 
 #include "USDStageImportOptions.generated.h"
 
-class UInterchangeGenericMaterialPipeline;
 struct FAnalyticsEventAttribute;
 
 UENUM(BlueprintType)
@@ -102,10 +101,6 @@ public:
 	UPROPERTY( BlueprintReadWrite, config, EditAnywhere, Category = "USD options", meta = ( EditCondition = bImportSkeletalAnimations ) )
 	EUsdRootMotionHandling RootMotionHandling = EUsdRootMotionHandling::NoAdditionalRootMotion;
 
-	// Settings to use when parsing MaterialX materials via Interchange
-	UPROPERTY(EditAnywhere, BlueprintReadOnly, Category = "DataToImport", AdvancedDisplay, meta = (DisplayName = "MaterialX Options"))
-	TObjectPtr<UInterchangeGenericMaterialPipeline> MaterialXOptions;
-
 	/** Whether to use the specified StageOptions instead of the stage's own settings */
 	UPROPERTY( BlueprintReadWrite, config, EditAnywhere, Category = "USD options" )
 	bool bOverrideStageOptions;
