USDImporter: add UnrealUSDWrapper functions for registering USD plugins by path(s)

diff --git a/Engine/Plugins/Importers/USDImporter/Source/UnrealUSDWrapper/Private/UnrealUSDWrapper.cpp b/Engine/Plugins/Importers/USDImporter/Source/UnrealUSDWrapper/Private/UnrealUSDWrapper.cpp
--- a/Engine/Plugins/Importers/USDImporter/Source/UnrealUSDWrapper/Private/UnrealUSDWrapper.cpp
+++ b/Engine/Plugins/Importers/USDImporter/Source/UnrealUSDWrapper/Private/UnrealUSDWrapper.cpp
@@ -519,6 +519,38 @@ double UnrealUSDWrapper::GetDefaultTimeCode()
 }
 #endif // USE_USD_SDK
 
+TArray<FString> UnrealUSDWrapper::RegisterPlugins(const FString& PathToPlugInfo)
+{
+	return RegisterPlugins(TArray<FString>{PathToPlugInfo});
+}
+
+TArray<FString> UnrealUSDWrapper::RegisterPlugins(const TArray<FString>& PathsToPlugInfo)
+{
+	TArray<FString> PluginNames;
+
+#if USE_USD_SDK
+	FScopedUsdAllocs Allocs;
+
+	std::vector<std::string> UsdPathsToPlugInfo;
+	UsdPathsToPlugInfo.reserve(PathsToPlugInfo.Num());
+
+	for (const FString& PathToPlugInfo : PathsToPlugInfo)
+	{
+		UsdPathsToPlugInfo.emplace_back(TCHAR_TO_UTF8(*PathToPlugInfo));
+	}
+
+	const pxr::PlugPluginPtrVector UsdPlugins =
+		pxr::PlugRegistry::GetInstance().RegisterPlugins(UsdPathsToPlugInfo);
+
+	for (const pxr::PlugPluginPtr& UsdPlugin : UsdPlugins)
+	{
+		PluginNames.Emplace(UTF8_TO_TCHAR(UsdPlugin->GetName().c_str()));
+	}
+#endif // USE_USD_SDK
+
+	return PluginNames;
+}
+
 TArray<FString> UnrealUSDWrapper::GetAllSupportedFileFormats()
 {
 	TArray<FString> Result;
diff --git a/Engine/Plugins/Importers/USDImporter/Source/UnrealUSDWrapper/Public/UnrealUSDWrapper.h b/Engine/Plugins/Importers/USDImporter/Source/UnrealUSDWrapper/Public/UnrealUSDWrapper.h
--- a/Engine/Plugins/Importers/USDImporter/Source/UnrealUSDWrapper/Public/UnrealUSDWrapper.h
+++ b/Engine/Plugins/Importers/USDImporter/Source/UnrealUSDWrapper/Public/UnrealUSDWrapper.h
@@ -179,6 +179,18 @@ public:
 	UNREALUSDWRAPPER_API static double GetDefaultTimeCode();
 #endif  // #if USE_USD_SDK
 
+	/**
+	 * Registers all USD plug-ins discovered at PathToPlugInfo.
+	 * @return An array containing the names of any newly registered plugins.
+	 */
+	UNREALUSDWRAPPER_API static TArray<FString> RegisterPlugins(const FString& PathToPlugInfo);
+
+	/**
+	 * Registers all USD plug-ins discovered in any of PathsToPlugInfo.
+	 * @return An array containing the names of any newly registered plugins.
+	 */
+	UNREALUSDWRAPPER_API static TArray<FString> RegisterPlugins(const TArray<FString>& PathsToPlugInfo);
+
 	/**
 	 * Returns the file extensions of all file formats supported by USD.
 	 *
