USD: Don't add Mesh prims into the SkeletalMesh if they lack the SkelBindingAPI.

diff --git a/Engine/Plugins/Importers/USDImporter/Source/USDSchemas/Private/USDSkelRootTranslator.cpp b/Engine/Plugins/Importers/USDImporter/Source/USDSchemas/Private/USDSkelRootTranslator.cpp
--- a/Engine/Plugins/Importers/USDImporter/Source/USDSchemas/Private/USDSkelRootTranslator.cpp
+++ b/Engine/Plugins/Importers/USDImporter/Source/USDSchemas/Private/USDSkelRootTranslator.cpp
@@ -489,7 +489,21 @@ namespace UsdSkelRootTranslatorImpl
 			// In USD, the skinning target need not be a mesh, but for Unreal we are only interested in skinning meshes
 			if ( pxr::UsdGeomMesh SkinningMesh = pxr::UsdGeomMesh( SkinningQuery.GetPrim() ) )
 			{
-				PathsToSkinnedPrims.Add(SkinningMesh.GetPrim().GetPath());
+				// Let's only care about prims with the SkelBindingAPI for now as we'll *need* joint influences and weights
+				if (SkinningQuery.GetPrim().HasAPI<pxr::UsdSkelBindingAPI>())
+				{
+					PathsToSkinnedPrims.Add(SkinningMesh.GetPrim().GetPath());
+				}
+				else
+				{
+					UE_LOG(
+						LogUsd,
+						Log,
+						TEXT("Ignoring skinned prim '%s' when generating Skeletal Mesh for SkelRoot '%s' as the prim doesn't have the SkelBindingAPI"),
+						*UsdToUnreal::ConvertPath(SkinningQuery.GetPrim().GetPrimPath()),
+						*UsdToUnreal::ConvertPath(InSkeletonRoot.GetPrim().GetPrimPath())
+					);
+				}
 			}
 		}
 
diff --git a/Engine/Plugins/Importers/USDImporter/Source/USDUtilities/Private/USDSkeletalDataConversion.cpp b/Engine/Plugins/Importers/USDImporter/Source/USDUtilities/Private/USDSkeletalDataConversion.cpp
--- a/Engine/Plugins/Importers/USDImporter/Source/USDUtilities/Private/USDSkeletalDataConversion.cpp
+++ b/Engine/Plugins/Importers/USDImporter/Source/USDUtilities/Private/USDSkeletalDataConversion.cpp
@@ -1254,6 +1254,10 @@ bool UsdToUnreal::ConvertSkinnedMesh(
 
 	const UsdPrim& SkinningPrim = SkinningQuery.GetPrim();
 	UsdSkelBindingAPI SkelBindingAPI(SkinningPrim);
+	if (!SkelBindingAPI)
+	{
+		return false;
+	}
 
 	// Ref. FFbxImporter::FillSkelMeshImporterFromFbx
 	UsdGeomMesh UsdMesh = UsdGeomMesh(SkinningPrim);
