USD: Fix GroomBinding assets not being set on components.

diff --git a/Engine/Plugins/Importers/USDImporter/Source/USDSchemas/Private/USDGroomTranslatorUtils.cpp b/Engine/Plugins/Importers/USDImporter/Source/USDSchemas/Private/USDGroomTranslatorUtils.cpp
--- a/Engine/Plugins/Importers/USDImporter/Source/USDSchemas/Private/USDGroomTranslatorUtils.cpp
+++ b/Engine/Plugins/Importers/USDImporter/Source/USDSchemas/Private/USDGroomTranslatorUtils.cpp
@@ -165,17 +165,16 @@ namespace UE::UsdGroomTranslatorUtils::Private
 			if (Targets.size() > 0)
 			{
 				const pxr::SdfPath& TargetPrimPath = Targets[0];
-				UObject* Mesh = InfoCache.GetSingleAssetForPrim(UE::FSdfPath{TargetPrimPath});
 
 				// Validate that the target prim and associated asset are of the expected type for the binding
 				pxr::UsdPrim TargetPrim = Prim.GetPrimAtPath(TargetPrimPath);
 				if (BindingType == EGroomBindingMeshType::SkeletalMesh && pxr::UsdSkelRoot(TargetPrim))
 				{
-					return Cast<USkeletalMesh>(Mesh);
+					return InfoCache.GetSingleAssetForPrim<USkeletalMesh>(UE::FSdfPath{TargetPrimPath});
 				}
 				else if (BindingType == EGroomBindingMeshType::GeometryCache && pxr::UsdGeomMesh(TargetPrim))
 				{
-					return Cast<UGeometryCache>(Mesh);
+					return InfoCache.GetSingleAssetForPrim<UGeometryCache>(UE::FSdfPath{TargetPrimPath});
 				}
 			}
 		}
@@ -216,13 +215,11 @@ namespace UsdGroomTranslatorUtils
 
 		// Determine the type of binding needed based on the prim mesh type
 		const FString PrimPath(UsdToUnreal::ConvertPath(Prim.GetPath()));
-		UObject* PrimAsset = InfoCache.GetSingleAssetForPrim(UE::FSdfPath{*PrimPath});
-
 		EGroomBindingMeshType GroomBindingType = EGroomBindingMeshType::SkeletalMesh;
-		UObject* TargetMesh = Cast<USkeletalMesh>(PrimAsset);
+		UObject* TargetMesh = InfoCache.GetSingleAssetForPrim<USkeletalMesh>(UE::FSdfPath{*PrimPath});
 		if (!TargetMesh)
 		{
-			TargetMesh = Cast<UGeometryCache>(PrimAsset);
+			TargetMesh = InfoCache.GetSingleAssetForPrim<UGeometryCache>(UE::FSdfPath{*PrimPath});
 			if (!TargetMesh)
 			{
 				return;
