USD: Fix "Save USD Layers" dialog showing when switching stage options.

diff --git a/Engine/Plugins/Importers/USDImporter/Source/USDStage/Private/USDStageActor.cpp b/Engine/Plugins/Importers/USDImporter/Source/USDStage/Private/USDStageActor.cpp
--- a/Engine/Plugins/Importers/USDImporter/Source/USDStage/Private/USDStageActor.cpp
+++ b/Engine/Plugins/Importers/USDImporter/Source/USDStage/Private/USDStageActor.cpp
@@ -2147,6 +2147,11 @@ void AUsdStageActor::OpenUsdStage()
 		AbsPath = RootLayer.FilePath;
 	}
 
+	if (UsdStage && (FPaths::IsSamePath(UsdStage.GetRootLayer().GetRealPath(), AbsPath)))
+	{
+		return;
+	}
+
 	OnPreStageChanged.Broadcast();
 
 	UsdStage = UnrealUSDWrapper::OpenStage(*AbsPath, InitialLoadSet);
@@ -2174,7 +2179,10 @@ void AUsdStageActor::OpenUsdStage()
 
 void AUsdStageActor::CloseUsdStage()
 {
-	OnPreStageChanged.Broadcast();
+	if (UsdStage)
+	{
+		OnPreStageChanged.Broadcast();
+	}
 
 	FUsdStageActorImpl::DiscardStage(UsdStage, this);
 	UsdStage = UE::FUsdStage();
@@ -2338,8 +2346,6 @@ void AUsdStageActor::LoadUsdStage()
 	FScopedSlowTask SlowTask(1.f, LOCTEXT("LoadingUDStage", "Loading USD Stage"));
 	SlowTask.MakeDialogDelayed(0.25f);
 
-	OnPreStageChanged.Broadcast();
-
 	// Block writing level sequence changes back to the USD stage until we finished this transaction, because once we do
 	// the movie scene and tracks will all trigger OnObjectTransacted. We listen for those on FUsdLevelSequenceHelperImpl::OnObjectTransacted,
 	// and would otherwise end up writing all of the data we just loaded back to the USD stage
