USD: Add support for exporting animation due to constraints to USD.

diff --git a/Engine/Plugins/Importers/USDImporter/Source/USDExporter/Private/LevelSequenceExporterUSD.cpp b/Engine/Plugins/Importers/USDImporter/Source/USDExporter/Private/LevelSequenceExporterUSD.cpp
--- a/Engine/Plugins/Importers/USDImporter/Source/USDExporter/Private/LevelSequenceExporterUSD.cpp
+++ b/Engine/Plugins/Importers/USDImporter/Source/USDExporter/Private/LevelSequenceExporterUSD.cpp
@@ -435,6 +435,9 @@ namespace UE
 				// File path of the exported USD root layer, in case we also exported the level along with the level sequence
 				FString LevelFilePath;
 
+				// World that is being animated for the LevelSequence bake
+				UWorld* World = nullptr;
+
 			public:
 				~FLevelSequenceExportContext()
 				{
@@ -851,11 +854,21 @@ namespace UE
 				double EndTimeCode = FFrameRate::TransformTime( EndFrameUETime, Resolution, StageFrameRate ).AsDecimal();
 				UsdStage.SetEndTimeCode( EndTimeCode );
 
+				const FConstraintsManagerController& Controller = FConstraintsManagerController::Get(Context.World);
+				static constexpr bool bSorted = true;
+				const TArray<TObjectPtr<UTickableConstraint>> AllConstraints = Controller.GetAllConstraints(bSorted);
+
 				for ( FFrameTime EvalTime = StartFrame; EvalTime <= EndFrame; EvalTime += Interval )
 				{
 					Context.Sequencer->SetLocalTimeDirectly( EvalTime );
 					Context.Sequencer->ForceEvaluate();
 
+					// Evaluate constraints (these run on tick in the editor, so here we must trigger them manually)
+					for (const UTickableConstraint* Constraint : AllConstraints)
+					{
+						Constraint->Evaluate();
+					}
+
 					FFrameTime KeyTime = FFrameRate::Snap( EvalTime, Resolution, DisplayRate ).FloorToFrame();
 					double UsdTimeCode = FFrameRate::TransformTime( KeyTime, Resolution, StageFrameRate ).AsDecimal();
 
@@ -1225,6 +1238,7 @@ bool ULevelSequenceExporterUsd::ExportBinary( UObject* Object, const TCHAR* Type
 	LevelSequenceExporterImpl::FLevelSequenceExportContext Context{ *LevelSequence, TempSequencer.ToSharedRef(), SpawnRegister.ToSharedRef() };
 	Context.ExportOptions = Options;
 	Context.bReplaceIdentical = ExportTask->bReplaceIdentical;
+	Context.World = World;
 
 	// Spawn (but hide) all spawnables so that they will also show up on the level export if we need them to.
 	// We have to traverse the template IDs when spawning spawnables, because we'll want to force each individual spawnable of each
diff --git a/Engine/Plugins/Importers/USDImporter/Source/USDExporter/USDExporter.Build.cs b/Engine/Plugins/Importers/USDImporter/Source/USDExporter/USDExporter.Build.cs
--- a/Engine/Plugins/Importers/USDImporter/Source/USDExporter/USDExporter.Build.cs
+++ b/Engine/Plugins/Importers/USDImporter/Source/USDExporter/USDExporter.Build.cs
@@ -26,6 +26,7 @@ namespace UnrealBuildTool.Rules
 					"Analytics",
 					"AnalyticsBlueprintLibrary",
 					"CinematicCamera",
+					"Constraints",
 					"ControlRig", // For detecting rigged skeletal animations when exporting level sequences to USD
 					"EditorStyle",
 					"GeometryCache",
