USD: Fix grooms and groom bindings not being imported;

Fix "Unknown asset 'X' encountered when collecting dependencies" when opening stages with grooms.

diff --git a/Engine/Plugins/Importers/USDImporter/Source/USDClasses/Private/USDClassesModule.cpp b/Engine/Plugins/Importers/USDImporter/Source/USDClasses/Private/USDClassesModule.cpp
--- a/Engine/Plugins/Importers/USDImporter/Source/USDClasses/Private/USDClassesModule.cpp
+++ b/Engine/Plugins/Importers/USDImporter/Source/USDClasses/Private/USDClassesModule.cpp
@@ -15,6 +15,9 @@
 #include "Engine/Texture2D.h"
 #include "EngineAnalytics.h"
 #include "GeometryCache.h"
+#include "GroomAsset.h"
+#include "GroomBindingAsset.h"
+#include "GroomCache.h"
 #include "HAL/FileManager.h"
 #include "Materials/Material.h"
 #include "Materials/MaterialInstance.h"
@@ -313,38 +316,22 @@ TSet<UObject*> IUsdClassesModule::GetAssetDependencies(UObject* Asset)
 		Result.Reserve(Result.Num() + MaterialInstance->TextureParameterValues.Num());
 		for (const FTextureParameterValue& TextureValue : MaterialInstance->TextureParameterValues)
 		{
-			if (UTexture* Texture = TextureValue.ParameterValue)
-			{
-				Result.Add(Texture);
-			}
+			Result.Add(TextureValue.ParameterValue);
 		}
 
 		// We'll have a dependency on our reference material too of course (this happens for Mdl
 		// materials for example, that create new UMaterial assets every time, and also material instances).
-		if (UMaterialInterface* ReferenceMaterial = MaterialInstance->Parent.Get())
-		{
-			Result.Add(ReferenceMaterial);
-		}
+		Result.Add(MaterialInstance->Parent.Get());
 	}
 	else if (USkeletalMesh* SkeletalMesh = Cast<USkeletalMesh>(Asset))
 	{
-		if (USkeleton* Skeleton = SkeletalMesh->GetSkeleton())
-		{
-			Result.Add(Skeleton);
-		}
-
-		if (UPhysicsAsset* PhysicsAsset = SkeletalMesh->GetPhysicsAsset())
-		{
-			Result.Add(PhysicsAsset);
-		}
+		Result.Add(SkeletalMesh->GetSkeleton());
+		Result.Add(SkeletalMesh->GetPhysicsAsset());
 
 		Result.Reserve(Result.Num() + SkeletalMesh->GetMaterials().Num());
 		for (const FSkeletalMaterial& SkeletalMaterial : SkeletalMesh->GetMaterials())
 		{
-			if (UMaterialInterface* UsedMaterial = SkeletalMaterial.MaterialInterface)
-			{
-				Result.Add(UsedMaterial);
-			}
+			Result.Add(SkeletalMaterial.MaterialInterface);
 		}
 	}
 	else if (UStaticMesh* StaticMesh = Cast<UStaticMesh>(Asset))
@@ -352,33 +339,36 @@ TSet<UObject*> IUsdClassesModule::GetAssetDependencies(UObject* Asset)
 		Result.Reserve(Result.Num() + StaticMesh->GetStaticMaterials().Num());
 		for (const FStaticMaterial& StaticMaterial : StaticMesh->GetStaticMaterials())
 		{
-			if (UMaterialInterface* UsedMaterial = StaticMaterial.MaterialInterface)
-			{
-				Result.Add(UsedMaterial);
-			}
+			Result.Add(StaticMaterial.MaterialInterface);
 		}
 	}
 	else if (UGeometryCache* GeometryCache = Cast<UGeometryCache>(Asset))
 	{
 		for (UMaterialInterface* UsedMaterial : GeometryCache->Materials)
 		{
-			if (UsedMaterial)
-			{
-				Result.Add(UsedMaterial);
-			}
+			Result.Add(UsedMaterial);
 		}
 	}
 	else if (UAnimSequence* AnimSequence = Cast<UAnimSequence>(Asset))
 	{
-		if (USkeletalMesh* Mesh = AnimSequence->GetPreviewMesh())
-		{
-			Result.Add(Mesh);
-		}
-
-		if (USkeleton* Skeleton = AnimSequence->GetSkeleton())
-		{
-			Result.Add(Skeleton);
-		}
+		Result.Add(AnimSequence->GetPreviewMesh());
+		Result.Add(AnimSequence->GetSkeleton());
+	}
+	else if (UGroomBindingAsset* GroomBinding = Cast<UGroomBindingAsset>(Asset))
+	{
+		Result.Add(GroomBinding->Groom);
+		Result.Add(GroomBinding->TargetSkeletalMesh);
+		Result.Add(GroomBinding->SourceSkeletalMesh);
+		Result.Add(GroomBinding->SourceGeometryCache);
+		Result.Add(GroomBinding->TargetGeometryCache);
+	}
+	else if (UGroomAsset* GroomAsset = Cast<UGroomAsset>(Asset))
+	{
+		// Do nothing. The atual groom assets have no additional dependencies
+	}
+	else if (UGroomCache* GroomCache = Cast<UGroomCache>(Asset))
+	{
+		// Do nothing. The groom cache doesn't have any additional dependencies
 	}
 	else if (UTexture* Texture = Cast<UTexture>(Asset))
 	{
@@ -401,6 +391,8 @@ TSet<UObject*> IUsdClassesModule::GetAssetDependencies(UObject* Asset)
 		UE_LOG(LogUsd, Warning, TEXT("Unknown asset '%s' encountered when collecting dependencies."), Asset ? *Asset->GetName() : TEXT("nullptr"));
 	}
 
+	// This way we don't have to nullptr check everything we add to the set
+	Result.Remove(nullptr);
 	return Result;
 }
 
diff --git a/Engine/Plugins/Importers/USDImporter/Source/USDClasses/USDClasses.Build.cs b/Engine/Plugins/Importers/USDImporter/Source/USDClasses/USDClasses.Build.cs
--- a/Engine/Plugins/Importers/USDImporter/Source/USDClasses/USDClasses.Build.cs
+++ b/Engine/Plugins/Importers/USDImporter/Source/USDClasses/USDClasses.Build.cs
@@ -26,6 +26,7 @@ namespace UnrealBuildTool.Rules
 				{
 					"DeveloperSettings",
 					"GeometryCache",
+					"HairStrandsCore",
 					"Json", // To read/write plugInfo.json files from UnrealUSDWrapper.cpp
 					"MeshDescription",
 					"RHI",
diff --git a/Engine/Plugins/Importers/USDImporter/Source/USDSchemas/Private/USDGroomTranslator.cpp b/Engine/Plugins/Importers/USDImporter/Source/USDSchemas/Private/USDGroomTranslator.cpp
--- a/Engine/Plugins/Importers/USDImporter/Source/USDSchemas/Private/USDGroomTranslator.cpp
+++ b/Engine/Plugins/Importers/USDImporter/Source/USDSchemas/Private/USDGroomTranslator.cpp
@@ -183,7 +183,13 @@ protected:
 						*IUsdClassesModule::SanitizeObjectName(FPaths::GetBaseFilename(PrimPathString))
 					);
 
-					FHairImportContext HairImportContext(ImportOptions.Get(), GetTransientPackage(), UGroomAsset::StaticClass(), AssetName, Context->ObjectFlags | EObjectFlags::RF_Public);
+					FHairImportContext HairImportContext(
+						ImportOptions.Get(),
+						GetTransientPackage(),
+						UGroomAsset::StaticClass(),
+						AssetName,
+						Context->ObjectFlags | RF_Public | RF_Transient
+					);
 					UGroomAsset* ExistingAsset = nullptr;
 					GroomAsset = FHairStrandsImporter::ImportHair(HairImportContext, HairDescription, ExistingAsset);
 					if (GroomAsset)
@@ -323,7 +329,7 @@ protected:
 			[this]() -> bool
 			{
 				const FString StrandsGroomCachePrimPath = UsdGroomTranslatorUtils::GetStrandsGroomCachePrimPath(PrimPath);
-				FHairImportContext HairImportContext(nullptr, GetTransientPackage(), nullptr, FName(), Context->ObjectFlags | EObjectFlags::RF_Public);
+				FHairImportContext HairImportContext(nullptr, GetTransientPackage(), nullptr, FName(), Context->ObjectFlags | RF_Public | RF_Transient);
 				FName UniqueName = MakeUniqueObjectName(
 					GetTransientPackage(),
 					UGroomCache::StaticClass(),
diff --git a/Engine/Plugins/Importers/USDImporter/Source/USDSchemas/Private/USDGroomTranslatorUtils.cpp b/Engine/Plugins/Importers/USDImporter/Source/USDSchemas/Private/USDGroomTranslatorUtils.cpp
--- a/Engine/Plugins/Importers/USDImporter/Source/USDSchemas/Private/USDGroomTranslatorUtils.cpp
+++ b/Engine/Plugins/Importers/USDImporter/Source/USDSchemas/Private/USDGroomTranslatorUtils.cpp
@@ -104,7 +104,11 @@ namespace UE::UsdGroomTranslatorUtils::Private
 			UGroomBindingAsset::StaticClass(),
 			*IUsdClassesModule::SanitizeObjectName(FPaths::GetBaseFilename(GroomBindingPath))
 		);
-		UGroomBindingAsset* GroomBinding = NewObject<UGroomBindingAsset>(GetTransientPackage(), BindingAssetName, ObjectFlags | RF_Public);
+		UGroomBindingAsset* GroomBinding = NewObject<UGroomBindingAsset>(
+			GetTransientPackage(),
+			BindingAssetName,
+			ObjectFlags | RF_Public | RF_Transient
+		);
 		if (GroomBinding)
 		{
 			GroomBinding->GroomBindingType = Settings.GroomBindingType;
diff --git a/Engine/Plugins/Importers/USDImporter/Source/USDStageImporter/Private/USDStageImporter.cpp b/Engine/Plugins/Importers/USDImporter/Source/USDStageImporter/Private/USDStageImporter.cpp
--- a/Engine/Plugins/Importers/USDImporter/Source/USDStageImporter/Private/USDStageImporter.cpp
+++ b/Engine/Plugins/Importers/USDImporter/Source/USDStageImporter/Private/USDStageImporter.cpp
@@ -39,6 +39,7 @@
 #include "EngineUtils.h"
 #include "GeometryCache.h"
 #include "GroomAsset.h"
+#include "GroomBindingAsset.h"
 #include "GroomCache.h"
 #include "HAL/FileManager.h"
 #include "IAssetTools.h"
@@ -821,6 +822,9 @@ namespace UsdStageImporterImpl
 		TArray<UObject*> AnimBlueprints;
 		TArray<UObject*> LevelSequences;
 		TArray<UObject*> GeometryCaches;
+		TArray<UObject*> Grooms;
+		TArray<UObject*> GroomCaches;
+		TArray<UObject*> GroomBindings;
 
 		TSet<FString> UniqueAssetNames;
 		TMap<UObject*, FString> AssetToContentFolder;
@@ -838,6 +842,7 @@ namespace UsdStageImporterImpl
 			const static FString SkeletalMeshesFolder = TEXT("SkeletalMeshes");
 			const static FString LevelSequencesFolder = TEXT("LevelSequences");
 			const static FString GeometryCachesFolder = TEXT("GeometryCaches");
+			const static FString GroomsFolder         = TEXT("Grooms");
 
 			const FString* AssetTypeFolderPtr = nullptr;
 			if (Asset->IsA(UMaterialInterface::StaticClass()))
@@ -890,6 +895,28 @@ namespace UsdStageImporterImpl
 				AssetTypeFolderPtr = &GeometryCachesFolder;
 				GeometryCaches.Add(Asset);
 			}
+			else if (Asset->IsA(UGroomAsset::StaticClass()))
+			{
+				AssetTypeFolderPtr = &GroomsFolder;
+				Grooms.Add(Asset);
+			}
+			else if (Asset->IsA(UGroomCache::StaticClass()))
+			{
+				AssetTypeFolderPtr = &GroomsFolder;
+				GroomCaches.Add(Asset);
+			}
+			else if (UGroomBindingAsset* GroomBinding = Cast<UGroomBindingAsset>(Asset))
+			{
+				if(GroomBinding->TargetSkeletalMesh)
+				{
+					AssetTypeFolderPtr = &SkeletalMeshesFolder;
+				}
+				else if (GroomBinding->TargetGeometryCache)
+				{
+					AssetTypeFolderPtr = &GeometryCachesFolder;
+				}
+				GroomBindings.Add(Asset);
+			}
 			else
 			{
 				// We don't know what to do with this asset
@@ -948,6 +975,9 @@ namespace UsdStageImporterImpl
 		PublishAssetType(LevelSequences);
 		PublishAssetType(AnimBlueprints);
 		PublishAssetType(AnimSequences);
+		PublishAssetType(GroomBindings);
+		PublishAssetType(GroomCaches);
+		PublishAssetType(Grooms);
 		PublishAssetType(SkeletalMeshes);
 		PublishAssetType(Skeletons);
 		PublishAssetType(PhysicsAssets);
@@ -1520,6 +1550,7 @@ namespace UsdStageImporterImpl
 			int32 NumTextures = 0;
 			int32 NumGeometryCaches = 0;
 			int32 NumGroomAssets = 0;
+			int32 NumGroomBindings = 0;
 			int32 NumGroomCaches = 0;
 			for ( UObject* ImportedAsset : ImportedAssets )
 			{
@@ -1560,6 +1591,10 @@ namespace UsdStageImporterImpl
 				{
 					++NumGroomAssets;
 				}
+				else if ( ImportedAsset->IsA<UGroomBindingAsset>() )
+				{
+					++NumGroomBindings;
+				}
 				else if ( ImportedAsset->IsA<UGroomCache>() )
 				{
 					++NumGroomCaches;
@@ -1573,6 +1608,7 @@ namespace UsdStageImporterImpl
 			EventAttributes.Emplace( TEXT( "NumTextures" ), LexToString( NumTextures ) );
 			EventAttributes.Emplace( TEXT( "NumGeometryCaches" ), LexToString( NumGeometryCaches ) );
 			EventAttributes.Emplace( TEXT( "NumGroomAssets" ), LexToString( NumGroomAssets ) );
+			EventAttributes.Emplace( TEXT( "NumGroomBindings" ), LexToString( NumGroomBindings ) );
 			EventAttributes.Emplace( TEXT( "NumGroomCaches" ), LexToString( NumGroomCaches ) );
 
 			FString RootLayerIdentifier = ImportContext.FilePath;
