USD: Add import option to enable/disable the import of groom assets.

diff --git a/Engine/Plugins/Importers/USDImporter/Source/USDSchemas/Private/USDGeometryCacheTranslator.cpp b/Engine/Plugins/Importers/USDImporter/Source/USDSchemas/Private/USDGeometryCacheTranslator.cpp
--- a/Engine/Plugins/Importers/USDImporter/Source/USDSchemas/Private/USDGeometryCacheTranslator.cpp
+++ b/Engine/Plugins/Importers/USDImporter/Source/USDSchemas/Private/USDGeometryCacheTranslator.cpp
@@ -901,7 +901,7 @@ USceneComponent* FUsdGeometryCacheTranslator::CreateComponents()
 				);
 
 				// Check if the prim has the GroomBinding schema and setup the component and assets necessary to bind the groom to the GeometryCache
-				if (UsdUtils::PrimHasSchema(GetPrim(), UnrealIdentifiers::GroomBindingAPI))
+				if (UsdUtils::PrimHasSchema(GetPrim(), UnrealIdentifiers::GroomBindingAPI) && Context->bAllowParsingGroomAssets)
 				{
 					UsdGroomTranslatorUtils::CreateGroomBindingAsset(
 						GetPrim(),
diff --git a/Engine/Plugins/Importers/USDImporter/Source/USDSchemas/Private/USDGroomTranslator.cpp b/Engine/Plugins/Importers/USDImporter/Source/USDSchemas/Private/USDGroomTranslator.cpp
--- a/Engine/Plugins/Importers/USDImporter/Source/USDSchemas/Private/USDGroomTranslator.cpp
+++ b/Engine/Plugins/Importers/USDImporter/Source/USDSchemas/Private/USDGroomTranslator.cpp
@@ -362,7 +362,7 @@ bool FUsdGroomTranslator::IsGroomPrim() const
 
 void FUsdGroomTranslator::CreateAssets()
 {
-	if (!IsGroomPrim())
+	if (!Context->bAllowParsingGroomAssets || !IsGroomPrim())
 	{
 		return Super::CreateAssets();
 	}
@@ -374,7 +374,7 @@ void FUsdGroomTranslator::CreateAssets()
 
 USceneComponent* FUsdGroomTranslator::CreateComponents()
 {
-	if (!IsGroomPrim())
+	if (!Context->bAllowParsingGroomAssets || !IsGroomPrim())
 	{
 		return Super::CreateComponents();
 	}
@@ -395,7 +395,7 @@ USceneComponent* FUsdGroomTranslator::CreateComponents()
 
 void FUsdGroomTranslator::UpdateComponents(USceneComponent* SceneComponent)
 {
-	if (!IsGroomPrim())
+	if (!Context->bAllowParsingGroomAssets || !IsGroomPrim())
 	{
 		Super::UpdateComponents(SceneComponent);
 	}
@@ -448,7 +448,7 @@ void FUsdGroomTranslator::UpdateComponents(USceneComponent* SceneComponent)
 
 bool FUsdGroomTranslator::CollapsesChildren(ECollapsingType CollapsingType) const
 {
-	if (!IsGroomPrim())
+	if (!Context->bAllowParsingGroomAssets || !IsGroomPrim())
 	{
 		return Super::CollapsesChildren(CollapsingType);
 	}
@@ -458,7 +458,7 @@ bool FUsdGroomTranslator::CollapsesChildren(ECollapsingType CollapsingType) cons
 
 bool FUsdGroomTranslator::CanBeCollapsed(ECollapsingType CollapsingType) const
 {
-	if (!IsGroomPrim())
+	if (!Context->bAllowParsingGroomAssets || !IsGroomPrim())
 	{
 		return Super::CanBeCollapsed(CollapsingType);
 	}
@@ -468,7 +468,7 @@ bool FUsdGroomTranslator::CanBeCollapsed(ECollapsingType CollapsingType) const
 
 TSet<UE::FSdfPath> FUsdGroomTranslator::CollectAuxiliaryPrims() const
 {
-	if (!IsGroomPrim())
+	if (!Context->bAllowParsingGroomAssets || !IsGroomPrim())
 	{
 		return Super::CollectAuxiliaryPrims();
 	}
diff --git a/Engine/Plugins/Importers/USDImporter/Source/USDSchemas/Private/USDSkelRootTranslator.cpp b/Engine/Plugins/Importers/USDImporter/Source/USDSchemas/Private/USDSkelRootTranslator.cpp
--- a/Engine/Plugins/Importers/USDImporter/Source/USDSchemas/Private/USDSkelRootTranslator.cpp
+++ b/Engine/Plugins/Importers/USDImporter/Source/USDSchemas/Private/USDSkelRootTranslator.cpp
@@ -1653,7 +1653,7 @@ USceneComponent* FUsdSkelRootTranslator::CreateComponents()
 
 #if WITH_EDITOR
 	// Check if the prim has the GroomBinding schema and setup the component and assets necessary to bind the groom to the SkeletalMesh
-	if ( UsdUtils::PrimHasSchema( GetPrim(), UnrealIdentifiers::GroomBindingAPI ) && Context->AssetCache && Context->InfoCache )
+	if (UsdUtils::PrimHasSchema(GetPrim(), UnrealIdentifiers::GroomBindingAPI) && Context->AssetCache && Context->InfoCache && Context->bAllowParsingGroomAssets)
 	{
 		UsdGroomTranslatorUtils::CreateGroomBindingAsset(
 			GetPrim(),
@@ -1688,7 +1688,6 @@ void FUsdSkelRootTranslator::UpdateComponents( USceneComponent* SceneComponent )
 	UE::FUsdPrim Prim = GetPrim();
 
 	const bool bPrimHasLiveLinkSchema = UsdUtils::PrimHasSchema( Prim, UnrealIdentifiers::LiveLinkAPI );
-	const bool bPrimHasControlRigSchema = UsdUtils::PrimHasSchema( Prim, UnrealIdentifiers::ControlRigAPI );
 
 	bool bPrimHasLiveLinkEnabled = bPrimHasLiveLinkSchema;
 	if ( bPrimHasLiveLinkSchema )
diff --git a/Engine/Plugins/Importers/USDImporter/Source/USDSchemas/Public/USDSchemaTranslator.h b/Engine/Plugins/Importers/USDImporter/Source/USDSchemas/Public/USDSchemaTranslator.h
--- a/Engine/Plugins/Importers/USDImporter/Source/USDSchemas/Public/USDSchemaTranslator.h
+++ b/Engine/Plugins/Importers/USDImporter/Source/USDSchemas/Public/USDSchemaTranslator.h
@@ -224,6 +224,9 @@ struct USDSCHEMAS_API FUsdSchemaTranslationContext : public TSharedFromThis< FUs
 	/** If true, we will also try creating UAnimSequence skeletal animation assets when parsing SkelRoot prims */
 	bool bAllowParsingSkeletalAnimations = true;
 
+	/** If true, means we will try generating GroomAssets, GroomCaches and GroomBindings */
+	bool bAllowParsingGroomAssets = true;
+
 	/** Skip the import of materials that aren't being used by any prim on the stage */
 	bool bTranslateOnlyUsedMaterials = false;
 
diff --git a/Engine/Plugins/Importers/USDImporter/Source/USDStageImporter/Private/USDStageImportOptions.cpp b/Engine/Plugins/Importers/USDImporter/Source/USDStageImporter/Private/USDStageImportOptions.cpp
--- a/Engine/Plugins/Importers/USDImporter/Source/USDStageImporter/Private/USDStageImportOptions.cpp
+++ b/Engine/Plugins/Importers/USDImporter/Source/USDStageImporter/Private/USDStageImportOptions.cpp
@@ -18,6 +18,7 @@ UUsdStageImportOptions::UUsdStageImportOptions(const FObjectInitializer& ObjectI
 	bImportSkeletalAnimations = true;
 	bImportLevelSequences = true;
 	bImportMaterials = true;
+	bImportGroomAssets = true;
 	bImportOnlyUsedMaterials = false;
 
 	PurposesToImport = (int32) (EUsdPurpose::Default | EUsdPurpose::Proxy | EUsdPurpose::Render | EUsdPurpose::Guide);
@@ -83,6 +84,7 @@ void UsdUtils::AddAnalyticsAttributes(
 	InOutAttributes.Emplace( TEXT( "ImportSkeletalAnimations" ), LexToString( Options.bImportSkeletalAnimations ) );
 	InOutAttributes.Emplace( TEXT( "ImportLevelSequences" ), LexToString( Options.bImportLevelSequences ) );
 	InOutAttributes.Emplace( TEXT( "ImportMaterials" ), LexToString( Options.bImportMaterials ) );
+	InOutAttributes.Emplace( TEXT( "ImportGroomAssets" ), LexToString( Options.bImportGroomAssets ) );
 	InOutAttributes.Emplace( TEXT( "ImportOnlyUsedMaterials" ), LexToString( Options.bImportOnlyUsedMaterials) );
 	if ( Options.PrimsToImport != TArray<FString>{ TEXT( "/" ) } )
 	{
diff --git a/Engine/Plugins/Importers/USDImporter/Source/USDStageImporter/Private/USDStageImporter.cpp b/Engine/Plugins/Importers/USDImporter/Source/USDStageImporter/Private/USDStageImporter.cpp
--- a/Engine/Plugins/Importers/USDImporter/Source/USDStageImporter/Private/USDStageImporter.cpp
+++ b/Engine/Plugins/Importers/USDImporter/Source/USDStageImporter/Private/USDStageImporter.cpp
@@ -1652,16 +1652,17 @@ namespace UsdStageImporterImpl
 			}
 
 			if (
-				( !ImportContext.ImportOptions->bImportGeometry && (
+				(!ImportContext.ImportOptions->bImportGeometry && (
 					Asset->IsA<UStaticMesh>() ||
 					Asset->IsA<USkeletalMesh>() ||
 					Asset->IsA<USkeleton>() ||
 					Asset->IsA<UPhysicsAsset>() ||
 					Asset->IsA<UGeometryCache>()
-				) ) ||
-				( !bImportSkeletalAnimations && ( Asset->IsA<UAnimSequence>() ) ) ||
-				( !ImportContext.ImportOptions->bImportLevelSequences && ( Asset->IsA<ULevelSequence>() ) ) ||
-				( !ImportContext.ImportOptions->bImportMaterials && ( Asset->IsA<UMaterialInterface>() || Asset->IsA<UTexture>() ) )
+				)) ||
+				(!bImportSkeletalAnimations && (Asset->IsA<UAnimSequence>())) ||
+				(!ImportContext.ImportOptions->bImportLevelSequences && (Asset->IsA<ULevelSequence>())) ||
+				(!ImportContext.ImportOptions->bImportMaterials && (Asset->IsA<UMaterialInterface>() || Asset->IsA<UTexture>())) ||
+				(!ImportContext.ImportOptions->bImportGroomAssets && (Asset->IsA<UGroomAsset>() || Asset->IsA<UGroomCache>() || Asset->IsA<UGroomBindingAsset>()))
 			)
 			{
 				ObjectsToRemap.Add( Asset, nullptr );
@@ -1829,6 +1830,7 @@ void UUsdStageImporter::ImportFromFile(FUsdStageImportContext& ImportContext)
 	TranslationContext->bMergeIdenticalMaterialSlots = ImportContext.ImportOptions->bMergeIdenticalMaterialSlots;
 	TranslationContext->bAllowInterpretingLODs = ImportContext.ImportOptions->bInterpretLODs;
 	TranslationContext->bAllowParsingSkeletalAnimations = ImportContext.ImportOptions->bImportGeometry && ImportContext.ImportOptions->bImportSkeletalAnimations;
+	TranslationContext->bAllowParsingGroomAssets = ImportContext.ImportOptions->bImportGroomAssets;
 	TranslationContext->bTranslateOnlyUsedMaterials = ImportContext.ImportOptions->bImportOnlyUsedMaterials;
 	TranslationContext->InfoCache = ImportContext.InfoCache;
 	TranslationContext->BlendShapesByPath = &BlendShapesByPath;
@@ -1951,6 +1953,7 @@ bool UUsdStageImporter::ReimportSingleAsset(
 	TranslationContext->bMergeIdenticalMaterialSlots = ImportContext.ImportOptions->bMergeIdenticalMaterialSlots;
 	TranslationContext->bAllowInterpretingLODs = ImportContext.ImportOptions->bInterpretLODs;
 	TranslationContext->bAllowParsingSkeletalAnimations = ImportContext.ImportOptions->bImportGeometry && ImportContext.ImportOptions->bImportSkeletalAnimations;
+	TranslationContext->bAllowParsingGroomAssets = ImportContext.ImportOptions->bImportGroomAssets;
 	TranslationContext->bTranslateOnlyUsedMaterials = ImportContext.ImportOptions->bImportOnlyUsedMaterials;
 	TranslationContext->InfoCache = ImportContext.InfoCache;
 	TranslationContext->BlendShapesByPath = &BlendShapesByPath;
diff --git a/Engine/Plugins/Importers/USDImporter/Source/USDStageImporter/Public/USDStageImportOptions.h b/Engine/Plugins/Importers/USDImporter/Source/USDStageImporter/Public/USDStageImportOptions.h
--- a/Engine/Plugins/Importers/USDImporter/Source/USDStageImporter/Public/USDStageImportOptions.h
+++ b/Engine/Plugins/Importers/USDImporter/Source/USDStageImporter/Public/USDStageImportOptions.h
@@ -66,6 +66,10 @@ public:
 	UPROPERTY(BlueprintReadWrite, config, EditAnywhere, Category = "DataToImport", meta = (DisplayName = "Materials & Textures"))
 	bool bImportMaterials;
 
+	/** Whether to import GroomAssets, GroomCaches and GroomBindings */
+	UPROPERTY(BlueprintReadWrite, config, EditAnywhere, Category = "DataToImport", meta = (DisplayName = "Groom Assets"))
+	bool bImportGroomAssets;
+
 	/**
 	 * If this is checked, only materials actively used by the stage and import settings will be parsed.
 	 * If this is unchecked, all materials present on the stage will be parsed.
