USD: Fix automated test failing due to USD Stage Editor not being opened on horde;

Revert previous diagnostic changes;
Add IsStageEditorOpened() to the Stage Editor blueprint API (plus tests for it).

diff --git a/Engine/Plugins/Importers/USDImporter/Source/USDStageEditor/Private/USDStageEditorBlueprintLibrary.cpp b/Engine/Plugins/Importers/USDImporter/Source/USDStageEditor/Private/USDStageEditorBlueprintLibrary.cpp
--- a/Engine/Plugins/Importers/USDImporter/Source/USDStageEditor/Private/USDStageEditorBlueprintLibrary.cpp
+++ b/Engine/Plugins/Importers/USDImporter/Source/USDStageEditor/Private/USDStageEditorBlueprintLibrary.cpp
@@ -22,6 +22,12 @@ bool UUsdStageEditorBlueprintLibrary::CloseStageEditor()
 	return StageEditorModule.CloseStageEditor();
 }
 
+bool UUsdStageEditorBlueprintLibrary::IsStageEditorOpened()
+{
+	IUsdStageEditorModule& StageEditorModule = FModuleManager::GetModuleChecked< IUsdStageEditorModule >( "USDStageEditor" );
+	return StageEditorModule.IsStageEditorOpened();
+}
+
 AUsdStageActor* UUsdStageEditorBlueprintLibrary::GetAttachedStageActor()
 {
 	IUsdStageEditorModule& StageEditorModule = FModuleManager::GetModuleChecked< IUsdStageEditorModule >( "USDStageEditor" );
diff --git a/Engine/Plugins/Importers/USDImporter/Source/USDStageEditor/Private/USDStageEditorModule.cpp b/Engine/Plugins/Importers/USDImporter/Source/USDStageEditor/Private/USDStageEditorModule.cpp
--- a/Engine/Plugins/Importers/USDImporter/Source/USDStageEditor/Private/USDStageEditorModule.cpp
+++ b/Engine/Plugins/Importers/USDImporter/Source/USDStageEditor/Private/USDStageEditorModule.cpp
@@ -227,7 +227,7 @@ namespace UE::UsdStageEditorModule::Private
 	}
 
 #if USE_USD_SDK
-	TSharedPtr<SUsdStage> GetUsdStageEditor( bool bOpenIfNeeded = false )
+	TSharedPtr<SUsdStage> GetUsdStageEditor( bool bOpenIfNeeded = true )
 	{
 		FLevelEditorModule& LevelEditorModule = FModuleManager::GetModuleChecked< FLevelEditorModule >( "LevelEditor" );
 		TSharedPtr< FTabManager > LevelEditorTabManager = LevelEditorModule.GetLevelEditorTabManager();
@@ -426,8 +426,7 @@ private:
 bool IUsdStageEditorModule::OpenStageEditor() const
 {
 #if USE_USD_SDK
-	const bool bOpenIfNeeded = true;
-	return UE::UsdStageEditorModule::Private::GetUsdStageEditor( bOpenIfNeeded ).IsValid();
+	return UE::UsdStageEditorModule::Private::GetUsdStageEditor().IsValid();
 #else
 	return false;
 #endif // USE_USD_SDK
@@ -447,6 +446,16 @@ bool IUsdStageEditorModule::CloseStageEditor() const
 	return false;
 }
 
+bool IUsdStageEditorModule::IsStageEditorOpened() const
+{
+#if USE_USD_SDK
+	const bool bOpenIfNeeded = false;
+	return UE::UsdStageEditorModule::Private::GetUsdStageEditor(bOpenIfNeeded).IsValid();
+#else
+	return false;
+#endif // USE_USD_SDK
+}
+
 AUsdStageActor* IUsdStageEditorModule::GetAttachedStageActor() const
 {
 #if USE_USD_SDK
diff --git a/Engine/Plugins/Importers/USDImporter/Source/USDStageEditor/Public/USDStageEditorBlueprintLibrary.h b/Engine/Plugins/Importers/USDImporter/Source/USDStageEditor/Public/USDStageEditorBlueprintLibrary.h
--- a/Engine/Plugins/Importers/USDImporter/Source/USDStageEditor/Public/USDStageEditorBlueprintLibrary.h
+++ b/Engine/Plugins/Importers/USDImporter/Source/USDStageEditor/Public/USDStageEditorBlueprintLibrary.h
@@ -19,18 +19,25 @@ class USDSTAGEEDITOR_API UUsdStageEditorBlueprintLibrary : public UBlueprintFunc
 public:
 	/**
 	 * Opens the the USD Stage Editor window, or focus it in case it is already open.
-	 * @return True if a stage is now opened and available.
+	 * @return True if a stage editor window is now opened and available.
 	 */
 	UFUNCTION( BlueprintCallable, Category = "USD|Stage Editor" )
 	static bool OpenStageEditor();
 
 	/**
 	 * Closes the USD Stage Editor window if it is opened. Does nothing in case it is already closed.
-	 * @return True if a stage was closed by this action.
+	 * @return True if a stage editor window was closed by this action.
 	 */
 	UFUNCTION( BlueprintCallable, Category = "USD|Stage Editor" )
 	static bool CloseStageEditor();
 
+	/**
+	 * Checks to see if an USD Stage Editor window is currently opened.
+	 * @return True if a stage editor window is now opened and available.
+	 */
+	UFUNCTION( BlueprintCallable, Category = "USD|Stage Editor" )
+	static bool IsStageEditorOpened();
+
 	/**
 	 * Gets which actor is currently attached to the USD Stage Editor, if any.
 	 * May return nullptr in case no actor is attached.
diff --git a/Engine/Plugins/Importers/USDImporter/Source/USDStageEditor/Public/USDStageEditorModule.h b/Engine/Plugins/Importers/USDImporter/Source/USDStageEditor/Public/USDStageEditorModule.h
--- a/Engine/Plugins/Importers/USDImporter/Source/USDStageEditor/Public/USDStageEditorModule.h
+++ b/Engine/Plugins/Importers/USDImporter/Source/USDStageEditor/Public/USDStageEditorModule.h
@@ -25,6 +25,7 @@ public:
 
 	bool OpenStageEditor() const;
 	bool CloseStageEditor() const;
+	bool IsStageEditorOpened() const;
 
 	AUsdStageActor* GetAttachedStageActor() const;
 	bool SetAttachedStageActor( AUsdStageActor* NewActor ) const;
