USD: Fix warning when loading Skeleton assets after importing.

diff --git a/Engine/Plugins/Importers/USDImporter/Source/USDSchemas/Private/USDSkelRootTranslator.cpp b/Engine/Plugins/Importers/USDImporter/Source/USDSchemas/Private/USDSkelRootTranslator.cpp
--- a/Engine/Plugins/Importers/USDImporter/Source/USDSchemas/Private/USDSkelRootTranslator.cpp
+++ b/Engine/Plugins/Importers/USDImporter/Source/USDSchemas/Private/USDSkelRootTranslator.cpp
@@ -1391,6 +1391,14 @@ namespace UsdSkelRootTranslatorImpl
 
 						Context->AssetCache->CacheAsset( SkeletalMeshHashString, SkeletalMesh );
 						Context->AssetCache->CacheAsset( SkeletalMeshHashString + TEXT( "_Skeleton" ), SkeletalMesh->GetSkeleton() );
+
+						// The PreviewSkeletalMesh property on the Skeleton is a soft object path, that is set within
+						// UsdToUnreal::GetSkeletalMeshFromImportData before the SkeletalMesh is part of the cache. When we cache the
+						// SkeletalMesh, it is renamed into the AssetCache package, and so PreviewSkeletalMesh is no longer pointing at it.
+						// Ideally we'd create all our asset cache assets directly outered to the AssetCache to avoid this issue, but
+						// for now I think it's only a problem for SkeletalMesh/Skeleton pairs (since they are created at the same time
+						// already knowing each other), so it's simpler to just do this
+						SkeletalMesh->GetSkeleton()->SetPreviewMesh(SkeletalMesh);
 					}
 
 					if ( bGeneratePhysicsAssets )
