USD: Fix potential nullptr dereference warning from static analysis (Horde issue 309734).

diff --git a/Engine/Plugins/Importers/USDImporter/Source/USDStageImporter/Private/USDStageImporter.cpp b/Engine/Plugins/Importers/USDImporter/Source/USDStageImporter/Private/USDStageImporter.cpp
--- a/Engine/Plugins/Importers/USDImporter/Source/USDStageImporter/Private/USDStageImporter.cpp
+++ b/Engine/Plugins/Importers/USDImporter/Source/USDStageImporter/Private/USDStageImporter.cpp
@@ -943,7 +943,7 @@ namespace UsdStageImporterImpl
 					}
 				}
 			}
-			else
+			else if (AssetTypeFolderPtr)
 			{
 				AssetToContentFolder.Add(Asset, *AssetTypeFolderPtr);
 			}
