USD: Partial revert of 25157830 so that we store import file paths on AssetImportData, which is needed for reimport.

diff --git a/Engine/Plugins/Importers/USDImporter/Source/USDClasses/Public/USDAssetImportData.h b/Engine/Plugins/Importers/USDImporter/Source/USDClasses/Public/USDAssetImportData.h
--- a/Engine/Plugins/Importers/USDImporter/Source/USDClasses/Public/USDAssetImportData.h
+++ b/Engine/Plugins/Importers/USDImporter/Source/USDClasses/Public/USDAssetImportData.h
@@ -10,11 +10,12 @@
 #include "USDAssetImportData.generated.h"
 
 UCLASS()
-class UE_DEPRECATED(5.3, "Prefer USD AssetUserData instead of AssetImportData") USDCLASSES_API UUsdAssetImportData : public UAssetImportData
+class USDCLASSES_API UUsdAssetImportData : public UAssetImportData
 {
 	GENERATED_BODY()
 
 public:
+	UE_DEPRECATED(5.3, "Use USD AssetUserData instead of AssetImportData for USD-specific info")
 	UPROPERTY()
 	FString PrimPath;
 
@@ -26,7 +27,7 @@ public:
 
 PRAGMA_DISABLE_DEPRECATION_WARNINGS
 UCLASS()
-class UE_DEPRECATED(5.3, "Prefer USD AssetUserData instead of AssetImportData") USDCLASSES_API UUsdAnimSequenceAssetImportData : public UUsdAssetImportData
+class UE_DEPRECATED(5.3, "Use USD AssetUserData instead of AssetImportData for USD-specific info") USDCLASSES_API UUsdAnimSequenceAssetImportData : public UUsdAssetImportData
 {
 	GENERATED_BODY()
 
@@ -44,7 +45,7 @@ public:
 
 /** We assign these to UStaticMeshes or USkeletalMeshes generated from USD */
 UCLASS()
-class UE_DEPRECATED(5.3, "Prefer USD AssetUserData instead of AssetImportData") USDCLASSES_API UUsdMeshAssetImportData : public UUsdAssetImportData
+class UE_DEPRECATED(5.3, "Use USD AssetUserData instead of AssetImportData for USD-specific info") USDCLASSES_API UUsdMeshAssetImportData : public UUsdAssetImportData
 {
 	GENERATED_BODY()
 
diff --git a/Engine/Plugins/Importers/USDImporter/Source/USDClasses/Public/USDAssetUserData.h b/Engine/Plugins/Importers/USDImporter/Source/USDClasses/Public/USDAssetUserData.h
--- a/Engine/Plugins/Importers/USDImporter/Source/USDClasses/Public/USDAssetUserData.h
+++ b/Engine/Plugins/Importers/USDImporter/Source/USDClasses/Public/USDAssetUserData.h
@@ -12,16 +12,8 @@ class USDCLASSES_API UUsdAssetUserData : public UAssetUserData
 	GENERATED_BODY()
 
 public:
-	UPROPERTY()
-	FString FilePath;
-
 	UPROPERTY()
 	FString PrimPath;
-
-	// Likely a UUSDStageImportOptions, but we don't declare it here
-	// to prevent an unnecessary module dependency on USDStageImporter
-	UPROPERTY()
-	TObjectPtr<class UObject> ImportOptions;
 };
 
 UCLASS()
diff --git a/Engine/Plugins/Importers/USDImporter/Source/USDSchemas/Private/USDGroomTranslator.cpp b/Engine/Plugins/Importers/USDImporter/Source/USDSchemas/Private/USDGroomTranslator.cpp
--- a/Engine/Plugins/Importers/USDImporter/Source/USDSchemas/Private/USDGroomTranslator.cpp
+++ b/Engine/Plugins/Importers/USDImporter/Source/USDSchemas/Private/USDGroomTranslator.cpp
@@ -1,7 +1,6 @@
 // Copyright Epic Games, Inc. All Rights Reserved.
 
 #include "USDGroomTranslator.h"
-#include "UObject/Package.h"
 
 #if USE_USD_SDK && WITH_EDITOR
 
@@ -15,6 +14,7 @@
 #include "HairStrandsImporter.h"
 #include "Misc/ArchiveMD5.h"
 #include "Templates/UniquePtr.h"
+#include "UObject/Package.h"
 #include "UObject/StrongObjectPtr.h"
 
 #include "UnrealUSDWrapper.h"
@@ -198,7 +198,6 @@ protected:
 
 						UUsdAssetUserData* UserData = NewObject<UUsdAssetUserData>(GroomAsset, TEXT("UUSDAssetUserData"));
 						UserData->PrimPath = PrimPathString;
-						UserData->ImportOptions = ImportOptions.Get();
 
 						GroomAsset->AddAssetUserData(UserData);
 					}
@@ -342,7 +341,6 @@ protected:
 				{
 					UUsdAssetUserData* UserData = NewObject<UUsdAssetUserData>(GroomCache, TEXT("UUSDAssetUserData"));
 					UserData->PrimPath = PrimPath.GetString();
-					UserData->ImportOptions = ImportOptions.Get();
 
 					GroomCache->AddAssetUserData(UserData);
 
diff --git a/Engine/Plugins/Importers/USDImporter/Source/USDSchemas/Private/USDShadeMaterialTranslator.cpp b/Engine/Plugins/Importers/USDImporter/Source/USDSchemas/Private/USDShadeMaterialTranslator.cpp
--- a/Engine/Plugins/Importers/USDImporter/Source/USDSchemas/Private/USDShadeMaterialTranslator.cpp
+++ b/Engine/Plugins/Importers/USDImporter/Source/USDSchemas/Private/USDShadeMaterialTranslator.cpp
@@ -13,6 +13,7 @@
 
 #include "UsdWrappers/SdfPath.h"
 
+#include "EditorFramework/AssetImportData.h"
 #include "Engine/Level.h"
 #include "Engine/Texture.h"
 #include "Materials/Material.h"
@@ -488,9 +489,9 @@ void FUsdShadeMaterialTranslator::PostImportMaterial(const FString& MaterialHash
 			else if (bIsOwnedByTransientPackage)
 			{
 				FString FilePath;
-				if(UUsdAssetUserData* TextureUserData = Texture->GetAssetUserData<UUsdAssetUserData>())
+				if(UAssetImportData* TextureImportData = Texture->AssetImportData.Get())
 				{
-					FilePath = TextureUserData->FilePath;
+					FilePath = TextureImportData->GetFirstFilename();
 				}
 				else
 				{
diff --git a/Engine/Plugins/Importers/USDImporter/Source/USDStageImporter/Private/USDStageAssetImportFactory.cpp b/Engine/Plugins/Importers/USDImporter/Source/USDStageImporter/Private/USDStageAssetImportFactory.cpp
--- a/Engine/Plugins/Importers/USDImporter/Source/USDStageImporter/Private/USDStageAssetImportFactory.cpp
+++ b/Engine/Plugins/Importers/USDImporter/Source/USDStageImporter/Private/USDStageAssetImportFactory.cpp
@@ -112,23 +112,11 @@ void UUsdStageAssetImportFactory::CleanUp()
 
 bool UUsdStageAssetImportFactory::CanReimport(UObject* Obj, TArray<FString>& OutFilenames)
 {
-	if (IInterface_AssetUserData* UserDataInterface = Cast<IInterface_AssetUserData>(Obj))
-	{
-		if(UUsdAssetUserData* UserData = UserDataInterface->GetAssetUserData<UUsdAssetUserData>())
-		{
-			OutFilenames.Add(UserData->FilePath);
-			return true;
-		}
-	}
-
-	// Fall back to also checking asset import data for now, in case it's an old asset
-	PRAGMA_DISABLE_DEPRECATION_WARNINGS
-	if (UUsdAssetImportData* ImportData = UsdUtils::GetAssetImportData(Obj))
+	if (UAssetImportData* ImportData = UsdUtils::GetAssetImportData(Obj))
 	{
 		OutFilenames.Add(ImportData->GetFirstFilename());
 		return true;
 	}
-	PRAGMA_ENABLE_DEPRECATION_WARNINGS
 
 	return false;
 }
@@ -140,20 +128,10 @@ void UUsdStageAssetImportFactory::SetReimportPaths(UObject* Obj, const TArray<FS
 		return;
 	}
 
-	if (IInterface_AssetUserData* UserDataInterface = Cast<IInterface_AssetUserData>(Obj))
-	{
-		if(UUsdAssetUserData* UserData = UserDataInterface->GetAssetUserData<UUsdAssetUserData>())
-		{
-			UserData->FilePath = NewReimportPaths[0];
-		}
-	}
-
-	PRAGMA_DISABLE_DEPRECATION_WARNINGS
 	if (UUsdAssetImportData* ImportData = UsdUtils::GetAssetImportData(Obj))
 	{
 		ImportData->UpdateFilenameOnly(NewReimportPaths[0]);
 	}
-	PRAGMA_ENABLE_DEPRECATION_WARNINGS
 }
 
 EReimportResult::Type UUsdStageAssetImportFactory::Reimport(UObject* Obj)
@@ -175,22 +153,14 @@ EReimportResult::Type UUsdStageAssetImportFactory::Reimport(UObject* Obj)
 	{
 		if (UUsdAssetUserData* UserData = UserDataInterface->GetAssetUserData<UUsdAssetUserData>())
 		{
-			ReimportFilePath = UserData->FilePath;
 			OriginalPrimPath = UserData->PrimPath;
-			ReimportOptions = UserData->ImportOptions;
 		}
 	}
 
-	if (!ReimportOptions || ReimportFilePath.IsEmpty())
+	if (UUsdAssetImportData* ImportData = UsdUtils::GetAssetImportData(Obj))
 	{
-		PRAGMA_DISABLE_DEPRECATION_WARNINGS
-		if (UUsdAssetImportData* ImportData = UsdUtils::GetAssetImportData(Obj))
-		{
-			ReimportFilePath = ImportData->GetFirstFilename();
-			OriginalPrimPath = ImportData->PrimPath;
-			ReimportOptions = ImportData->ImportOptions;
-		}
-		PRAGMA_ENABLE_DEPRECATION_WARNINGS
+		ReimportFilePath = ImportData->GetFirstFilename();
+		ReimportOptions = ImportData->ImportOptions;
 	}
 
 	if (ReimportFilePath.IsEmpty() || OriginalPrimPath.IsEmpty())
diff --git a/Engine/Plugins/Importers/USDImporter/Source/USDStageImporter/Private/USDStageImporter.cpp b/Engine/Plugins/Importers/USDImporter/Source/USDStageImporter/Private/USDStageImporter.cpp
--- a/Engine/Plugins/Importers/USDImporter/Source/USDStageImporter/Private/USDStageImporter.cpp
+++ b/Engine/Plugins/Importers/USDImporter/Source/USDStageImporter/Private/USDStageImporter.cpp
@@ -4,6 +4,7 @@
 
 #include "USDAssetCache.h"
 #include "USDAssetCache2.h"
+#include "USDAssetImportData.h"
 #include "USDAssetUserData.h"
 #include "USDClassesModule.h"
 #include "USDConversionUtils.h"
@@ -566,30 +567,30 @@ namespace UsdStageImporterImpl
 		return AssetName;
 	}
 
-	void UpdateAssetUserData(UObject* Asset, const FString& MainFilePath, UUsdStageImportOptions* ImportOptions)
+	void UpdateAssetImportData(UObject* Asset, const FString& MainFilePath, UUsdStageImportOptions* ImportOptions)
 	{
 		if (!Asset)
 		{
 			return;
 		}
 
-		UUsdAssetUserData* UserData = nullptr;
-		if (IInterface_AssetUserData* UserDataInterface = Cast<IInterface_AssetUserData>(Asset))
-		{
-			UserData = UserDataInterface->GetAssetUserData<UUsdAssetUserData>();
-		}
-		if (!UserData)
+		UUsdAssetImportData* ImportData = UsdUtils::GetAssetImportData(Asset);
+		if (!ImportData)
 		{
-			return;
+			ImportData = NewObject<UUsdAssetImportData>(Asset);
+			UsdUtils::SetAssetImportData(Asset, ImportData);
 		}
 
-		// Don't force update as textures will already come with this preset to their actual texture path
-		if (!Asset->IsA<UTexture>())
+		if (ImportData)
 		{
-			UserData->FilePath = MainFilePath;
-		}
+			// Don't force update textures as they will already have this preset to their actual texture path
+			if (!Asset->IsA<UTexture>())
+			{
+				ImportData->UpdateFilenameOnly(MainFilePath);
+			}
 
-		UserData->ImportOptions = ImportOptions;
+			ImportData->ImportOptions = ImportOptions;
+		}
 	}
 
 	void UpdateAssetUserData(
@@ -600,7 +601,7 @@ namespace UsdStageImporterImpl
 	{
 		for (UObject* Asset : UsedAssetsAndDependencies)
 		{
-			UpdateAssetUserData(Asset, MainFilePath, ImportOptions);
+			UpdateAssetImportData(Asset, MainFilePath, ImportOptions);
 		}
 	}
 
@@ -1991,7 +1992,7 @@ bool UUsdStageImporter::ReimportSingleAsset(
 
 	if ( ReimportedObject )
 	{
-		UsdStageImporterImpl::UpdateAssetUserData(ReimportedObject, ImportContext.FilePath, ImportContext.ImportOptions);
+		UsdStageImporterImpl::UpdateAssetImportData(ReimportedObject, ImportContext.FilePath, ImportContext.ImportOptions);
 
 		// Assign things from the original assets before we publish the reimported asset, overwriting it
 		UsdStageImporterImpl::CopyOriginalMaterialAssignment(ImportContext, OriginalAsset, ReimportedObject );
diff --git a/Engine/Plugins/Importers/USDImporter/Source/USDStageImporter/Public/USDStageImporter.h b/Engine/Plugins/Importers/USDImporter/Source/USDStageImporter/Public/USDStageImporter.h
--- a/Engine/Plugins/Importers/USDImporter/Source/USDStageImporter/Public/USDStageImporter.h
+++ b/Engine/Plugins/Importers/USDImporter/Source/USDStageImporter/Public/USDStageImporter.h
@@ -6,7 +6,6 @@
 
 #include "CoreMinimal.h"
 
-class UUsdAssetImportData;
 struct FUsdStageImportContext;
 
 class USDSTAGEIMPORTER_API UUsdStageImporter
diff --git a/Engine/Plugins/Importers/USDImporter/Source/USDUtilities/Private/USDConversionUtils.cpp b/Engine/Plugins/Importers/USDImporter/Source/USDUtilities/Private/USDConversionUtils.cpp
--- a/Engine/Plugins/Importers/USDImporter/Source/USDUtilities/Private/USDConversionUtils.cpp
+++ b/Engine/Plugins/Importers/USDImporter/Source/USDUtilities/Private/USDConversionUtils.cpp
@@ -1243,7 +1243,6 @@ double UsdUtils::GetEarliestTimeCode()
 #endif
 }
 
-PRAGMA_DISABLE_DEPRECATION_WARNINGS
 UUsdAssetImportData* UsdUtils::GetAssetImportData( UObject* Asset )
 {
 	UUsdAssetImportData* ImportData = nullptr;
@@ -1310,7 +1309,49 @@ UUsdAssetImportData* UsdUtils::GetAssetImportData( UObject* Asset )
 #endif
 	return ImportData;
 }
-PRAGMA_ENABLE_DEPRECATION_WARNINGS
+
+void UsdUtils::SetAssetImportData(UObject* Asset, UAssetImportData* ImportData)
+{
+	if (!Asset)
+	{
+		return;
+	}
+
+#if WITH_EDITOR
+	if (UStaticMesh* Mesh = Cast<UStaticMesh>(Asset))
+	{
+		Mesh->AssetImportData = ImportData;
+	}
+	else if (USkeletalMesh* SkMesh = Cast<USkeletalMesh>(Asset))
+	{
+		SkMesh->SetAssetImportData(ImportData);
+	}
+	else if (UAnimSequence* SkelAnim = Cast<UAnimSequence>(Asset))
+	{
+		SkelAnim->AssetImportData = ImportData;
+	}
+	else if (UMaterialInterface* Material = Cast<UMaterialInterface>(Asset))
+	{
+		Material->AssetImportData = ImportData;
+	}
+	else if (UTexture* Texture = Cast<UTexture>(Asset))
+	{
+		Texture->AssetImportData = ImportData;
+	}
+	else if (UGeometryCache* GeometryCache = Cast<UGeometryCache>(Asset))
+	{
+		GeometryCache->AssetImportData = ImportData;
+	}
+	else if (UGroomAsset* Groom = Cast<UGroomAsset>(Asset))
+	{
+		Groom->AssetImportData = ImportData;
+	}
+	else if (UGroomCache* GroomCache = Cast<UGroomCache>(Asset))
+	{
+		GroomCache->AssetImportData = ImportData;
+	}
+#endif // WITH_EDITOR
+}
 
 namespace UE::UsdConversionUtils::Private
 {
diff --git a/Engine/Plugins/Importers/USDImporter/Source/USDUtilities/Private/USDShadeConversion.cpp b/Engine/Plugins/Importers/USDImporter/Source/USDUtilities/Private/USDShadeConversion.cpp
--- a/Engine/Plugins/Importers/USDImporter/Source/USDUtilities/Private/USDShadeConversion.cpp
+++ b/Engine/Plugins/Importers/USDImporter/Source/USDUtilities/Private/USDShadeConversion.cpp
@@ -7,6 +7,7 @@
 #include "UnrealUSDWrapper.h"
 #include "USDAssetCache.h"
 #include "USDAssetCache2.h"
+#include "USDAssetImportData.h"
 #include "USDAssetUserData.h"
 #include "USDAttributeUtils.h"
 #include "USDClassesModule.h"
@@ -1336,8 +1337,13 @@ namespace UE
 						{
 							UUsdAssetUserData* UserData = NewObject<UUsdAssetUserData>(Texture, TEXT("USDAssetUserData"));
 							UserData->PrimPath = PrimPath;
-							UserData->FilePath = ResolvedTexturePath;
 							Texture->AddAssetUserData(UserData);
+
+							// We set this even if we're not going to import so that we can track our original texture filepath
+							// in case we later do an Actions->Import
+							UUsdAssetImportData* ImportData = NewObject<UUsdAssetImportData>(Texture);
+							ImportData->UpdateFilenameOnly(ResolvedTexturePath);
+							Texture->AssetImportData = ImportData;
 						}
 					}
 				}
diff --git a/Engine/Plugins/Importers/USDImporter/Source/USDUtilities/Public/USDConversionUtils.h b/Engine/Plugins/Importers/USDImporter/Source/USDUtilities/Public/USDConversionUtils.h
--- a/Engine/Plugins/Importers/USDImporter/Source/USDUtilities/Public/USDConversionUtils.h
+++ b/Engine/Plugins/Importers/USDImporter/Source/USDUtilities/Public/USDConversionUtils.h
@@ -45,6 +45,7 @@ PXR_NAMESPACE_CLOSE_SCOPE
 #endif // #if USE_USD_SDK
 
 class USceneComponent;
+class UAssetImportData;
 class UUsdAssetImportData;
 enum class EUsdDuplicateType : uint8;
 enum class EUsdUpAxis : uint8;
@@ -269,10 +270,14 @@ namespace UsdUtils
 	/** Returns the earliest possible timecode. Use it to always fetch the first frame of an animated attribute */
 	USDUTILITIES_API double GetEarliestTimeCode();
 
-	PRAGMA_DISABLE_DEPRECATION_WARNINGS
-	UE_DEPRECATED(5.3, "No longer used as the USDImporter plugin will now use AssetUserData over AssetImportData for better runtime compatibility")
-	USDUTILITIES_API UUsdAssetImportData* GetAssetImportData( UObject* Asset );
-	PRAGMA_ENABLE_DEPRECATION_WARNINGS
+	/**
+	 * Utilities to allow getting and setting our AssetImportData to an asset from a base UObject*.
+	 * Note that not all asset types support AssetImportData, and in some cases when retrieving it for e.g. a Skeleton,
+	 * we'll actually check it's preview mesh instead (since Skeletons don't have AssetImportData). The setter won't do
+	 * anything if you try setting asset import data on e.g. a Skeleton, on the other hand.
+	 */
+	USDUTILITIES_API UUsdAssetImportData* GetAssetImportData(UObject* Asset);
+	USDUTILITIES_API void SetAssetImportData(UObject* Asset, UAssetImportData* ImportData);
 
 	/** Adds a reference on Prim to the layer at AbsoluteFilePath */
 	USDUTILITIES_API void AddReference( UE::FUsdPrim& Prim, const TCHAR* AbsoluteFilePath, const UE::FSdfPath& TargetPrimPath = {}, double TimeCodeOffset = 0.0, double TimeCodeScale = 1.0 );
