USD: Fix crash when undoing transform key addition;

Fix ensure when adding track to actor binding whose component binding is already bound;
Fix crash when undoing adding an AUsdStageActor to a persistent LevelSequence;
Fix ensure related to component Mobility when redoing the addition of transform tracks to the AUsdStageActor's Sequence.

diff --git a/Engine/Plugins/Importers/USDImporter/Source/USDStage/Private/USDLevelSequenceHelper.cpp b/Engine/Plugins/Importers/USDImporter/Source/USDStage/Private/USDLevelSequenceHelper.cpp
--- a/Engine/Plugins/Importers/USDImporter/Source/USDStage/Private/USDLevelSequenceHelper.cpp
+++ b/Engine/Plugins/Importers/USDImporter/Source/USDStage/Private/USDLevelSequenceHelper.cpp
@@ -626,7 +626,7 @@ private:
 
 		// For now we support one binding per component type (mostly so we can fit a binding to a scene component and
 		// camera component for a Camera prim twin)
-		TMap< const UClass*, FGuid > ComponentClassToBindingGuid;
+		TMap< const UClass*, FGuid > ObjectClassToBindingGuid;
 	};
 
 	TMap< TWeakObjectPtr< const UUsdPrimTwin >, FPrimTwinBindings > PrimTwinToBindings;
@@ -2433,7 +2433,7 @@ void FUsdLevelSequenceHelperImpl::RemovePossessable(const UUsdPrimTwin& PrimTwin
 	// ones don't modify the Sequence and change properties, so we must modify them here
 	Bindings->Sequence->Modify();
 
-	for (const TPair< const UClass*, FGuid >& Pair : Bindings->ComponentClassToBindingGuid)
+	for (const TPair< const UClass*, FGuid >& Pair : Bindings->ObjectClassToBindingGuid)
 	{
 		const FGuid& ComponentPossessableGuid = Pair.Value;
 
@@ -2804,7 +2804,7 @@ FGuid FUsdLevelSequenceHelperImpl::GetOrCreateComponentBinding(const UUsdPrimTwi
 	ensure(Bindings.Sequence == nullptr || Bindings.Sequence == &Sequence);
 	Bindings.Sequence = &Sequence;
 
-	if (FGuid* ExistingGuid = Bindings.ComponentClassToBindingGuid.Find(ComponentToBind.GetClass()))
+	if (FGuid* ExistingGuid = Bindings.ObjectClassToBindingGuid.Find(ComponentToBind.GetClass()))
 	{
 		return *ExistingGuid;
 	}
@@ -2851,7 +2851,7 @@ FGuid FUsdLevelSequenceHelperImpl::GetOrCreateComponentBinding(const UUsdPrimTwi
 
 	// Bind component
 	Sequence.BindPossessableObject(ComponentBinding, ComponentToBind, ComponentContext);
-	Bindings.ComponentClassToBindingGuid.Emplace(ComponentToBind.GetClass(), ComponentBinding);
+	Bindings.ObjectClassToBindingGuid.Emplace(ComponentToBind.GetClass(), ComponentBinding);
 	return ComponentBinding;
 }
 
@@ -2967,7 +2967,7 @@ void FUsdLevelSequenceHelperImpl::HandleMovieSceneChange(UMovieScene& MovieScene
 			continue;
 		}
 
-		for (TMap< const UClass*, FGuid >::TIterator BindingIt = Bindings.ComponentClassToBindingGuid.CreateIterator();
+		for (TMap< const UClass*, FGuid >::TIterator BindingIt = Bindings.ObjectClassToBindingGuid.CreateIterator();
 			BindingIt;
 			++BindingIt)
 		{
@@ -3295,8 +3295,6 @@ void FUsdLevelSequenceHelperImpl::HandleTrackChange(const UMovieSceneTrack& Trac
 		return;
 	}
 
-	ensure( BoundSceneComponent->Mobility != EComponentMobility::Static );
-
 	UUsdPrimTwin* PrimTwin = StageActor->RootUsdTwin->Find(BoundSceneComponent);
 
 	// If we exported/created this Camera prim ourselves, we'll have a decomposed parent Xform and a child Camera prim (to mirror
@@ -3344,6 +3342,8 @@ void FUsdLevelSequenceHelperImpl::HandleTrackChange(const UMovieSceneTrack& Trac
 
 	if (PrimTwin)
 	{
+		ensure(BoundSceneComponent->Mobility != EComponentMobility::Static);
+
 		FScopedBlockNoticeListening BlockNotices(StageActor.Get());
 		UE::FUsdPrim UsdPrim = UsdStage.GetPrimAtPath(UE::FSdfPath(*PrimTwin->PrimPath));
 
@@ -3352,10 +3352,10 @@ void FUsdLevelSequenceHelperImpl::HandleTrackChange(const UMovieSceneTrack& Trac
 		Bindings.Sequence = Sequence;
 
 		// Make sure we track this binding
-		const UClass* ComponentClass = BoundSceneComponent->GetClass();
-		FGuid* FoundExistingGuid = Bindings.ComponentClassToBindingGuid.Find(ComponentClass);
+		const UClass* ComponentClass = BoundObject->GetClass();
+		FGuid* FoundExistingGuid = Bindings.ObjectClassToBindingGuid.Find(ComponentClass);
 		ensure(!FoundExistingGuid || *FoundExistingGuid == PossessableGuid);
-		Bindings.ComponentClassToBindingGuid.Emplace(ComponentClass, PossessableGuid);
+		Bindings.ObjectClassToBindingGuid.Emplace(ComponentClass, PossessableGuid);
 
 		if (bIsMuteChange)
 		{
diff --git a/Engine/Plugins/Importers/USDImporter/Source/USDStage/Private/USDStageActor.cpp b/Engine/Plugins/Importers/USDImporter/Source/USDStage/Private/USDStageActor.cpp
--- a/Engine/Plugins/Importers/USDImporter/Source/USDStage/Private/USDStageActor.cpp
+++ b/Engine/Plugins/Importers/USDImporter/Source/USDStage/Private/USDStageActor.cpp
@@ -2634,17 +2634,22 @@ void AUsdStageActor::ReloadAnimations()
 	// want to listen to it and have valid level sequences
 	if (!IsTemplate())
 	{
-		bool bLevelSequenceEditorWasOpened = false;
-		if (LevelSequence)
-		{
-			// The sequencer won't update on its own, so let's at least force it closed
 #if WITH_EDITOR
-			if (GIsEditor && GEditor)
+		TArray<TSharedPtr<ISequencer>> SequencersToReset;
+		if (IUsdStageModule* UsdStageModule = FModuleManager::Get().GetModulePtr<IUsdStageModule>(TEXT("UsdStage")))
+		{
+			for (const TWeakPtr<ISequencer>& ExistingSequencer : UsdStageModule->GetExistingSequencers())
 			{
-				bLevelSequenceEditorWasOpened = GEditor->GetEditorSubsystem<UAssetEditorSubsystem>()->CloseAllEditorsForAsset(LevelSequence) > 0;
+				if (TSharedPtr<ISequencer> PinnedSequencer = ExistingSequencer.Pin())
+				{
+					if (PinnedSequencer->GetRootMovieSceneSequence() == LevelSequence)
+					{
+						SequencersToReset.Add(PinnedSequencer);
+					}
+				}
 			}
-#endif // WITH_EDITOR
 		}
+#endif // WITH_EDITOR
 
 		// We need to guarantee we'll record our change of LevelSequence into the transaction, as Init() will create a new one
 		const bool bMarkDirty = false;
@@ -2654,9 +2659,12 @@ void AUsdStageActor::ReloadAnimations()
 		LevelSequenceHelper.BindToUsdStageActor(this);
 
 #if WITH_EDITOR
-		if (GIsEditor && GEditor && LevelSequence && bLevelSequenceEditorWasOpened)
+		for (TSharedPtr<ISequencer>& Sequencer : SequencersToReset)
 		{
-			GEditor->GetEditorSubsystem<UAssetEditorSubsystem>()->OpenEditorForAsset(LevelSequence);
+			if (LevelSequence && Sequencer->GetRootMovieSceneSequence() != LevelSequence)
+			{
+				Sequencer->ResetToNewRootSequence(*LevelSequence);
+			}
 		}
 #endif // WITH_EDITOR
 	}
@@ -2816,37 +2824,97 @@ void AUsdStageActor::HandleTransactionStateChanged(const FTransactionContext& In
 	// to show that one instead. If we undo the Reload, that new LevelSequence will be deleted and the Sequencer will be left open trying to display it,
 	// which leads to crashes. Here we try detecting for that case and close/reopen the sequencer to show the correct one.
 #if WITH_EDITOR
-	if (GIsEditor && GEditor && (InTransactionState == ETransactionStateEventType::UndoRedoStarted || InTransactionState == ETransactionStateEventType::UndoRedoFinalized))
+	if (GIsEditor && LevelSequence
+		&& (InTransactionState == ETransactionStateEventType::UndoRedoStarted || InTransactionState == ETransactionStateEventType::UndoRedoFinalized))
 	{
 		if (UTransactor* Trans = GEditor->Trans)
 		{
-			static TSet<AUsdStageActor*> ActorsThatClosedTheSequencer;
+			static TMap<AUsdStageActor*, TArray<TWeakPtr<ISequencer>>> ActorsToSequencers;
 
 			int32 CurrentTransactionIndex = Trans->FindTransactionIndex(InTransactionContext.TransactionId);
-			const FTransaction* Transaction = Trans->GetTransaction(CurrentTransactionIndex);
-
-			if (Transaction)
+			if (const FTransaction* Transaction = Trans->GetTransaction(CurrentTransactionIndex))
 			{
 				TArray<UObject*> TransactionObjects;
 				Transaction->GetTransactionObjects(TransactionObjects);
 
-				if (TransactionObjects.Contains(this))
+				// We really just want the transactions that contain *our* LevelSequence, but it seems like when we swap LevelSequences the newly
+				// created LevelSequence is not in the TransactionObjects, so we would fail to detect the right transaction on redo (as our "current
+				// LevelSequence" would have been this new one, that is not part of TransactionObjects)
+				bool bTransactionContainsLevelSequence = false;
+				bool bTransactionContainsThis = false;
+				for (UObject* TransactionObject : TransactionObjects)
+				{
+					if (TransactionObject == this)
+					{
+						bTransactionContainsThis = true;
+					}
+					else if (Cast<ULevelSequence>(TransactionObject))
+					{
+						bTransactionContainsLevelSequence = true;
+					}
+				}
+
+				if (bTransactionContainsLevelSequence && bTransactionContainsThis)
 				{
 					if (InTransactionState == ETransactionStateEventType::UndoRedoStarted)
 					{
-						const bool bLevelSequenceEditorWasOpened = GEditor->GetEditorSubsystem<UAssetEditorSubsystem>()->CloseAllEditorsForAsset(LevelSequence) > 0;
-						if (bLevelSequenceEditorWasOpened)
+						if (IUsdStageModule* UsdStageModule = FModuleManager::Get().GetModulePtr<IUsdStageModule>(TEXT("UsdStage")))
 						{
-							ActorsThatClosedTheSequencer.Add(this);
+							TArray<TWeakPtr<ISequencer>> SequencersToReset;
+							for (const TWeakPtr<ISequencer>& ExistingSequencer : UsdStageModule->GetExistingSequencers())
+							{
+								if (TSharedPtr<ISequencer> PinnedSequencer = ExistingSequencer.Pin())
+								{
+									if (PinnedSequencer->GetRootMovieSceneSequence() == LevelSequence)
+									{
+										SequencersToReset.Add(PinnedSequencer);
+
+										// Hack for solving UE-171596
+										// In this transaction we will switch LevelSequences, and have a Sequencer opened displaying our current Sequence.
+										// - We cannot leave this Sequencer displaying our old LevelSequence, because it will go PendingKill, and as the
+										//   Sequencer fetches it through WeakPtrs it will not find a valid LevelSequence and crash (this was the reason for
+										//   the original UE-127253 hack above). This means on UndoRedoStarted we *must* do something;
+										// - We cannot set our new LevelSequence into it yet of course, because it hasn't been created yet (it will be spawned
+										//   by the undo system after UndoRedoStarted);
+										// - We cannot close this Sequencer, because of this "DeferredModify" mechanism that pushes some updates to the end
+										//   of the transaction (to UndoRedoFinalized). If one of those updates executes after we close the Sequencer and before
+										//   we fix things up (which it can always do as the order of execution of the delegates is not deterministic), it will
+										//   crash (this is the issue at UE-171596);
+										//
+										// This means we're forced to give *some valid LevelSequence* to the Sequencer for the split second while we switch
+										// our actual generated LevelSequence.
+										static TStrongObjectPtr<ULevelSequence> DummySequencePtr = nullptr;
+										ULevelSequence* DummySequence = DummySequencePtr.Get();
+										if(!DummySequence)
+										{
+											DummySequence = NewObject<ULevelSequence>();
+											DummySequence->Initialize();
+
+											DummySequencePtr.Reset(DummySequence);
+										}
+										PinnedSequencer->ResetToNewRootSequence(*DummySequence);
+									}
+								}
+							}
+							ActorsToSequencers.Add(this, SequencersToReset);
 						}
 					}
 
 					if (InTransactionState == ETransactionStateEventType::UndoRedoFinalized)
 					{
-						if (ActorsThatClosedTheSequencer.Contains(this))
+						if (TArray<TWeakPtr<ISequencer>>* FoundSequencers = ActorsToSequencers.Find(this))
 						{
-							GEditor->GetEditorSubsystem<UAssetEditorSubsystem>()->OpenEditorForAsset(LevelSequence);
-							ActorsThatClosedTheSequencer.Remove(this);
+							for (TWeakPtr<ISequencer>& Sequencer : *FoundSequencers)
+							{
+								if (TSharedPtr<ISequencer> PinnedSequencer = Sequencer.Pin())
+								{
+									if (LevelSequence && PinnedSequencer->GetRootMovieSceneSequence() != LevelSequence)
+									{
+										PinnedSequencer->ResetToNewRootSequence(*LevelSequence);
+									}
+								}
+							}
+							ActorsToSequencers.Remove(this);
 						}
 					}
 				}
