USD: Always rebuild the info cache before importing as we may have new import options, so our info cache may be stale.

diff --git a/Engine/Plugins/Importers/USDImporter/Source/USDStageEditorViewModels/Private/USDStageViewModel.cpp b/Engine/Plugins/Importers/USDImporter/Source/USDStageEditorViewModels/Private/USDStageViewModel.cpp
--- a/Engine/Plugins/Importers/USDImporter/Source/USDStageEditorViewModels/Private/USDStageViewModel.cpp
+++ b/Engine/Plugins/Importers/USDImporter/Source/USDStageEditorViewModels/Private/USDStageViewModel.cpp
@@ -351,8 +351,6 @@ void FUsdStageViewModel::ImportStage( const TCHAR* TargetContentFolder, UUsdStag
 			// Let the importer reuse our assets, but force it to spawn new actors and components always
 			// This allows a different setting for asset/component collapsing, and doesn't require modifying the PrimTwins
 			ImportContext.AssetCache = StageActor->UsdAssetCache;
-			ImportContext.InfoCache = StageActor->GetInfoCache();
-			ImportContext.LevelSequenceHelper.SetInfoCache(StageActor->GetInfoCache());
 
 			ImportContext.TargetSceneActorAttachParent = StageActor->GetRootComponent()->GetAttachParent();
 			ImportContext.TargetSceneActorTargetTransform = StageActor->GetActorTransform();
diff --git a/Engine/Plugins/Importers/USDImporter/Source/USDStageImporter/Private/USDStageImportContext.cpp b/Engine/Plugins/Importers/USDImporter/Source/USDStageImporter/Private/USDStageImportContext.cpp
--- a/Engine/Plugins/Importers/USDImporter/Source/USDStageImporter/Private/USDStageImportContext.cpp
+++ b/Engine/Plugins/Importers/USDImporter/Source/USDStageImporter/Private/USDStageImportContext.cpp
@@ -89,7 +89,6 @@ void FUsdStageImportContext::Reset()
 	ImportedAsset = nullptr;
 	LevelSequenceHelper.Clear();
 	AssetCache = nullptr;
-	InfoCache = nullptr;
 	MaterialToPrimvarToUVIndex.Empty();
 	Stage = UE::FUsdStage{};
 	ImportObjectFlags = EObjectFlags::RF_NoFlags;
diff --git a/Engine/Plugins/Importers/USDImporter/Source/USDStageImporter/Private/USDStageImporter.cpp b/Engine/Plugins/Importers/USDImporter/Source/USDStageImporter/Private/USDStageImporter.cpp
--- a/Engine/Plugins/Importers/USDImporter/Source/USDStageImporter/Private/USDStageImporter.cpp
+++ b/Engine/Plugins/Importers/USDImporter/Source/USDStageImporter/Private/USDStageImporter.cpp
@@ -282,10 +282,9 @@ namespace UsdStageImporterImpl
 			TranslationContext.InfoCache = MakeShared<FUsdInfoCache>();
 		}
 
-		if ( TranslationContext.InfoCache->IsEmpty() )
-		{
-			TranslationContext.InfoCache->RebuildCacheForSubtree( TranslationContext.Stage.GetPseudoRoot(), TranslationContext );
-		}
+		// It's better to always rebuild the info cache because our import options may have changed
+		// from the options used when the cache was first built, which could change collapsing states, etc.
+		TranslationContext.InfoCache->RebuildCacheForSubtree(TranslationContext.Stage.GetPseudoRoot(), TranslationContext);
 	}
 
 	void ImportMaterials(FUsdStageImportContext& ImportContext, FUsdSchemaTranslationContext& TranslationContext)
@@ -1800,15 +1799,9 @@ void UUsdStageImporter::ImportFromFile(FUsdStageImportContext& ImportContext)
 		);
 		ImportContext.AssetCache = NewObject<UUsdAssetCache2>();
 	}
-	if (!ImportContext.InfoCache)
-	{
-		UE_LOG(LogUsd, Log, TEXT("Generating a temporary USD Info Cache when importing '%s'."),
-			*ImportContext.FilePath
-		);
-		ImportContext.InfoCache = MakeShared<FUsdInfoCache>();
-	}
+	TSharedPtr<FUsdInfoCache> InfoCache = MakeShared<FUsdInfoCache>();
 	ImportContext.AssetCache->MarkAssetsAsStale();
-	ImportContext.LevelSequenceHelper.SetInfoCache(ImportContext.InfoCache);
+	ImportContext.LevelSequenceHelper.SetInfoCache(InfoCache);
 	ImportContext.LevelSequenceHelper.Init( ImportContext.Stage );  // Must happen after the context gets an InfoCache!
 
 	// Shotgun approach to recreate all render states because we may want to reimport/delete/reassign a material/static/skeletalmesh while it is currently being drawn
@@ -1833,7 +1826,7 @@ void UUsdStageImporter::ImportFromFile(FUsdStageImportContext& ImportContext)
 	TranslationContext->bAllowParsingSkeletalAnimations = ImportContext.ImportOptions->bImportGeometry && ImportContext.ImportOptions->bImportSkeletalAnimations;
 	TranslationContext->bAllowParsingGroomAssets = ImportContext.ImportOptions->bImportGroomAssets;
 	TranslationContext->bTranslateOnlyUsedMaterials = ImportContext.ImportOptions->bImportOnlyUsedMaterials;
-	TranslationContext->InfoCache = ImportContext.InfoCache;
+	TranslationContext->InfoCache = InfoCache;
 	TranslationContext->BlendShapesByPath = &BlendShapesByPath;
 	TranslationContext->GroomInterpolationSettings = ImportContext.ImportOptions->GroomInterpolationSettings;
 	{
@@ -1924,15 +1917,9 @@ bool UUsdStageImporter::ReimportSingleAsset(
 		);
 		ImportContext.AssetCache = NewObject<UUsdAssetCache2>();
 	}
-	if (!ImportContext.InfoCache)
-	{
-		UE_LOG(LogUsd, Log, TEXT("Generating a temporary USD Info Cache when importing '%s'."),
-			*ImportContext.FilePath
-		);
-		ImportContext.InfoCache = MakeShared<FUsdInfoCache>();
-	}
+	TSharedPtr<FUsdInfoCache> InfoCache = MakeShared<FUsdInfoCache>();
 	ImportContext.AssetCache->MarkAssetsAsStale();
-	ImportContext.LevelSequenceHelper.SetInfoCache(ImportContext.InfoCache);
+	ImportContext.LevelSequenceHelper.SetInfoCache(InfoCache);
 	ImportContext.LevelSequenceHelper.Init(ImportContext.Stage);  // Must happen after the context gets an InfoCache!
 
 	// Shotgun approach to recreate all render states because we may want to reimport/delete/reassign a material/static/skeletalmesh while it is currently being drawn
@@ -1956,7 +1943,7 @@ bool UUsdStageImporter::ReimportSingleAsset(
 	TranslationContext->bAllowParsingSkeletalAnimations = ImportContext.ImportOptions->bImportGeometry && ImportContext.ImportOptions->bImportSkeletalAnimations;
 	TranslationContext->bAllowParsingGroomAssets = ImportContext.ImportOptions->bImportGroomAssets;
 	TranslationContext->bTranslateOnlyUsedMaterials = ImportContext.ImportOptions->bImportOnlyUsedMaterials;
-	TranslationContext->InfoCache = ImportContext.InfoCache;
+	TranslationContext->InfoCache = InfoCache;
 	TranslationContext->BlendShapesByPath = &BlendShapesByPath;
 	TranslationContext->GroomInterpolationSettings = ImportContext.ImportOptions->GroomInterpolationSettings;
 	{
diff --git a/Engine/Plugins/Importers/USDImporter/Source/USDStageImporter/Public/USDStageImportContext.h b/Engine/Plugins/Importers/USDImporter/Source/USDStageImporter/Public/USDStageImportContext.h
--- a/Engine/Plugins/Importers/USDImporter/Source/USDStageImporter/Public/USDStageImportContext.h
+++ b/Engine/Plugins/Importers/USDImporter/Source/USDStageImporter/Public/USDStageImportContext.h
@@ -61,6 +61,7 @@ struct USDSTAGEIMPORTER_API FUsdStageImportContext
 	TObjectPtr<UUsdAssetCache2> AssetCache;
 
 	/** Caches various information about prims that are expensive to query */
+	UE_DEPRECATED(5.3, "The import process now always builds its own InfoCache, so this member is no longer used")
 	TSharedPtr<FUsdInfoCache> InfoCache;
 
 	/**
@@ -97,6 +98,12 @@ struct USDSTAGEIMPORTER_API FUsdStageImportContext
 public:
 	FUsdStageImportContext();
 
+	PRAGMA_DISABLE_DEPRECATION_WARNINGS
+	~FUsdStageImportContext() = default;
+	FUsdStageImportContext(const FUsdStageImportContext& Other) = default;
+	FUsdStageImportContext& operator=(const FUsdStageImportContext& Other) = default;
+	PRAGMA_ENABLE_DEPRECATION_WARNINGS
+
 	bool Init(const FString& InName, const FString& InFilePath, const FString& InInitialPackagePath, EObjectFlags InFlags, bool bInIsAutomated, bool bIsReimport = false, bool bAllowActorImport = true);
 	void Reset();
 
