USD: Refactor/cleanup event listening to better handle sublayer added/removed events.

diff --git a/Engine/Plugins/Importers/USDImporter/Source/USDStage/Private/USDLevelSequenceHelper.cpp b/Engine/Plugins/Importers/USDImporter/Source/USDStage/Private/USDLevelSequenceHelper.cpp
--- a/Engine/Plugins/Importers/USDImporter/Source/USDStage/Private/USDLevelSequenceHelper.cpp
+++ b/Engine/Plugins/Importers/USDImporter/Source/USDStage/Private/USDLevelSequenceHelper.cpp
@@ -2728,12 +2728,12 @@ void FUsdLevelSequenceHelperImpl::OnUsdObjectsChanged(const UsdUtils::FObjectCha
 
 	FScopedBlockMonitoringChangesForTransaction BlockMonitoring{*this};
 
-	for (const TPair<FString, TArray<UsdUtils::FObjectChangeNotice>>& InfoChange : InfoChanges)
+	for (const TPair<FString, TArray<UsdUtils::FSdfChangeListEntry>>& InfoChange : InfoChanges)
 	{
 		const FString& PrimPath = InfoChange.Key;
 		if (PrimPath == TEXT("/"))
 		{
-			for (const UsdUtils::FObjectChangeNotice& ObjectChange : InfoChange.Value)
+			for (const UsdUtils::FSdfChangeListEntry& ObjectChange : InfoChange.Value)
 			{
 				for (const UsdUtils::FAttributeChange& AttributeChange : ObjectChange.AttributeChanges)
 				{
diff --git a/Engine/Plugins/Importers/USDImporter/Source/USDStage/Private/USDStageActor.cpp b/Engine/Plugins/Importers/USDImporter/Source/USDStage/Private/USDStageActor.cpp
--- a/Engine/Plugins/Importers/USDImporter/Source/USDStage/Private/USDStageActor.cpp
+++ b/Engine/Plugins/Importers/USDImporter/Source/USDStage/Private/USDStageActor.cpp
@@ -963,15 +963,15 @@ AUsdStageActor::AUsdStageActor()
 
 		UsdListener.GetOnObjectsChanged().AddUObject(this, &AUsdStageActor::OnUsdObjectsChanged);
 
-		UsdListener.GetOnLayersChanged().AddLambda(
-			[&, this](const TArray< FString >& ChangedLayerIdentifiers)
+		UsdListener.GetOnSdfLayersChanged().AddLambda(
+			[&, this](const UsdUtils::FLayerToSdfChangeList& LayersToChangeList)
 			{
-				if (!IsListeningToUsdNotices() || ChangedLayerIdentifiers.Num() == 0)
+				if (!IsListeningToUsdNotices() || LayersToChangeList.Num() == 0)
 				{
 					return;
 				}
 
-				TOptional< TGuardValue< ITransaction* > > SuppressTransaction;
+				TOptional<TGuardValue<ITransaction*>> SuppressTransaction;
 				if (this->GetOutermost()->HasAnyPackageFlags(PKG_PlayInEditor))
 				{
 					SuppressTransaction.Emplace(GUndo, nullptr);
@@ -983,32 +983,52 @@ AUsdStageActor::AUsdStageActor()
 					return;
 				}
 
-				TSet<FString> UsedLayerIdentifiers;
-				TArray<UE::FSdfLayer> UsedLayers = Stage.GetUsedLayers();
-				UsedLayerIdentifiers.Reserve(UsedLayers.Num());
-				for (const UE::FSdfLayer& UsedLayer : UsedLayers)
+				// Check to see if any of the stage's layers reloaded, or if we added/removed any layer
+				TSet<UE::FSdfLayer> UsedLayers{Stage.GetUsedLayers()};
+				for (const TPair<UE::FSdfLayerWeak, UsdUtils::FSdfChangeList>& LayerToChangeList : LayersToChangeList)
 				{
-					UsedLayerIdentifiers.Add(UsedLayer.GetIdentifier());
-				}
-
-				// Check to see if any layer reloaded. If so, rebuild all of our animations as a single layer changing
-				// might propagate timecodes through all level sequences
-				for (const FString& ChangedLayerIdentifier : ChangedLayerIdentifiers)
-				{
-					if (!UsedLayerIdentifiers.Contains(ChangedLayerIdentifier))
+					if (!UsedLayers.Contains(LayerToChangeList.Key))
 					{
 						continue;
 					}
 
-					UE_LOG(LogUsd, Verbose, TEXT("Reloading animations because layer '%s' was added/removed/reloaded"), *ChangedLayerIdentifier);
-					ReloadAnimations();
+					for (const TPair<UE::FSdfPath, UsdUtils::FSdfChangeListEntry>& Change : LayerToChangeList.Value)
+					{
+						if (Change.Value.Flags.bDidReloadContent)
+						{
+							UE_LOG(
+								LogUsd,
+								Verbose,
+								TEXT("Reloading animations because layer '%s' was reloaded"),
+								*LayerToChangeList.Key.GetIdentifier()
+							);
+							ReloadAnimations();
+
+							// Make sure our PrimsToAnimate and the LevelSequenceHelper are kept in sync, because we'll use PrimsToAnimate to
+							// check whether we need to call LevelSequenceHelper::AddPrim within AUsdStageActor::ExpandPrim. Without this reset
+							// our prims would already be in here by the time we're checking if we need to add tracks or not, and we wouldn't re-add
+							// the tracks
+							PrimsToAnimate.Reset();
+							return;
+						}
 
-					// Make sure our PrimsToAnimate and the LevelSequenceHelper are kept in sync, because we'll use PrimsToAnimate to
-					// check whether we need to call LevelSequenceHelper::AddPrim within AUsdStageActor::ExpandPrim. Without this reset
-					// our prims would already be in here by the time we're checking if we need to add tracks or not, and we wouldn't re-add
-					// the tracks
-					PrimsToAnimate.Reset();
-					return;
+						for (const TPair<FString, UsdUtils::ESubLayerChangeType>& SubLayerChange : Change.Value.SubLayerChanges)
+						{
+							if (SubLayerChange.Value == UsdUtils::ESubLayerChangeType::SubLayerAdded
+								|| SubLayerChange.Value == UsdUtils::ESubLayerChangeType::SubLayerRemoved)
+							{
+								UE_LOG(
+									LogUsd,
+									Verbose,
+									TEXT("Reloading animations because layer '%s' was added or removed"),
+									*LayerToChangeList.Key.GetIdentifier()
+								);
+								ReloadAnimations();
+								PrimsToAnimate.Reset();
+								return;
+							}
+						}
+					}
 				}
 			}
 		);
@@ -1189,13 +1209,13 @@ void AUsdStageActor::OnUsdObjectsChanged(const UsdUtils::FObjectChangesByPath& I
 
 	// The most important thing here is to iterate in parent to child order, so build SortedPrimsChangedList
 	TMap< FString, bool > SortedPrimsChangedList;
-	for (const TPair<FString, TArray<UsdUtils::FObjectChangeNotice>>& InfoChange : InfoChanges)
+	for (const TPair<FString, TArray<UsdUtils::FSdfChangeListEntry>>& InfoChange : InfoChanges)
 	{
 		UE::FSdfPath PrimPath = UE::FSdfPath(*InfoChange.Key).StripAllVariantSelections();
 
 		// Upgrade some info changes into resync changes
 		bool bIsResync = false;
-		for (const UsdUtils::FObjectChangeNotice& ObjectChange : InfoChange.Value)
+		for (const UsdUtils::FSdfChangeListEntry& ObjectChange : InfoChange.Value)
 		{
 			// This is in charge of resyncing components (and so calling CreateComponents instead of just
 			// UpdateComponents) whenever we change material assignments. This is important because
@@ -1254,7 +1274,7 @@ void AUsdStageActor::OnUsdObjectsChanged(const UsdUtils::FObjectChangesByPath& I
 		SortedPrimsChangedList.Add(UE::FSdfPath(*InfoChange.Key).StripAllVariantSelections().GetString(), bIsResync);
 	}
 	// Do Resyncs after so that they overwrite pure info changes if we have any
-	for (const TPair<FString, TArray<UsdUtils::FObjectChangeNotice>>& ResyncChange : ResyncChanges)
+	for (const TPair<FString, TArray<UsdUtils::FSdfChangeListEntry>>& ResyncChange : ResyncChanges)
 	{
 		UE::FSdfPath PrimPath = UE::FSdfPath(*ResyncChange.Key).StripAllVariantSelections();
 
diff --git a/Engine/Plugins/Importers/USDImporter/Source/USDStage/Private/USDTransactor.cpp b/Engine/Plugins/Importers/USDImporter/Source/USDStage/Private/USDTransactor.cpp
--- a/Engine/Plugins/Importers/USDImporter/Source/USDStage/Private/USDTransactor.cpp
+++ b/Engine/Plugins/Importers/USDImporter/Source/USDStage/Private/USDTransactor.cpp
@@ -141,15 +141,15 @@ namespace UsdUtils
 	{
 		InOutEdits.Edits.Reserve( InChanges.Num() + InOutEdits.Edits.Num() );
 
-		for ( const TPair<FString, TArray<FObjectChangeNotice>>& ChangesByPrimPath : InChanges )
+		for ( const TPair<FString, TArray<FSdfChangeListEntry>>& ChangesByPrimPath : InChanges )
 		{
 			const FString& PrimPath = ChangesByPrimPath.Key;
-			const TArray<FObjectChangeNotice>& Changes = ChangesByPrimPath.Value;
+			const TArray<FSdfChangeListEntry>& Changes = ChangesByPrimPath.Value;
 
 			FTransactorRecordedEdit& Edit = InOutEdits.Edits.Emplace_GetRef();
 			Edit.ObjectPath = PrimPath;
 
-			for ( const FObjectChangeNotice& Change : Changes )
+			for ( const FSdfChangeListEntry& Change : Changes )
 			{
 				FTransactorObjectChange& ConvertedChange = Edit.ObjectChanges.Emplace_GetRef();
 				ConvertedChange.Flags = Change.Flags;
diff --git a/Engine/Plugins/Importers/USDImporter/Source/USDStage/Public/USDTransactor.h b/Engine/Plugins/Importers/USDImporter/Source/USDStage/Public/USDTransactor.h
--- a/Engine/Plugins/Importers/USDImporter/Source/USDStage/Public/USDTransactor.h
+++ b/Engine/Plugins/Importers/USDImporter/Source/USDStage/Public/USDTransactor.h
@@ -16,8 +16,8 @@ namespace UsdUtils
 {
 	class FUsdTransactorImpl;
 
-	struct FObjectChangeNotice;
-	using FObjectChangesByPath = TMap<FString, TArray<FObjectChangeNotice>>;
+	struct FSdfChangeListEntry;
+	using FObjectChangesByPath = TMap<FString, TArray<FSdfChangeListEntry>>;
 }
 
 namespace UE
diff --git a/Engine/Plugins/Importers/USDImporter/Source/USDStageEditor/Private/SUSDStage.cpp b/Engine/Plugins/Importers/USDImporter/Source/USDStageEditor/Private/SUSDStage.cpp
--- a/Engine/Plugins/Importers/USDImporter/Source/USDStageEditor/Private/SUSDStage.cpp
+++ b/Engine/Plugins/Importers/USDImporter/Source/USDStageEditor/Private/SUSDStage.cpp
@@ -663,32 +663,64 @@ void SUsdStage::SetupStageActorDelegates()
 			}
 		);
 
-		OnLayersChangedHandle = ViewModel.UsdStageActor->GetUsdListener().GetOnLayersChanged().AddLambda(
-			[this](const TArray<FString>& LayersNames)
+		TFunction<bool(float)> DelayedRefreshOfLayersTreeView = [this](float Time)
+		{
+			if (bEditorIsShuttingDown)
 			{
-				FTSTicker::GetCoreTicker().AddTicker(FTickerDelegate::CreateLambda(
-					[this](float Time)
+				return false;
+			}
+
+			AUsdStageActor* StageActor = ViewModel.UsdStageActor.Get();
+			if (this->UsdLayersTreeView && StageActor)
+			{
+				constexpr bool bResync = false;
+				UsdLayersTreeView->Refresh(StageActor->GetBaseUsdStage(), StageActor->GetIsolatedUsdStage(), bResync);
+			}
+
+			// Returning false means this is a one-off, and won't repeat
+			return false;
+		};
+
+		OnSdfLayersChangedHandle = ViewModel.UsdStageActor->GetUsdListener().GetOnSdfLayersChanged().AddLambda(
+			[this, DelayedRefreshOfLayersTreeView](const UsdUtils::FLayerToSdfChangeList& LayersToChangeList)
+			{
+				// The layers changed events are global, and so fired for every stage actor when any layer of
+				// any stage is modified. Let's update only when a layer used by our stage is updated
+				bool bRelevantToCurrentStage = false;
+				if (AUsdStageActor* StageActor = ViewModel.UsdStageActor.Get())
+				{
+					if (const UE::FUsdStage& CurrentStage = StageActor->GetUsdStage())
 					{
-						if (bEditorIsShuttingDown)
+						TSet<UE::FSdfLayer> UsedLayers{CurrentStage.GetUsedLayers()};
+						for (const TPair<UE::FSdfLayerWeak, UsdUtils::FSdfChangeList>& LayerToChangeList : LayersToChangeList)
 						{
-							return false;
+							if (UsedLayers.Contains(LayerToChangeList.Key))
+							{
+								bRelevantToCurrentStage = true;
+								break;
+							}
 						}
+					}
+				}
+				if (!bRelevantToCurrentStage)
+				{
+					return;
+				}
 
-						AUsdStageActor* StageActor = ViewModel.UsdStageActor.Get();
-						if (this->UsdLayersTreeView && StageActor)
-						{
-							constexpr bool bResync = false;
-							UsdLayersTreeView->Refresh(
-								StageActor->GetBaseUsdStage(),
-								StageActor->GetIsolatedUsdStage(),
-								bResync
-							);
-						}
+				FTSTicker::GetCoreTicker().AddTicker(FTickerDelegate::CreateLambda(DelayedRefreshOfLayersTreeView));
+			}
+		);
 
-						// Returning false means this is a one-off, and won't repeat
-						return false;
-					}
-				));
+		OnSdfLayerDirtinessChangedHandle = ViewModel.UsdStageActor->GetUsdListener().GetOnSdfLayerDirtinessChanged().AddLambda(
+			[DelayedRefreshOfLayersTreeView]()
+			{
+				// The layer dirtiness notice as we're using here is global, which means any layer of any other stage
+				// going dirty will affect us. We *could* make this notice layer-specific on the listener, but it would force us to
+				// register/unregister to every single layer... since the updates done in DelayedRefreshOfLayersTreeView are
+				// relatively light, it's probably for the best to just refresh every time than to support the churn of
+				// a bunch of separate listeners and registration/unregistration happening every time, potentially also
+				// crashing if done incorrectly when a layer is created/added/removed/etc.
+				FTSTicker::GetCoreTicker().AddTicker(FTickerDelegate::CreateLambda(DelayedRefreshOfLayersTreeView));
 			}
 		);
 	}
@@ -704,7 +736,8 @@ void SUsdStage::ClearStageActorDelegates()
 		StageActor->OnActorDestroyed.Remove ( OnActorDestroyedHandle );
 
 		StageActor->GetUsdListener().GetOnStageEditTargetChanged().Remove( OnStageEditTargetChangedHandle );
-		StageActor->GetUsdListener().GetOnLayersChanged().Remove( OnLayersChangedHandle );
+		StageActor->GetUsdListener().GetOnSdfLayersChanged().Remove( OnSdfLayersChangedHandle );
+		StageActor->GetUsdListener().GetOnSdfLayerDirtinessChanged().Remove( OnSdfLayerDirtinessChangedHandle );
 	}
 }
 
diff --git a/Engine/Plugins/Importers/USDImporter/Source/USDStageEditor/Private/SUSDStage.h b/Engine/Plugins/Importers/USDImporter/Source/USDStageEditor/Private/SUSDStage.h
--- a/Engine/Plugins/Importers/USDImporter/Source/USDStageEditor/Private/SUSDStage.h
+++ b/Engine/Plugins/Importers/USDImporter/Source/USDStageEditor/Private/SUSDStage.h
@@ -124,7 +124,8 @@ protected:
 	FDelegateHandle OnStageChangedHandle;
 	FDelegateHandle OnStageEditTargetChangedHandle;
 	FDelegateHandle OnPrimChangedHandle;
-	FDelegateHandle OnLayersChangedHandle;
+	FDelegateHandle OnSdfLayersChangedHandle;
+	FDelegateHandle OnSdfLayerDirtinessChangedHandle;
 	FDelegateHandle OnViewportSelectionChangedHandle;
 	FDelegateHandle PostPIEStartedHandle;
 	FDelegateHandle EndPIEHandle;
diff --git a/Engine/Plugins/Importers/USDImporter/Source/UnrealUSDWrapper/Private/USDListener.cpp b/Engine/Plugins/Importers/USDImporter/Source/UnrealUSDWrapper/Private/USDListener.cpp
--- a/Engine/Plugins/Importers/USDImporter/Source/UnrealUSDWrapper/Private/USDListener.cpp
+++ b/Engine/Plugins/Importers/USDImporter/Source/UnrealUSDWrapper/Private/USDListener.cpp
@@ -5,6 +5,8 @@
 #include "USDLog.h"
 #include "USDMemory.h"
 
+#include "UsdWrappers/ForwardDeclarations.h"
+#include "UsdWrappers/SdfLayer.h"
 #include "UsdWrappers/SdfPath.h"
 #include "UsdWrappers/UsdStage.h"
 #include "UsdWrappers/VtValue.h"
@@ -29,6 +31,101 @@
 namespace UsdToUnreal
 {
 #if USE_USD_SDK
+	void ConvertSdfChangeListEntry(
+		const pxr::SdfPath& ChangedObject,
+		const pxr::SdfChangeList::Entry& Entry,
+		UsdUtils::FSdfChangeListEntry& OutConverted
+	)
+	{
+		using namespace pxr;
+
+		// For most changes we'll only get one of these, but sometimes multiple changes are fired in sequence
+		// (e.g. if you change framesPerSecond, it will send a notice for it but also for the matching, updated timeCodesPerSecond)
+		for ( const std::pair<TfToken, std::pair<VtValue, VtValue>>& AttributeChange : Entry.infoChanged ) // Note: infoChanged here is just a naming conflict, there's no "resyncChanged"
+		{
+			UsdUtils::FAttributeChange& ConvertedAttributeChange = OutConverted.AttributeChanges.Emplace_GetRef();
+
+			const FString FieldToken = ANSI_TO_TCHAR( AttributeChange.first.GetString().c_str() );
+
+			// For regular properties (most common case) FullFieldPath will already carry the property and FieldToken will be "default", "timeSamples", "variability", etc.
+			if ( ChangedObject.IsPropertyPath() )
+			{
+				ConvertedAttributeChange.PropertyName = ANSI_TO_TCHAR( ChangedObject.GetName().c_str() );
+				ConvertedAttributeChange.Field = FieldToken;
+			}
+			// For stage properties it seems like USD likes to send just "/" as the FullFieldPath, and the actual property name in FieldToken (e.g. "metersPerUnit", "timeCodesPerSecond", etc.)
+			// We send "default" here for consistency
+			else if ( ChangedObject.IsAbsoluteRootPath() || FieldToken == FString( TEXT( "kind" ) ) )
+			{
+				ConvertedAttributeChange.PropertyName = FieldToken;
+				ConvertedAttributeChange.Field = TEXT( "default" );
+			}
+
+			ConvertedAttributeChange.OldValue = UE::FVtValue{ AttributeChange.second.first };
+			ConvertedAttributeChange.NewValue = UE::FVtValue{ AttributeChange.second.second };
+		}
+
+		// Some notices (like creating/removing a property) don't have any actual infoChanged entries, so we need to create a fake one in here in order to convey the PropertyName, if applicable
+		if ( Entry.infoChanged.size() == 0 &&
+			( Entry.flags.didAddProperty ||
+				Entry.flags.didAddPropertyWithOnlyRequiredFields ||
+				Entry.flags.didRemoveProperty ||
+				Entry.flags.didRemovePropertyWithOnlyRequiredFields ||
+				Entry.flags.didChangeAttributeTimeSamples ) )
+		{
+			UsdUtils::FAttributeChange& ConvertedAttributeChange = OutConverted.AttributeChanges.Emplace_GetRef();
+			ConvertedAttributeChange.Field = Entry.flags.didChangeAttributeTimeSamples ? TEXT( "timeSamples" ) : TEXT( "default" );
+
+			if ( ChangedObject.IsPropertyPath() )
+			{
+				ConvertedAttributeChange.PropertyName = ANSI_TO_TCHAR( ChangedObject.GetName().c_str() );
+			}
+		}
+
+		// These should be packed just the same, but just in case we do member by member here instead of memcopying it over
+		UsdUtils::FPrimChangeFlags& ConvertedFlags = OutConverted.Flags;
+		ConvertedFlags.bDidChangeIdentifier = Entry.flags.didChangeIdentifier;
+		ConvertedFlags.bDidChangeResolvedPath = Entry.flags.didChangeResolvedPath;
+		ConvertedFlags.bDidReplaceContent = Entry.flags.didReplaceContent;
+		ConvertedFlags.bDidReloadContent = Entry.flags.didReloadContent;
+		ConvertedFlags.bDidReorderChildren = Entry.flags.didReorderChildren;
+		ConvertedFlags.bDidReorderProperties = Entry.flags.didReorderProperties;
+		ConvertedFlags.bDidRename = Entry.flags.didRename;
+		ConvertedFlags.bDidChangePrimVariantSets = Entry.flags.didChangePrimVariantSets;
+		ConvertedFlags.bDidChangePrimInheritPaths = Entry.flags.didChangePrimInheritPaths;
+		ConvertedFlags.bDidChangePrimSpecializes = Entry.flags.didChangePrimSpecializes;
+		ConvertedFlags.bDidChangePrimReferences = Entry.flags.didChangePrimReferences;
+		ConvertedFlags.bDidChangeAttributeTimeSamples = Entry.flags.didChangeAttributeTimeSamples;
+		ConvertedFlags.bDidChangeAttributeConnection = Entry.flags.didChangeAttributeConnection;
+		ConvertedFlags.bDidChangeRelationshipTargets = Entry.flags.didChangeRelationshipTargets;
+		ConvertedFlags.bDidAddTarget = Entry.flags.didAddTarget;
+		ConvertedFlags.bDidRemoveTarget = Entry.flags.didRemoveTarget;
+		ConvertedFlags.bDidAddInertPrim = Entry.flags.didAddInertPrim;
+		ConvertedFlags.bDidAddNonInertPrim = Entry.flags.didAddNonInertPrim;
+		ConvertedFlags.bDidRemoveInertPrim = Entry.flags.didRemoveInertPrim;
+		ConvertedFlags.bDidRemoveNonInertPrim = Entry.flags.didRemoveNonInertPrim;
+		ConvertedFlags.bDidAddPropertyWithOnlyRequiredFields = Entry.flags.didAddPropertyWithOnlyRequiredFields;
+		ConvertedFlags.bDidAddProperty = Entry.flags.didAddProperty;
+		ConvertedFlags.bDidRemovePropertyWithOnlyRequiredFields = Entry.flags.didRemovePropertyWithOnlyRequiredFields;
+		ConvertedFlags.bDidRemoveProperty = Entry.flags.didRemoveProperty;
+
+		static_assert( static_cast< int >( UsdUtils::ESubLayerChangeType::SubLayerAdded ) == static_cast< int >( pxr::SdfChangeList::SubLayerChangeType::SubLayerAdded ), "Enum values changed!" );
+		static_assert( static_cast< int >( UsdUtils::ESubLayerChangeType::SubLayerOffset ) == static_cast< int >( pxr::SdfChangeList::SubLayerChangeType::SubLayerOffset ), "Enum values changed!" );
+		static_assert( static_cast< int >( UsdUtils::ESubLayerChangeType::SubLayerRemoved ) == static_cast< int >( pxr::SdfChangeList::SubLayerChangeType::SubLayerRemoved ), "Enum values changed!" );
+		for ( const std::pair<std::string, pxr::SdfChangeList::SubLayerChangeType>& SubLayerChange : Entry.subLayerChanges )
+		{
+			OutConverted.SubLayerChanges.Add(
+				TPair<FString, UsdUtils::ESubLayerChangeType>(
+					ANSI_TO_TCHAR( SubLayerChange.first.c_str() ),
+					static_cast< UsdUtils::ESubLayerChangeType >( SubLayerChange.second )
+				)
+			);
+		}
+
+		OutConverted.OldPath = ANSI_TO_TCHAR( Entry.oldPath.GetString().c_str() );
+		OutConverted.OldIdentifier = ANSI_TO_TCHAR( Entry.oldIdentifier.c_str() );
+	}
+
 	bool ConvertPathRange( const pxr::UsdNotice::ObjectsChanged::PathRange& PathRange, UsdUtils::FObjectChangesByPath& OutChanges )
 	{
 		using namespace pxr;
@@ -37,9 +134,9 @@ namespace UsdToUnreal
 
 		OutChanges.Reset();
 
-		for ( UsdNotice::ObjectsChanged::PathRange::const_iterator It = PathRange.begin(); It != PathRange.end(); ++It )
+		for (UsdNotice::ObjectsChanged::PathRange::const_iterator It = PathRange.begin(); It != PathRange.end(); ++It)
 		{
-			const FString PrimPath = ANSI_TO_TCHAR( It->GetAbsoluteRootOrPrimPath().GetAsString().c_str() );
+			const FString PrimPath = ANSI_TO_TCHAR(It->GetAbsoluteRootOrPrimPath().GetAsString().c_str());
 
 			if (pxr::UsdPrim::IsPathInPrototype(UE::FSdfPath(*PrimPath)))
 			{
@@ -47,102 +144,20 @@ namespace UsdToUnreal
 			}
 
 			// Something like "/Root/Prim.some_field", or "/"
-			const FString FullFieldPath = ANSI_TO_TCHAR( It->GetAsString().c_str() );
+			const FString FullFieldPath = ANSI_TO_TCHAR(It->GetAsString().c_str());
 
-			TArray<UsdUtils::FObjectChangeNotice>& ConvertedChanges = OutChanges.FindOrAdd( PrimPath );
+			TArray<UsdUtils::FSdfChangeListEntry>& ConvertedChanges = OutChanges.FindOrAdd(PrimPath);
 
 			// Changes may be empty, but we should still pass along this overall notice because sending a root
 			// resync notice with no actual change item inside is how USD signals that a layer has been added/removed/resynced
 			const std::vector<const SdfChangeList::Entry*>& Changes = It.base()->second;
-			for ( const SdfChangeList::Entry* Entry : Changes )
+			for (const SdfChangeList::Entry* Entry : Changes)
 			{
-				UsdUtils::FObjectChangeNotice& ConvertedEntry = ConvertedChanges.Emplace_GetRef();
-
-				// For most changes we'll only get one of these, but sometimes multiple changes are fired in sequence
-				// (e.g. if you change framesPerSecond, it will send a notice for it but also for the matching, updated timeCodesPerSecond)
-				for ( const std::pair<TfToken, std::pair<VtValue, VtValue>>& AttributeChange : Entry->infoChanged ) // Note: infoChanged here is just a naming conflict, there's no "resyncChanged"
+				UsdUtils::FSdfChangeListEntry& ConvertedEntry = ConvertedChanges.Emplace_GetRef();
+				if (Entry)
 				{
-					UsdUtils::FAttributeChange& ConvertedAttributeChange = ConvertedEntry.AttributeChanges.Emplace_GetRef();
-
-					const FString FieldToken = ANSI_TO_TCHAR( AttributeChange.first.GetString().c_str() );
-
-					// For regular properties (most common case) FullFieldPath will already carry the property and FieldToken will be "default", "timeSamples", "variability", etc.
-					if ( It->IsPropertyPath() )
-					{
-						ConvertedAttributeChange.PropertyName = ANSI_TO_TCHAR( It->GetName().c_str() );
-						ConvertedAttributeChange.Field = FieldToken;
-					}
-					// For stage properties it seems like USD likes to send just "/" as the FullFieldPath, and the actual property name in FieldToken (e.g. "metersPerUnit", "timeCodesPerSecond", etc.)
-					// We send "default" here for consistency
-					else if ( It->IsAbsoluteRootPath() || FieldToken == FString( TEXT( "kind" ) ) )
-					{
-						ConvertedAttributeChange.PropertyName = FieldToken;
-						ConvertedAttributeChange.Field = TEXT( "default" );
-					}
-
-					ConvertedAttributeChange.OldValue = UE::FVtValue{ AttributeChange.second.first };
-					ConvertedAttributeChange.NewValue = UE::FVtValue{ AttributeChange.second.second };
-				}
-
-				// Some notices (like creating/removing a property) don't have any actual infoChanged entries, so we need to create a fake one in here in order to convey the PropertyName, if applicable
-				if ( Entry->infoChanged.size() == 0 &&
-					( Entry->flags.didAddProperty ||
-					  Entry->flags.didAddPropertyWithOnlyRequiredFields ||
-					  Entry->flags.didRemoveProperty ||
-					  Entry->flags.didRemovePropertyWithOnlyRequiredFields ||
-					  Entry->flags.didChangeAttributeTimeSamples ) )
-				{
-					UsdUtils::FAttributeChange& ConvertedAttributeChange = ConvertedEntry.AttributeChanges.Emplace_GetRef();
-					ConvertedAttributeChange.Field = Entry->flags.didChangeAttributeTimeSamples ? TEXT( "timeSamples" ) : TEXT( "default" );
-
-					if ( It->IsPropertyPath() )
-					{
-						ConvertedAttributeChange.PropertyName = ANSI_TO_TCHAR( It->GetName().c_str() );
-					}
+					ConvertSdfChangeListEntry(*It, *Entry, ConvertedEntry);
 				}
-
-				// These should be packed just the same, but just in case we do member by member here instead of memcopying it over
-				UsdUtils::FPrimChangeFlags& ConvertedFlags = ConvertedEntry.Flags;
-				ConvertedFlags.bDidChangeIdentifier = Entry->flags.didChangeIdentifier;
-				ConvertedFlags.bDidChangeResolvedPath = Entry->flags.didChangeResolvedPath;
-				ConvertedFlags.bDidReplaceContent = Entry->flags.didReplaceContent;
-				ConvertedFlags.bDidReloadContent = Entry->flags.didReloadContent;
-				ConvertedFlags.bDidReorderChildren = Entry->flags.didReorderChildren;
-				ConvertedFlags.bDidReorderProperties = Entry->flags.didReorderProperties;
-				ConvertedFlags.bDidRename = Entry->flags.didRename;
-				ConvertedFlags.bDidChangePrimVariantSets = Entry->flags.didChangePrimVariantSets;
-				ConvertedFlags.bDidChangePrimInheritPaths = Entry->flags.didChangePrimInheritPaths;
-				ConvertedFlags.bDidChangePrimSpecializes = Entry->flags.didChangePrimSpecializes;
-				ConvertedFlags.bDidChangePrimReferences = Entry->flags.didChangePrimReferences;
-				ConvertedFlags.bDidChangeAttributeTimeSamples = Entry->flags.didChangeAttributeTimeSamples;
-				ConvertedFlags.bDidChangeAttributeConnection = Entry->flags.didChangeAttributeConnection;
-				ConvertedFlags.bDidChangeRelationshipTargets = Entry->flags.didChangeRelationshipTargets;
-				ConvertedFlags.bDidAddTarget = Entry->flags.didAddTarget;
-				ConvertedFlags.bDidRemoveTarget = Entry->flags.didRemoveTarget;
-				ConvertedFlags.bDidAddInertPrim = Entry->flags.didAddInertPrim;
-				ConvertedFlags.bDidAddNonInertPrim = Entry->flags.didAddNonInertPrim;
-				ConvertedFlags.bDidRemoveInertPrim = Entry->flags.didRemoveInertPrim;
-				ConvertedFlags.bDidRemoveNonInertPrim = Entry->flags.didRemoveNonInertPrim;
-				ConvertedFlags.bDidAddPropertyWithOnlyRequiredFields = Entry->flags.didAddPropertyWithOnlyRequiredFields;
-				ConvertedFlags.bDidAddProperty = Entry->flags.didAddProperty;
-				ConvertedFlags.bDidRemovePropertyWithOnlyRequiredFields = Entry->flags.didRemovePropertyWithOnlyRequiredFields;
-				ConvertedFlags.bDidRemoveProperty = Entry->flags.didRemoveProperty;
-
-				static_assert( static_cast< int >( UsdUtils::ESubLayerChangeType::SubLayerAdded ) == static_cast< int >( pxr::SdfChangeList::SubLayerChangeType::SubLayerAdded ), "Enum values changed!" );
-				static_assert( static_cast< int >( UsdUtils::ESubLayerChangeType::SubLayerOffset ) == static_cast< int >( pxr::SdfChangeList::SubLayerChangeType::SubLayerOffset ), "Enum values changed!" );
-				static_assert( static_cast< int >( UsdUtils::ESubLayerChangeType::SubLayerRemoved ) == static_cast< int >( pxr::SdfChangeList::SubLayerChangeType::SubLayerRemoved ), "Enum values changed!" );
-				for ( const std::pair<std::string, pxr::SdfChangeList::SubLayerChangeType>& SubLayerChange : Entry->subLayerChanges )
-				{
-					ConvertedEntry.SubLayerChanges.Add(
-						TPair<FString, UsdUtils::ESubLayerChangeType>(
-							ANSI_TO_TCHAR( SubLayerChange.first.c_str() ),
-							static_cast< UsdUtils::ESubLayerChangeType >( SubLayerChange.second )
-						)
-					);
-				}
-
-				ConvertedEntry.OldPath = ANSI_TO_TCHAR( Entry->oldPath.GetString().c_str() );
-				ConvertedEntry.OldIdentifier = ANSI_TO_TCHAR( Entry->oldIdentifier.c_str() );
 			}
 		}
 
@@ -154,12 +169,12 @@ namespace UsdToUnreal
 		ConvertPathRange( InNotice.GetChangedInfoOnlyPaths(), OutInfoChanges );
 		ConvertPathRange( InNotice.GetResyncedPaths(), OutResyncChanges );
 
-		for ( TPair< FString, TArray<UsdUtils::FObjectChangeNotice> >& InfoPair : OutInfoChanges )
+		for ( TPair< FString, TArray<UsdUtils::FSdfChangeListEntry> >& InfoPair : OutInfoChanges )
 		{
 			const FString& PrimPath = InfoPair.Key;
-			TArray< UsdUtils::FObjectChangeNotice >* AnalogueResyncChanges = nullptr;
+			TArray< UsdUtils::FSdfChangeListEntry >* AnalogueResyncChanges = nullptr;
 
-			for ( TArray<UsdUtils::FObjectChangeNotice>::TIterator ChangeIt = InfoPair.Value.CreateIterator(); ChangeIt; ++ChangeIt )
+			for ( TArray<UsdUtils::FSdfChangeListEntry>::TIterator ChangeIt = InfoPair.Value.CreateIterator(); ChangeIt; ++ChangeIt )
 			{
 				bool bUpgrade = false;
 
@@ -208,12 +223,18 @@ class FUsdListenerImpl
 #endif // #if USE_USD_SDK
 {
 public:
+	PRAGMA_DISABLE_DEPRECATION_WARNINGS
 	FUsdListenerImpl() = default;
+	PRAGMA_ENABLE_DEPRECATION_WARNINGS
 
 	virtual ~FUsdListenerImpl();
 
 	FUsdListener::FOnStageEditTargetChanged OnStageEditTargetChanged;
+	PRAGMA_DISABLE_DEPRECATION_WARNINGS
 	FUsdListener::FOnLayersChanged OnLayersChanged;
+	PRAGMA_ENABLE_DEPRECATION_WARNINGS
+	FUsdListener::FOnSdfLayersChanged OnSdfLayersChanged;
+	FUsdListener::FOnSdfLayerDirtinessChanged OnSdfLayerDirtinessChanged;
 	FUsdListener::FOnObjectsChanged OnObjectsChanged;
 
 	FThreadSafeCounter IsBlocked;
@@ -224,16 +245,16 @@ public:
 	void Register( const pxr::UsdStageRefPtr& Stage );
 
 protected:
-	void HandleObjectsChangedNotice( const pxr::UsdNotice::ObjectsChanged& Notice, const pxr::UsdStageWeakPtr& Sender );
-	void HandleStageEditTargetChangedNotice( const pxr::UsdNotice::StageEditTargetChanged& Notice, const pxr::UsdStageWeakPtr& Sender );
-	void HandleLayersChangedNotice( const pxr::SdfNotice::LayersDidChange& Notice );
-	void HandleLayerDirtinessChangedNotice( const pxr::SdfNotice::LayerDirtinessChanged& Notice );
+	void HandleObjectsChangedNotice(const pxr::UsdNotice::ObjectsChanged& Notice, const pxr::UsdStageWeakPtr& Sender);
+	void HandleStageEditTargetChangedNotice(const pxr::UsdNotice::StageEditTargetChanged& Notice, const pxr::UsdStageWeakPtr& Sender);
+	void HandleLayersChangedNotice(const pxr::SdfNotice::LayersDidChange& Notice);
+	void HandleLayerDirtinessChangedNotice(const pxr::SdfNotice::LayerDirtinessChanged& Notice);
 
 private:
 	pxr::TfNotice::Key RegisteredObjectsChangedKey;
 	pxr::TfNotice::Key RegisteredStageEditTargetChangedKey;
-	pxr::TfNotice::Key RegisteredLayersChangedKey;
-	pxr::TfNotice::Key RegisteredLayerDirtinessChangedKey;
+	pxr::TfNotice::Key RegisteredSdfLayersChangedKey;
+	pxr::TfNotice::Key RegisteredSdfLayerDirtinessChangedKey;
 #endif // #if USE_USD_SDK
 };
 
@@ -277,10 +298,22 @@ FUsdListener::FOnStageEditTargetChanged& FUsdListener::GetOnStageEditTargetChang
 	return Impl->OnStageEditTargetChanged;
 }
 
+PRAGMA_DISABLE_DEPRECATION_WARNINGS
 FUsdListener::FOnLayersChanged& FUsdListener::GetOnLayersChanged()
 {
 	return Impl->OnLayersChanged;
 }
+PRAGMA_ENABLE_DEPRECATION_WARNINGS
+
+FUsdListener::FOnSdfLayersChanged& FUsdListener::GetOnSdfLayersChanged()
+{
+	return Impl->OnSdfLayersChanged;
+}
+
+FUsdListener::FOnSdfLayerDirtinessChanged& FUsdListener::GetOnSdfLayerDirtinessChanged()
+{
+	return Impl->OnSdfLayerDirtinessChanged;
+}
 
 FUsdListener::FOnObjectsChanged& FUsdListener::GetOnObjectsChanged()
 {
@@ -296,36 +329,40 @@ void FUsdListenerImpl::Register( const pxr::UsdStageRefPtr& Stage )
 	{
 		pxr::TfNotice::Revoke( RegisteredObjectsChangedKey );
 	}
-
 	RegisteredObjectsChangedKey = pxr::TfNotice::Register( pxr::TfWeakPtr< FUsdListenerImpl >( this ), &FUsdListenerImpl::HandleObjectsChangedNotice, Stage );
 
 	if ( RegisteredStageEditTargetChangedKey.IsValid() )
 	{
 		pxr::TfNotice::Revoke( RegisteredStageEditTargetChangedKey );
 	}
-
 	RegisteredStageEditTargetChangedKey = pxr::TfNotice::Register( pxr::TfWeakPtr< FUsdListenerImpl >( this ), &FUsdListenerImpl::HandleStageEditTargetChangedNotice, Stage );
 
-	if (RegisteredLayersChangedKey.IsValid())
+	if (RegisteredSdfLayersChangedKey.IsValid())
 	{
-		pxr::TfNotice::Revoke(RegisteredLayersChangedKey);
+		pxr::TfNotice::Revoke(RegisteredSdfLayersChangedKey);
 	}
+	RegisteredSdfLayersChangedKey = pxr::TfNotice::Register( pxr::TfWeakPtr< FUsdListenerImpl >( this ), &FUsdListenerImpl::HandleLayersChangedNotice );
 
-	RegisteredLayersChangedKey = pxr::TfNotice::Register( pxr::TfWeakPtr< FUsdListenerImpl >( this ), &FUsdListenerImpl::HandleLayersChangedNotice );
-	RegisteredLayerDirtinessChangedKey = pxr::TfNotice::Register( pxr::TfWeakPtr< FUsdListenerImpl >( this ), &FUsdListenerImpl::HandleLayerDirtinessChangedNotice );
+	if (RegisteredSdfLayerDirtinessChangedKey.IsValid())
+	{
+		pxr::TfNotice::Revoke(RegisteredSdfLayerDirtinessChangedKey);
+	}
+	RegisteredSdfLayerDirtinessChangedKey = pxr::TfNotice::Register( pxr::TfWeakPtr< FUsdListenerImpl >( this ), &FUsdListenerImpl::HandleLayerDirtinessChangedNotice );
 }
 #endif // #if USE_USD_SDK
 
+PRAGMA_DISABLE_DEPRECATION_WARNINGS
 FUsdListenerImpl::~FUsdListenerImpl()
 {
 #if USE_USD_SDK
 	FScopedUsdAllocs UsdAllocs;
 	pxr::TfNotice::Revoke( RegisteredObjectsChangedKey );
 	pxr::TfNotice::Revoke( RegisteredStageEditTargetChangedKey );
-	pxr::TfNotice::Revoke( RegisteredLayersChangedKey );
-	pxr::TfNotice::Revoke( RegisteredLayerDirtinessChangedKey );
+	pxr::TfNotice::Revoke( RegisteredSdfLayersChangedKey );
+	pxr::TfNotice::Revoke( RegisteredSdfLayerDirtinessChangedKey );
 #endif // #if USE_USD_SDK
 }
+PRAGMA_ENABLE_DEPRECATION_WARNINGS
 
 #if USE_USD_SDK
 void FUsdListenerImpl::HandleObjectsChangedNotice( const pxr::UsdNotice::ObjectsChanged& Notice, const pxr::UsdStageWeakPtr& Sender )
@@ -358,26 +395,34 @@ void FUsdListenerImpl::HandleStageEditTargetChangedNotice( const pxr::UsdNotice:
 
 void FUsdListenerImpl::HandleLayersChangedNotice( const pxr::SdfNotice::LayersDidChange& Notice )
 {
-	if ( !OnLayersChanged.IsBound() || IsBlocked.GetValue() > 0 )
+PRAGMA_DISABLE_DEPRECATION_WARNINGS
+	if ((!OnLayersChanged.IsBound() && !OnSdfLayersChanged.IsBound()) || IsBlocked.GetValue() > 0)
 	{
 		return;
 	}
 
-	TArray< FString > LayersNames;
-
-	FScopedUsdAllocs UsdAllocs;
-	[&](const pxr::SdfLayerChangeListVec& ChangeVec)
+	TArray<FString> LayersNames;
+	UsdUtils::FLayerToSdfChangeList ConvertedLayerToChangeList;
 	{
-		// Check to see if any layer reloaded. If so, rebuild all of our animations as a single layer changing
-		// might propagate timecodes through all level sequences
-		for (const std::pair<pxr::SdfLayerHandle, pxr::SdfChangeList>& ChangeVecItem : ChangeVec)
+		FScopedUnrealAllocs Allocs;
+
+		const std::vector<std::pair<pxr::SdfLayerHandle, pxr::SdfChangeList>>& UsdChangeLists = Notice.GetChangeListVec();
+		ConvertedLayerToChangeList.Reserve(UsdChangeLists.size());
+
+		for (const std::pair<pxr::SdfLayerHandle, pxr::SdfChangeList>& ChangeVecItem : UsdChangeLists)
 		{
+			const pxr::SdfLayerHandle& Layer = ChangeVecItem.first;
 			const pxr::SdfChangeList::EntryList& ChangeList = ChangeVecItem.second.GetEntryList();
+
+			TPair<UE::FSdfLayerWeak, UsdUtils::FSdfChangeList>& ConvertedChangeList = ConvertedLayerToChangeList.Emplace_GetRef();
+			ConvertedChangeList.Key = UE::FSdfLayerWeak{Layer};
+			ConvertedChangeList.Value.Reserve(ChangeList.size());
+
 			for (const std::pair<pxr::SdfPath, pxr::SdfChangeList::Entry>& Change : ChangeList)
 			{
-				// Sadly we can't do much with these notices as USD doesn't provide information of which layer the
-				// layer was added *to*. Since this is a global notice not even tied to a particular stage,
-				// it doesn't help us know when to update what
+				TPair<UE::FSdfPath, UsdUtils::FSdfChangeListEntry>& ConvertedChange = ConvertedChangeList.Value.Emplace_GetRef();
+				ConvertedChange.Key = UE::FSdfPath{Change.first};
+				UsdToUnreal::ConvertSdfChangeListEntry(Change.first, Change.second, ConvertedChange.Value);
 
 				const pxr::SdfChangeList::Entry::_Flags& Flags = Change.second.flags;
 				if (Flags.didReloadContent)
@@ -386,23 +431,24 @@ void FUsdListenerImpl::HandleLayersChangedNotice( const pxr::SdfNotice::LayersDi
 				}
 			}
 		}
-	}( Notice.GetChangeListVec() );
+	}
 
 	FScopedUnrealAllocs UnrealAllocs;
 	OnLayersChanged.Broadcast( LayersNames );
+	OnSdfLayersChanged.Broadcast(ConvertedLayerToChangeList);
+PRAGMA_ENABLE_DEPRECATION_WARNINGS
 }
 
-void FUsdListenerImpl::HandleLayerDirtinessChangedNotice( const pxr::SdfNotice::LayerDirtinessChanged& Notice )
+void FUsdListenerImpl::HandleLayerDirtinessChangedNotice(const pxr::SdfNotice::LayerDirtinessChanged& Notice)
 {
-	if ( !OnLayersChanged.IsBound() || IsBlocked.GetValue() > 0 )
+	if (!OnSdfLayerDirtinessChanged.IsBound() || IsBlocked.GetValue() > 0)
 	{
 		return;
 	}
 
 	FScopedUnrealAllocs UnrealAllocs;
 
-	TArray< FString > LayersNames; // We don't know which layer dirtiness has changed unfortunately so we'll provide an empty list for now.
-	OnLayersChanged.Broadcast(LayersNames);
+	OnSdfLayerDirtinessChanged.Broadcast();
 }
 #endif // #if USE_USD_SDK
 
diff --git a/Engine/Plugins/Importers/USDImporter/Source/UnrealUSDWrapper/Private/UsdWrappers/SdfLayer.cpp b/Engine/Plugins/Importers/USDImporter/Source/UnrealUSDWrapper/Private/UsdWrappers/SdfLayer.cpp
--- a/Engine/Plugins/Importers/USDImporter/Source/UnrealUSDWrapper/Private/UsdWrappers/SdfLayer.cpp
+++ b/Engine/Plugins/Importers/USDImporter/Source/UnrealUSDWrapper/Private/UsdWrappers/SdfLayer.cpp
@@ -153,7 +153,7 @@ namespace UE
 		return ( bool ) Impl->GetInner();
 	}
 
-	uint32 GetTypeHash( const UE::FSdfLayerWeak& Layer )
+	UNREALUSDWRAPPER_API uint32 GetTypeHash(const UE::FSdfLayerWeak& Layer)
 	{
 		uint32 Result = 0;
 #if USE_USD_SDK
@@ -163,8 +163,8 @@ namespace UE
 		// TfWeakPtrFacade, which is how TfWeakPtr seem to be hashed in USD.
 		// We have to clamp the size_t to uint32 here but deferring to USD is likely the best approach to guarantee
 		// that objects with non-colliding hashes in USD also not collide in UE.
-		Result = static_cast< uint32 >( pxr::TfHash()( Layer.Impl->GetInner() ) );
-#endif // #if USE_USD_SDK
+		Result = static_cast<uint32>(pxr::TfHash()(Layer.Impl->GetInner()));
+#endif	  // #if USE_USD_SDK
 		return Result;
 	}
 
diff --git a/Engine/Plugins/Importers/USDImporter/Source/UnrealUSDWrapper/Public/USDListener.h b/Engine/Plugins/Importers/USDImporter/Source/UnrealUSDWrapper/Public/USDListener.h
--- a/Engine/Plugins/Importers/USDImporter/Source/UnrealUSDWrapper/Public/USDListener.h
+++ b/Engine/Plugins/Importers/USDImporter/Source/UnrealUSDWrapper/Public/USDListener.h
@@ -11,6 +11,11 @@
 
 class FUsdListenerImpl;
 
+namespace UE
+{
+	class FSdfPath;
+}
+
 namespace UsdUtils
 {
 	/** Analogous to pxr::SdfChangeList::Entry::_Flags */
@@ -68,7 +73,7 @@ namespace UsdUtils
 	};
 
 	/** Analogous to pxr::SdfChangeList::Entry, describes a generic change to an object */
-	struct FObjectChangeNotice
+	struct FSdfChangeListEntry
 	{
 		TArray<FAttributeChange> AttributeChanges;
 		FPrimChangeFlags Flags;
@@ -76,13 +81,17 @@ namespace UsdUtils
 		FString OldIdentifier;						// Empty if Flags.bDIdChangeIdentifier is not set
 		TArray<TPair<FString, ESubLayerChangeType>> SubLayerChanges;
 	};
+	using FObjectChangeNotice = FSdfChangeListEntry; // Renamed in 5.3 as it is used for layer changes now too
+
+	using FSdfChangeList = TArray<TPair<UE::FSdfPath, FSdfChangeListEntry>>;
+	using FLayerToSdfChangeList = TArray<TPair<UE::FSdfLayerWeak, FSdfChangeList>>;
 
 	/**
 	 * Describes USD object changes by object path.
 	 * The only difference to the USD data structures is that we store them by object path here for convenience, so
 	 * the key can be "/" to signify a stage change, or something like "/MyRoot/SomePrim" to indicate a particular prim change
 	 */
-	using FObjectChangesByPath = TMap<FString, TArray<FObjectChangeNotice>>;
+	using FObjectChangesByPath = TMap<FString, TArray<FSdfChangeListEntry>>;
 }
 
 /**
@@ -109,11 +118,21 @@ public:
 	void Unblock();
 	bool IsBlocked() const;
 
+	// Stage-specific events
 	DECLARE_EVENT( FUsdListener, FOnStageEditTargetChanged );
 	FOnStageEditTargetChanged& GetOnStageEditTargetChanged();
 
-	DECLARE_EVENT_OneParam( FUsdListener, FOnLayersChanged, const TArray< FString >& );
+	DECLARE_EVENT_OneParam( FUsdListener, UE_DEPRECATED(5.3, "Use FOnSdfLayersChanged") FOnLayersChanged, const TArray< FString >& );
+	UE_DEPRECATED(5.3, "Use GetOnSdfLayersChanged")
+  	PRAGMA_DISABLE_DEPRECATION_WARNINGS
 	FOnLayersChanged& GetOnLayersChanged();
+  	PRAGMA_ENABLE_DEPRECATION_WARNINGS
+
+	DECLARE_EVENT_OneParam(FUsdListener, FOnSdfLayersChanged, const UsdUtils::FLayerToSdfChangeList&);
+	FOnSdfLayersChanged& GetOnSdfLayersChanged();
+
+	DECLARE_EVENT(FUsdListener, FOnSdfLayerDirtinessChanged);
+	FOnSdfLayerDirtinessChanged& GetOnSdfLayerDirtinessChanged();
 
 	DECLARE_EVENT_TwoParams( FUsdListener, FOnObjectsChanged, const UsdUtils::FObjectChangesByPath& /* InfoChanges */, const UsdUtils::FObjectChangesByPath& /* ResyncChanges */ );
 	FOnObjectsChanged& GetOnObjectsChanged();
