USD: Fix compilation when building without editor.

diff --git a/Engine/Plugins/Importers/USDImporter/Source/USDSchemas/Private/USDShadeMaterialTranslator.cpp b/Engine/Plugins/Importers/USDImporter/Source/USDSchemas/Private/USDShadeMaterialTranslator.cpp
--- a/Engine/Plugins/Importers/USDImporter/Source/USDSchemas/Private/USDShadeMaterialTranslator.cpp
+++ b/Engine/Plugins/Importers/USDImporter/Source/USDSchemas/Private/USDShadeMaterialTranslator.cpp
@@ -489,11 +489,13 @@ void FUsdShadeMaterialTranslator::PostImportMaterial(const FString& MaterialHash
 			else if (bIsOwnedByTransientPackage)
 			{
 				FString FilePath;
+#if WITH_EDITOR
 				if(UAssetImportData* TextureImportData = Texture->AssetImportData.Get())
 				{
 					FilePath = TextureImportData->GetFirstFilename();
 				}
 				else
+#endif // WITH_EDITOR
 				{
 					FilePath = Texture->GetName();
 				}
