USD: remove comments in usdShaders plugInfo.json

Comments are not officially supported in JSON, and so are
also not supported by UE's JSON deserializer.

diff --git a/Engine/Plugins/Importers/USDImporter/Resources/UsdResources/Win64/plugins/usdShaders/resources/plugInfo.json b/Engine/Plugins/Importers/USDImporter/Resources/UsdResources/Win64/plugins/usdShaders/resources/plugInfo.json
--- a/Engine/Plugins/Importers/USDImporter/Resources/UsdResources/Win64/plugins/usdShaders/resources/plugInfo.json
+++ b/Engine/Plugins/Importers/USDImporter/Resources/UsdResources/Win64/plugins/usdShaders/resources/plugInfo.json
@@ -8,8 +8,6 @@
                         "displayName": "USD-based shader definition discovery plugin."
                     }
                 },
-                # Allow usdShaders/shaders to be accessible by hio/glslfx.cpp's 
-                #   ShaderResourceRegistry
                 "ShaderResources": "shaders"
             },
             "LibraryPath": "../../../../../../../../Binaries/Win64/usdShaders.dll",
