USD: Consider child prim Kinds when deciding if a prim collapses children.

diff --git a/Engine/Plugins/Importers/USDImporter/Source/USDSchemas/Private/USDGeomXformableTranslator.cpp b/Engine/Plugins/Importers/USDImporter/Source/USDSchemas/Private/USDGeomXformableTranslator.cpp
--- a/Engine/Plugins/Importers/USDImporter/Source/USDSchemas/Private/USDGeomXformableTranslator.cpp
+++ b/Engine/Plugins/Importers/USDImporter/Source/USDSchemas/Private/USDGeomXformableTranslator.cpp
@@ -687,7 +687,12 @@ bool FUsdGeomXformableTranslator::CollapsesChildren( ECollapsingType CollapsingT
 	if ( Model )
 	{
 		EUsdDefaultKind PrimKind = UsdUtils::GetDefaultKind( Prim );
-		bCollapsesChildren = Context->KindsToCollapse != EUsdDefaultKind::None && ( EnumHasAnyFlags( Context->KindsToCollapse, PrimKind ) || ( PrimKind == EUsdDefaultKind::None && GCollapsePrimsWithoutKind ) );
+
+		// Note that this is false if PrimKind is None
+		const bool bPrimKindShouldCollapse = EnumHasAnyFlags( Context->KindsToCollapse, PrimKind );
+
+		bCollapsesChildren = Context->KindsToCollapse != EUsdDefaultKind::None &&
+			( bPrimKindShouldCollapse || ( PrimKind == EUsdDefaultKind::None && GCollapsePrimsWithoutKind ) );
 
 		if ( !bCollapsesChildren )
 		{
@@ -741,7 +746,10 @@ bool FUsdGeomXformableTranslator::CanBeCollapsed( ECollapsingType CollapsingType
 		return false;
 	}
 
-	return true;
+	EUsdDefaultKind PrimKind = UsdUtils::GetDefaultKind(GetPrim());
+	// Note that this is false if PrimKind is None
+	const bool bThisPrimCanCollapse = EnumHasAnyFlags(Context->KindsToCollapse, PrimKind);
+	return bThisPrimCanCollapse || (PrimKind == EUsdDefaultKind::None && GCollapsePrimsWithoutKind);
 }
 
 TSet<UE::FSdfPath> FUsdGeomXformableTranslator::CollectAuxiliaryPrims() const
