USD: Fix empty material overrides layers when exporting levels.

diff --git a/Engine/Plugins/Importers/USDImporter/Source/USDUtilities/Private/USDPrimConversion.cpp b/Engine/Plugins/Importers/USDImporter/Source/USDUtilities/Private/USDPrimConversion.cpp
--- a/Engine/Plugins/Importers/USDImporter/Source/USDUtilities/Private/USDPrimConversion.cpp
+++ b/Engine/Plugins/Importers/USDImporter/Source/USDUtilities/Private/USDPrimConversion.cpp
@@ -396,7 +396,7 @@ void UsdToUnreal::PropagateTransform(const pxr::UsdStageRefPtr& Stage, const pxr
 		if (!bResetXformStack)
 		{
 			OutTransform *= CurrentTransform;
-			
+
 			if (Leaf != Root)
 			{
 				if (!Leaf.IsPseudoRoot())
@@ -2051,24 +2051,32 @@ bool UnrealToUsd::ConvertMeshComponent( const pxr::UsdStageRefPtr& Stage, const
 					continue;
 				}
 
+				// If we have user data this is one of our meshes, so we know exactly the prim that corresponds to each
+				// material slot. Let's use that
 				// const_cast as there's no const access to asset user data on IInterface_AssetUserData, but we won't
 				// modify anything
-				const UUsdMeshAssetUserData* UserData =
-					const_cast<UGeometryCache*>(GeometryCache)->GetAssetUserData<UUsdMeshAssetUserData>();
-				if (!UserData)
-				{
-					return false;
-				}
-
-				if (const FUsdPrimPathList* SourcePrimPaths = UserData->MaterialSlotToPrimPaths.Find(MatIndex))
+				if (const UUsdMeshAssetUserData* UserData =
+					const_cast<UGeometryCache*>(GeometryCache)->GetAssetUserData<UUsdMeshAssetUserData>())
 				{
-					for ( const FString& PrimPath : SourcePrimPaths->PrimPaths )
+					if (const FUsdPrimPathList* SourcePrimPaths = UserData->MaterialSlotToPrimPaths.Find(MatIndex))
 					{
-						pxr::SdfPath OverridePrimPath = UnrealToUsd::ConvertPath( *PrimPath ).Get();
-						pxr::UsdPrim MeshPrim = Stage->OverridePrim( OverridePrimPath );
-						UE::USDPrimConversionImpl::Private::AuthorMaterialOverride( MeshPrim, Override->GetPathName() );
+						for ( const FString& PrimPath : SourcePrimPaths->PrimPaths )
+						{
+							pxr::SdfPath OverridePrimPath = UnrealToUsd::ConvertPath( *PrimPath ).Get();
+							pxr::UsdPrim MeshPrim = Stage->OverridePrim( OverridePrimPath );
+							UE::USDPrimConversionImpl::Private::AuthorMaterialOverride( MeshPrim, Override->GetPathName() );
+						}
 					}
 				}
+				// If we don't, we have to fallback to writing the same prim patterns that the mesh exporters
+				// generate when exporting meshes, so that we can override its opinions. This happens when exporting
+				// geometry cache / static mesh / skeletal mesh components, for example
+				else
+				{
+					pxr::SdfPath OverridePrimPath = UsdPrim.GetPath();
+					pxr::UsdPrim MeshPrim = Stage->OverridePrim(OverridePrimPath);
+					UE::USDPrimConversionImpl::Private::AuthorMaterialOverride(MeshPrim, Override->GetPathName());
+				}
 			}
 		}
 	}
@@ -2076,11 +2084,10 @@ bool UnrealToUsd::ConvertMeshComponent( const pxr::UsdStageRefPtr& Stage, const
 	{
 		if ( UStaticMesh* Mesh = StaticMeshComponent->GetStaticMesh() )
 		{
+			int32 NumLODs = Mesh->GetNumLODs();
+			const bool bHasLODs = NumLODs > 1;
+
 			const UUsdMeshAssetUserData* UserData = Mesh->GetAssetUserData<UUsdMeshAssetUserData>();
-			if (!UserData)
-			{
-				return false;
-			}
 
 			const TArray<UMaterialInterface*>& Overrides = MeshComponent->OverrideMaterials;
 			for ( int32 MatIndex = 0; MatIndex < Overrides.Num(); ++MatIndex )
@@ -2091,7 +2098,7 @@ bool UnrealToUsd::ConvertMeshComponent( const pxr::UsdStageRefPtr& Stage, const
 					continue;
 				}
 
-				for ( int32 LODIndex = 0; LODIndex < Mesh->GetNumLODs(); ++LODIndex )
+				for ( int32 LODIndex = 0; LODIndex < NumLODs; ++LODIndex )
 				{
 					int32 NumSections = Mesh->GetNumSections( LODIndex );
 					const bool bHasSubsets = NumSections > 1;
@@ -2104,14 +2111,53 @@ bool UnrealToUsd::ConvertMeshComponent( const pxr::UsdStageRefPtr& Stage, const
 							continue;
 						}
 
-						if (const FUsdPrimPathList* SourcePrimPaths = UserData->MaterialSlotToPrimPaths.Find(SectionMatIndex))
+						if (UserData)
+						{
+							if (const FUsdPrimPathList* SourcePrimPaths = UserData->MaterialSlotToPrimPaths.Find(SectionMatIndex))
+							{
+								for ( const FString& PrimPath : SourcePrimPaths->PrimPaths )
+								{
+									pxr::SdfPath OverridePrimPath = UnrealToUsd::ConvertPath( *PrimPath ).Get();
+									pxr::UsdPrim MeshPrim = Stage->OverridePrim( OverridePrimPath );
+									UE::USDPrimConversionImpl::Private::AuthorMaterialOverride( MeshPrim, Override->GetPathName() );
+								}
+							}
+						}
+						else
 						{
-							for ( const FString& PrimPath : SourcePrimPaths->PrimPaths )
+							pxr::SdfPath OverridePrimPath = UsdPrim.GetPath();
+
+							// If we have only 1 LOD, the asset's DefaultPrim will be the Mesh prim directly.
+							// If we have multiple, the default prim won't have any schema, but will contain separate
+							// Mesh prims for each LOD named "LOD0", "LOD1", etc., switched via a "LOD" variant set
+							if (bHasLODs)
 							{
-								pxr::SdfPath OverridePrimPath = UnrealToUsd::ConvertPath( *PrimPath ).Get();
-								pxr::UsdPrim MeshPrim = Stage->OverridePrim( OverridePrimPath );
-								UE::USDPrimConversionImpl::Private::AuthorMaterialOverride( MeshPrim, Override->GetPathName() );
+								OverridePrimPath = OverridePrimPath.AppendPath(
+									UnrealToUsd::ConvertPath(*FString::Printf(TEXT("LOD%d"), LODIndex)).Get()
+								);
 							}
+
+							// If our LOD has only one section, its material assignment will be authored directly on the Mesh prim.
+							// If it has more than one material slot, we'll author UsdGeomSubset for each LOD Section, and author the material
+							// assignment there
+							if (bHasSubsets)
+							{
+								// Assume the UE sections are in the same order as the USD ones
+								std::vector<pxr::UsdGeomSubset> GeomSubsets = pxr::UsdShadeMaterialBindingAPI(UsdPrim).GetMaterialBindSubsets();
+								if (SectionIndex < GeomSubsets.size())
+								{
+									OverridePrimPath = OverridePrimPath.AppendChild(GeomSubsets[SectionIndex].GetPrim().GetName());
+								}
+								else
+								{
+									OverridePrimPath = OverridePrimPath.AppendPath(
+										UnrealToUsd::ConvertPath(*FString::Printf(TEXT("Section%d"), SectionIndex)).Get()
+									);
+								}
+							}
+
+							pxr::UsdPrim MeshPrim = Stage->OverridePrim(OverridePrimPath);
+							UE::USDPrimConversionImpl::Private::AuthorMaterialOverride(MeshPrim, Override->GetPathName());
 						}
 					}
 				}
@@ -2136,10 +2182,6 @@ bool UnrealToUsd::ConvertMeshComponent( const pxr::UsdStageRefPtr& Stage, const
 
 			const UUsdMeshAssetUserData* UserData =
 				const_cast<USkeletalMesh*>(SkeletalMesh)->GetAssetUserData<UUsdMeshAssetUserData>();
-			if (!UserData)
-			{
-				return false;
-			}
 
 			TIndirectArray<FSkeletalMeshLODRenderData>& LodRenderData = RenderData->LODRenderData;
 			if ( LodRenderData.Num() == 0 )
@@ -2147,6 +2189,21 @@ bool UnrealToUsd::ConvertMeshComponent( const pxr::UsdStageRefPtr& Stage, const
 				return false;
 			}
 
+			int32 NumLODs = SkeletalMesh->GetLODNum();
+			const bool bHasLODs = NumLODs > 1;
+
+			FString MeshName;
+			if ( !bHasLODs )
+			{
+				for ( const pxr::UsdPrim& Child : UsdPrim.GetChildren() )
+				{
+					if ( pxr::UsdGeomMesh Mesh{ Child } )
+					{
+						MeshName = UsdToUnreal::ConvertToken( Child.GetName() );
+					}
+				}
+			}
+
 			const TArray<UMaterialInterface*>& Overrides = MeshComponent->OverrideMaterials;
 			for ( int32 MatIndex = 0; MatIndex < Overrides.Num(); ++MatIndex )
 			{
@@ -2156,7 +2213,7 @@ bool UnrealToUsd::ConvertMeshComponent( const pxr::UsdStageRefPtr& Stage, const
 					continue;
 				}
 
-				for ( int32 LODIndex = 0; LODIndex < SkeletalMesh->GetLODNum(); ++LODIndex )
+				for ( int32 LODIndex = 0; LODIndex < NumLODs; ++LODIndex )
 				{
 					if ( !LodRenderData.IsValidIndex( LODIndex ) )
 					{
@@ -2179,19 +2236,62 @@ bool UnrealToUsd::ConvertMeshComponent( const pxr::UsdStageRefPtr& Stage, const
 							SectionMatIndex = LODInfo->LODMaterialMap[ SectionIndex ];
 						}
 
-						if ( SectionMatIndex != MatIndex )
+						if (SectionMatIndex != MatIndex)
 						{
 							continue;
 						}
 
-						if (const FUsdPrimPathList* SourcePrimPaths = UserData->MaterialSlotToPrimPaths.Find(SectionMatIndex))
+						if (UserData)
+						{
+							if (const FUsdPrimPathList* SourcePrimPaths = UserData->MaterialSlotToPrimPaths.Find(SectionMatIndex))
+							{
+								for (const FString& PrimPath : SourcePrimPaths->PrimPaths)
+								{
+									pxr::SdfPath OverridePrimPath = UnrealToUsd::ConvertPath(*PrimPath).Get();
+									pxr::UsdPrim MeshPrim = Stage->OverridePrim(OverridePrimPath);
+									UE::USDPrimConversionImpl::Private::AuthorMaterialOverride(MeshPrim, Override->GetPathName());
+								}
+							}
+						}
+						else
 						{
-							for ( const FString& PrimPath : SourcePrimPaths->PrimPaths )
+							pxr::SdfPath OverridePrimPath = UsdPrim.GetPath();
+
+							// If we have only 1 LOD, the asset's DefaultPrim will be a SkelRoot, and the Mesh will be a subprim
+							// with the same name.If we have the default prim is also the SkelRoot, but will contain separate
+							// Mesh prims for each LOD named "LOD0", "LOD1", etc., switched via a "LOD" variant set
+							if (bHasLODs)
+							{
+								OverridePrimPath = OverridePrimPath.AppendPath(
+									UnrealToUsd::ConvertPath(*FString::Printf(TEXT("LOD%d"), LODIndex)).Get()
+								);
+							}
+							else
 							{
-								pxr::SdfPath OverridePrimPath = UnrealToUsd::ConvertPath( *PrimPath ).Get();
-								pxr::UsdPrim MeshPrim = Stage->OverridePrim( OverridePrimPath );
-								UE::USDPrimConversionImpl::Private::AuthorMaterialOverride( MeshPrim, Override->GetPathName() );
+								OverridePrimPath = OverridePrimPath.AppendElementString(UnrealToUsd::ConvertString(*MeshName).Get());
 							}
+
+							// If our LOD has only one section, its material assignment will be authored directly on the Mesh prim.
+							// If it has more than one material slot, we'll author UsdGeomSubset for each LOD Section, and author the material
+							// assignment there
+							if (bHasSubsets)
+							{
+								// Assume the UE sections are in the same order as the USD ones
+								std::vector<pxr::UsdGeomSubset> GeomSubsets = pxr::UsdShadeMaterialBindingAPI(UsdPrim).GetMaterialBindSubsets();
+								if (SectionIndex < GeomSubsets.size())
+								{
+									OverridePrimPath = OverridePrimPath.AppendChild(GeomSubsets[SectionIndex].GetPrim().GetName());
+								}
+								else
+								{
+									OverridePrimPath = OverridePrimPath.AppendPath(
+										UnrealToUsd::ConvertPath(*FString::Printf(TEXT("Section%d"), SectionIndex)).Get()
+									);
+								}
+							}
+
+							pxr::UsdPrim MeshPrim = Stage->OverridePrim(OverridePrimPath);
+							UE::USDPrimConversionImpl::Private::AuthorMaterialOverride(MeshPrim, Override->GetPathName());
 						}
 					}
 				}
