USD: Fix crash when translating point instancers with unexpected prims inside their prototypes.

diff --git a/Engine/Plugins/Importers/USDImporter/Source/USDSchemas/Private/USDGeomPointInstancerTranslator.cpp b/Engine/Plugins/Importers/USDImporter/Source/USDSchemas/Private/USDGeomPointInstancerTranslator.cpp
--- a/Engine/Plugins/Importers/USDImporter/Source/USDSchemas/Private/USDGeomPointInstancerTranslator.cpp
+++ b/Engine/Plugins/Importers/USDImporter/Source/USDSchemas/Private/USDGeomPointInstancerTranslator.cpp
@@ -567,10 +567,13 @@ TSet<UE::FSdfPath> FUsdGeomPointInstancerTranslator::CollectAuxiliaryPrims() con
 					UE::FUsdTyped(ChildPrim.Get())
 				);
 
-				TSet<UE::FSdfPath> RecursiveDependencies = ChildSchemaTranslator->CollectAuxiliaryPrims();
-				for (const UE::FSdfPath& RecursiveDependency : RecursiveDependencies)
+				if (ChildSchemaTranslator)
 				{
-					Result.Add(RecursiveDependency);
+					TSet<UE::FSdfPath> RecursiveDependencies = ChildSchemaTranslator->CollectAuxiliaryPrims();
+					for (const UE::FSdfPath& RecursiveDependency : RecursiveDependencies)
+					{
+						Result.Add(RecursiveDependency);
+					}
 				}
 			}
 		}
@@ -589,7 +592,7 @@ TSet<UE::FSdfPath> FUsdGeomPointInstancerTranslator::CollectAuxiliaryPrims() con
 					Context,
 					UE::FUsdTyped(*PrimRangeIt)
 				);
-				if (SchemaTranslator->CollapsesChildren(ECollapsingType::Assets))
+				if (SchemaTranslator && SchemaTranslator->CollapsesChildren(ECollapsingType::Assets))
 				{
 					PrimRangeIt.PruneChildren();
 				}
