USD: Fix traversal of geometry cache hierarchies with invisible parent prims.

diff --git a/Engine/Plugins/Importers/USDImporter/Source/USDSchemas/Private/USDGeometryCacheTranslator.cpp b/Engine/Plugins/Importers/USDImporter/Source/USDSchemas/Private/USDGeometryCacheTranslator.cpp
--- a/Engine/Plugins/Importers/USDImporter/Source/USDSchemas/Private/USDGeometryCacheTranslator.cpp
+++ b/Engine/Plugins/Importers/USDImporter/Source/USDSchemas/Private/USDGeometryCacheTranslator.cpp
@@ -651,30 +651,7 @@ namespace UsdGeometryCacheTranslatorImpl
 			OutMeshData.BatchesInfo.Add(BatchInfo);
 		}
 	}
-
-	bool IsPurposefullyVisible(const pxr::UsdPrim& Prim, EUsdPurpose Purpose)
-	{
-		FScopedUsdAllocs UsdAllocs;
-
-		if (!EnumHasAllFlags(Purpose, IUsdPrim::GetPurpose(Prim)))
-		{
-			return false;
-		}
-
-		if (pxr::UsdGeomImageable UsdGeomImageable = pxr::UsdGeomImageable(Prim))
-		{
-			if (pxr::UsdAttribute VisibilityAttr = UsdGeomImageable.GetVisibilityAttr())
-			{
-				pxr::TfToken VisibilityToken;
-				if (VisibilityAttr.Get(&VisibilityToken) && VisibilityToken == pxr::UsdGeomTokens->invisible)
-				{
-					return false;
-				}
-			}
-		}
-		return true;
-	}
-}
+}	 // namespace UsdGeometryCacheTranslatorImpl
 
 class FGeometryCacheCreateAssetsTaskChain : public FBuildStaticMeshTaskChain
 {
@@ -746,14 +723,14 @@ void FGeometryCacheCreateAssetsTaskChain::SetupTasks()
 			{
 				// Collect the visible meshes, animated or not, under Prim to be processed into the GeometryCache
 				FScopedUsdAllocs UsdAllocs;
-				TArray<TUsdStore<pxr::UsdPrim>> ChildPrims = UsdUtils::GetAllPrimsOfType(GetPrim(), pxr::TfType::Find<pxr::UsdGeomMesh>());
+				TArray<UE::FUsdPrim> ChildVisiblePrims = UsdUtils::GetVisibleChildren(GetPrim(), Context->PurposesToLoad);
 
-				MeshPrimPaths.Reserve(ChildPrims.Num());
-				for (const TUsdStore<pxr::UsdPrim>& ChildPrim : ChildPrims)
+				MeshPrimPaths.Reserve(ChildVisiblePrims.Num());
+				for (const UE::FUsdPrim& ChildPrim : ChildVisiblePrims)
 				{
-					if (UsdGeometryCacheTranslatorImpl::IsPurposefullyVisible(ChildPrim.Get(), Context->PurposesToLoad))
+					if (ChildPrim.IsA(TEXT("Mesh")))
 					{
-						MeshPrimPaths.Add(UE::FSdfPath(ChildPrim.Get().GetPrimPath()));
+						MeshPrimPaths.Add(UE::FSdfPath(ChildPrim.GetPrimPath()));
 					}
 				}
 			}
@@ -1035,25 +1012,19 @@ TSet<UE::FSdfPath> FUsdGeometryCacheTranslator::CollectAuxiliaryPrims() const
 		return Super::CollectAuxiliaryPrims();
 	}
 
-	TSet<UE::FSdfPath> Result;
-	{
-		FScopedUsdAllocs UsdAllocs;
+	TSet<UE::FSdfPath> AuxPrims;
 
-		// Here, we collect all meshes even non-animated ones since they'll be collapse into the cache
-		TArray<TUsdStore<pxr::UsdPrim>> ChildPrims = UsdUtils::GetAllPrimsOfType(GetPrim(), pxr::TfType::Find<pxr::UsdGeomMesh>());
-
-		Result.Reserve(ChildPrims.Num());
-		for (const TUsdStore<pxr::UsdPrim>& ChildPrim : ChildPrims)
+	// Here, we collect all meshes even non-animated ones since they'll be collapse into the cache
+	TArray<UE::FUsdPrim> VisibleChildPrims = UsdUtils::GetVisibleChildren(GetPrim(), Context->PurposesToLoad);
+	AuxPrims.Reserve(VisibleChildPrims.Num());
+	for (const UE::FUsdPrim& VisibleChild : VisibleChildPrims)
+	{
+		if (VisibleChild.IsA(TEXT("Imageable")))
 		{
-			if (UsdGeometryCacheTranslatorImpl::IsPurposefullyVisible(ChildPrim.Get(), Context->PurposesToLoad))
-			{
-				Result.Add(UE::FSdfPath(ChildPrim.Get().GetPrimPath()));
-			}
-
-			// Should probably collect the Xformable as well
+			AuxPrims.Add(VisibleChild.GetPrimPath());
 		}
 	}
-	return Result;
+	return AuxPrims;
 }
 
 bool FUsdGeometryCacheTranslator::IsPotentialGeometryCacheRoot() const
diff --git a/Engine/Plugins/Importers/USDImporter/Source/USDUtilities/Private/USDConversionUtils.cpp b/Engine/Plugins/Importers/USDImporter/Source/USDUtilities/Private/USDConversionUtils.cpp
--- a/Engine/Plugins/Importers/USDImporter/Source/USDUtilities/Private/USDConversionUtils.cpp
+++ b/Engine/Plugins/Importers/USDImporter/Source/USDUtilities/Private/USDConversionUtils.cpp
@@ -1896,6 +1896,47 @@ bool UsdUtils::HasInvisibleParent( const UE::FUsdPrim& Prim, const UE::FUsdPrim&
 	return false;
 }
 
+TArray<UE::FUsdPrim> UsdUtils::GetVisibleChildren(const UE::FUsdPrim& Prim, EUsdPurpose AllowedPurposes)
+{
+	TArray<UE::FUsdPrim> VisiblePrims;
+
+#if USE_USD_SDK
+	FScopedUsdAllocs UsdAllocs;
+
+	TFunction<void(const pxr::UsdPrim& Prim)> RecursivelyCollectVisibleMeshes;
+	RecursivelyCollectVisibleMeshes = [&RecursivelyCollectVisibleMeshes, &VisiblePrims, AllowedPurposes](const pxr::UsdPrim& Prim)
+	{
+		if (!Prim || !EnumHasAllFlags(AllowedPurposes, IUsdPrim::GetPurpose(Prim)))
+		{
+			return;
+		}
+
+		if (pxr::UsdGeomImageable UsdGeomImageable = pxr::UsdGeomImageable(Prim))
+		{
+			if (pxr::UsdAttribute VisibilityAttr = UsdGeomImageable.GetVisibilityAttr())
+			{
+				pxr::TfToken VisibilityToken;
+				if (VisibilityAttr.Get(&VisibilityToken) && VisibilityToken == pxr::UsdGeomTokens->invisible)
+				{
+					// We don't propagate the (in)visibility token, we just flat out stop recursing instead
+					return;
+				}
+			}
+		}
+
+		VisiblePrims.Add(UE::FUsdPrim{Prim});
+
+		for (const pxr::UsdPrim& ChildPrim : Prim.GetFilteredChildren(pxr::UsdTraverseInstanceProxies()))
+		{
+			RecursivelyCollectVisibleMeshes(ChildPrim);
+		}
+	};
+	RecursivelyCollectVisibleMeshes(Prim);
+#endif // USE_USD_SDK
+
+	return VisiblePrims;
+}
+
 UE::FSdfPath UsdUtils::GetPrimSpecPathForLayer( const UE::FUsdPrim& Prim, const UE::FSdfLayer& Layer )
 {
 	UE::FSdfPath Result;
diff --git a/Engine/Plugins/Importers/USDImporter/Source/USDUtilities/Public/USDConversionUtils.h b/Engine/Plugins/Importers/USDImporter/Source/USDUtilities/Public/USDConversionUtils.h
--- a/Engine/Plugins/Importers/USDImporter/Source/USDUtilities/Public/USDConversionUtils.h
+++ b/Engine/Plugins/Importers/USDImporter/Source/USDUtilities/Public/USDConversionUtils.h
@@ -348,6 +348,12 @@ namespace UsdUtils
 	 */
 	USDUTILITIES_API bool HasInvisibleParent( const UE::FUsdPrim& Prim, const UE::FUsdPrim& RootPrim, double TimeCode = UsdUtils::GetDefaultTimeCode() );
 
+	/**
+	 * Returns the prims in the subtree of Prim (potentially including Prim itself) that are visible due to the
+	 * visibility attribute and prim purpose
+	 */
+	USDUTILITIES_API TArray<UE::FUsdPrim> GetVisibleChildren(const UE::FUsdPrim& Prim, EUsdPurpose AllowedPurposes);
+
 	/**
 	 * Returns a path exactly like Prim.GetPrimPath(), except that if the prim is within variant sets, it will return the
 	 * full path with variant selections in it (i.e. the spec path), like "/Root/Child{Varset=Var}Inner" instead of just
