USD: Handle translating multiple MaterialX files referenced by the same prim.

diff --git a/Engine/Plugins/Importers/USDImporter/Source/USDSchemas/Private/Custom/MaterialXUSDShadeMaterialTranslator.cpp b/Engine/Plugins/Importers/USDImporter/Source/USDSchemas/Private/Custom/MaterialXUSDShadeMaterialTranslator.cpp
--- a/Engine/Plugins/Importers/USDImporter/Source/USDSchemas/Private/Custom/MaterialXUSDShadeMaterialTranslator.cpp
+++ b/Engine/Plugins/Importers/USDImporter/Source/USDSchemas/Private/Custom/MaterialXUSDShadeMaterialTranslator.cpp
@@ -49,114 +49,12 @@
 
 FName FMaterialXUsdShadeMaterialTranslator::MaterialXRenderContext = TEXT("mtlx");
 
-void FMaterialXUsdShadeMaterialTranslator::CreateAssets()
+namespace UE::USDMaterialXTranslator::Private
 {
-#if WITH_EDITOR
-	// We handle MaterialX materials here by leveraging USD's own usdMtlx plugin. What it does is translate references
-	// to .mtlx files into hierarchies of Material/Shader/NodeGraph prims as the stage is opened. To actually use
-	// the materials contained in the .mtlx files, other prims from the stage will create material bindings to those
-	// generated Material prims. For example, a stage that wishes to use MaterialX materials may have this:
-	//
-	// def Mesh "SomeMesh"
-	// {
-	//     ...
-	//     rel material:binding = </MaterialX/Materials/TextureMaterial>
-	// }
-	//
-	// def Scope "MaterialX" (
-	//     references = [
-	//         @./textureTest.mtlx@</MaterialX>,
-	//     ]
-	// )
-	// {
-	// }
-	//
-	// Note how SomeMesh references the "TextureMaterial" material: That's the name of one of the surfacematerials
-	// inside the MaterialX file. This because when the stage is opened, the usdMtlx plugin will generate and
-	// compose a child "Materials" prim and additional child prims something like the below, which is what the external
-	// prims will be referencing:
-	//
-	// def Scope "MaterialX"
-	// {
-	//     def "Materials"
-	//     {
-	//         def Material "TextureMaterial"
-	//         {
-	//             float inputs:base = 1
-	//             color3f inputs:base_color
-	//             float inputs:coat
-	//             float inputs:coat_affect_color
-	//             ...
-	//             token outputs:mtlx:surface.connect = </MaterialX/Materials/TextureMaterial/ND_standard_surface_surfaceshader.outputs:surface>
-	//
-	//             def Shader "ND_standard_surface_surfaceshader"
-	//             {
-	//                 uniform token info:id = "ND_standard_surface_surfaceshader"
-	//                 float inputs:base.connect = </MaterialX/Materials/TextureMaterial.inputs:base>
-	//                 color3f inputs:base_color.connect = </MaterialX/Materials/TextureMaterial/NG_imagetex1.outputs:out_color_0>
-	//                 float inputs:coat.connect = </MaterialX/Materials/TextureMaterial.inputs:coat>
-	//                 float inputs:coat_affect_color.connect = </MaterialX/Materials/TextureMaterial.inputs:coat_affect_color>
-	//                 ...
-	//             }
-	//         }
-	//     }
-	// }
-	//
-	// Our goal in here is to translate those generated Materials (like "TextureMaterial" above). This because we want
-	// to generate UMaterialInterfaces via Unreal's MaterialX plugin instead, and link them to those generated Material
-	// prims within our info cache. This way, the rest or our USDImporter plugin doesn't need to care or know where
-	// this material came from: It will find an UMaterialInterface linked to that Material prim and it will happily use
-	// that as any other material.
+	TArray<FString> GetMaterialXFilePaths(const pxr::UsdPrim& Prim)
+	{
+		TArray<FString> Result;
 
-	if (Context->RenderContext != MaterialXRenderContext || !GUseInterchangeMaterialXTranslator)
-	{
-		Super::CreateAssets();
-		return;
-	}
-
-	if (!Context->AssetCache || !Context->InfoCache)
-	{
-		return;
-	}
-
-	if (Context->bTranslateOnlyUsedMaterials && !Context->InfoCache->IsMaterialUsed(PrimPath))
-	{
-		UE_LOG(LogUsd, Verbose, TEXT("Skipping creating assets for material prim '%s' as it is not currently bound by any prim."), *PrimPath.GetString());
-		return;
-	}
-
-	FScopedUsdAllocs UsdAllocs;
-
-	pxr::UsdPrim Prim = GetPrim();
-	pxr::UsdShadeMaterial ShadeMaterial(Prim);
-	if (!ShadeMaterial)
-	{
-		return;
-	}
-
-	// We check for the mtlx surface output directly, because ComputeSurfaceSource will return a valid SurfaceShader
-	// in case the material just has a regular universal render context output.
-	// This is just for checking though: We will defer back to USD to let it ComputeSurfaceSource with whatever logic it has
-	const static pxr::TfToken RenderContextToken = UnrealToUsd::ConvertToken(*MaterialXRenderContext.ToString()).Get();
-	pxr::UsdShadeOutput MtlxOutput = ShadeMaterial.GetSurfaceOutput(RenderContextToken);
-	if (!MtlxOutput)
-	{
-		Super::CreateAssets();
-		return;
-	}
-	pxr::UsdShadeShader SurfaceShader = ShadeMaterial.ComputeSurfaceSource(RenderContextToken);
-	if (!SurfaceShader)
-	{
-		// This really shouldn't ever happen if we have an actual 'mtlx' output, but anyway
-		Super::CreateAssets();
-		return;
-	}
-
-	// This material prim has the mtlx render context, so maybe it is one of the ones generated by usdMtlx.
-	// Let's traverse upwards and try finding a .mtlx file reference in one of our parents.
-
-	TFunction<bool(const pxr::UsdPrim&, FString&)> GetMaterialXFilePath = [](const pxr::UsdPrim& Prim, FString& OutMaterialXFilePath)
-	{
 		pxr::UsdPrimCompositionQuery PrimCompositionQuery = pxr::UsdPrimCompositionQuery::GetDirectReferences(Prim);
 		for (const pxr::UsdPrimCompositionQueryArc& CompositionArc : PrimCompositionQuery.GetCompositionArcs())
 		{
@@ -173,91 +71,57 @@
 					FString AbsoluteFilePath = UsdToUnreal::ConvertString(TargetLayer->GetRealPath());
 
 					FString Extension = FPaths::GetExtension(AbsoluteFilePath);
-					if (Extension == MaterialXRenderContext && FPaths::FileExists(AbsoluteFilePath))
+					if (Extension == FMaterialXUsdShadeMaterialTranslator::MaterialXRenderContext && FPaths::FileExists(AbsoluteFilePath))
 					{
-						OutMaterialXFilePath = AbsoluteFilePath;
-						return true;
+						Result.Add(AbsoluteFilePath);
 					}
 				}
 			}
 		}
 
-		return false;
+		return Result;
 	};
 
-	FString MaterialXFilePath;
-	pxr::UsdPrim MaterialXReferencerPrim = Prim;
-	if (!GetMaterialXFilePath(Prim, MaterialXFilePath))
-	{
-		// We know the usdMtlx plugin always generates a "Materials" schemaless prim to contain all the generated
-		// Materials, so let's use that too.
-		pxr::UsdPrim ParentPrim = Prim.GetParent();
-		while (ParentPrim && ParentPrim.GetName() != "Materials")
-		{
-			ParentPrim = ParentPrim.GetParent();
-		}
-
-		if (ParentPrim)
-		{
-			// This prim likely holds the reference to the MaterialX file, but let's search upwards too
-			pxr::UsdPrim MaterialXReferencerCandidate = ParentPrim.GetParent();
-			while (MaterialXReferencerCandidate)
-			{
-				if (GetMaterialXFilePath(MaterialXReferencerCandidate, MaterialXFilePath))
-				{
-					break;
-				}
-
-				MaterialXReferencerCandidate = MaterialXReferencerCandidate.GetParent();
-			}
-			MaterialXReferencerPrim = MaterialXReferencerCandidate;
-		}
-	}
-
-	if (MaterialXFilePath.IsEmpty() || !MaterialXReferencerPrim)
-	{
-		UE_LOG(LogUsd, Warning, TEXT("Recognized potential MaterialX materials on prim '%s', but failed to find a valid referenced MaterialX file. Reverting to parsing the generated Material prims instead."), *PrimPath.GetString());
-		Super::CreateAssets();
-		return;
-	}
-
 	// Quick shallow parse of MaterialX documents ourselves so that we can find all the referenced MaterialX files and
 	// build a robust hash to use for caching our generated assets
-	FSHA1 SHA1;
+	FString HashMaterialXFile(const FString& MaterialXFilePath)
 	{
+		if (!FPaths::FileExists(MaterialXFilePath))
+		{
+			return {};
+		}
+
+		FSHA1 SHA1;
+
 		TSet<FString> ReferencedMaterialXFiles;
 		TFunction<void(const FString&)> CollectReferencedMaterialXFilesRecursive;
 		CollectReferencedMaterialXFilesRecursive =
-			[&ReferencedMaterialXFiles, &CollectReferencedMaterialXFilesRecursive]
-			(const FString& ActiveDocumentFullPath)
+			[&ReferencedMaterialXFiles, &CollectReferencedMaterialXFilesRecursive](const FString& ActiveDocumentFullPath)
+		{
+			namespace mx = MaterialX;
+
+			mx::DocumentPtr Document = mx::createDocument();
+			mx::readFromXmlFile(Document, TCHAR_TO_UTF8(*ActiveDocumentFullPath));
+
+			for (const std::string& ReferencedURI : Document->getReferencedSourceUris())
 			{
-				namespace mx = MaterialX;
+				FString UEReferencedURI = UTF8_TO_TCHAR(ReferencedURI.c_str());
 
-				mx::DocumentPtr Document = mx::createDocument();
-				mx::readFromXmlFile(Document, TCHAR_TO_UTF8(*ActiveDocumentFullPath));
+				// This could be relative to the referencer .mtlx file
+				UEReferencedURI = FPaths::ConvertRelativePathToFull(FPaths::GetPath(ActiveDocumentFullPath), UEReferencedURI);
 
-				for (const std::string& ReferencedURI : Document->getReferencedSourceUris())
+				if (ReferencedMaterialXFiles.Contains(UEReferencedURI))
 				{
-					FString UEReferencedURI = UTF8_TO_TCHAR(ReferencedURI.c_str());
+					continue;
+				}
 
-					// This could be relative to the referencer .mtlx file
-					UEReferencedURI = FPaths::ConvertRelativePathToFull(
-						FPaths::GetPath(ActiveDocumentFullPath),
-						UEReferencedURI
-					);
+				if (FPaths::FileExists(UEReferencedURI))
+				{
+					ReferencedMaterialXFiles.Add(UEReferencedURI);
+					CollectReferencedMaterialXFilesRecursive(UEReferencedURI);
+				}
+			}
+		};
 
-					if (ReferencedMaterialXFiles.Contains(UEReferencedURI))
-					{
-						continue;
-					}
-
-					if (FPaths::FileExists(UEReferencedURI))
-					{
-						ReferencedMaterialXFiles.Add(UEReferencedURI);
-						CollectReferencedMaterialXFilesRecursive(UEReferencedURI);
-					}
-				}
-			};
-
 		// Note we don't have to add MaterialXFilePath ourselves because getReferencedSourceUris
 		// always includes at least the active document itself anyway
@@ -284,18 +148,19 @@
 				}
 			}
 		}
+
+		FSHAHash Hash;
+		SHA1.Final();
+		SHA1.GetHash(&Hash.Hash[0]);
+		return Hash.ToString();
 	}
 
-	FSHAHash Hash;
-	SHA1.Final();
-	SHA1.GetHash(&Hash.Hash[0]);
-	FString MaterialXHash = Hash.ToString();
-
-	FString TargetMaterialAssetHash = MaterialXHash + TEXT("/") + UsdToUnreal::ConvertString(Prim.GetName());
-	UMaterialInterface* ParsedMaterial = Cast< UMaterialInterface >(Context->AssetCache->GetCachedAsset(TargetMaterialAssetHash));
-	if (!ParsedMaterial)
+	bool TranslateMaterialXFile(const FString& MaterialXFilePath, const FString& FileHash, const pxr::UsdPrim& MaterialXReferencerPrim, UUsdAssetCache2& AssetCache)
 	{
-		// If we don't have this asset in the cache yet, we need to parse the MaterialX file
+		if (!FPaths::FileExists(MaterialXFilePath))
+		{
+			return false;
+		}
 
 		UInterchangeManager& Manager = UInterchangeManager::GetInterchangeManager();
 		UInterchangeSourceData* SourceData = Manager.CreateSourceData(MaterialXFilePath);
@@ -319,9 +184,10 @@
 		// Currently Interchange cannot import directly into the transient package, so we'll import into a temporary folder and then
 		// move everything over
 		InterchangeParameters.OnAssetsImportDoneNative.BindLambda(
-			[this, &MaterialXHash, &TargetMaterialAssetHash, &ParsedMaterial, &ReferencerPrimPath, &TempPackagePath]
-			(const TArray<UObject*>& ImportedObjects)
+			[&AssetCache, &FileHash, &ReferencerPrimPath, &TempPackagePath, &MaterialXFilePath](const TArray<UObject*>& ImportedObjects)
 			{
+				UE_LOG(LogUsd, Verbose, TEXT("Translated '%d' assets from MaterialX file '%s'"), ImportedObjects.Num(), *MaterialXFilePath);
+
 				// Move all our assets to the transient package
 				for (UObject* ImportedObject : ImportedObjects)
 				{
@@ -350,26 +216,20 @@
 					{
 						// MaterialX names are unique, and can only have alphanumeric and the "_" character, so we should
 						// always have a solid enough mapping to assume UAsset name == Prim name == MaterialX name
-						FString NewMaterialHash = MaterialXHash + TEXT("/") + Material->GetFName().GetPlainNameString();
+						FString NewMaterialHash = FileHash + TEXT("/") + Material->GetFName().GetPlainNameString();
 
 						// We generate all assets from the MaterialX file once, but it's possible we're just updating a
 						// single material prim here. If we were to cache all assets here, we'd potentially be trying
 						// to overwrite the existing assets that are being used by other prims and wouldn't otherwise
 						// be discarded, so make sure we don't do that
-						UMaterialInterface* ExistingMaterial = Cast<UMaterialInterface>(Context->AssetCache->GetCachedAsset(NewMaterialHash));
+						UMaterialInterface* ExistingMaterial = Cast<UMaterialInterface>(AssetCache.GetCachedAsset(NewMaterialHash));
 						if (!ExistingMaterial)
 						{
-							Context->AssetCache->CacheAsset(NewMaterialHash, Material);
+							AssetCache.CacheAsset(NewMaterialHash, Material);
 						}
-
-						if (NewMaterialHash == TargetMaterialAssetHash)
-						{
-							ParsedMaterial = Material;
-						}
 					}
 					else if (UTexture* Texture = Cast<UTexture>(ImportedObject))
 					{
-						// Note: This filename will be the .mtlx file for now, which is why we append the texture name manually later
 						const FString FilePath = Texture->AssetImportData ? Texture->AssetImportData->GetFirstFilename() : TEXT("");
 
 						FString TextureHash = UsdUtils::GetTextureHash(
@@ -381,14 +241,14 @@
 						);
 
 						// See comment on the analogous part of the Material case above
-						UTexture* ExistingTexture = Cast<UTexture>(Context->AssetCache->GetCachedAsset(TextureHash));
+						UTexture* ExistingTexture = Cast<UTexture>(AssetCache.GetCachedAsset(TextureHash));
 						if (!ExistingTexture)
 						{
 							UUsdAssetUserData* UserData = NewObject<UUsdAssetUserData>(Texture, TEXT("USDAssetUserData"));
 							UserData->PrimPath = ReferencerPrimPath;
 							Texture->AddAssetUserData(UserData);
 
-							Context->AssetCache->CacheAsset(TextureHash, Texture);
+							AssetCache.CacheAsset(TextureHash, Texture);
 						}
 					}
 					else
@@ -406,31 +266,243 @@
 			}
 		);
 
-		const bool bSuccess = Manager.ImportAsset(TempPackagePath, SourceData, InterchangeParameters);
-		if (!bSuccess)
+		return Manager.ImportAsset(TempPackagePath, SourceData, InterchangeParameters);
+	}
+}
+
+void FMaterialXUsdShadeMaterialTranslator::CreateAssets()
+{
+#if WITH_EDITOR
+	// We handle MaterialX materials here by leveraging USD's own usdMtlx plugin. What it does is translate references
+	// to .mtlx files into hierarchies of Material/Shader/NodeGraph prims as the stage is opened. To actually use
+	// the materials contained in the .mtlx files, other prims from the stage will create material bindings to those
+	// generated Material prims. For example, a stage that wishes to use MaterialX materials may have this:
+	//
+	// def Mesh "SomeMesh"
+	// {
+	//     ...
+	//     rel material:binding = </MaterialX/Materials/TextureMaterial>
+	// }
+	//
+	// def Scope "MaterialX" (
+	//     references = [
+	//         @./textureTest.mtlx@</MaterialX>,
+	//     ]
+	// )
+	// {
+	// }
+	//
+	// Note how SomeMesh references the "TextureMaterial" material: That's the name of one of the surfacematerials
+	// inside the MaterialX file. This because when the stage is opened, the usdMtlx plugin will generate and
+	// compose a child "Materials" prim and additional child prims something like the below, which is what the external
+	// prims will be referencing:
+	//
+	// def Scope "MaterialX"
+	// {
+	//     def "Materials"
+	//     {
+	//         def Material "TextureMaterial"
+	//         {
+	//             float inputs:base = 1
+	//             color3f inputs:base_color
+	//             float inputs:coat
+	//             float inputs:coat_affect_color
+	//             ...
+	//             token outputs:mtlx:surface.connect = </MaterialX/Materials/TextureMaterial/ND_standard_surface_surfaceshader.outputs:surface>
+	//
+	//             def Shader "ND_standard_surface_surfaceshader"
+	//             {
+	//                 uniform token info:id = "ND_standard_surface_surfaceshader"
+	//                 float inputs:base.connect = </MaterialX/Materials/TextureMaterial.inputs:base>
+	//                 color3f inputs:base_color.connect = </MaterialX/Materials/TextureMaterial/NG_imagetex1.outputs:out_color_0>
+	//                 float inputs:coat.connect = </MaterialX/Materials/TextureMaterial.inputs:coat>
+	//                 float inputs:coat_affect_color.connect = </MaterialX/Materials/TextureMaterial.inputs:coat_affect_color>
+	//                 ...
+	//             }
+	//         }
+	//     }
+	// }
+	//
+	// Our goal in here is to translate those generated Materials (like "TextureMaterial" above). This because we want
+	// to generate UMaterialInterfaces via Unreal's MaterialX plugin instead, and link them to those generated Material
+	// prims within our info cache. This way, the rest or our USDImporter plugin doesn't need to care or know where
+	// this material came from: It will find an UMaterialInterface linked to that Material prim and it will happily use
+	// that as any other material.
+
+	if (Context->RenderContext != MaterialXRenderContext || !GUseInterchangeMaterialXTranslator)
+	{
+		Super::CreateAssets();
+		return;
+	}
+
+	if (!Context->AssetCache || !Context->InfoCache)
+	{
+		return;
+	}
+
+	if (Context->bTranslateOnlyUsedMaterials && !Context->InfoCache->IsMaterialUsed(PrimPath))
+	{
+		UE_LOG(LogUsd, Verbose, TEXT("Skipping creating assets for material prim '%s' as it is not currently bound by any prim."), *PrimPath.GetString());
+		return;
+	}
+
+	FScopedUsdAllocs UsdAllocs;
+
+	pxr::UsdPrim Prim = GetPrim();
+	pxr::UsdShadeMaterial ShadeMaterial(Prim);
+	if (!ShadeMaterial)
+	{
+		return;
+	}
+
+	// We check for the mtlx surface output directly, because ComputeSurfaceSource will return a valid SurfaceShader
+	// in case the material just has a regular universal render context output.
+	// This is just for checking though: We will defer back to USD to let it ComputeSurfaceSource with whatever logic it has
+	const static pxr::TfToken RenderContextToken = UnrealToUsd::ConvertToken(*MaterialXRenderContext.ToString()).Get();
+	pxr::UsdShadeOutput MtlxOutput = ShadeMaterial.GetSurfaceOutput(RenderContextToken);
+	if (!MtlxOutput)
+	{
+		Super::CreateAssets();
+		return;
+	}
+	pxr::UsdShadeShader SurfaceShader = ShadeMaterial.ComputeSurfaceSource(RenderContextToken);
+	if (!SurfaceShader)
+	{
+		// This really shouldn't ever happen if we have an actual 'mtlx' output, but anyway
+		Super::CreateAssets();
+		return;
+	}
+
+	// This material prim has the mtlx render context, so maybe it is one of the ones generated by usdMtlx.
+	// Let's traverse upwards and try finding a .mtlx file reference in one of our parents.
+
+	TArray<FString> MaterialXFilePaths = UE::USDMaterialXTranslator::Private::GetMaterialXFilePaths(Prim);
+	pxr::UsdPrim MaterialXReferencerPrim = Prim;
+	if (MaterialXFilePaths.Num() == 0)
+	{
+		// We know the usdMtlx plugin always generates a "Materials" schemaless prim to contain all the generated
+		// Materials, so let's use that too.
+		pxr::UsdPrim ParentPrim = Prim.GetParent();
+		while (ParentPrim && ParentPrim.GetName() != "Materials")
 		{
-			UE_LOG(LogUsd, Warning, TEXT("Recognized potential MaterialX materials on prim '%s', but MaterialX parsing of file '%s' failed. Reverting to parsing the generated Material prims instead."),
-				*PrimPath.GetString(),
-				*MaterialXFilePath
+			ParentPrim = ParentPrim.GetParent();
+		}
+
+		if (ParentPrim)
+		{
+			// This prim likely holds the reference to the MaterialX file, but let's search upwards too
+			pxr::UsdPrim MaterialXReferencerCandidate = ParentPrim.GetParent();
+			while (MaterialXReferencerCandidate)
+			{
+				MaterialXFilePaths = UE::USDMaterialXTranslator::Private::GetMaterialXFilePaths(MaterialXReferencerCandidate);
+				if (MaterialXFilePaths.Num() > 0)
+				{
+					break;
+				}
+
+				MaterialXReferencerCandidate = MaterialXReferencerCandidate.GetParent();
+			}
+			MaterialXReferencerPrim = MaterialXReferencerCandidate;
+		}
+	}
+
+	if (MaterialXFilePaths.IsEmpty() || !MaterialXReferencerPrim)
+	{
+		UE_LOG(LogUsd, Warning, TEXT("Recognized potential MaterialX materials on prim '%s', but failed to find a valid referenced MaterialX file. Reverting to parsing the generated Material prims instead."), *PrimPath.GetString());
+		Super::CreateAssets();
+		return;
+	}
+
+	FString TargetHashSuffix = TEXT("/") + UsdToUnreal::ConvertString(Prim.GetName());
+
+	FString FoundMaterialAssetHash;
+	UMaterialInterface* ParsedMaterial = nullptr;
+
+	// Try to find the parsed material already in the asset cache assuming it came from any one of the MaterialX file paths
+	TArray<FString> MaterialXFileHashes;
+	MaterialXFileHashes.Reserve(MaterialXFilePaths.Num());
+	for (const FString& MaterialXFilePath : MaterialXFilePaths)
+	{
+		FString& MaterialXHash = MaterialXFileHashes.Emplace_GetRef(UE::USDMaterialXTranslator::Private::HashMaterialXFile(MaterialXFilePath));
+		if (MaterialXHash.IsEmpty())
+		{
+			continue;
+		}
+
+		FString MaterialAssetHashForThisFile = MaterialXHash + TargetHashSuffix;
+		ParsedMaterial = Cast<UMaterialInterface>(Context->AssetCache->GetCachedAsset(MaterialAssetHashForThisFile));
+		if (ParsedMaterial)
+		{
+			FoundMaterialAssetHash = MaterialAssetHashForThisFile;
+			break;
+		}
+	}
+
+	if (!ParsedMaterial)
+	{
+		// Translate all MaterialX files
+		ensure(MaterialXFilePaths.Num() == MaterialXFileHashes.Num());
+		for (int32 Index = 0; Index < MaterialXFilePaths.Num(); ++Index)
+		{
+			const FString& MaterialXFilePath = MaterialXFilePaths[Index];
+			const FString& MaterialXFileHash = MaterialXFileHashes[Index];
+
+			const bool bSuccess = UE::USDMaterialXTranslator::Private::TranslateMaterialXFile(
+				MaterialXFilePath,
+				MaterialXFileHash,
+				MaterialXReferencerPrim,
+				*Context->AssetCache
 			);
-			Super::CreateAssets();
-			return;
+
+			if (!bSuccess)
+			{
+				UE_LOG(
+					LogUsd,
+					Warning,
+					TEXT("Recognized potential MaterialX materials on prim '%s', but MaterialX parsing of file '%s' failed."),
+					*PrimPath.GetString(),
+					*MaterialXFilePath
+				);
+			}
+
+			// Check if we found our target material when parsing this file
+			if (!ParsedMaterial)
+			{
+				FString MaterialAssetHashForThisFile = MaterialXFileHash + TargetHashSuffix;
+				ParsedMaterial = Cast<UMaterialInterface>(Context->AssetCache->GetCachedAsset(MaterialAssetHashForThisFile));
+				if (ParsedMaterial)
+				{
+					FoundMaterialAssetHash = MaterialAssetHashForThisFile;
+
+					// Notice we don't break here: Let's parse all referenced MaterialX files anyway
+				}
+			}
 		}
 	}
 
 	if (ParsedMaterial)
 	{
-		PostImportMaterial(TargetMaterialAssetHash, ParsedMaterial);
+		ensure(!FoundMaterialAssetHash.IsEmpty());
+		PostImportMaterial(FoundMaterialAssetHash, ParsedMaterial);
 	}
 	else
 	{
-		UE_LOG(LogUsd, Warning, TEXT("Parsing of MaterialX file '%s' on prim '%s' succeeded, but we failed to find the target material '%s' within the file."),
-			*MaterialXFilePath,
+		FString FilePaths;
+		for (const FString& MaterialXFilePath : MaterialXFilePaths)
+		{
+			FilePaths += MaterialXFilePath + TEXT(", ");
+		}
+		FilePaths.RemoveFromEnd(TEXT(", "));
+
+		UE_LOG(
+			LogUsd,
+			Warning,
+			TEXT("Failed to find target Material '%s' after parsing all Material X files [%s]. Reverting back to parsing the USD Material prim "
+				 "generated by usdMtlx directly."),
 			*PrimPath.GetString(),
-			*UsdToUnreal::ConvertString(Prim.GetName())
+			*FilePaths
 		);
 		Super::CreateAssets();
-		return;
 	}
 #else
 	Super::CreateAssets();
