USD: Register auxiliary prims when building the info cache, as we may need that info to know which translator to call when an update arrives;

Refactor handling of USD notices to handle many issues;
Fix potential issue where the order of traversal of the MainPrimPaths set when updating prims could cause extra components to be generated;
Fix issue where changing prim kinds wouldn't cause components to recollapse/uncollapse.

diff --git a/Engine/Plugins/Importers/USDImporter/Source/USDSchemas/Private/Custom/MDLUSDShadeMaterialTranslator.cpp b/Engine/Plugins/Importers/USDImporter/Source/USDSchemas/Private/Custom/MDLUSDShadeMaterialTranslator.cpp
--- a/Engine/Plugins/Importers/USDImporter/Source/USDSchemas/Private/Custom/MDLUSDShadeMaterialTranslator.cpp
+++ b/Engine/Plugins/Importers/USDImporter/Source/USDSchemas/Private/Custom/MDLUSDShadeMaterialTranslator.cpp
@@ -102,9 +102,6 @@ void FMdlUsdShadeMaterialTranslator::CreateAssets()
 		return;
 	}
 
-	// Note that we rely on FUsdShadeMaterialTranslator's implementation of CollectAuxiliaryPrims
-	RegisterAuxiliaryPrims();
-
 	const pxr::TfToken MdlToken = UnrealToUsd::ConvertToken( *MdlRenderContext.ToString() ).Get();
 
 	pxr::UsdShadeShader SurfaceShader = ShadeMaterial.ComputeSurfaceSource( MdlToken );
diff --git a/Engine/Plugins/Importers/USDImporter/Source/USDSchemas/Private/USDGeomMeshTranslator.cpp b/Engine/Plugins/Importers/USDImporter/Source/USDSchemas/Private/USDGeomMeshTranslator.cpp
--- a/Engine/Plugins/Importers/USDImporter/Source/USDSchemas/Private/USDGeomMeshTranslator.cpp
+++ b/Engine/Plugins/Importers/USDImporter/Source/USDSchemas/Private/USDGeomMeshTranslator.cpp
@@ -1004,12 +1004,10 @@ void FUsdGeomMeshTranslator::CreateAssets()
 		return Super::CreateAssets();
 	}
 
-	RegisterAuxiliaryPrims();
+	TSharedRef< FGeomMeshCreateAssetsTaskChain > AssetsTaskChain = MakeShared< FGeomMeshCreateAssetsTaskChain >(Context, PrimPath);
 
-		TSharedRef< FGeomMeshCreateAssetsTaskChain > AssetsTaskChain = MakeShared< FGeomMeshCreateAssetsTaskChain >(Context, PrimPath);
-
-		Context->TranslatorTasks.Add(MoveTemp(AssetsTaskChain));
-	}
+	Context->TranslatorTasks.Add(MoveTemp(AssetsTaskChain));
+}
 
 USceneComponent* FUsdGeomMeshTranslator::CreateComponents()
 {
@@ -1126,6 +1124,11 @@ TSet<UE::FSdfPath> FUsdGeomMeshTranslator::CollectAuxiliaryPrims() const
 		return Super::CollectAuxiliaryPrims();
 	}
 
+	if (!Context->bIsBuildingInfoCache)
+	{
+		return Context->InfoCache->GetAuxiliaryPrims(PrimPath);
+	}
+
 	TSet<UE::FSdfPath> Result;
 	{
 		FScopedUsdAllocs UsdAllocs;
diff --git a/Engine/Plugins/Importers/USDImporter/Source/USDSchemas/Private/USDGeomPointInstancerTranslator.cpp b/Engine/Plugins/Importers/USDImporter/Source/USDSchemas/Private/USDGeomPointInstancerTranslator.cpp
--- a/Engine/Plugins/Importers/USDImporter/Source/USDSchemas/Private/USDGeomPointInstancerTranslator.cpp
+++ b/Engine/Plugins/Importers/USDImporter/Source/USDSchemas/Private/USDGeomPointInstancerTranslator.cpp
@@ -152,8 +152,6 @@ void FUsdGeomPointInstancerTranslator::CreateAssets()
 {
 	TRACE_CPUPROFILER_EVENT_SCOPE( FUsdGeomPointInstancerTranslator::CreateAssets );
 
-	RegisterAuxiliaryPrims();
-
 	FScopedUsdAllocs Allocs;
 
 	pxr::UsdPrim Prim = GetPrim();
@@ -523,6 +521,11 @@ bool FUsdGeomPointInstancerTranslator::CanBeCollapsed( ECollapsingType Collapsin
 
 TSet<UE::FSdfPath> FUsdGeomPointInstancerTranslator::CollectAuxiliaryPrims() const
 {
+	if (!Context->bIsBuildingInfoCache)
+	{
+		return Context->InfoCache->GetAuxiliaryPrims(PrimPath);
+	}
+
 	FScopedUsdAllocs UsdAllocs;
 
 	pxr::UsdPrim Prim = GetPrim();
diff --git a/Engine/Plugins/Importers/USDImporter/Source/USDSchemas/Private/USDGeomXformableTranslator.cpp b/Engine/Plugins/Importers/USDImporter/Source/USDSchemas/Private/USDGeomXformableTranslator.cpp
--- a/Engine/Plugins/Importers/USDImporter/Source/USDSchemas/Private/USDGeomXformableTranslator.cpp
+++ b/Engine/Plugins/Importers/USDImporter/Source/USDSchemas/Private/USDGeomXformableTranslator.cpp
@@ -268,8 +268,6 @@ void FUsdGeomXformableTranslator::CreateAssets()
 		return;
 	}
 
-	RegisterAuxiliaryPrims();
-
 	Context->TranslatorTasks.Add( MakeShared< FUsdGeomXformableCreateAssetsTaskChain >( Context, PrimPath ) );
 }
 
@@ -672,9 +670,9 @@ bool FUsdGeomXformableTranslator::CollapsesChildren( ECollapsingType CollapsingT
 {
 	TRACE_CPUPROFILER_EVENT_SCOPE( FUsdGeomXformableTranslator::CollapsesChildren );
 
-	if ( Context->InfoCache.IsValid() )
+	if (!Context->bIsBuildingInfoCache)
 	{
-		return Context->InfoCache->DoesPathCollapseChildren( PrimPath, CollapsingType );
+		return Context->InfoCache->DoesPathCollapseChildren(PrimPath, CollapsingType);
 	}
 
 	bool bCollapsesChildren = false;
@@ -754,9 +752,9 @@ bool FUsdGeomXformableTranslator::CanBeCollapsed( ECollapsingType CollapsingType
 
 TSet<UE::FSdfPath> FUsdGeomXformableTranslator::CollectAuxiliaryPrims() const
 {
-	if (!Context->InfoCache.IsValid())
+	if (!Context->bIsBuildingInfoCache)
 	{
-		return {};
+		return Context->InfoCache->GetAuxiliaryPrims(PrimPath);
 	}
 
 	if (!Context->InfoCache->DoesPathCollapseChildren(PrimPath, ECollapsingType::Assets))
diff --git a/Engine/Plugins/Importers/USDImporter/Source/USDSchemas/Private/USDGeometryCacheTranslator.cpp b/Engine/Plugins/Importers/USDImporter/Source/USDSchemas/Private/USDGeometryCacheTranslator.cpp
--- a/Engine/Plugins/Importers/USDImporter/Source/USDSchemas/Private/USDGeometryCacheTranslator.cpp
+++ b/Engine/Plugins/Importers/USDImporter/Source/USDSchemas/Private/USDGeometryCacheTranslator.cpp
@@ -826,8 +826,6 @@ void FUsdGeometryCacheTranslator::CreateAssets()
 		return;
 	}
 
-	RegisterAuxiliaryPrims();
-
 	// Create the GeometryCache TaskChain
 	TSharedRef<FGeometryCacheCreateAssetsTaskChain> AssetsTaskChain = MakeShared<FGeometryCacheCreateAssetsTaskChain>(Context, PrimPath);
 
@@ -989,19 +987,11 @@ void FUsdGeometryCacheTranslator::UpdateComponents(USceneComponent* SceneCompone
 
 bool FUsdGeometryCacheTranslator::CollapsesChildren(ECollapsingType CollapsingType) const
 {
-	if (!Context->InfoCache.IsValid())
-	{
-		return Super::CollapsesChildren(CollapsingType);
-	}
 	return IsPotentialGeometryCacheRoot() ? true : Super::CollapsesChildren(CollapsingType);
 }
 
 bool FUsdGeometryCacheTranslator::CanBeCollapsed(ECollapsingType CollapsingType) const
 {
-	if (!Context->InfoCache.IsValid())
-	{
-		return Super::CanBeCollapsed(CollapsingType);
-	}
 	return IsPotentialGeometryCacheRoot() ? false : Super::CanBeCollapsed(CollapsingType);
 }
 
@@ -1012,6 +1002,11 @@ TSet<UE::FSdfPath> FUsdGeometryCacheTranslator::CollectAuxiliaryPrims() const
 		return Super::CollectAuxiliaryPrims();
 	}
 
+	if (!Context->bIsBuildingInfoCache)
+	{
+		return Context->InfoCache->GetAuxiliaryPrims(PrimPath);
+	}
+
 	TSet<UE::FSdfPath> AuxPrims;
 
 	// Here, we collect all meshes even non-animated ones since they'll be collapse into the cache
@@ -1029,8 +1024,6 @@ TSet<UE::FSdfPath> FUsdGeometryCacheTranslator::CollectAuxiliaryPrims() const
 
 bool FUsdGeometryCacheTranslator::IsPotentialGeometryCacheRoot() const
 {
-	check(Context->InfoCache.IsValid());
-
 	// The logic to check for GeometryCache is completely in the InfoCache
 	return Context->InfoCache->IsPotentialGeometryCacheRoot(PrimPath);
 }
diff --git a/Engine/Plugins/Importers/USDImporter/Source/USDSchemas/Private/USDGroomTranslator.cpp b/Engine/Plugins/Importers/USDImporter/Source/USDSchemas/Private/USDGroomTranslator.cpp
--- a/Engine/Plugins/Importers/USDImporter/Source/USDSchemas/Private/USDGroomTranslator.cpp
+++ b/Engine/Plugins/Importers/USDImporter/Source/USDSchemas/Private/USDGroomTranslator.cpp
@@ -365,8 +365,6 @@ void FUsdGroomTranslator::CreateAssets()
 		return Super::CreateAssets();
 	}
 
-	RegisterAuxiliaryPrims();
-
 	Context->TranslatorTasks.Add(MakeShared<FUsdGroomCreateAssetsTaskChain>(Context, PrimPath));
 }
 
@@ -471,9 +469,9 @@ TSet<UE::FSdfPath> FUsdGroomTranslator::CollectAuxiliaryPrims() const
 		return Super::CollectAuxiliaryPrims();
 	}
 
-	if (!Context->InfoCache.IsValid())
+	if (!Context->bIsBuildingInfoCache)
 	{
-		return {};
+		return Context->InfoCache->GetAuxiliaryPrims(PrimPath);
 	}
 
 	if (!Context->InfoCache->DoesPathCollapseChildren(PrimPath, ECollapsingType::Assets))
diff --git a/Engine/Plugins/Importers/USDImporter/Source/USDSchemas/Private/USDInfoCache.cpp b/Engine/Plugins/Importers/USDImporter/Source/USDSchemas/Private/USDInfoCache.cpp
--- a/Engine/Plugins/Importers/USDImporter/Source/USDSchemas/Private/USDInfoCache.cpp
+++ b/Engine/Plugins/Importers/USDImporter/Source/USDSchemas/Private/USDInfoCache.cpp
@@ -97,6 +97,24 @@ struct FUsdInfoCache::FUsdInfoCacheImpl
 		AllowedExtensionsForGeometryCacheSource.Add(TEXT("abc"));
 	}
 
+	FUsdInfoCacheImpl(const FUsdInfoCacheImpl& Other)
+		: FUsdInfoCacheImpl()
+	{
+		// Probably not necessary
+		FReadScopeLock ScopedInfoMapLock(Other.InfoMapLock);
+		FReadScopeLock ScopedPrimPathToAssetsLock(Other.PrimPathToAssetsLock);
+		FReadScopeLock ScopedMaterialUsersLock(Other.MaterialUsersLock);
+		FReadScopeLock ScopedAuxPrimsLock(Other.AuxiliaryPrimsLock);
+
+		InfoMap = Other.InfoMap;
+		PrimPathToAssets = Other.PrimPathToAssets;
+		AssetToPrimPaths = Other.AssetToPrimPaths;
+		MaterialUsers = Other.MaterialUsers;
+		AuxToMainPrims = Other.AuxToMainPrims;
+		MainToAuxPrims = Other.MainToAuxPrims;
+		AllowedExtensionsForGeometryCacheSource = Other.AllowedExtensionsForGeometryCacheSource;
+	}
+
 	// Information we must have about all prims on the stage
 	TMap< UE::FSdfPath, UE::UsdInfoCache::Private::FUsdPrimInfo > InfoMap;
 	mutable FRWLock InfoMapLock;
@@ -119,11 +137,39 @@ struct FUsdInfoCache::FUsdInfoCacheImpl
 
 	// Geometry cache can come from a reference or payload of these file types
 	TArray<FString> AllowedExtensionsForGeometryCacheSource;
+
+public:
+	void RegisterAuxiliaryPrims(const UE::FSdfPath& MainPrimPath, const TSet<UE::FSdfPath>& AuxPrimPaths)
+	{
+		if (AuxPrimPaths.Num() == 0)
+		{
+			return;
+		}
+
+		FWriteScopeLock ScopeLock(AuxiliaryPrimsLock);
+
+		MainToAuxPrims.FindOrAdd(MainPrimPath).Append(AuxPrimPaths);
+
+		for (const UE::FSdfPath& AuxPrimPath : AuxPrimPaths)
+		{
+			UE_LOG(LogUsd, Verbose, TEXT("Registering main prim '%s' and aux prim '%s'"),
+				*MainPrimPath.GetString(),
+				*AuxPrimPath.GetString()
+			);
+
+			AuxToMainPrims.FindOrAdd(AuxPrimPath).Add(MainPrimPath);
+		}
+	}
 };
 
 FUsdInfoCache::FUsdInfoCache()
 {
-	Impl = MakeUnique< FUsdInfoCache::FUsdInfoCacheImpl>();
+	Impl = MakeUnique<FUsdInfoCache::FUsdInfoCacheImpl>();
+}
+
+FUsdInfoCache::FUsdInfoCache(const FUsdInfoCache& Other)
+{
+	Impl = MakeUnique<FUsdInfoCache::FUsdInfoCacheImpl>(*Other.Impl);
 }
 
 FUsdInfoCache::~FUsdInfoCache()
@@ -168,6 +214,21 @@ bool FUsdInfoCache::ContainsInfoAboutPrim( const UE::FSdfPath& Path ) const
 	return false;
 }
 
+TSet<UE::FSdfPath> FUsdInfoCache::GetKnownPrims() const
+{
+	TSet<UE::FSdfPath> Result;
+
+	if (FUsdInfoCacheImpl* ImplPtr = Impl.Get())
+	{
+		FReadScopeLock ScopeLock(ImplPtr->InfoMapLock);
+
+		ImplPtr->InfoMap.GetKeys(Result);
+		return Result;
+	}
+
+	return Result;
+}
+
 bool FUsdInfoCache::IsPathCollapsed( const UE::FSdfPath& Path, ECollapsingType CollapsingType ) const
 {
 	if ( FUsdInfoCacheImpl* ImplPtr = Impl.Get() )
@@ -246,67 +307,6 @@ UE::FSdfPath FUsdInfoCache::UnwindToNonCollapsedPath( const UE::FSdfPath& Path,
 	return Path;
 }
 
-void FUsdInfoCache::RegisterAuxiliaryPrim(const UE::FSdfPath& MainPrimPath, const UE::FSdfPath& AuxPrimPath) const
-{
-	if (FUsdInfoCacheImpl* ImplPtr = Impl.Get())
-	{
-		FWriteScopeLock ScopeLock(ImplPtr->AuxiliaryPrimsLock);
-
-		UE_LOG(LogUsd, Verbose, TEXT("Registering main prim '%s' and aux prim '%s'"),
-			*MainPrimPath.GetString(),
-			*AuxPrimPath.GetString()
-		);
-
-		ImplPtr->MainToAuxPrims.FindOrAdd(MainPrimPath).Add(AuxPrimPath);
-		ImplPtr->AuxToMainPrims.FindOrAdd(AuxPrimPath).Add(MainPrimPath);
-	}
-}
-
-void FUsdInfoCache::RegisterAuxiliaryPrims(const UE::FSdfPath& MainPrimPath, const TSet<UE::FSdfPath>& AuxPrimPaths) const
-{
-	if (FUsdInfoCacheImpl* ImplPtr = Impl.Get())
-	{
-		FWriteScopeLock ScopeLock(ImplPtr->AuxiliaryPrimsLock);
-
-		ImplPtr->MainToAuxPrims.FindOrAdd(MainPrimPath).Append(AuxPrimPaths);
-
-		for (const UE::FSdfPath& AuxPrimPath : AuxPrimPaths)
-		{
-			UE_LOG(LogUsd, Verbose, TEXT("Registering main prim '%s' and aux prim '%s'"),
-				*MainPrimPath.GetString(),
-				*AuxPrimPath.GetString()
-			);
-
-			ImplPtr->AuxToMainPrims.FindOrAdd(AuxPrimPath).Add(MainPrimPath);
-		}
-	}
-}
-
-void FUsdInfoCache::RemoveAuxiliaryPrims(const UE::FSdfPath& MainPrimPath) const
-{
-	if (FUsdInfoCacheImpl* ImplPtr = Impl.Get())
-	{
-		FWriteScopeLock ScopeLock(ImplPtr->AuxiliaryPrimsLock);
-
-		TSet<UE::FSdfPath> AuxPrimPaths;
-		ImplPtr->MainToAuxPrims.RemoveAndCopyValue(MainPrimPath, AuxPrimPaths);
-
-		for (const UE::FSdfPath& AuxPrimPath : AuxPrimPaths)
-		{
-			UE_LOG(LogUsd, Verbose, TEXT("Removing auxilliary prim '%s' from main prim '%s'"),
-				*AuxPrimPath.GetString(),
-				*MainPrimPath.GetString()
-			);
-
-			TSet<UE::FSdfPath>* MainPrimPaths = ImplPtr->AuxToMainPrims.Find(AuxPrimPath);
-			if (ensure(MainPrimPaths))
-			{
-				MainPrimPaths->Remove(MainPrimPath);
-			}
-		}
-	}
-}
-
 TSet<UE::FSdfPath> FUsdInfoCache::GetMainPrims(const UE::FSdfPath& AuxPrimPath) const
 {
 	if (FUsdInfoCacheImpl* ImplPtr = Impl.Get())
@@ -506,14 +506,12 @@ namespace UE::USDInfoCacheImpl::Private
 
 		if (!UsdPrim.IsPseudoRoot())
 		{
+			pxr::UsdShadeMaterial ShadeMaterial;
+
 			pxr::UsdShadeMaterialBindingAPI BindingAPI{UsdPrim};
 			if (BindingAPI)
 			{
-				if (pxr::UsdShadeMaterial ShadeMaterial = BindingAPI.ComputeBoundMaterial(MaterialPurposeToken))
-				{
-					FWriteScopeLock ScopeLock(Impl.MaterialUsersLock);
-					Impl.MaterialUsers.FindOrAdd(UE::FSdfPath{ShadeMaterial.GetPrim().GetPath()}).Add(UE::FSdfPath{UsdPrimPath});
-				}
+				ShadeMaterial = BindingAPI.ComputeBoundMaterial(MaterialPurposeToken);
 			}
 			else if (pxr::UsdRelationship Relationship = UsdPrim.GetRelationship(pxr::UsdShadeTokens->materialBinding))
 			{
@@ -524,12 +522,43 @@ namespace UE::USDInfoCacheImpl::Private
 				{
 					const pxr::SdfPath& TargetMaterialPrimPath = Targets[0];
 					pxr::UsdPrim MaterialPrim = UsdPrim.GetStage()->GetPrimAtPath(TargetMaterialPrimPath);
-					pxr::UsdShadeMaterial UsdShadeMaterial{MaterialPrim};
-					if (UsdShadeMaterial)
-					{
-						FWriteScopeLock ScopeLock(Impl.MaterialUsersLock);
-						Impl.MaterialUsers.FindOrAdd(UE::FSdfPath{TargetMaterialPrimPath}).Add(UE::FSdfPath{UsdPrimPath});
-					}
+					ShadeMaterial = pxr::UsdShadeMaterial{MaterialPrim};
+				}
+			}
+
+			if (ShadeMaterial)
+			{
+				UE::FSdfPath UsedMaterialPath{ShadeMaterial.GetPrim().GetPath()};
+
+				FWriteScopeLock ScopeLock(Impl.MaterialUsersLock);
+				TSet<UE::FSdfPath>& MaterialUsersForPrim = Impl.MaterialUsers.FindOrAdd(UsedMaterialPath);
+
+				MaterialUsersForPrim.Add(PrimPath);
+
+				UE_LOG(
+					LogUsd,
+					Verbose,
+					TEXT("Registering prim '%s' as a material user of material prim '%s'"),
+					*PrimPath.GetString(),
+					*UsedMaterialPath.GetString()
+				);
+
+				// Our notice handling is somewhat stricter now, and we have no good way of upgrading a simple material info change
+				// into a resync change of the StaticMeshComponent when we change a material that is bound directly to a
+				// UsdGeomSubset, since the GeomMesh translator doesn't collapse. We'll unwind this path later when fetching material
+				// users, so collapsed static meshes are handled OK, skeletal meshes are handled OK, we just need this one exception
+				// for handling uncollapsed static meshes, because by default Mesh prims don't "collapse" their child UsdGeomSubsets
+				if (pxr::UsdGeomSubset Subset{UsdPrim})
+				{
+					UE::FSdfPath ParentMesh{PrimPath.GetParentPath()};
+					UE_LOG(
+						LogUsd,
+						Verbose,
+						TEXT("Registering prim '%s' as a material user of material prim '%s'"),
+						*ParentMesh.GetString(),
+						*UsedMaterialPath.GetString()
+					);
+					MaterialUsersForPrim.Add(ParentMesh);
 				}
 			}
 		}
@@ -854,9 +883,10 @@ namespace UE::USDInfoCacheImpl::Private
 		bool bIsAssetCollapsed = !AssetCollapsedRoot.IsEmpty();
 		bool bIsComponentCollapsed = !ComponentCollapsedRoot.IsEmpty();
 
-		if (!bIsAssetCollapsed || !bIsComponentCollapsed)
+		TSharedPtr< FUsdSchemaTranslator > SchemaTranslator = Registry.CreateTranslatorForSchema(Context.AsShared(), UE::FUsdTyped(UsdPrim));
+		if (SchemaTranslator)
 		{
-			if (TSharedPtr< FUsdSchemaTranslator > SchemaTranslator = Registry.CreateTranslatorForSchema(Context.AsShared(), UE::FUsdTyped(UsdPrim)))
+			if (!bIsAssetCollapsed || !bIsComponentCollapsed)
 			{
 				const bool bCanMeshSubtreeBeCollapsed = CanMeshSubtreeBeCollapsed(UsdPrim, Context, Impl, SchemaTranslator);
 
@@ -917,6 +947,24 @@ namespace UE::USDInfoCacheImpl::Private
 			Info.AssetCollapsedRoot = *AssetCollapsedRootOverride;
 			Info.ComponentCollapsedRoot = *ComponentCollapsedRootOverride;
 		}
+
+		// This really should be a separate pass, but it does no harm here and we have so many passes already...
+		// This needs to happen after we set things into the FUsdPrimInfo for this prim right above this, as it may query
+		// whether this prim (or any of its children) collapse
+		if (SchemaTranslator)
+		{
+			const bool bIsCollapsedBySomeParent = AssetCollapsedRootOverride && !AssetCollapsedRootOverride->IsEmpty()
+												  && *AssetCollapsedRootOverride != UsdPrimPath;
+
+			// We don't care about prims that were collapsed by another. This because whenever the collapse root
+			// registers its auxiliary prims it will already account for all of the collapsed prims that are relevant,
+			// according to the translator type. If we registered aux prims for all prims here, we'd get useless aux
+			// prim links between e.g. all parent and child prims *within* the collapsed subtree
+			if (!bIsCollapsedBySomeParent)
+			{
+				Impl.RegisterAuxiliaryPrims(PrimPath, SchemaTranslator->CollectAuxiliaryPrims());
+			}
+		}
 	}
 
 	void FindValidGeometryCacheRoot(const pxr::UsdPrim& UsdPrim, FUsdSchemaTranslationContext& Context, FUsdInfoCache::FUsdInfoCacheImpl& Impl, UE::UsdInfoCache::Private::EGeometryCachePrimState& OutState)
@@ -1049,7 +1097,7 @@ namespace UE::USDInfoCacheImpl::Private
 		// if the prim can serve as a root and collapse its children into a GeometryCache. The sole condition for
 		// being a valid root is that all the branches of the subtree are collapsible.
 		EGeometryCachePrimState ChildrenState = EGeometryCachePrimState::None;
-		for (EGeometryCachePrimState ChildState : States) 
+		for (EGeometryCachePrimState ChildState : States)
 		{
 			ChildrenState |= ChildState;
 		}
@@ -1394,8 +1442,7 @@ void FUsdInfoCache::RebuildCacheForSubtree( const UE::FUsdPrim& Prim, FUsdSchema
 	// We can't deallocate our info cache pointer with the Usd allocator
 	FScopedUnrealAllocs UEAllocs;
 
-	// We don't want the translation context to try using its info cache during the rebuild process, as that's the entire point
-	TGuardValue< TSharedPtr<FUsdInfoCache> > Guard{ Context.InfoCache, nullptr };
+	TGuardValue<bool> Guard{Context.bIsBuildingInfoCache, true};
 	{
 		FScopedUsdAllocs Allocs;
 
@@ -1406,8 +1453,7 @@ void FUsdInfoCache::RebuildCacheForSubtree( const UE::FUsdPrim& Prim, FUsdSchema
 		}
 
 		// Don't call Clear() here as we don't want to get rid of PrimPathToAssets because we're rebuilding the cache,
-		// as that info is also linked to the asset cache. The same for the aux prim maps, that currently gets
-		// filled in on-demand.
+		// as that info is also linked to the asset cache
 		{
 			FWriteScopeLock ScopeLock(ImplPtr->InfoMapLock);
 			ImplPtr->InfoMap.Empty();
@@ -1416,6 +1462,11 @@ void FUsdInfoCache::RebuildCacheForSubtree( const UE::FUsdPrim& Prim, FUsdSchema
 			FWriteScopeLock ScopeLock(ImplPtr->MaterialUsersLock);
 			ImplPtr->MaterialUsers.Empty();
 		}
+		{
+			FWriteScopeLock ScopeLock(ImplPtr->AuxiliaryPrimsLock);
+			ImplPtr->AuxToMainPrims.Empty();
+			ImplPtr->MainToAuxPrims.Empty();
+		}
 
 		IUsdSchemasModule& UsdSchemasModule = FModuleManager::Get().LoadModuleChecked< IUsdSchemasModule >( TEXT( "USDSchemas" ) );
 		FUsdSchemaTranslatorRegistry& Registry = UsdSchemasModule.GetTranslatorRegistry();
diff --git a/Engine/Plugins/Importers/USDImporter/Source/USDSchemas/Private/USDSchemaTranslator.cpp b/Engine/Plugins/Importers/USDImporter/Source/USDSchemas/Private/USDSchemaTranslator.cpp
--- a/Engine/Plugins/Importers/USDImporter/Source/USDSchemas/Private/USDSchemaTranslator.cpp
+++ b/Engine/Plugins/Importers/USDImporter/Source/USDSchemas/Private/USDSchemaTranslator.cpp
@@ -201,26 +201,14 @@ void FUsdSchemaTranslationContext::CompleteTasks()
 	}
 }
 
-void FUsdSchemaTranslator::RegisterAuxiliaryPrims()
-{
-	if (Context->InfoCache.IsValid())
-	{
-		TSet<UE::FSdfPath> AuxPrims = CollectAuxiliaryPrims();
-		if (AuxPrims.Num() > 0)
-		{
-			Context->InfoCache->RegisterAuxiliaryPrims(PrimPath, AuxPrims);
-		}
-	}
-}
-
 bool FUsdSchemaTranslator::IsCollapsed( ECollapsingType CollapsingType ) const
 {
 #if USE_USD_SDK
 	TRACE_CPUPROFILER_EVENT_SCOPE( FUsdSchemaTranslator::IsCollapsed );
 
-	if ( Context->InfoCache.IsValid() )
+	if (!Context->bIsBuildingInfoCache)
 	{
-		return Context->InfoCache->IsPathCollapsed( PrimPath, CollapsingType );
+		return Context->InfoCache->IsPathCollapsed(PrimPath, CollapsingType);
 	}
 
 	// This is merely a fallback, and we should never need this
diff --git a/Engine/Plugins/Importers/USDImporter/Source/USDSchemas/Private/USDShadeMaterialTranslator.cpp b/Engine/Plugins/Importers/USDImporter/Source/USDSchemas/Private/USDShadeMaterialTranslator.cpp
--- a/Engine/Plugins/Importers/USDImporter/Source/USDSchemas/Private/USDShadeMaterialTranslator.cpp
+++ b/Engine/Plugins/Importers/USDImporter/Source/USDSchemas/Private/USDShadeMaterialTranslator.cpp
@@ -250,8 +250,6 @@ void FUsdShadeMaterialTranslator::CreateAssets()
 		}
 	}
 
-	RegisterAuxiliaryPrims();
-
 	const pxr::TfToken RenderContextToken =
 		Context->RenderContext.IsNone() ?
 			pxr::UsdShadeTokens->universalRenderContext :
@@ -553,6 +551,11 @@ void FUsdShadeMaterialTranslator::PostImportMaterial(const FString& MaterialHash
 
 TSet<UE::FSdfPath> FUsdShadeMaterialTranslator::CollectAuxiliaryPrims() const
 {
+	if (!Context->bIsBuildingInfoCache)
+	{
+		return Context->InfoCache->GetAuxiliaryPrims(PrimPath);
+	}
+
 	TSet<UE::FSdfPath> Result;
 	{
 		TFunction<void(const pxr::UsdShadeInput&)> TraverseShadeInput;
diff --git a/Engine/Plugins/Importers/USDImporter/Source/USDSchemas/Private/USDSkelRootTranslator.cpp b/Engine/Plugins/Importers/USDImporter/Source/USDSchemas/Private/USDSkelRootTranslator.cpp
--- a/Engine/Plugins/Importers/USDImporter/Source/USDSchemas/Private/USDSkelRootTranslator.cpp
+++ b/Engine/Plugins/Importers/USDImporter/Source/USDSchemas/Private/USDSkelRootTranslator.cpp
@@ -1642,8 +1642,6 @@ void FUsdSkelRootTranslator::CreateAssets()
 	TRACE_CPUPROFILER_EVENT_SCOPE(FUsdSkelRootTranslator::CreateAssets);
 
 #if WITH_EDITOR
-	RegisterAuxiliaryPrims();
-
 	// Importing skeletal meshes actually works in Standalone mode, but we intentionally block it here
 	// to not confuse users as to why it doesn't work at runtime
 	TSharedRef< UsdSkelRootTranslatorImpl::FSkelRootCreateAssetsTaskChain > AssetsTaskChain =
@@ -1877,9 +1875,9 @@ bool FUsdSkelRootTranslator::CanBeCollapsed(ECollapsingType CollapsingType) cons
 
 TSet<UE::FSdfPath> FUsdSkelRootTranslator::CollectAuxiliaryPrims() const
 {
-	if (!Context->InfoCache.IsValid())
+	if (!Context->bIsBuildingInfoCache)
 	{
-		return {};
+		return Context->InfoCache->GetAuxiliaryPrims(PrimPath);
 	}
 
 	if (!Context->InfoCache->DoesPathCollapseChildren(PrimPath, ECollapsingType::Assets))
diff --git a/Engine/Plugins/Importers/USDImporter/Source/USDSchemas/Public/USDInfoCache.h b/Engine/Plugins/Importers/USDImporter/Source/USDSchemas/Public/USDInfoCache.h
--- a/Engine/Plugins/Importers/USDImporter/Source/USDSchemas/Public/USDInfoCache.h
+++ b/Engine/Plugins/Importers/USDImporter/Source/USDSchemas/Public/USDInfoCache.h
@@ -35,6 +35,7 @@ public:
 	struct FUsdInfoCacheImpl;
 
 	FUsdInfoCache();
+	FUsdInfoCache(const FUsdInfoCache& Other);
 	virtual ~FUsdInfoCache();
 
 	bool Serialize(FArchive& Ar);
@@ -42,6 +43,9 @@ public:
 	// Returns whether we contain any info about prim at 'Path' at all
 	bool ContainsInfoAboutPrim(const UE::FSdfPath& Path) const;
 
+	// Returns a list of all prims we have generic info about
+	TSet<UE::FSdfPath> GetKnownPrims() const;
+
 	void RebuildCacheForSubtree(const UE::FUsdPrim& Prim, FUsdSchemaTranslationContext& Context);
 
 	void Clear();
@@ -55,12 +59,6 @@ public:
 	UE::FSdfPath UnwindToNonCollapsedPath(const UE::FSdfPath& Path, ECollapsingType CollapsingType) const;
 
 public:
-	// Records that whenever we want to translate MainPrim, we need to read one or more attributes off AuxPrim
-	void RegisterAuxiliaryPrim(const UE::FSdfPath& MainPrimPath, const UE::FSdfPath& AuxPrimPath) const;
-	void RegisterAuxiliaryPrims(const UE::FSdfPath& MainPrimPath, const TSet<UE::FSdfPath>& AuxPrimPaths) const;
-
-	void RemoveAuxiliaryPrims(const UE::FSdfPath& MainPrimPath) const;
-
 	// Returns the paths to prims that, when translated into assets or components, also require reading the prim at
 	// 'Path'. e.g. providing the path to a Shader prim will return the paths to all Material prims for which the
 	// translation involves reading that particular Shader.
diff --git a/Engine/Plugins/Importers/USDImporter/Source/USDSchemas/Public/USDSchemaTranslator.h b/Engine/Plugins/Importers/USDImporter/Source/USDSchemas/Public/USDSchemaTranslator.h
--- a/Engine/Plugins/Importers/USDImporter/Source/USDSchemas/Public/USDSchemaTranslator.h
+++ b/Engine/Plugins/Importers/USDImporter/Source/USDSchemas/Public/USDSchemaTranslator.h
@@ -140,6 +140,12 @@ struct USDSCHEMAS_API FUsdSchemaTranslationContext : public TSharedFromThis< FUs
 	/** True if we're a context created by the USDStageImporter to fully import to persistent assets and actors */
 	bool bIsImporting = false;
 
+	/**
+	 * True if we're building the InfoCache assigned to this context. This usually means we shouldn't query it for information, and should instead
+	 * compute it manually so that it can be cached
+	 */
+	bool bIsBuildingInfoCache = false;
+
 	/** pxr::UsdStage we're translating from */
 	UE::FUsdStage Stage;
 
@@ -287,12 +293,6 @@ public:
 
 	virtual ~FUsdSchemaTranslator() = default;
 
-	/**
-	 * Translates the prim at PrimPath, creating assets.
-	 * Derived classes that wish to implement FUsdSchemaTranslator::CreateAssets should manually call
-	 * FUsdSchemaTranslator::RegisterAuxiliaryPrims and implement CollectAuxiliaryPrims to catalog their
-	 * dependency prims and regenerate assets when those prims are updated.
-	 */
 	virtual void CreateAssets() {}
 
 	virtual USceneComponent* CreateComponents() { return nullptr; }
@@ -312,10 +312,6 @@ public:
 
 	UE::FUsdPrim GetPrim() const { return Context->Stage.GetPrimAtPath(PrimPath); }
 
-protected:
-	/** Calls CollectAuxiliaryPrims to record the current dependencies in the Context's InfoCache */
-	void RegisterAuxiliaryPrims();
-
 protected:
 	UE::FSdfPath PrimPath;
 	TSharedRef< FUsdSchemaTranslationContext > Context;
diff --git a/Engine/Plugins/Importers/USDImporter/Source/USDStage/Private/USDStageActor.cpp b/Engine/Plugins/Importers/USDImporter/Source/USDStage/Private/USDStageActor.cpp
--- a/Engine/Plugins/Importers/USDImporter/Source/USDStage/Private/USDStageActor.cpp
+++ b/Engine/Plugins/Importers/USDImporter/Source/USDStage/Private/USDStageActor.cpp
@@ -1208,7 +1208,7 @@ void AUsdStageActor::OnUsdObjectsChanged(const UsdUtils::FObjectChangesByPath& I
 	bool bNeedsAnimationReload = false;
 
 	// The most important thing here is to iterate in parent to child order, so build SortedPrimsChangedList
-	TMap< FString, bool > SortedPrimsChangedList;
+	TMap<UE::FSdfPath, bool> SortedPrimsChangedList;
 	for (const TPair<FString, TArray<UsdUtils::FSdfChangeListEntry>>& InfoChange : InfoChanges)
 	{
 		UE::FSdfPath PrimPath = UE::FSdfPath(*InfoChange.Key).StripAllVariantSelections();
@@ -1224,6 +1224,21 @@ void AUsdStageActor::OnUsdObjectsChanged(const UsdUtils::FObjectChangesByPath& I
 			{
 				bIsResync = true;
 				bHasResync = true;
+
+				// If the material that was changed was on a geom subset, we need to instead pretend it
+				// happened to its parent Mesh prim instead. If the mesh is collapsed this won't matter,
+				// but if it's not collapsed we currently need this to make sure we resync the static mesh
+				// component. The Mesh prim will declare that the subset is its aux prim, but we don't propagate
+				// resyncs from aux to main prims now, so even if we resynced the UsdGeomSubset the Mesh
+				// would not resync.
+				// TODO: Better way of handling material override updates. Ideally we wouldn't be resyncing
+				// the Mesh prim just to update material overrides...
+				UE::FUsdPrim ChangedPrim = Stage.GetPrimAtPath(PrimPath);
+				if(ChangedPrim && ChangedPrim.IsA(TEXT("GeomSubset")))
+				{
+					PrimPath = PrimPath.GetParentPath();
+				}
+
 				continue;
 			}
 
@@ -1271,7 +1286,7 @@ void AUsdStageActor::OnUsdObjectsChanged(const UsdUtils::FObjectChangesByPath& I
 		// We may need the full spec path with variant selections later, but for traversal and retrieving prims from the stage we always need
 		// the prim path without any variant selections in it (i.e. GetPrimAtPath("/Root{Varset=Var}Child") doesn't work, we need GetPrimAtPath("/Root/Child")),
 		// and USD sometimes emits changes with the variant selection path (like during renames).
-		SortedPrimsChangedList.Add(UE::FSdfPath(*InfoChange.Key).StripAllVariantSelections().GetString(), bIsResync);
+		SortedPrimsChangedList.Add(PrimPath, bIsResync);
 	}
 	// Do Resyncs after so that they overwrite pure info changes if we have any
 	for (const TPair<FString, TArray<UsdUtils::FSdfChangeListEntry>>& ResyncChange : ResyncChanges)
@@ -1279,41 +1294,9 @@ void AUsdStageActor::OnUsdObjectsChanged(const UsdUtils::FObjectChangesByPath& I
 		UE::FSdfPath PrimPath = UE::FSdfPath(*ResyncChange.Key).StripAllVariantSelections();
 
 		const bool bIsResync = true;
-		SortedPrimsChangedList.Add(PrimPath.GetString(), bIsResync);
+		SortedPrimsChangedList.Add(PrimPath, bIsResync);
 	}
 
-	// If we have just resynced a child prim that was collapsed, we actually need to resync from its collapsed root
-	// downwards instead.
-	// UpdateComponents/ReloadAssets below will already do this of course, but the point of this is that we unwind
-	// the collapsed prim paths *before* we rebuild the cache: We're really interested in whether this prim was
-	// collapsed or not before this change, as the change itself may have made the prim uncollapsible now
-	{
-		TMap<FString, bool> NewEntries;
-		NewEntries.Reserve(SortedPrimsChangedList.Num());
-
-		for (TMap< FString, bool >::TIterator ChangeIt = SortedPrimsChangedList.CreateIterator(); ChangeIt; ++ChangeIt)
-		{
-			if (ChangeIt.Value())
-			{
-				UE::FSdfPath ChildPath{*ChangeIt.Key()};
-
-				// Usually the cache should contain info about any prim on the stage, but if this is a
-				// notice about a prim being *created*, it may not have anything about this prim yet
-				if (InfoCache->ContainsInfoAboutPrim(ChildPath))
-				{
-					NewEntries.Add(InfoCache->UnwindToNonCollapsedPath(ChildPath, ECollapsingType::Assets).GetString(), true);
-					NewEntries.Add(InfoCache->UnwindToNonCollapsedPath(ChildPath, ECollapsingType::Components).GetString(), true);
-
-					ChangeIt.RemoveCurrent();
-				}
-			}
-		}
-
-		SortedPrimsChangedList.Append(NewEntries);
-	}
-
-	SortedPrimsChangedList.KeySort([](const FString& A, const FString& B) -> bool { return A.Len() < B.Len(); });
-
 	// During PIE, the PIE and the editor world will respond to notices. We have to prevent any PIE
 	// objects from being added to the transaction however, or else it will be discarded when finalized.
 	// We need to keep the transaction, or else we may end up with actors outside of the transaction
@@ -1325,190 +1308,356 @@ void AUsdStageActor::OnUsdObjectsChanged(const UsdUtils::FObjectChangesByPath& I
 		SuppressTransaction.Emplace(GUndo, nullptr);
 	}
 
-	FScopedSlowTask RefreshStageTask(SortedPrimsChangedList.Num(), LOCTEXT("RefreshingUSDStage", "Refreshing USD Stage"));
+	FScopedSlowTask RefreshStageTask(4, LOCTEXT("ProcessingUSDStageUpdates", "Processing USD Stage updates"));
 	RefreshStageTask.MakeDialogDelayed(0.25f);
 
 	FScopedUsdMessageLog ScopedMessageLog;
 
-	TSet< UE::FSdfPath > UpdatedAssets;
-	TSet< UE::FSdfPath > ResyncedAssets;
-	TSet< UE::FSdfPath > UpdatedComponents;
-	TSet< UE::FSdfPath > ResyncedComponents;
-
-	bool bDeselected = false;
-
-	if (bHasResync && InfoCache.IsValid())
+	TSharedPtr<FUsdInfoCache> OldInfoCache = InfoCache;
+	if (InfoCache.IsValid() && bHasResync)
 	{
+		// Take a copy of the info cache here: We want to keep the old one for this function call as it helps us cleanup the old assets and components
+		// For now we don't need to do this for info changes, only for resync changes
+		InfoCache = MakeShared<FUsdInfoCache>(*OldInfoCache);
+
 		// TODO: Selective rebuild of only the required parts of the cache.
-		// If a prim changes from CanBeCollapsed to not (or vice-versa), that means its parent may change, and its grandparent, etc. so we'd need to check a large
-		// part of the tree (here we're just rebuilding the whole thing for now).
-		// However, if we know that only Prim resynced, we can traverse the tree root down and if we reach Prim and its collapsing state
-		// hasn't updated from before, we don't have to update its subtree at all, or sibling subtrees.
-		// Note that whenver we do a selective rebuild of the info cache we'll need to be very careful when to update
-		// certain info cache maps: For example whenever we delete a prim we need to make sure it's removed from MaterialUsers, etc.
-		TSharedRef< FUsdSchemaTranslationContext > TranslationContext = FUsdStageActorImpl::CreateUsdSchemaTranslationContext(this, TEXT("/"));
+		// If a prim changes from CanBeCollapsed to not (or vice-versa), that means its parent may change, and its grandparent, etc. so we'd need
+		// to check a large part of the tree (here we're just rebuilding the whole thing for now). However, if we know that only Prim resynced, we
+		// can traverse the tree root down and if we reach Prim and its collapsing state hasn't updated from before, we don't have to update its
+		// subtree at all, or sibling subtrees. Note that whenver we do a selective rebuild of the info cache we'll need to be very careful when
+		// to update certain info cache maps: For example whenever we delete a prim we need to make sure it's removed from MaterialUsers, etc.
+		TSharedRef<FUsdSchemaTranslationContext> TranslationContext = FUsdStageActorImpl::CreateUsdSchemaTranslationContext(this, TEXT("/"));
 		InfoCache->RebuildCacheForSubtree(Stage.GetPseudoRoot(), TranslationContext.Get());
 	}
 
-	bool bHasLoadedOrAbandonedAssets = false;
-	for (const TPair< FString, bool >& PrimChangedInfo : SortedPrimsChangedList)
+	TFunction<void(TMap<UE::FSdfPath, bool>&)> SortAndCleanPrimsToUpdate = [](TMap<UE::FSdfPath, bool>& InOutMap)
 	{
-		RefreshStageTask.EnterProgressFrame();
+		// Force traversal of parent before child (not needed for cleaning up, but very much needed later for reloading assets/components)
+		InOutMap.KeySort(
+			[](const UE::FSdfPath& A, const UE::FSdfPath& B) -> bool
+			{
+				return A < B;
+			}
+		);
 
-		const UE::FSdfPath PrimPath = UE::FSdfPath{*PrimChangedInfo.Key};
-		const bool bIsResync = PrimChangedInfo.Value;
+		// Strip child paths for resyncs as processing a parent means we already process the children anyway.
+		// This is not the same for info changes: We may have an info change for a parent and child component
+		// in the same change block, and we really want to call UpdateComponents for both then
+		TMap<UE::FSdfPath, bool> CleanedPairs;
+		CleanedPairs.Reserve(InOutMap.Num());
+		TSet<UE::FSdfPath> ResyncedPaths;
+		ResyncedPaths.Reserve(InOutMap.Num());
+		for (const TPair<UE::FSdfPath, bool>& Pair : InOutMap)
+		{
+			const UE::FSdfPath& ThisPrim = Pair.Key;
+			bool bIsResync = Pair.Value;
 
-		if (bIsResync && !bDeselected)
+			bool bRemoveThisPath = false;
+			for (const UE::FSdfPath& Prefix : ThisPrim.GetPrefixes())
+			{
+				if (ResyncedPaths.Contains(Prefix))
+				{
+					bRemoveThisPath = true;
+					break;
+				}
+			}
+			if (bRemoveThisPath)
+			{
+				continue;
+			}
+
+			if (bIsResync)
+			{
+				ResyncedPaths.Add(ThisPrim);
+			}
+
+			CleanedPairs.Add(Pair);
+		}
+		Swap(CleanedPairs, InOutMap);
+	};
+
+	TSet<UE::FSdfPath> KnownOldPrims;
+
+	// Traverses the info caches to find out which prims we need to update
+	TFunction<void(const UE::FSdfPath&, bool, TMap<UE::FSdfPath, bool>&, TSet<UE::FSdfPath>&)> RecursiveCollectPrimsToUpdate;
+	RecursiveCollectPrimsToUpdate =
+		[this, &OldInfoCache, &RecursiveCollectPrimsToUpdate, &KnownOldPrims, &Stage](
+			const UE::FSdfPath& PrimPath,
+			bool bIsResync,
+			TMap<UE::FSdfPath, bool>& OutPrimsToUpdate,
+			TSet<UE::FSdfPath>& InOutVisitedPaths
+		)
+	{
+		// Allow revisiting a prim if we're resyncing: This is a cheap way of making sure we mark the
+		// resynced prims as resynced, or else we may recurse into them while traversing another prim and preemptively
+		// mark them as visited, and because we always recurse with resync=false we'd assume those aren't resyncs.
+		// By that same reason this shouldn't be that expensive, as we'll only ever potentially revisit the prims
+		// that are the actual roots of the resyncs
+		if (InOutVisitedPaths.Contains(PrimPath) && !bIsResync)
+		{
+			return;
+		}
+		InOutVisitedPaths.Add(PrimPath);
+
+		// In some cases USD sends us notices about prims that don't exist anymore: If you rename X to Y,
+		// both X and Y will be on the notice change list, even though X doesn't exist on the stage anymore.
+		// It's easy to just ignore those here by doing this, but we could also pay attention to the flags
+		// on the notice and try to guess if a rename took place, if needed
+		if (InfoCache->ContainsInfoAboutPrim(PrimPath))
 		{
-			FUsdStageActorImpl::DeselectActorsAndComponents(this);
-			bDeselected = true;
+			// We always want the unwound path here. We'll take care to only ever register main prims that are
+			// themselves uncollapsed or collapsed roots, but there's nothing stopping the user from manually
+			// modifying directly a collapsed prim that is not an aux prim of it's collapsed root (e.g. some
+			// parent prim of a point instancer prototype). If we retrieved a main prim for a prim like that,
+			// we'd only get that prim itself, and assume we need to spawn assets/components for it (which we
+			// really don't if it's collapsed)
+			UE::FSdfPath UnwoundPath = InfoCache->UnwindToNonCollapsedPath(PrimPath, ECollapsingType::Assets);
+			OutPrimsToUpdate.FindOrAdd(UnwoundPath) |= bIsResync;
 		}
 
-		// Return if the path or any of its higher level paths are already processed
-		auto IsPathAlreadyProcessed = [](TSet< UE::FSdfPath >& PathsProcessed, UE::FSdfPath PathToProcess) -> bool
+		// If we're going to cleanup assets/components for a prim according to an old dependency (that doesn't
+		// necessarily exist anymore), we need to make sure we try recreating assets/components for it now, even if
+		// it wasn't part of the actual change, so we'll visit both the new and old info caches.
+		// Even if this was just an info change we still need to do this: We may need to regeneate an asset from it
+		if (OldInfoCache->ContainsInfoAboutPrim(PrimPath))
 		{
-			FString SubPath;
-			FString ParentPath;
+			UE::FSdfPath UnwoundPath = OldInfoCache->UnwindToNonCollapsedPath(PrimPath, ECollapsingType::Assets);
+			OutPrimsToUpdate.FindOrAdd(UnwoundPath) |= bIsResync;
+		}
 
-			if (PathsProcessed.Contains(UE::FSdfPath::AbsoluteRootPath()))
+		// If our original USD notice resyncs PrimPath, its subtree will need to be rebuilt, yes,
+		// but external prims that depend on prim path (its "main prims") won't need to be *recursively* resynced.
+		// Their hierarchies are fine, they just need to be updated to the fact that PrimPath changed. That is
+		// at most a component update, or regenerating the asset for that particular main prim, but it's entire
+		// hierarchy doesn't need to be rebuilt
+		const bool bRecursiveResync = false;
+
+		// Imagine we have a stage like this:
+		// 		/parent/child1
+		// 		/other
+		// And "child1" is marked as an aux prim for "other". What happens if we resync "parent"? Since a resync means
+		// the subtree is arbitrarily rebuilt, it means we probably want to update "other" too, which is what this does.
+		// Note that parent and child could be fully independent, uncollapsed prims, without main/aux links between them.
+		// We have to do this on both the stage and old info cache because the change may also have meant that aux/main links
+		// have been modified (i.e. "other" could depend on "child" only now, or only on the old state of the stage, but
+		// we'll still have those assets on the UE level either way, so we need to refresh them)
+		if (bIsResync)
+		{
+			if (KnownOldPrims.Num() == 0)
 			{
-				return true;
+				KnownOldPrims = OldInfoCache->GetKnownPrims();
 			}
 
-			while (!PathToProcess.IsEmpty() && !PathsProcessed.Contains(PathToProcess))
+			if(UE::FUsdPrim Prim = Stage.GetPrimAtPath(PrimPath))
 			{
-				if (!PathToProcess.IsAbsoluteRootPath())
+				for (const UE::FUsdPrim& ChildPrim : Prim.GetChildren())
 				{
-					PathToProcess = PathToProcess.GetParentPath();
+					UE::FSdfPath ChildPrimPath = ChildPrim.GetPrimPath();
+					if (ChildPrimPath.HasPrefix(PrimPath))
+					{
+						RecursiveCollectPrimsToUpdate(ChildPrimPath, bRecursiveResync, OutPrimsToUpdate, InOutVisitedPaths);
+					}
 				}
-				else
+			}
+			// TODO: Find a better way of doing this, as this is extremely expensive
+			for (const UE::FSdfPath& KnownOldPrim : KnownOldPrims)
+			{
+				if (KnownOldPrim.HasPrefix(PrimPath))
 				{
-					return false;
+					RecursiveCollectPrimsToUpdate(KnownOldPrim, bRecursiveResync, OutPrimsToUpdate, InOutVisitedPaths);
 				}
 			}
+		}
 
-			return !PathToProcess.IsEmpty() && PathsProcessed.Contains(PathToProcess);
-		};
+		TSet<UE::FSdfPath> NewMainPrims = InfoCache->GetMainPrims(PrimPath);
+		TSet<UE::FSdfPath> OldMainPrims = OldInfoCache->GetMainPrims(PrimPath);
+		OutPrimsToUpdate.Reserve(OutPrimsToUpdate.Num() + NewMainPrims.Num() + OldMainPrims.Num());
+		for (const UE::FSdfPath& NewPrimPath : NewMainPrims)
+		{
+			RecursiveCollectPrimsToUpdate(NewPrimPath, bRecursiveResync, OutPrimsToUpdate, InOutVisitedPaths);
+		}
+		for (const UE::FSdfPath& OldPrimPath : OldMainPrims)
+		{
+			RecursiveCollectPrimsToUpdate(OldPrimPath, bRecursiveResync, OutPrimsToUpdate, InOutVisitedPaths);
+		}
+	};
 
-		TFunction<void(const UE::FSdfPath& InPrimPath, const bool bInResync)> UpdateComponents;
-		UpdateComponents = [&](const UE::FSdfPath& InPrimPath, const bool bInResync)
+	// Collect all the paths to update
+	TMap<UE::FSdfPath, bool> PrimsToUpdate;
+	{
+		PrimsToUpdate.Reserve(SortedPrimsChangedList.Num());
+
+		// Recursively append main prims to the list of PrimsToUpdate
+		TSet<UE::FSdfPath> VisitedPaths;
+		for (const TPair<UE::FSdfPath, bool>& PrimChangedInfo : SortedPrimsChangedList)
 		{
-			TSet< UE::FSdfPath >& RefreshedComponents = bInResync ? ResyncedComponents : UpdatedComponents;
+			const UE::FSdfPath PrimPath = PrimChangedInfo.Key;
+			const bool bIsResync = PrimChangedInfo.Value;
 
-			TSet<UE::FSdfPath> MainPrimPaths = InfoCache->GetMainPrims(InPrimPath);
-			for (const UE::FSdfPath& MainPrimPath : MainPrimPaths)
-			{
-				// Note: We don't need to "unwind" MainPrimPath to a collapse root here: The entries for main prims
-				// on the Main/Aux maps are always prims that were translated into assets before, so we know they
-				// must be "unwound" paths already
+			// Note how we're not modifying SortedPrimsChangedList in-place and are instead adding to a new PrimsToUpdate list.
+			// The intent is that we really only want to process uncollapsed/collapse root main prims, but what is actually on these notices is up to
+			// USD, and could have anything
+			RecursiveCollectPrimsToUpdate(PrimPath, bIsResync, PrimsToUpdate, VisitedPaths);
+		}
 
-				if (IsPathAlreadyProcessed(RefreshedComponents, MainPrimPath))
-				{
-					continue;
-				}
-				RefreshedComponents.Add(MainPrimPath);
-				if (bInResync)
-				{
-					// Consider that the path has been updated in the case of a resync
-					UpdatedComponents.Add(MainPrimPath);
-				}
+		SortAndCleanPrimsToUpdate(PrimsToUpdate);
+	}
 
-				UE_LOG(LogUsd, Verbose, TEXT("Prim '%s' was updated. Updating components for prim '%s'"),
-					*InPrimPath.GetString(),
-					*MainPrimPath.GetString()
-				);
+	if (bHasResync)
+	{
+		FUsdStageActorImpl::DeselectActorsAndComponents(this);
+	}
 
-				// We could share this context, add all tasks for all consumers and then call CompleteTasks once, but
-				// that could cause trouble in case we had a triangle of dependencies (A depends on B an C,
-				// and B depends on C) as that would prevent dependencies from being parsed before dependents
-				TSharedRef< FUsdSchemaTranslationContext > TranslationContext = FUsdStageActorImpl::CreateUsdSchemaTranslationContext(
-					this,
-					MainPrimPath.GetString()
-				);
+	bool bHasLoadedOrAbandonedAssets = false;
 
-				UpdatePrim(MainPrimPath, bInResync, *TranslationContext);
-				TranslationContext->CompleteTasks();
+	RefreshStageTask.EnterProgressFrame();
+	FScopedSlowTask CleanUpAssetsTask(PrimsToUpdate.Num(), LOCTEXT("CleaningUpAssets", "Cleaning up old assets"));
+	for (const TPair<UE::FSdfPath, bool>& PrimChangedInfo : PrimsToUpdate)
+	{
+		CleanUpAssetsTask.EnterProgressFrame();
 
-				// Now update the consumers of *this* prim
-				UpdateComponents(MainPrimPath, bInResync);
-			}
-		};
+		const UE::FSdfPath& PrimPath = PrimChangedInfo.Key;
+		const bool bIsResync = PrimChangedInfo.Value;
 
-		TFunction<void(const UE::FSdfPath& InPrimPath, const bool bInResync)> ReloadAssets;
-		ReloadAssets = [&](const UE::FSdfPath& InPrimPath, const bool bInResync)
+		// If it's a new prim, we may not have old info about it
+		if (!OldInfoCache->ContainsInfoAboutPrim(PrimPath))
 		{
-			TSet< UE::FSdfPath >& RefreshedAssets = bInResync ? ResyncedAssets : UpdatedAssets;
+			continue;
+		}
 
-			TSet<UE::FSdfPath> MainPrimPaths = InfoCache->GetMainPrims(InPrimPath);
-			for (const UE::FSdfPath& MainPrimPath : MainPrimPaths)
-			{
-				if (IsPathAlreadyProcessed(RefreshedAssets, MainPrimPath))
-				{
-					continue;
-				}
-				RefreshedAssets.Add(MainPrimPath);
+		UE_LOG(LogUsd, Verbose, TEXT("Cleaning up assets for prim path '%s'"), *PrimPath.GetString());
 
-				UE_LOG(LogUsd, Verbose, TEXT("Prim '%s' was updated. Reloading assets for prim '%s'"),
-					*InPrimPath.GetString(),
-					*MainPrimPath.GetString()
-				);
+		const bool bForEntireSubtree = bIsResync;
+		bHasLoadedOrAbandonedAssets |= UnloadAssets(PrimPath, bForEntireSubtree);
+	}
 
-				UE::FUsdPrim PrimToUpdate = GetOrLoadUsdStage().GetPrimAtPath(MainPrimPath);
-				if (!PrimToUpdate)
-				{
-					continue;
-				}
+	// Regenerate assets before cleaning up components because if material assignments change, we may
+	// discover we need to clear and regenerate additional mesh components to update material overrides
+	TSet<UE::FSdfPath> MaterialUserPrims;
+	RefreshStageTask.EnterProgressFrame();
+	FScopedSlowTask RegenerateAssetsTask(PrimsToUpdate.Num(), LOCTEXT("RegeneratingAssets", "Regenerating assets"));
+	for (const TPair<UE::FSdfPath, bool>& PrimChangedInfo : PrimsToUpdate)
+	{
+		RegenerateAssetsTask.EnterProgressFrame();
 
-				TSharedRef< FUsdSchemaTranslationContext > TranslationContext = FUsdStageActorImpl::CreateUsdSchemaTranslationContext(
-					this,
-					MainPrimPath.GetString()
-				);
+		const UE::FSdfPath& PrimPath = PrimChangedInfo.Key;
+		const bool bIsResync = PrimChangedInfo.Value;
 
-				UMaterialInterface* ExistingMaterial = InfoCache->GetSingleAssetForPrim<UMaterialInterface>(MainPrimPath);
+		UE::FUsdPrim PrimToUpdate = Stage.GetPrimAtPath(PrimPath);
 
-				if (bInResync)
-				{
-					bHasLoadedOrAbandonedAssets |= LoadAssets(*TranslationContext, PrimToUpdate);
+		// It's OK to not have info about a prim if it's an old prim that only exists on the old info cache.
+		// If the new info cache has info about this prim then it must exist on the stage right now
+		if (!InfoCache->ContainsInfoAboutPrim(PrimPath))
+		{
+			ensure(!PrimToUpdate);
+			continue;
+		}
+		ensure(PrimToUpdate);
 
-					// Resyncing also includes "updating" the prim
-					UpdatedAssets.Add(MainPrimPath);
-				}
-				else
-				{
-					bHasLoadedOrAbandonedAssets |= LoadAsset(*TranslationContext, PrimToUpdate);
-				}
+		UE_LOG(LogUsd, Verbose, TEXT("Regenerating assets for prim path '%s'"), *PrimPath.GetString());
 
-				UMaterialInterface* NewMaterial = InfoCache->GetSingleAssetForPrim<UMaterialInterface>(
-					MainPrimPath
-				);
+		TSharedRef<FUsdSchemaTranslationContext> TranslationContext = FUsdStageActorImpl::CreateUsdSchemaTranslationContext(
+			this,
+			PrimPath.GetString()
+		);
 
-				// For UE-120185: If we recreated a material for a prim path we also need to update all components that were using it.
-				// This could be fleshed out further if other asset types require this refresh of "consumer components" but materials
-				// seem to be the only ones that do at the moment.
-				// Note that even after UE-157644 this is still useful: Material prims are not marked as a dependency
-				// of Mesh prims, otherwise we'd have to regenerate the StaticMesh itself when material info changed.
-				// Ideally we'd just set a new material override on the component, which is what this does
-				if (ExistingMaterial && (ExistingMaterial != NewMaterial) && InfoCache.IsValid())
-				{
-					for (const UE::FSdfPath& MaterialUserPath : InfoCache->GetMaterialUsers(MainPrimPath))
-					{
-						const bool bResyncComponent = true; // We need to force resync to reassign materials
-						UE::FUsdPrim MaterialPrim = GetOrLoadUsdStage().GetPrimAtPath(MaterialUserPath);
-						UpdateComponents(MaterialUserPath, bResyncComponent);
-					}
-				}
+		bool bThisPrimLoadedAssets = false;
+		if (bIsResync)
+		{
+			bThisPrimLoadedAssets |= LoadAssets(*TranslationContext, PrimToUpdate);
+		}
+		else
+		{
+			bThisPrimLoadedAssets |= LoadAsset(*TranslationContext, PrimToUpdate);
+		}
+		bHasLoadedOrAbandonedAssets |= bThisPrimLoadedAssets;
+
+		// For UE-120185: If we recreated a material for a prim path we also need to update all components that were using it.
+		// This could be fleshed out further if other asset types require this refresh of "consumer components" but materials
+		// seem to be the only ones that do at the moment.
+		// Note that even after UE-157644 this is still useful: Material prims are not marked as a dependency
+		// of Mesh prims, otherwise we'd have to regenerate the StaticMesh itself when material info changed.
+		// Ideally we'd just set a new material override on the component, which is what this does
+		const bool bPrimGeneratesMaterials = InfoCache->GetSingleAssetForPrim<UMaterialInterface>(PrimPath) != nullptr;
+		if (bThisPrimLoadedAssets && bPrimGeneratesMaterials)
+		{
+			MaterialUserPrims.Append(InfoCache->GetMaterialUsers(PrimPath));
+		}
+	}
+
+	if(MaterialUserPrims.Num())
+	{
+		TSet<UE::FSdfPath> VisitedPaths;
+		PrimsToUpdate.Reserve(PrimsToUpdate.Num() + MaterialUserPrims.Num());
+		for(const UE::FSdfPath& MaterialUserPrim : MaterialUserPrims)
+		{
+			// This also needs to be done recursivly because our "material user prim" may be just something
+			// like a UsdGeomSubset. If we're going to resync its component, we actually need to resync whatever
+			// component it was collapsed into instead
+			const bool bIsResync = true;
+			RecursiveCollectPrimsToUpdate(MaterialUserPrim, bIsResync, PrimsToUpdate, VisitedPaths);
+		}
+		SortAndCleanPrimsToUpdate(PrimsToUpdate);
+	}
+
+	RefreshStageTask.EnterProgressFrame();
+	FScopedSlowTask CleanUpComponentsTask(PrimsToUpdate.Num(), LOCTEXT("CleaningUpComponents", "Cleaning up actors and components"));
+	for (const TPair<UE::FSdfPath, bool>& PrimChangedInfo : PrimsToUpdate)
+	{
+		CleanUpComponentsTask.EnterProgressFrame();
+
+		const UE::FSdfPath PrimPath = PrimChangedInfo.Key;
+		const bool bIsResync = PrimChangedInfo.Value;
+
+		// If it's a new prim, we may not have old info about it
+		if (!OldInfoCache->ContainsInfoAboutPrim(PrimPath))
+		{
+			continue;
+		}
+
+		UE_LOG(LogUsd, Verbose, TEXT("Cleaning up actors and components for prim path '%s'"), *PrimPath.GetString());
 
-				// Now update the main prims of *this* prim as an auxiliary prim
-				ReloadAssets(MainPrimPath, bInResync);
+		if (bIsResync && PrimPath.IsAbsoluteRootOrPrimPath())
+		{
+			if (UUsdPrimTwin* UsdPrimTwin = GetRootPrimTwin()->Find(PrimPath.GetString()))
+			{
+				UsdPrimTwin->Clear();
 			}
-		};
+		}
+	}
+
+	RefreshStageTask.EnterProgressFrame();
+	FScopedSlowTask RegenerateComponentsTask(PrimsToUpdate.Num(), LOCTEXT("RegeneratingComponents", "Regenerating components"));
+	for (const TPair<UE::FSdfPath, bool>& PrimChangedInfo : PrimsToUpdate)
+	{
+		RegenerateComponentsTask.EnterProgressFrame();
 
-		ReloadAssets(PrimPath, bIsResync);
-		UpdateComponents(PrimPath, bIsResync);
+		const UE::FSdfPath PrimPath = PrimChangedInfo.Key;
+		const bool bIsResync = PrimChangedInfo.Value;
 
+		if (!InfoCache->ContainsInfoAboutPrim(PrimPath))
+		{
+			continue;
+		}
+
+		UE_LOG(LogUsd, Verbose, TEXT("Regenerating components for prim path '%s'"), *PrimPath.GetString());
+
+		TSharedRef<FUsdSchemaTranslationContext> TranslationContext = FUsdStageActorImpl::CreateUsdSchemaTranslationContext(
+			this,
+			PrimPath.GetString()
+		);
+
+		UpdatePrim(PrimPath, bIsResync, *TranslationContext);
+		TranslationContext->CompleteTasks();
+	}
+
+	// Separate pass because we need to update the stage editor even if we don't have info about this prim
+	// anymore, as that's how it refreshes whenever we delete a prim
+	for (const TPair<UE::FSdfPath, bool>& PrimChangedInfo : PrimsToUpdate)
+	{
 		if (HasAuthorityOverStage())
 		{
-			OnPrimChanged.Broadcast(PrimChangedInfo.Key, PrimChangedInfo.Value);
+			OnPrimChanged.Broadcast(PrimChangedInfo.Key.GetString(), PrimChangedInfo.Value);
 		}
 	}
 
@@ -1603,7 +1752,7 @@ UUsdPrimTwin* AUsdStageActor::GetOrCreatePrimTwin(const UE::FSdfPath& UsdPrimPat
 	return UsdPrimTwin;
 }
 
-UUsdPrimTwin* AUsdStageActor::ExpandPrim(const UE::FUsdPrim& Prim, FUsdSchemaTranslationContext& TranslationContext)
+UUsdPrimTwin* AUsdStageActor::ExpandPrim(const UE::FUsdPrim& Prim, bool bResync, FUsdSchemaTranslationContext& TranslationContext)
 {
 	UUsdPrimTwin* UsdPrimTwin = nullptr;
 #if USE_USD_SDK
@@ -1629,7 +1778,7 @@ UUsdPrimTwin* AUsdStageActor::ExpandPrim(const UE::FUsdPrim& Prim, FUsdSchemaTra
 
 	if (TSharedPtr< FUsdSchemaTranslator > SchemaTranslator = UsdSchemasModule.GetTranslatorRegistry().CreateTranslatorForSchema(TranslationContext.AsShared(), UE::FUsdTyped(Prim)))
 	{
-		if (!UsdPrimTwin->SceneComponent.IsValid())
+		if (bResync && !UsdPrimTwin->SceneComponent.IsValid())
 		{
 			UsdPrimTwin->SceneComponent = SchemaTranslator->CreateComponents();
 		}
@@ -1658,7 +1807,7 @@ UUsdPrimTwin* AUsdStageActor::ExpandPrim(const UE::FUsdPrim& Prim, FUsdSchemaTra
 
 		for (const UE::FUsdPrim& ChildPrim : PrimChildren)
 		{
-			ExpandPrim(ChildPrim, TranslationContext);
+			ExpandPrim(ChildPrim, bResync, TranslationContext);
 		}
 	}
 
@@ -1777,17 +1926,8 @@ void AUsdStageActor::UpdatePrim(const UE::FSdfPath& InUsdPrimPath, bool bResync,
 
 	if (UsdPrimPath.IsAbsoluteRootOrPrimPath())
 	{
-		if (bResync)
-		{
-			FString PrimPath = UsdPrimPath.GetString();
-			if (UUsdPrimTwin* UsdPrimTwin = GetRootPrimTwin()->Find(PrimPath))
-			{
-				UsdPrimTwin->Clear();
-			}
-		}
-
 		UE::FUsdPrim PrimToExpand = GetOrLoadUsdStage().GetPrimAtPath(UsdPrimPath);
-		UUsdPrimTwin* UsdPrimTwin = ExpandPrim(PrimToExpand, TranslationContext);
+		ExpandPrim(PrimToExpand, bResync, TranslationContext);
 
 #if WITH_EDITOR
 		if (GIsEditor && GEditor && !IsGarbageCollecting()) // Make sure we're not in standalone either
@@ -3771,21 +3911,35 @@ void AUsdStageActor::OnSkelAnimationBaked(const FString& SkelRootPrimPath)
 #endif // #if USE_USD_SDK
 }
 
-bool AUsdStageActor::LoadAsset(FUsdSchemaTranslationContext& TranslationContext, const UE::FUsdPrim& Prim)
+bool AUsdStageActor::UnloadAssets(const UE::FSdfPath& StartPrimPath, bool bForEntireSubtree)
 {
-	TRACE_CPUPROFILER_EVENT_SCOPE(AUsdStageActor::LoadAsset);
-
 	if (!UsdAssetCache || !InfoCache)
 	{
 		return false;
 	}
 
-	// Clear existing prim/asset association
 	bool bAssetsDiscarded = false;
+
+	TSet<UE::FSdfPath> PrimPathsToRemove;
+	if (bForEntireSubtree)
 	{
-		UE::FSdfPath UsdPrimPath = Prim.GetPrimPath();
-		TSet<TWeakObjectPtr<UObject>> OldAssets = InfoCache->RemoveAllAssetPrimLinks(UsdPrimPath);
+		for (const TPair<UE::FSdfPath, TSet<TWeakObjectPtr<UObject>>>& PrimPathToAssetIt : InfoCache->GetAllAssetPrimLinks())
+		{
+			const UE::FSdfPath& LinkPrimPath = PrimPathToAssetIt.Key;
+			if (LinkPrimPath.HasPrefix(StartPrimPath) || LinkPrimPath == StartPrimPath)
+			{
+				PrimPathsToRemove.Add(LinkPrimPath);
+			}
+		}
+	}
+	else
+	{
+		PrimPathsToRemove.Add(StartPrimPath);
+	}
 
+	for (const UE::FSdfPath& PrimPathToRemove : PrimPathsToRemove)
+	{
+		TSet<TWeakObjectPtr<UObject>> OldAssets = InfoCache->RemoveAllAssetPrimLinks(PrimPathToRemove);
 		for (const TWeakObjectPtr<UObject>& OldAsset : OldAssets)
 		{
 			// If there are any other prim paths linked to this asset that we *won't* be removing/reparsing
@@ -3793,13 +3947,12 @@ bool AUsdStageActor::LoadAsset(FUsdSchemaTranslationContext& TranslationContext,
 			bool bAssetStillReferenced = false;
 			for (const UE::FSdfPath& LinkedPrim : InfoCache->GetPrimsForAsset(OldAsset.Get()))
 			{
-				if (LinkedPrim != UsdPrimPath)
+				if (!PrimPathsToRemove.Contains(LinkedPrim))
 				{
 					bAssetStillReferenced = true;
 					break;
 				}
 			}
-
 			if (bAssetStillReferenced)
 			{
 				continue;
@@ -3809,6 +3962,18 @@ bool AUsdStageActor::LoadAsset(FUsdSchemaTranslationContext& TranslationContext,
 		}
 	}
 
+	return bAssetsDiscarded;
+}
+
+bool AUsdStageActor::LoadAsset(FUsdSchemaTranslationContext& TranslationContext, const UE::FUsdPrim& Prim)
+{
+	TRACE_CPUPROFILER_EVENT_SCOPE(AUsdStageActor::LoadAsset);
+
+	if (!UsdAssetCache || !InfoCache)
+	{
+		return false;
+	}
+
 	// Mark the assets as non transactional so that they don't get serialized in the transaction buffer
 	TGuardValue< EObjectFlags > ContextFlagsGuard(TranslationContext.ObjectFlags, TranslationContext.ObjectFlags & ~RF_Transactional);
 
@@ -3826,7 +3991,7 @@ bool AUsdStageActor::LoadAsset(FUsdSchemaTranslationContext& TranslationContext,
 
 	TranslationContext.CompleteTasks(); // Finish the asset tasks before moving on
 
-	return bAssetsDiscarded || UsdAssetCache->GetNumAssets() != StartNumAssets;
+	return UsdAssetCache->GetNumAssets() != StartNumAssets;
 }
 
 bool AUsdStageActor::LoadAssets(FUsdSchemaTranslationContext& TranslationContext, const UE::FUsdPrim& StartPrim)
@@ -3838,46 +4003,6 @@ bool AUsdStageActor::LoadAssets(FUsdSchemaTranslationContext& TranslationContext
 		return false;
 	}
 
-	// Clear existing prim/asset association
-	bool bAssetsDiscarded = false;
-	{
-		UE::FSdfPath StartPrimPath = StartPrim.GetPrimPath();
-		TSet<UE::FSdfPath> PrimPathsToRemove;
-		for (const TPair<UE::FSdfPath, TSet<TWeakObjectPtr<UObject>>>& PrimPathToAssetIt : InfoCache->GetAllAssetPrimLinks())
-		{
-			const UE::FSdfPath& PrimPath = PrimPathToAssetIt.Key;
-			if (PrimPath.HasPrefix(StartPrimPath) || PrimPath == StartPrimPath)
-			{
-				PrimPathsToRemove.Add(PrimPath);
-			}
-		}
-
-		for (const UE::FSdfPath& PrimPathToRemove : PrimPathsToRemove)
-		{
-			TSet<TWeakObjectPtr<UObject>> OldAssets = InfoCache->RemoveAllAssetPrimLinks(PrimPathToRemove);
-			for (const TWeakObjectPtr<UObject>& OldAsset : OldAssets)
-			{
-				// If there are any other prim paths linked to this asset that we *won't* be removing/reparsing
-				// in here, it means our stage actor as a whole is still "referencing" that asset
-				bool bAssetStillReferenced = false;
-				for (const UE::FSdfPath& LinkedPrim : InfoCache->GetPrimsForAsset(OldAsset.Get()))
-				{
-					if (!PrimPathsToRemove.Contains(LinkedPrim))
-					{
-						bAssetStillReferenced = true;
-						break;
-					}
-				}
-				if (bAssetStillReferenced)
-				{
-					continue;
-				}
-
-				bAssetsDiscarded |= UsdAssetCache->RemoveAssetReference(OldAsset.Get(), this);
-			}
-		}
-	}
-
 	// Mark the assets as non transactional so that they don't get serialized in the transaction buffer
 	TGuardValue< EObjectFlags > ContextFlagsGuard(TranslationContext.ObjectFlags, TranslationContext.ObjectFlags & ~RF_Transactional);
 
@@ -3933,7 +4058,7 @@ bool AUsdStageActor::LoadAssets(FUsdSchemaTranslationContext& TranslationContext
 		CreateAssetsForPrims(AllPrimAssets, AssetsProgress);
 	}
 
-	return bAssetsDiscarded || UsdAssetCache->GetNumAssets() != StartNumAssets;
+	return UsdAssetCache->GetNumAssets() != StartNumAssets;
 }
 
 void AUsdStageActor::AnimatePrims()
diff --git a/Engine/Plugins/Importers/USDImporter/Source/USDStage/Public/USDStageActor.h b/Engine/Plugins/Importers/USDImporter/Source/USDStage/Public/USDStageActor.h
--- a/Engine/Plugins/Importers/USDImporter/Source/USDStage/Public/USDStageActor.h
+++ b/Engine/Plugins/Importers/USDImporter/Source/USDStage/Public/USDStageActor.h
@@ -266,6 +266,8 @@ public:
 	USDSTAGE_API AUsdStageActor();
 
 protected:
+	bool UnloadAssets(const UE::FSdfPath& PrimPath, bool bForEntireSubtree);
+
 	/** Loads the asset for a single prim */
 	bool LoadAsset(FUsdSchemaTranslationContext& TranslationContext, const UE::FUsdPrim& Prim);
 
@@ -278,7 +280,7 @@ protected:
 	UUsdPrimTwin* GetRootPrimTwin();
 	UUsdPrimTwin* GetOrCreatePrimTwin(const UE::FSdfPath& UsdPrimPath);
 
-	UUsdPrimTwin* ExpandPrim(const UE::FUsdPrim& Prim, FUsdSchemaTranslationContext& TranslationContext);
+	UUsdPrimTwin* ExpandPrim(const UE::FUsdPrim& Prim, bool bResync, FUsdSchemaTranslationContext& TranslationContext);
 	void UpdatePrim(const UE::FSdfPath& UsdPrimPath, bool bResync, FUsdSchemaTranslationContext& TranslationContext);
 
 	void OpenUsdStage();
diff --git a/Engine/Plugins/Importers/USDImporter/Source/UnrealUSDWrapper/Private/UsdWrappers/SdfPath.cpp b/Engine/Plugins/Importers/USDImporter/Source/UnrealUSDWrapper/Private/UsdWrappers/SdfPath.cpp
--- a/Engine/Plugins/Importers/USDImporter/Source/UnrealUSDWrapper/Private/UsdWrappers/SdfPath.cpp
+++ b/Engine/Plugins/Importers/USDImporter/Source/UnrealUSDWrapper/Private/UsdWrappers/SdfPath.cpp
@@ -112,6 +112,24 @@ namespace UE
 		return !( *this == Other );
 	}
 
+	bool FSdfPath::operator<(const FSdfPath& Other) const
+	{
+#if USE_USD_SDK
+		return Impl->PxrSdfPath.Get() < Other.Impl->PxrSdfPath.Get();
+#else
+		return false;
+#endif // #if USE_USD_SDK
+	}
+
+	bool FSdfPath::operator<=(const FSdfPath& Other) const
+	{
+#if USE_USD_SDK
+		return Impl->PxrSdfPath.Get() <= Other.Impl->PxrSdfPath.Get();
+#else
+		return false;
+#endif // #if USE_USD_SDK
+	}
+
 	uint32 GetTypeHash( const UE::FSdfPath& Path )
 	{
 		uint32 Result = 0;
diff --git a/Engine/Plugins/Importers/USDImporter/Source/UnrealUSDWrapper/Public/UsdWrappers/SdfPath.h b/Engine/Plugins/Importers/USDImporter/Source/UnrealUSDWrapper/Public/UsdWrappers/SdfPath.h
--- a/Engine/Plugins/Importers/USDImporter/Source/UnrealUSDWrapper/Public/UsdWrappers/SdfPath.h
+++ b/Engine/Plugins/Importers/USDImporter/Source/UnrealUSDWrapper/Public/UsdWrappers/SdfPath.h
@@ -42,6 +42,8 @@ namespace UE
 
 		bool operator==( const FSdfPath& Other ) const;
 		bool operator!=( const FSdfPath& Other ) const;
+		bool operator<(const FSdfPath& Other) const;
+		bool operator<=(const FSdfPath& Other) const;
 
 		friend UNREALUSDWRAPPER_API uint32 GetTypeHash( const UE::FSdfPath& Path );
 		friend UNREALUSDWRAPPER_API FArchive& operator<<( FArchive& Ar, FSdfPath& Path );
