USD: Fix referencing exported assets with LODs when chosing a LOD range.

Use C++ function to add references when exporting, as that fixes the referencer prim's type;
Fix material overrides when exporting a different LOD range;
Fix material overrides ending up on other layers due to mixup with stage edit targets;
Always export meshes with LODs to Xform prims, which child Mesh prims.

diff --git a/Engine/Plugins/Importers/USDImporter/Source/USDExporter/Private/USDConversionBlueprintContext.cpp b/Engine/Plugins/Importers/USDImporter/Source/USDExporter/Private/USDConversionBlueprintContext.cpp
--- a/Engine/Plugins/Importers/USDImporter/Source/USDExporter/Private/USDConversionBlueprintContext.cpp
+++ b/Engine/Plugins/Importers/USDImporter/Source/USDExporter/Private/USDConversionBlueprintContext.cpp
@@ -425,6 +425,27 @@ bool UUsdConversionBlueprintContext::ConvertLandscapeProxyActorMaterial( ALandsc
 #endif // USE_USD_SDK
 }
 
+bool UUsdConversionBlueprintContext::ConvertMaterialOverrides(
+	const UObject* MeshAsset,
+	const TArray<UMaterialInterface*> MaterialOverrides,
+	const FString& PrimPath,
+	int32 LowestLOD,
+	int32 HighestLOD
+)
+{
+#if USE_USD_SDK
+	UE::FUsdPrim Prim = UnrealToUsdImpl::GetPrim(Stage, PrimPath);
+	if (!Prim || !MeshAsset)
+	{
+		return false;
+	}
+
+	return UnrealToUsd::ConvertMaterialOverrides(MeshAsset, MaterialOverrides, Prim, LowestLOD, HighestLOD);
+#else
+	return false;
+#endif	  // USE_USD_SDK
+}
+
 void UUsdConversionBlueprintContext::ReplaceUnrealMaterialsWithBaked( const FFilePath& LayerToAuthorIn, const TMap<FString, FString>& BakedMaterials, bool bIsAssetLayer, bool bUsePayload, bool bRemoveUnrealMaterials )
 {
 #if USE_USD_SDK
diff --git a/Engine/Plugins/Importers/USDImporter/Source/USDExporter/Private/USDConversionBlueprintLibrary.cpp b/Engine/Plugins/Importers/USDImporter/Source/USDExporter/Private/USDConversionBlueprintLibrary.cpp
--- a/Engine/Plugins/Importers/USDImporter/Source/USDExporter/Private/USDConversionBlueprintLibrary.cpp
+++ b/Engine/Plugins/Importers/USDImporter/Source/USDExporter/Private/USDConversionBlueprintLibrary.cpp
@@ -400,6 +400,29 @@ void UUsdConversionBlueprintLibrary::InsertSubLayer( const FString& ParentLayerP
 #endif // USE_USD_SDK
 }
 
+void UUsdConversionBlueprintLibrary::AddReference( const FString& ReferencingStagePath, const FString& ReferencingPrimPath, const FString& TargetStagePath )
+{
+#if USE_USD_SDK
+	TArray<UE::FUsdStage> PreviouslyOpenedStages = UnrealUSDWrapper::GetAllStagesFromCache();
+
+	// Open using the stage cache as it's very likely this stage is already in there anyway
+	UE::FUsdStage ReferencingStage = UnrealUSDWrapper::OpenStage( *ReferencingStagePath, EUsdInitialLoadSet::LoadAll );
+	if ( ReferencingStage )
+	{
+		if ( UE::FUsdPrim ReferencingPrim = ReferencingStage.GetPrimAtPath( UE::FSdfPath( *ReferencingPrimPath ) ) )
+		{
+			UsdUtils::AddReference( ReferencingPrim, *TargetStagePath );
+		}
+	}
+
+	// Cleanup or else the stage cache will keep these stages open forever
+	if ( !PreviouslyOpenedStages.Contains( ReferencingStage ) )
+	{
+		UnrealUSDWrapper::EraseStageFromCache( ReferencingStage );
+	}
+#endif // USE_USD_SDK
+}
+
 void UUsdConversionBlueprintLibrary::AddPayload( const FString& ReferencingStagePath, const FString& ReferencingPrimPath, const FString& TargetStagePath )
 {
 #if USE_USD_SDK
diff --git a/Engine/Plugins/Importers/USDImporter/Source/USDExporter/Public/USDConversionBlueprintContext.h b/Engine/Plugins/Importers/USDImporter/Source/USDExporter/Public/USDConversionBlueprintContext.h
--- a/Engine/Plugins/Importers/USDImporter/Source/USDExporter/Public/USDConversionBlueprintContext.h
+++ b/Engine/Plugins/Importers/USDImporter/Source/USDExporter/Public/USDConversionBlueprintContext.h
@@ -137,6 +137,15 @@ public:
 	bool ConvertLandscapeProxyActorMaterial( ALandscapeProxy* Actor, const FString& PrimPath, const TArray<FPropertyEntry>& PropertiesToBake, const FIntPoint& DefaultTextureSize, const FDirectoryPath& TexturesDir, float TimeCode = 3.402823466e+38F );
 
 public:
+	UFUNCTION( BlueprintCallable, Category = "Conversion utils" )
+	bool ConvertMaterialOverrides(
+		const UObject* MeshAsset,
+		const TArray<UMaterialInterface*> MaterialOverrides,
+		const FString& PrimPath,
+		int32 LowestLOD,
+		int32 HighestLOD
+	);
+
 	/**
 	 * Traverses the context's stage and authors material binding attributes for all `unrealMaterials` that were baked into USD material assets.
 	 * @param LayerToAuthorIn - File path to the layer where the material binding opinions are authored
diff --git a/Engine/Plugins/Importers/USDImporter/Source/USDExporter/Public/USDConversionBlueprintLibrary.h b/Engine/Plugins/Importers/USDImporter/Source/USDExporter/Public/USDConversionBlueprintLibrary.h
--- a/Engine/Plugins/Importers/USDImporter/Source/USDExporter/Public/USDConversionBlueprintLibrary.h
+++ b/Engine/Plugins/Importers/USDImporter/Source/USDExporter/Public/USDConversionBlueprintLibrary.h
@@ -81,6 +81,9 @@ public:
 	UFUNCTION( BlueprintCallable, Category = "Layer utils" )
 	static void InsertSubLayer( const FString& ParentLayerPath, const FString& SubLayerPath, int32 Index = -1 );
 
+	UFUNCTION( BlueprintCallable, Category = "Layer utils" )
+	static void AddReference( const FString& ReferencingStagePath, const FString& ReferencingPrimPath, const FString& TargetStagePath );
+
 	UFUNCTION( BlueprintCallable, Category = "Layer utils" )
 	static void AddPayload( const FString& ReferencingStagePath, const FString& ReferencingPrimPath, const FString& TargetStagePath );
 
diff --git a/Engine/Plugins/Importers/USDImporter/Source/USDUtilities/Private/USDGeomMeshConversion.cpp b/Engine/Plugins/Importers/USDImporter/Source/USDUtilities/Private/USDGeomMeshConversion.cpp
--- a/Engine/Plugins/Importers/USDImporter/Source/USDUtilities/Private/USDGeomMeshConversion.cpp
+++ b/Engine/Plugins/Importers/USDImporter/Source/USDUtilities/Private/USDGeomMeshConversion.cpp
@@ -2183,6 +2183,13 @@ bool UnrealToUsd::ConvertStaticMesh( const UStaticMesh* StaticMesh, pxr::UsdPrim
 		MaterialAssignments.clear();
 	}
 
+	// Do this outside the variant edit context or else it's going to be a weaker opinion than the stuff outside
+	// the variant, and it won't really do anything for UsdPrim if it already exists.
+	// Use an Xform for the parent prim because it will be our defaultPrim for this layer, and our referencer code will
+	// try copying the schema of the defaultPrim onto the referencer prim to make sure they match. If we were typeless
+	// here, so would our referencer and we wouldn't be able to put a transform on it
+	UsdPrim = Stage->DefinePrim(UsdPrim.GetPath(), UnrealToUsd::ConvertToken(bExportMultipleLODs ? TEXT("Xform") : TEXT("Mesh")).Get());
+
 	for ( int32 LODIndex = LowestMeshLOD; LODIndex <= HighestMeshLOD; ++LODIndex )
 	{
 		const FStaticMeshLODResources& RenderMesh = StaticMesh->GetLODForExport( LODIndex );
@@ -2245,8 +2252,6 @@ bool UnrealToUsd::ConvertStaticMesh( const UStaticMesh* StaticMesh, pxr::UsdPrim
 		}
 		else
 		{
-			// Make sure the parent prim has the Mesh schema and add the mesh data directly to it
-			UsdPrim = Stage->DefinePrim( UsdPrim.GetPath(), UnrealToUsd::ConvertToken( TEXT("Mesh") ).Get() );
 			TargetMesh = pxr::UsdGeomMesh{ UsdPrim };
 
 			MaterialPrim = MaterialStage->OverridePrim( UsdPrim.GetPath() );
@@ -2297,6 +2302,9 @@ bool UnrealToUsd::ConvertMeshDescriptions( const TArray<FMeshDescription>& LODIn
 	pxr::SdfPath ParentPrimPath = UsdPrim.GetPath();
 	std::string LowestLODAdded = "";
 
+	// Check the comment on the analogous line on ConvertStaticMesh
+	UsdPrim = Stage->DefinePrim(UsdPrim.GetPath(), UnrealToUsd::ConvertToken(bExportMultipleLODs ? TEXT("Xform") : TEXT("Mesh")).Get());
+
 	for ( int32 LODIndex = 0; LODIndex < NumLODs; ++LODIndex )
 	{
 		const FMeshDescription& MeshDescription = LODIndexToMeshDescription[ LODIndex ];
@@ -2333,8 +2341,6 @@ bool UnrealToUsd::ConvertMeshDescriptions( const TArray<FMeshDescription>& LODIn
 		}
 		else
 		{
-			// Make sure the parent prim has the Mesh schema and add the mesh data directly to it
-			UsdPrim = Stage->DefinePrim( UsdPrim.GetPath(), UnrealToUsd::ConvertToken( TEXT( "Mesh" ) ).Get() );
 			TargetMesh = pxr::UsdGeomMesh{ UsdPrim };
 		}
 
diff --git a/Engine/Plugins/Importers/USDImporter/Source/USDUtilities/Private/USDPrimConversion.cpp b/Engine/Plugins/Importers/USDImporter/Source/USDUtilities/Private/USDPrimConversion.cpp
--- a/Engine/Plugins/Importers/USDImporter/Source/USDUtilities/Private/USDPrimConversion.cpp
+++ b/Engine/Plugins/Importers/USDImporter/Source/USDUtilities/Private/USDPrimConversion.cpp
@@ -2027,283 +2027,50 @@ bool UnrealToUsd::ConvertSceneComponent( const pxr::UsdStageRefPtr& Stage, const
 	return true;
 }
 
-bool UnrealToUsd::ConvertMeshComponent( const pxr::UsdStageRefPtr& Stage, const UMeshComponent* MeshComponent, pxr::UsdPrim& UsdPrim )
+bool UnrealToUsd::ConvertMeshComponent(const pxr::UsdStageRefPtr& Stage, const UMeshComponent* MeshComponent, pxr::UsdPrim& UsdPrim)
 {
-	if ( !UsdPrim || !MeshComponent )
+	if (!UsdPrim || !MeshComponent)
 	{
 		return false;
 	}
 
-	FScopedUsdAllocs Allocs;
+	UObject* MeshAsset = nullptr;
 
-#if WITH_EDITOR
 	// Handle material overrides
-	if ( const UGeometryCacheComponent* GeometryCacheComponent = Cast<const UGeometryCacheComponent>( MeshComponent ) )
+	if (const UGeometryCacheComponent* GeometryCacheComponent = Cast<const UGeometryCacheComponent>(MeshComponent))
 	{
-		if ( const UGeometryCache* GeometryCache = GeometryCacheComponent->GetGeometryCache() )
-		{
-			const TArray<UMaterialInterface*>& Overrides = MeshComponent->OverrideMaterials;
-			for ( int32 MatIndex = 0; MatIndex < Overrides.Num(); ++MatIndex )
-			{
-				UMaterialInterface* Override = Overrides[ MatIndex ];
-				if ( !Override )
-				{
-					continue;
-				}
-
-				// If we have user data this is one of our meshes, so we know exactly the prim that corresponds to each
-				// material slot. Let's use that
-				// const_cast as there's no const access to asset user data on IInterface_AssetUserData, but we won't
-				// modify anything
-				if (const UUsdMeshAssetUserData* UserData =
-					const_cast<UGeometryCache*>(GeometryCache)->GetAssetUserData<UUsdMeshAssetUserData>())
-				{
-					if (const FUsdPrimPathList* SourcePrimPaths = UserData->MaterialSlotToPrimPaths.Find(MatIndex))
-					{
-						for ( const FString& PrimPath : SourcePrimPaths->PrimPaths )
-						{
-							pxr::SdfPath OverridePrimPath = UnrealToUsd::ConvertPath( *PrimPath ).Get();
-							pxr::UsdPrim MeshPrim = Stage->OverridePrim( OverridePrimPath );
-							UE::USDPrimConversionImpl::Private::AuthorMaterialOverride( MeshPrim, Override->GetPathName() );
-						}
-					}
-				}
-				// If we don't, we have to fallback to writing the same prim patterns that the mesh exporters
-				// generate when exporting meshes, so that we can override its opinions. This happens when exporting
-				// geometry cache / static mesh / skeletal mesh components, for example
-				else
-				{
-					pxr::SdfPath OverridePrimPath = UsdPrim.GetPath();
-					pxr::UsdPrim MeshPrim = Stage->OverridePrim(OverridePrimPath);
-					UE::USDPrimConversionImpl::Private::AuthorMaterialOverride(MeshPrim, Override->GetPathName());
-				}
-			}
-		}
+		MeshAsset = GeometryCacheComponent->GetGeometryCache();
 	}
-	else if ( const UStaticMeshComponent* StaticMeshComponent = Cast<const UStaticMeshComponent>( MeshComponent ) )
+	else if (const UStaticMeshComponent* StaticMeshComponent = Cast<const UStaticMeshComponent>(MeshComponent))
 	{
-		if ( UStaticMesh* Mesh = StaticMeshComponent->GetStaticMesh() )
-		{
-			int32 NumLODs = Mesh->GetNumLODs();
-			const bool bHasLODs = NumLODs > 1;
-
-			const UUsdMeshAssetUserData* UserData = Mesh->GetAssetUserData<UUsdMeshAssetUserData>();
-
-			const TArray<UMaterialInterface*>& Overrides = MeshComponent->OverrideMaterials;
-			for ( int32 MatIndex = 0; MatIndex < Overrides.Num(); ++MatIndex )
-			{
-				UMaterialInterface* Override = Overrides[MatIndex];
-				if ( !Override )
-				{
-					continue;
-				}
-
-				for ( int32 LODIndex = 0; LODIndex < NumLODs; ++LODIndex )
-				{
-					int32 NumSections = Mesh->GetNumSections( LODIndex );
-					const bool bHasSubsets = NumSections > 1;
-
-					for ( int32 SectionIndex = 0; SectionIndex < NumSections; ++SectionIndex )
-					{
-						int32 SectionMatIndex = Mesh->GetSectionInfoMap().Get( LODIndex, SectionIndex ).MaterialIndex;
-						if ( SectionMatIndex != MatIndex )
-						{
-							continue;
-						}
-
-						if (UserData)
-						{
-							if (const FUsdPrimPathList* SourcePrimPaths = UserData->MaterialSlotToPrimPaths.Find(SectionMatIndex))
-							{
-								for ( const FString& PrimPath : SourcePrimPaths->PrimPaths )
-								{
-									pxr::SdfPath OverridePrimPath = UnrealToUsd::ConvertPath( *PrimPath ).Get();
-									pxr::UsdPrim MeshPrim = Stage->OverridePrim( OverridePrimPath );
-									UE::USDPrimConversionImpl::Private::AuthorMaterialOverride( MeshPrim, Override->GetPathName() );
-								}
-							}
-						}
-						else
-						{
-							pxr::SdfPath OverridePrimPath = UsdPrim.GetPath();
-
-							// If we have only 1 LOD, the asset's DefaultPrim will be the Mesh prim directly.
-							// If we have multiple, the default prim won't have any schema, but will contain separate
-							// Mesh prims for each LOD named "LOD0", "LOD1", etc., switched via a "LOD" variant set
-							if (bHasLODs)
-							{
-								OverridePrimPath = OverridePrimPath.AppendPath(
-									UnrealToUsd::ConvertPath(*FString::Printf(TEXT("LOD%d"), LODIndex)).Get()
-								);
-							}
-
-							// If our LOD has only one section, its material assignment will be authored directly on the Mesh prim.
-							// If it has more than one material slot, we'll author UsdGeomSubset for each LOD Section, and author the material
-							// assignment there
-							if (bHasSubsets)
-							{
-								// Assume the UE sections are in the same order as the USD ones
-								std::vector<pxr::UsdGeomSubset> GeomSubsets = pxr::UsdShadeMaterialBindingAPI(UsdPrim).GetMaterialBindSubsets();
-								if (SectionIndex < GeomSubsets.size())
-								{
-									OverridePrimPath = OverridePrimPath.AppendChild(GeomSubsets[SectionIndex].GetPrim().GetName());
-								}
-								else
-								{
-									OverridePrimPath = OverridePrimPath.AppendPath(
-										UnrealToUsd::ConvertPath(*FString::Printf(TEXT("Section%d"), SectionIndex)).Get()
-									);
-								}
-							}
-
-							pxr::UsdPrim MeshPrim = Stage->OverridePrim(OverridePrimPath);
-							UE::USDPrimConversionImpl::Private::AuthorMaterialOverride(MeshPrim, Override->GetPathName());
-						}
-					}
-				}
-			}
-		}
+		MeshAsset = StaticMeshComponent->GetStaticMesh();
 	}
-	else if ( const USkinnedMeshComponent* SkinnedMeshComponent = Cast<const USkinnedMeshComponent>( MeshComponent ) )
+	else if (const USkinnedMeshComponent* SkinnedMeshComponent = Cast<const USkinnedMeshComponent>(MeshComponent))
 	{
-		pxr::UsdSkelRoot SkelRoot{ UsdPrim };
-		if ( !SkelRoot )
+		FScopedUsdAllocs Allocs;
+		pxr::UsdSkelRoot SkelRoot{UsdPrim};
+		if (!SkelRoot)
 		{
 			return false;
 		}
 
-		if ( const USkeletalMesh* SkeletalMesh = Cast<USkeletalMesh>( SkinnedMeshComponent->GetSkinnedAsset() ) )
-		{
-			FSkeletalMeshRenderData* RenderData = SkeletalMesh->GetResourceForRendering();
-			if ( !RenderData )
-			{
-				return false;
-			}
-
-			const UUsdMeshAssetUserData* UserData =
-				const_cast<USkeletalMesh*>(SkeletalMesh)->GetAssetUserData<UUsdMeshAssetUserData>();
-
-			TIndirectArray<FSkeletalMeshLODRenderData>& LodRenderData = RenderData->LODRenderData;
-			if ( LodRenderData.Num() == 0 )
-			{
-				return false;
-			}
-
-			int32 NumLODs = SkeletalMesh->GetLODNum();
-			const bool bHasLODs = NumLODs > 1;
-
-			FString MeshName;
-			if ( !bHasLODs )
-			{
-				for ( const pxr::UsdPrim& Child : UsdPrim.GetChildren() )
-				{
-					if ( pxr::UsdGeomMesh Mesh{ Child } )
-					{
-						MeshName = UsdToUnreal::ConvertToken( Child.GetName() );
-					}
-				}
-			}
-
-			const TArray<UMaterialInterface*>& Overrides = MeshComponent->OverrideMaterials;
-			for ( int32 MatIndex = 0; MatIndex < Overrides.Num(); ++MatIndex )
-			{
-				UMaterialInterface* Override = Overrides[ MatIndex ];
-				if ( !Override )
-				{
-					continue;
-				}
-
-				for ( int32 LODIndex = 0; LODIndex < NumLODs; ++LODIndex )
-				{
-					if ( !LodRenderData.IsValidIndex( LODIndex ) )
-					{
-						continue;
-					}
-
-					const FSkeletalMeshLODInfo* LODInfo = SkeletalMesh->GetLODInfo( LODIndex );
-
-					const TArray<FSkelMeshRenderSection>& Sections = LodRenderData[ LODIndex ].RenderSections;
-					int32 NumSections = Sections.Num();
-					const bool bHasSubsets = NumSections > 1;
-
-					for ( int32 SectionIndex = 0; SectionIndex < NumSections; ++SectionIndex )
-					{
-						int32 SectionMatIndex = Sections[ SectionIndex ].MaterialIndex;
-
-						// If we have a LODInfo map, we need to reroute the material index through it
-						if ( LODInfo && LODInfo->LODMaterialMap.IsValidIndex( SectionIndex ) )
-						{
-							SectionMatIndex = LODInfo->LODMaterialMap[ SectionIndex ];
-						}
-
-						if (SectionMatIndex != MatIndex)
-						{
-							continue;
-						}
-
-						if (UserData)
-						{
-							if (const FUsdPrimPathList* SourcePrimPaths = UserData->MaterialSlotToPrimPaths.Find(SectionMatIndex))
-							{
-								for (const FString& PrimPath : SourcePrimPaths->PrimPaths)
-								{
-									pxr::SdfPath OverridePrimPath = UnrealToUsd::ConvertPath(*PrimPath).Get();
-									pxr::UsdPrim MeshPrim = Stage->OverridePrim(OverridePrimPath);
-									UE::USDPrimConversionImpl::Private::AuthorMaterialOverride(MeshPrim, Override->GetPathName());
-								}
-							}
-						}
-						else
-						{
-							pxr::SdfPath OverridePrimPath = UsdPrim.GetPath();
-
-							// If we have only 1 LOD, the asset's DefaultPrim will be a SkelRoot, and the Mesh will be a subprim
-							// with the same name.If we have the default prim is also the SkelRoot, but will contain separate
-							// Mesh prims for each LOD named "LOD0", "LOD1", etc., switched via a "LOD" variant set
-							if (bHasLODs)
-							{
-								OverridePrimPath = OverridePrimPath.AppendPath(
-									UnrealToUsd::ConvertPath(*FString::Printf(TEXT("LOD%d"), LODIndex)).Get()
-								);
-							}
-							else
-							{
-								OverridePrimPath = OverridePrimPath.AppendElementString(UnrealToUsd::ConvertString(*MeshName).Get());
-							}
-
-							// If our LOD has only one section, its material assignment will be authored directly on the Mesh prim.
-							// If it has more than one material slot, we'll author UsdGeomSubset for each LOD Section, and author the material
-							// assignment there
-							if (bHasSubsets)
-							{
-								// Assume the UE sections are in the same order as the USD ones
-								std::vector<pxr::UsdGeomSubset> GeomSubsets = pxr::UsdShadeMaterialBindingAPI(UsdPrim).GetMaterialBindSubsets();
-								if (SectionIndex < GeomSubsets.size())
-								{
-									OverridePrimPath = OverridePrimPath.AppendChild(GeomSubsets[SectionIndex].GetPrim().GetName());
-								}
-								else
-								{
-									OverridePrimPath = OverridePrimPath.AppendPath(
-										UnrealToUsd::ConvertPath(*FString::Printf(TEXT("Section%d"), SectionIndex)).Get()
-									);
-								}
-							}
+		MeshAsset = SkinnedMeshComponent->GetSkinnedAsset();
+	}
 
-							pxr::UsdPrim MeshPrim = Stage->OverridePrim(OverridePrimPath);
-							UE::USDPrimConversionImpl::Private::AuthorMaterialOverride(MeshPrim, Override->GetPathName());
-						}
-					}
-				}
-			}
-		}
+	// Component doesn't have any mesh so this function doesn't need to do anything
+	if (!MeshAsset)
+	{
+		return true;
 	}
-#endif // WITH_EDITOR
 
-	return true;
+	return UnrealToUsd::ConvertMaterialOverrides(MeshAsset, MeshComponent->OverrideMaterials, UsdPrim);
 }
 
-bool UnrealToUsd::ConvertHierarchicalInstancedStaticMeshComponent( const UHierarchicalInstancedStaticMeshComponent* HISMComponent, pxr::UsdPrim& UsdPrim, double TimeCode )
+bool UnrealToUsd::ConvertHierarchicalInstancedStaticMeshComponent(
+	const UHierarchicalInstancedStaticMeshComponent* HISMComponent,
+	pxr::UsdPrim& UsdPrim,
+	double TimeCode
+)
 {
 	using namespace pxr;
 
@@ -2381,6 +2148,308 @@ bool UnrealToUsd::ConvertHierarchicalInstancedStaticMeshComponent( const UHierar
 	return true;
 }
 
+bool UnrealToUsd::ConvertMaterialOverrides(
+	const UObject* MeshAsset,
+	const TArray<UMaterialInterface*> MaterialOverrides,
+	pxr::UsdPrim& UsdPrim,
+	int32 LowestLOD,
+	int32 HighestLOD
+)
+{
+	if (!UsdPrim)
+	{
+		return false;
+	}
+
+	FScopedUsdAllocs Allocs;
+	pxr::UsdStageRefPtr Stage = UsdPrim.GetStage();
+
+	const bool bAllLODs = LowestLOD == INDEX_NONE && HighestLOD == INDEX_NONE;
+
+	if (const UGeometryCache* GeometryCache = Cast<const UGeometryCache>(MeshAsset))
+	{
+		for (int32 MatIndex = 0; MatIndex < MaterialOverrides.Num(); ++MatIndex)
+		{
+			UMaterialInterface* Override = MaterialOverrides[MatIndex];
+			if (!Override)
+			{
+				continue;
+			}
+
+			// If we have user data this is one of our meshes, so we know exactly the prim that corresponds to each
+			// material slot. Let's use that
+			// const_cast as there's no const access to asset user data on IInterface_AssetUserData, but we won't
+			// modify anything
+			if (const UUsdMeshAssetUserData* UserData = const_cast<UGeometryCache*>(GeometryCache)->GetAssetUserData<UUsdMeshAssetUserData>())
+			{
+				if (const FUsdPrimPathList* SourcePrimPaths = UserData->MaterialSlotToPrimPaths.Find(MatIndex))
+				{
+					for (const FString& PrimPath : SourcePrimPaths->PrimPaths)
+					{
+						pxr::SdfPath OverridePrimPath = UnrealToUsd::ConvertPath(*PrimPath).Get();
+						pxr::UsdPrim MeshPrim = Stage->OverridePrim(OverridePrimPath);
+						UE::USDPrimConversionImpl::Private::AuthorMaterialOverride(MeshPrim, Override->GetPathName());
+					}
+				}
+			}
+			// If we don't, we have to fallback to writing the same prim patterns that the mesh exporters
+			// generate when exporting meshes, so that we can override its opinions. This happens when exporting
+			// geometry cache / static mesh / skeletal mesh components, for example
+			else
+			{
+				pxr::SdfPath OverridePrimPath = UsdPrim.GetPath();
+				pxr::UsdPrim MeshPrim = Stage->OverridePrim(OverridePrimPath);
+				UE::USDPrimConversionImpl::Private::AuthorMaterialOverride(MeshPrim, Override->GetPathName());
+			}
+		}
+	}
+	else if (const UStaticMesh* StaticMesh = Cast<const UStaticMesh>(MeshAsset))
+	{
+		int32 NumLODs = StaticMesh->GetNumLODs();
+		if (bAllLODs)
+		{
+			HighestLOD = NumLODs - 1;
+			LowestLOD = 0;
+		}
+		else
+		{
+			// Make sure they're both >= 0 (the options dialog slider is clamped, but this may be called directly)
+			LowestLOD = FMath::Clamp(LowestLOD, 0, NumLODs - 1);
+			HighestLOD = FMath::Clamp(HighestLOD, 0, NumLODs - 1);
+
+			// Make sure Lowest <= Highest
+			int32 Temp = FMath::Min(LowestLOD, HighestLOD);
+			HighestLOD = FMath::Max(LowestLOD, HighestLOD);
+			LowestLOD = Temp;
+
+			// Make sure there's at least one LOD
+			NumLODs = FMath::Max(HighestLOD - LowestLOD + 1, 1);
+		}
+		const bool bHasLODs = NumLODs > 1;
+
+		const UUsdMeshAssetUserData* UserData = const_cast<UStaticMesh*>(StaticMesh)->GetAssetUserData<UUsdMeshAssetUserData>();
+
+		for (int32 MatIndex = 0; MatIndex < MaterialOverrides.Num(); ++MatIndex)
+		{
+			UMaterialInterface* Override = MaterialOverrides[MatIndex];
+			if (!Override)
+			{
+				continue;
+			}
+
+			for (int32 LODIndex = LowestLOD; LODIndex <= HighestLOD; ++LODIndex)
+			{
+				int32 NumSections = StaticMesh->GetNumSections(LODIndex);
+				const bool bHasSubsets = NumSections > 1;
+
+				for (int32 SectionIndex = 0; SectionIndex < NumSections; ++SectionIndex)
+				{
+					int32 SectionMatIndex = StaticMesh->GetSectionInfoMap().Get(LODIndex, SectionIndex).MaterialIndex;
+					if (SectionMatIndex != MatIndex)
+					{
+						continue;
+					}
+
+					if (UserData)
+					{
+						if (const FUsdPrimPathList* SourcePrimPaths = UserData->MaterialSlotToPrimPaths.Find(SectionMatIndex))
+						{
+							for (const FString& PrimPath : SourcePrimPaths->PrimPaths)
+							{
+								pxr::SdfPath OverridePrimPath = UnrealToUsd::ConvertPath(*PrimPath).Get();
+								pxr::UsdPrim MeshPrim = Stage->OverridePrim(OverridePrimPath);
+								UE::USDPrimConversionImpl::Private::AuthorMaterialOverride(MeshPrim, Override->GetPathName());
+							}
+						}
+					}
+					else
+					{
+						pxr::SdfPath OverridePrimPath = UsdPrim.GetPath();
+
+						// If we have only 1 LOD, the asset's DefaultPrim will be the Mesh prim directly.
+						// If we have multiple, the default prim won't have any schema, but will contain separate
+						// Mesh prims for each LOD named "LOD0", "LOD1", etc., switched via a "LOD" variant set
+						if (bHasLODs)
+						{
+							OverridePrimPath = OverridePrimPath.AppendPath(UnrealToUsd::ConvertPath(*FString::Printf(TEXT("LOD%d"), LODIndex)).Get());
+						}
+
+						// If our LOD has only one section, its material assignment will be authored directly on the Mesh prim.
+						// If it has more than one material slot, we'll author UsdGeomSubset for each LOD Section, and author the material
+						// assignment there
+						if (bHasSubsets)
+						{
+							// Assume the UE sections are in the same order as the USD ones
+							std::vector<pxr::UsdGeomSubset> GeomSubsets = pxr::UsdShadeMaterialBindingAPI(UsdPrim).GetMaterialBindSubsets();
+							if (SectionIndex < static_cast<int32>(GeomSubsets.size()))
+							{
+								OverridePrimPath = OverridePrimPath.AppendChild(GeomSubsets[SectionIndex].GetPrim().GetName());
+							}
+							else
+							{
+								OverridePrimPath = OverridePrimPath.AppendPath(
+									UnrealToUsd::ConvertPath(*FString::Printf(TEXT("Section%d"), SectionIndex)).Get()
+								);
+							}
+						}
+
+						pxr::UsdPrim MeshPrim = Stage->OverridePrim(OverridePrimPath);
+						UE::USDPrimConversionImpl::Private::AuthorMaterialOverride(MeshPrim, Override->GetPathName());
+					}
+				}
+			}
+		}
+	}
+	else if (const USkeletalMesh* SkeletalMesh = Cast<const USkeletalMesh>(MeshAsset))
+	{
+		FSkeletalMeshRenderData* RenderData = SkeletalMesh->GetResourceForRendering();
+		if (!RenderData)
+		{
+			return false;
+		}
+
+		const UUsdMeshAssetUserData* UserData = const_cast<USkeletalMesh*>(SkeletalMesh)->GetAssetUserData<UUsdMeshAssetUserData>();
+
+		TIndirectArray<FSkeletalMeshLODRenderData>& LodRenderData = RenderData->LODRenderData;
+		if (LodRenderData.Num() == 0)
+		{
+			return false;
+		}
+
+		int32 NumLODs = SkeletalMesh->GetLODNum();
+		if (bAllLODs)
+		{
+			HighestLOD = NumLODs - 1;
+			LowestLOD = 0;
+		}
+		else
+		{
+			// Make sure they're both >= 0 (the options dialog slider is clamped, but this may be called directly)
+			LowestLOD = FMath::Clamp(LowestLOD, 0, NumLODs - 1);
+			HighestLOD = FMath::Clamp(HighestLOD, 0, NumLODs - 1);
+
+			// Make sure Lowest <= Highest
+			int32 Temp = FMath::Min(LowestLOD, HighestLOD);
+			HighestLOD = FMath::Max(LowestLOD, HighestLOD);
+			LowestLOD = Temp;
+
+			// Make sure there's at least one LOD
+			NumLODs = FMath::Max(HighestLOD - LowestLOD + 1, 1);
+		}
+		const bool bHasLODs = NumLODs > 1;
+
+		FString MeshName;
+		if (!bHasLODs)
+		{
+			for (const pxr::UsdPrim& Child : UsdPrim.GetChildren())
+			{
+				if (pxr::UsdGeomMesh Mesh{Child})
+				{
+					MeshName = UsdToUnreal::ConvertToken(Child.GetName());
+					break;
+				}
+			}
+		}
+
+		for (int32 MatIndex = 0; MatIndex < MaterialOverrides.Num(); ++MatIndex)
+		{
+			UMaterialInterface* Override = MaterialOverrides[MatIndex];
+			if (!Override)
+			{
+				continue;
+			}
+
+			for (int32 LODIndex = 0; LODIndex < NumLODs; ++LODIndex)
+			{
+				if (!LodRenderData.IsValidIndex(LODIndex))
+				{
+					continue;
+				}
+
+				const FSkeletalMeshLODInfo* LODInfo = SkeletalMesh->GetLODInfo(LODIndex);
+
+				const TArray<FSkelMeshRenderSection>& Sections = LodRenderData[LODIndex].RenderSections;
+				int32 NumSections = Sections.Num();
+				const bool bHasSubsets = NumSections > 1;
+
+				for (int32 SectionIndex = 0; SectionIndex < NumSections; ++SectionIndex)
+				{
+					int32 SectionMatIndex = Sections[SectionIndex].MaterialIndex;
+
+					// If we have a LODInfo map, we need to reroute the material index through it
+					if (LODInfo && LODInfo->LODMaterialMap.IsValidIndex(SectionIndex))
+					{
+						SectionMatIndex = LODInfo->LODMaterialMap[SectionIndex];
+					}
+
+					if (SectionMatIndex != MatIndex)
+					{
+						continue;
+					}
+
+					if (UserData)
+					{
+						if (const FUsdPrimPathList* SourcePrimPaths = UserData->MaterialSlotToPrimPaths.Find(SectionMatIndex))
+						{
+							for (const FString& PrimPath : SourcePrimPaths->PrimPaths)
+							{
+								pxr::SdfPath OverridePrimPath = UnrealToUsd::ConvertPath(*PrimPath).Get();
+								pxr::UsdPrim MeshPrim = Stage->OverridePrim(OverridePrimPath);
+								UE::USDPrimConversionImpl::Private::AuthorMaterialOverride(MeshPrim, Override->GetPathName());
+							}
+						}
+					}
+					else
+					{
+						pxr::SdfPath OverridePrimPath = UsdPrim.GetPath();
+
+						// If we have only 1 LOD, the asset's DefaultPrim will be a SkelRoot, and the Mesh will be a subprim
+						// with the same name. If we have multiple LODS, the default prim is also the SkelRoot but will contain separate
+						// Mesh prims for each LOD named "LOD0", "LOD1", etc., switched via a "LOD" variant set
+						if (bHasLODs)
+						{
+							OverridePrimPath = OverridePrimPath.AppendPath(UnrealToUsd::ConvertPath(*FString::Printf(TEXT("LOD%d"), LODIndex)).Get());
+						}
+						else
+						{
+							OverridePrimPath = OverridePrimPath.AppendElementString(UnrealToUsd::ConvertString(*MeshName).Get());
+						}
+
+						// If our LOD has only one section, its material assignment will be authored directly on the Mesh prim.
+						// If it has more than one material slot, we'll author UsdGeomSubset for each LOD Section, and author the material
+						// assignment there
+						if (bHasSubsets)
+						{
+							// Assume the UE sections are in the same order as the USD ones
+							std::vector<pxr::UsdGeomSubset> GeomSubsets = pxr::UsdShadeMaterialBindingAPI(UsdPrim).GetMaterialBindSubsets();
+							if (SectionIndex < static_cast<int32>(GeomSubsets.size()))
+							{
+								OverridePrimPath = OverridePrimPath.AppendChild(GeomSubsets[SectionIndex].GetPrim().GetName());
+							}
+							else
+							{
+								OverridePrimPath = OverridePrimPath.AppendPath(
+									UnrealToUsd::ConvertPath(*FString::Printf(TEXT("Section%d"), SectionIndex)).Get()
+								);
+							}
+						}
+
+						pxr::UsdPrim MeshPrim = Stage->OverridePrim(OverridePrimPath);
+						UE::USDPrimConversionImpl::Private::AuthorMaterialOverride(MeshPrim, Override->GetPathName());
+					}
+				}
+			}
+		}
+	}
+	else
+	{
+		ensure(false);
+		return false;
+	}
+
+	return true;
+}
+
 bool UnrealToUsd::ConvertXformable( const FTransform& RelativeTransform, pxr::UsdPrim& UsdPrim, double TimeCode )
 {
 	if ( !UsdPrim )
diff --git a/Engine/Plugins/Importers/USDImporter/Source/USDUtilities/Public/USDPrimConversion.h b/Engine/Plugins/Importers/USDImporter/Source/USDUtilities/Public/USDPrimConversion.h
--- a/Engine/Plugins/Importers/USDImporter/Source/USDUtilities/Public/USDPrimConversion.h
+++ b/Engine/Plugins/Importers/USDImporter/Source/USDUtilities/Public/USDPrimConversion.h
@@ -132,9 +132,31 @@ namespace UnrealToUsd
 
 	USDUTILITIES_API bool ConvertSceneComponent( const pxr::UsdStageRefPtr& Stage, const USceneComponent* SceneComponent, pxr::UsdPrim& UsdPrim );
 
-	USDUTILITIES_API bool ConvertMeshComponent( const pxr::UsdStageRefPtr& Stage, const UMeshComponent* MeshComponent, pxr::UsdPrim& UsdPrim );
+	USDUTILITIES_API bool ConvertMeshComponent(
+		const pxr::UsdStageRefPtr& Stage,
+		const UMeshComponent* MeshComponent,
+		pxr::UsdPrim& UsdPrim
+	);
+
+	USDUTILITIES_API bool ConvertHierarchicalInstancedStaticMeshComponent(
+		const UHierarchicalInstancedStaticMeshComponent* HISMComponent,
+		pxr::UsdPrim& UsdPrim,
+		double TimeCode = UsdUtils::GetDefaultTimeCode()
+	);
 
-	USDUTILITIES_API bool ConvertHierarchicalInstancedStaticMeshComponent( const UHierarchicalInstancedStaticMeshComponent* HISMComponent, pxr::UsdPrim& UsdPrim, double TimeCode = UsdUtils::GetDefaultTimeCode() );
+	/**
+	 * Sets the material overrides as opinions on MeshPrim, taking care to figure out whether it needs to create LOD or subsection prims, and
+	 * which ones given the chosen export LOD range.
+	 * Leaving the LOD range at default of INDEX_NONE will author the material override opinions assuming all LODs of the MeshAsset were exported.
+	 * MeshAsset should be an UStaticMesh, USkeletalMesh or UGeometryCache.
+	 */
+	USDUTILITIES_API bool ConvertMaterialOverrides(
+		const UObject* MeshAsset,
+		const TArray<UMaterialInterface*> MaterialOverrides,
+		pxr::UsdPrim& UsdPrim,
+		int32 LowestLOD = INDEX_NONE,
+		int32 HighestLOD = INDEX_NONE
+	);
 
 	/**
 	 * Converts a UMovieScene3DTransformTrack to a UsdGeomXformable
