USD: Add cvar to allow translating unreferenced Material prims when opening the stage.

diff --git a/Engine/Plugins/Importers/USDImporter/Source/USDStage/Private/USDStageActor.cpp b/Engine/Plugins/Importers/USDImporter/Source/USDStage/Private/USDStageActor.cpp
--- a/Engine/Plugins/Importers/USDImporter/Source/USDStage/Private/USDStageActor.cpp
+++ b/Engine/Plugins/Importers/USDImporter/Source/USDStage/Private/USDStageActor.cpp
@@ -111,6 +111,12 @@ static FAutoConsoleVariableRef CVarRegenerateSkeletalAssetsOnControlRigBake(
 	GRegenerateSkeletalAssetsOnControlRigBake,
 	TEXT("Whether to regenerate the assets associated with a SkelRoot (mesh, skeleton, anim sequence, etc.) whenever we modify Control Rig tracks. The USD Stage itself is always updated however."));
 
+static bool GTranslateOnlyUsedMaterialsWhenOpeningStage = true;
+static FAutoConsoleVariableRef CVarTranslateOnlyUsedMaterialsWhenOpeningStage(
+	TEXT("USD.TranslateOnlyUsedMaterialsWhenOpeningStage"),
+	GTranslateOnlyUsedMaterialsWhenOpeningStage,
+	TEXT("If enabled, only Material prims bound by at least one Mesh are translated into Unreal material assets. If disabled, all Material prims are translated into Unreal material assets."));
+
 static const EObjectFlags DefaultObjFlag = EObjectFlags::RF_Transactional | EObjectFlags::RF_Transient;
 
 AUsdStageActor::FOnActorLoaded AUsdStageActor::OnActorLoaded;
@@ -135,7 +141,7 @@ struct FUsdStageActorImpl
 		TranslationContext->RootMotionHandling = StageActor->RootMotionHandling;
 		TranslationContext->BlendShapesByPath = &StageActor->BlendShapesByPath;
 		TranslationContext->InfoCache = StageActor->InfoCache;
-		TranslationContext->bTranslateOnlyUsedMaterials = true;  // There is no point in turning this off when just opening a stage
+		TranslationContext->bTranslateOnlyUsedMaterials = GTranslateOnlyUsedMaterialsWhenOpeningStage;
 
 		// Its more convenient to toggle between variants using the USDStage window, as opposed to parsing LODs
 		TranslationContext->bAllowInterpretingLODs = false;
