USD: Add support for registering material users from collection-based material bindings;

Register material users for inherited direct material bindings.

diff --git a/Engine/Plugins/Importers/USDImporter/Source/USDSchemas/Private/USDInfoCache.cpp b/Engine/Plugins/Importers/USDImporter/Source/USDSchemas/Private/USDInfoCache.cpp
--- a/Engine/Plugins/Importers/USDImporter/Source/USDSchemas/Private/USDInfoCache.cpp
+++ b/Engine/Plugins/Importers/USDImporter/Source/USDSchemas/Private/USDInfoCache.cpp
@@ -21,6 +21,7 @@
 #if USE_USD_SDK
 #include "USDIncludesStart.h"
 	#include "pxr/usd/sdf/path.h"
+	#include "pxr/usd/usd/collectionAPI.h"
 	#include "pxr/usd/usd/prim.h"
 	#include "pxr/usd/usd/primCompositionQuery.h"
 	#include "pxr/usd/usdGeom/mesh.h"
@@ -489,7 +490,8 @@ namespace UE::USDInfoCacheImpl::Private
 		TMap< UE::FSdfPath, TArray<UsdUtils::FUsdPrimMaterialSlot>>& InOutSubtreeToMaterialSlots,
 		TArray< FString >& InOutPointInstancerPaths,
 		uint64& OutSubtreeVertexCount,
-		TArray<UsdUtils::FUsdPrimMaterialSlot>& OutSubtreeSlots
+		TArray<UsdUtils::FUsdPrimMaterialSlot>& OutSubtreeSlots,
+		bool bPossibleInheritedBindings
 	)
 	{
 		TRACE_CPUPROFILER_EVENT_SCOPE( UE::USDInfoCacheImpl::Private::RecursivePropagateVertexAndMaterialSlotCounts );
@@ -503,16 +505,101 @@ namespace UE::USDInfoCacheImpl::Private
 
 		pxr::SdfPath UsdPrimPath = UsdPrim.GetPrimPath();
 		UE::FSdfPath PrimPath{UsdPrimPath};
+		pxr::UsdStageRefPtr Stage = UsdPrim.GetStage();
 
+		// Material bindings are inherited down to child prims, so if we detect a binding on a parent Xform,
+		// we should register the child Mesh prims as users of the material too (regardless of collapsing).
+		// Note that we only consider this for direct bindings: Collection-based bindings will already provide the exhaustive
+		// list of all the prims that they should apply to when we call ComputeIncludedPaths
+		bool bPrimHasInheritableMaterialBindings = false;
+
+		// Register material users
 		if (!UsdPrim.IsPseudoRoot())
 		{
 			pxr::UsdShadeMaterial ShadeMaterial;
+			TMap<UE::FSdfPath, TSet<UE::FSdfPath>> NewMaterialUsers;
 
 			pxr::UsdShadeMaterialBindingAPI BindingAPI{UsdPrim};
-			if (BindingAPI)
+			if (BindingAPI || bPossibleInheritedBindings)
 			{
+				// Check for material users via collections-based material bindings
+				{
+					// When retrieving the relationships directly we'll always need to check the universal render context
+					// manually, as it won't automatically "compute the fallback" for us like when we ComputeBoundMaterial()
+					std::unordered_set<pxr::TfToken, pxr::TfHash> MaterialPurposeTokens{
+						MaterialPurposeToken,
+						pxr::UsdShadeTokens->universalRenderContext
+					};
+					for(const pxr::TfToken& SomeMaterialPurposeToken : MaterialPurposeTokens)
+					{
+						// Each one of those relationships must have two targets: A collection, and a material
+						for (const pxr::UsdRelationship& Rel : BindingAPI.GetCollectionBindingRels(SomeMaterialPurposeToken))
+						{
+							const pxr::SdfPath* CollectionPath = nullptr;
+							const pxr::SdfPath* MaterialPath = nullptr;
+
+							std::vector<pxr::SdfPath> PathVector;
+							if (Rel.GetTargets(&PathVector))
+							{
+								for (const pxr::SdfPath& Path : PathVector)
+								{
+									if (Path.IsPrimPath())
+									{
+										MaterialPath = &Path;
+									}
+									else if (Path.IsPropertyPath())
+									{
+										CollectionPath = &Path;
+									}
+								}
+							}
+
+							if (!CollectionPath || !MaterialPath || PathVector.size() != 2)
+							{
+								// Emit this warning here as USD doesn't seem to and just seems to just ignores this relationship instead
+								UE_LOG(
+									LogUsd,
+									Warning,
+									TEXT("Prim '%s' describes a collection-based material binding, but the relationship '%s' is invalid: It should contain exactly one Material path and one path to a collection relationship"),
+									*PrimPath.GetString(),
+									*UsdToUnreal::ConvertToken(Rel.GetName())
+								);
+								continue;
+							}
+
+							if (pxr::UsdCollectionAPI Collection = pxr::UsdCollectionAPI::Get(Stage, *CollectionPath))
+							{
+								std::set<pxr::SdfPath> IncludedPaths = Collection.ComputeIncludedPaths(Collection.ComputeMembershipQuery(), Stage);
+
+								for (const pxr::SdfPath& IncludedPath : IncludedPaths)
+								{
+									NewMaterialUsers.FindOrAdd(UE::FSdfPath{*MaterialPath}).Add(UE::FSdfPath{IncludedPath});
+								}
+							}
+							else
+							{
+								UE_LOG(
+									LogUsd,
+									Warning,
+									TEXT("Failed to find collection at path '%s' when processing collection-based material bindings on prim '%s'"),
+									*UsdToUnreal::ConvertPath(CollectionPath->GetPrimPath()),
+									*PrimPath.GetString()
+								);
+							}
+						}
+					}
+				}
+
+				// Check for material bindings directly for this prim
 				ShadeMaterial = BindingAPI.ComputeBoundMaterial(MaterialPurposeToken);
+				if (ShadeMaterial)
+				{
+					bPrimHasInheritableMaterialBindings = true;
+					NewMaterialUsers.FindOrAdd(UE::FSdfPath{ShadeMaterial.GetPrim().GetPath()}).Add(UE::FSdfPath{UsdPrimPath});
+				}
 			}
+			// Temporary fallback for prims that don't have the MaterialBindingAPI but do have the relationship.
+			// USD will emit a warning for these though
 			else if (pxr::UsdRelationship Relationship = UsdPrim.GetRelationship(pxr::UsdShadeTokens->materialBinding))
 			{
 				pxr::SdfPathVector Targets;
@@ -521,44 +608,58 @@ namespace UE::USDInfoCacheImpl::Private
 				if (Targets.size() > 0)
 				{
 					const pxr::SdfPath& TargetMaterialPrimPath = Targets[0];
-					pxr::UsdPrim MaterialPrim = UsdPrim.GetStage()->GetPrimAtPath(TargetMaterialPrimPath);
+					pxr::UsdPrim MaterialPrim = Stage->GetPrimAtPath(TargetMaterialPrimPath);
 					ShadeMaterial = pxr::UsdShadeMaterial{MaterialPrim};
+					if (ShadeMaterial)
+					{
+						bPrimHasInheritableMaterialBindings = true;
+						NewMaterialUsers.FindOrAdd(UE::FSdfPath{TargetMaterialPrimPath}).Add(UE::FSdfPath{UsdPrimPath});
+					}
 				}
 			}
 
-			if (ShadeMaterial)
+			FWriteScopeLock ScopeLock(Impl.MaterialUsersLock);
+			for(const TPair<UE::FSdfPath, TSet<UE::FSdfPath>>& NewMaterialToUsers : NewMaterialUsers)
 			{
-				UE::FSdfPath UsedMaterialPath{ShadeMaterial.GetPrim().GetPath()};
-
-				FWriteScopeLock ScopeLock(Impl.MaterialUsersLock);
-				TSet<UE::FSdfPath>& MaterialUsersForPrim = Impl.MaterialUsers.FindOrAdd(UsedMaterialPath);
-
-				MaterialUsersForPrim.Add(PrimPath);
-
-				UE_LOG(
-					LogUsd,
-					Verbose,
-					TEXT("Registering prim '%s' as a material user of material prim '%s'"),
-					*PrimPath.GetString(),
-					*UsedMaterialPath.GetString()
-				);
-
-				// Our notice handling is somewhat stricter now, and we have no good way of upgrading a simple material info change
-				// into a resync change of the StaticMeshComponent when we change a material that is bound directly to a
-				// UsdGeomSubset, since the GeomMesh translator doesn't collapse. We'll unwind this path later when fetching material
-				// users, so collapsed static meshes are handled OK, skeletal meshes are handled OK, we just need this one exception
-				// for handling uncollapsed static meshes, because by default Mesh prims don't "collapse" their child UsdGeomSubsets
-				if (pxr::UsdGeomSubset Subset{UsdPrim})
+				TSet<UE::FSdfPath>& UserPrimPaths = Impl.MaterialUsers.FindOrAdd(NewMaterialToUsers.Key);
+				UserPrimPaths.Reserve(UserPrimPaths.Num() + NewMaterialToUsers.Value.Num());
+				for (const UE::FSdfPath& NewUserPath : NewMaterialToUsers.Value)
 				{
-					UE::FSdfPath ParentMesh{PrimPath.GetParentPath()};
-					UE_LOG(
-						LogUsd,
-						Verbose,
-						TEXT("Registering prim '%s' as a material user of material prim '%s'"),
-						*ParentMesh.GetString(),
-						*UsedMaterialPath.GetString()
-					);
-					MaterialUsersForPrim.Add(ParentMesh);
+					pxr::UsdPrim UserPrim = Stage->GetPrimAtPath(NewUserPath);
+
+					// Do this filtering here because Collection.ComputeIncludedPaths() can be very aggressive and return
+					// literally *all prims* below an included prim path. That's fine and it really does mean that any Mesh prim
+					// in there could use the collection-based material binding, but nevertheless we don't want to register that
+					// e.g. Shader prims or SkelAnimation prims are "material users"
+					if (UserPrim.IsA<pxr::UsdGeomImageable>())
+					{
+						UE_LOG(
+							LogUsd,
+							Verbose,
+							TEXT("Registering material user '%s' of material '%s'"),
+							*NewUserPath.GetString(),
+							*NewMaterialToUsers.Key.GetString()
+						);
+						UserPrimPaths.Add(NewUserPath);
+					}
+					// If a UsdGeomSubset is a material user, make its Mesh parent prim into a user too.
+					// Our notice handling is somewhat stricter now, and we have no good way of upgrading a simple material info change
+					// into a resync change of the StaticMeshComponent when we change a material that is bound directly to a
+					// UsdGeomSubset, since the GeomMesh translator doesn't collapse. We'll unwind this path later when fetching material
+					// users, so collapsed static meshes are handled OK, skeletal meshes are handled OK, we just need this one exception
+					// for handling uncollapsed static meshes, because by default Mesh prims don't "collapse" their child UsdGeomSubsets
+					else if (UserPrim.IsA<pxr::UsdGeomSubset>())
+					{
+						UE_LOG(
+							LogUsd,
+							Verbose,
+							TEXT("Registering parent Mesh prim '%s' of UsdGeomSubset '%s' as a material user of material prim '%s'"),
+							*NewUserPath.GetString(),
+							*PrimPath.GetString(),
+							*NewMaterialToUsers.Key.GetString()
+						);
+						UserPrimPaths.Add(NewUserPath.GetParentPath());
+					}
 				}
 			}
 		}
@@ -593,7 +694,8 @@ namespace UE::USDInfoCacheImpl::Private
 					InOutSubtreeToMaterialSlots,
 					InOutPointInstancerPaths,
 					ChildSubtreeVertexCounts[ Index ],
-					ChildSubtreeMaterialSlots[ Index ]
+					ChildSubtreeMaterialSlots[ Index ],
+					bPrimHasInheritableMaterialBindings || bPossibleInheritedBindings
 				);
 			}
 		);
@@ -1484,6 +1586,7 @@ void FUsdInfoCache::RebuildCacheForSubtree( const UE::FUsdPrim& Prim, FUsdSchema
 		// translator needs to know when it would generate too large a static mesh
 		uint64 SubtreeVertexCount = 0;
 		TArray<UsdUtils::FUsdPrimMaterialSlot> SubtreeSlots;
+		const bool bPossibleInheritedBindings = false;
 		UE::USDInfoCacheImpl::Private::RecursivePropagateVertexAndMaterialSlotCounts(
 			UsdPrim,
 			Context,
@@ -1493,7 +1596,8 @@ void FUsdInfoCache::RebuildCacheForSubtree( const UE::FUsdPrim& Prim, FUsdSchema
 			TempSubtreeSlots,
 			PointInstancerPaths,
 			SubtreeVertexCount,
-			SubtreeSlots
+			SubtreeSlots,
+			bPossibleInheritedBindings
 		);
 
 		UE::USDInfoCacheImpl::Private::UpdateInfoForPointInstancers(
