USD: Fix broken build when building without the editor.

diff --git a/Engine/Plugins/Importers/USDImporter/Source/USDSchemas/Private/Custom/MaterialXUSDShadeMaterialTranslator.cpp b/Engine/Plugins/Importers/USDImporter/Source/USDSchemas/Private/Custom/MaterialXUSDShadeMaterialTranslator.cpp
--- a/Engine/Plugins/Importers/USDImporter/Source/USDSchemas/Private/Custom/MaterialXUSDShadeMaterialTranslator.cpp
+++ b/Engine/Plugins/Importers/USDImporter/Source/USDSchemas/Private/Custom/MaterialXUSDShadeMaterialTranslator.cpp
@@ -51,6 +51,7 @@ FName FMaterialXUsdShadeMaterialTranslator::MaterialXRenderContext = TEXT("mtlx"
 
 namespace UE::USDMaterialXTranslator::Private
 {
+#if WITH_EDITOR
 	TArray<FString> GetMaterialXFilePaths(const pxr::UsdPrim& Prim)
 	{
 		TArray<FString> Result;
@@ -268,6 +269,7 @@ namespace UE::USDMaterialXTranslator::Private
 
 		return Manager.ImportAsset(TempPackagePath, SourceData, InterchangeParameters);
 	}
+#endif // WITH_EDITOR
 }
 
 void FMaterialXUsdShadeMaterialTranslator::CreateAssets()
diff --git a/Engine/Plugins/Importers/USDImporter/Source/USDUtilities/Private/USDPrimConversion.cpp b/Engine/Plugins/Importers/USDImporter/Source/USDUtilities/Private/USDPrimConversion.cpp
--- a/Engine/Plugins/Importers/USDImporter/Source/USDUtilities/Private/USDPrimConversion.cpp
+++ b/Engine/Plugins/Importers/USDImporter/Source/USDUtilities/Private/USDPrimConversion.cpp
@@ -2156,6 +2156,7 @@ bool UnrealToUsd::ConvertMaterialOverrides(
 	int32 HighestLOD
 )
 {
+#if WITH_EDITOR
 	if (!UsdPrim)
 	{
 		return false;
@@ -2446,6 +2447,7 @@ bool UnrealToUsd::ConvertMaterialOverrides(
 		ensure(false);
 		return false;
 	}
+#endif // WITH_EDITOR
 
 	return true;
 }
