USD: Fix logic bug that causes us to miss the MaterialX referencer prim in some cases.

diff --git a/Engine/Plugins/Importers/USDImporter/Source/USDSchemas/Private/Custom/MaterialXUSDShadeMaterialTranslator.cpp b/Engine/Plugins/Importers/USDImporter/Source/USDSchemas/Private/Custom/MaterialXUSDShadeMaterialTranslator.cpp
--- a/Engine/Plugins/Importers/USDImporter/Source/USDSchemas/Private/Custom/MaterialXUSDShadeMaterialTranslator.cpp
+++ b/Engine/Plugins/Importers/USDImporter/Source/USDSchemas/Private/Custom/MaterialXUSDShadeMaterialTranslator.cpp
@@ -393,7 +393,7 @@ void FMaterialXUsdShadeMaterialTranslator::CreateAssets()
 		if (ParentPrim)
 		{
 			// This prim likely holds the reference to the MaterialX file, but let's search upwards too
-			pxr::UsdPrim MaterialXReferencerCandidate = ParentPrim.GetParent();
+			pxr::UsdPrim MaterialXReferencerCandidate = ParentPrim;
 			while (MaterialXReferencerCandidate)
 			{
 				MaterialXFilePaths = UE::USDMaterialXTranslator::Private::GetMaterialXFilePaths(MaterialXReferencerCandidate);
