USD: Cancel loading the stage if the user presses X on the dialog to setup the asset cache.

diff --git a/Engine/Plugins/Importers/USDImporter/Source/USDClassesEditor/Private/USDClassesEditorModule.cpp b/Engine/Plugins/Importers/USDImporter/Source/USDClassesEditor/Private/USDClassesEditorModule.cpp
--- a/Engine/Plugins/Importers/USDImporter/Source/USDClassesEditor/Private/USDClassesEditorModule.cpp
+++ b/Engine/Plugins/Importers/USDImporter/Source/USDClassesEditor/Private/USDClassesEditorModule.cpp
@@ -14,6 +14,14 @@
 #define LOCTEXT_NAMESPACE "USDClassesEditorModule"
 
 UUsdAssetCache2* IUsdClassesEditorModule::ShowMissingDefaultAssetCacheDialog()
+{
+	UUsdAssetCache2* Result = nullptr;
+	bool bOutUserAccepted = false;
+	ShowMissingDefaultAssetCacheDialog(Result, bOutUserAccepted);
+	return Result;
+}
+
+void IUsdClassesEditorModule::ShowMissingDefaultAssetCacheDialog(UUsdAssetCache2*& OutCreatedCache, bool& bOutUserAccepted)
 {
 	TSharedPtr<SWindow> ParentWindow;
 
@@ -40,13 +48,16 @@ UUsdAssetCache2* IUsdClassesEditorModule::ShowMissingDefaultAssetCacheDialog()
 	const bool bSlowTaskWindow = false;
 	FSlateApplication::Get().AddModalWindow(Window, ParentWindow, bSlowTaskWindow);
 
-	UUsdAssetCache2* NewCache = nullptr;
 	if (OptionsWindow->UserAccepted())
 	{
-		NewCache = OptionsWindow->GetCreatedCache();
+		bOutUserAccepted = true;
+		OutCreatedCache = OptionsWindow->GetCreatedCache();
+	}
+	else
+	{
+		bOutUserAccepted = false;
+		OutCreatedCache = nullptr;
 	}
-
-	return NewCache;
 }
 
 class FUsdClassesEditorModule : public IUsdClassesEditorModule
diff --git a/Engine/Plugins/Importers/USDImporter/Source/USDClassesEditor/Public/USDClassesEditorModule.h b/Engine/Plugins/Importers/USDImporter/Source/USDClassesEditor/Public/USDClassesEditorModule.h
--- a/Engine/Plugins/Importers/USDImporter/Source/USDClassesEditor/Public/USDClassesEditorModule.h
+++ b/Engine/Plugins/Importers/USDImporter/Source/USDClassesEditor/Public/USDClassesEditorModule.h
@@ -10,5 +10,8 @@ class IUsdClassesEditorModule : public IModuleInterface
 {
 public:
     /** Shows the dialog that asks the user to set a default asset cache for the project */
+	UE_DEPRECATED(5.3, "Use the other signature that also returns whether the user accepted the dialog or not")
 	USDCLASSESEDITOR_API static UUsdAssetCache2* ShowMissingDefaultAssetCacheDialog();
+
+	USDCLASSESEDITOR_API static void ShowMissingDefaultAssetCacheDialog(UUsdAssetCache2*& OutCreatedCache, bool& bOutUserAccepted);
 };
\ No newline at end of file
diff --git a/Engine/Plugins/Importers/USDImporter/Source/USDStage/Private/USDStageActor.cpp b/Engine/Plugins/Importers/USDImporter/Source/USDStage/Private/USDStageActor.cpp
--- a/Engine/Plugins/Importers/USDImporter/Source/USDStage/Private/USDStageActor.cpp
+++ b/Engine/Plugins/Importers/USDImporter/Source/USDStage/Private/USDStageActor.cpp
@@ -1192,13 +1192,12 @@ void AUsdStageActor::OnUsdObjectsChanged(const UsdUtils::FObjectChangesByPath& I
 		}
 	}
 
-	EnsureAssetCache();
-
-	TOptional<FUsdScopedAssetCacheReferencer> ScopedReferencer;
-	if (UsdAssetCache)
+	SetupAssetCacheIfNeeded();
+	if (!UsdAssetCache)
 	{
-		ScopedReferencer.Emplace(UsdAssetCache, this);
+		return;
 	}
+	FUsdScopedAssetCacheReferencer ScopedReferencer{UsdAssetCache, this};
 
 	// Mark the level as dirty since we received a notice about our stage having changed in some way.
 	// The main goal of this is to trigger the "save layers" dialog if we then save the UE level
@@ -1678,7 +1677,7 @@ void AUsdStageActor::OnUsdObjectsChanged(const UsdUtils::FObjectChangesByPath& I
 		FUsdStageActorImpl::RepairExternalSequencerBindings();
 	}
 
-	if (bHasLoadedOrAbandonedAssets && UsdAssetCache)
+	if (bHasLoadedOrAbandonedAssets)
 	{
 		UsdAssetCache->RefreshStorage();
 	}
@@ -2504,7 +2503,23 @@ void AUsdStageActor::LoadUsdStage()
 	// Ensure an asset cache if we're going to load something
 	if (!RootLayer.FilePath.IsEmpty())
 	{
-		EnsureAssetCache();
+		SetupAssetCacheIfNeeded();
+
+		// The only way we're coming out of SetupAssetCacheIfNeeded without an asset cache is if the user canceled the
+		// asset cache dialog. In that case let's just abort loading this stage entirely and revert our RootLayer
+		// to empty
+		if (!UsdAssetCache)
+		{
+			UE_LOG(
+				LogUsd,
+				Log,
+				TEXT("Cancelling the loading of USD Stage '%s' as the Asset Cache setup was explicitly canceled."),
+				*RootLayer.FilePath
+			);
+
+			RootLayer.FilePath.Empty();
+			return;
+		}
 	}
 
 	double StartTime = FPlatformTime::Cycles64();
@@ -2700,7 +2715,7 @@ void AUsdStageActor::UnloadUsdStage()
 	CloseUsdStage();
 }
 
-void AUsdStageActor::EnsureAssetCache()
+void AUsdStageActor::SetupAssetCacheIfNeeded()
 {
 	if (!UsdAssetCache)
 	{
@@ -2725,7 +2740,11 @@ void AUsdStageActor::EnsureAssetCache()
 		{
 			if (ProjectSettings->bShowCreateDefaultAssetCacheDialog)
 			{
-				if (UUsdAssetCache2* NewCache = IUsdClassesEditorModule::ShowMissingDefaultAssetCacheDialog())
+				bool bOutUserAccepted = false;
+				UUsdAssetCache2* NewCache = nullptr;
+				IUsdClassesEditorModule::ShowMissingDefaultAssetCacheDialog(NewCache, bOutUserAccepted);
+
+				if (bOutUserAccepted)
 				{
 					ProjectSettings->DefaultAssetCache = NewCache;
 					ProjectSettings->SaveConfig();
@@ -2737,6 +2756,10 @@ void AUsdStageActor::EnsureAssetCache()
 						*UsdAssetCache->GetPathName()
 					);
 				}
+				else
+				{
+					return;
+				}
 			}
 		}
 	}
diff --git a/Engine/Plugins/Importers/USDImporter/Source/USDStage/Public/USDStageActor.h b/Engine/Plugins/Importers/USDImporter/Source/USDStage/Public/USDStageActor.h
--- a/Engine/Plugins/Importers/USDImporter/Source/USDStage/Public/USDStageActor.h
+++ b/Engine/Plugins/Importers/USDImporter/Source/USDStage/Public/USDStageActor.h
@@ -289,7 +289,7 @@ protected:
 	void LoadUsdStage();
 	void UnloadUsdStage();
 
-	void EnsureAssetCache();
+	void SetupAssetCacheIfNeeded();
 
 	bool HasAuthorityOverStage() const;
 
