USD: Emit a warning instead of crashing the editor when USD skeleton has more than one root bone.

diff --git a/Engine/Plugins/Importers/USDImporter/Source/USDSchemas/Private/USDSkelRootTranslator.cpp b/Engine/Plugins/Importers/USDImporter/Source/USDSchemas/Private/USDSkelRootTranslator.cpp
--- a/Engine/Plugins/Importers/USDImporter/Source/USDSchemas/Private/USDSkelRootTranslator.cpp
+++ b/Engine/Plugins/Importers/USDImporter/Source/USDSchemas/Private/USDSkelRootTranslator.cpp
@@ -1315,7 +1315,7 @@ namespace UsdSkelRootTranslatorImpl
 				Options.RenderContext = RenderContextToken;
 				Options.MaterialPurpose = MaterialPurposeToken;
 
-				const bool bContinueTaskChain = UsdSkelRootTranslatorImpl::LoadAllSkeletalData(
+				bool bContinueTaskChain = UsdSkelRootTranslatorImpl::LoadAllSkeletalData(
 					SkeletonCache.Get(),
 					pxr::UsdSkelRoot( GetPrim() ),
 					LODIndexToSkeletalMeshImportData,
@@ -1329,6 +1329,31 @@ namespace UsdSkelRootTranslatorImpl
 					Options
 				);
 
+				int32 NumRootBones = 0;
+				for (const SkeletalMeshImportData::FBone& Bone : SkeletonBones)
+				{
+					if (Bone.ParentIndex == -1)
+					{
+						NumRootBones++;
+					}
+
+					if (NumRootBones > 1)
+					{
+						break;
+					}
+				}
+
+				if (NumRootBones != 1)
+				{
+					UE_LOG(
+						LogUsd,
+						Warning,
+						TEXT("Ignoring SkelRoot '%s' as the bound Skeleton prim must have exactly one root bone!"),
+						*PrimPath.GetString()
+					);
+					bContinueTaskChain = false;
+				}
+
 				return bContinueTaskChain;
 			} );
 
