USD: Retain property and property metadata selection when prim info panel updates.

diff --git a/Engine/Plugins/Importers/USDImporter/Source/USDStageEditor/Private/SUSDPrimInfo.cpp b/Engine/Plugins/Importers/USDImporter/Source/USDStageEditor/Private/SUSDPrimInfo.cpp
--- a/Engine/Plugins/Importers/USDImporter/Source/USDStageEditor/Private/SUSDPrimInfo.cpp
+++ b/Engine/Plugins/Importers/USDImporter/Source/USDStageEditor/Private/SUSDPrimInfo.cpp
@@ -120,10 +120,33 @@ void SUsdPrimInfo::SetPrimPath(const UE::FUsdStageWeak& UsdStage, const TCHAR* P
 {
 	if ( PropertiesList )
 	{
+		TArray<FString> OldSelectedProperties = PropertiesList->GetSelectedFieldNames();
+		TArray<FString> OldSelectedPropertyMetadata = PropertyMetadataPanel->GetSelectedFieldNames();
+
+		// This is the main way with which we resync the properties list in case the prim received some change while
+		// we were displaying it
 		PropertiesList->SetObjectPath(UsdStage, PrimPath);
 
-		PropertyMetadataPanel->SetObjectPath({}, TEXT(""));
-		PropertyMetadataPanel->ClearSelection();
+		// Try restoring the property selection while we change selected prims
+		// This will also restore the object path on the PropertyMetadataPanel because PropertiesList::OnSelectionChanged
+		// will fire
+		PropertiesList->SetSelectedFieldNames(OldSelectedProperties);
+		TArray<FString> NewSelectedProperties = PropertiesList->GetSelectedFieldNames();
+
+		// If we haven't managed to select everything we had before, reset (and hide) the metadata panel.
+		// We need this because apparently the list view doesn't generate a selection changed event when it's items are
+		// fully rebuilt, which will happen on SetObjectPath. If SetSelectedFieldNames can't select anything, nothing will
+		// update the metadata panel
+		if(OldSelectedProperties.Num() != NewSelectedProperties.Num() || NewSelectedProperties.Num() != 1)
+		{
+			PropertyMetadataPanel->SetObjectPath({}, TEXT(""));
+			PropertyMetadataPanel->SetVisibility(EVisibility::Collapsed);
+		}
+		else
+		{
+			// If we had a metadata panel open, try restoring the selection within that too
+			PropertyMetadataPanel->SetSelectedFieldNames(OldSelectedPropertyMetadata);
+		}
 	}
 
 	if ( IntegrationsPanel )
