USD: Fix "Don't create a persistent asset cache" button canceling opening the stage.

diff --git a/Engine/Plugins/Importers/USDImporter/Source/USDClassesEditor/Private/USDDefaultAssetCacheDialog.cpp b/Engine/Plugins/Importers/USDImporter/Source/USDClassesEditor/Private/USDDefaultAssetCacheDialog.cpp
--- a/Engine/Plugins/Importers/USDImporter/Source/USDClassesEditor/Private/USDDefaultAssetCacheDialog.cpp
+++ b/Engine/Plugins/Importers/USDImporter/Source/USDClassesEditor/Private/USDDefaultAssetCacheDialog.cpp
@@ -262,7 +262,7 @@ FReply SUsdDefaultAssetCacheDialog::OnCreateNew()
 
 FReply SUsdDefaultAssetCacheDialog::OnDontCreate()
 {
-	bAccepted = false;
+	bAccepted = true;
 	if (Window.IsValid())
 	{
 		Window.Pin()->RequestDestroyWindow();
