USD: Fixed support for USD geometry cache track. Now, it can also be driven by the Time property of the USD stage and is enabled by default.

diff --git a/Engine/Plugins/Importers/USDImporter/Source/GeometryCacheUSD/Private/GeometryCacheTrackUSD.cpp b/Engine/Plugins/Importers/USDImporter/Source/GeometryCacheUSD/Private/GeometryCacheTrackUSD.cpp
--- a/Engine/Plugins/Importers/USDImporter/Source/GeometryCacheUSD/Private/GeometryCacheTrackUSD.cpp
+++ b/Engine/Plugins/Importers/USDImporter/Source/GeometryCacheUSD/Private/GeometryCacheTrackUSD.cpp
@@ -11,10 +11,7 @@
 #include "HAL/IConsoleManager.h"
 #include "IGeometryCacheStreamer.h"
 
-// Currently disabled by default as it prevents AUsdStageActor::ExpandPrim from animating Mesh prim transforms,
-// and fixing it is not trivial: UGeometryCacheTrackUsd currently either works for animating via the Time property
-// *OR* via the geometry cache Sequencer tracks, while we still need the Time property animation to work in both cases
-static bool GDisableGeoCacheTracks = true;
+static bool GDisableGeoCacheTracks = false;
 static FAutoConsoleVariableRef CVarDisableGeoCacheTracks(
 	TEXT("USD.DisableGeoCacheTracks"),
 	GDisableGeoCacheTracks,
@@ -97,34 +94,23 @@ const bool UGeometryCacheTrackUsd::UpdateBoundsData(const float Time, const bool
 
 const int32 UGeometryCacheTrackUsd::FindSampleIndexFromTime(const float Time, const bool bLooping) const
 {
-	if (!GDisableGeoCacheTracks)
+	// Time is relative to the start of the section on the track
+	float SampleTime = Time;
+	if (bLooping)
 	{
-		// Time is relative to the start of the section on the track
-		float SampleTime = Time;
-		if (bLooping)
-		{
-			SampleTime = GeometyCacheHelpers::WrapAnimationTime(Time, Duration);
-		}
-		// Which is converted to an index
-		int32 FrameIndex = FMath::RoundToInt(Time * FramesPerSecond);
-
-		// The final computed frame index must be offset by the start frame index
-		return FMath::Clamp(FrameIndex + StartFrameIndex, StartFrameIndex, EndFrameIndex - 1);
+		SampleTime = GeometyCacheHelpers::WrapAnimationTime(Time, Duration);
 	}
-	// Treat the time as the frame index
-	int32 FrameIndex = (int32)Time;
-	return FMath::Clamp(FrameIndex, StartFrameIndex, EndFrameIndex - 1);
+	// Which is converted to an index
+	int32 FrameIndex = FMath::RoundToInt(Time * FramesPerSecond);
+
+	// The final computed frame index must be offset by the start frame index
+	return FMath::Clamp(FrameIndex + StartFrameIndex, StartFrameIndex, EndFrameIndex - 1);
 }
 
 float UGeometryCacheTrackUsd::GetTimeFromSampleIndex(int32 SampleIndex) const
 {
-	if (!GDisableGeoCacheTracks)
-	{
-		// Time is relative to the start of the section on the track
-		return float((SampleIndex - StartFrameIndex) / FramesPerSecond);
-	}
-	// Treat the frame index as time
-	return SampleIndex;
+	// Time is relative to the start of the section on the track
+	return float((SampleIndex - StartFrameIndex) / FramesPerSecond);
 }
 
 void UGeometryCacheTrackUsd::GetFractionalFrameIndexFromTime(const float Time, const bool bLooping, int& OutFrameIndex, float& OutFraction) const
@@ -132,22 +118,15 @@ void UGeometryCacheTrackUsd::GetFractionalFrameIndexFromTime(const float Time, c
 	OutFrameIndex = FindSampleIndexFromTime(Time, bLooping);
 
 	float AdjustedTime = Time;
-	if (!GDisableGeoCacheTracks)
-	{
-		if (bLooping)
-		{
-			AdjustedTime = GeometyCacheHelpers::WrapAnimationTime(Time, Duration);
-		}
-		// Time at ThisFrameIndex with index normalized to 0
-		const float FrameIndexTime = (OutFrameIndex - StartFrameIndex) / FramesPerSecond;
-		OutFraction = float((AdjustedTime - FrameIndexTime) * FramesPerSecond);
-	}
-	else
+	if (bLooping)
 	{
-		// Clamp the Time (which is the FrameNumber) to the range of the USD GeometryCache track
-		// EndFrameIndex - 1 since last frame index is not included
-		OutFraction = FMath::Clamp(Time, float(StartFrameIndex), float(EndFrameIndex - 1)) - OutFrameIndex;
+		AdjustedTime = GeometyCacheHelpers::WrapAnimationTime(Time, Duration);
 	}
+	// Time at ThisFrameIndex with index normalized to 0
+	const float FrameIndexTime = (OutFrameIndex - StartFrameIndex) / FramesPerSecond;
+	OutFraction = float((AdjustedTime - FrameIndexTime) * FramesPerSecond);
+	// The fractional part is clamped to (-0.5, 0.5) for subframe interpolation that is centered around FrameIndex
+	OutFraction = FMath::Clamp(OutFraction, -0.5f, 0.5f);
 }
 
 const FGeometryCacheTrackSampleInfo& UGeometryCacheTrackUsd::GetSampleInfo(float Time, bool bLooping)
@@ -276,11 +255,8 @@ void UGeometryCacheTrackUsd::Initialize(
 		FramesPerSecond = 1;
 	}
 
-	// The main difference between GDisableGeoCacheTracks true or false is in how time
-	// is converted to frame index. When true, the time is basically the frame indices
-	// and no conversion is actually needed.
 	const int32 NumFrames = EndFrameIndex - StartFrameIndex + 1;
-	Duration = GDisableGeoCacheTracks ? NumFrames : NumFrames / FramesPerSecond;
+	Duration = NumFrames / FramesPerSecond;
 
 	UsdStream.Reset(new FGeometryCacheUsdStream(this, InReadFunc));
 	IGeometryCacheStreamer::Get().RegisterTrack(this, UsdStream.Get());
diff --git a/Engine/Plugins/Importers/USDImporter/Source/USDSchemas/Private/USDGeometryCacheTranslator.cpp b/Engine/Plugins/Importers/USDImporter/Source/USDSchemas/Private/USDGeometryCacheTranslator.cpp
--- a/Engine/Plugins/Importers/USDImporter/Source/USDSchemas/Private/USDGeometryCacheTranslator.cpp
+++ b/Engine/Plugins/Importers/USDImporter/Source/USDSchemas/Private/USDGeometryCacheTranslator.cpp
@@ -341,7 +341,7 @@ namespace UsdGeometryCacheTranslatorImpl
 	}
 
 	UGeometryCache* CreateGeometryCache(const FString& RootPrimPath, const FMeshDescription& MeshDescription, const TArray<UE::FSdfPath>& MeshPaths, const TArray<int32>& MaterialOffsets,
-										TSharedRef<FUsdSchemaTranslationContext> Context, bool& bOutIsNew)
+										TSharedRef<FUsdSchemaTranslationContext> Context, bool& bOutIsNew, float& StartOffsetTime)
 	{
 		// Compute the asset hash from the merged mesh description
 		FSHA1 SHA1;
@@ -378,6 +378,8 @@ namespace UsdGeometryCacheTranslatorImpl
 			if (!Context->bIsImporting)
 			{
 				Args = GetReadMeshDataArgs(Context, RootPrimPath);
+				// StartOffsetTime is the offset applied to the GeometryCache section on the sequencer track, so not relevant when importing
+				StartOffsetTime = static_cast<float>(Args->StartFrame) / Args->FramesPerSecond;
 			}
 
 			// Create a track for each mesh to be processed and add it to the GeometryCache
@@ -489,7 +491,7 @@ namespace UsdGeometryCacheTranslatorImpl
 
 					// Parallel frame read: frame data can be read concurrently but have to be processed in order for AddMeshSample
 					ParallelFor(NumFrameThreads, [&Args, &MeshPrim, &FrameWrittenEvent, &WriteFrameIndex, &Mutex, NumFrameThreads, MaterialOffset, StreamableTrack, bConstantTopology](int32 FrameThreadIndex)
-							{
+					{
 						int32 FrameIndex = Args.StartFrame + FrameThreadIndex;
 
 						while (FrameIndex < Args.EndFrame)
@@ -763,8 +765,9 @@ void FGeometryCacheCreateAssetsTaskChain::SetupTasks()
 			}
 
 			bool bIsNew = true;
+			float StartTimeOffset = 0.0f;
 			const FString PrimPathString = PrimPath.GetString();
-			GeometryCache.Reset(UsdGeometryCacheTranslatorImpl::CreateGeometryCache(PrimPathString, LODIndexToMeshDescription[0], MeshPrimPaths, MaterialOffsets, Context, bIsNew));
+			GeometryCache.Reset(UsdGeometryCacheTranslatorImpl::CreateGeometryCache(PrimPathString, LODIndexToMeshDescription[0], MeshPrimPaths, MaterialOffsets, Context, bIsNew, StartTimeOffset));
 
 			if (GeometryCache && bIsNew)
 			{
@@ -773,6 +776,11 @@ void FGeometryCacheCreateAssetsTaskChain::SetupTasks()
 				UserData->PrimvarToUVIndex = LODIndexToMaterialInfo[0].PrimvarToUVIndex;  // We use the same primvar mapping for all LODs
 				GeometryCache->AddAssetUserData(UserData);
 
+				UUsdAnimSequenceAssetUserData* AnimUserData = NewObject<UUsdAnimSequenceAssetUserData>(GeometryCache.Get(), TEXT("UsdAnimUserData"));
+				AnimUserData->PrimPath = PrimPathString;
+				AnimUserData->LayerStartOffsetSeconds = StartTimeOffset;
+				GeometryCache->AddAssetUserData(AnimUserData);
+
 				// Only the original creator of the prim at creation time gets to set the material assignments
 				// directly on the geometry cache, all others prims ensure their materials via material overrides on the
 				// components
@@ -939,9 +947,10 @@ void FUsdGeometryCacheTranslator::UpdateComponents(USceneComponent* SceneCompone
 		}
 
 		// Manually tick USD GeometryCache only when their tracks are disabled in Sequencer
+		// but also need to tick for the initial setup
 		static IConsoleVariable* CVar = IConsoleManager::Get().FindConsoleVariable(TEXT("USD.DisableGeoCacheTracks"));
 		bool bDisableGeoCacheTracks = CVar && CVar->GetBool();
-		if (bDisableGeoCacheTracks)
+		if (bDisableGeoCacheTracks || !Context->bSequencerIsAnimating)
 		{
 			float TimeCode = Context->Time;
 			if (FMath::IsNaN(TimeCode))
@@ -958,7 +967,25 @@ void FUsdGeometryCacheTranslator::UpdateComponents(USceneComponent* SceneCompone
 			// geometry cache to register itself at Context->Time so that it will synchronously load that frame right away.
 			// Otherwise the geometry cache will start at t=0 regardless of Context->Time
 			GeometryCacheComponent->SetManualTick(true);
-			GeometryCacheComponent->TickAtThisTime(TimeCode, true, false, true);
+
+			// Looping is disabled since the animation is driven by Sequencer
+			const bool bIsLooping = false;
+			GeometryCacheComponent->SetLooping(bIsLooping);
+
+			float LayerStartOffsetSeconds = 0.0f;
+			if (GeometryCache)
+			{
+				if (UUsdAnimSequenceAssetUserData* UserData = GeometryCache->GetAssetUserData<UUsdAnimSequenceAssetUserData>())
+				{
+					LayerStartOffsetSeconds = UserData->LayerStartOffsetSeconds;
+				}
+			}
+
+			const double FramesPerSecond = Context->Stage.GetTimeCodesPerSecond();
+			const float AdjustedTime = static_cast<float>(TimeCode / FramesPerSecond - LayerStartOffsetSeconds);
+			const bool bIsRunning = true;
+			const bool bIsBackwards = false;
+			GeometryCacheComponent->TickAtThisTime(AdjustedTime, bIsRunning, bIsBackwards, bIsLooping);
 		}
 
 		// If the prim has a GroomBinding schema, apply the target groom to its associated GroomComponent
diff --git a/Engine/Plugins/Importers/USDImporter/Source/USDStage/Private/USDLevelSequenceHelper.cpp b/Engine/Plugins/Importers/USDImporter/Source/USDStage/Private/USDLevelSequenceHelper.cpp
--- a/Engine/Plugins/Importers/USDImporter/Source/USDStage/Private/USDLevelSequenceHelper.cpp
+++ b/Engine/Plugins/Importers/USDImporter/Source/USDStage/Private/USDLevelSequenceHelper.cpp
@@ -1045,16 +1045,17 @@ ULevelSequence* FUsdLevelSequenceHelperImpl::FindOrAddSequenceForAttribute(const
 	{
 		if (UE::FSdfLayer AttributeLayer = UsdUtils::FindLayerForAttribute(Attribute, 0.0))
 		{
-			if (UsdStage)
+			const FString SequenceIdentifier = AttributeLayer.GetIdentifier();
+			Sequence = FindOrAddSequenceForLayer(AttributeLayer, SequenceIdentifier, SequenceIdentifier);
+
+#if WITH_EDITORONLY_DATA
+			// Make level sequences for non-local layers read-only: Our current approach is that only
+			// local layers can be written to, so there is no point in doing otherwise
+			if (Sequence && UsdStage && !UsdStage.HasLocalLayer(AttributeLayer))
 			{
-				// Only create level sequences for layers in the local layer stack: Our current approach is that only
-				// those layers can be written to, so there is no point in doing otherwise
-				if (UsdStage.HasLocalLayer(AttributeLayer))
-				{
-					const FString SequenceIdentifier = AttributeLayer.GetIdentifier();
-					Sequence = FindOrAddSequenceForLayer(AttributeLayer, SequenceIdentifier, SequenceIdentifier);
-				}
+				Sequence->GetMovieScene()->SetReadOnly(true);
 			}
+#endif
 		}
 	}
 
@@ -1231,27 +1232,49 @@ void FUsdLevelSequenceHelperImpl::CreateSubSequenceSection(ULevelSequence& Seque
 	}
 
 	const double TimeCodesPerSecond = Layer.GetTimeCodesPerSecond();
+	const bool bIsAlembicSublayer = SubLayerIdentifier->EndsWith(TEXT(".abc"));
+	TRange< FFrameNumber > SubSectionRange;
+	FFrameNumber StartFrame;
+
+	if (!bIsAlembicSublayer)
+	{
+		// Section full duration is always [0, endTimeCode]. The play range varies: For the root layer it will be [startTimeCode, endTimeCode],
+		// but for sublayers it will be [0, endTimeCode] too in order to match how USD composes sublayers with non-zero startTimeCode
+		const double SubDurationTimeCodes = SubLayer.GetEndTimeCode() * SubLayerOffset.Scale;
+		const double SubDurationSeconds = SubDurationTimeCodes / TimeCodesPerSecond;
 
-	// Section full duration is always [0, endTimeCode]. The play range varies: For the root layer it will be [startTimeCode, endTimeCode],
-	// but for sublayers it will be [0, endTimeCode] too in order to match how USD composes sublayers with non-zero startTimeCode
-	const double SubDurationTimeCodes = SubLayer.GetEndTimeCode() * SubLayerOffset.Scale;
-	const double SubDurationSeconds = SubDurationTimeCodes / TimeCodesPerSecond;
+		const double SubStartTimeSeconds = SubLayerOffset.Offset / TimeCodesPerSecond;
+		const double SubEndTimeSeconds = SubStartTimeSeconds + SubDurationSeconds;
 
-	const double SubStartTimeSeconds = SubLayerOffset.Offset / TimeCodesPerSecond;
-	const double SubEndTimeSeconds = SubStartTimeSeconds + SubDurationSeconds;
+		StartFrame = UsdLevelSequenceHelperImpl::RoundAsFrameNumber(TickResolution, SubStartTimeSeconds);
+		const FFrameNumber EndFrame = UsdLevelSequenceHelperImpl::RoundAsFrameNumber(TickResolution, SubEndTimeSeconds);
 
-	const FFrameNumber StartFrame = UsdLevelSequenceHelperImpl::RoundAsFrameNumber(TickResolution, SubStartTimeSeconds);
-	const FFrameNumber EndFrame = UsdLevelSequenceHelperImpl::RoundAsFrameNumber(TickResolution, SubEndTimeSeconds);
+		// Don't clip subsections with their duration, so that the root layer's [startTimeCode, endTimeCode] range is the only thing clipping
+		// anything, as this is how USD seems to behave. Even if a middle sublayer has startTimeCode == endTimeCode, its animations
+		// (or its child sublayers') won't be clipped by it and play according to the stage's range
+		const double StageEndTimeSeconds = UsdStage.GetEndTimeCode() / UsdStage.GetTimeCodesPerSecond();
+		const FFrameNumber StageEndFrame = UsdLevelSequenceHelperImpl::RoundAsFrameNumber(TickResolution, StageEndTimeSeconds);
+
+		// Max here because StartFrame can theoretically be larger than StageEndFrame, which would generate a range where the upper bound is smaller
+		// than the lower bound, which can trigger asserts
+		SubSectionRange = TRange< FFrameNumber >{StartFrame, FMath::Max(StageEndFrame, EndFrame)};
+	}
+	else
+	{
+		// One issue with a sublayer from Alembic is that the usdAbc plugin does not retrieve the frame rate of the archive.
+		// Another is that the start time does not necessarily represents the actual start of the animation. That's why
+		// there's an option to "skip empty frames" when importing an Alembic. So instead take the start/end timecodes from
+		// the parent layer. That way the user can define the animation range needed.
+		SubLayer = Layer;
 
-	// Don't clip subsections with their duration, so that the root layer's [startTimeCode, endTimeCode] range is the only thing clipping
-	// anything, as this is how USD seems to behave. Even if a middle sublayer has startTimeCode == endTimeCode, its animations
-	// (or its child sublayers') won't be clipped by it and play according to the stage's range
-	const double StageEndTimeSeconds = UsdStage.GetEndTimeCode() / UsdStage.GetTimeCodesPerSecond();
-	const FFrameNumber StageEndFrame = UsdLevelSequenceHelperImpl::RoundAsFrameNumber(TickResolution, StageEndTimeSeconds);
+		const double SubStartTimeSeconds = SubLayer.GetStartTimeCode() * SubLayerOffset.Scale / TimeCodesPerSecond;
+		const double SubEndTimeSeconds = SubLayer.GetEndTimeCode() * SubLayerOffset.Scale / TimeCodesPerSecond;
 
-	// Max here because StartFrame can theoretically be larger than StageEndFrame, which would generate a range where the upper bound is smaller
-	// than the lower bound, which can trigger asserts
-	TRange< FFrameNumber > SubSectionRange{StartFrame, FMath::Max(StageEndFrame, EndFrame)};
+		StartFrame = UsdLevelSequenceHelperImpl::RoundAsFrameNumber(TickResolution, SubStartTimeSeconds);
+		const FFrameNumber EndFrame = UsdLevelSequenceHelperImpl::RoundAsFrameNumber(TickResolution, SubEndTimeSeconds);
+
+		SubSectionRange = TRange<FFrameNumber>{StartFrame, EndFrame};
+	}
 
 	UMovieSceneSubSection* SubSection = FindSubSequenceSection(Sequence, SubSequence);
 	if (SubSection)
@@ -1962,11 +1985,13 @@ void FUsdLevelSequenceHelperImpl::AddGeometryCacheTracks(const UUsdPrimTwin& Pri
 		GeometryCacheTrack->RemoveAllAnimationData();
 
 		double LayerStartOffsetSeconds = 0;
-		if (GeometryCache->Tracks.Num())
+		if (GeometryCacheSequence == MainLevelSequence)
 		{
-			// Use the geometry cache start frame index converted to time as the offset
-			UGeometryCacheTrackUsd* UsdTrack = Cast<UGeometryCacheTrackUsd>(GeometryCache->Tracks[0]);
-			LayerStartOffsetSeconds = UsdTrack->StartFrameIndex / UsdTrack->FramesPerSecond;
+			// The LayerStartOffset needs to be applied only for the track in the main sequence. For subsequences, it's the subsequence section that is offset.
+			if (UUsdAnimSequenceAssetUserData* UserData = GeometryCache->GetAssetUserData<UUsdAnimSequenceAssetUserData>())
+			{
+				LayerStartOffsetSeconds = UserData->LayerStartOffsetSeconds;
+			}
 		}
 
 		const FFrameNumber StartOffsetTick = FFrameTime::FromDecimal(LayerStartOffsetSeconds * MovieScene->GetTickResolution().AsDecimal()).RoundToFrame();
@@ -2050,6 +2075,12 @@ void FUsdLevelSequenceHelperImpl::AddPrim(UUsdPrimTwin& PrimTwin, bool bForceVis
 	UE::FSdfLayer PrimLayer = UsdUtils::FindLayerForPrim(UsdPrim);
 	ULevelSequence* PrimSequence = FindSequenceForIdentifier(PrimLayer.GetIdentifier());
 
+	// PrimSequence is needed for a subsequence section to be created so fall back to the main sequence
+	if (!PrimSequence)
+	{
+		PrimSequence = MainLevelSequence;
+	}
+
 	TArray< UE::FUsdAttribute > PrimAttributes = UsdPrim.GetAttributes();
 
 	for (const UE::FUsdAttribute& PrimAttribute : PrimAttributes)
diff --git a/Engine/Plugins/Importers/USDImporter/Source/USDStage/Private/USDStageActor.cpp b/Engine/Plugins/Importers/USDImporter/Source/USDStage/Private/USDStageActor.cpp
--- a/Engine/Plugins/Importers/USDImporter/Source/USDStage/Private/USDStageActor.cpp
+++ b/Engine/Plugins/Importers/USDImporter/Source/USDStage/Private/USDStageActor.cpp
@@ -1865,15 +1865,7 @@ UUsdPrimTwin* AUsdStageActor::ExpandPrim(const UE::FUsdPrim& Prim, bool bResync,
 			UE::FSdfPath PrimPath(*UsdPrimTwin->PrimPath);
 			UE::FUsdPrim UsdPrim(UsdStage.GetPrimAtPath(PrimPath));
 
-			static IConsoleVariable* CVar = IConsoleManager::Get().FindConsoleVariable(TEXT("USD.DisableGeoCacheTracks"));
-			const bool bDisableGeoCacheTracks = CVar && CVar->GetBool();
-
-			// Optimization: manually animate mesh prims only when the USD GeometryCaches tracks are disabled
-			const bool bAnimatePrim = bDisableGeoCacheTracks || !UsdPrim.IsA(TEXT("Mesh"));
-			if (bAnimatePrim)
-			{
-				PrimsToAnimate.Add(UsdPrimTwin->PrimPath);
-			}
+			PrimsToAnimate.Add(UsdPrimTwin->PrimPath);
 			LevelSequenceHelper.AddPrim(*UsdPrimTwin);
 		}
 	}
