USD: Use texture filenames instead of the material prim names for the imported texture asset names.

diff --git a/Engine/Plugins/Importers/USDImporter/Source/USDStageImporter/Private/USDStageImporter.cpp b/Engine/Plugins/Importers/USDImporter/Source/USDStageImporter/Private/USDStageImporter.cpp
--- a/Engine/Plugins/Importers/USDImporter/Source/USDStageImporter/Private/USDStageImporter.cpp
+++ b/Engine/Plugins/Importers/USDImporter/Source/USDStageImporter/Private/USDStageImporter.cpp
@@ -531,6 +531,12 @@ namespace UsdStageImporterImpl
 		else if (UTexture* Texture = Cast<UTexture>(Asset))
 		{
 			AssetPrefix = TEXT("T_");
+
+			// Use the actual asset name as it matches the original texture filename, which is more useful
+			// than ending up with "T_<MaterialName>" instead (the textures receive the material prim paths).
+			// We could also fetch it from AssetImportData but then we'd have to process the path a bit in case of
+			// textures inside USDZ files. It's just easier to take from the asset name itself instead
+			AssetName = Texture->GetFName().GetPlainNameString();
 		}
 		else if ( ULevelSequence* LevelSequence = Cast<ULevelSequence>( Asset ) )
 		{
