USD: Never add dynamic bindings to transient level sequences.

diff --git a/Engine/Plugins/Importers/USDImporter/Source/USDStage/Private/USDStageActor.cpp b/Engine/Plugins/Importers/USDImporter/Source/USDStage/Private/USDStageActor.cpp
--- a/Engine/Plugins/Importers/USDImporter/Source/USDStage/Private/USDStageActor.cpp
+++ b/Engine/Plugins/Importers/USDImporter/Source/USDStage/Private/USDStageActor.cpp
@@ -3473,6 +3473,12 @@ void AUsdStageActor::OnActorAddedToSequencer(AActor* NewActor, const FGuid Guid,
 		return;
 	}
 
+	// We never need dynamic bindings on our (or other stage actor's) transient level sequences
+	if (Sequence == LevelSequence || Sequence->HasAnyFlags(RF_Transient))
+	{
+		return;
+	}
+
 	UMovieScene* MovieScene = Sequence->GetMovieScene();
 	if (!MovieScene)
 	{
@@ -3518,6 +3524,12 @@ void AUsdStageActor::OnMovieSceneDataChanged(EMovieSceneDataChangeType ChangeTyp
 		return;
 	}
 
+	// We never need dynamic bindings on our (or other stage actor's) transient level sequences
+	if (Sequence == LevelSequence || Sequence->HasAnyFlags(RF_Transient))
+	{
+		return;
+	}
+
 	UUsdLevelSequenceAssetUserData* AssetUserData = Sequence->GetAssetUserData<UUsdLevelSequenceAssetUserData>();
 	if(!AssetUserData)
 	{
