USD: Fix material overrides potentially writing to the wrong prim if a mesh asset is reused from the asset cache.

diff --git a/Engine/Plugins/Importers/USDImporter/Source/USDSchemas/Private/MeshTranslationImpl.cpp b/Engine/Plugins/Importers/USDImporter/Source/USDSchemas/Private/MeshTranslationImpl.cpp
--- a/Engine/Plugins/Importers/USDImporter/Source/USDSchemas/Private/MeshTranslationImpl.cpp
+++ b/Engine/Plugins/Importers/USDImporter/Source/USDSchemas/Private/MeshTranslationImpl.cpp
@@ -858,6 +858,29 @@ void MeshTranslationImpl::SetMaterialOverrides(
 	}
 }
 
+void MeshTranslationImpl::RecordSourcePrimsForMaterialSlots(
+	const TArray<UsdUtils::FUsdPrimMaterialAssignmentInfo>& LODIndexToMaterialInfo,
+	UUsdMeshAssetUserData* UserData
+)
+{
+	if (!UserData)
+	{
+		return;
+	}
+
+	uint32 SlotIndex = 0;
+	for (int32 LODIndex = 0; LODIndex < LODIndexToMaterialInfo.Num(); ++LODIndex)
+	{
+		const TArray<UsdUtils::FUsdPrimMaterialSlot>& LODSlots = LODIndexToMaterialInfo[LODIndex].Slots;
+
+		for (int32 LODSlotIndex = 0; LODSlotIndex < LODSlots.Num(); ++LODSlotIndex, ++SlotIndex)
+		{
+			const UsdUtils::FUsdPrimMaterialSlot& Slot = LODSlots[LODSlotIndex];
+			UserData->MaterialSlotToPrimPaths.FindOrAdd(SlotIndex).PrimPaths.Append(Slot.PrimPaths.Array());
+		}
+	}
+}
+
 UMaterialInterface* MeshTranslationImpl::GetReferencePreviewSurfaceMaterial(EUsdReferenceMaterialProperties ReferenceMaterialProperties)
 {
 	const UUsdProjectSettings* Settings = GetDefault<UUsdProjectSettings>();
diff --git a/Engine/Plugins/Importers/USDImporter/Source/USDSchemas/Private/MeshTranslationImpl.h b/Engine/Plugins/Importers/USDImporter/Source/USDSchemas/Private/MeshTranslationImpl.h
--- a/Engine/Plugins/Importers/USDImporter/Source/USDSchemas/Private/MeshTranslationImpl.h
+++ b/Engine/Plugins/Importers/USDImporter/Source/USDSchemas/Private/MeshTranslationImpl.h
@@ -16,9 +16,10 @@ PXR_NAMESPACE_OPEN_SCOPE
 PXR_NAMESPACE_CLOSE_SCOPE
 
 class FUsdInfoCache;
-class UUsdAssetCache2;
 class UMaterialInterface;
 class UMeshComponent;
+class UUsdAssetCache2;
+class UUsdMeshAssetUserData;
 namespace UsdUtils
 {
 	struct FUsdPrimMaterialSlot;
@@ -54,6 +55,11 @@ namespace MeshTranslationImpl
 		const FName& MaterialPurpose
 	);
 
+	void RecordSourcePrimsForMaterialSlots(
+		const TArray<UsdUtils::FUsdPrimMaterialAssignmentInfo>& LODIndexToMaterialInfo,
+		UUsdMeshAssetUserData* UserData
+	);
+
 	enum class EUsdReferenceMaterialProperties
 	{
 		None = 0,
diff --git a/Engine/Plugins/Importers/USDImporter/Source/USDSchemas/Private/USDGeomMeshTranslator.cpp b/Engine/Plugins/Importers/USDImporter/Source/USDSchemas/Private/USDGeomMeshTranslator.cpp
--- a/Engine/Plugins/Importers/USDImporter/Source/USDSchemas/Private/USDGeomMeshTranslator.cpp
+++ b/Engine/Plugins/Importers/USDImporter/Source/USDSchemas/Private/USDGeomMeshTranslator.cpp
@@ -207,12 +207,6 @@ namespace UsdGeomMeshTranslatorImpl
 			ExistingAssignments.Add(StaticMaterial.MaterialInterface);
 		}
 
-		UUsdMeshAssetUserData* UserData = StaticMesh.GetAssetUserData<UUsdMeshAssetUserData>();
-		ensureMsgf(UserData, TEXT("Static Mesh '%s' generated for prim '%s' should have an UUsdMeshAssetUserData at this point!"),
-			*StaticMesh.GetPathName(),
-			*UsdToUnreal::ConvertPath(UsdPrim.GetPrimPath())
-		);
-
 		TMap<const UsdUtils::FUsdPrimMaterialSlot*, UMaterialInterface*> ResolvedMaterials = MeshTranslationImpl::ResolveMaterialAssignmentInfo(
 			UsdPrim,
 			LODIndexToMaterialInfo,
@@ -254,11 +248,6 @@ namespace UsdGeomMeshTranslatorImpl
 					bMaterialAssignementsHaveChanged = true;
 				}
 
-				if (UserData)
-				{
-					UserData->MaterialSlotToPrimPaths.FindOrAdd(StaticMeshSlotIndex).PrimPaths = Slot.PrimPaths.Array();
-				}
-
 #if WITH_EDITOR
 				// Setup the section map so that our LOD material index is properly mapped to the static mesh material index
 				// At runtime we don't ever parse these variants as LODs so we don't need this
@@ -817,14 +806,17 @@ void FBuildStaticMeshTaskChain::SetupTasks()
 				StaticMesh->NaniteSettings.bEnabled = bShouldEnableNanite;
 #endif // WITH_EDITOR
 
-				if (bIsNew)
+				UUsdMeshAssetUserData* UserData = StaticMesh->GetAssetUserData<UUsdMeshAssetUserData>();
+				if (!UserData)
 				{
-					UUsdMeshAssetUserData* UserData = NewObject<UUsdMeshAssetUserData>(StaticMesh, TEXT("UUSDAssetUserData"));
+					UserData = NewObject<UUsdMeshAssetUserData>(StaticMesh, TEXT("UUSDAssetUserData"));
 					UserData->PrimPath = MeshName;
-					UserData->PrimvarToUVIndex = LODIndexToMaterialInfo[0].PrimvarToUVIndex;  // We use the same primvar mapping for all LODs
+					UserData->PrimvarToUVIndex = LODIndexToMaterialInfo[0].PrimvarToUVIndex;	// We use the same primvar mapping for all LODs
 					StaticMesh->AddAssetUserData(UserData);
 				}
 
+				MeshTranslationImpl::RecordSourcePrimsForMaterialSlots(LODIndexToMaterialInfo, UserData);
+
 				// Only the original creator of the prim at creation time gets to set the material assignments
 				// directly on the mesh, all others prims ensure their materials via material overrides on the
 				// components
diff --git a/Engine/Plugins/Importers/USDImporter/Source/USDSchemas/Private/USDGeomXformableTranslator.cpp b/Engine/Plugins/Importers/USDImporter/Source/USDSchemas/Private/USDGeomXformableTranslator.cpp
--- a/Engine/Plugins/Importers/USDImporter/Source/USDSchemas/Private/USDGeomXformableTranslator.cpp
+++ b/Engine/Plugins/Importers/USDImporter/Source/USDSchemas/Private/USDGeomXformableTranslator.cpp
@@ -289,8 +289,7 @@ USceneComponent* FUsdGeomXformableTranslator::CreateComponents()
 	// the same: A single static mesh will be shared between them and one of the task chains will manage to put their
 	// material assignments on the mesh directly. To ensure the correct materials for the second subtree, we need to
 	// set overrides.
-	// Note how we don't have to handle geometry caches in here like the geom mesh translator does though: We don't
-	// collapse geometry caches, at least for now
+	// Note how we don't have to handle geometry caches in here as they're handled by the geometry cache translator now
 	if (UStaticMeshComponent* StaticMeshComponent = Cast<UStaticMeshComponent>(SceneComponent))
 	{
 		if (Context->InfoCache)
diff --git a/Engine/Plugins/Importers/USDImporter/Source/USDSchemas/Private/USDGeometryCacheTranslator.cpp b/Engine/Plugins/Importers/USDImporter/Source/USDSchemas/Private/USDGeometryCacheTranslator.cpp
--- a/Engine/Plugins/Importers/USDImporter/Source/USDSchemas/Private/USDGeometryCacheTranslator.cpp
+++ b/Engine/Plugins/Importers/USDImporter/Source/USDSchemas/Private/USDGeometryCacheTranslator.cpp
@@ -61,12 +61,6 @@ namespace UsdGeometryCacheTranslatorImpl
 	bool ProcessGeometryCacheMaterials(const pxr::UsdPrim& UsdPrim, const TArray<UsdUtils::FUsdPrimMaterialAssignmentInfo>& LODIndexToMaterialInfo,
 									   UGeometryCache& GeometryCache, UUsdAssetCache2& AssetCache, FUsdInfoCache* InfoCache, float Time, EObjectFlags Flags)
 	{
-		UUsdMeshAssetUserData* MeshUserData = GeometryCache.GetAssetUserData<UUsdMeshAssetUserData>();
-		ensureMsgf(MeshUserData, TEXT("Geometry Cache '%s' generated for prim '%s' should have an UUsdMeshAssetUserData at this point!"),
-			*GeometryCache.GetPathName(),
-			*UsdToUnreal::ConvertPath(UsdPrim.GetPrimPath())
-		);
-
 		TMap<const UsdUtils::FUsdPrimMaterialSlot*, UMaterialInterface*> ResolvedMaterials = MeshTranslationImpl::ResolveMaterialAssignmentInfo(
 			UsdPrim, LODIndexToMaterialInfo, AssetCache, *InfoCache, Flags);
 
@@ -100,8 +94,6 @@ namespace UsdGeometryCacheTranslatorImpl
 					GeometryCache.Materials[SlotIndex] = Material;
 					bMaterialAssignementsHaveChanged = true;
 				}
-
-				MeshUserData->MaterialSlotToPrimPaths.FindOrAdd(SlotIndex).PrimPaths = Slot.PrimPaths.Array();
 			}
 		}
 
@@ -769,36 +761,45 @@ void FGeometryCacheCreateAssetsTaskChain::SetupTasks()
 			const FString PrimPathString = PrimPath.GetString();
 			GeometryCache.Reset(UsdGeometryCacheTranslatorImpl::CreateGeometryCache(PrimPathString, LODIndexToMeshDescription[0], MeshPrimPaths, MaterialOffsets, Context, bIsNew, StartTimeOffset));
 
-			if (GeometryCache && bIsNew)
+			if (GeometryCache)
 			{
-				UUsdMeshAssetUserData* UserData = NewObject<UUsdMeshAssetUserData>(GeometryCache.Get(), TEXT("UUSDAssetUserData"));
-				UserData->PrimPath = PrimPathString;
-				UserData->PrimvarToUVIndex = LODIndexToMaterialInfo[0].PrimvarToUVIndex;  // We use the same primvar mapping for all LODs
-				GeometryCache->AddAssetUserData(UserData);
-
-				UUsdAnimSequenceAssetUserData* AnimUserData = NewObject<UUsdAnimSequenceAssetUserData>(GeometryCache.Get(), TEXT("UsdAnimUserData"));
-				AnimUserData->PrimPath = PrimPathString;
-				AnimUserData->LayerStartOffsetSeconds = StartTimeOffset;
-				GeometryCache->AddAssetUserData(AnimUserData);
-
-				// Only the original creator of the prim at creation time gets to set the material assignments
-				// directly on the geometry cache, all others prims ensure their materials via material overrides on the
-				// components
-				UsdGeometryCacheTranslatorImpl::ProcessGeometryCacheMaterials(
-					GetPrim(),
-					LODIndexToMaterialInfo,
-					*GeometryCache,
-					*Context->AssetCache.Get(),
-					Context->InfoCache.Get(),
-					Context->Time,
-					Context->ObjectFlags
-				);
-			}
+				UUsdMeshAssetUserData* UserData = GeometryCache->GetAssetUserData<UUsdMeshAssetUserData>();
+				if (!UserData)
+				{
+					UserData = NewObject<UUsdMeshAssetUserData>(GeometryCache.Get(), TEXT("UUSDAssetUserData"));
+					UserData->PrimPath = PrimPathString;
+					UserData->PrimvarToUVIndex = LODIndexToMaterialInfo[0].PrimvarToUVIndex;	// We use the same primvar mapping for all LODs
+					GeometryCache->AddAssetUserData(UserData);
+
+					UUsdAnimSequenceAssetUserData* AnimUserData = NewObject<UUsdAnimSequenceAssetUserData>(GeometryCache.Get(), TEXT("UsdAnimUserData"));
+					AnimUserData->PrimPath = PrimPathString;
+					AnimUserData->LayerStartOffsetSeconds = StartTimeOffset;
+					GeometryCache->AddAssetUserData(AnimUserData);
+				}
 
-			if (GeometryCache && Context->InfoCache)
-			{
-				const UE::FSdfPath& TargetPath = AlternativePrimToLinkAssetsTo.IsSet() ? AlternativePrimToLinkAssetsTo.GetValue() : PrimPath;
-				Context->InfoCache->LinkAssetToPrim(TargetPath, GeometryCache.Get());
+				MeshTranslationImpl::RecordSourcePrimsForMaterialSlots(LODIndexToMaterialInfo, UserData);
+
+				if (bIsNew)
+				{
+					// Only the original creator of the prim at creation time gets to set the material assignments
+					// directly on the geometry cache, all others prims ensure their materials via material overrides on the
+					// components
+					UsdGeometryCacheTranslatorImpl::ProcessGeometryCacheMaterials(
+						GetPrim(),
+						LODIndexToMaterialInfo,
+						*GeometryCache,
+						*Context->AssetCache.Get(),
+						Context->InfoCache.Get(),
+						Context->Time,
+						Context->ObjectFlags
+					);
+				}
+
+				if (Context->InfoCache)
+				{
+					const UE::FSdfPath& TargetPath = AlternativePrimToLinkAssetsTo.IsSet() ? AlternativePrimToLinkAssetsTo.GetValue() : PrimPath;
+					Context->InfoCache->LinkAssetToPrim(TargetPath, GeometryCache.Get());
+				}
 			}
 
 			// Continue with the import steps
diff --git a/Engine/Plugins/Importers/USDImporter/Source/USDSchemas/Private/USDSkelRootTranslator.cpp b/Engine/Plugins/Importers/USDImporter/Source/USDSchemas/Private/USDSkelRootTranslator.cpp
--- a/Engine/Plugins/Importers/USDImporter/Source/USDSchemas/Private/USDSkelRootTranslator.cpp
+++ b/Engine/Plugins/Importers/USDImporter/Source/USDSchemas/Private/USDSkelRootTranslator.cpp
@@ -185,11 +185,6 @@ namespace UsdSkelRootTranslatorImpl
 					bMaterialsHaveChanged = true;
 				}
 
-				if (UserData)
-				{
-					UserData->MaterialSlotToPrimPaths.FindOrAdd(SkeletalMeshSlotIndex).PrimPaths = Slot.PrimPaths.Array();
-				}
-
 				// Already have a material at that LOD remap slot, need to reassign
 				if ( LODMaterialMap.IsValidIndex( LODSlotIndex ) )
 				{
@@ -1388,15 +1383,21 @@ namespace UsdSkelRootTranslatorImpl
 					);
 				}
 
-				if ( SkeletalMesh )
+				if (SkeletalMesh)
 				{
-					if (bIsNew)
+					UUsdMeshAssetUserData* UserData = SkeletalMesh->GetAssetUserData<UUsdMeshAssetUserData>();
+					if (!UserData)
 					{
-						UUsdMeshAssetUserData* UserData = NewObject<UUsdMeshAssetUserData>(SkeletalMesh, TEXT("USDAssetUserData"));
+						UserData = NewObject<UUsdMeshAssetUserData>(SkeletalMesh, TEXT("USDAssetUserData"));
 						UserData->PrimPath = SkelRootPath;
-						UserData->PrimvarToUVIndex = LODIndexToMaterialInfo[0].PrimvarToUVIndex;  // We use the same primvar mapping for all LODs
+						UserData->PrimvarToUVIndex = LODIndexToMaterialInfo[0].PrimvarToUVIndex;	// We use the same primvar mapping for all LODs
 						SkeletalMesh->AddAssetUserData(UserData);
+					}
+
+					MeshTranslationImpl::RecordSourcePrimsForMaterialSlots(LODIndexToMaterialInfo, UserData);
 
+					if (bIsNew)
+					{
 						const bool bMaterialsHaveChanged = UsdSkelRootTranslatorImpl::ProcessMaterials(
 							GetPrim(),
 							LODIndexToMaterialInfo,
@@ -1408,10 +1409,10 @@ namespace UsdSkelRootTranslatorImpl
 							NewBlendShapes.Num() > 0
 						);
 
-						if ( bMaterialsHaveChanged )
+						if (bMaterialsHaveChanged)
 						{
 							const bool bRebuildAll = true;
-							SkeletalMesh->UpdateUVChannelData( bRebuildAll );
+							SkeletalMesh->UpdateUVChannelData(bRebuildAll);
 						}
 
 						Context->AssetCache->CacheAsset( SkeletalMeshHashString, SkeletalMesh );
diff --git a/Engine/Plugins/Importers/USDImporter/Source/USDUtilities/Private/USDPrimConversion.cpp b/Engine/Plugins/Importers/USDImporter/Source/USDUtilities/Private/USDPrimConversion.cpp
--- a/Engine/Plugins/Importers/USDImporter/Source/USDUtilities/Private/USDPrimConversion.cpp
+++ b/Engine/Plugins/Importers/USDImporter/Source/USDUtilities/Private/USDPrimConversion.cpp
@@ -2164,6 +2164,7 @@ bool UnrealToUsd::ConvertMaterialOverrides(
 
 	FScopedUsdAllocs Allocs;
 	pxr::UsdStageRefPtr Stage = UsdPrim.GetStage();
+	FString UsdPrimPath = UsdToUnreal::ConvertPath(UsdPrim.GetPrimPath());
 
 	const bool bAllLODs = LowestLOD == INDEX_NONE && HighestLOD == INDEX_NONE;
 
@@ -2187,9 +2188,15 @@ bool UnrealToUsd::ConvertMaterialOverrides(
 				{
 					for (const FString& PrimPath : SourcePrimPaths->PrimPaths)
 					{
-						pxr::SdfPath OverridePrimPath = UnrealToUsd::ConvertPath(*PrimPath).Get();
-						pxr::UsdPrim MeshPrim = Stage->OverridePrim(OverridePrimPath);
-						UE::USDPrimConversionImpl::Private::AuthorMaterialOverride(MeshPrim, Override->GetPathName());
+						// Our mesh assets are shared between multiple prims via the asset cache, and all user prims of that asset will record
+						// their source prim paths for each material slot. In here we just want to apply overrides to the prims that
+						// correspond to the modified component, and not the others
+						if (PrimPath.StartsWith(UsdPrimPath))
+						{
+							pxr::SdfPath OverridePrimPath = UnrealToUsd::ConvertPath(*PrimPath).Get();
+							pxr::UsdPrim MeshPrim = Stage->OverridePrim(OverridePrimPath);
+							UE::USDPrimConversionImpl::Private::AuthorMaterialOverride(MeshPrim, Override->GetPathName());
+						}
 					}
 				}
 			}
@@ -2257,9 +2264,13 @@ bool UnrealToUsd::ConvertMaterialOverrides(
 						{
 							for (const FString& PrimPath : SourcePrimPaths->PrimPaths)
 							{
-								pxr::SdfPath OverridePrimPath = UnrealToUsd::ConvertPath(*PrimPath).Get();
-								pxr::UsdPrim MeshPrim = Stage->OverridePrim(OverridePrimPath);
-								UE::USDPrimConversionImpl::Private::AuthorMaterialOverride(MeshPrim, Override->GetPathName());
+								// See comment on the analogue part for the geometry cache component
+								if (PrimPath.StartsWith(UsdPrimPath))
+								{
+									pxr::SdfPath OverridePrimPath = UnrealToUsd::ConvertPath(*PrimPath).Get();
+									pxr::UsdPrim MeshPrim = Stage->OverridePrim(OverridePrimPath);
+									UE::USDPrimConversionImpl::Private::AuthorMaterialOverride(MeshPrim, Override->GetPathName());
+								}
 							}
 						}
 					}
@@ -2394,9 +2405,13 @@ bool UnrealToUsd::ConvertMaterialOverrides(
 						{
 							for (const FString& PrimPath : SourcePrimPaths->PrimPaths)
 							{
-								pxr::SdfPath OverridePrimPath = UnrealToUsd::ConvertPath(*PrimPath).Get();
-								pxr::UsdPrim MeshPrim = Stage->OverridePrim(OverridePrimPath);
-								UE::USDPrimConversionImpl::Private::AuthorMaterialOverride(MeshPrim, Override->GetPathName());
+								// See comment on the analogue part for the geometry cache component
+								if (PrimPath.StartsWith(UsdPrimPath))
+								{
+									pxr::SdfPath OverridePrimPath = UnrealToUsd::ConvertPath(*PrimPath).Get();
+									pxr::UsdPrim MeshPrim = Stage->OverridePrim(OverridePrimPath);
+									UE::USDPrimConversionImpl::Private::AuthorMaterialOverride(MeshPrim, Override->GetPathName());
+								}
 							}
 						}
 					}
