USD: Fix level export with no levels selected generating some USD files with default data.

diff --git a/Engine/Plugins/Importers/USDImporter/Source/USDExporter/Private/USDConversionBlueprintLibrary.cpp b/Engine/Plugins/Importers/USDImporter/Source/USDExporter/Private/USDConversionBlueprintLibrary.cpp
--- a/Engine/Plugins/Importers/USDImporter/Source/USDExporter/Private/USDConversionBlueprintLibrary.cpp
+++ b/Engine/Plugins/Importers/USDImporter/Source/USDExporter/Private/USDConversionBlueprintLibrary.cpp
@@ -168,6 +168,37 @@ namespace UE
 	}
 }
 
+int32 UUsdConversionBlueprintLibrary::GetNumLevelsToExport(UWorld* World, const TSet<FString>& LevelsToIgnore)
+{
+	int32 Count = 0;
+
+	if (!World)
+	{
+		return Count;
+	}
+
+	if (!LevelsToIgnore.Contains("Persistent Level"))
+	{
+		Count += 1;
+	}
+
+	for (ULevelStreaming* StreamingLevel : World->GetStreamingLevels())
+	{
+		if (StreamingLevel)
+		{
+			const FString LevelName = FPaths::GetBaseFilename(StreamingLevel->GetWorldAssetPackageName());
+			if (!LevelsToIgnore.Contains(LevelName))
+			{
+				continue;
+			}
+
+			Count += 1;
+		}
+	}
+
+	return Count;
+}
+
 void UUsdConversionBlueprintLibrary::StreamInRequiredLevels( UWorld* World, const TSet<FString>& LevelsToIgnore )
 {
 	if ( !World )
diff --git a/Engine/Plugins/Importers/USDImporter/Source/USDExporter/Public/USDConversionBlueprintLibrary.h b/Engine/Plugins/Importers/USDImporter/Source/USDExporter/Public/USDConversionBlueprintLibrary.h
--- a/Engine/Plugins/Importers/USDImporter/Source/USDExporter/Public/USDConversionBlueprintLibrary.h
+++ b/Engine/Plugins/Importers/USDImporter/Source/USDExporter/Public/USDConversionBlueprintLibrary.h
@@ -20,6 +20,10 @@ class USDEXPORTER_API UUsdConversionBlueprintLibrary : public UBlueprintFunction
 	GENERATED_BODY()
 
 public:
+	/** Returns how many total Unreal levels (persistent + all sublevels) will be exported if we consider LevelsToIgnore */
+	UFUNCTION( BlueprintCallable, Category = "World utils" )
+	static int32 GetNumLevelsToExport( UWorld* World, const TSet<FString>& LevelsToIgnore );
+
 	/** Fully streams in and displays all levels whose names are not in LevelsToIgnore */
 	UFUNCTION( BlueprintCallable, Category = "World utils" )
 	static void StreamInRequiredLevels( UWorld* World, const TSet<FString>& LevelsToIgnore );
