USD: Allow recording multiple source prim paths on AssetUserData.

diff --git a/Engine/Plugins/Importers/USDImporter/Source/USDClasses/Public/USDAssetUserData.h b/Engine/Plugins/Importers/USDImporter/Source/USDClasses/Public/USDAssetUserData.h
--- a/Engine/Plugins/Importers/USDImporter/Source/USDClasses/Public/USDAssetUserData.h
+++ b/Engine/Plugins/Importers/USDImporter/Source/USDClasses/Public/USDAssetUserData.h
@@ -13,7 +13,7 @@ class USDCLASSES_API UUsdAssetUserData : public UAssetUserData
 
 public:
 	UPROPERTY()
-	FString PrimPath;
+	TArray<FString> PrimPaths;
 };
 
 UCLASS()
diff --git a/Engine/Plugins/Importers/USDImporter/Source/USDSchemas/Private/Custom/MaterialXUSDShadeMaterialTranslator.cpp b/Engine/Plugins/Importers/USDImporter/Source/USDSchemas/Private/Custom/MaterialXUSDShadeMaterialTranslator.cpp
--- a/Engine/Plugins/Importers/USDImporter/Source/USDSchemas/Private/Custom/MaterialXUSDShadeMaterialTranslator.cpp
+++ b/Engine/Plugins/Importers/USDImporter/Source/USDSchemas/Private/Custom/MaterialXUSDShadeMaterialTranslator.cpp
@@ -246,7 +246,7 @@ namespace UE::USDMaterialXTranslator::Private
 						if (!ExistingTexture)
 						{
 							UUsdAssetUserData* UserData = NewObject<UUsdAssetUserData>(Texture, TEXT("USDAssetUserData"));
-							UserData->PrimPath = ReferencerPrimPath;
+							UserData->PrimPaths = {ReferencerPrimPath};
 							Texture->AddAssetUserData(UserData);
 
 							AssetCache.CacheAsset(TextureHash, Texture);
diff --git a/Engine/Plugins/Importers/USDImporter/Source/USDSchemas/Private/USDGeomMeshTranslator.cpp b/Engine/Plugins/Importers/USDImporter/Source/USDSchemas/Private/USDGeomMeshTranslator.cpp
--- a/Engine/Plugins/Importers/USDImporter/Source/USDSchemas/Private/USDGeomMeshTranslator.cpp
+++ b/Engine/Plugins/Importers/USDImporter/Source/USDSchemas/Private/USDGeomMeshTranslator.cpp
@@ -810,10 +810,10 @@ void FBuildStaticMeshTaskChain::SetupTasks()
 				if (!UserData)
 				{
 					UserData = NewObject<UUsdMeshAssetUserData>(StaticMesh, TEXT("UUSDAssetUserData"));
-					UserData->PrimPath = MeshName;
 					UserData->PrimvarToUVIndex = LODIndexToMaterialInfo[0].PrimvarToUVIndex;	// We use the same primvar mapping for all LODs
 					StaticMesh->AddAssetUserData(UserData);
 				}
+				UserData->PrimPaths.AddUnique(MeshName);
 
 				MeshTranslationImpl::RecordSourcePrimsForMaterialSlots(LODIndexToMaterialInfo, UserData);
 
diff --git a/Engine/Plugins/Importers/USDImporter/Source/USDSchemas/Private/USDGeomPointInstancerTranslator.cpp b/Engine/Plugins/Importers/USDImporter/Source/USDSchemas/Private/USDGeomPointInstancerTranslator.cpp
--- a/Engine/Plugins/Importers/USDImporter/Source/USDSchemas/Private/USDGeomPointInstancerTranslator.cpp
+++ b/Engine/Plugins/Importers/USDImporter/Source/USDSchemas/Private/USDGeomPointInstancerTranslator.cpp
@@ -407,8 +407,11 @@ void FUsdGeomPointInstancerTranslator::UpdateComponents(USceneComponent* PointIn
 	{
 		if (UUsdAssetUserData* UserData = PrototypeMesh->GetAssetUserData<UUsdAssetUserData>())
 		{
-			pxr::SdfPath PrototypePath = UnrealToUsd::ConvertPath(*UserData->PrimPath).Get();
-			PrototypeMeshes.insert({PrototypePath, PrototypeMesh});
+			for (const FString& SourcePrimPath : UserData->PrimPaths)
+			{
+				pxr::SdfPath PrototypePath = UnrealToUsd::ConvertPath(*SourcePrimPath).Get();
+				PrototypeMeshes.insert({PrototypePath, PrototypeMesh});
+			}
 		}
 	}
 
diff --git a/Engine/Plugins/Importers/USDImporter/Source/USDSchemas/Private/USDGeometryCacheTranslator.cpp b/Engine/Plugins/Importers/USDImporter/Source/USDSchemas/Private/USDGeometryCacheTranslator.cpp
--- a/Engine/Plugins/Importers/USDImporter/Source/USDSchemas/Private/USDGeometryCacheTranslator.cpp
+++ b/Engine/Plugins/Importers/USDImporter/Source/USDSchemas/Private/USDGeometryCacheTranslator.cpp
@@ -767,12 +767,10 @@ void FGeometryCacheCreateAssetsTaskChain::SetupTasks()
 				if (!UserData)
 				{
 					UserData = NewObject<UUsdMeshAssetUserData>(GeometryCache.Get(), TEXT("UUSDAssetUserData"));
-					UserData->PrimPath = PrimPathString;
 					UserData->PrimvarToUVIndex = LODIndexToMaterialInfo[0].PrimvarToUVIndex;	// We use the same primvar mapping for all LODs
 					GeometryCache->AddAssetUserData(UserData);
 
 					UUsdAnimSequenceAssetUserData* AnimUserData = NewObject<UUsdAnimSequenceAssetUserData>(GeometryCache.Get(), TEXT("UsdAnimUserData"));
-					AnimUserData->PrimPath = PrimPathString;
 					AnimUserData->LayerStartOffsetSeconds = StartTimeOffset;
 					GeometryCache->AddAssetUserData(AnimUserData);
 				}
diff --git a/Engine/Plugins/Importers/USDImporter/Source/USDSchemas/Private/USDGroomTranslator.cpp b/Engine/Plugins/Importers/USDImporter/Source/USDSchemas/Private/USDGroomTranslator.cpp
--- a/Engine/Plugins/Importers/USDImporter/Source/USDSchemas/Private/USDGroomTranslator.cpp
+++ b/Engine/Plugins/Importers/USDImporter/Source/USDSchemas/Private/USDGroomTranslator.cpp
@@ -197,7 +197,7 @@ protected:
 						Context->AssetCache->CacheAsset(SHAHash.ToString(), GroomAsset);
 
 						UUsdAssetUserData* UserData = NewObject<UUsdAssetUserData>(GroomAsset, TEXT("UUSDAssetUserData"));
-						UserData->PrimPath = PrimPathString;
+						UserData->PrimPaths = {PrimPathString};
 
 						GroomAsset->AddAssetUserData(UserData);
 					}
@@ -340,7 +340,7 @@ protected:
 				if (GroomCache)
 				{
 					UUsdAssetUserData* UserData = NewObject<UUsdAssetUserData>(GroomCache, TEXT("UUSDAssetUserData"));
-					UserData->PrimPath = PrimPath.GetString();
+					UserData->PrimPaths = {PrimPath.GetString()};
 
 					GroomCache->AddAssetUserData(UserData);
 
diff --git a/Engine/Plugins/Importers/USDImporter/Source/USDSchemas/Private/USDShadeMaterialTranslator.cpp b/Engine/Plugins/Importers/USDImporter/Source/USDSchemas/Private/USDShadeMaterialTranslator.cpp
--- a/Engine/Plugins/Importers/USDImporter/Source/USDSchemas/Private/USDShadeMaterialTranslator.cpp
+++ b/Engine/Plugins/Importers/USDImporter/Source/USDSchemas/Private/USDShadeMaterialTranslator.cpp
@@ -279,7 +279,7 @@ void FUsdShadeMaterialTranslator::CreateAssets()
 			if ( UMaterialInstanceConstant* NewMaterial = NewObject<UMaterialInstanceConstant>( GetTransientPackage(), InstanceName, Context->ObjectFlags | EObjectFlags::RF_Transient ) )
 			{
 				UUsdMaterialAssetUserData* UserData = NewObject<UUsdMaterialAssetUserData>(NewMaterial, TEXT("USDAssetUserData"));
-				UserData->PrimPath = PrimPath.GetString();
+				UserData->PrimPaths = {PrimPath.GetString()};
 				NewMaterial->AddAssetUserData(UserData);
 
 				const bool bSuccess = UsdToUnreal::ConvertMaterial(
@@ -377,7 +377,7 @@ void FUsdShadeMaterialTranslator::CreateAssets()
 				if ( UMaterialInstanceDynamic* NewMaterial = UMaterialInstanceDynamic::Create( ReferenceMaterial, GetTransientPackage(), InstanceName ) )
 				{
 					UUsdMaterialAssetUserData* UserData = NewObject<UUsdMaterialAssetUserData>(NewMaterial, TEXT("USDAssetUserData"));
-					UserData->PrimPath = PrimPath.GetString();
+					UserData->PrimPaths = {PrimPath.GetString()};
 					NewMaterial->AddAssetUserData(UserData);
 
 					NewMaterial->SetFlags(RF_Transient);
@@ -457,7 +457,7 @@ void FUsdShadeMaterialTranslator::PostImportMaterial(const FString& MaterialHash
 		UserData = NewObject<UUsdMaterialAssetUserData>(ImportedMaterial, TEXT("USDAssetUserData"));
 		ImportedMaterial->AddAssetUserData(UserData);
 	}
-	UserData->PrimPath = PrimPath.GetString();
+	UserData->PrimPaths.AddUnique(PrimPath.GetString());
 
 	// Note that this needs to run even if we found this material in the asset cache already, otherwise we won't
 	// re-register the prim asset links when we reload a stage
@@ -529,7 +529,7 @@ void FUsdShadeMaterialTranslator::PostImportMaterial(const FString& MaterialHash
 					TextureUserData = NewObject<UUsdAssetUserData>(ImportedMaterial, TEXT("USDAssetUserData"));
 					Texture->AddAssetUserData(TextureUserData);
 				}
-				TextureUserData->PrimPath = PrimPath.GetString();
+				TextureUserData->PrimPaths.AddUnique(PrimPath.GetString());
 
 				Context->InfoCache->LinkAssetToPrim(PrimPath, Texture);
 			}
diff --git a/Engine/Plugins/Importers/USDImporter/Source/USDSchemas/Private/USDSkelRootTranslator.cpp b/Engine/Plugins/Importers/USDImporter/Source/USDSchemas/Private/USDSkelRootTranslator.cpp
--- a/Engine/Plugins/Importers/USDImporter/Source/USDSchemas/Private/USDSkelRootTranslator.cpp
+++ b/Engine/Plugins/Importers/USDImporter/Source/USDSchemas/Private/USDSkelRootTranslator.cpp
@@ -1389,7 +1389,6 @@ namespace UsdSkelRootTranslatorImpl
 					if (!UserData)
 					{
 						UserData = NewObject<UUsdMeshAssetUserData>(SkeletalMesh, TEXT("USDAssetUserData"));
-						UserData->PrimPath = SkelRootPath;
 						UserData->PrimvarToUVIndex = LODIndexToMaterialInfo[0].PrimvarToUVIndex;	// We use the same primvar mapping for all LODs
 						SkeletalMesh->AddAssetUserData(UserData);
 					}
@@ -1630,7 +1629,7 @@ namespace UsdSkelRootTranslatorImpl
 							if (AnimSequence->GetDataModel()->GetNumBoneTracks() != 0 || AnimSequence->GetDataModel()->GetNumberOfFloatCurves() != 0 )
 							{
 								UUsdAnimSequenceAssetUserData* UserData = NewObject<UUsdAnimSequenceAssetUserData>(AnimSequence, TEXT("USDAssetUserData"));
-								UserData->PrimPath = SkelAnimationPrimPath;
+								UserData->PrimPaths = {SkelAnimationPrimPath};
 								UserData->LayerStartOffsetSeconds = LayerStartOffsetSeconds;
 
 								AnimSequence->AddAssetUserData(UserData);
diff --git a/Engine/Plugins/Importers/USDImporter/Source/USDStageImporter/Private/USDStageAssetImportFactory.cpp b/Engine/Plugins/Importers/USDImporter/Source/USDStageImporter/Private/USDStageAssetImportFactory.cpp
--- a/Engine/Plugins/Importers/USDImporter/Source/USDStageImporter/Private/USDStageAssetImportFactory.cpp
+++ b/Engine/Plugins/Importers/USDImporter/Source/USDStageImporter/Private/USDStageAssetImportFactory.cpp
@@ -153,7 +153,10 @@ EReimportResult::Type UUsdStageAssetImportFactory::Reimport(UObject* Obj)
 	{
 		if (UUsdAssetUserData* UserData = UserDataInterface->GetAssetUserData<UUsdAssetUserData>())
 		{
-			OriginalPrimPath = UserData->PrimPath;
+			if (!UserData->PrimPaths.IsEmpty())
+			{
+				OriginalPrimPath = UserData->PrimPaths[0];
+			}
 		}
 	}
 
diff --git a/Engine/Plugins/Importers/USDImporter/Source/USDStageImporter/Private/USDStageImporter.cpp b/Engine/Plugins/Importers/USDImporter/Source/USDStageImporter/Private/USDStageImporter.cpp
--- a/Engine/Plugins/Importers/USDImporter/Source/USDStageImporter/Private/USDStageImporter.cpp
+++ b/Engine/Plugins/Importers/USDImporter/Source/USDStageImporter/Private/USDStageImporter.cpp
@@ -438,8 +438,11 @@ namespace UsdStageImporterImpl
 		{
 			if (UUsdAssetUserData* UserData = UserDataInterface->GetAssetUserData<UUsdAssetUserData>())
 			{
-				PrimPath = UserData->PrimPath;
-				AssetPath = PrimPath;
+				if (!UserData->PrimPaths.IsEmpty())
+				{
+					PrimPath = UserData->PrimPaths[0];
+					AssetPath = PrimPath;
+				}
 			}
 		}
 
@@ -940,11 +943,11 @@ namespace UsdStageImporterImpl
 						// as we'll always have at least the skeletal mesh and the skeleton
 						if (AssetTypeFolderPtr == &SkeletalMeshesFolder)
 						{
-							AssetToContentFolder.Add(Asset, UserData->PrimPath);
+							AssetToContentFolder.Add(Asset, UserData->PrimPaths.IsEmpty() ? FString{} : UserData->PrimPaths[0]);
 						}
 						else
 						{
-							AssetToContentFolder.Add(Asset, FPaths::GetPath(UserData->PrimPath));
+							AssetToContentFolder.Add(Asset, FPaths::GetPath(UserData->PrimPaths.IsEmpty() ? FString{} : UserData->PrimPaths[0]));
 						}
 					}
 				}
@@ -1976,7 +1979,7 @@ bool UUsdStageImporter::ReimportSingleAsset(
 		if ( Asset &&
 			 UserData &&
 			 Asset->GetClass() == OriginalAsset->GetClass() &&
-			 UserData->PrimPath.Equals(OriginalPrimPath, ESearchCase::CaseSensitive))
+			 UserData->PrimPaths.Contains(OriginalPrimPath))
 		{
 			ReimportedObject = Asset;
 			break;
diff --git a/Engine/Plugins/Importers/USDImporter/Source/USDUtilities/Private/USDShadeConversion.cpp b/Engine/Plugins/Importers/USDImporter/Source/USDUtilities/Private/USDShadeConversion.cpp
--- a/Engine/Plugins/Importers/USDImporter/Source/USDUtilities/Private/USDShadeConversion.cpp
+++ b/Engine/Plugins/Importers/USDImporter/Source/USDUtilities/Private/USDShadeConversion.cpp
@@ -706,9 +706,12 @@ namespace UE
 								FString MaterialPrimPath;
 								if ( Material )
 								{
-									if(UUsdAssetUserData* UserData = Material->GetAssetUserData<UUsdAssetUserData>())
+									if (UUsdAssetUserData* UserData = Material->GetAssetUserData<UUsdAssetUserData>())
 									{
-										MaterialPrimPath = UserData->PrimPath;
+										if (!UserData->PrimPaths.IsEmpty())
+										{
+											MaterialPrimPath = UserData->PrimPaths[0];
+										}
 									}
 								}
 
@@ -1336,7 +1339,7 @@ namespace UE
 						if ( Texture )
 						{
 							UUsdAssetUserData* UserData = NewObject<UUsdAssetUserData>(Texture, TEXT("USDAssetUserData"));
-							UserData->PrimPath = PrimPath;
+							UserData->PrimPaths = {PrimPath};
 							Texture->AddAssetUserData(UserData);
 
 							// We set this even if we're not going to import so that we can track our original texture filepath
@@ -2770,14 +2773,18 @@ void UsdUtils::NotifyIfVirtualTexturesNeeded( UTexture* Texture )
 	FString TexturePath = Texture->GetName();
 	if (UUsdAssetUserData* UserData = Texture->GetAssetUserData<UUsdAssetUserData>())
 	{
-		TexturePath = UserData->PrimPath;
+		if (!UserData->PrimPaths.IsEmpty())
+		{
+			TexturePath = UserData->PrimPaths[0];
+		}
 	}
 
 	if ( !UseVirtualTexturing( GMaxRHIFeatureLevel ) )
 	{
 		FUsdLogManager::LogMessage(
 			EMessageSeverity::Warning,
-			FText::Format( LOCTEXT( "DisabledVirtualTexturing", "Texture '{0}' requires Virtual Textures, but the feature is disabled for this project" ),
+			FText::Format( LOCTEXT( "DisabledVirtualTexturing", "Texture '{0}' (from prim '{1}') requires Virtual Textures, but the feature is disabled for this project" ),
+				FText::FromString( Texture->GetName() ),
 				FText::FromString( TexturePath ) )
 		);
 	}
