USD: Use thread-safe delegates for events fired in response to USD notices as those can come from multiple threads concurrently.

diff --git a/Engine/Plugins/Importers/USDImporter/Source/UnrealUSDWrapper/Public/USDListener.h b/Engine/Plugins/Importers/USDImporter/Source/UnrealUSDWrapper/Public/USDListener.h
--- a/Engine/Plugins/Importers/USDImporter/Source/UnrealUSDWrapper/Public/USDListener.h
+++ b/Engine/Plugins/Importers/USDImporter/Source/UnrealUSDWrapper/Public/USDListener.h
@@ -119,22 +119,26 @@ public:
 	bool IsBlocked() const;
 
 	// Stage-specific events
-	DECLARE_EVENT( FUsdListener, FOnStageEditTargetChanged );
+	DECLARE_TS_MULTICAST_DELEGATE(FOnStageEditTargetChanged);
 	FOnStageEditTargetChanged& GetOnStageEditTargetChanged();
 
-	DECLARE_EVENT_OneParam( FUsdListener, UE_DEPRECATED(5.3, "Use FOnSdfLayersChanged") FOnLayersChanged, const TArray< FString >& );
+	DECLARE_TS_MULTICAST_DELEGATE_OneParam(FOnLayersChanged, const TArray<FString>&);
 	UE_DEPRECATED(5.3, "Use GetOnSdfLayersChanged")
-  	PRAGMA_DISABLE_DEPRECATION_WARNINGS
+	PRAGMA_DISABLE_DEPRECATION_WARNINGS
 	FOnLayersChanged& GetOnLayersChanged();
-  	PRAGMA_ENABLE_DEPRECATION_WARNINGS
+	PRAGMA_ENABLE_DEPRECATION_WARNINGS
 
-	DECLARE_EVENT_OneParam(FUsdListener, FOnSdfLayersChanged, const UsdUtils::FLayerToSdfChangeList&);
+	DECLARE_TS_MULTICAST_DELEGATE_OneParam(FOnSdfLayersChanged, const UsdUtils::FLayerToSdfChangeList&);
 	FOnSdfLayersChanged& GetOnSdfLayersChanged();
 
-	DECLARE_EVENT(FUsdListener, FOnSdfLayerDirtinessChanged);
+	DECLARE_TS_MULTICAST_DELEGATE(FOnSdfLayerDirtinessChanged);
 	FOnSdfLayerDirtinessChanged& GetOnSdfLayerDirtinessChanged();
 
-	DECLARE_EVENT_TwoParams( FUsdListener, FOnObjectsChanged, const UsdUtils::FObjectChangesByPath& /* InfoChanges */, const UsdUtils::FObjectChangesByPath& /* ResyncChanges */ );
+	DECLARE_TS_MULTICAST_DELEGATE_TwoParams(
+		FOnObjectsChanged,
+		const UsdUtils::FObjectChangesByPath& /* InfoChanges */,
+		const UsdUtils::FObjectChangesByPath& /* ResyncChanges */
+	);
 	FOnObjectsChanged& GetOnObjectsChanged();
 
 private:
