USD: Fix importing of exported LOD meshes not spawning static mesh components;

Fix Xformable translator collapsing prims with LOD setup.

diff --git a/Engine/Plugins/Importers/USDImporter/Source/USDSchemas/Private/USDGeomMeshTranslator.cpp b/Engine/Plugins/Importers/USDImporter/Source/USDSchemas/Private/USDGeomMeshTranslator.cpp
--- a/Engine/Plugins/Importers/USDImporter/Source/USDSchemas/Private/USDGeomMeshTranslator.cpp
+++ b/Engine/Plugins/Importers/USDImporter/Source/USDSchemas/Private/USDGeomMeshTranslator.cpp
@@ -813,7 +813,6 @@ void FBuildStaticMeshTaskChain::SetupTasks()
 					UserData->PrimvarToUVIndex = LODIndexToMaterialInfo[0].PrimvarToUVIndex;	// We use the same primvar mapping for all LODs
 					StaticMesh->AddAssetUserData(UserData);
 				}
-				UserData->PrimPaths.AddUnique(MeshName);
 
 				MeshTranslationImpl::RecordSourcePrimsForMaterialSlots(LODIndexToMaterialInfo, UserData);
 
@@ -1146,7 +1145,7 @@ TSet<UE::FSdfPath> FUsdGeomMeshTranslator::CollectAuxiliaryPrims() const
 bool FUsdGeomMeshTranslator::IsMeshPrim() const
 {
 	UE::FUsdPrim Prim = GetPrim();
-	if(Prim && Prim.IsA(TEXT("Mesh")))
+	if (Prim && (Prim.IsA(TEXT("Mesh")) || (Context->bAllowInterpretingLODs && UsdUtils::DoesPrimContainMeshLODs(Prim))))
 	{
 		return true;
 	}
diff --git a/Engine/Plugins/Importers/USDImporter/Source/USDSchemas/Private/USDGeomXformableTranslator.cpp b/Engine/Plugins/Importers/USDImporter/Source/USDSchemas/Private/USDGeomXformableTranslator.cpp
--- a/Engine/Plugins/Importers/USDImporter/Source/USDSchemas/Private/USDGeomXformableTranslator.cpp
+++ b/Engine/Plugins/Importers/USDImporter/Source/USDSchemas/Private/USDGeomXformableTranslator.cpp
@@ -733,12 +733,10 @@ bool FUsdGeomXformableTranslator::CanBeCollapsed( ECollapsingType CollapsingType
 		return false;
 	}
 
-	if ( UsdUtils::IsAnimated( UsdPrim ) )
-	{
-		return false;
-	}
-
-	if ( UsdUtils::PrimHasSchema( UsdPrim, UnrealIdentifiers::LiveLinkAPI ) )
+	if (UsdUtils::IsAnimated(UsdPrim) ||
+		UsdUtils::PrimHasSchema(UsdPrim, UnrealIdentifiers::LiveLinkAPI) ||
+		(Context->bAllowInterpretingLODs && UsdUtils::DoesPrimContainMeshLODs(UsdPrim))
+	)
 	{
 		return false;
 	}
