USD: Always use the latest asset of a particular asset type when retrieving assets for prims.

diff --git a/Engine/Plugins/Importers/USDImporter/Source/USDSchemas/Private/MeshTranslationImpl.cpp b/Engine/Plugins/Importers/USDImporter/Source/USDSchemas/Private/MeshTranslationImpl.cpp
--- a/Engine/Plugins/Importers/USDImporter/Source/USDSchemas/Private/MeshTranslationImpl.cpp
+++ b/Engine/Plugins/Importers/USDImporter/Source/USDSchemas/Private/MeshTranslationImpl.cpp
@@ -495,7 +495,7 @@ TMap<const UsdUtils::FUsdPrimMaterialSlot*, UMaterialInterface*> MeshTranslation
 
 				UE::FSdfPath MaterialPrimPath{*Slot.MaterialSource};
 
-				TSet<UMaterialInterface*> ExistingMaterials = InfoCache.GetAssetsForPrim<UMaterialInterface>(MaterialPrimPath);
+				TArray<UMaterialInterface*> ExistingMaterials = InfoCache.GetAssetsForPrim<UMaterialInterface>(MaterialPrimPath);
 
 				for (UMaterialInterface* ExistingMaterial : ExistingMaterials)
 				{
diff --git a/Engine/Plugins/Importers/USDImporter/Source/USDSchemas/Private/USDGeomPointInstancerTranslator.cpp b/Engine/Plugins/Importers/USDImporter/Source/USDSchemas/Private/USDGeomPointInstancerTranslator.cpp
--- a/Engine/Plugins/Importers/USDImporter/Source/USDSchemas/Private/USDGeomPointInstancerTranslator.cpp
+++ b/Engine/Plugins/Importers/USDImporter/Source/USDSchemas/Private/USDGeomPointInstancerTranslator.cpp
@@ -400,17 +400,17 @@ void FUsdGeomPointInstancerTranslator::UpdateComponents(USceneComponent* PointIn
 	// prototype, as we translate these with task pools and some of those prototypes may have generated nullptr.
 	// We always put the prototype path on the asset import data though, so here we use that to figure out where
 	// each mesh should go
-	TSet<UStaticMesh*> PrototypeMeshSet = InfoCache.GetAssetsForPrim<UStaticMesh>(PrimPath);
+	TArray<UStaticMesh*> PrototypeMeshArr = InfoCache.GetAssetsForPrim<UStaticMesh>(PrimPath);
 	std::unordered_map<pxr::SdfPath, UStaticMesh*, pxr::SdfPath::Hash> PrototypeMeshes;
-	PrototypeMeshes.reserve(PrototypeMeshSet.Num());
-	for (UStaticMesh* PrototypeMesh : PrototypeMeshSet)
+	PrototypeMeshes.reserve(PrototypeMeshArr.Num());
+	for (UStaticMesh* PrototypeMesh : PrototypeMeshArr)
 	{
 		if (UUsdAssetUserData* UserData = PrototypeMesh->GetAssetUserData<UUsdAssetUserData>())
 		{
 			for (const FString& SourcePrimPath : UserData->PrimPaths)
 			{
 				pxr::SdfPath PrototypePath = UnrealToUsd::ConvertPath(*SourcePrimPath).Get();
-				PrototypeMeshes.insert({PrototypePath, PrototypeMesh});
+				PrototypeMeshes[PrototypePath] = PrototypeMesh;
 			}
 		}
 	}
diff --git a/Engine/Plugins/Importers/USDImporter/Source/USDSchemas/Private/USDInfoCache.cpp b/Engine/Plugins/Importers/USDImporter/Source/USDSchemas/Private/USDInfoCache.cpp
--- a/Engine/Plugins/Importers/USDImporter/Source/USDSchemas/Private/USDInfoCache.cpp
+++ b/Engine/Plugins/Importers/USDImporter/Source/USDSchemas/Private/USDInfoCache.cpp
@@ -121,8 +121,8 @@ struct FUsdInfoCache::FUsdInfoCacheImpl
 	mutable FRWLock InfoMapLock;
 
 	// Information we may have about a subset of prims
-	TMap<UE::FSdfPath, TSet<TWeakObjectPtr<UObject>>> PrimPathToAssets;
-	TMap<TWeakObjectPtr<UObject>, TSet<UE::FSdfPath>> AssetToPrimPaths;
+	TMap<UE::FSdfPath, TArray<TWeakObjectPtr<UObject>>> PrimPathToAssets;
+	TMap<TWeakObjectPtr<UObject>, TArray<UE::FSdfPath>> AssetToPrimPaths;
 	mutable FRWLock PrimPathToAssetsLock;
 
 	// Paths to material prims to the mesh prims they are bound to in the scene, given the current settings for
@@ -1449,11 +1449,11 @@ void FUsdInfoCache::LinkAssetToPrim(const UE::FSdfPath& Path, UObject* Asset)
 		*Path.GetString()
 	);
 
-	ImplPtr->PrimPathToAssets.FindOrAdd(Path).Add(Asset);
-	ImplPtr->AssetToPrimPaths.FindOrAdd(Asset).Add(Path);
+	ImplPtr->PrimPathToAssets.FindOrAdd(Path).AddUnique(Asset);
+	ImplPtr->AssetToPrimPaths.FindOrAdd(Asset).AddUnique(Path);
 }
 
-TSet<TWeakObjectPtr<UObject>> FUsdInfoCache::RemoveAllAssetPrimLinks(const UE::FSdfPath& Path)
+TArray<TWeakObjectPtr<UObject>> FUsdInfoCache::RemoveAllAssetPrimLinks(const UE::FSdfPath& Path)
 {
 	FUsdInfoCacheImpl* ImplPtr = Impl.Get();
 	if (!ImplPtr)
@@ -1466,12 +1466,12 @@ TSet<TWeakObjectPtr<UObject>> FUsdInfoCache::RemoveAllAssetPrimLinks(const UE::F
 		*Path.GetString()
 	);
 
-	TSet<TWeakObjectPtr<UObject>> Assets;
+	TArray<TWeakObjectPtr<UObject>> Assets;
 	ImplPtr->PrimPathToAssets.RemoveAndCopyValue(Path, Assets);
 
 	for (const TWeakObjectPtr<UObject>& Asset : Assets)
 	{
-		if (TSet<UE::FSdfPath>* PrimPaths = ImplPtr->AssetToPrimPaths.Find(Asset))
+		if (TArray<UE::FSdfPath>* PrimPaths = ImplPtr->AssetToPrimPaths.Find(Asset))
 		{
 			PrimPaths->Remove(Path);
 		}
@@ -1480,7 +1480,7 @@ TSet<TWeakObjectPtr<UObject>> FUsdInfoCache::RemoveAllAssetPrimLinks(const UE::F
 	return Assets;
 }
 
-TSet<TWeakObjectPtr<UObject>> FUsdInfoCache::GetAllAssetsForPrim(const UE::FSdfPath& Path) const
+TArray<TWeakObjectPtr<UObject>> FUsdInfoCache::GetAllAssetsForPrim(const UE::FSdfPath& Path) const
 {
 	FUsdInfoCacheImpl* ImplPtr = Impl.Get();
 	if (!ImplPtr)
@@ -1489,7 +1489,7 @@ TSet<TWeakObjectPtr<UObject>> FUsdInfoCache::GetAllAssetsForPrim(const UE::FSdfP
 	}
 	FReadScopeLock ScopeLock(ImplPtr->PrimPathToAssetsLock);
 
-	if (const TSet<TWeakObjectPtr<UObject>>* FoundAssets = ImplPtr->PrimPathToAssets.Find(Path))
+	if (const TArray<TWeakObjectPtr<UObject>>* FoundAssets = ImplPtr->PrimPathToAssets.Find(Path))
 	{
 		return *FoundAssets;
 	}
@@ -1497,7 +1497,7 @@ TSet<TWeakObjectPtr<UObject>> FUsdInfoCache::GetAllAssetsForPrim(const UE::FSdfP
 	return {};
 }
 
-TSet<UE::FSdfPath> FUsdInfoCache::GetPrimsForAsset(UObject* Asset) const
+TArray<UE::FSdfPath> FUsdInfoCache::GetPrimsForAsset(UObject* Asset) const
 {
 	if (!Asset)
 	{
@@ -1511,7 +1511,7 @@ TSet<UE::FSdfPath> FUsdInfoCache::GetPrimsForAsset(UObject* Asset) const
 	}
 	FReadScopeLock ScopeLock(ImplPtr->PrimPathToAssetsLock);
 
-	if (const TSet<UE::FSdfPath>* FoundPrims = ImplPtr->AssetToPrimPaths.Find(Asset))
+	if (const TArray<UE::FSdfPath>* FoundPrims = ImplPtr->AssetToPrimPaths.Find(Asset))
 	{
 		return *FoundPrims;
 	}
@@ -1519,7 +1519,7 @@ TSet<UE::FSdfPath> FUsdInfoCache::GetPrimsForAsset(UObject* Asset) const
 	return {};
 }
 
-TMap<UE::FSdfPath, TSet<TWeakObjectPtr<UObject>>> FUsdInfoCache::GetAllAssetPrimLinks() const
+TMap<UE::FSdfPath, TArray<TWeakObjectPtr<UObject>>> FUsdInfoCache::GetAllAssetPrimLinks() const
 {
 	FUsdInfoCacheImpl* ImplPtr = Impl.Get();
 	if (!ImplPtr)
diff --git a/Engine/Plugins/Importers/USDImporter/Source/USDSchemas/Public/USDInfoCache.h b/Engine/Plugins/Importers/USDImporter/Source/USDSchemas/Public/USDInfoCache.h
--- a/Engine/Plugins/Importers/USDImporter/Source/USDSchemas/Public/USDInfoCache.h
+++ b/Engine/Plugins/Importers/USDImporter/Source/USDSchemas/Public/USDInfoCache.h
@@ -90,17 +90,20 @@ public:
 public:
 	void LinkAssetToPrim(const UE::FSdfPath& Path, UObject* Asset);
 
-	TSet<TWeakObjectPtr<UObject>> RemoveAllAssetPrimLinks(const UE::FSdfPath& Path);
+	TArray<TWeakObjectPtr<UObject>> RemoveAllAssetPrimLinks(const UE::FSdfPath& Path);
 
-	TSet<TWeakObjectPtr<UObject>> GetAllAssetsForPrim(const UE::FSdfPath& Path) const;
+	TArray<TWeakObjectPtr<UObject>> GetAllAssetsForPrim(const UE::FSdfPath& Path) const;
 
 	template<typename T = UObject>
 	T* GetSingleAssetForPrim(const UE::FSdfPath& Path) const
 	{
-		TSet<TWeakObjectPtr<UObject>> Assets = GetAllAssetsForPrim(Path);
-		for (const TWeakObjectPtr<UObject>& Asset : Assets)
+		TArray<TWeakObjectPtr<UObject>> Assets = GetAllAssetsForPrim(Path);
+
+		// Search back to front so that if we generate a new version of an asset type we prefer
+		// returning that
+		for (int32 Index = Assets.Num() - 1; Index >= 0; --Index)
 		{
-			if (T* CastAsset = Cast<T>(Asset.Get()))
+			if (T* CastAsset = Cast<T>(Assets[Index].Get()))
 			{
 				return CastAsset;
 			}
@@ -110,11 +113,11 @@ public:
 	}
 
 	template<typename T>
-	TSet<T*> GetAssetsForPrim(const UE::FSdfPath& Path) const
+	TArray<T*> GetAssetsForPrim(const UE::FSdfPath& Path) const
 	{
-		TSet<TWeakObjectPtr<UObject>> Assets = GetAllAssetsForPrim(Path);
+		TArray<TWeakObjectPtr<UObject>> Assets = GetAllAssetsForPrim(Path);
 
-		TSet<T*> CastAssets;
+		TArray<T*> CastAssets;
 		CastAssets.Reserve(Assets.Num());
 
 		for (const TWeakObjectPtr<UObject>& Asset : Assets)
@@ -128,8 +131,8 @@ public:
 		return CastAssets;
 	}
 
-	TSet<UE::FSdfPath> GetPrimsForAsset(UObject* Asset) const;
-	TMap<UE::FSdfPath, TSet<TWeakObjectPtr<UObject>>> GetAllAssetPrimLinks() const;
+	TArray<UE::FSdfPath> GetPrimsForAsset(UObject* Asset) const;
+	TMap<UE::FSdfPath, TArray<TWeakObjectPtr<UObject>>> GetAllAssetPrimLinks() const;
 
 private:
 	TUniquePtr<FUsdInfoCacheImpl> Impl;
diff --git a/Engine/Plugins/Importers/USDImporter/Source/USDStage/Private/USDStageActor.cpp b/Engine/Plugins/Importers/USDImporter/Source/USDStage/Private/USDStageActor.cpp
--- a/Engine/Plugins/Importers/USDImporter/Source/USDStage/Private/USDStageActor.cpp
+++ b/Engine/Plugins/Importers/USDImporter/Source/USDStage/Private/USDStageActor.cpp
@@ -2243,16 +2243,16 @@ TArray<UObject*> AUsdStageActor::GetGeneratedAssets(const FString& PrimPath)
 	// Prefer checking the prim directly, but also check its collapsed root if it is collapsed.
 	// This because we have some exception cases like USkeleton/UAnimSequences that can be found by querying the actual
 	// Skeleton/SKelAnimation prims even though they are collapsed into the SkelRoot prim.
-	TSet<TWeakObjectPtr<UObject>> AssetsSet = InfoCache->GetAllAssetsForPrim(UsdPath);
-	if (AssetsSet.Num() == 0 && InfoCache->IsPathCollapsed(UsdPath, ECollapsingType::Assets))
+	TArray<TWeakObjectPtr<UObject>> AssetsPtrs = InfoCache->GetAllAssetsForPrim(UsdPath);
+	if (AssetsPtrs.Num() == 0 && InfoCache->IsPathCollapsed(UsdPath, ECollapsingType::Assets))
 	{
 		UsdPath = InfoCache->UnwindToNonCollapsedPath(UsdPath, ECollapsingType::Assets);
-		AssetsSet = InfoCache->GetAllAssetsForPrim(UsdPath);
+		AssetsPtrs = InfoCache->GetAllAssetsForPrim(UsdPath);
 	}
 
 	TArray<UObject*> Assets;
-	Assets.Reserve(AssetsSet.Num());
-	for (const TWeakObjectPtr<UObject>& Asset : AssetsSet)
+	Assets.Reserve(AssetsPtrs.Num());
+	for (const TWeakObjectPtr<UObject>& Asset : AssetsPtrs)
 	{
 		Assets.Add(Asset.Get());
 	}
@@ -2270,10 +2270,10 @@ FString AUsdStageActor::GetSourcePrimPath(UObject* Object)
 	}
 	else if (InfoCache)
 	{
-		const TSet<UE::FSdfPath> FoundPaths = InfoCache->GetPrimsForAsset(Object);
+		const TArray<UE::FSdfPath> FoundPaths = InfoCache->GetPrimsForAsset(Object);
 		if (FoundPaths.Num() > 0)
 		{
-			return FoundPaths.CreateConstIterator()->GetString();
+			return FoundPaths[0].GetString();
 		}
 	}
 
@@ -3946,6 +3946,12 @@ void AUsdStageActor::OnSkelAnimationBaked(const FString& SkelRootPrimPath)
 
 bool AUsdStageActor::UnloadAssets(const UE::FSdfPath& StartPrimPath, bool bForEntireSubtree)
 {
+	// Note that whenever we change a stage option (like render context, kinds to collapse, etc.) we may generate
+	// new assets for a prim but we won't call this function, which means we will still temporarily keep the old
+	// assets in the asset cache, and they will count as "referenced". That is not great, although they will *still*
+	// be tracked via the info cache asset prim links, so if at any time they resync the old assets will still be
+	// found below when iterating the prim links, and we will discard them either way.
+
 	if (!UsdAssetCache || !InfoCache)
 	{
 		return false;
@@ -3956,7 +3962,7 @@ bool AUsdStageActor::UnloadAssets(const UE::FSdfPath& StartPrimPath, bool bForEn
 	TSet<UE::FSdfPath> PrimPathsToRemove;
 	if (bForEntireSubtree)
 	{
-		for (const TPair<UE::FSdfPath, TSet<TWeakObjectPtr<UObject>>>& PrimPathToAssetIt : InfoCache->GetAllAssetPrimLinks())
+		for (const TPair<UE::FSdfPath, TArray<TWeakObjectPtr<UObject>>>& PrimPathToAssetIt : InfoCache->GetAllAssetPrimLinks())
 		{
 			const UE::FSdfPath& LinkPrimPath = PrimPathToAssetIt.Key;
 			if (LinkPrimPath.HasPrefix(StartPrimPath) || LinkPrimPath == StartPrimPath)
@@ -3972,7 +3978,7 @@ bool AUsdStageActor::UnloadAssets(const UE::FSdfPath& StartPrimPath, bool bForEn
 
 	for (const UE::FSdfPath& PrimPathToRemove : PrimPathsToRemove)
 	{
-		TSet<TWeakObjectPtr<UObject>> OldAssets = InfoCache->RemoveAllAssetPrimLinks(PrimPathToRemove);
+		TArray<TWeakObjectPtr<UObject>> OldAssets = InfoCache->RemoveAllAssetPrimLinks(PrimPathToRemove);
 		for (const TWeakObjectPtr<UObject>& OldAsset : OldAssets)
 		{
 			// If there are any other prim paths linked to this asset that we *won't* be removing/reparsing
