USD Importer: when spawning actors during USD import, set the actor's transient flag after spawn to make sure they get an external package if needed.
Also don't filter out transient components for transient actors when computing streaming bounds.

diff --git a/Engine/Plugins/Importers/USDImporter/Source/USDSchemas/Private/USDGeomXformableTranslator.cpp b/Engine/Plugins/Importers/USDImporter/Source/USDSchemas/Private/USDGeomXformableTranslator.cpp
--- a/Engine/Plugins/Importers/USDImporter/Source/USDSchemas/Private/USDGeomXformableTranslator.cpp
+++ b/Engine/Plugins/Importers/USDImporter/Source/USDSchemas/Private/USDGeomXformableTranslator.cpp
@@ -416,14 +416,17 @@ USceneComponent* FUsdGeomXformableTranslator::CreateComponentsEx( TOptional< TSu
 
 	// Can't have public or standalone on spawned actors and components because that
 	// will lead to asserts when trying to collect them during a level change, or when
-	// trying to replace them (right-clicking from the world outliner)
-	EObjectFlags ComponentFlags = Context->ObjectFlags & ~RF_Standalone & ~RF_Public;
+	// trying to replace them (right-clicking from the world outliner). Also, must set
+	// the transient flag after spawn to make sure the spawned actor get an external
+	// package if needed.
+	const EObjectFlags PreComponentFlags = Context->ObjectFlags & ~(RF_Standalone | RF_Public | RF_Transient);
+	const EObjectFlags PostComponentFlags = Context->ObjectFlags & RF_Transient;
 
 	if ( bNeedsActor.GetValue() )
 	{
 		// Spawn actor
 		FActorSpawnParameters SpawnParameters;
-		SpawnParameters.ObjectFlags = ComponentFlags;
+		SpawnParameters.ObjectFlags = PreComponentFlags;
 		SpawnParameters.OverrideLevel =  Context->Level;
 		SpawnParameters.Name = Prim.GetName();
 		SpawnParameters.NameMode = FActorSpawnParameters::ESpawnActorNameMode::Requested; // Will generate a unique name in case of a conflict
@@ -433,6 +436,8 @@ USceneComponent* FUsdGeomXformableTranslator::CreateComponentsEx( TOptional< TSu
 
 		if ( SpawnedActor )
 		{
+			SpawnedActor->SetFlags(PostComponentFlags);
+
 #if WITH_EDITOR
 			const bool bMarkDirty = false;
 			SpawnedActor->SetActorLabel( Prim.GetName().ToString(), bMarkDirty );
@@ -511,7 +516,8 @@ USceneComponent* FUsdGeomXformableTranslator::CreateComponentsEx( TOptional< TSu
 				ComponentType.GetValue(),
 				*IUsdClassesModule::SanitizeObjectName(Prim.GetName().ToString())
 			);
-			SceneComponent = NewObject< USceneComponent >( ComponentOuter, ComponentType.GetValue(), ComponentName, ComponentFlags );
+			SceneComponent = NewObject< USceneComponent >( ComponentOuter, ComponentType.GetValue(), ComponentName, PreComponentFlags );
+			SceneComponent->SetFlags(PostComponentFlags);
 
 			if ( AActor* Owner = SceneComponent->GetOwner() )
 			{
diff --git a/Engine/Source/Runtime/Engine/Private/Actor.cpp b/Engine/Source/Runtime/Engine/Private/Actor.cpp
--- a/Engine/Source/Runtime/Engine/Private/Actor.cpp
+++ b/Engine/Source/Runtime/Engine/Private/Actor.cpp
@@ -1138,11 +1138,13 @@ static bool IsComponentStreamingRelevant(const AActor* InActor, const UActorComp
 		return false;
 	}
 
-	if (InComponent->HasAnyFlags(RF_Transient))
+	// Transient components shoudn't be part of the streeaming bounds, unless the actor itself is transient.
+	if (!InActor->HasAnyFlags(RF_Transient) && InComponent->HasAnyFlags(RF_Transient))
 	{
 		return false;
 	}
 
+	// Editor-only components shoudn't be part of the streeaming bounds, unless the actor itself is editor-only.
 	if (!InActor->IsEditorOnly() && InComponent->IsEditorOnly())
 	{
 		return false;
