USD: Remove some unnecessary allocator juggling causing some crashes.

diff --git a/Engine/Plugins/Importers/USDImporter/Source/USDSchemas/Private/USDGeomMeshTranslator.cpp b/Engine/Plugins/Importers/USDImporter/Source/USDSchemas/Private/USDGeomMeshTranslator.cpp
--- a/Engine/Plugins/Importers/USDImporter/Source/USDSchemas/Private/USDGeomMeshTranslator.cpp
+++ b/Engine/Plugins/Importers/USDImporter/Source/USDSchemas/Private/USDGeomMeshTranslator.cpp
@@ -288,25 +288,18 @@ namespace UsdGeomMeshTranslatorImpl
 	// We use a separate function for this because there is a very specific set of conditions where we successfully can do this, and we
 	// want to fall back to just parsing UsdMesh as a simple single-LOD mesh if we fail.
 	bool TryLoadingMultipleLODs(
-		const pxr::UsdTyped& UsdMesh,
+		const UE::FUsdPrim& MeshPrim,
 		TArray<FMeshDescription>& OutLODIndexToMeshDescription,
 		TArray<UsdUtils::FUsdPrimMaterialAssignmentInfo>& OutLODIndexToMaterialInfo,
 		const UsdToUnreal::FUsdMeshConversionOptions& Options
 	)
 	{
-		FScopedUsdAllocs Allocs;
-
-		pxr::UsdPrim UsdMeshPrim = UsdMesh.GetPrim();
-		if (!UsdMeshPrim)
+		if(!MeshPrim)
 		{
 			return false;
 		}
 
-		pxr::UsdPrim ParentPrim = UsdMeshPrim.GetParent();
-		if (!ParentPrim)
-		{
-			return false;
-		}
+		UE::FUsdPrim ParentPrim = MeshPrim.GetParent();
 
 		TMap<int32, FMeshDescription> LODIndexToMeshDescriptionMap;
 		TMap<int32, UsdUtils::FUsdPrimMaterialAssignmentInfo> LODIndexToMaterialInfoMap;
@@ -319,7 +312,7 @@ namespace UsdGeomMeshTranslatorImpl
 		// Here we are choosing to preemptively traverse all LODs to fetch the combined primvars that we'll use for
 		// each UV index for the static mesh as a whole
 		TFunction<bool(const pxr::UsdGeomMesh&, int32)> CombinePrimvars =
-			[&AllPrimvars, &PreferredPrimvars, &LODIndexToMaterialInfoMap]
+			[&AllPrimvars, &PreferredPrimvars]
 			(const pxr::UsdGeomMesh& LODMesh, int32 LODIndex)
 			{
 				TArray<TUsdStore<pxr::UsdGeomPrimvar>> MeshPrimvars = UsdUtils::GetUVSetPrimvars(
@@ -366,32 +359,36 @@ namespace UsdGeomMeshTranslatorImpl
 
 			bool bSuccess = true;
 
-			// The user can't manually hide/unhide a particular LOD in the engine after it is imported, so we should probably bake
-			// the particular LOD visibility into the combined mesh. Note how we don't use computed visibility here, as we only really
-			// care if this mesh in particular has been marked as invisible
-			pxr::TfToken Visibility;
-			pxr::UsdAttribute VisibilityAttr = LODMesh.GetVisibilityAttr();
-			if (VisibilityAttr && VisibilityAttr.Get(&Visibility, Options.TimeCode) && Visibility == pxr::UsdGeomTokens->inherited)
 			{
-				// If we're interpreting LODs we must bake the transform from each LOD Mesh into the vertices, because there's no guarantee
-				// all LODs have the same transform, so we can't just put the transforms directly on the component. If we are not interpreting
-				// LODs we can do that though
-				// TODO: Handle resetXformOp here
-				bool bResetXformStack = false;
-				FTransform MeshTransform = FTransform::Identity;
-				bSuccess &= UsdToUnreal::ConvertXformable(LODMesh.GetPrim().GetStage(), LODMesh, MeshTransform, Options.TimeCode.GetValue(), &bResetXformStack);
-
-				if (bSuccess)
+				FScopedUsdAllocs Allocs;
+				// The user can't manually hide/unhide a particular LOD in the engine after it is imported, so we should probably bake
+				// the particular LOD visibility into the combined mesh. Note how we don't use computed visibility here, as we only really
+				// care if this mesh in particular has been marked as invisible
+				pxr::TfToken Visibility;
+				pxr::UsdAttribute VisibilityAttr = LODMesh.GetVisibilityAttr();
+				if (VisibilityAttr && VisibilityAttr.Get(&Visibility, Options.TimeCode) && Visibility == pxr::UsdGeomTokens->inherited)
 				{
-					OptionsCopy.AdditionalTransform = MeshTransform * Options.AdditionalTransform;
-					OptionsCopy.bMergeIdenticalMaterialSlots = false;  // We only merge slots when collapsing, and we never collapse LODs
-
-					bSuccess &= UsdToUnreal::ConvertGeomMesh(
+					// If we're interpreting LODs we must bake the transform from each LOD Mesh into the vertices, because there's no guarantee
+					// all LODs have the same transform, so we can't just put the transforms directly on the component. If we are not interpreting
+					// LODs we can do that though
+					// TODO: Handle resetXformOp here
+					bool bResetXformStack = false;
+					FTransform MeshTransform = FTransform::Identity;
+					bSuccess &= UsdToUnreal::ConvertXformable(
+						LODMesh.GetPrim().GetStage(),
 						LODMesh,
-						TempMeshDescription,
-						TempMaterialInfo,
-						OptionsCopy
+						MeshTransform,
+						Options.TimeCode.GetValue(),
+						&bResetXformStack
 					);
+
+					if (bSuccess)
+					{
+						OptionsCopy.AdditionalTransform = MeshTransform * Options.AdditionalTransform;
+						OptionsCopy.bMergeIdenticalMaterialSlots = false;	 // We only merge slots when collapsing, and we never collapse LODs
+
+						bSuccess &= UsdToUnreal::ConvertGeomMesh(LODMesh, TempMeshDescription, TempMaterialInfo, OptionsCopy);
+					}
 				}
 			}
 
@@ -423,32 +420,29 @@ namespace UsdGeomMeshTranslatorImpl
 	}
 
 	void LoadMeshDescriptions(
-		pxr::UsdTyped UsdMesh,
+		UE::FUsdPrim MeshPrim,
 		TArray<FMeshDescription>& OutLODIndexToMeshDescription,
 		TArray<UsdUtils::FUsdPrimMaterialAssignmentInfo>& OutLODIndexToMaterialInfo,
 		const UsdToUnreal::FUsdMeshConversionOptions& Options,
 		bool bInterpretLODs = false
 	)
 	{
-		if (!UsdMesh)
+		if (!MeshPrim)
 		{
 			return;
 		}
 
-		FScopedUsdAllocs Allocs;
-
-		pxr::UsdPrim Prim = UsdMesh.GetPrim();
-		pxr::UsdStageRefPtr Stage = Prim.GetStage();
-		pxr::SdfPath Path = Prim.GetPrimPath();
+		UE::FUsdStage Stage = MeshPrim.GetStage();
+		UE::FSdfPath Path = MeshPrim.GetPrimPath();
 
 		bool bInterpretedLODs = false;
 		if (bInterpretLODs)
 		{
-			bInterpretedLODs = TryLoadingMultipleLODs(UsdMesh, OutLODIndexToMeshDescription, OutLODIndexToMaterialInfo, Options);
+			bInterpretedLODs = TryLoadingMultipleLODs(MeshPrim, OutLODIndexToMeshDescription, OutLODIndexToMaterialInfo, Options);
 
 			// Have to be very careful here as flipping through LODs invalidates prim references, so we need to
 			// re-acquire them
-			UsdMesh = pxr::UsdGeomMesh{Stage->GetPrimAtPath(Path)};
+			MeshPrim = Stage.GetPrimAtPath(Path);
 		}
 
 		// If we've managed to interpret LODs, we won't place our mesh transform on the static mesh component itself
@@ -460,7 +454,7 @@ namespace UsdGeomMeshTranslatorImpl
 		{
 			// TODO: Handle resetXformOp here
 			bool bResetXformStack = false;
-			bSuccess &= UsdToUnreal::ConvertXformable(Stage, UsdMesh, MeshTransform, Options.TimeCode.GetValue(), &bResetXformStack);
+			bSuccess &= UsdToUnreal::ConvertXformable(Stage, UE::FUsdTyped{MeshPrim}, MeshTransform, Options.TimeCode.GetValue(), &bResetXformStack);
 		}
 
 		if (!bInterpretedLODs)
@@ -477,8 +471,11 @@ namespace UsdGeomMeshTranslatorImpl
 				OptionsCopy.AdditionalTransform = MeshTransform * Options.AdditionalTransform;
 				OptionsCopy.bMergeIdenticalMaterialSlots = false;  // We only merge for collapsed meshes
 
+				FScopedUsdAllocs Allocs;
+				pxr::UsdGeomMesh UsdMesh{MeshPrim};
+
 				bSuccess &= UsdToUnreal::ConvertGeomMesh(
-					pxr::UsdGeomMesh{UsdMesh},
+					UsdMesh,
 					TempMeshDescription,
 					TempMaterialInfo,
 					OptionsCopy
@@ -966,7 +963,7 @@ void FGeomMeshCreateAssetsTaskChain::SetupTasks()
 			Options.AdditionalTransform = AdditionalTransform;
 
 			UsdGeomMeshTranslatorImpl::LoadMeshDescriptions(
-				pxr::UsdTyped(GetPrim()),
+				GetPrim(),
 				LODIndexToMeshDescription,
 				LODIndexToMaterialInfo,
 				Options,
diff --git a/Engine/Plugins/Importers/USDImporter/Source/USDSchemas/Private/USDGeometryCacheTranslator.cpp b/Engine/Plugins/Importers/USDImporter/Source/USDSchemas/Private/USDGeometryCacheTranslator.cpp
--- a/Engine/Plugins/Importers/USDImporter/Source/USDSchemas/Private/USDGeometryCacheTranslator.cpp
+++ b/Engine/Plugins/Importers/USDImporter/Source/USDSchemas/Private/USDGeometryCacheTranslator.cpp
@@ -107,8 +107,6 @@ namespace UsdGeometryCacheTranslatorImpl
 			return;
 		}
 
-		FScopedUsdAllocs Allocs;
-
 		FMeshDescription TempMeshDescription;
 		UsdUtils::FUsdPrimMaterialAssignmentInfo TempMaterialInfo;
 
diff --git a/Engine/Plugins/Importers/USDImporter/Source/USDUtilities/Private/USDGeomMeshConversion.cpp b/Engine/Plugins/Importers/USDImporter/Source/USDUtilities/Private/USDGeomMeshConversion.cpp
--- a/Engine/Plugins/Importers/USDImporter/Source/USDUtilities/Private/USDGeomMeshConversion.cpp
+++ b/Engine/Plugins/Importers/USDImporter/Source/USDUtilities/Private/USDGeomMeshConversion.cpp
@@ -1318,32 +1318,45 @@ bool UsdToUnreal::ConvertPointInstancerToMesh(
 		return false;
 	}
 
-	FScopedUsdAllocs Allocs;
-
-	pxr::UsdStageRefPtr Stage = PointInstancer.GetPrim().GetStage();
-	if ( !Stage )
-	{
-		return false;
-	}
-
 	// Bake each prototype to a single mesh description and material assignment struct
 	TArray<FMeshDescription> PrototypeMeshDescriptions;
 	TArray<UsdUtils::FUsdPrimMaterialAssignmentInfo> PrototypeMaterialAssignments;
 	TArray<TMap<FPolygonGroupID, FPolygonGroupID>> PrototypePolygonGroupRemapping;
 	uint32 NumPrototypes = 0;
+	UE::FUsdStage Stage;
 	{
-		const pxr::UsdRelationship& Prototypes = PointInstancer.GetPrototypesRel();
-
-		pxr::SdfPathVector PrototypePaths;
-		if ( !Prototypes.GetTargets( &PrototypePaths ) )
+		TArray<UE::FSdfPath> PrototypePaths;
 		{
-			return false;
-		}
+			Stage = UE::FUsdStage{PointInstancer.GetPrim().GetStage()};
+			if (!Stage)
+			{
+				return false;
+			}
 
-		NumPrototypes = PrototypePaths.size();
-		if ( NumPrototypes == 0 )
-		{
-			return true;
+			TOptional<FScopedUsdAllocs> Allocs;
+			Allocs.Emplace();
+
+			const pxr::UsdRelationship& Prototypes = PointInstancer.GetPrototypesRel();
+
+			pxr::SdfPathVector UsdPrototypePaths;
+			if ( !Prototypes.GetTargets( &UsdPrototypePaths ) )
+			{
+				return false;
+			}
+
+			NumPrototypes = UsdPrototypePaths.size();
+			if (NumPrototypes == 0)
+			{
+				return true;
+			}
+
+			Allocs.Reset();
+			PrototypePaths.Reserve(NumPrototypes);
+			for(const pxr::SdfPath& UsdPath : UsdPrototypePaths)
+			{
+				PrototypePaths.Add(UE::FSdfPath{UsdPath});
+			}
+			Allocs.Emplace();
 		}
 
 		PrototypeMeshDescriptions.SetNum( NumPrototypes );
@@ -1357,9 +1370,9 @@ bool UsdToUnreal::ConvertPointInstancerToMesh(
 
 		for ( uint32 PrototypeIndex = 0; PrototypeIndex < NumPrototypes; ++PrototypeIndex )
 		{
-			const pxr::SdfPath& PrototypePath = PrototypePaths[ PrototypeIndex ];
+			const UE::FSdfPath& PrototypePath = PrototypePaths[ PrototypeIndex ];
 
-			pxr::UsdPrim PrototypeUsdPrim = Stage->GetPrimAtPath( PrototypePath );
+			UE::FUsdPrim PrototypeUsdPrim = Stage.GetPrimAtPath( PrototypePath );
 			if ( !PrototypeUsdPrim )
 			{
 				UE_LOG( LogUsd, Warning, TEXT( "Failed to find prototype '%s' for PointInstancer '%s' within ConvertPointInstancerToMesh" ),
@@ -1442,7 +1455,6 @@ bool UsdToUnreal::ConvertPointInstancerToMesh(
 	for ( uint32 PrototypeIndex = 0; PrototypeIndex < NumPrototypes; ++PrototypeIndex )
 	{
 		const FMeshDescription& PrototypeMeshDescription = PrototypeMeshDescriptions[ PrototypeIndex ];
-		UsdUtils::FUsdPrimMaterialAssignmentInfo& PrototypeMaterialAssignment = PrototypeMaterialAssignments[ PrototypeIndex ];
 
 		// We may generate some empty meshes in case a prototype is invisible, for example
 		if ( PrototypeMeshDescription.IsEmpty() )
@@ -2497,49 +2509,50 @@ int32 UsdUtils::GetNumberOfLODVariants( const pxr::UsdPrim& Prim )
 
 bool UsdUtils::IterateLODMeshes( const pxr::UsdPrim& ParentPrim, TFunction<bool( const pxr::UsdGeomMesh & LODMesh, int32 LODIndex )> Func )
 {
-	if ( !ParentPrim )
+	if (!ParentPrim)
 	{
 		return false;
 	}
 
-	FScopedUsdAllocs Allocs;
+	TOptional<FScopedUsdAllocs> Allocs;
+	Allocs.Emplace();
 
 	const std::string LODString = UnrealIdentifiers::LOD.GetString();
 
 	pxr::UsdVariantSets VariantSets = ParentPrim.GetVariantSets();
-	if ( !VariantSets.HasVariantSet( LODString ) )
+	if (!VariantSets.HasVariantSet(LODString))
 	{
 		return false;
 	}
 
-	pxr::UsdVariantSet LODVariantSet = VariantSets.GetVariantSet( LODString );
+	pxr::UsdVariantSet LODVariantSet = VariantSets.GetVariantSet(LODString);
 	const std::string OriginalVariant = LODVariantSet.GetVariantSelection();
 
 	pxr::UsdStageRefPtr Stage = ParentPrim.GetStage();
-	pxr::UsdEditContext( Stage, Stage->GetRootLayer() );
+	pxr::UsdEditContext(Stage, Stage->GetRootLayer());
 
 	bool bHasValidVariant = false;
-	for ( const std::string& LODVariantName : VariantSets.GetVariantSet( LODString ).GetVariantNames() )
+	for (const std::string& LODVariantName : VariantSets.GetVariantSet(LODString).GetVariantNames())
 	{
-		int32 LODIndex = UsdGeomMeshImpl::GetLODIndexFromName( LODVariantName );
-		if ( LODIndex == INDEX_NONE )
+		int32 LODIndex = UsdGeomMeshImpl::GetLODIndexFromName(LODVariantName);
+		if (LODIndex == INDEX_NONE)
 		{
 			continue;
 		}
 
-		LODVariantSet.SetVariantSelection( LODVariantName );
+		LODVariantSet.SetVariantSelection(LODVariantName);
 
 		pxr::UsdGeomMesh LODMesh;
-		pxr::TfToken TargetChildNameToken{ LODVariantName };
+		pxr::TfToken TargetChildNameToken{LODVariantName};
 
 		// Search for our LOD child mesh
 		pxr::UsdPrimSiblingRange PrimRange = ParentPrim.GetChildren();
-		for ( pxr::UsdPrimSiblingRange::iterator PrimRangeIt = PrimRange.begin(); PrimRangeIt != PrimRange.end(); ++PrimRangeIt )
+		for (pxr::UsdPrimSiblingRange::iterator PrimRangeIt = PrimRange.begin(); PrimRangeIt != PrimRange.end(); ++PrimRangeIt)
 		{
 			const pxr::UsdPrim& Child = *PrimRangeIt;
-			if ( pxr::UsdGeomMesh ChildMesh{ Child } )
+			if (pxr::UsdGeomMesh ChildMesh{Child})
 			{
-				if ( Child.GetName() == TargetChildNameToken )
+				if (Child.GetName() == TargetChildNameToken)
 				{
 					LODMesh = ChildMesh;
 					// Don't break here so we can show warnings if the user has other prims here (that we may end up ignoring)
@@ -2547,28 +2560,35 @@ bool UsdUtils::IterateLODMeshes( const pxr::UsdPrim& ParentPrim, TFunction<bool(
 				}
 				else
 				{
-					UE_LOG(LogUsd, Warning, TEXT("Unexpected prim '%s' inside LOD variant '%s'. For automatic parsing of LODs, each LOD variant should contain only a single Mesh prim named the same as the variant!"),
-						*UsdToUnreal::ConvertPath( Child.GetPath() ),
-						*UsdToUnreal::ConvertString( LODVariantName )
+					UE_LOG(
+						LogUsd,
+						Warning,
+						TEXT("Unexpected prim '%s' inside LOD variant '%s'. For automatic parsing of LODs, each LOD variant should contain only a "
+							 "single Mesh prim named the same as the variant!"),
+						*UsdToUnreal::ConvertPath(Child.GetPath()),
+						*UsdToUnreal::ConvertString(LODVariantName)
 					);
 				}
 			}
 		}
-		if ( !LODMesh )
+		if (!LODMesh)
 		{
 			continue;
 		}
 
 		bHasValidVariant = true;
 
+		// Reset our forced allocator as we don't know what Func expects
+		Allocs.Reset();
 		bool bContinue = Func(LODMesh, LODIndex);
-		if ( !bContinue )
+		Allocs.Emplace();
+		if (!bContinue)
 		{
 			break;
 		}
 	}
 
-	LODVariantSet.SetVariantSelection( OriginalVariant );
+	LODVariantSet.SetVariantSelection(OriginalVariant);
 	return bHasValidVariant;
 }
 
