USDImporter: fix allocator management scoping in UsdUtils::GetPointInstancerTransforms()

diff --git a/Engine/Plugins/Importers/USDImporter/Source/USDUtilities/Private/USDGeomMeshConversion.cpp b/Engine/Plugins/Importers/USDImporter/Source/USDUtilities/Private/USDGeomMeshConversion.cpp
--- a/Engine/Plugins/Importers/USDImporter/Source/USDUtilities/Private/USDGeomMeshConversion.cpp
+++ b/Engine/Plugins/Importers/USDImporter/Source/USDUtilities/Private/USDGeomMeshConversion.cpp
@@ -3027,50 +3027,52 @@ FString UsdUtils::HashGeomMeshPrim( const UE::FUsdStage& Stage, const FString& P
 	return Hash;
 }
 
-bool UsdUtils::GetPointInstancerTransforms( const FUsdStageInfo& StageInfo, const pxr::UsdGeomPointInstancer& PointInstancer, const int32 ProtoIndex, pxr::UsdTimeCode EvalTime, TArray<FTransform>& OutInstanceTransforms )
+bool UsdUtils::GetPointInstancerTransforms(const FUsdStageInfo& StageInfo, const pxr::UsdGeomPointInstancer& PointInstancer, const int32 ProtoIndex, pxr::UsdTimeCode EvalTime, TArray<FTransform>& OutInstanceTransforms)
 {
-	TRACE_CPUPROFILER_EVENT_SCOPE( GetPointInstancerTransforms );
+	TRACE_CPUPROFILER_EVENT_SCOPE(GetPointInstancerTransforms);
 
-	if ( !PointInstancer )
+	if (!PointInstancer)
 	{
 		return false;
 	}
 
 	FScopedUsdAllocs UsdAllocs;
 
-	pxr::VtArray< int > ProtoIndices = UsdUtils::GetUsdValue< pxr::VtArray< int > >( PointInstancer.GetProtoIndicesAttr(), EvalTime );
+	const pxr::VtArray<int> UsdProtoIndices = UsdUtils::GetUsdValue<pxr::VtArray<int>>(PointInstancer.GetProtoIndicesAttr(), EvalTime);
 
 	pxr::VtMatrix4dArray UsdInstanceTransforms;
 
 	// We don't want the prototype root prim's transforms to be included in these, as they'll already be baked into the meshes themselves
-	if ( !PointInstancer.ComputeInstanceTransformsAtTime( &UsdInstanceTransforms, EvalTime, EvalTime, pxr::UsdGeomPointInstancer::ExcludeProtoXform ) )
+	if (!PointInstancer.ComputeInstanceTransformsAtTime(&UsdInstanceTransforms, EvalTime, EvalTime, pxr::UsdGeomPointInstancer::ExcludeProtoXform))
 	{
 		return false;
 	}
 
 	int32 Index = 0;
 
-	FScopedUnrealAllocs UnrealAllocs;
-
 	const int32 NumInstances = GMaxInstancesPerPointInstancer >= 0
-		? FMath::Min( static_cast< int32 >( UsdInstanceTransforms.size() ), GMaxInstancesPerPointInstancer )
-		: static_cast< int32 >( UsdInstanceTransforms.size() );
+		? FMath::Min(static_cast<int32>(UsdInstanceTransforms.size()), GMaxInstancesPerPointInstancer)
+		: static_cast<int32>(UsdInstanceTransforms.size());
 
-	OutInstanceTransforms.Reset( NumInstances );
-
-	for ( pxr::GfMatrix4d& UsdMatrix : UsdInstanceTransforms )
 	{
-		if ( Index == NumInstances )
-		{
-			break;
-		}
+		FScopedUnrealAllocs UnrealAllocs;
 
-		if ( ProtoIndices[ Index ] == ProtoIndex )
+		OutInstanceTransforms.Reset(NumInstances);
+
+		for (const pxr::GfMatrix4d& UsdMatrix : UsdInstanceTransforms)
 		{
-			OutInstanceTransforms.Add( UsdToUnreal::ConvertMatrix( StageInfo, UsdMatrix ) );
-		}
+			if (Index == NumInstances)
+			{
+				break;
+			}
 
-		++Index;
+			if (UsdProtoIndices[Index] == ProtoIndex)
+			{
+				OutInstanceTransforms.Add(UsdToUnreal::ConvertMatrix(StageInfo, UsdMatrix));
+			}
+
+			++Index;
+		}
 	}
 
 	return true;
