USD: Fix ensure when closing stages with the transient LevelSequence opened.

diff --git a/Engine/Plugins/Importers/USDImporter/Source/USDStage/Private/USDStageActor.cpp b/Engine/Plugins/Importers/USDImporter/Source/USDStage/Private/USDStageActor.cpp
--- a/Engine/Plugins/Importers/USDImporter/Source/USDStage/Private/USDStageActor.cpp
+++ b/Engine/Plugins/Importers/USDImporter/Source/USDStage/Private/USDStageActor.cpp
@@ -2660,15 +2660,9 @@ void AUsdStageActor::UnloadUsdStage()
 	// and could lead to stage changes
 	BlockMonitoringLevelSequenceForThisTransaction();
 
-	if (InfoCache)
-	{
-		InfoCache->Clear();
-	}
-
-	ObjectsToWatch.Reset();
-	BlendShapesByPath.Reset();
-	MaterialToPrimvarToUVIndex.Reset();
-
+	// Close the Sequencer before dropping the info cache, as the Sequencer closing may
+	// trigger one last SetTime call on the stage actor (to revert to the preanimated state),
+	// and if we try animating things and calling UpdateComponents we may need the info cache
 	if (LevelSequence)
 	{
 #if WITH_EDITOR
@@ -2681,6 +2675,10 @@ void AUsdStageActor::UnloadUsdStage()
 	}
 	LevelSequenceHelper.Clear();
 
+	ObjectsToWatch.Reset();
+	BlendShapesByPath.Reset();
+	MaterialToPrimvarToUVIndex.Reset();
+
 	if (RootUsdTwin)
 	{
 		RootUsdTwin->Clear();
@@ -2702,6 +2700,11 @@ void AUsdStageActor::UnloadUsdStage()
 		UsdAssetCache->RefreshStorage();
 	}
 
+	if (InfoCache)
+	{
+		InfoCache->Clear();
+	}
+
 	OnStageUnloaded.Broadcast();
 
 	CloseUsdStage();
