USD: Fix missing components when Light prims have child Mesh prims.

diff --git a/Engine/Plugins/Importers/USDImporter/Source/USDSchemas/Private/USDGeomXformableTranslator.cpp b/Engine/Plugins/Importers/USDImporter/Source/USDSchemas/Private/USDGeomXformableTranslator.cpp
--- a/Engine/Plugins/Importers/USDImporter/Source/USDSchemas/Private/USDGeomXformableTranslator.cpp
+++ b/Engine/Plugins/Importers/USDImporter/Source/USDSchemas/Private/USDGeomXformableTranslator.cpp
@@ -550,19 +550,20 @@ USceneComponent* FUsdGeomXformableTranslator::CreateComponentsEx( TOptional< TSu
 		// so we must force them movable here
 		const bool bIsAssociating = Context->Level && Context->Level->bIsAssociatingLevel;
 		const bool bParentIsMovable = Context->ParentComponent && Context->ParentComponent->Mobility == EComponentMobility::Movable;
+		const bool bParentIsStationary = Context->ParentComponent && Context->ParentComponent->Mobility == EComponentMobility::Stationary;
 
 		// Don't call SetMobility as it would trigger a reregister, queuing unnecessary rhi commands since this is a brand new component
-		if ( bIsAssociating || bParentIsMovable )
+		if (bIsAssociating || bParentIsMovable || UsdUtils::IsAnimated(Prim))
 		{
 			SceneComponent->Mobility = EComponentMobility::Movable;
 		}
+		else if (bParentIsStationary || SceneComponent->IsA<ULightComponentBase>())
+		{
+			SceneComponent->Mobility = EComponentMobility::Stationary;
+		}
 		else
 		{
-			SceneComponent->Mobility = UsdUtils::IsAnimated( Prim )
-				? EComponentMobility::Movable
-				: SceneComponent->IsA<ULightComponentBase>()	// Lights need to be stationary by default
-					? EComponentMobility::Stationary
-					: EComponentMobility::Static;
+			SceneComponent->Mobility = EComponentMobility::Static;
 		}
 
 		// Attach to parent
