USD: Fix missing meshes on HISM generated from point instancers with LOD prototypes.

diff --git a/Engine/Plugins/Importers/USDImporter/Source/USDSchemas/Private/USDGeomPointInstancerTranslator.cpp b/Engine/Plugins/Importers/USDImporter/Source/USDSchemas/Private/USDGeomPointInstancerTranslator.cpp
--- a/Engine/Plugins/Importers/USDImporter/Source/USDSchemas/Private/USDGeomPointInstancerTranslator.cpp
+++ b/Engine/Plugins/Importers/USDImporter/Source/USDSchemas/Private/USDGeomPointInstancerTranslator.cpp
@@ -421,13 +421,18 @@ void FUsdGeomPointInstancerTranslator::UpdateComponents(USceneComponent* PointIn
 	{
 		PrototypePathsSlowTask.EnterProgressFrame();
 
-		const pxr::SdfPath& PrototypePath = PrototypePaths[ PrototypeIndex ];
-		FString PrototypePathStr = UsdToUnreal::ConvertPath( PrototypePath );
+		pxr::SdfPath PrototypePath = PrototypePaths[PrototypeIndex];
 
 		pxr::UsdPrim PrototypeUsdPrim = Prim.GetStage()->GetPrimAtPath( PrototypePath );
 		if ( !PrototypeUsdPrim )
 		{
-			UE_LOG( LogUsd, Warning, TEXT( "Failed to find prototype '%s' for PointInstancer '%s' when updating components" ), *PrototypePathStr, *PrimPath.GetString() );
+			UE_LOG(
+				LogUsd,
+				Warning,
+				TEXT("Failed to find prototype '%s' for PointInstancer '%s' when updating components"),
+				*UsdToUnreal::ConvertPath(PrototypePath),
+				*PrimPath.GetString()
+			);
 			continue;
 		}
 
@@ -436,7 +441,13 @@ void FUsdGeomPointInstancerTranslator::UpdateComponents(USceneComponent* PointIn
 		// The user could have just manually deleted the component, so we must check
 		if (!AttachedHISMs.IsValidIndex(PrototypeIndex))
 		{
-			UE_LOG(LogUsd, Warning, TEXT("Failed to find corresponding HISM component for prototype '%s' of PointInstancer '%s'. Cancelling component update"), *PrototypePathStr, *PrimPath.GetString());
+			UE_LOG(
+				LogUsd,
+				Warning,
+				TEXT("Failed to find corresponding HISM component for prototype '%s' of PointInstancer '%s'. Cancelling component update"),
+				*UsdToUnreal::ConvertPath(PrototypePath),
+				*PrimPath.GetString()
+			);
 			break;
 		}
 		HISMComponent = AttachedHISMs[PrototypeIndex];
@@ -452,7 +463,7 @@ void FUsdGeomPointInstancerTranslator::UpdateComponents(USceneComponent* PointIn
 				if ( pxr::UsdGeomMesh ChildMesh{ Child } )
 				{
 					PrototypeUsdPrim = Child;
-					PrototypePathStr = UsdToUnreal::ConvertPath( Child.GetPrimPath() );
+					PrototypePath = Child.GetPrimPath();
 					break;
 				}
 			}
