USD: Rename dynamic binding payload variable for clarity.

diff --git a/Engine/Plugins/Importers/USDImporter/Source/USDStage/Private/USDDynamicBindingResolverLibrary.cpp b/Engine/Plugins/Importers/USDImporter/Source/USDStage/Private/USDDynamicBindingResolverLibrary.cpp
--- a/Engine/Plugins/Importers/USDImporter/Source/USDStage/Private/USDDynamicBindingResolverLibrary.cpp
+++ b/Engine/Plugins/Importers/USDImporter/Source/USDStage/Private/USDDynamicBindingResolverLibrary.cpp
@@ -12,7 +12,7 @@
 FMovieSceneDynamicBindingResolveResult UUsdDynamicBindingResolverLibrary::ResolveWithStageActor(
 	UObject* WorldContextObject,
 	const FMovieSceneDynamicBindingResolveParams& Params,
-	const FString& StageActorNameFilter,
+	const FString& StageActorIDNameFilter,
 	const FString& RootLayerFilter,
 	const FString& PrimPath
 )
@@ -38,7 +38,7 @@ FMovieSceneDynamicBindingResolveResult UUsdDynamicBindingResolverLibrary::Resolv
 			}
 		}
 
-		const bool bNoNameFilter = StageActorNameFilter.IsEmpty();
+		const bool bNoNameFilter = StageActorIDNameFilter.IsEmpty();
 		const bool bNoRootLayerFilter = RootLayerFilter.IsEmpty();
 
 		UWorld* World = WorldContextObject->GetWorld();
@@ -46,7 +46,7 @@ FMovieSceneDynamicBindingResolveResult UUsdDynamicBindingResolverLibrary::Resolv
 		{
 			if (AUsdStageActor* Actor = *ActorIt)
 			{
-				const bool bNameOK = bNoNameFilter || Actor->GetName() == StageActorNameFilter;
+				const bool bNameOK = bNoNameFilter || Actor->GetName() == StageActorIDNameFilter;
 				const bool bRootLayerOK = bNoRootLayerFilter || Actor->RootLayer.FilePath == RootLayerFilter;
 
 				if (bNameOK && bRootLayerOK)
diff --git a/Engine/Plugins/Importers/USDImporter/Source/USDStage/Private/USDStageActor.cpp b/Engine/Plugins/Importers/USDImporter/Source/USDStage/Private/USDStageActor.cpp
--- a/Engine/Plugins/Importers/USDImporter/Source/USDStage/Private/USDStageActor.cpp
+++ b/Engine/Plugins/Importers/USDImporter/Source/USDStage/Private/USDStageActor.cpp
@@ -640,7 +640,7 @@ struct FUsdStageActorImpl
 			// Store a path to this very actor on the binding, so that it can find us later and ask how to resolve a particular
 			// prim path
 			FMovieSceneDynamicBindingPayloadVariable& ActorPathVariable = Possessable.DynamicBinding.PayloadVariables.FindOrAdd(
-				TEXT("StageActorNameFilter")
+				TEXT("StageActorIDNameFilter")
 			);
 			ActorPathVariable.Value = DefaultActorFilter;
 
diff --git a/Engine/Plugins/Importers/USDImporter/Source/USDStage/Public/USDDynamicBindingResolverLibrary.h b/Engine/Plugins/Importers/USDImporter/Source/USDStage/Public/USDDynamicBindingResolverLibrary.h
--- a/Engine/Plugins/Importers/USDImporter/Source/USDStage/Public/USDDynamicBindingResolverLibrary.h
+++ b/Engine/Plugins/Importers/USDImporter/Source/USDStage/Public/USDDynamicBindingResolverLibrary.h
@@ -22,7 +22,7 @@ public:
 	 * should bind to
 	 * @param WorldContextObject - Some UObject that lives in the UWorld that we're talking about
 	 * @param Params - The binding to resolve
-	 * @param StageActorNameFilter - The "ID Name"/FName of a AUsdStageActor to restrict our search for actor and components to. Can be
+	 * @param StageActorIDNameFilter - The "ID Name"/FName of a AUsdStageActor to restrict our search for actor and components to. Can be
 	 * empty.
 	 * @param RootLayerFilter - The root layer file path to restrict our search for actor and components to. This should match what is shown
 	 * on the Stage Actor's RootLayer property. Can be empty.
@@ -34,7 +34,7 @@ public:
 	static FMovieSceneDynamicBindingResolveResult ResolveWithStageActor(
 		UObject* WorldContextObject,
 		const FMovieSceneDynamicBindingResolveParams& Params,
-		const FString& StageActorNameFilter,
+		const FString& StageActorIDNameFilter,
 		const FString& RootLayerFilter,
 		const FString& PrimPath
 	);
