USD: Prevent writing component changes to the stage while they're bound to an opened Sequencer track;

Fix changes to camera component properties not being tracked.

diff --git a/Engine/Plugins/Importers/USDImporter/Source/USDClasses/Public/USDProjectSettings.h b/Engine/Plugins/Importers/USDImporter/Source/USDClasses/Public/USDProjectSettings.h
--- a/Engine/Plugins/Importers/USDImporter/Source/USDClasses/Public/USDProjectSettings.h
+++ b/Engine/Plugins/Importers/USDImporter/Source/USDClasses/Public/USDProjectSettings.h
@@ -80,6 +80,10 @@ public:
 	UPROPERTY( config, EditAnywhere, Category = "USD|Dialogs" )
 	bool bShowWarningOnIncompleteDuplication = true;
 
+	// Whether to show the warning dialog when authoring a transforms directly to a camera component
+	UPROPERTY( config, EditAnywhere, Category = "USD|Dialogs" )
+	bool bShowTransformOnCameraComponentWarning = true;
+
 	// Whether to show the warning dialog when authoring a transform track directly to a camera component
 	UPROPERTY( config, EditAnywhere, Category = "USD|Dialogs" )
 	bool bShowTransformTrackOnCameraComponentWarning = true;
diff --git a/Engine/Plugins/Importers/USDImporter/Source/USDStage/Private/USDStageActor.cpp b/Engine/Plugins/Importers/USDImporter/Source/USDStage/Private/USDStageActor.cpp
--- a/Engine/Plugins/Importers/USDImporter/Source/USDStage/Private/USDStageActor.cpp
+++ b/Engine/Plugins/Importers/USDImporter/Source/USDStage/Private/USDStageActor.cpp
@@ -74,6 +74,7 @@
 #include "PhysicsEngine/BodySetup.h"
 #include "Rendering/SkeletalMeshLODImporterData.h"
 #include "Roles/LiveLinkTransformRole.h"
+#include "Sections/MovieSceneSubSection.h"
 #include "StaticMeshAttributes.h"
 #include "StaticMeshOperations.h"
 #include "Tracks/MovieScene3DTransformTrack.h"
@@ -763,6 +764,126 @@ struct FUsdStageActorImpl
 		FKismetEditorUtilities::CompileBlueprint(DirectorBlueprint);
 #endif	  // WITH_EDITOR
 	}
+
+	static void GetDescendantMovieSceneSequences(UMovieSceneSequence* InSequence, TSet<UMovieSceneSequence*>& OutAllSequences)
+	{
+		if (InSequence == nullptr || OutAllSequences.Contains(InSequence))
+		{
+			return;
+		}
+
+		OutAllSequences.Add(InSequence);
+
+		UMovieScene* MovieScene = InSequence->GetMovieScene();
+		if (!MovieScene)
+		{
+			return;
+		}
+
+		for (UMovieSceneSection* Section : MovieScene->GetAllSections())
+		{
+			UMovieSceneSubSection* SubSection = Cast<UMovieSceneSubSection>(Section);
+			if (SubSection != nullptr)
+			{
+				UMovieSceneSequence* SubSequence = SubSection->GetSequence();
+				if (SubSequence != nullptr)
+				{
+					GetDescendantMovieSceneSequences(SubSequence, OutAllSequences);
+				}
+			}
+		}
+	}
+
+	static void ShowTransformOnCameraComponentWarning(const UActorComponent* Component)
+	{
+		const UCineCameraComponent* CameraComponent = Cast<const UCineCameraComponent>(Component);
+		if (!CameraComponent)
+		{
+			return;
+		}
+		const AActor* OwnerActor = CameraComponent->GetOwner();
+		if (!OwnerActor)
+		{
+			return;
+		}
+
+		FObjectKey NewComponentKey{Component};
+		static TSet<FObjectKey> WarnedComponents;
+		if (WarnedComponents.Contains(NewComponentKey))
+		{
+			return;
+		}
+		WarnedComponents.Add(NewComponentKey);
+
+		const FText Text = LOCTEXT("TransformOnCameraComponentText", "USD: Transform on camera component");
+
+		const FText SubText = FText::Format(
+			LOCTEXT(
+				"TransformOnCameraComponentSubText",
+				"The transform of camera component '{0}' was modified, but the new value will not be written out to the USD stage.\n\nIn order to "
+				"write to the Camera prim "
+				"transform, please modify the transform of the Cine Camera Actor (or its root Scene Component) instead."
+			),
+			FText::FromString(Component->GetName())
+		);
+
+		UE_LOG(LogUsd, Warning, TEXT("%s"), *SubText.ToString().Replace(TEXT("\n\n"), TEXT(" ")));
+
+		const UUsdProjectSettings* Settings = GetDefault<UUsdProjectSettings>();
+		if (Settings && Settings->bShowTransformOnCameraComponentWarning)
+		{
+			static TWeakPtr<SNotificationItem> Notification;
+
+			FNotificationInfo Toast(Text);
+			Toast.SubText = SubText;
+			Toast.Image = FCoreStyle::Get().GetBrush(TEXT("MessageLog.Warning"));
+			Toast.CheckBoxText = LOCTEXT("DontAskAgain", "Don't prompt again");
+			Toast.bUseLargeFont = false;
+			Toast.bFireAndForget = false;
+			Toast.FadeOutDuration = 0.0f;
+			Toast.ExpireDuration = 0.0f;
+			Toast.bUseThrobber = false;
+			Toast.bUseSuccessFailIcons = false;
+			Toast.ButtonDetails.Emplace(
+				LOCTEXT("OverridenOpinionMessageOk", "Ok"),
+				FText::GetEmpty(),
+				FSimpleDelegate::CreateLambda(
+					[]()
+					{
+						if (TSharedPtr<SNotificationItem> PinnedNotification = Notification.Pin())
+						{
+							PinnedNotification->SetCompletionState(SNotificationItem::CS_Success);
+							PinnedNotification->ExpireAndFadeout();
+						}
+					}
+				)
+			);
+			// This is flipped because the default checkbox message is "Don't prompt again"
+			Toast.CheckBoxState = Settings->bShowTransformOnCameraComponentWarning ? ECheckBoxState::Unchecked : ECheckBoxState::Checked;
+			Toast.CheckBoxStateChanged = FOnCheckStateChanged::CreateStatic(
+				[](ECheckBoxState NewState)
+				{
+					if (UUsdProjectSettings* Settings = GetMutableDefault<UUsdProjectSettings>())
+					{
+						// This is flipped because the default checkbox message is "Don't prompt again"
+						Settings->bShowTransformOnCameraComponentWarning = NewState == ECheckBoxState::Unchecked;
+						Settings->SaveConfig();
+					}
+				}
+			);
+
+			// Only show one at a time
+			if (!Notification.IsValid())
+			{
+				Notification = FSlateNotificationManager::Get().AddNotification(Toast);
+			}
+
+			if (TSharedPtr<SNotificationItem> PinnedNotification = Notification.Pin())
+			{
+				PinnedNotification->SetCompletionState(SNotificationItem::CS_Pending);
+			}
+		}
+	}
 };
 
 /**
@@ -1790,6 +1911,13 @@ UUsdPrimTwin* AUsdStageActor::ExpandPrim(const UE::FUsdPrim& Prim, bool bResync,
 		else
 		{
 			ObjectsToWatch.Remove(UsdPrimTwin->SceneComponent.Get());
+			if (Prim.IsA(TEXT("Camera")))
+			{
+				if (ACineCameraActor* CameraActor = Cast<ACineCameraActor>(SceneComponent->GetOwner()))
+				{
+					ObjectsToWatch.Remove(CameraActor->GetCineCameraComponent());
+				}
+			}
 			SchemaTranslator->UpdateComponents(UsdPrimTwin->SceneComponent.Get());
 		}
 
@@ -1816,18 +1944,27 @@ UUsdPrimTwin* AUsdStageActor::ExpandPrim(const UE::FUsdPrim& Prim, bool bResync,
 		}
 	}
 
-	if (UsdPrimTwin->SceneComponent.IsValid())
+	USceneComponent* TwinSceneComponent = UsdPrimTwin->SceneComponent.Get();
+	if (TwinSceneComponent)
 	{
 #if WITH_EDITOR
-		UsdPrimTwin->SceneComponent->PostEditChange();
+		TwinSceneComponent->PostEditChange();
 #endif // WITH_EDITOR
 
-		if (!UsdPrimTwin->SceneComponent->IsRegistered())
+		if (!TwinSceneComponent->IsRegistered())
 		{
-			UsdPrimTwin->SceneComponent->RegisterComponent();
+			TwinSceneComponent->RegisterComponent();
 		}
 
-		ObjectsToWatch.Add(UsdPrimTwin->SceneComponent.Get(), UsdPrimTwin->PrimPath);
+		ObjectsToWatch.Add(TwinSceneComponent, UsdPrimTwin->PrimPath);
+		// Make sure we monitor direct changes to camera properties on the component as well as the actor
+		if (Prim.IsA(TEXT("Camera")))
+		{
+			if (ACineCameraActor* CameraActor = Cast<ACineCameraActor>(TwinSceneComponent->GetOwner()))
+			{
+				ObjectsToWatch.Add(CameraActor->GetCineCameraComponent(), UsdPrimTwin->PrimPath);
+			}
+		}
 	}
 
 	// Update the prim animated status
@@ -3661,30 +3798,23 @@ void AUsdStageActor::OnObjectPropertyChanged(UObject* ObjectBeingModified, FProp
 		FUsdStageActorImpl::AllowListComponentHierarchy(GetRootComponent(), VisitedObjects);
 	}
 
-	// So that we can detect when the user enables/disables live link properties on a ULiveLinkComponentController that may
-	// be controlling a component that we *do* care about
-	ULiveLinkComponentController* Controller = Cast< ULiveLinkComponentController >(ObjectBeingModified);
-	if (Controller)
+	// We have to accept actor and component events here, because actor transform changes do not trigger root component
+	// transform property events, and component property changes don't trigger actor property change events
+	bool bIsActorEvent = false;
+	UActorComponent* ComponentBeingModified = Cast<UActorComponent>(ObjectBeingModified);
+	if (!ComponentBeingModified || !ObjectsToWatch.Contains(ComponentBeingModified))
 	{
-		if (UActorComponent* ControlledComponent = Controller->GetControlledComponent(ULiveLinkTransformRole::StaticClass()))
+		if (AActor* ActorBeingModified = Cast<AActor>(ObjectBeingModified))
 		{
-			ObjectBeingModified = ControlledComponent;
-		}
-	}
-
-	UObject* PrimObject = ObjectBeingModified;
+			bIsActorEvent = true;
 
-	if (!ObjectsToWatch.Contains(ObjectBeingModified))
-	{
-		if (AActor* ActorBeingModified = Cast< AActor >(ObjectBeingModified))
-		{
 			if (!ObjectsToWatch.Contains(ActorBeingModified->GetRootComponent()))
 			{
 				return;
 			}
 			else
 			{
-				PrimObject = ActorBeingModified->GetRootComponent();
+				ComponentBeingModified = ActorBeingModified->GetRootComponent();
 			}
 		}
 		else
@@ -3693,9 +3823,125 @@ void AUsdStageActor::OnObjectPropertyChanged(UObject* ObjectBeingModified, FProp
 		}
 	}
 
+	// So that we can detect when the user enables/disables live link properties on a ULiveLinkComponentController that may
+	// be controlling a component that we *do* care about
+	ULiveLinkComponentController* Controller = Cast<ULiveLinkComponentController>(ComponentBeingModified);
+	if (Controller)
+	{
+		if (UActorComponent* ControlledComponent = Controller->GetControlledComponent(ULiveLinkTransformRole::StaticClass()))
+		{
+			ComponentBeingModified = ControlledComponent;
+		}
+	}
+
+	const static TSet<FName> TransformProperties = {
+		USceneComponent::GetRelativeLocationPropertyName(),
+		USceneComponent::GetRelativeRotationPropertyName(),
+		USceneComponent::GetRelativeScale3DPropertyName()
+	};
+	const bool bIsTransformChange = TransformProperties.Contains(PropertyChangedEvent.GetPropertyName());
+
+	// When we change an actor property that is just a mirror of a component property (e.g. light intensity, or camera aperture)
+	// UE will emit a property changed event on the actual component with the expected PropertyChangedEvent, and also emit a strange
+	// property changed event for the actor, with the PropertyChangedEvent claiming the component property changed (it didn't, it's still
+	// pointing at the same component). We can *almost* fully ignore these events where the object modified is an actor then, so we
+	// don't have false positives/negatives due to these strange events, except that changing the actor transform doesn't seem to
+	// fire a component transform property changed event... so we allow that case to pass through
+	if (bIsActorEvent && !bIsTransformChange)
+	{
+		return;
+	}
+
+	// Try to suppress writing anything to the stage if we're modifying a property that is animated with a track
+	// on a persistent LevelSequence currently opened in the sequencer. Otherwise we'd be constantly writing out
+	// default (non-animated) opinions for attributes that the user is trying to animate on their persistent LevelSequences.
+	// This is also important because whenever the user closes that Sequence, the modifed properties will be reverted
+	// on the UE level, but not on the stage
+#if WITH_EDITOR
+	{
+		UObject* Context = ComponentBeingModified->GetOwner();
+		AActor* OwnerActor = ComponentBeingModified->GetOwner();
+		UObject* ActorContext = OwnerActor->GetWorld();
+
+		const bool bIsRootComponent = ComponentBeingModified->GetOwner()->GetRootComponent() == ComponentBeingModified;
+
+		IUsdStageModule& UsdStageModule = FModuleManager::Get().LoadModuleChecked<IUsdStageModule>(TEXT("UsdStage"));
+		for (const TWeakPtr<ISequencer>& ExistingSequencer : UsdStageModule.GetExistingSequencers())
+		{
+			if (TSharedPtr<ISequencer> PinnedSequencer = ExistingSequencer.Pin())
+			{
+				if (UMovieSceneSequence* RootSequence = PinnedSequencer->GetRootMovieSceneSequence())
+				{
+					TSet<UMovieSceneSequence*> AllSequences;
+					FUsdStageActorImpl::GetDescendantMovieSceneSequences(RootSequence, AllSequences);
+
+					for (UMovieSceneSequence* Sequence : AllSequences)
+					{
+						UMovieScene* MovieScene = Sequence->GetMovieScene();
+						if (!MovieScene)
+						{
+							continue;
+						}
+
+						TArray<FGuid> BindingsToCheck;
+						BindingsToCheck.Add(Sequence->FindBindingFromObject(ComponentBeingModified, Context));
+						if (bIsRootComponent)
+						{
+							// Maybe all the sequence has is a track directly on the actor. That's still enough to
+							// supress a root component animation in case the property is just mirrored on the actor,
+							// so let's try checking for that
+							BindingsToCheck.Add(Sequence->FindBindingFromObject(OwnerActor, ActorContext));
+						}
+
+						for (const FGuid& BindingGuid : BindingsToCheck)
+						{
+							FMovieSceneBinding* Binding = MovieScene->FindBinding(BindingGuid);
+							if (!Binding)
+							{
+								continue;
+							}
+
+							for (const UMovieSceneTrack* Track : Binding->GetTracks())
+							{
+								// Ignore muted tracks
+								if (Track->IsEvalDisabled())
+								{
+									continue;
+								}
+
+								if (bIsTransformChange && Track->IsA<UMovieScene3DTransformTrack>())
+								{
+									return;
+								}
+
+								if (const UMovieScenePropertyTrack* PropertyTrack = Cast<const UMovieScenePropertyTrack>(Track))
+								{
+									if (PropertyTrack->GetPropertyName() == PropertyChangedEvent.GetPropertyName())
+									{
+										return;
+									}
+								}
+							}
+						}
+					}
+				}
+			}
+		}
+	}
+#endif // WITH_EDITOR
+
+	// We spawn Cine Camera Actors for Camera prims, but those have two components by default. Our convention is to place
+	// camera stuff on the camera component (not much choice there), but use the transform of the scene (root) component.
+	// Here we ignore transform changes of the camera component, emitting a warning if appropriate
+	if (bIsTransformChange && ComponentBeingModified->IsA<UCineCameraComponent>())
+	{
+		FUsdStageActorImpl::ShowTransformOnCameraComponentWarning(ComponentBeingModified);
+		return;
+	}
+
 	const UE::FUsdStage& CurrentStage = static_cast<const AUsdStageActor*>(this)->GetUsdStage();
 
-	FString PrimPath = ObjectsToWatch[PrimObject];
+	FString PrimPath = ObjectsToWatch[ComponentBeingModified];
 
 	if (UUsdPrimTwin* UsdPrimTwin = GetRootPrimTwin()->Find(PrimPath))
 	{
