USD: Fix broken baked material references when exporting on Linux;

Revert debug prints on export test now that issue has been found.

diff --git a/Engine/Plugins/Importers/USDImporter/Source/USDUtilities/Private/USDShadeConversion.cpp b/Engine/Plugins/Importers/USDImporter/Source/USDUtilities/Private/USDShadeConversion.cpp
--- a/Engine/Plugins/Importers/USDImporter/Source/USDUtilities/Private/USDShadeConversion.cpp
+++ b/Engine/Plugins/Importers/USDImporter/Source/USDUtilities/Private/USDShadeConversion.cpp
@@ -3053,11 +3053,12 @@ bool UsdUtils::SetUnrealSurfaceOutput( pxr::UsdPrim& MaterialPrim, const FString
 		return false;
 	}
 
-	UnrealShader.CreateImplementationSourceAttr( pxr::VtValue{ pxr::UsdShadeTokens->sourceAsset } );
-	UnrealShader.SetSourceAsset(
+	// Let SetSourceAsset call CreateImplementationSourceAttr internally as it will create the attribute with the correct metadata.
+	// For some reason, if we try doing this on Linux we get an attribute that seems to always output "id" when we're exporting material bindings.
+	ensure(UnrealShader.SetSourceAsset(
 		UnrealMaterialPathName.IsEmpty() ? pxr::SdfAssetPath{} : pxr::SdfAssetPath{ UnrealToUsd::ConvertString( *UnrealMaterialPathName ).Get() },
 		UnrealIdentifiers::Unreal
-	);
+	));
 	pxr::UsdShadeOutput ShaderOutput = UnrealShader.CreateOutput( UnrealToUsd::ConvertToken( TEXT( "out" ) ).Get(), pxr::SdfValueTypeNames->Token );
 
 	pxr::UsdShadeOutput MaterialOutput = Material.CreateSurfaceOutput( UnrealIdentifiers::Unreal );
