USD: Suppress downstream transactions when importing.

diff --git a/Engine/Plugins/Importers/USDImporter/Source/USDStageImporter/Private/USDStageImporter.cpp b/Engine/Plugins/Importers/USDImporter/Source/USDStageImporter/Private/USDStageImporter.cpp
--- a/Engine/Plugins/Importers/USDImporter/Source/USDStageImporter/Private/USDStageImporter.cpp
+++ b/Engine/Plugins/Importers/USDImporter/Source/USDStageImporter/Private/USDStageImporter.cpp
@@ -1754,6 +1754,16 @@ void UUsdStageImporter::ImportFromFile(FUsdStageImportContext& ImportContext)
 
 	double StartTime = FPlatformTime::Cycles64();
 
+	// If we don't have any transaction yet, let's temporarly disable the transaction buffer, to prevent code
+	// downstream from accidentally creating their own transactions. This happens for example on USkeleton::AccumulateCurveMetaData,
+	// and can lead to thousands of transactions showing up on the editor and a huge performance cost due to serialization spam.
+	// See also UEditorEngine::CanTransact.
+	TOptional<TGuardValue<TObjectPtr<class UTransactor>>> TransactionSuppressor;
+	if (!GIsTransacting)
+	{
+		TransactionSuppressor.Emplace(GEditor->Trans, nullptr);
+	}
+
 	// Load some default options in case we don't have an options object, so that we don't have to check for this every time.
 	// We could also GetMutableDefault<UUsdStageImportOptions>() instead, but I think in this case it's expected that these
 	// should be the default value and not the config
