USD: Parallel morph target creation.

diff --git a/Engine/Plugins/Importers/USDImporter/Source/USDUtilities/Private/USDSkeletalDataConversion.cpp b/Engine/Plugins/Importers/USDImporter/Source/USDUtilities/Private/USDSkeletalDataConversion.cpp
--- a/Engine/Plugins/Importers/USDImporter/Source/USDUtilities/Private/USDSkeletalDataConversion.cpp
+++ b/Engine/Plugins/Importers/USDImporter/Source/USDUtilities/Private/USDSkeletalDataConversion.cpp
@@ -26,6 +26,7 @@
 #include "Animation/AnimCurveTypes.h"
 #include "AnimationRuntime.h"
 #include "AnimEncoding.h"
+#include "Async/ParallelFor.h"
 #include "BoneWeights.h"
 #include "ControlRig.h"
 #include "Evaluation/MovieSceneSequenceTransform.h"
@@ -556,6 +557,8 @@ namespace UsdToUnrealImpl
 
 	void ComputeSourceNormals( SkelDataConversionImpl::FMeshDataBundle& UnmorphedShape )
 	{
+		TRACE_CPUPROFILER_EVENT_SCOPE(USDSkeletalDataConversion::ComputeSourceNormals);
+
 		IMeshUtilities& MeshUtilities = FModuleManager::Get().LoadModuleChecked<IMeshUtilities>(TEXT("MeshUtilities"));
 
 		// Calculate base normals for the mesh so that we can compute tangent deltas if we need to
@@ -575,6 +578,8 @@ namespace UsdToUnrealImpl
 
 	void CreateMorphTargets( UsdUtils::FBlendShapeMap& BlendShapes, const TArray<FSkeletalMeshImportData>& LODIndexToSkeletalMeshImportData, USkeletalMesh* SkeletalMesh )
 	{
+		TRACE_CPUPROFILER_EVENT_SCOPE(USDSkeletalDataConversion::CreateMorphTargets);
+
 		FSkeletalMeshModel* ImportedResource = SkeletalMesh->GetImportedModel();
 		if ( LODIndexToSkeletalMeshImportData.Num() != ImportedResource->LODModels.Num() )
 		{
@@ -591,9 +596,9 @@ namespace UsdToUnrealImpl
 		OrigIndexToBuiltIndicesPerLOD.Reserve( NumLODs );
 
 		TArray<bool> AreNormalsComputedForLODIndex;
-		AreNormalsComputedForLODIndex.SetNumZeroed( NumLODs );
+		AreNormalsComputedForLODIndex.SetNumZeroed(NumLODs);
 
-		for ( int32 LODIndex = 0; LODIndex < NumLODs; ++LODIndex )
+		for (int32 LODIndex = 0; LODIndex < NumLODs; ++LODIndex)
 		{
 			FSkeletalMeshLODModel& LODModel = ImportedResource->LODModels[LODIndex];
 
@@ -602,74 +607,110 @@ namespace UsdToUnrealImpl
 			// to original vertex indices, so we'll need to map them to the post-build vertex indices
 			TArray<int32>& BuildIndexToOrigIndex = LODModel.MeshToImportVertexMap;
 			TMap<int32, TArray<int32>>& OrigIndexToBuiltIndices = OrigIndexToBuiltIndicesPerLOD.Emplace_GetRef();
-			OrigIndexToBuiltIndices.Reserve( BuildIndexToOrigIndex.Num() );
-			for ( int32 BuiltIndex = 0; BuiltIndex < BuildIndexToOrigIndex.Num(); ++BuiltIndex )
+			OrigIndexToBuiltIndices.Reserve(BuildIndexToOrigIndex.Num());
+			for (int32 BuiltIndex = 0; BuiltIndex < BuildIndexToOrigIndex.Num(); ++BuiltIndex)
 			{
-				int32 OrigIndex = BuildIndexToOrigIndex[ BuiltIndex ];
-				OrigIndexToBuiltIndices.FindOrAdd( OrigIndex ).Add( BuiltIndex );
+				int32 OrigIndex = BuildIndexToOrigIndex[BuiltIndex];
+				OrigIndexToBuiltIndices.FindOrAdd(OrigIndex).Add(BuiltIndex);
 			}
 
 			SkelDataConversionImpl::FMeshDataBundle& LODMeshBundle = TempMeshBundlesPerLOD.Emplace_GetRef();
-			SkelDataConversionImpl::ConvertImportDataToMeshData( LODIndexToSkeletalMeshImportData[LODIndex], LODMeshBundle );
+			SkelDataConversionImpl::ConvertImportDataToMeshData(LODIndexToSkeletalMeshImportData[LODIndex], LODMeshBundle);
 		}
 
+		struct FMorphTargetJob
+		{
+			UsdUtils::FUsdBlendShape* BlendShape = nullptr;
+			UMorphTarget* MorphTarget = nullptr;
+		};
+		TArray<FMorphTargetJob> MorphTargetJobs;
+		MorphTargetJobs.Reserve(BlendShapes.Num());
+
 		bool bHasValidMorphTarget = false;
-		for ( TPair<FString, UsdUtils::FUsdBlendShape> BlendShapeByPath : BlendShapes )
+		for (TPair<FString, UsdUtils::FUsdBlendShape>& BlendShapeByPath : BlendShapes)
 		{
-			UsdUtils::FUsdBlendShape& BlendShape = BlendShapeByPath.Value;
+			TRACE_CPUPROFILER_EVENT_SCOPE(USDSkeletalDataConversion::CreatingMorphTargetObject);
 
-			if ( !BlendShape.IsValid() )
+			UsdUtils::FUsdBlendShape& BlendShape = BlendShapeByPath.Value;
+			if (!BlendShape.IsValid())
 			{
 				continue;
 			}
+
 			bHasValidMorphTarget = true;
 
-			UMorphTarget* MorphTarget = NewObject<UMorphTarget>( SkeletalMesh, *BlendShape.Name );
+			FMorphTargetJob& Job = MorphTargetJobs.Emplace_GetRef();
+			Job.BlendShape = &BlendShape;
+			Job.MorphTarget = NewObject<UMorphTarget>(SkeletalMesh, *BlendShape.Name);
+			Job.MorphTarget->BaseSkelMesh = SkeletalMesh;
+			SkeletalMesh->GetMorphTargets().Add(Job.MorphTarget);
 
-			for ( int32 LODIndex : BlendShape.LODIndicesThatUseThis )
+			// Recompute normals for the final morphed shape in case it doesn't have authored normals
+			// This is required or else the morphed shape will reuse the unmorphed normals, and lighting may look incorrect for
+			// aggressive morph targets.
+			// Note that this should happen *before* we call UpdatesDeltasToMeshBuild, because our MeshDataBundle refers to import data,
+			// and so should our BlendShape
+			for (int32 LODIndex : Job.BlendShape->LODIndicesThatUseThis)
 			{
-				SkelDataConversionImpl::FMeshDataBundle& UnmorphedShape = TempMeshBundlesPerLOD[ LODIndex ];
-
-				// Recompute normals for the final morphed shape in case it doesn't have authored normals
-				// This is required or else the morphed shape will reuse the unmorphed normals, and lighting may look incorrect for
-				// aggressive morph targets.
-				// Note that this should happen *before* we call UpdatesDeltasToMeshBuild, because our MeshDataBundle refers to import data,
-				// and so should our BlendShape
-				if ( !BlendShape.bHasAuthoredTangents )
+				if (!Job.BlendShape->bHasAuthoredTangents)
 				{
-					if ( !AreNormalsComputedForLODIndex[ LODIndex ] )
+					SkelDataConversionImpl::FMeshDataBundle& UnmorphedShape = TempMeshBundlesPerLOD[LODIndex];
+					if (!AreNormalsComputedForLODIndex[LODIndex])
 					{
-						ComputeSourceNormals( UnmorphedShape );
-						AreNormalsComputedForLODIndex[ LODIndex ] = true;
+						ComputeSourceNormals(UnmorphedShape);
+						AreNormalsComputedForLODIndex[LODIndex] = true;
 					}
-
-					SkelDataConversionImpl::ComputeTangentDeltas( UnmorphedShape, BlendShape );
 				}
+			}
+		}
 
-				TArray<FMorphTargetDelta> Vertices;
-				if ( BlendShape.LODIndicesThatUseThis.Num() > 1 )
+		const int32 MinBatchSize = 1;
+
+		ParallelFor(
+			TEXT("CreateMorphTarget"),
+			MorphTargetJobs.Num(),
+			MinBatchSize,
+			[&MorphTargetJobs, &OrigIndexToBuiltIndicesPerLOD, &TempMeshBundlesPerLOD, ImportedResource](int32 Index)
+			{
+				TRACE_CPUPROFILER_EVENT_SCOPE(USDSkeletalDataConversion::CreateMorphTargetJob);
+
+				FMorphTargetJob& Job = MorphTargetJobs[Index];
+				if (!Job.BlendShape || !Job.MorphTarget)
 				{
-					// Need to copy this here because different LODs may build differently, and so the deltas may need to be updated differently
-					Vertices = BlendShape.Vertices;
+					return;
 				}
-				else
+
+				for (int32 LODIndex : Job.BlendShape->LODIndicesThatUseThis)
 				{
-					Vertices = MoveTemp(BlendShape.Vertices);
-				}
+					const SkelDataConversionImpl::FMeshDataBundle& UnmorphedShape = TempMeshBundlesPerLOD[LODIndex];
 
-				SkelDataConversionImpl::UpdatesDeltasToMeshBuild( Vertices, OrigIndexToBuiltIndicesPerLOD[ LODIndex ] );
+					if (!Job.BlendShape->bHasAuthoredTangents)
+					{
+						SkelDataConversionImpl::ComputeTangentDeltas(UnmorphedShape, *Job.BlendShape);
+					}
 
-				const bool bCompareNormal = true;
-				FSkeletalMeshLODModel& LODModel = ImportedResource->LODModels[LODIndex];
-				MorphTarget->PopulateDeltas( Vertices, LODIndex, LODModel.Sections, bCompareNormal );
-			}
+					TArray<FMorphTargetDelta> Vertices;
+					if (Job.BlendShape->LODIndicesThatUseThis.Num() > 1)
+					{
+						// Need to copy this here because different LODs may build differently, and so the deltas may need to be updated differently
+						Vertices = Job.BlendShape->Vertices;
+					}
+					else
+					{
+						Vertices = MoveTemp(Job.BlendShape->Vertices);
+					}
 
-			// Don't need this data anymore as it has been moved into UMorphTarget
-			BlendShape.Vertices.Empty();
+					SkelDataConversionImpl::UpdatesDeltasToMeshBuild(Vertices, OrigIndexToBuiltIndicesPerLOD[LODIndex]);
 
-			MorphTarget->BaseSkelMesh = SkeletalMesh;
-			SkeletalMesh->GetMorphTargets().Add( MorphTarget );
-		}
+					const bool bCompareNormal = true;
+					FSkeletalMeshLODModel& LODModel = ImportedResource->LODModels[LODIndex];
+					Job.MorphTarget->PopulateDeltas(Vertices, LODIndex, LODModel.Sections, bCompareNormal);
+				}
+
+				// Don't need this data anymore as it has been moved into UMorphTarget
+				Job.BlendShape->Vertices.Empty();
+			}
+		);
 
 		if ( bHasValidMorphTarget )
 		{
@@ -2362,6 +2403,8 @@ bool UsdToUnreal::ConvertBlendShape( const pxr::UsdSkelBlendShape& UsdBlendShape
 
 bool UsdToUnreal::ConvertBlendShape( const pxr::UsdSkelBlendShape& UsdBlendShape, const FUsdStageInfo& StageInfo, int32 LODIndex, uint32 PointIndexOffset, TSet<FString>& UsedMorphTargetNames, UsdUtils::FBlendShapeMap& OutBlendShapes )
 {
+	TRACE_CPUPROFILER_EVENT_SCOPE( UsdToUnreal::ConvertBlendShape );
+
 	FScopedUsdAllocs Allocs;
 
 	pxr::UsdAttribute OffsetsAttr = UsdBlendShape.GetOffsetsAttr();
