USD: Show referenced prims on references panel, and make internal references explicit.

diff --git a/Engine/Plugins/Importers/USDImporter/Source/USDStageEditor/Private/Widgets/SUSDReferencesList.cpp b/Engine/Plugins/Importers/USDImporter/Source/USDStageEditor/Private/Widgets/SUSDReferencesList.cpp
--- a/Engine/Plugins/Importers/USDImporter/Source/USDStageEditor/Private/Widgets/SUSDReferencesList.cpp
+++ b/Engine/Plugins/Importers/USDImporter/Source/USDStageEditor/Private/Widgets/SUSDReferencesList.cpp
@@ -15,6 +15,8 @@
 
 #if USE_USD_SDK
 
+#define LOCTEXT_NAMESPACE "USDReferencesList"
+
 namespace UsdReferencesListConstants
 {
 	const FMargin RowPadding( 6.0f, 2.5f, 2.0f, 2.5f );
@@ -35,9 +37,16 @@ TSharedRef< SWidget > SUsdReferenceRow::GenerateWidgetForColumn( const FName& Co
 
 	if ( ColumnName == TEXT("AssetPath") )
 	{
-		SAssignNew( ColumnWidget, STextBlock )
-		.Text( FText::FromString( Reference->AssetPath ) )
-		.Font( FAppStyle::GetFontStyle( UsdReferencesListConstants::NormalFont ) );
+		SAssignNew(ColumnWidget, STextBlock)
+		.Text(FText::FromString(Reference->AssetPath.IsEmpty() ? TEXT("(internal reference)") : Reference->AssetPath))
+		.Font(FAppStyle::GetFontStyle(UsdReferencesListConstants::NormalFont));
+	}
+
+	if ( ColumnName == TEXT("PrimPath") )
+	{
+		SAssignNew(ColumnWidget, STextBlock)
+		.Text(FText::FromString(Reference->PrimPath.IsEmpty() ? TEXT("(default prim)") : Reference->PrimPath))
+		.Font(FAppStyle::GetFontStyle(UsdReferencesListConstants::NormalFont));
 	}
 
 	return SNew( SBox )
@@ -57,11 +66,15 @@ TSharedRef< SWidget > SUsdReferenceRow::GenerateWidgetForColumn( const FName& Co
 
 void SUsdReferencesList::Construct( const FArguments& InArgs )
 {
-	SAssignNew( HeaderRowWidget, SHeaderRow )
+	SAssignNew(HeaderRowWidget, SHeaderRow)
+
+	+SHeaderRow::Column(FName(TEXT("AssetPath")))
+	.DefaultLabel(LOCTEXT("ReferencedPath", "Referenced layers"))
+	.FillWidth(100.f)
 
-	+SHeaderRow::Column( FName( TEXT("AssetPath") ) )
-	.DefaultLabel( NSLOCTEXT( "USDReferencesList", "References", "References" ) )
-	.FillWidth( 100.f );
+	+SHeaderRow::Column(FName(TEXT("PrimPath")))
+	.DefaultLabel(LOCTEXT("ReferencedPrim", "Referenced prims"))
+	.FillWidth(100.f);
 
 	SListView::Construct
 	(
@@ -88,4 +101,6 @@ void SUsdReferencesList::SetPrimPath( const UE::FUsdStageWeak& UsdStage, const T
 	RequestListRefresh();
 }
 
+#undef LOCTEXT_NAMESPACE
+
 #endif // #if USE_USD_SDK
diff --git a/Engine/Plugins/Importers/USDImporter/Source/USDStageEditorViewModels/Private/USDReferencesViewModel.cpp b/Engine/Plugins/Importers/USDImporter/Source/USDStageEditorViewModels/Private/USDReferencesViewModel.cpp
--- a/Engine/Plugins/Importers/USDImporter/Source/USDStageEditorViewModels/Private/USDReferencesViewModel.cpp
+++ b/Engine/Plugins/Importers/USDImporter/Source/USDStageEditorViewModels/Private/USDReferencesViewModel.cpp
@@ -46,6 +46,7 @@ void FUsdReferencesViewModel::UpdateReferences( const UE::FUsdStageWeak& UsdStag
 				{
 					FUsdReference Reference;
 					Reference.AssetPath = UsdToUnreal::ConvertString( UsdReference.GetAssetPath() );
+					Reference.PrimPath = UsdToUnreal::ConvertPath( UsdReference.GetPrimPath() );
 
 					References.Add( MakeSharedUnreal< FUsdReference >( MoveTemp( Reference ) ) );
 				}
diff --git a/Engine/Plugins/Importers/USDImporter/Source/USDStageEditorViewModels/Public/USDReferencesViewModel.h b/Engine/Plugins/Importers/USDImporter/Source/USDStageEditorViewModels/Public/USDReferencesViewModel.h
--- a/Engine/Plugins/Importers/USDImporter/Source/USDStageEditorViewModels/Public/USDReferencesViewModel.h
+++ b/Engine/Plugins/Importers/USDImporter/Source/USDStageEditorViewModels/Public/USDReferencesViewModel.h
@@ -10,6 +10,7 @@ class USDSTAGEEDITORVIEWMODELS_API FUsdReference : public TSharedFromThis< FUsdR
 {
 public:
 	FString AssetPath;
+	FString PrimPath;
 };
 
 class USDSTAGEEDITORVIEWMODELS_API FUsdReferencesViewModel
