USD: Fix warning about not being able to export to layer firing in wrong scenarios due to drive letter case.

diff --git a/Engine/Plugins/Importers/USDImporter/Source/USDExporter/Private/USDExporterModule.cpp b/Engine/Plugins/Importers/USDImporter/Source/USDExporter/Private/USDExporterModule.cpp
--- a/Engine/Plugins/Importers/USDImporter/Source/USDExporter/Private/USDExporterModule.cpp
+++ b/Engine/Plugins/Importers/USDImporter/Source/USDExporter/Private/USDExporterModule.cpp
@@ -101,7 +101,7 @@ bool IUsdExporterModule::CanExportToLayer( const FString& TargetFilePath )
 {
 	for ( const UE::FSdfLayerWeak& LoadedLayer : UE::FSdfLayerWeak::GetLoadedLayers() )
 	{
-		if ( LoadedLayer.GetIdentifier() == TargetFilePath )
+		if (FPaths::IsSamePath(LoadedLayer.GetRealPath(), TargetFilePath))
 		{
 			FText ErrorMessage = FText::Format(
 				LOCTEXT( "FailedExportLayerOpenSubText", "Failed to export to the USD layer '{0}' as a layer with this identifier is already open in another USD Stage. Please use a different file path or try closing that other USD stage and exporting again." ),
