USD: Fix warning/error spam when we fail to parse any material parameter from preview surface materials.

diff --git a/Engine/Plugins/Importers/USDImporter/Source/USDUtilities/Private/USDShadeConversion.cpp b/Engine/Plugins/Importers/USDImporter/Source/USDUtilities/Private/USDShadeConversion.cpp
--- a/Engine/Plugins/Importers/USDImporter/Source/USDUtilities/Private/USDShadeConversion.cpp
+++ b/Engine/Plugins/Importers/USDImporter/Source/USDUtilities/Private/USDShadeConversion.cpp
@@ -2186,7 +2186,9 @@ bool UsdToUnreal::ConvertMaterial(
 		UserData->ParameterToPrimvar = ParameterToPrimvar;
 	}
 
-	return MaterialParameters.Num() > 0;
+	// We used to only reture true in case we managed to convert at least one parameter, but we don't want
+	// callers to interpret that we "failed to convert the material" if we couldn't find any usable parameter
+	return true;
 }
 
 bool UsdToUnreal::ConvertMaterial(
@@ -2196,7 +2198,6 @@ bool UsdToUnreal::ConvertMaterial(
 	const TCHAR* RenderContext
 )
 {
-	bool bHasMaterialInfo = false;
 #if WITH_EDITOR
 	pxr::TfToken RenderContextToken = pxr::UsdShadeTokens->universalRenderContext;
 	if (RenderContext)
@@ -2318,7 +2319,7 @@ bool UsdToUnreal::ConvertMaterial(
 
 	// auto so we can use the MaterialInput generic parameter for FColorMaterialInput, FScalarMaterialInput and FVectorMaterialInput
 	auto ConnectMaterialInput =
-		[&PrimvarToUVIndex, &bHasMaterialInfo]
+		[&PrimvarToUVIndex]
 		(UMaterial& Material, auto& MaterialInput, const UsdShadeConversionImpl::FParameterValue& InParameterValue) -> bool
 		{
 			UMaterialExpression* Expression = UsdShadeConversionImpl::GetExpressionForValue(Material, InParameterValue);
@@ -2327,8 +2328,6 @@ bool UsdToUnreal::ConvertMaterial(
 				return false;
 			}
 
-			bHasMaterialInfo = true;
-
 			MaterialInput.Expression = Expression;
 
 			if (const UsdShadeConversionImpl::FTextureParameterValue* TextureParameterValue = InParameterValue.TryGet<UsdShadeConversionImpl::FTextureParameterValue>())
@@ -2461,8 +2460,10 @@ bool UsdToUnreal::ConvertMaterial(
 		UserData->ParameterToPrimvar = ParameterToPrimvar;
 	}
 
+	return true;
+#else
+	return false;
 #endif // WITH_EDITOR
-	return bHasMaterialInfo;
 }
 
 // Deprecated
