USD: Fix actors spawned on Actions -> Import workflow ending up transient.

diff --git a/Engine/Plugins/Importers/USDImporter/Source/USDStageEditorViewModels/Private/USDStageViewModel.cpp b/Engine/Plugins/Importers/USDImporter/Source/USDStageEditorViewModels/Private/USDStageViewModel.cpp
--- a/Engine/Plugins/Importers/USDImporter/Source/USDStageEditorViewModels/Private/USDStageViewModel.cpp
+++ b/Engine/Plugins/Importers/USDImporter/Source/USDStageEditorViewModels/Private/USDStageViewModel.cpp
@@ -335,9 +335,7 @@ void FUsdStageViewModel::ImportStage( const TCHAR* TargetContentFolder, UUsdStag
 
 		const bool bIsAutomated = TargetContentFolder && Options;
 
-		// Even when importing we still want to first create these assets as transient. Only the publishing process
-		// itself will remove the transient flag if everything succeeded
-		if ( ImportContext.Init( StageName, RootPath, TEXT("/Game/"), RF_Public | RF_Transactional | RF_Transient, bIsAutomated ) )
+		if (ImportContext.Init(StageName, RootPath, TEXT("/Game/"), RF_Public | RF_Transactional | RF_Standalone, bIsAutomated))
 		{
 			FScopedTransaction Transaction( FText::Format(LOCTEXT("ImportTransaction", "Import USD stage '{0}'"), FText::FromString(StageName)));
 
