USD: Fix UV transform support on Linux.

diff --git a/Engine/Plugins/Importers/USDImporter/Source/USDUtilities/Private/USDShadeConversion.cpp b/Engine/Plugins/Importers/USDImporter/Source/USDUtilities/Private/USDShadeConversion.cpp
--- a/Engine/Plugins/Importers/USDImporter/Source/USDUtilities/Private/USDShadeConversion.cpp
+++ b/Engine/Plugins/Importers/USDImporter/Source/USDUtilities/Private/USDShadeConversion.cpp
@@ -469,13 +469,12 @@ namespace UE
 					// UV transform shader node
 					if (ConnectedNodeId == UnrealIdentifiers::UsdTransform2d)
 					{
-						FScale2f Scale;
+						FScale2f Scale{1.0f, 1.0f};
 						if (pxr::UsdShadeInput ScaleInput = Connectable.GetInput(UnrealIdentifiers::Scale))
 						{
-							pxr::VtValue InputValue;
-							if (ScaleInput.Get(&InputValue) && InputValue.IsHolding<pxr::GfVec2f>())
+							pxr::GfVec2f VecValue;
+							if (ScaleInput.Get(&VecValue))
 							{
-								pxr::GfVec2f VecValue = InputValue.Get<pxr::GfVec2f>();
 								Scale = FScale2f{VecValue[0], VecValue[1]};
 							}
 						}
@@ -483,20 +482,19 @@ namespace UE
 						float Rotation = 0.0f;
 						if (pxr::UsdShadeInput RotationInput = Connectable.GetInput(UnrealIdentifiers::Rotation))
 						{
-							pxr::VtValue InputValue;
-							if (RotationInput.Get(&InputValue) && InputValue.IsHolding<float>())
+							float FloatValue;
+							if (RotationInput.Get(&FloatValue))
 							{
-								Rotation = InputValue.Get<float>();
+								Rotation = FloatValue;
 							}
 						}
 
-						FVector2f Translation;
+						FVector2f Translation{0.0f, 0.0f};
 						if (pxr::UsdShadeInput TranslationInput = Connectable.GetInput(UnrealIdentifiers::Translation))
 						{
-							pxr::VtValue InputValue;
-							if (TranslationInput.Get(&InputValue) && InputValue.IsHolding<pxr::GfVec2f>())
+							pxr::GfVec2f VecValue;
+							if (TranslationInput.Get(&VecValue))
 							{
-								pxr::GfVec2f VecValue = InputValue.Get<pxr::GfVec2f>();
 								Translation = FVector2f{VecValue[0], VecValue[1]};
 							}
 						}
