USD caching for tetrahedral mesh sims.  Converts the FFleshCacheAdapter to read and write a monolithic usd file on Win64.

Introduces FComponentCacheAdapter::Finalize() callback, and optionally feeds more data to InitializeForRecord() and InitializeForPlayback(), enabling options to be passed from the UI to cache adapters.

Create USDImporter/SimCacheUSD plugin.  Implements tetrahedral mesh gprim.  Contains facilities for managing stages, and 'value clips' streaming.

preflight: /64bb269c9a2f949266dfb175

diff --git a/Engine/Plugins/Experimental/ChaosCaching/Source/ChaosCaching/Private/Chaos/Adapters/CacheAdapter.cpp b/Engine/Plugins/Experimental/ChaosCaching/Source/ChaosCaching/Private/Chaos/Adapters/CacheAdapter.cpp
--- a/Engine/Plugins/Experimental/ChaosCaching/Source/ChaosCaching/Private/Chaos/Adapters/CacheAdapter.cpp
+++ b/Engine/Plugins/Experimental/ChaosCaching/Source/ChaosCaching/Private/Chaos/Adapters/CacheAdapter.cpp
@@ -1,6 +1,7 @@
 // Copyright Epic Games, Inc. All Rights Reserved.
 
 #include "Chaos/Adapters/CacheAdapter.h"
+#include "Chaos/CacheManagerActor.h"
 #include "Features/IModularFeatures.h"
 #include "Components/PrimitiveComponent.h"
 
@@ -17,6 +18,16 @@ DEFINE_LOG_CATEGORY(LogCacheAdapter);
 
 namespace Chaos
 {
+	bool FComponentCacheAdapter::InitializeForRecord(UPrimitiveComponent* InComponent, FObservedComponent& InObserved)
+	{
+		return InitializeForRecord(InComponent, InObserved.Cache);
+	}
+
+	bool FComponentCacheAdapter::InitializeForPlayback(UPrimitiveComponent* InComponent, FObservedComponent& InObserved, float InTime)
+	{
+		return InitializeForPlayback(InComponent, InObserved.Cache, InTime);
+	}
+
 	void RegisterAdapter(FComponentCacheAdapter* InAdapter)
 	{
 		check(InAdapter);
diff --git a/Engine/Plugins/Experimental/ChaosCaching/Source/ChaosCaching/Private/Chaos/CacheManagerActor.cpp b/Engine/Plugins/Experimental/ChaosCaching/Source/ChaosCaching/Private/Chaos/CacheManagerActor.cpp
--- a/Engine/Plugins/Experimental/ChaosCaching/Source/ChaosCaching/Private/Chaos/CacheManagerActor.cpp
+++ b/Engine/Plugins/Experimental/ChaosCaching/Source/ChaosCaching/Private/Chaos/CacheManagerActor.cpp
@@ -127,6 +127,76 @@ void AChaosCacheManager::TickActor(float DeltaTime, enum ELevelTick TickType, FA
 	}
 }
 
+#if WITH_EDITOR
+bool
+AChaosCacheManager::ContainsProperty(
+	const UStruct* Struct,
+	const void* InProperty) const
+{
+	for (FPropertyValueIterator PropIt(FProperty::StaticClass(), Struct, this); PropIt; ++PropIt)
+	{
+		const void* CurrProperty = PropIt.Key();
+		if (InProperty == CurrProperty)
+		{
+			return true;
+		}
+	}
+	return false;
+}
+
+bool AChaosCacheManager::CanEditChange(const FProperty* InProperty) const
+{
+	const bool bRetVal = Super::CanEditChange(InProperty);
+	if (!bRetVal)
+	{
+		return false;
+	}
+
+	if (InProperty->GetFName() == GET_MEMBER_NAME_CHECKED(FObservedComponent, USDCacheDirectory))
+	{
+		// Find the FObservedComponent this FProperty corresponds to, and determine if the
+		// associated component has an associated adapter that supports the USD workflow.
+
+		// We hard code the type name of the component we look for, because currently the 
+		// tetrahedral component is the only one supported, and it exists in ChaosFlesh, which is 
+		// a downstream dependency of Chaos. Once all adapters for deformables support the USD 
+		// caching workflow, then we can target a base class that exists in Chaos.
+		FName TetCompName(TEXT("DeformableTetrahedralComponent"));
+		for (const FObservedComponent& Observed : ObservedComponents)
+		{
+			TUniquePtr<FStructOnScope> StructOnScope(
+				new FStructOnScope(FObservedComponent::StaticStruct(), (uint8*)&Observed));
+			if (StructOnScope)
+			{
+				if (const UStruct* Struct = StructOnScope->GetStruct())
+				{
+					if (ContainsProperty(Struct, InProperty))
+					{
+						if (UPrimitiveComponent* Comp = Observed.GetComponent(const_cast<AChaosCacheManager*>(this)))
+						{
+							UClass* Class = Comp->GetClass();
+							do {
+								FName ClassFName = Class->GetFName();
+								if (Class->GetFName() == TetCompName)
+								{
+									return true;
+								}
+								Class = Class->GetSuperClass();
+							} while (Class);
+						}
+					}
+				}
+			}
+		}
+		// Otherwise, no.
+		return false;
+	}
+
+	return bRetVal;
+}
+#endif
+
+
 #if WITH_EDITOR
 void AChaosCacheManager::PostEditChangeProperty(FPropertyChangedEvent& PropertyChangedEvent)
 {
@@ -377,7 +447,7 @@ void AChaosCacheManager::BeginEvaluate()
 								Observed.TickRecord.SetSpaceTransform(GetTransform());
 								Observed.TickRecord.SetLastTime(StartTime);
 								OpenPlaybackCaches.Add(TTuple<FCacheUserToken, UChaosCache*>(MoveTemp(Token), Observed.Cache));
-								CurrAdapter->InitializeForPlayback(Comp, Observed.Cache, StartTime);
+								CurrAdapter->InitializeForPlayback(Comp, Observed, StartTime);
 							}
 						}
 						else
@@ -425,7 +495,7 @@ void AChaosCacheManager::BeginEvaluate()
 						// We'll record the observed component in Cache Manager's local space.
 						Observed.TickRecord.SetSpaceTransform(GetTransform());
 						OpenRecordCaches.Add(TTuple<FCacheUserToken, UChaosCache*>(MoveTemp(Token), Observed.Cache));
-						CurrAdapter->InitializeForRecord(Comp, Observed.Cache);
+						CurrAdapter->InitializeForRecord(Comp, Observed);
 
 						// Ensure we enable the actor tick to flush out the pending record writes
 						bRequiresRecord =  true;
@@ -469,6 +539,14 @@ void AChaosCacheManager::EndEvaluate()
 
 	WaitForObservedComponentSolverTasks();
 
+	// Build list of adapters for our observed components
+	IModularFeatures& ModularFeatures = IModularFeatures::Get();
+	TArray<Chaos::FComponentCacheAdapter*> Adapters = ModularFeatures.GetModularFeatureImplementations<Chaos::FComponentCacheAdapter>(Chaos::FComponentCacheAdapter::FeatureName);
+	for (Chaos::FComponentCacheAdapter* Adapter : Adapters)
+	{
+		Adapter->Finalize();
+	}
+
 	for(TPair<Chaos::FPhysicsSolverEvents*, FPerSolverData> PerSolver : PerSolverData)
 	{
 		Chaos::FPhysicsSolverEvents* CurrSolver = PerSolver.Key;
diff --git a/Engine/Plugins/Experimental/ChaosCaching/Source/ChaosCaching/Public/Chaos/Adapters/CacheAdapter.h b/Engine/Plugins/Experimental/ChaosCaching/Source/ChaosCaching/Public/Chaos/Adapters/CacheAdapter.h
--- a/Engine/Plugins/Experimental/ChaosCaching/Source/ChaosCaching/Public/Chaos/Adapters/CacheAdapter.h
+++ b/Engine/Plugins/Experimental/ChaosCaching/Source/ChaosCaching/Public/Chaos/Adapters/CacheAdapter.h
@@ -13,6 +13,7 @@ class UClass;
 class UChaosCache;
 class UPrimitiveComponent;
 struct FPlaybackTickRecord;
+struct FObservedComponent;
 
 DECLARE_LOG_CATEGORY_EXTERN(LogCacheAdapter, Log, All);
 
@@ -122,6 +123,11 @@ namespace Chaos
 		 */
 		virtual void Initialize() {}
 
+		/**
+		 * Called from the game thread to perform any final steps before deallocation.
+		 */
+		virtual void Finalize() {}
+
 		/**
 		* Called from the game thread to set rest state from an evaluated time in the cache.
 		*/
@@ -132,7 +138,8 @@ namespace Chaos
 		 * @param InComponent Target component to initialize
 		 * @param InCache Target cache to initialize
 		 */
-		virtual bool InitializeForRecord(UPrimitiveComponent* InComponent, UChaosCache* InCache) = 0;
+		virtual bool InitializeForRecord(UPrimitiveComponent* InComponent, UChaosCache* InCache) { unimplemented(); return false; }
+		virtual bool CHAOSCACHING_API InitializeForRecord(UPrimitiveComponent* InComponent, FObservedComponent& InObserved);
 
 		/**
 		 * Called from the game thread to initialize a component and cache ready to playback a cache
@@ -140,7 +147,8 @@ namespace Chaos
 		 * @param InCache Target cache to initialize
 		 * @param InTime Time from cache start to evaluate initial conditions
 		 */
-		virtual bool InitializeForPlayback(UPrimitiveComponent* InComponent, UChaosCache* InCache, float InTime) = 0;
+		virtual bool InitializeForPlayback(UPrimitiveComponent* InComponent, UChaosCache* InCache, float InTime) { unimplemented(); return false; }
+		virtual bool CHAOSCACHING_API InitializeForPlayback(UPrimitiveComponent* InComponent, FObservedComponent& InObserved, float InTime);
 
 		/**
 		 * Called by a cache observer when a component should be recorded to a cache through this adapter.
diff --git a/Engine/Plugins/Experimental/ChaosCaching/Source/ChaosCaching/Public/Chaos/CacheManagerActor.h b/Engine/Plugins/Experimental/ChaosCaching/Source/ChaosCaching/Public/Chaos/CacheManagerActor.h
--- a/Engine/Plugins/Experimental/ChaosCaching/Source/ChaosCaching/Public/Chaos/CacheManagerActor.h
+++ b/Engine/Plugins/Experimental/ChaosCaching/Source/ChaosCaching/Public/Chaos/CacheManagerActor.h
@@ -108,6 +108,10 @@ struct FObservedComponent
 	UPROPERTY(EditAnywhere, BlueprintReadWrite, Category = "Caching")
 	bool bPlaybackEnabled;
 
+	/** USD cache directory, if supported. */
+	UPROPERTY(EditAnywhere, BlueprintReadWrite, Category = "Caching")
+	FDirectoryPath USDCacheDirectory;
+
 	/** 
 	* Capture the state of bNotifyBreaks of the component before cache manager takes control. 
 	* this is because when recording the cache needs the component to have bNotifyBreaks set on the component 
@@ -134,6 +138,7 @@ struct FObservedComponent
 
 private:
 	friend class AChaosCacheManager;
+	friend class Chaos::FComponentCacheAdapter;
 
 	bool         bTriggered;          // Whether the observed component is active
 	Chaos::FReal AbsoluteTime;        // Time since BeginPlay
@@ -206,6 +211,11 @@ public:
 	friend class IChaosCachingEditorPlugin;
 #endif
 
+#if WITH_EDITOR
+	bool ContainsProperty(const UStruct* Struct, const void* InProperty) const;
+	CHAOSCACHING_API virtual bool CanEditChange(const FProperty* InProperty) const override;
+#endif
+
 #if WITH_EDITOR
 	CHAOSCACHING_API virtual void PostEditChangeProperty(FPropertyChangedEvent& PropertyChangedEvent) override;
 #endif
diff --git a/Engine/Plugins/Experimental/ChaosFlesh/ChaosFlesh.uplugin b/Engine/Plugins/Experimental/ChaosFlesh/ChaosFlesh.uplugin
--- a/Engine/Plugins/Experimental/ChaosFlesh/ChaosFlesh.uplugin
+++ b/Engine/Plugins/Experimental/ChaosFlesh/ChaosFlesh.uplugin
@@ -61,6 +61,10 @@
 		{
 			"Name": "DeformerGraph",
 			"Enabled": true
-		}
+		},
+        {
+            "Name": "USDImporter",
+            "Enabled": true
+        }
 	]
 }
diff --git a/Engine/Plugins/Experimental/ChaosFlesh/Source/ChaosFleshEngine/ChaosFleshEngine.Build.cs b/Engine/Plugins/Experimental/ChaosFlesh/Source/ChaosFleshEngine/ChaosFleshEngine.Build.cs
--- a/Engine/Plugins/Experimental/ChaosFlesh/Source/ChaosFleshEngine/ChaosFleshEngine.Build.cs
+++ b/Engine/Plugins/Experimental/ChaosFlesh/Source/ChaosFleshEngine/ChaosFleshEngine.Build.cs
@@ -1,5 +1,7 @@
 // Copyright Epic Games, Inc. All Rights Reserved.
 
+using System.Linq;
+
 namespace UnrealBuildTool.Rules
 {
 	public class ChaosFleshEngine : ModuleRules
@@ -46,7 +48,11 @@ namespace UnrealBuildTool.Rules
 					"Projects",
 					"RenderCore",
                     "RHI",
-					"Renderer"
+					"Renderer",
+					"SimCacheUSD",
+					"UnrealUSDWrapper",
+					"USDClasses",
+					"USDUtilities",
 				}
 				);
 
@@ -58,6 +64,70 @@ namespace UnrealBuildTool.Rules
 				);
 
 			PrivateDefinitions.Add("CHAOS_INCLUDE_LEVEL_1=1");
+
+			// Build flavor selection copied from UnrealUSDWrapper, then modified.
+			// Currently only Win64 is supported.
+			if (EnableUsdSdk(Target) && (Target.Type == TargetType.Editor && Target.Platform == UnrealTargetPlatform.Win64))
+			{
+				bUseRTTI = true;
+				PublicDefinitions.Add("DO_USD_CACHING=1");
+			}
+			else
+			{
+				PublicDefinitions.Add("DO_USD_CACHING=0");
+			}
+		}
+		bool EnableUsdSdk(ReadOnlyTargetRules Target)
+		{
+			// USD SDK has been built against Python 3 and won't launch if the editor is using Python 2
+
+			bool bEnableUsdSdk = (
+				Target.WindowsPlatform.Compiler != WindowsCompiler.Clang &&
+				Target.StaticAnalyzer == StaticAnalyzer.None
+			);
+
+			// Don't enable USD when running the include tool because it has issues parsing Boost headers
+			if (Target.GlobalDefinitions.Contains("UE_INCLUDE_TOOL=1"))
+			{
+				bEnableUsdSdk = false;
+			}
+
+			// If you want to use USD in a monolithic target, you'll have to use the ANSI allocator.
+			// USD always uses the ANSI C allocators directly. In a DLL UE build (so not monolithic) we can just override the operators new and delete
+			// on each module with versions that use either the ANSI (so USD-compatible) allocators or the UE allocators (ModuleBoilerplate.h) when appropriate.
+			// In a monolithic build we can't do that, as the primary game module will already define overrides for operator new and delete with
+			// the standard UE allocators: Since we can only have one operator new/delete override on the entire monolithic executable, we can't define our own overrides.
+			// Additionally, the ANSI allocator does not work properly with FMallocPoisonProxy. Consequently, FMallocPoisonProxy has to be disabled.
+			// The only way around it is by forcing the ansi allocator and disabling FMallocPoisonProxy in your project's target file
+			// (YourProject/Source/YourProject.Target.cs) file like this:
+			//
+			//		public class YourProject : TargetRules
+			//		{
+			//			public YourProject(TargetInfo Target) : base(Target)
+			//			{
+			//				...
+			//				GlobalDefinitions.Add("FORCE_ANSI_ALLOCATOR=1");
+			//				GlobalDefinitions.Add("UE_USE_MALLOC_FILL_BYTES=0");
+			//				...
+			//			}
+			//		}
+			//
+			// This will force the entire built executable to use the ANSI C allocators for everything (by disabling the UE overrides in ModuleBoilerplate.h) while
+			// FMallocPoisonProxy is disabled, and so UE and USD allocations will be compatible.
+			// Note that by that point everything will be using the USD-compatible ANSI allocators anyway, so our overrides in USDMemory.h are also disabled, as they're unnecessary.
+			// Also note that we're forced to use dynamic linking for monolithic targets mainly because static linking the USD libraries disables support for user USD plugins, and secondly
+			// because those static libraries would need to be linked with the --whole-archive argument, and there is currently no standard way of doing that in UE.
+			if (bEnableUsdSdk && Target.LinkType == TargetLinkType.Monolithic && !Target.GlobalDefinitions.Contains("FORCE_ANSI_ALLOCATOR=1") && !Target.GlobalDefinitions.Contains("UE_USE_MALLOC_FILL_BYTES=0"))
+			{
+				PublicDefinitions.Add("USD_FORCE_DISABLED=1");
+				bEnableUsdSdk = false;
+			}
+			else
+			{
+				PublicDefinitions.Add("USD_FORCE_DISABLED=0");
+			}
+
+			return bEnableUsdSdk;
 		}
 	}
 }
diff --git a/Engine/Plugins/Experimental/ChaosFlesh/Source/ChaosFleshEngine/Private/ChaosCache/FleshComponentCacheAdapter.cpp b/Engine/Plugins/Experimental/ChaosFlesh/Source/ChaosFleshEngine/Private/ChaosCache/FleshComponentCacheAdapter.cpp
--- a/Engine/Plugins/Experimental/ChaosFlesh/Source/ChaosFleshEngine/Private/ChaosCache/FleshComponentCacheAdapter.cpp
+++ b/Engine/Plugins/Experimental/ChaosFlesh/Source/ChaosFleshEngine/Private/ChaosCache/FleshComponentCacheAdapter.cpp
@@ -2,15 +2,50 @@
 
 #include "ChaosCache/FleshComponentCacheAdapter.h"
 
-
 #include "Chaos/ChaosCache.h"
+#include "Chaos/CacheManagerActor.h"
 #include "Chaos/PBDEvolution.h"
 #include "ChaosFlesh/ChaosDeformableSolverComponent.h"
 #include "ChaosFlesh/FleshComponent.h"
 
+#if USE_USD_SDK && DO_USD_CACHING
+
+#include "HAL/PlatformFile.h"
+#include "HAL/PlatformFileManager.h"
+#include "Misc/Paths.h"
+
+#include "SimCacheUSDOperations.h"
+#include "USDConversionUtils.h" // for GetPrimPathForObject()
+
+#endif // USE_USD_SDK && DO_USD_CACHING
 
 namespace Chaos
 {
+	FFleshCacheAdapterCVarParams CVarParams;
+
+#if USE_USD_SDK && DO_USD_CACHING
+
+	FAutoConsoleVariableRef CVarDeformableFleshCacheWriteBinaryUSD(
+		TEXT("p.Chaos.Caching.USD.WriteBinary"),
+		CVarParams.bWriteBinary,
+		TEXT("Write binary (usdc) cache files. [def: true]"));
+	FAutoConsoleVariableRef CVarDeformableFleshCacheNoClobberUSD(
+		TEXT("p.Chaos.Caching.USD.NoClobber"),
+		CVarParams.bNoClobber,
+		TEXT("Rename rather than over write existing cach files. [def: true]"));
+
+#endif // USE_USD_SDK && DO_USD_CACHING
+
+	FFleshCacheAdapter::~FFleshCacheAdapter()
+	{
+#if USE_USD_SDK && DO_USD_CACHING
+		if (MonolithStage)
+		{
+			UE::SimCacheUSD::CloseStage(MonolithStage);
+		}
+#endif // USE_USD_SDK && DO_USD_CACHING
+	}
+
 	FComponentCacheAdapter::SupportType FFleshCacheAdapter::SupportsComponentClass(UClass* InComponentClass) const
 	{
 		const UClass* Desired = GetDesiredClass();
@@ -49,6 +84,36 @@ namespace Chaos
 				const uint32 NumParticles = Particles.Size();
 				if (NumParticles > 0)
 				{
+#if USE_USD_SDK && DO_USD_CACHING
+					FScopedUsdAllocs UsdAllocs; // Use USD memory allocator
+
+					// Update time range.
+					MinTime = FMath::Min(InTime, MinTime);
+					MaxTime = FMath::Max(InTime, MaxTime);
+
+					pxr::VtArray<pxr::GfVec3f> Points = UE::SimCacheUSD::ToVtVec3Array(Particles.XArray());
+					pxr::VtArray<pxr::GfVec3f> Vels   = UE::SimCacheUSD::ToVtVec3Array(Particles.GetV());
+					
+					// Try to not fill the disk with redundant data.  Finalize() does last save 
+					// to set the final frame range.
+					if (UE::SimCacheUSD::ValuesDiffer(Points, PrevPoints, 0.5) ||
+						UE::SimCacheUSD::ValuesDiffer(Vels, PrevVels, 0.5))
+					{
+						PrevPoints = Points;
+						PrevVels = Vels;
+
+						if (MonolithStage)
+						{
+							if (!UE::SimCacheUSD::WritePoints(MonolithStage, PrimPath, InTime, Points, Vels))
+							{
+								// TODO: Warn
+								return;
+							}
+							// Save every n frames?
+							UE::SimCacheUSD::SaveStage(MonolithStage, MinTime, MaxTime);
+						}
+					}
+#else // USE_USD_SDK && DO_USD_CACHING
 					const Softs::FSolverVec3* ParticleXs = &Particles.X(0);
 					const Softs::FSolverVec3* ParticleVs = &Particles.V(0);
 
@@ -85,6 +150,8 @@ namespace Chaos
 					OutFrame.PendingChannelsData.Add(PositionXName, PendingPX);
 					OutFrame.PendingChannelsData.Add(PositionYName, PendingPY);
 					OutFrame.PendingChannelsData.Add(PositionZName, PendingPZ);
+
+#endif // USE_USD_SDK && DO_USD_CACHING
 				}
 			}
 		}
@@ -98,6 +165,85 @@ namespace Chaos
 
 			if (FEvolution* Evolution = PhysicsThreadAccess.GetEvolution())
 			{
+#if USE_USD_SDK && DO_USD_CACHING
+				if (MonolithStage)
+				{
+					FScopedUsdAllocs UEAllocs; // Use USD memory allocator
+					pxr::VtArray<pxr::GfVec3f> Points0, Points1;
+					pxr::VtArray<pxr::GfVec3f> Vels0, Vels1;
+
+					double TargetTime = InTime;// TickRecord.GetTime();
+					double Prev = -TNumericLimits<double>::Max();
+					double Next = -TNumericLimits<double>::Max();
+					double PrevV = -TNumericLimits<double>::Max();
+					double NextV = -TNumericLimits<double>::Max();
+					if (!UE::SimCacheUSD::GetBracketingTimeSamples(
+							MonolithStage, PrimPath, UE::SimCacheUSD::GetPointsAttrName(), TargetTime, &Prev, &Next) ||
+						!UE::SimCacheUSD::GetBracketingTimeSamples(
+							MonolithStage, PrimPath, UE::SimCacheUSD::GetVelocityAttrName(), TargetTime, &PrevV, &NextV) ||
+						Prev != PrevV ||
+						Next != NextV)
+					{
+						return;
+					}
+
+					if (!UE::SimCacheUSD::ReadPoints(MonolithStage, PrimPath, Prev, Points0, Vels0) ||
+						Points0.size() != Vels0.size())
+					{
+						return;
+					}
+
+					FParticles& Particles = Evolution->Particles();
+					const int32 NumParticles = Particles.Size();
+					Softs::FSolverVec3* ParticleXs = &Particles.X(0);
+					Softs::FSolverVec3* ParticleVs = &Particles.V(0);
+
+					int32 NumCachedParticles = static_cast<int32>(Points0.size());
+					if (NumCachedParticles > NumParticles)
+					{
+						// Cached particles doesn't match solver particles.  Truncate.
+						NumCachedParticles = NumParticles;
+					}
+
+					// < time range start, > time range end, or exact hit
+					if (FMath::IsNearlyEqual(Prev, Next))
+					{
+						// Directly set the result of the cache into the solver particles
+						for (int32 CachedIndex = 0; CachedIndex < NumCachedParticles; ++CachedIndex)
+						{
+							// Note that VtArray::operator[] is non-const access and will cause trigger 
+							// the copy-on-write memcopy!  VtArray::cdata() avoids that.
+							const pxr::GfVec3f& P0 = Points0.cdata()[CachedIndex];
+							const pxr::GfVec3f& V0 = Vels0.cdata()[CachedIndex];
+							ParticleXs[CachedIndex].Set(P0[0], P0[1], P0[2]);
+							ParticleVs[CachedIndex].Set(V0[0], V0[1], V0[2]);
+						}
+						return;
+					}
+
+					if (!UE::SimCacheUSD::ReadPoints(MonolithStage, PrimPath, Next, Points1, Vels1) ||
+						Points1.size() != Vels1.size() ||
+						Points0.size() != Points1.size())
+					{
+						return;
+					}
+					double Duration = Next - Prev;
+					double Alpha = Duration > UE_SMALL_NUMBER ? (TargetTime - Prev) / Duration : 0.5;
+					for (int32 CachedIndex = 0; CachedIndex < NumCachedParticles; ++CachedIndex)
+					{
+						// Note that VtArray::operator[] is non-const access and will cause trigger 
+						// the copy-on-write memcopy!  VtArray::cdata() avoids that.
+						const pxr::GfVec3f& P0 = Points0.cdata()[CachedIndex];
+						const pxr::GfVec3f& P1 = Points1.cdata()[CachedIndex];
+						const pxr::GfVec3f& V0 = Vels0.cdata()[CachedIndex];
+						const pxr::GfVec3f& V1 = Vels1.cdata()[CachedIndex];
+						pxr::GfVec3f Pos = (1.0 - Alpha) * P0 + Alpha * P1;
+						pxr::GfVec3f Vel = (1.0 - Alpha) * V0 + Alpha * V1;
+						ParticleXs[CachedIndex].Set(Pos[0], Pos[1], Pos[2]);
+						ParticleVs[CachedIndex].Set(Vel[0], Vel[1], Vel[2]);
+					}
+				}
+#else // USE_USD_SDK && DO_USD_CACHING
 				FCacheEvaluationContext Context(TickRecord);
 				Context.bEvaluateTransform = false;
 				Context.bEvaluateCurves = false;
@@ -142,6 +288,8 @@ namespace Chaos
 						}
 					}
 				}
+
+#endif // USE_USD_SDK && DO_USD_CACHING
 			}
 		}
 	}
@@ -157,7 +305,11 @@ namespace Chaos
 	{
 		// If we have a flesh mesh we can play back any cache as long as it has one or more tracks
 		const UFleshComponent* FleshComp = CastChecked<UFleshComponent>(InComponent);
+#if USE_USD_SDK && DO_USD_CACHING
+		return FleshComp != nullptr;
+#else
 		return FleshComp && InCache->ChannelCurveToParticle.Num() > 0;
+#endif // USE_USD_SDK && DO_USD_CACHING
 	}
 
 	Chaos::Softs::FDeformableSolver* FFleshCacheAdapter::GetDeformableSolver(UPrimitiveComponent* InComponent) const
@@ -179,8 +331,6 @@ namespace Chaos
 	{
 		return nullptr;
 	}
-
-
 	
 	Chaos::FPhysicsSolverEvents* FFleshCacheAdapter::BuildEventsSolver(UPrimitiveComponent* InComponent) const
 	{
@@ -193,16 +343,58 @@ namespace Chaos
 		return nullptr;
 	}
 	
-	
 	void FFleshCacheAdapter::SetRestState(UPrimitiveComponent* InComponent, UChaosCache* InCache, const FTransform& InRootTransform, Chaos::FReal InTime) const
 	{
+#if !USE_USD_SDK && DO_USD_CACHING
 		if (!InCache || InCache->GetDuration() == 0.0f)
 		{
 			return;
 		} 
+#endif // USE_USD_SDK && DO_USD_CACHING
 
 		if (UFleshComponent* FleshComp = CastChecked<UFleshComponent>(InComponent))
 		{
+#if USE_USD_SDK && DO_USD_CACHING
+			FleshComp->ResetDynamicCollection();
+			if (UFleshDynamicAsset* DynamicCollection = FleshComp->GetDynamicCollection())
+			{
+				TManagedArray<FVector3f>& DynamicVertex = DynamicCollection->GetPositions();
+				//TManagedArray<FVector3f>& DynamicVertex = DynamicCollection->GetVelocities();
+				const int32 NumDynamicVertex = DynamicVertex.Num();
+
+				if (MonolithStage)
+				{
+					FScopedUsdAllocs UEAllocs; // Use USD memory allocator
+					pxr::VtArray<pxr::GfVec3f> Points;
+					pxr::VtArray<pxr::GfVec3f> Vels;
+
+					if (!UE::SimCacheUSD::ReadPoints(MonolithStage, PrimPath, -TNumericLimits<double>::Max(), Points, Vels))
+					{
+						return;
+					}
+
+					int32 NumCachedParticles = static_cast<int32>(Points.size());
+					if (NumDynamicVertex == NumCachedParticles)
+					{
+						for (int32 CachedIndex = 0; CachedIndex < NumCachedParticles; ++CachedIndex)
+						{
+							// Note that VtArray::operator[] is non-const access and will cause trigger 
+							// the copy-on-write memcopy!  VtArray::cdata() avoids that.
+							const pxr::GfVec3f& P = Points.cdata()[CachedIndex];
+							DynamicVertex[CachedIndex].Set(P[0], P[1], P[2]);
+
+							auto UEVertd = [](Chaos::FVec3 V) { return FVector3d(V.X, V.Y, V.Z); };
+							auto UEVertf = [](FVector3d V) { return FVector3f((float)V.X, (float)V.Y, (float)V.Z); };
+							DynamicVertex[CachedIndex] =
+								UEVertf(FleshComp->GetComponentTransform().InverseTransformPosition(
+									UEVertd(DynamicVertex[CachedIndex])));
+						}
+					}
+				}
+			}
+
+#else // USE_USD_SDK && DO_USD_CACHING
+
 			FPlaybackTickRecord TickRecord;
 			TickRecord.SetLastTime(InTime);
 
@@ -256,10 +448,49 @@ namespace Chaos
 					}
 				}
 			}
+
+#endif // USE_USD_SDK && DO_USD_CACHING
 		}
 	}
 
-	bool FFleshCacheAdapter::InitializeForRecord(UPrimitiveComponent* InComponent, UChaosCache* InCache)
+	FString GetCacheDirectory(const FObservedComponent& InObserved)
+	{
+		FString CacheDir = InObserved.USDCacheDirectory.Path;
+		if (CacheDir.IsEmpty())
+		{
+			FString ProjectDir = FPaths::ProjectContentDir();
+			CacheDir = FPaths::Combine(ProjectDir, FString(TEXT("SimCache")));
+		}
+		return CacheDir;
+	}
+
+	FString GetCacheFileName(const UFleshComponent* FleshComp)
+	{
+		const UObject* CurrObject = FleshComp;
+		const AActor* Actor = nullptr;
+		do {
+			CurrObject = CurrObject->GetOuter();
+			Actor = Cast<AActor>(CurrObject);
+		} while (!Actor);
+		const UObject* ActorParent = Actor ? Actor->GetOuter() : nullptr;
+
+		FString CompName;
+		if (ActorParent)
+		{
+			CompName = FleshComp->GetPathName(ActorParent);
+		}
+		else
+		{
+			FleshComp->GetName(CompName);
+		}
+		if (!CompName.Len())
+		{
+			CompName = FString(TEXT("FleshCache"));
+		}
+		return CompName;
+	}
+
+	bool FFleshCacheAdapter::InitializeForRecord(UPrimitiveComponent* InComponent, FObservedComponent& InObserved)
 	{
 		if( FDeformableSolver* Solver = GetDeformableSolver(InComponent))
 		{
@@ -268,12 +499,87 @@ namespace Chaos
 			if (UFleshComponent* FleshComp = CastChecked<UFleshComponent>(InComponent))
 			{
 				FleshComp->ResetDynamicCollection();
+
+#if USE_USD_SDK && DO_USD_CACHING
+				//
+				// USD caching
+				//
+
+				FString CompName = GetCacheFileName(FleshComp);
+				FString CacheDir = GetCacheDirectory(InObserved);
+				FPlatformFileManager& FileManager = FPlatformFileManager::Get();
+				IPlatformFile& PlatformFile = FileManager.GetPlatformFile();
+				if (!PlatformFile.DirectoryExists(*CacheDir))
+				{
+					if (!PlatformFile.CreateDirectory(*CacheDir))
+					{
+						return false;
+					}
+				}
+
+				const UFleshAsset* RestCollectionAsset = FleshComp->GetRestCollection();
+				const FFleshCollection* RestCollection = RestCollectionAsset ? RestCollectionAsset->GetCollection() : nullptr;
+				if (!RestCollection)
+				{
+					return false;
+				}
+
+				PrimPath = UsdUtils::GetPrimPathForObject(FleshComp);
+				if (bUseMonolith)
+				{
+					bReadOnly = false;
+					FString Ext = CVarParams.bWriteBinary ? FString(TEXT("usd")) : FString(TEXT("usda"));
+					FString FileName = FString::Printf(TEXT("%s.%s"), *CompName, *Ext);
+					FilePath = FPaths::Combine(CacheDir, FileName);
+					if (CVarParams.bNoClobber)
+					{
+						if (PlatformFile.FileExists(*FilePath))
+						{
+							// Rename the file to 'path/to/file.usd' to 'path/to/file_#.usd', where '#' 
+							// is a unique version number.
+							FString UniqueFilePath;
+							int32 i = 1;
+							do {
+								FString UniqueCompName = FString::Printf(TEXT("%s_%d.%s"), *CompName, i++, *Ext);
+								UniqueFilePath = FPaths::Combine(CacheDir, UniqueCompName);
+							} while (PlatformFile.FileExists(*UniqueFilePath));
+
+							if (!PlatformFile.MoveFile(*UniqueFilePath, *FilePath))
+							{
+								return false;
+							}
+						}
+					}
+					else
+					{
+						if (!PlatformFile.DeleteFile(*FilePath))
+						{
+							return false;
+						}
+					}
+
+					if (MonolithStage)
+					{
+						UE::SimCacheUSD::CloseStage(MonolithStage);
+					}
+					if (!UE::SimCacheUSD::NewStage(FilePath, MonolithStage))
+					{
+						return false;
+					}
+					if (!UE::SimCacheUSD::WriteTetMesh(MonolithStage, PrimPath, *RestCollection))
+					{
+						return false;
+					}
+				}
+
+#endif // USE_USD_SDK && DO_USD_CACHING
+
 			}
 		}
 		return true;
 	}
 
-	bool FFleshCacheAdapter::InitializeForPlayback(UPrimitiveComponent* InComponent, UChaosCache* InCache, float InTime)
+	bool FFleshCacheAdapter::InitializeForPlayback(UPrimitiveComponent* InComponent, FObservedComponent& InObserved, float InTime)
 	{
 		EnsureIsInGameThreadContext();
 		
@@ -284,8 +590,70 @@ namespace Chaos
 			if (UFleshComponent* FleshComp = CastChecked<UFleshComponent>(InComponent))
 			{
 				FleshComp->ResetDynamicCollection();
+
+#if USE_USD_SDK && DO_USD_CACHING
+				//
+				// USD caching
+				//
+
+				FString CacheDir = GetCacheDirectory(InObserved);
+				FString CompName = GetCacheFileName(FleshComp);
+				PrimPath = UsdUtils::GetPrimPathForObject(FleshComp);
+				if (bUseMonolith)
+				{
+					bReadOnly = true;
+					FString Ext = CVarParams.bWriteBinary ? FString(TEXT("usd")) : FString(TEXT("usda"));
+					FString FileName = FString::Printf(TEXT("%s.%s"), *CompName, *Ext);
+					FilePath = FPaths::Combine(CacheDir, FileName);
+
+					FPlatformFileManager& FileManager = FPlatformFileManager::Get();
+					IPlatformFile& PlatformFile = FileManager.GetPlatformFile();
+					if (PlatformFile.FileExists(*FilePath))
+					{
+						if (MonolithStage)
+						{
+							UE::SimCacheUSD::CloseStage(MonolithStage);
+						}
+						if (!UE::SimCacheUSD::OpenStage(FilePath, MonolithStage))
+						{
+							return false;
+						}
+					}
+					else
+					{
+						return false;
+					}
+				}
+
+#endif // USE_USD_SDK && DO_USD_CACHING
+
 			}
 		}
 		return true;
 	}
+
+	void FFleshCacheAdapter::Finalize()
+	{
+#if USE_USD_SDK && DO_USD_CACHING
+		// Detach shared memory arrays.
+		PrevPoints.clear();
+		PrevVels.clear();
+		
+		if (MonolithStage)
+		{
+			if (!bReadOnly)
+			{
+				UE::SimCacheUSD::SaveStage(MonolithStage, MinTime, MaxTime);
+			}
+			UE::SimCacheUSD::CloseStage(MonolithStage);
+			MonolithStage = UE::FUsdStage();
+		}
+
+		PrimPath.Empty();
+		FilePath.Empty();
+		MinTime = TNumericLimits<double>::Max();
+		MaxTime = -TNumericLimits<double>::Max();
+#endif // USE_USD_SDK && DO_USD_CACHING
+	}
+
 }    // namespace Chaos
diff --git a/Engine/Plugins/Experimental/ChaosFlesh/Source/ChaosFleshEngine/Public/ChaosCache/FleshComponentCacheAdapter.h b/Engine/Plugins/Experimental/ChaosFlesh/Source/ChaosFleshEngine/Public/ChaosCache/FleshComponentCacheAdapter.h
--- a/Engine/Plugins/Experimental/ChaosFlesh/Source/ChaosFleshEngine/Public/ChaosCache/FleshComponentCacheAdapter.h
+++ b/Engine/Plugins/Experimental/ChaosFlesh/Source/ChaosFleshEngine/Public/ChaosCache/FleshComponentCacheAdapter.h
@@ -7,6 +7,16 @@
 #include "Chaos/Deformable/ChaosDeformableSolverProxy.h"
 #include "Chaos/Deformable/ChaosDeformableSolverTypes.h"
 
+#if USE_USD_SDK && DO_USD_CACHING
+
+#include "UsdWrappers/UsdStage.h"
+#include "USDIncludesStart.h"
+	#include "pxr/base/gf/vec3f.h"
+	#include "pxr/base/vt/array.h"
+#include "USDIncludesEnd.h"
+
+#endif // USE_USD_SDK && DO_USD_CACHING
+
 namespace Chaos::Softs
 {
 	class FDeformableSolver;
@@ -16,6 +26,12 @@ namespace Chaos::Softs
 
 namespace Chaos
 {
+	struct FFleshCacheAdapterCVarParams
+	{
+		// USD caching options
+		bool bWriteBinary = true;
+		bool bNoClobber = true;
+	};
 
 	/**
 	 * Skeletal mesh cache adapter to be able to cache cloth simulation datas through the chaos cache system
@@ -27,7 +43,7 @@ namespace Chaos
 		typedef Chaos::Softs::FPBDEvolution FEvolution;
 		typedef Chaos::Softs::FSolverParticles FParticles;
 
-		virtual ~FFleshCacheAdapter() = default;
+		virtual ~FFleshCacheAdapter();
 
 		// ~Begin FComponentCacheAdapter interface
 		virtual SupportType					SupportsComponentClass(UClass* InComponentClass) const override;
@@ -38,21 +54,43 @@ namespace Chaos
 		virtual Chaos::FPhysicsSolver*		GetComponentSolver(UPrimitiveComponent* InComponent) const override;
 		virtual Chaos::FPhysicsSolverEvents* BuildEventsSolver(UPrimitiveComponent* InComponent) const override;
 		virtual void						SetRestState(UPrimitiveComponent* InComponent, UChaosCache* InCache, const FTransform& InRootTransform, Chaos::FReal InTime) const override;
-		virtual bool						InitializeForRecord(UPrimitiveComponent* InComponent, UChaosCache* InCache) override;
-		virtual bool						InitializeForPlayback(UPrimitiveComponent* InComponent, UChaosCache* InCache, float InTime) override;
+		virtual bool						InitializeForRecord(UPrimitiveComponent* InComponent, FObservedComponent& InObserved) override;
+		virtual bool						InitializeForPlayback(UPrimitiveComponent* InComponent, FObservedComponent& InObserved, float InTime) override;
 		virtual void						Record_PostSolve(UPrimitiveComponent* InComp, const FTransform& InRootTransform, FPendingFrameWrite& OutFrame, Chaos::FReal InTime) const override;
 		virtual void						Playback_PreSolve(UPrimitiveComponent* InComponent,
 												UChaosCache*										InCache,
 												Chaos::FReal                                       InTime,
 												FPlaybackTickRecord&								TickRecord,
 												TArray<TPBDRigidParticleHandle<Chaos::FReal, 3>*>& OutUpdatedRigids) const override;
+		virtual void						Finalize() override;
 		// ~End FComponentCacheAdapter interface
 
 
 		FDeformableSolver* GetDeformableSolver(UPrimitiveComponent* InComponent) const;
 
 
-	private :
+	private:
+#if USE_USD_SDK && DO_USD_CACHING
+
+		bool bReadOnly = false;
+		bool bUseMonolith = true; // move to cvarparams when value clips is an option
+
+		/**
+		 * Cache files are stored in a "SimCache" directory at the root of the project.
+		 * File name is derived from flesh component name.
+		 */
+		FString FilePath;
+		FString PrimPath;
+
+		mutable double MinTime = TNumericLimits<double>::Max();
+		mutable double MaxTime = -TNumericLimits<double>::Max();
+
+		mutable pxr::VtArray<pxr::GfVec3f> PrevPoints;
+		mutable pxr::VtArray<pxr::GfVec3f> PrevVels;
+
+		mutable UE::FUsdStage MonolithStage;
+
+#else // USE_USD_SDK && DO_USD_CACHING
 
 		inline static const FName VelocityXName = TEXT("VelocityX");
 		inline static const FName VelocityYName = TEXT("VelocityY");
@@ -60,5 +98,7 @@ namespace Chaos
 		inline static const FName PositionXName = TEXT("PositionX");
 		inline static const FName PositionYName = TEXT("PositionY");
 		inline static const FName PositionZName = TEXT("PositionZ");
+
+#endif // USE_USD_SDK && DO_USD_CACHING
 	};
 }    // namespace Chaos
diff --git a/Engine/Plugins/Importers/USDImporter/Source/SimCacheUSD/Private/SimCacheUSDModule.cpp b/Engine/Plugins/Importers/USDImporter/Source/SimCacheUSD/Private/SimCacheUSDModule.cpp
new file mode 100644
--- /dev/null
+++ b/Engine/Plugins/Importers/USDImporter/Source/SimCacheUSD/Private/SimCacheUSDModule.cpp
@@ -0,0 +1,39 @@
+// Copyright Epic Games, Inc. All Rights Reserved.
+
+#include "Modules/ModuleInterface.h"
+#include "Modules/ModuleManager.h"
+#include "Interfaces/IPluginManager.h"
+#include "Misc/Paths.h"
+
+#include "UnrealUSDWrapper.h"
+#include "USDMemory.h"
+
+
+class FSimCacheUSDModule : public IModuleInterface
+{
+	//~ Begin IModuleInterface API
+	virtual void StartupModule() override;
+	//~ End IModuleInterface API
+
+};
+
+void FSimCacheUSDModule::StartupModule()
+{
+#if USE_USD_SDK
+	// Register the SimCacheUSD plugin with USD.
+	IPluginManager& UEPluginManager = IPluginManager::Get();
+	FString USDImporterDir = UEPluginManager.FindPlugin(TEXT("USDImporter"))->GetBaseDir();
+
+	const FString SimCacheUSDResourcesDir =
+		FPaths::ConvertRelativePathToFull(
+			FPaths::Combine(
+				USDImporterDir, 
+				FString(TEXT("SimCacheUSD")),
+				FString(TEXT("Resources"))));
+
+	UnrealUSDWrapper::RegisterPlugins(SimCacheUSDResourcesDir);
+#endif // USE_USD_SDK
+}
+
+
+IMPLEMENT_MODULE_USD(FSimCacheUSDModule, SimCacheUSD);
diff --git a/Engine/Plugins/Importers/USDImporter/Source/SimCacheUSD/Private/SimCacheUSDOperations.cpp b/Engine/Plugins/Importers/USDImporter/Source/SimCacheUSD/Private/SimCacheUSDOperations.cpp
new file mode 100644
--- /dev/null
+++ b/Engine/Plugins/Importers/USDImporter/Source/SimCacheUSD/Private/SimCacheUSDOperations.cpp
@@ -0,0 +1,677 @@
+// Copyright Epic Games, Inc. All Rights Reserved.
+
+#include "SimCacheUSDOperations.h"
+#include "UEUsdGeomTetMesh.h"
+#include "SimCacheUSDUtil.h"
+
+#if USE_USD_SDK
+
+#include "Containers/StringConv.h"
+#include "GeometryCollection/GeometryCollection.h"
+#include "GeometryCollection/ManagedArray.h"
+#include "GeometryCollection/ManagedArrayCollection.h"
+#include "Misc/Paths.h"
+
+#include "UnrealUSDWrapper.h"
+#include "USDConversionUtils.h"
+#include "USDLog.h"
+#include "USDMemory.h"
+
+#include "UsdWrappers/SdfLayer.h"
+#include "UsdWrappers/SdfPath.h"
+#include "UsdWrappers/UsdAttribute.h"
+#include "UsdWrappers/UsdPrim.h"
+#include "UsdWrappers/UsdStage.h"
+
+#include "USDIncludesStart.h"
+	#include "pxr/usd/sdf/primSpec.h"
+	#include "pxr/usd/usd/stageCacheContext.h"
+	#include "pxr/usd/usdGeom/mesh.h"
+	#include "pxr/usd/usdGeom/pointBased.h"
+	#include "pxr/usd/usdGeom/tokens.h"
+	#include "pxr/usd/usdUtils/stageCache.h"
+#include "USDIncludesEnd.h"
+
+bool 
+UE::SimCacheUSD::NewStage(const FString& StageName, UE::FUsdStage& UsdStage)
+{
+	if (FPaths::FileExists(StageName))
+	{
+		UE_LOG(LogUsd, Error, TEXT("Cannot clobber existing file: '%s'"), *StageName);
+		return false;
+	}
+
+	// USD caches all stages you open/create, unless you tell it not to.
+	// For convenience, we're using the global/default stage cache.
+	UsdStage = UnrealUSDWrapper::NewStage(*StageName);
+	if (!UsdStage)
+	{
+		UE_LOG(LogUsd, Error, TEXT("Failed to create new USD stage: '%s'"), *StageName);
+		return false;
+	}
+
+	// TODO: Ryan
+	//UsdUtils::SetUsdStageMetersPerUnit(UsdStage, ...);
+	UsdUtils::SetUsdStageUpAxis(UsdStage, pxr::UsdGeomTokens->z);
+	return true;
+}
+
+bool 
+UE::SimCacheUSD::OpenStage(const FString& StageName, UE::FUsdStage& UsdStage)
+{
+	if (!FPaths::FileExists(StageName))
+	{
+		UE_LOG(LogUsd, Error, TEXT("File not found: '%s'"), *StageName);
+		return false;
+	}
+
+	// USD caches all stages you open/create, unless you tell it not to.
+	// For convenience, we're using the global/default stage cache.
+	UsdStage = UnrealUSDWrapper::OpenStage(
+			*StageName,						
+			EUsdInitialLoadSet::LoadAll,	// load payload?
+			true,							// use stage cache
+			false);							// force reload from disk
+	if (!UsdStage)
+	{
+		UE_LOG(LogUsd, Error, TEXT("Failed to open USD stage: '%s'"), *StageName);
+		return false;
+	}
+	return true;
+}
+
+bool 
+UE::SimCacheUSD::SaveStage(UE::FUsdStage& Stage, const double FirstFrame, const double LastFrame)
+{
+	if (FirstFrame != -TNumericLimits<double>::Max() && LastFrame != -TNumericLimits<double>::Max())
+	{
+		const double StageStart = Stage.GetStartTimeCode();
+		const double StageEnd = Stage.GetEndTimeCode();
+		const bool bHasRange = static_cast<pxr::UsdStageRefPtr>(Stage)->HasAuthoredTimeCodeRange();
+		if (!bHasRange || FirstFrame < StageStart)
+		{
+			Stage.SetStartTimeCode(FirstFrame);
+		}
+		if (!bHasRange || LastFrame > StageEnd)
+		{
+			Stage.SetEndTimeCode(LastFrame);
+		}
+	}
+	return Stage.GetRootLayer().Save();
+}
+
+bool
+UE::SimCacheUSD::CloseStage(const UE::FUsdStage& Stage)
+{
+	UnrealUSDWrapper::EraseStageFromCache(Stage);
+	return true;
+}
+
+bool 
+UE::SimCacheUSD::CloseStage(const FString& StageName)
+{
+	TArray<UE::FUsdStage> CachedStages = UnrealUSDWrapper::GetAllStagesFromCache();
+	for (UE::FUsdStage& Stage : CachedStages)
+	{
+		if (Stage.GetRootLayer().GetRealPath() == StageName)
+		{
+			return SimCacheUSD::CloseStage(Stage);
+		}
+	}
+	return true;
+}
+
+void
+UE::SimCacheUSD::GenerateValueClipStageNames(
+	const FString& ParentName, 
+	FString& TopologyName, 
+	FString& TimeVaryingTemplate)
+{
+	FScopedUnrealAllocs UEAllocs; // Use UE memory allocator
+	std::string ParentNameStr(TCHAR_TO_ANSI(*ParentName));
+	std::string TopologyNameStr = pxr::UsdUtilsGenerateClipTopologyName(ParentNameStr);
+	TopologyName = FString(ANSI_TO_TCHAR(TopologyNameStr.c_str()));
+	// Usd supports an integral template format, and a floating point.  We're using the
+	// float format here.  The integral format is "path/name.####.usd".
+	TimeVaryingTemplate = TopologyName.Replace(TEXT("topology"), TEXT("###.###"));
+}
+
+FString
+UE::SimCacheUSD::GenerateValueClipTimeVaryingStageName(
+	const FString& TimeVaryingTemplate, 
+	const double Time)
+{
+	FScopedUnrealAllocs UEAllocs; // Use UE memory allocator
+	if (Time == -TNumericLimits<double>::Max())
+	{
+		return TimeVaryingTemplate.Replace(TEXT("###.###"), TEXT("default"));
+	}
+	FString Template = TimeVaryingTemplate.Replace(TEXT("###.###"), TEXT("{0}"));
+	FString TimeStr = FString::Printf(TEXT("%06.3f"), Time);
+	return FString::Format(*Template, { TimeStr });
+}
+
+bool
+UE::SimCacheUSD::NewValueClipsStages(
+	const FString& ParentStageName,
+	const FString& TopologyStageName,
+	UE::FUsdStage& ParentStage,
+	UE::FUsdStage& TopologyStage)
+{
+	if (!SimCacheUSD::NewStage(ParentStageName, ParentStage)) return false;
+	if (!SimCacheUSD::NewStage(TopologyStageName, TopologyStage)) return false;
+	return true;
+}
+
+bool
+UE::SimCacheUSD::NewValueClipsFrameStage(
+	const FString& TimeVaryingStageTemplate,
+	const double Time,
+	FString& FrameStageName,
+	UE::FUsdStage& FrameStage)
+{
+	FrameStageName =
+		SimCacheUSD::GenerateValueClipTimeVaryingStageName(TimeVaryingStageTemplate, Time);
+	return SimCacheUSD::NewStage(FrameStageName, FrameStage);
+}
+
+bool 
+UE::SimCacheUSD::InitValueClipsTemplate(
+	UE::FUsdStage& ParentStage,
+	UE::FUsdStage& TopologyStage,
+	const FString& ParentStageName,
+	const FString& TopologyStageName,
+	const FString& TimeVaryingStageTemplate,
+	const TArray<FString>& PrimPaths,
+	const double StartTime,
+	const double EndTime,
+	const double Stride)
+{
+	FScopedUsdAllocs UEAllocs; // Use USD memory allocator
+
+	UE::FSdfLayer ParentLayer = ParentStage.GetRootLayer();
+	UE::FSdfLayer TopologyLayer = ParentStage.GetRootLayer();
+
+	for (const FString& PrimPath : PrimPaths)
+	{
+		// Get or create the prim.
+		pxr::SdfPrimSpecHandle Prim = pxr::SdfCreatePrimInLayer(ParentLayer, UE::FSdfPath(*PrimPath));
+
+		// Get the existing metadata, if it exists.
+		pxr::VtDictionary Clips;
+		pxr::VtDictionary ClipSetDict;
+		if (Prim->HasInfo(pxr::UsdTokens->clips))
+		{
+			pxr::VtValue PrevClipInfoValue = Prim->GetInfo(pxr::UsdTokens->clips);
+			if (PrevClipInfoValue.IsHolding<pxr::VtDictionary>())
+			{
+				Clips = PrevClipInfoValue.Get<pxr::VtDictionary>();
+				auto it = Clips.find(pxr::UsdClipsAPISetNames->default_);
+				if (it != Clips.end())
+				{
+					pxr::VtValue PrevClipsEntryValue = it->second;
+					if (PrevClipsEntryValue.IsHolding<pxr::VtDictionary>())
+					{
+						ClipSetDict = PrevClipsEntryValue.Get<pxr::VtDictionary>();
+					}
+				}
+			}
+		}
+
+		// Add the topology layer.
+		if (!ParentStageName.IsEmpty() && !TopologyStageName.IsEmpty() && TopologyLayer)
+		{
+			FString TopologyId = TopologyStageName;
+			FPaths::MakePathRelativeTo(TopologyId, *ParentStageName);
+			std::string TopologyIdStr(TCHAR_TO_ANSI(*TopologyId));
+
+			// Add the sub layer in the parent pointing to the topology.
+			static_cast<pxr::SdfLayerRefPtr>(ParentLayer)->InsertSubLayerPath(
+				TopologyIdStr, 0); // 0 means strongest
+			
+			// Set root layer metadata
+			ClipSetDict[pxr::UsdClipsAPIInfoKeys->manifestAssetPath] = 
+				pxr::SdfAssetPath(TopologyIdStr);
+		}
+
+		// Add time varying stage template.
+		if (!TimeVaryingStageTemplate.IsEmpty())
+		{
+			std::string TimeVaryingStageTemplateStr(TCHAR_TO_ANSI(*TimeVaryingStageTemplate));
+			ClipSetDict[pxr::UsdClipsAPIInfoKeys->templateAssetPath] = 
+				TimeVaryingStageTemplateStr;
+		}
+
+		// Add the prim path.
+		std::string PrimPathStr(TCHAR_TO_ANSI(*PrimPath));
+		ClipSetDict[pxr::UsdClipsAPIInfoKeys->primPath] = PrimPathStr;
+
+		// Add start and end time.
+		if (fabs(StartTime) != TNumericLimits<double>::Max() &&
+			fabs(EndTime) != TNumericLimits<double>::Max())
+		{
+			ClipSetDict[pxr::UsdClipsAPIInfoKeys->templateStartTime] = StartTime;
+			ClipSetDict[pxr::UsdClipsAPIInfoKeys->templateEndTime] = EndTime;
+		}
+
+		// Add sampling stride.
+		if (Stride > 0.0)
+		{
+			ClipSetDict[pxr::UsdClipsAPIInfoKeys->templateStride] = Stride;
+		}
+
+		// Set metadata info on Prim.
+		Clips[pxr::UsdClipsAPISetNames->default_] = ClipSetDict;
+		Prim->SetInfo(pxr::UsdTokens->clips, pxr::VtValue::Take(Clips));
+	}
+	ParentLayer.SetStartTimeCode(StartTime);
+	ParentLayer.SetEndTimeCode(EndTime);
+	return true;
+}
+
+bool
+UE::SimCacheUSD::InitValueClipsTemplate(
+	const FString& ParentStageName,
+	const FString& TopologyStageName,
+	const FString& TimeVaryingStageTemplate,
+	const TArray<FString>& PrimPaths,
+	const double StartTime,
+	const double EndTime,
+	const double Stride)
+{
+	UE::FUsdStage ParentStage;
+	if (!SimCacheUSD::OpenStage(ParentStageName, ParentStage)) return false;
+	UE::FUsdStage TopologyStage;
+	if (!SimCacheUSD::OpenStage(TopologyStageName, TopologyStage)) return false;
+
+	return SimCacheUSD::InitValueClipsTemplate(
+		ParentStage,
+		TopologyStage,
+		ParentStageName,
+		TopologyStageName,
+		TimeVaryingStageTemplate,
+		PrimPaths,
+		StartTime,
+		EndTime,
+		Stride);
+}
+
+bool 
+UE::SimCacheUSD::WriteTetMesh(
+	UE::FUsdStage& Stage, 
+	const FString& PrimPath,
+	const FManagedArrayCollection& Collection, 
+	const int32 StructureIndex)
+{
+	FScopedUsdAllocs UEAllocs; // Use USD memory allocator
+
+	const TManagedArray<FIntVector4>* Tetrahedron =
+		Collection.FindAttribute<FIntVector4>(
+			/*FTetrahedralCollection::TetrahedronAttribute*/FName("Tetrahedron"), 
+			/*FTetrahedralCollection::TetrahedralGroup*/FName("Tetrahedral"));
+	const TManagedArray<int32>* TetrahedronStart =
+		Collection.FindAttribute<int32>(
+			/*FTetrahedralCollection::TetrahedronStartAttribute*/FName("TetrahedronStart"), 
+			/*FGeometryCollection::GeometryGroup*/FName("Tetrahedral"));
+	const TManagedArray<int32>* TetrahedronCount =
+		Collection.FindAttribute<int32>(
+			/*FTetrahedralCollection::TetrahedronCountAttribute*/FName("TetrahedronCount"), 
+			/*FGeometryCollection::GeometryGroup*/FName("Geometry"));
+
+	const TManagedArray<FIntVector>* Triangle =
+		Collection.FindAttribute<FIntVector>(
+			"Indices", /*FGeometryCollection::FacesGroup*/FName("Faces"));
+	const TManagedArray<int32>* FacesStart =
+		Collection.FindAttribute<int32>(
+			"FaceStart", /*FGeometryCollection::GeometryGroup*/FName("Geometry"));
+	const TManagedArray<int32>* FacesCount =
+		Collection.FindAttribute<int32>(
+			"FaceCount", /*FGeometryCollection::GeometryGroup*/FName("Geometry"));
+
+	const TManagedArray<int32>* VertexStart =
+		Collection.FindAttribute<int32>(
+			"VertexStart", /*FGeometryCollection::GeometryGroup*/FName("Geometry"));
+	const TManagedArray<int32>* VertexCount =
+		Collection.FindAttribute<int32>(
+			"VertexCount", /*FGeometryCollection::GeometryGroup*/FName("Geometry"));
+
+	const TManagedArray<FVector3f>* Vertex =
+		Collection.FindAttribute<FVector3f>(
+			"Vertex", "Vertices");
+
+	int32 VertStart = 0;
+	int32 VertCount = Vertex->Num();
+	int32 TetMeshStart = 0;
+	int32 TetMeshCount = Tetrahedron->Num();
+	int32 TriMeshStart = 0;
+	int32 TriMeshCount = Triangle->Num();
+	if (StructureIndex != INDEX_NONE)
+	{
+		if (!VertexStart->IsValidIndex(StructureIndex) ||
+			!VertexCount->IsValidIndex(StructureIndex))
+		{
+			UE_LOG(LogUsd, Error, TEXT("Invalid vertex index: %d, valid range [0, %d)."), StructureIndex, VertexStart->Num());
+			return false;
+		}
+		VertStart = (*VertexStart)[StructureIndex];
+		VertCount = (*VertexCount)[StructureIndex];
+		if (!TetrahedronStart->IsValidIndex(StructureIndex) ||
+			!TetrahedronCount->IsValidIndex(StructureIndex))
+		{
+			UE_LOG(LogUsd, Error, TEXT("Invalid tet mesh index: %d, valid range [0, %d)."), StructureIndex, TetrahedronStart->Num());
+			return false;
+		}
+		TetMeshStart = (*TetrahedronStart)[StructureIndex];
+		TetMeshCount = (*TetrahedronCount)[StructureIndex];
+		if (!FacesStart->IsValidIndex(StructureIndex) ||
+			!FacesCount->IsValidIndex(StructureIndex))
+		{
+			// Not a deal breaker, but highly unusual.
+			UE_LOG(LogUsd, Warning, TEXT("Invalid tet mesh faces index: %d, valid range [0, %d),"), StructureIndex, FacesStart->Num());
+		}
+		else
+		{
+			TriMeshStart = (*FacesStart)[StructureIndex];
+			TriMeshCount = (*FacesCount)[StructureIndex];
+		}
+
+		UE_LOG(LogUsd, Log, TEXT("Extracting tetrahedral mesh %d of %d from collection; vertex range [%d, %d), tet [%d, %d), tri [%d, %d)."),
+			StructureIndex, VertexStart->Num(),
+			VertStart, VertStart + VertCount,
+			TetMeshStart, TetMeshStart + TetMeshCount,
+			TriMeshStart, TriMeshStart + TriMeshCount);
+	}
+
+	UE::FSdfPath Path(*PrimPath);
+
+	// Author a UsdGeomPointBased override prim.
+	if (!UE::SimCacheUSD::WritePoints(Stage, PrimPath, -TNumericLimits<double>::Max(), Collection, StructureIndex))
+	{
+		UE_LOG(LogUsd, Error, TEXT("Failed to create UsdGeomPointBased at path: '%s' on USD stage: '%s'."),
+			*PrimPath, *Stage.GetRootLayer().GetRealPath());
+		return false;
+	}
+
+	// Promote to UsdGeomMesh.
+	if (pxr::UsdGeomMesh Mesh = pxr::UsdGeomMesh::Define(Stage, Path))
+	{
+		pxr::UsdPrim Prim = Mesh.GetPrim();
+
+		// Surface tris vertices and counts
+		pxr::VtArray<int> VtTriVerts(TriMeshCount * 3);
+		int32 k = 0;
+		for (int32 i = 0; i < TriMeshCount; i++)
+			for (int32 j = 0; j < 3; j++, k++)
+				VtTriVerts[k] = (*Triangle)[TriMeshStart + i][j] - VertStart;
+		Mesh.CreateFaceVertexIndicesAttr().Set(VtTriVerts);
+		pxr::VtArray<int> VtTriCounts(TriMeshCount);
+		for (int32 i = 0; i < TriMeshCount; i++)
+			VtTriCounts[i] = 3; // (This makes me die inside)
+		Mesh.CreateFaceVertexCountsAttr().Set(VtTriCounts);
+
+		// Subdivision scheme set to 'none' in UEUsdGeomTetMesh constructor, but we do it here anyway.
+		Mesh.CreateSubdivisionSchemeAttr().Set(pxr::UsdGeomTokens->none);
+	}
+	else
+	{
+		UE_LOG(LogUsd, Warning, TEXT("Failed to create UsdGeomMesh at path: '%s' on USD stage: '%s'."),
+			*PrimPath, *Stage.GetRootLayer().GetRealPath());
+	}
+
+	// Promote to UEUsdGeomTetMesh.
+	// This next line succeeds, but farts out a:
+	//   TF_DIAGNOSTIC_CODING_ERROR_TYPE: Failed to find plugin for schema type 'UEUsdGeomTetMesh'
+	// I think that means there's an issue with the plugInfo.json, and I think that means that if
+	// USD was to happen upon a prim of this type, it might not know which plugin to load.  For 
+	// our caching purposes, I think that's of no consequence; other than the fact that operator bool 
+	// returns false for TetMesh, which I think may be related, but may not.  PointBased's operator 
+	// bool does the same thing.  What's more is that we currently have no consumer for the tet 
+	// topology, so this really doesn't matter.  At least not yet.
+	pxr::UEUsdGeomTetMesh TetMesh = pxr::UEUsdGeomTetMesh::Define(Stage, Path);
+	if (/*TetMesh*/true)
+	{
+		pxr::UsdPrim Prim = TetMesh.GetPrim();
+
+		// Tetrahedra
+		check(sizeof(pxr::GfVec4i) == sizeof(FIntVector4));
+		const pxr::GfVec4i* GfTets = reinterpret_cast<const pxr::GfVec4i*>(&(*Tetrahedron)[TetMeshStart]);
+		pxr::VtArray<pxr::GfVec4i> VtTets(TetMeshCount);
+		for (int32 i = 0; i < TetMeshCount; i++) 
+			VtTets[i] = GfTets[i] - pxr::GfVec4i(VertStart);
+		TetMesh.CreateTetVertexIndicesAttr().Set(VtTets);
+	}
+	else
+	{
+		UE_LOG(LogUsd, Warning, TEXT("Failed to create UEUsdGeomTetMesh at path: '%s' on USD stage: '%s'."),
+			*PrimPath, *Stage.GetRootLayer().GetRealPath());
+	}
+
+	UE::SimCacheUSD::SaveStage(Stage, -TNumericLimits<double>::Max(), -TNumericLimits<double>::Max());
+	return true;
+}
+
+bool
+UE::SimCacheUSD::WritePoints(
+	UE::FUsdStage& Stage,
+	const FString& PrimPath,
+	const double Time,
+	pxr::VtArray<pxr::GfVec3f>& VtPoints,
+	pxr::VtArray<pxr::GfVec3f>& VtVels)
+{
+	pxr::GfMatrix4d GfMat(1.0);
+
+	UE::FSdfPath Path(*PrimPath);
+	pxr::UsdPrim Prim = Stage.GetPrimAtPath(Path);
+	//UE::FUsdPrim Prim = PrimPath.IsEmpty() ? Stage.GetDefaultPrim() : Stage.GetPrimAtPath(UE::FSdfPath(*PrimPath));
+	if (!Prim)
+	{
+		// We're authoring a prim that doesn't exist.  Use an override prim.
+		Prim = Stage.OverridePrim(Path);
+		//pxr::UsdGeomPointBased PointBased(Prim);
+		//PointBased.CreatePointsAttr();
+		//PointBased.CreateExtentAttr();
+
+		UE::SimCacheUSDUtil::Private::DefineAncestorTransforms(Stage, Path);
+
+		return UE::SimCacheUSDUtil::Private::SetPointsExtentAndPrimXform(
+			Prim, GfMat, VtPoints, VtVels,
+			Time == -TNumericLimits<double>::Max() ? pxr::UsdTimeCode::Default() : pxr::UsdTimeCode(Time));
+	}
+	else
+	{
+		// We're probably authoring an override prim, UsdGeomPointBased, or a UsdGeomPointInstancer.
+		return UE::SimCacheUSDUtil::Private::SetPointsExtentAndPrimXform(
+			Prim, GfMat, VtPoints, VtVels,
+			Time == -TNumericLimits<double>::Max() ? pxr::UsdTimeCode::Default() : pxr::UsdTimeCode(Time));
+	}
+}
+
+bool 
+UE::SimCacheUSD::WritePoints(
+	UE::FUsdStage& Stage, 
+	const FString& PrimPath, 
+	const double Time, 
+	const FManagedArrayCollection& Collection, 
+	const int32 StructureIndex)
+{
+	FScopedUsdAllocs UEAllocs; // Use USD memory allocator
+
+	const TManagedArray<int32>* VertexStart =
+		Collection.FindAttribute<int32>(
+			"VertexStart", /*FGeometryCollection::GeometryGroup*/ FName("Geometry"));
+	const TManagedArray<int32>* VertexCount =
+		Collection.FindAttribute<int32>(
+			"VertexCount", /*FGeometryCollection::GeometryGroup*/FName("Geometry"));
+
+	const TManagedArray<FVector3f>* Vertex =
+		Collection.FindAttribute<FVector3f>(
+			"Vertex", "Vertices");
+
+	int32 VertStart = 0;
+	int32 VertCount = Vertex->Num();
+	if (StructureIndex != INDEX_NONE)
+	{
+		if (!VertexStart->IsValidIndex(StructureIndex) ||
+			!VertexCount->IsValidIndex(StructureIndex))
+		{
+			UE_LOG(LogUsd, Error, TEXT("Invalid vertex index: %d, valid range [0, %d)."), StructureIndex, VertexStart->Num());
+			return false;
+		}
+		VertStart = (*VertexStart)[StructureIndex];
+		VertCount = (*VertexCount)[StructureIndex];
+	}
+
+	check(sizeof(pxr::GfVec3f) == sizeof(FVector3f)); // reinterpret_cast works?
+	pxr::VtArray<pxr::GfVec3f> VtPoints(VertCount); // VtArray is shared mem with copy-on-write.
+	for (int32 i = 0; i < VertCount; i++)
+		VtPoints[i] = *reinterpret_cast<const pxr::GfVec3f*>(&(*Vertex)[VertStart + i]);
+	pxr::GfMatrix4d GfMat(1.0);
+
+	UE::FSdfPath Path(*PrimPath);
+	pxr::UsdPrim Prim = Stage.GetPrimAtPath(Path);
+	if(!Prim)
+	{
+		// We're authoring a prim that doesn't exist.  Use an override prim.
+		Prim = Stage.OverridePrim(Path);
+		//pxr::UsdGeomPointBased PointBased(Prim);
+		//PointBased.CreatePointsAttr();
+		//PointBased.CreateExtentAttr();
+
+		UE::SimCacheUSDUtil::Private::DefineAncestorTransforms(Stage, Path);
+
+		return UE::SimCacheUSDUtil::Private::SetPointsExtentAndPrimXform(
+			Prim, GfMat, VtPoints,
+			Time == -TNumericLimits<double>::Max() ? pxr::UsdTimeCode::Default() : pxr::UsdTimeCode(Time));
+	}
+	else 
+	{
+		// We're probably authoring an override prim, UsdGeomPointBased, or a UsdGeomPointInstancer.
+		return UE::SimCacheUSDUtil::Private::SetPointsExtentAndPrimXform(
+			Prim, GfMat, VtPoints,
+			Time == -TNumericLimits<double>::Max() ? pxr::UsdTimeCode::Default() : pxr::UsdTimeCode(Time));
+	}
+}
+
+bool 
+UE::SimCacheUSD::ReadTimeSamples(
+	const UE::FUsdStage& Stage, 
+	const FString& PrimPath, 
+	const FString& AttrName, 
+	TArray<double>& TimeSamples)
+{
+	UE::FSdfPath Path(*PrimPath);
+	UE::FUsdPrim Prim = Stage.GetPrimAtPath(Path);
+	if (!Prim)
+	{
+		UE_LOG(LogUsd, Error, TEXT("No prim found at path '%s'."), *PrimPath);
+		return false;
+	}
+	UE::FUsdAttribute Attr = Prim.GetAttribute(*AttrName);
+	if (!Attr)
+	{
+		UE_LOG(LogUsd, Error, TEXT("No attribute '%s' found on prim '%s'."), *AttrName, *PrimPath);
+		return false;
+	}
+	Attr.GetTimeSamples(TimeSamples);
+	return true;
+}
+
+bool 
+UE::SimCacheUSD::ReadTimeSamples(
+	const UE::FUsdStage& Stage, 
+	const FString& PrimPath, 
+	TArray<double>& TimeSamples)
+{
+	return SimCacheUSD::ReadTimeSamples(Stage, PrimPath, FString(pxr::UsdGeomTokens->points.GetText()), TimeSamples);
+}
+
+FString
+UE::SimCacheUSD::GetPointsAttrName()
+{
+	FScopedUnrealAllocs UEAllocs; // Use UE memory allocator
+	return FString(pxr::UsdGeomTokens->points.GetText());
+}
+
+FString
+UE::SimCacheUSD::GetVelocityAttrName()
+{
+	FScopedUnrealAllocs UEAllocs; // Use UE memory allocator
+	return FString(pxr::UsdGeomTokens->velocities.GetText());
+}
+
+bool 
+UE::SimCacheUSD::GetBracketingTimeSamples(
+	const UE::FUsdStage& Stage, 
+	const FString& PrimPath, 
+	const FString& AttrName, 
+	const double TargetTime, 
+	double* Lower, 
+	double* Upper)
+{
+	FScopedUsdAllocs UEAllocs; // Use USD memory allocator
+
+	UE::FSdfPath Path(*PrimPath);
+	UE::FUsdPrim Prim = Stage.GetPrimAtPath(Path);
+	if (!Prim)
+	{
+		UE_LOG(LogUsd, Error, TEXT("No prim found at path '%s'."), *PrimPath);
+		return false;
+	}
+	UE::FUsdAttribute Attr = Prim.GetAttribute(*AttrName);
+	if (!Attr)
+	{
+		UE_LOG(LogUsd, Error, TEXT("No attribute '%s' found on prim '%s'."), *AttrName, *PrimPath);
+		return false;
+	}
+
+	bool bHasTimeSamples = false;
+	return static_cast<pxr::UsdAttribute&>(Attr).GetBracketingTimeSamples(
+		TargetTime, Lower, Upper, &bHasTimeSamples) && bHasTimeSamples;
+}
+
+
+bool
+UE::SimCacheUSD::ReadPoints(
+	const UE::FUsdStage& Stage,
+	const FString& PrimPath,
+	const FString& AttrName,
+	const double Time,
+	pxr::VtArray<pxr::GfVec3f>& VtPoints)
+{
+	FScopedUsdAllocs UEAllocs; // Use USD memory allocator
+	
+	UE::FSdfPath Path(*PrimPath);
+	UE::FUsdPrim Prim = Stage.GetPrimAtPath(Path);
+	if (!Prim)
+	{
+		UE_LOG(LogUsd, Error, TEXT("No prim found at path '%s'."), *PrimPath);
+		return false;
+	}
+	UE::FUsdAttribute Attr = Prim.GetAttribute(*AttrName);
+	if (!Attr)
+	{
+		UE_LOG(LogUsd, Error, TEXT("No attribute '%s' found on prim '%s'."), *AttrName, *PrimPath);
+		return false;
+	}
+
+	// VtArray is copy-on-write. Be careful to use const API.
+	static_cast<pxr::UsdAttribute&>(Attr).Get(
+		&VtPoints, 
+		Time == -TNumericLimits<double>::Max() ? pxr::UsdTimeCode::Default() : pxr::UsdTimeCode(Time));
+
+	return true;
+}
+
+bool
+UE::SimCacheUSD::ReadPoints(
+	const UE::FUsdStage& Stage,
+	const FString& PrimPath,
+	const double Time,
+	pxr::VtArray<pxr::GfVec3f>& VtPoints,
+	pxr::VtArray<pxr::GfVec3f>& VtVels)
+{
+	SimCacheUSD::ReadPoints(Stage, PrimPath, FString(pxr::UsdGeomTokens->velocities.GetText()), Time, VtVels);
+	return SimCacheUSD::ReadPoints(Stage, PrimPath, FString(pxr::UsdGeomTokens->points.GetText()), Time, VtPoints);
+}
+
+#endif // USE_USD_SDK
diff --git a/Engine/Plugins/Importers/USDImporter/Source/SimCacheUSD/Private/SimCacheUSDUtil.cpp b/Engine/Plugins/Importers/USDImporter/Source/SimCacheUSD/Private/SimCacheUSDUtil.cpp
new file mode 100644
--- /dev/null
+++ b/Engine/Plugins/Importers/USDImporter/Source/SimCacheUSD/Private/SimCacheUSDUtil.cpp
@@ -0,0 +1,100 @@
+// Copyright Epic Games, Inc. All Rights Reserved.
+
+#include "SimCacheUSDUtil.h"
+
+#if USE_USD_SDK
+
+#include "USDIncludesStart.h"
+	#include "pxr/usd/usdGeom/pointBased.h"
+	#include "pxr/usd/usdGeom/xform.h"
+#include "USDIncludesEnd.h"
+
+void 
+UE::SimCacheUSDUtil::Private::SaveStage(
+	const pxr::UsdStageRefPtr& Stage,
+	const pxr::UsdTimeCode& FirstFrame,
+	const pxr::UsdTimeCode& LastFrame)
+{
+	FScopedUsdAllocs UEAllocs; // Use USD memory allocator
+	if (FirstFrame.IsNumeric() && LastFrame.IsNumeric())
+	{
+		// Only expand the existing frame range.
+		const double StageStart = Stage->GetStartTimeCode();
+		const double StageEnd = Stage->GetEndTimeCode();
+		const bool HasRange = Stage->HasAuthoredTimeCodeRange();
+		if (!HasRange || FirstFrame.GetValue() < StageStart)
+		{
+			Stage->SetStartTimeCode(FirstFrame.GetValue());
+		}
+		if (!HasRange || LastFrame.GetValue() > StageEnd)
+		{
+			Stage->SetEndTimeCode(LastFrame.GetValue());
+		}
+	}
+	Stage->GetRootLayer()->Save();
+}
+
+void 
+UE::SimCacheUSDUtil::Private::DefineAncestorTransforms(
+	pxr::UsdStageRefPtr& Stage,
+	const pxr::SdfPath& PrimPath)
+{
+	FScopedUsdAllocs UEAllocs; // Use USD memory allocator
+	pxr::SdfPath Path(PrimPath.GetParentPath());
+	//while (Path != pxr::SdfPath::EmptyPath())
+	while (!Path.IsAbsoluteRootPath())
+	{
+		pxr::UsdGeomXform::Define(Stage, Path);
+		Path = Path.GetParentPath();
+	}
+}
+
+void
+UE::SimCacheUSDUtil::Private::ComputeExtent(
+	const pxr::VtArray<pxr::GfVec3f>& Points,
+	pxr::VtVec3fArray& Extent)
+{
+	FScopedUsdAllocs UEAllocs; // Use USD memory allocator
+	pxr::GfRange3f Bounds;
+	for (size_t i = 0; i < Points.size(); i++)
+	{
+		Bounds = Bounds.UnionWith(Points[i]);
+	}
+	Extent[0] = Bounds.GetMin();
+	Extent[1] = Bounds.GetMax();
+}
+
+bool 
+UE::SimCacheUSDUtil::Private::SetPointsExtentAndPrimXform(
+	pxr::UsdPrim& Prim, 
+	const pxr::GfMatrix4d& PrimXf, 
+	const pxr::VtArray<pxr::GfVec3f>& Points, 
+	const pxr::VtArray<pxr::GfVec3f>& Vels, 
+	const pxr::UsdTimeCode& Time)
+{
+	FScopedUsdAllocs UEAllocs; // Use USD memory allocator
+
+	pxr::VtVec3fArray Extent(2);
+	UE::SimCacheUSDUtil::Private::ComputeExtent(Points, Extent);
+
+	pxr::UsdGeomPointBased PointBased(Prim);
+	PointBased.MakeMatrixXform().Set(PrimXf, Time);
+	PointBased.CreateExtentAttr().Set(Extent, Time);
+	if (!PointBased.CreatePointsAttr().Set(Points, Time)) return false;
+	if (!Vels.empty() && !PointBased.CreateVelocitiesAttr().Set(Vels, Time)) return false;
+	return true;
+}
+
+bool 
+UE::SimCacheUSDUtil::Private::SetPointsExtentAndPrimXform(
+	pxr::UsdPrim& Prim,
+	const pxr::GfMatrix4d& PrimXf,
+	const pxr::VtArray<pxr::GfVec3f>& Points,
+	const pxr::UsdTimeCode& Time)
+{
+	FScopedUsdAllocs UEAllocs; // Use USD memory allocator
+	pxr::VtArray<pxr::GfVec3f> Vels;
+	return UE::SimCacheUSDUtil::Private::SetPointsExtentAndPrimXform(Prim, PrimXf, Points, Vels, Time);
+}
+
+#endif // USE_USD_SDK
diff --git a/Engine/Plugins/Importers/USDImporter/Source/SimCacheUSD/Private/SimCacheUSDUtil.h b/Engine/Plugins/Importers/USDImporter/Source/SimCacheUSD/Private/SimCacheUSDUtil.h
new file mode 100644
--- /dev/null
+++ b/Engine/Plugins/Importers/USDImporter/Source/SimCacheUSD/Private/SimCacheUSDUtil.h
@@ -0,0 +1,101 @@
+// Copyright Epic Games, Inc. All Rights Reserved.
+#pragma once
+#if USE_USD_SDK
+
+#include "USDMemory.h"
+#include "USDIncludesStart.h"
+	#include "pxr/usd/usd/usdaFileFormat.h"
+	#include "pxr/usd/usd/usdcFileFormat.h"
+	#include "pxr/usd/usd/usdFileFormat.h"
+
+	#include "pxr/usd/usd/tokens.h"
+	#include "pxr/usd/usdUtils/stitchClips.h"
+
+	#include "pxr/base/tf/fileUtils.h"
+	#include "pxr/base/tf/pathUtils.h"
+
+	#include "pxr/base/vt/array.h"
+	#include "pxr/base/vt/dictionary.h"
+	#include "pxr/base/vt/value.h"
+#include "USDIncludesEnd.h"
+
+namespace UE::SimCacheUSDUtil::Private {
+
+	/**
+	 * Convert from \c pxr::VtArray<T1> to \c TArray<T2>, where \p T1 and \p T2 are binary compatible.
+	 *
+	 * Note that \c VtArray is shared copy-on-write, which makes them trivially copied. This function
+	 * copies memory, which should be avoided.
+	 */
+	template <class T1, class T2>
+	void CastConvert(const pxr::VtArray<T1>& VtArray, TArray<T2>& UEArray)
+	{
+		check(sizeof(T1) == sizeof(T2)); // reinterpret_cast works?
+		check(VtArray.size() <= TNumericLimits<int32>::Max()); // int32 index works?
+		FScopedUnrealAllocs UEAllocs; // Use UE memory allocator
+		UEArray.SetNum(VtArray.size());
+		int32 i = 0;
+		for (typename pxr::VtArray<T1>::const_iterator it = VtArray.cbegin(), itEnd = VtArray.cend();
+			it != itEnd; ++it)
+		{
+			UEArray[i++] = *reinterpret_cast<const T2*>(&*it);
+		}
+	}
+
+	/**
+	 * Linearly interpolate values in \p VtArray1 and \p VtArray2 by interpolant \p Alpha ([0,1]), 
+	 * converting from \p T1 to \p T2, where \p T1 and \p T2 are binary compatible.
+	 *
+	 * Note that \c VtArray is shared copy-on-write, which makes them trivially copied. This function
+	 * copies memory, which should be avoided.
+	 */
+	template <class T1, class T2>
+	void CastLerp(const pxr::VtArray<T1>& VtArray1, const pxr::VtArray<T1>& VtArray2, TArray<T2>& UEArray, const float Alpha)
+	{
+		check(sizeof(T1) == sizeof(T2)); // reinterpret_cast works?
+		check(VtArray1.size() <= TNumericLimits<int32>::Max()); // int32 index works?
+		check(VtArray1.size() == VtArray2.size());
+		FScopedUnrealAllocs UEAllocs; // Use UE memory allocator
+		UEArray.SetNum(VtArray1.size());
+		int32 i = 0;
+		if (Alpha <= TNumericLimits<float>::Min())
+		{
+			for (typename pxr::VtArray<T1>::const_iterator
+				it1 = VtArray1.cbegin(), itEnd1 = VtArray1.end();
+				it1 != itEnd1; ++it1)
+			{
+				UEArray[i++] = *reinterpret_cast<const T2*>(&*it1);
+			}
+		}
+		else if (Alpha >= 1.0 - TNumericLimits<float>::Min())
+		{
+			for (typename pxr::VtArray<T1>::const_iterator
+				it2 = VtArray2.cbegin(), itEnd2 = VtArray2.end();
+				it2 != itEnd2; ++it2)
+			{
+				UEArray[i++] = *reinterpret_cast<const T2*>(&*it2);
+			}
+		}
+		else
+		{
+			for (typename pxr::VtArray<T1>::const_iterator
+				it1 = VtArray1.cbegin(), itEnd1 = VtArray1.end(),
+				it2 = VtArray2.cbegin(), itEnd2 = VtArray2.cend();
+				it1 != itEnd1 && it2 != itEnd2; ++it1, ++it2)
+			{
+				UEArray[i++] = 
+					(1.0 - Alpha) * *reinterpret_cast<const T2*>(&*it1) + 
+					       Alpha  * *reinterpret_cast<const T2*>(&*it2);
+			}
+		}
+	}
+
+	void SaveStage(const pxr::UsdStageRefPtr& Stage, const pxr::UsdTimeCode& FirstFrame, const pxr::UsdTimeCode &LastFrame);
+	void DefineAncestorTransforms(pxr::UsdStageRefPtr& Stage, const pxr::SdfPath& PrimPath);
+	void ComputeExtent(const pxr::VtArray<pxr::GfVec3f>& Points, pxr::VtVec3fArray& Extent);
+	bool SetPointsExtentAndPrimXform(pxr::UsdPrim& Prim, const pxr::GfMatrix4d& PrimXf, const pxr::VtArray<pxr::GfVec3f>& Points, const pxr::VtArray<pxr::GfVec3f>& Vels, const pxr::UsdTimeCode& Time = pxr::UsdTimeCode::Default());
+	bool SetPointsExtentAndPrimXform(pxr::UsdPrim& Prim, const pxr::GfMatrix4d& PrimXf, const pxr::VtArray<pxr::GfVec3f>& Points, const pxr::UsdTimeCode& Time = pxr::UsdTimeCode::Default());
+
+} // namespace UE::SimCacheUSDUtil::Private
+
+#endif // USE_USD_SDK
\ No newline at end of file
diff --git a/Engine/Plugins/Importers/USDImporter/Source/SimCacheUSD/Private/UEUsdGeomTetMesh.cpp b/Engine/Plugins/Importers/USDImporter/Source/SimCacheUSD/Private/UEUsdGeomTetMesh.cpp
new file mode 100644
--- /dev/null
+++ b/Engine/Plugins/Importers/USDImporter/Source/SimCacheUSD/Private/UEUsdGeomTetMesh.cpp
@@ -0,0 +1,172 @@
+// Copyright Epic Games, Inc. All Rights Reserved.
+
+#include "UEUsdGeomTetMesh.h"
+#include "UEUsdGeomTokens.h"
+
+#if USE_USD_SDK
+
+#include "USDMemory.h"
+#include "USDIncludesStart.h"
+	#include "pxr/usd/sdf/assetPath.h"
+	#include "pxr/usd/sdf/types.h"
+	#include "pxr/usd/usd/schemaRegistry.h"
+	#include "pxr/usd/usd/typed.h"
+
+PXR_NAMESPACE_OPEN_SCOPE
+// Register the schema with the TfType system - need to do this inside 
+// USDIncludesStart/End.h so that warnings are disabled, and re-enabled 
+// in USDIncludesEnd.h.
+#pragma warning(disable: 4191) /* 'reinterpret_cast': unsafe conversion during TF_REGISTRY_FUNCTION(TfType) */
+TF_REGISTRY_FUNCTION(TfType)
+{
+	FScopedUsdAllocs UEAllocs; // Use USD memory allocator
+	TfType::Define<UEUsdGeomTetMesh, TfType::Bases<UsdGeomMesh> >();
+	TfType::AddAlias<UsdSchemaBase, UEUsdGeomTetMesh>("UETetMesh");
+}
+PXR_NAMESPACE_CLOSE_SCOPE
+
+#include "USDIncludesEnd.h"
+
+PXR_NAMESPACE_OPEN_SCOPE
+
+UEUsdGeomTetMesh::UEUsdGeomTetMesh(const UsdPrim& Prim)
+    : UsdGeomMesh(Prim)          
+{
+    if(pxr::UsdAttribute SubdivisionAttr = UsdGeomMesh::CreateSubdivisionSchemeAttr())
+    {
+        SubdivisionAttr.Set(pxr::UsdGeomTokens->none);
+    }
+}
+
+UEUsdGeomTetMesh::UEUsdGeomTetMesh(const UsdSchemaBase& SchemaObj)
+    : UsdGeomMesh(SchemaObj)          
+{
+    if(pxr::UsdAttribute SubdivisionAttr = UsdGeomMesh::CreateSubdivisionSchemeAttr())
+    {
+        SubdivisionAttr.Set(pxr::UsdGeomTokens->none);
+    }
+}
+
+UEUsdGeomTetMesh::~UEUsdGeomTetMesh()
+{}
+
+UEUsdGeomTetMesh
+UEUsdGeomTetMesh::Get(const UsdStagePtr &Stage, const SdfPath &Path)
+{
+	if (!Stage)
+	{
+		TF_CODING_ERROR("Invalid stage");
+		return UEUsdGeomTetMesh();
+	}
+	return UEUsdGeomTetMesh(Stage->GetPrimAtPath(Path));
+}
+
+UEUsdGeomTetMesh
+UEUsdGeomTetMesh::Define(const UsdStagePtr &Stage, const SdfPath &Path)
+{
+    static TfToken UsdPrimTypeName("UETetMesh");
+    if(!Stage)
+    {
+        TF_CODING_ERROR("Invalid stage");
+        return UEUsdGeomTetMesh();
+    }
+    return UEUsdGeomTetMesh(Stage->DefinePrim(Path, UsdPrimTypeName));
+}
+
+UsdSchemaKind
+UEUsdGeomTetMesh::_GetSchemaKind() const
+{
+    return UEUsdGeomTetMesh::schemaKind;
+}
+
+const TfType &
+UEUsdGeomTetMesh::_GetStaticTfType()
+{
+    static TfType Type = TfType::Find<UEUsdGeomTetMesh>();
+    return Type;
+}
+
+bool
+UEUsdGeomTetMesh::_IsTypedSchema()
+{
+    static bool bIsTyped = _GetStaticTfType().IsA<UsdTyped>();
+    return bIsTyped;
+}
+
+const TfType &
+UEUsdGeomTetMesh::_GetTfType() const
+{
+    return _GetStaticTfType();
+}
+
+UsdAttribute
+UEUsdGeomTetMesh::GetTetVertexIndicesAttr() const
+{
+    return GetPrim().GetAttribute(UEUsdGeomTokens->tetVertexIndices);
+}
+
+UsdAttribute
+UEUsdGeomTetMesh::CreateTetVertexIndicesAttr(VtValue const &DefaultValue, bool bWriteSparsely) const
+{
+    return UsdSchemaBase::_CreateAttr(
+		UEUsdGeomTokens->tetVertexIndices,
+		SdfValueTypeNames->Int4Array,
+		false,    // custom
+		SdfVariabilityVarying,
+		DefaultValue,
+		bWriteSparsely);
+}
+
+UsdAttribute 
+UEUsdGeomTetMesh::GetTetOrientationAttr() const
+{
+	return GetPrim().GetAttribute(UEUsdGeomTokens->tetOrientation);
+}
+
+UsdAttribute 
+UEUsdGeomTetMesh::CreateTetOrientationAttr(VtValue const& DefaultValue, bool bWriteSparsely) const
+{
+	return UsdSchemaBase::_CreateAttr(
+		UEUsdGeomTokens->tetOrientation,
+		SdfValueTypeNames->Token,
+		false,	// custom
+		SdfVariabilityUniform,
+		DefaultValue,
+		bWriteSparsely);
+}
+
+namespace {
+	static inline TfTokenVector
+	_ConcatenateAttributeNames(const TfTokenVector& Left, const TfTokenVector& Right)
+	{
+	    TfTokenVector Result;
+	    Result.reserve(Left.size() + Right.size());
+	    Result.insert(Result.end(), Left.begin(), Left.end());
+	    Result.insert(Result.end(), Right.begin(), Right.end());
+	    return Result;
+	}
+} // namespace
+
+const TfTokenVector&
+UEUsdGeomTetMesh::GetSchemaAttributeNames(bool bIncludeInerited)
+{
+    static TfTokenVector LocalNames = {
+		UEUsdGeomTokens->tetOrientation,
+		UEUsdGeomTokens->tetVertexIndices,
+	};
+    static TfTokenVector AllNames =
+        _ConcatenateAttributeNames(
+          UsdGeomPointBased::GetSchemaAttributeNames(true),
+          LocalNames);
+    return bIncludeInerited ? AllNames : LocalNames;
+}
+
+// ===================================================================== //
+// Feel free to add custom code below this line. It will be preserved by
+// the code generator.
+// ===================================================================== //
+// --(BEGIN CUSTOM CODE)--
+
+PXR_NAMESPACE_CLOSE_SCOPE
+
+#endif // USE_USD_SDK
diff --git a/Engine/Plugins/Importers/USDImporter/Source/SimCacheUSD/Private/UEUsdGeomTetMesh.h b/Engine/Plugins/Importers/USDImporter/Source/SimCacheUSD/Private/UEUsdGeomTetMesh.h
new file mode 100644
--- /dev/null
+++ b/Engine/Plugins/Importers/USDImporter/Source/SimCacheUSD/Private/UEUsdGeomTetMesh.h
@@ -0,0 +1,104 @@
+// Copyright Epic Games, Inc. All Rights Reserved.
+#pragma once
+#if USE_USD_SDK
+
+#include "USDIncludesStart.h"
+	#include "pxr/base/gf/vec4i.h"
+	#include "pxr/base/tf/type.h"
+	#include "pxr/base/vt/value.h"
+
+	#include "pxr/usd/usdGeom/mesh.h"
+	#include "pxr/usd/usd/prim.h"
+	#include "pxr/usd/usd/stage.h"
+	#include "pxr/usd/usdGeom/tokens.h"
+	#include "pxr/usd/usd/timeCode.h"
+#include "USDIncludesEnd.h"
+
+PXR_NAMESPACE_OPEN_SCOPE
+
+class SdfAssetPath;
+
+/**
+ * \class UEUsdGeomTetMesh
+ *
+ * A Tetrahedral mesh prim, derived from \c UsdGeomMesh, which is used for 
+ * drawing the exterior hull of the tetrahedron mesh.  This class adds a
+ * topology array for tetrahedra consisting of \c GfVec4i, and an attribute
+ * that specifies the tetrahedron orientation/winding order.
+ * 
+ * The constructor defines and sets the subdivision scheme attribute to `none`
+ * so that by default the geometry rendered by the \c UsdGeomMesh base class
+ * is polygonal.
+ */
+class UEUsdGeomTetMesh : public UsdGeomMesh
+{
+public:
+    static const UsdSchemaKind schemaKind = UsdSchemaKind::ConcreteTyped;                                                                          
+
+    /** Construct a UEUsdGeomTetMesh from a UsdPrim \p prim. */
+    explicit UEUsdGeomTetMesh(const UsdPrim& Prim=UsdPrim());
+    explicit UEUsdGeomTetMesh(const UsdSchemaBase& SchemaObj);
+
+    virtual ~UEUsdGeomTetMesh();
+
+    static const TfTokenVector &
+    GetSchemaAttributeNames(bool bIncludeInherited=true);
+
+    static UEUsdGeomTetMesh
+    Get(const UsdStagePtr &Stage, const SdfPath &Path);
+
+    static UEUsdGeomTetMesh
+    Define(const UsdStagePtr &Stage, const SdfPath &Path);
+   
+protected:
+    UsdSchemaKind _GetSchemaKind() const override;
+
+private:
+    // Needs to invoke _GetStaticTfType.
+    friend class UsdSchemaRegistry;
+    static const TfType &_GetStaticTfType();
+
+    static bool _IsTypedSchema();
+
+    // Override SchemaBase virtuals.
+    virtual const TfType &_GetTfType() const;
+
+public:
+    /**
+	 * Tetrahedron indices. Each element of this array contains the
+     * indices of the vertices of a tetrahedron. The points are authored
+     * into the "points" attribute, just as they are for other PointBased    
+     * prims.
+     *
+     *   C++ Type: VtArray<GfVec4i>
+     *   Usd Type: SdfValueTypeNames->Int4Array
+     *   Variability: SdfVariabilityVarying
+     *   Fallback Value: No Fallback
+	 */
+    UsdAttribute GetTetVertexIndicesAttr() const;
+    
+    /**
+	 * See GetTetVertexIndicesAttr(), and also
+     * \ref Usd_Create_Or_Get_Property for when to use Get vs Create.
+     * If specified, author \p DefaultValue as the attribute's default,
+     * sparsely (when it makes sense to do so) if \p WriteSparsely is \c true -
+     * the default for \p writeSparsely is \c false.
+	 */
+    UsdAttribute CreateTetVertexIndicesAttr(VtValue const &DefaultValue = VtValue(), bool WriteSparsely=false) const;
+    
+public:
+	/**
+	 * Tetrahedron orientation. Given [A,B,C,D] if your right hand 
+	 * is at A, if your fingers wrap to A-B, then A-C, and your thumb 
+	 * points at D.
+	 */
+	UsdAttribute GetTetOrientationAttr() const;
+
+	/**
+	* See GetTetOrientationAttr().
+	*/
+	UsdAttribute CreateTetOrientationAttr(VtValue const& DefaultValue = VtValue(), bool WriteSparsely = false) const;
+};
+
+PXR_NAMESPACE_CLOSE_SCOPE
+#endif // USE_USD_SDK
\ No newline at end of file
diff --git a/Engine/Plugins/Importers/USDImporter/Source/SimCacheUSD/Private/UEUsdGeomTokens.cpp b/Engine/Plugins/Importers/USDImporter/Source/SimCacheUSD/Private/UEUsdGeomTokens.cpp
new file mode 100644
--- /dev/null
+++ b/Engine/Plugins/Importers/USDImporter/Source/SimCacheUSD/Private/UEUsdGeomTokens.cpp
@@ -0,0 +1,22 @@
+// Copyright Epic Games, Inc. All Rights Reserved.
+
+#include "UEUsdGeomTokens.h"
+
+#if USE_USD_SDK
+
+PXR_NAMESPACE_OPEN_SCOPE
+
+UEUsdGeomTokensType::UEUsdGeomTokensType()
+    : tetOrientation("tetOrientation", TfToken::Immortal)
+    , tetVertexIndices("tetVertexIndices", TfToken::Immortal)
+    , allTokens({
+            tetOrientation,
+            tetVertexIndices
+        })
+{}
+
+TfStaticData<UEUsdGeomTokensType> UEUsdGeomTokens;
+
+PXR_NAMESPACE_CLOSE_SCOPE
+
+#endif // USE_USD_SDK
diff --git a/Engine/Plugins/Importers/USDImporter/Source/SimCacheUSD/Private/UEUsdGeomTokens.h b/Engine/Plugins/Importers/USDImporter/Source/SimCacheUSD/Private/UEUsdGeomTokens.h
new file mode 100644
--- /dev/null
+++ b/Engine/Plugins/Importers/USDImporter/Source/SimCacheUSD/Private/UEUsdGeomTokens.h
@@ -0,0 +1,38 @@
+// Copyright Epic Games, Inc. All Rights Reserved.
+
+#pragma once
+#if USE_USD_SDK
+
+#include "USDIncludesStart.h"
+	#include "pxr/base/tf/staticData.h"
+	#include "pxr/base/tf/token.h"
+#include "USDIncludesEnd.h"
+
+#include <vector>
+
+PXR_NAMESPACE_OPEN_SCOPE
+
+struct UEUsdGeomTokensType 
+{
+	SIMCACHEUSD_API UEUsdGeomTokensType();
+
+    /// \brief "tetOrientation"
+    ///
+    /// UEUsdGeomTetMesh
+    const TfToken tetOrientation;
+    /// \brief "tetVetexIndices"
+    ///
+    /// UEUsdGeomTetMesh
+    const TfToken tetVertexIndices;
+    /// A vector of all of the tokens listed above.
+    const std::vector<TfToken> allTokens;
+};
+
+/// \var UEUsdGeomTokens
+///
+/// A global variable for access to static tokens.
+extern SIMCACHEUSD_API TfStaticData<UEUsdGeomTokensType> UEUsdGeomTokens;
+
+PXR_NAMESPACE_CLOSE_SCOPE
+
+#endif // USE_USD_SDK
diff --git a/Engine/Plugins/Importers/USDImporter/Source/SimCacheUSD/Private/wrapUEUsdGeomTetMesh.h b/Engine/Plugins/Importers/USDImporter/Source/SimCacheUSD/Private/wrapUEUsdGeomTetMesh.h
new file mode 100644
--- /dev/null
+++ b/Engine/Plugins/Importers/USDImporter/Source/SimCacheUSD/Private/wrapUEUsdGeomTetMesh.h
@@ -0,0 +1,73 @@
+// Copyright Epic Games, Inc. All Rights Reserved.
+#pragma once
+
+// TODO: Figure out whether we want/need python bindings, and what the right way to do them is.
+#if 0//USE_USD_SDK
+
+// python bindings:
+#include "pxr/usd/usd/pyConversions.h"
+#include "pxr/base/tf/pyContainerConversions.h"
+#include "pxr/base/tf/pyResultsConversions.h"
+#include "pxr/base/tf/pyUtils.h"
+#include "pxr/base/tf/wrapTypeHelpers.h"
+#include <boost/python.hpp>
+
+PXR_NAMESPACE_OPEN_SCOPE
+
+#define WRAP_CUSTOM
+	template <class Cls> static void _CustomWrapCode(Cls& _class)
+
+		// fwd decl.
+		WRAP_CUSTOM;
+
+	static UsdAttribute
+		_CreateTetVertexIndicesAttr(
+			UEUsdGeomTetMesh& self,
+			object defaultVal, bool writeSparsely)
+	{
+		return self.CreateTetVertexIndicesAttr(
+			UsdPythonToSdfType(defaultVal, SdfValueTypeNames->Int4Array), writeSparsely);
+	}
+
+	void wrapUEUsdGeomTetMesh()
+	{
+		typedef UEUsdGeomTetMesh This;
+
+		class_<This, boost::python::bases<UsdGeomMesh> >
+			cls("TetMesh");
+		cls
+			.def(init<UsdPrim>(arg("prim")))
+			.def(init<UsdSchemaBase const&>(arg("schemaObj")))
+			.def(TfTypePythonClass())
+
+			.def("Get", &This::Get, (arg("stage"), arg("path")))
+			.staticmethod("Get")
+
+			.def("GetSchemaAttributeNames",
+				&This::GetSchemaAttributeNames,
+				arg("includeInherited") = true,
+				boost::python::return_value_policy<TfPySequenceToList>())
+			.staticmethod("GetSchemaAttributeNames")
+
+			.def("_GetStaticTfType", &TfType const& (*)()) TfType::Find<This>,
+			boost::python::return_value_policy<boost::python::return_by_value>())
+			.staticmethod("_GetStaticTfType")
+
+			.def(!self)
+
+			.def("GetTetVertexIndicesAttr",
+				&This::GetTetVertexIndicesAttr)
+			.def("CreateTetVertexIndicesAttr",
+				&_CreateTetVertexIndicesAttr,
+				(boost::python::arg("defaultValue") = object(),
+					boost::python::arg("writeSparsely") = false))
+			;
+
+		_CustomWrapCode(cls);
+	}
+
+}
+
+PXR_NAMESPACE_CLOSE_SCOPE
+
+#endif // USE_USD_SDK
diff --git a/Engine/Plugins/Importers/USDImporter/Source/SimCacheUSD/Public/SimCacheUSDOperations.h b/Engine/Plugins/Importers/USDImporter/Source/SimCacheUSD/Public/SimCacheUSDOperations.h
new file mode 100644
--- /dev/null
+++ b/Engine/Plugins/Importers/USDImporter/Source/SimCacheUSD/Public/SimCacheUSDOperations.h
@@ -0,0 +1,153 @@
+// Copyright Epic Games, Inc. All Rights Reserved.
+
+#pragma once
+#if USE_USD_SDK
+
+#include "GeometryCollection/GeometryCollection.h"
+#include "GeometryCollection/ManagedArray.h"
+#include "Math/UnrealMathUtility.h"
+#include "UsdWrappers/UsdStage.h"
+
+#include "USDIncludesStart.h"
+	#include "pxr/base/gf/vec3f.h"
+	#include "pxr/base/vt/array.h"
+#include "USDIncludesEnd.h"
+
+namespace UE::SimCacheUSD
+{
+	template <class TV>
+	pxr::VtArray<pxr::GfVec3f> ToVtVec3Array(const TArray<TV>& V3Array)
+	{
+		FScopedUsdAllocs UsdAllocs; // Use USD memory allocator
+		pxr::VtArray<pxr::GfVec3f> RetVal(static_cast<size_t>(V3Array.Num()));
+		for (int32 i = 0; i < V3Array.Num(); i++) RetVal[i].Set(V3Array[i][0], V3Array[i][1], V3Array[i][2]);
+		return RetVal;
+	}
+
+	bool ValuesDiffer(const pxr::VtArray<pxr::GfVec3f>& A, const pxr::VtArray<pxr::GfVec3f>& B, const float Tolerance = 1.0e-8, const uint64 stride = 1)
+	{
+		if (A.size() != B.size()) return true;
+		for (size_t i = 0; i < A.size(); i += stride)
+		{
+			for (size_t j = 0; j < 3; j++)
+			{
+				if (!FMath::IsNearlyEqual(A[i][j], B[i][j], Tolerance))
+					return true;
+			}
+		}
+		return false;
+	}
+
+	/** Create a new stage, resident in the USD stage cache. */
+	SIMCACHEUSD_API bool NewStage(const FString& StageName, UE::FUsdStage& Stage);
+	/** Opens an existing stage, resident in the USD stage cache. */
+	SIMCACHEUSD_API bool OpenStage(const FString& StageName, UE::FUsdStage& Stage);
+
+	/** Save the stage and update the frame range (if not -DBL_MAX). */
+	SIMCACHEUSD_API bool SaveStage(UE::FUsdStage& Stage, const double FirstFrame, const double LastFrame);
+
+	/** Close and remove a stage from the USD stage cache. */
+	SIMCACHEUSD_API bool CloseStage(const UE::FUsdStage& Stage);
+	SIMCACHEUSD_API bool CloseStage(const FString& StageName);
+
+	/**
+	 * "Value Clips" is USD's concept of splitting data up into multiple 
+	 * files, the it looks like this:
+	 *
+	 *		path/to/file.usd - top level root or "parent" file that contains references to other files
+	 *		path/to/file.topology.usd - unvarying data goes here
+	 *		path/to/file.manifest.usd - we don't use this file
+	 *		path/to/file.#.usd - data for frame # goes here, where # is some number.
+	 *
+	 * Writing value clips scenes requires juggling multiple stages. Reading value clips scenes
+	 * should be somewhat simpler as USD's composition should pick up file references.
+	 */
+
+	/** Given 'path/to/file.usd', yields \p TopologyName 'path/to/file.topology.usd' and \p TimeVaryingTemplate 'path/to/file.###.###.usd'. */
+	SIMCACHEUSD_API void
+	GenerateValueClipStageNames(const FString& ParentName, FString& TopologyName, FString& TimeVaryingTemplate);
+	/** Given 'path/to/file.###.###.usd' and 1.23, yields 'path/to/file.001.230.usd'. */
+	SIMCACHEUSD_API FString
+	GenerateValueClipTimeVaryingStageName(const FString& TimeVaryingTemplate, const double Time);
+
+	/** Create new parent and topology stages for value clips. */
+	SIMCACHEUSD_API bool
+	NewValueClipsStages(
+		const FString& ParentStageName,
+		const FString& TopologyStageName,
+		UE::FUsdStage& ParentStage,
+		UE::FUsdStage& TopologyStage);
+	/** Create new value clips frame stage. Be sure to close these stages. */
+	SIMCACHEUSD_API bool
+	NewValueClipsFrameStage(
+		const FString& TimeVaryingStageTemplate,
+		const double Time,
+		FString& FrameStageName,
+		UE::FUsdStage& FrameStage);
+	/** 
+	 * Init layers and metadata for value clips. 
+	 * 
+	 * \p ParentStage - the stage that serves as the hub and references others.
+	 * \p TopologyStage - stage that will hod unvarying scene data.
+	 * \p ParentStageName
+	 * \p TopologyStageName
+	 * \p TimeVaryingStageTemplate - template pattern for per-frame USD data files.
+	 * \p PrimPaths - the prim(s) to author value clips metadata on.  It will apply to all prims 
+	 *    below this point in the hierarchy.  Authoring the root scope is not allowed by USD.
+	 * \p StartTime - the first time sample that should be available at this point in time.
+	 * \p EndTime - the last time sample that should be available at this point in time.
+	 * \p Stride - the size of the steps taken between \p StartTime and \p EndTime.
+	 */
+	SIMCACHEUSD_API bool
+	InitValueClipsTemplate(
+		UE::FUsdStage& ParentStage,
+		UE::FUsdStage& TopologyStage,
+		const FString& ParentStageName,
+		const FString& TopologyStageName,
+		const FString& TimeVaryingStageTemplate,
+		const TArray<FString>& PrimPaths,
+		const double StartTime,
+		const double EndTime,
+		const double Stride);
+	/** Helper function that opens/retreives parent and topology stages from cache. */
+	SIMCACHEUSD_API bool
+	InitValueClipsTemplate(
+		const FString& ParentStageName,
+		const FString& TopologyStageName,
+		const FString& TimeVaryingStageTemplate,
+		const TArray<FString>& PrimPaths,
+		const double StartTime,
+		const double EndTime,
+		const double Stride);
+
+	/**
+	 * Define \c UEUsdGeomTetMesh at \c PrimPath with ancenstor transforms, and authors points at
+	 * usd "default" time from tet mesh \p StructureIndex (or all if \c INDEX_NONE) from \c Collection.
+	 */
+	SIMCACHEUSD_API bool WriteTetMesh(UE::FUsdStage& Stage, const FString& PrimPath, const FManagedArrayCollection& Collection, const int32 StructureIndex = INDEX_NONE);
+
+	/**
+	 * Define \c UsdGeomPointBased at \c PrimPath with ancestor transforms, and author points at 
+	 * \p Time (or usd "default" time if \c -DBL_MAX) from \p Collection vertices from geometry 
+	 * \p StructureIndex (or all if \c StructureIndex == \c INDEX_NONE).
+	 */
+	SIMCACHEUSD_API bool WritePoints(UE::FUsdStage& Stage, const FString& PrimPath, const double Time, const FManagedArrayCollection& Collection, const int32 StructureIndex = INDEX_NONE);
+	/** Write points and velocities to \c UsdGeomPointBased at \c PrimPath. */
+	SIMCACHEUSD_API bool WritePoints(UE::FUsdStage& Stage, const FString& PrimPath, const double Time, pxr::VtArray<pxr::GfVec3f>& VtPoints, pxr::VtArray<pxr::GfVec3f>& VtVels);
+
+	/** Get time samples for an attribute. */
+	SIMCACHEUSD_API bool ReadTimeSamples(const UE::FUsdStage& Stage, const FString& PrimPath, const FString& AttrName, TArray<double>& TimeSamples);
+	/** Get time samples for the points attribute. */
+	SIMCACHEUSD_API bool ReadTimeSamples(const UE::FUsdStage& Stage, const FString& PrimPath, TArray<double>& TimeSamples);
+
+	SIMCACHEUSD_API FString GetPointsAttrName();
+	SIMCACHEUSD_API FString GetVelocityAttrName();
+	SIMCACHEUSD_API bool GetBracketingTimeSamples(const UE::FUsdStage& Stage, const FString& PrimPath, const FString& AttrName, const double TargetTime, double* Lower, double* Upper);
+
+	/** Get points from an attribute. Default time is used if \p Time is \c -DBL_MAX. */
+	SIMCACHEUSD_API bool ReadPoints(const UE::FUsdStage& Stage, const FString& PrimPath, const FString& AttrPath, const double Time, pxr::VtArray<pxr::GfVec3f>& Points);
+	/** Get points from the points attribute. Default time is used if \p Time is \c -DBL_MAX. */
+	SIMCACHEUSD_API bool ReadPoints(const UE::FUsdStage& Stage, const FString& PrimPath, const double Time, pxr::VtArray<pxr::GfVec3f>& Points, pxr::VtArray<pxr::GfVec3f>& VtVels);
+
+} // UE::SimCacheUSD
+#endif // USE_USD_SDK
diff --git a/Engine/Plugins/Importers/USDImporter/Source/SimCacheUSD/SimCacheUSD.build.cs b/Engine/Plugins/Importers/USDImporter/Source/SimCacheUSD/SimCacheUSD.build.cs
new file mode 100644
--- /dev/null
+++ b/Engine/Plugins/Importers/USDImporter/Source/SimCacheUSD/SimCacheUSD.build.cs
@@ -0,0 +1,41 @@
+// Copyright Epic Games, Inc. All Rights Reserved.
+
+using System.Linq;
+using System.IO;
+using System.Collections.Generic;
+
+namespace UnrealBuildTool.Rules
+{
+	public class SimCacheUSD : ModuleRules
+	{
+		public SimCacheUSD(ReadOnlyTargetRules Target) : base(Target)
+		{
+			// Does not compile with C++20:
+			// error C4002: too many arguments for function-like macro invocation 'TF_PP_CAT_IMPL'
+			// warning C5103: pasting '"TF_LOG_STACK_TRACE_ON_ERROR"' and '"TF_LOG_STACK_TRACE_ON_WARNING"' does not result in a valid preprocessing token
+			CppStandard = CppStandardVersion.Cpp17;
+
+			bUseRTTI = true;
+
+			PublicDependencyModuleNames.AddRange(
+				new string[]
+				{
+				"Boost",
+				"Core",
+				"CoreUObject",
+				"Engine",
+				"Projects",
+				"RHI",
+				}
+			);
+
+			PrivateDependencyModuleNames.AddRange(
+				new string[] {
+				"UnrealUSDWrapper",
+				"USDClasses",
+				"USDUtilities",
+				}
+			);
+		}
+	}
+}
\ No newline at end of file
diff --git a/Engine/Plugins/Importers/USDImporter/USDImporter.uplugin b/Engine/Plugins/Importers/USDImporter/USDImporter.uplugin
--- a/Engine/Plugins/Importers/USDImporter/USDImporter.uplugin
+++ b/Engine/Plugins/Importers/USDImporter/USDImporter.uplugin
@@ -73,6 +73,11 @@
 			"Type": "Editor",
 			"LoadingPhase": "Default"
 		},
+		{
+			"Name": "SimCacheUSD",
+			"Type": "Editor",
+			"LoadingPhase": "Default"
+		},
 		{
 			"Name": "USDTests",
 			"Type": "Editor",
diff --git a/Engine/Source/Runtime/Experimental/Chaos/Public/GeometryCollection/ManagedArrayCollection.h b/Engine/Source/Runtime/Experimental/Chaos/Public/GeometryCollection/ManagedArrayCollection.h
--- a/Engine/Source/Runtime/Experimental/Chaos/Public/GeometryCollection/ManagedArrayCollection.h
+++ b/Engine/Source/Runtime/Experimental/Chaos/Public/GeometryCollection/ManagedArrayCollection.h
@@ -578,10 +578,10 @@ private:
 		int32 Size;
 	};
 
-	static FKeyType MakeMapKey(FName Name, FName Group)
+	static inline FKeyType MakeMapKey(FName Name, FName Group) // inline for linking issues
 	{
 		return FKeyType(Name, Group);
-	};
+	}
 
 	friend struct FManagedArrayCollectionValueTypeWrapper;
 	struct FValueType
