USD: Discard old asset prim links when calling LoadUsdStage.

diff --git a/Engine/Plugins/Importers/USDImporter/Source/USDSchemas/Private/USDInfoCache.cpp b/Engine/Plugins/Importers/USDImporter/Source/USDSchemas/Private/USDInfoCache.cpp
--- a/Engine/Plugins/Importers/USDImporter/Source/USDSchemas/Private/USDInfoCache.cpp
+++ b/Engine/Plugins/Importers/USDImporter/Source/USDSchemas/Private/USDInfoCache.cpp
@@ -1480,6 +1480,21 @@ TArray<TWeakObjectPtr<UObject>> FUsdInfoCache::RemoveAllAssetPrimLinks(const UE:
 	return Assets;
 }
 
+void FUsdInfoCache::RemoveAllAssetPrimLinks()
+{
+	FUsdInfoCacheImpl* ImplPtr = Impl.Get();
+	if (!ImplPtr)
+	{
+		return;
+	}
+	FWriteScopeLock ScopeLock(ImplPtr->PrimPathToAssetsLock);
+
+	UE_LOG(LogUsd, Verbose, TEXT("Removing all asset prim links"));
+
+	ImplPtr->PrimPathToAssets.Empty();
+	ImplPtr->AssetToPrimPaths.Empty();
+}
+
 TArray<TWeakObjectPtr<UObject>> FUsdInfoCache::GetAllAssetsForPrim(const UE::FSdfPath& Path) const
 {
 	FUsdInfoCacheImpl* ImplPtr = Impl.Get();
diff --git a/Engine/Plugins/Importers/USDImporter/Source/USDSchemas/Public/USDInfoCache.h b/Engine/Plugins/Importers/USDImporter/Source/USDSchemas/Public/USDInfoCache.h
--- a/Engine/Plugins/Importers/USDImporter/Source/USDSchemas/Public/USDInfoCache.h
+++ b/Engine/Plugins/Importers/USDImporter/Source/USDSchemas/Public/USDInfoCache.h
@@ -91,6 +91,7 @@ public:
 	void LinkAssetToPrim(const UE::FSdfPath& Path, UObject* Asset);
 
 	TArray<TWeakObjectPtr<UObject>> RemoveAllAssetPrimLinks(const UE::FSdfPath& Path);
+	void RemoveAllAssetPrimLinks();
 
 	TArray<TWeakObjectPtr<UObject>> GetAllAssetsForPrim(const UE::FSdfPath& Path) const;
 
diff --git a/Engine/Plugins/Importers/USDImporter/Source/USDStage/Private/USDStageActor.cpp b/Engine/Plugins/Importers/USDImporter/Source/USDStage/Private/USDStageActor.cpp
--- a/Engine/Plugins/Importers/USDImporter/Source/USDStage/Private/USDStageActor.cpp
+++ b/Engine/Plugins/Importers/USDImporter/Source/USDStage/Private/USDStageActor.cpp
@@ -2717,6 +2717,7 @@ void AUsdStageActor::LoadUsdStage()
 	TSharedRef< FUsdSchemaTranslationContext > TranslationContext = FUsdStageActorImpl::CreateUsdSchemaTranslationContext(this, RootTwin->PrimPath);
 
 	SlowTask.EnterProgressFrame(0.1f);
+	InfoCache->RemoveAllAssetPrimLinks();  // These are usually not reset when rebuilding the cache so we must call manually
 	InfoCache->RebuildCacheForSubtree(StageToLoad.GetPseudoRoot(), TranslationContext.Get());
 
 	SlowTask.EnterProgressFrame(0.7f);
