USD: Always bake diffuseColor as black as a fallback because it has a different default in usd than it does in UE.

diff --git a/Engine/Plugins/Importers/USDImporter/Source/USDUtilities/Private/USDShadeConversion.cpp b/Engine/Plugins/Importers/USDImporter/Source/USDUtilities/Private/USDShadeConversion.cpp
--- a/Engine/Plugins/Importers/USDImporter/Source/USDUtilities/Private/USDShadeConversion.cpp
+++ b/Engine/Plugins/Importers/USDImporter/Source/USDUtilities/Private/USDShadeConversion.cpp
@@ -1661,6 +1661,15 @@ namespace UE
 					PropertiesToProcess.Append( UsedProperties );
 				}
 
+				// We always write BaseColor because the default in UE is full black, but in USD seems to be 0.18. If we left
+				// BaseColor unbound, a material that relies on opacity/other channels and leaves BaseColor disconnected for the black
+				// value would end up looking gray in usdview/other DCCs
+				if (!PropertiesToProcess.Contains(MP_BaseColor))
+				{
+					PropertiesToProcess.Add(MP_BaseColor);
+				}
+				const float Zero = 0.0f;
+
 				// Fill in outputs
 				for ( const EMaterialProperty& Property : PropertiesToProcess )
 				{
@@ -1743,6 +1752,12 @@ namespace UE
 						}
 					}
 
+					if (Property == MP_BaseColor && NumSamples == 0 && !UserConstantValue)
+					{
+						UserConstantValue = &Zero;
+						bPropertyValueIsConstant = true;
+					}
+
 					if ( ( NumSamples == 0 || !SampleSize ) && !UserConstantValue )
 					{
 						UE_LOG( LogUsd, Log, TEXT( "Skipping material property %d as we have no valid data to use." ), Property );
@@ -1821,7 +1836,7 @@ namespace UE
 								ConstantValue = ConstantLinearValue;
 							}
 						}
-						FallbackValue = pxr::GfVec4f{ 0.0, 1.0, 0.0, 1.0f };
+						FallbackValue = pxr::GfVec4f{ 0.0, 0.0, 1.0, 1.0f };
 						break;
 					case MP_Tangent:
 						InputToken = UnrealIdentifiers::Tangent;
@@ -1838,7 +1853,7 @@ namespace UE
 								ConstantValue = ConstantLinearValue;
 							}
 						}
-						FallbackValue = pxr::GfVec4f{ 0.0, 1.0, 0.0, 1.0f };
+						FallbackValue = pxr::GfVec4f{ 1.0, 0.0, 0.0, 1.0f };
 						break;
 					case MP_EmissiveColor:
 						InputToken = UnrealIdentifiers::EmissiveColor;
