USD: Fix how degenerate preview surface materials are generated when parsing UnrealMaterial prims now.

diff --git a/Engine/Plugins/Importers/USDImporter/Source/USDSchemas/Private/USDShadeMaterialTranslator.cpp b/Engine/Plugins/Importers/USDImporter/Source/USDSchemas/Private/USDShadeMaterialTranslator.cpp
--- a/Engine/Plugins/Importers/USDImporter/Source/USDSchemas/Private/USDShadeMaterialTranslator.cpp
+++ b/Engine/Plugins/Importers/USDImporter/Source/USDSchemas/Private/USDShadeMaterialTranslator.cpp
@@ -254,6 +254,19 @@ void FUsdShadeMaterialTranslator::CreateAssets()
 		Context->RenderContext.IsNone() ?
 			pxr::UsdShadeTokens->universalRenderContext :
 			UnrealToUsd::ConvertToken( *Context->RenderContext.ToString() ).Get();
+
+	// If this material has a valid surface output for the 'unreal' render context and we're using it, don't bother
+	// generating any new UMaterialInterface asset because when resolving material assignments for this material
+	// all consumers will just use the referenced UAsset anyway
+	if (RenderContextToken == UnrealIdentifiers::Unreal)
+	{
+		TOptional<FString> UnrealMaterial = UsdUtils::GetUnrealSurfaceOutput(ShadeMaterial.GetPrim());
+		if (UnrealMaterial.IsSet())
+		{
+			return;
+		}
+	}
+
 	FString MaterialHashString = UsdUtils::HashShadeMaterial( ShadeMaterial, RenderContextToken ).ToString();
 
 	UMaterialInterface* ConvertedMaterial = nullptr;
