Resolve several module public include paths not using references

#rnx

diff --git a/Engine/Plugins/Experimental/AppleVision/Source/AppleVisionBlueprintSupport/AppleVisionBlueprintSupport.Build.cs b/Engine/Plugins/Experimental/AppleVision/Source/AppleVisionBlueprintSupport/AppleVisionBlueprintSupport.Build.cs
--- a/Engine/Plugins/Experimental/AppleVision/Source/AppleVisionBlueprintSupport/AppleVisionBlueprintSupport.Build.cs
+++ b/Engine/Plugins/Experimental/AppleVision/Source/AppleVisionBlueprintSupport/AppleVisionBlueprintSupport.Build.cs
@@ -7,32 +7,24 @@ namespace UnrealBuildTool.Rules
 		public AppleVisionBlueprintSupport(ReadOnlyTargetRules Target) : base(Target)
 		{
 
-            PrivateIncludePaths.AddRange(
-                new string[] {
-                    "AppleVision/Public",
-                    // ... add other private include paths required here ...
-                }
-                );
-
-
-            PublicDependencyModuleNames.AddRange(
-                new string[]
-                {
-                    "Core",
-                    "Engine",
-                    "BlueprintGraph",
-                    "AppleVision"
-                    // ... add other public dependencies that you statically link with here ...
-                }
-                );
+			PublicDependencyModuleNames.AddRange(
+				new string[]
+				{
+					"Core",
+					"Engine",
+					"BlueprintGraph",
+					"AppleVision"
+					// ... add other public dependencies that you statically link with here ...
+				}
+			);
 
 			PrivateDependencyModuleNames.AddRange(
 				new string[]
 				{
-    				"CoreUObject"
+					"CoreUObject"
 					// ... add private dependencies that you statically link with here ...
 				}
-				);
+			);
 		}
 	}
 }
diff --git a/Engine/Plugins/Importers/USDImporter/Source/UnrealUSDWrapper/UnrealUSDWrapper.Build.cs b/Engine/Plugins/Importers/USDImporter/Source/UnrealUSDWrapper/UnrealUSDWrapper.Build.cs
--- a/Engine/Plugins/Importers/USDImporter/Source/UnrealUSDWrapper/UnrealUSDWrapper.Build.cs
+++ b/Engine/Plugins/Importers/USDImporter/Source/UnrealUSDWrapper/UnrealUSDWrapper.Build.cs
@@ -54,12 +54,12 @@ namespace UnrealBuildTool.Rules
 					// Don't need to handle it for Windows or Mac as IntelTBB.Build.cs already does it
 
 					// Python3
-					PublicIncludePaths.Add(Path.Combine(PythonSourceTPSDir, "include"));
+					PublicSystemIncludePaths.Add(Path.Combine(PythonSourceTPSDir, "include"));
 					PublicSystemLibraryPaths.Add(Path.Combine(PythonSourceTPSDir, "libs"));
 					RuntimeDependencies.Add(Path.Combine("$(TargetOutputDir)", "python39.dll"), Path.Combine(PythonBinaryTPSDir, "python39.dll"));
 
 					// USD
-					PublicIncludePaths.Add(Path.Combine(ModuleDirectory, "..", "ThirdParty", "USD", "include"));
+					PublicSystemIncludePaths.Add(Path.Combine(ModuleDirectory, "..", "ThirdParty", "USD", "include"));
 					PublicSystemLibraryPaths.Add(USDLibsDir);
 					foreach (string UsdLib in Directory.EnumerateFiles(USDLibsDir, "*.lib", SearchOption.AllDirectories))
 					{
@@ -85,13 +85,13 @@ namespace UnrealBuildTool.Rules
 					RuntimeDependencies.Add(Path.Combine(IntelTBBBinaries, "libtbb.so.2"));
 
 					// Python3
-					PublicIncludePaths.Add(Path.Combine(PythonSourceTPSDir, "include"));
+					PublicSystemIncludePaths.Add(Path.Combine(PythonSourceTPSDir, "include"));
 					PublicSystemLibraryPaths.Add(Path.Combine(PythonBinaryTPSDir, "lib"));
 					PrivateRuntimeLibraryPaths.Add(Path.Combine(PythonBinaryTPSDir, "bin"));
 					RuntimeDependencies.Add(Path.Combine(PythonBinaryTPSDir, "bin", "python3.9"));
 
 					// USD
-					PublicIncludePaths.Add(Path.Combine(ModuleDirectory, "..", "ThirdParty", "USD", "include"));
+					PublicSystemIncludePaths.Add(Path.Combine(ModuleDirectory, "..", "ThirdParty", "USD", "include"));
 					var USDBinDir = Path.Combine(ModuleDirectory, "..", "ThirdParty", "Linux", "bin", Target.Architecture.LinuxName);
 					PrivateRuntimeLibraryPaths.Add(USDBinDir);
 					foreach (string LibPath in Directory.EnumerateFiles(USDBinDir, "*.so", SearchOption.AllDirectories))
@@ -116,13 +116,13 @@ namespace UnrealBuildTool.Rules
 					// Don't need to handle it for Windows or Mac as IntelTBB.Build.cs already does it
 
 					// Python3
-					PublicIncludePaths.Add(Path.Combine(PythonSourceTPSDir, "include"));
+					PublicSystemIncludePaths.Add(Path.Combine(PythonSourceTPSDir, "include"));
 					PublicSystemLibraryPaths.Add(Path.Combine(PythonBinaryTPSDir, "lib"));
 					PrivateRuntimeLibraryPaths.Add(Path.Combine(PythonBinaryTPSDir, "bin"));
 					RuntimeModulePaths.Add(Path.Combine(PythonBinaryTPSDir, "lib", "libpython3.9.dylib"));
 
 					// USD
-					PublicIncludePaths.Add(Path.Combine(ModuleDirectory, "..", "ThirdParty", "USD", "include"));
+					PublicSystemIncludePaths.Add(Path.Combine(ModuleDirectory, "..", "ThirdParty", "USD", "include"));
 					var USDBinDir = Path.Combine(ModuleDirectory, "..", "ThirdParty", "Mac", "bin");
 					foreach (string LibPath in Directory.EnumerateFiles(USDBinDir, "*.dylib", SearchOption.AllDirectories))
 					{
diff --git a/Engine/Plugins/Media/AvfMedia/Source/AvfMediaCapture/AvfMediaCapture.Build.cs b/Engine/Plugins/Media/AvfMedia/Source/AvfMediaCapture/AvfMediaCapture.Build.cs
--- a/Engine/Plugins/Media/AvfMedia/Source/AvfMediaCapture/AvfMediaCapture.Build.cs
+++ b/Engine/Plugins/Media/AvfMedia/Source/AvfMediaCapture/AvfMediaCapture.Build.cs
@@ -1,7 +1,5 @@
 // Copyright Epic Games, Inc. All Rights Reserved.
 
-using System;
-using System.IO;
 namespace UnrealBuildTool.Rules
 {
 	public class AvfMediaCapture : ModuleRules
@@ -9,9 +7,9 @@ namespace UnrealBuildTool.Rules
 		public AvfMediaCapture(ReadOnlyTargetRules Target) : base(Target)
 		{
 			DynamicallyLoadedModuleNames.AddRange(
-			new string[] {
-				"Media",
-			});
+				new string[] {
+					"Media",
+				});
 
 			PrivateDependencyModuleNames.AddRange(
 				new string[] {
@@ -22,18 +20,17 @@ namespace UnrealBuildTool.Rules
 					"RenderCore",
 					"RHI"
 				});
-			
+
 			PrivateIncludePathModuleNames.AddRange(
 				new string[] {
 					"Media",
 				});
-			
-			var EngineDir = Path.GetFullPath(Target.RelativeEnginePath);
-			PublicIncludePaths.AddRange(
+
+			PublicIncludePathModuleNames.AddRange(
 				new string[] {
-					Path.Combine(EngineDir, "Plugins/Media/AvfMedia/Source/AvfMedia/Public/")
+					"AvfMedia",
 				});
-			
+
 			PublicFrameworks.AddRange(
 				new string[] {
 					"CoreMedia",
diff --git a/Engine/Plugins/Runtime/AndroidDeviceProfileSelector/Source/AndroidDeviceProfileCommandlets/AndroidDeviceProfileCommandlets.Build.cs b/Engine/Plugins/Runtime/AndroidDeviceProfileSelector/Source/AndroidDeviceProfileCommandlets/AndroidDeviceProfileCommandlets.Build.cs
--- a/Engine/Plugins/Runtime/AndroidDeviceProfileSelector/Source/AndroidDeviceProfileCommandlets/AndroidDeviceProfileCommandlets.Build.cs
+++ b/Engine/Plugins/Runtime/AndroidDeviceProfileSelector/Source/AndroidDeviceProfileCommandlets/AndroidDeviceProfileCommandlets.Build.cs
@@ -1,21 +1,24 @@
 // Copyright Epic Games, Inc. All Rights Reserved.
-using EpicGames.Core;
-using System.IO;
 using UnrealBuildTool;
 
 public class AndroidDeviceProfileCommandlets : ModuleRules
 {
 	public AndroidDeviceProfileCommandlets(ReadOnlyTargetRules Target) : base(Target)
 	{
-		PublicIncludePaths.Add("Developer/Android/AndroidDeviceDetection/Public");
-
 		PublicDependencyModuleNames.AddRange(
 			new string[]
 			{
 				"Core",
 				"CoreUObject",
 			}
-			);
+		);
+
+		PublicIncludePathModuleNames.AddRange(
+			new string[]
+			{
+				"AndroidDeviceDetection",
+			}
+		);
 
 		PrivateDependencyModuleNames.AddRange(
 			new string[]
@@ -26,6 +29,6 @@ public class AndroidDeviceProfileCommandlets : ModuleRules
 				"AndroidDeviceProfileSelector",
 				"PIEPreviewDeviceSpecification"
 			}
-			);
+		);
 	}
 }
diff --git a/Engine/Plugins/Runtime/OpenXRViveTracker/Source/OpenXRViveTracker/OpenXRViveTracker.Build.cs b/Engine/Plugins/Runtime/OpenXRViveTracker/Source/OpenXRViveTracker/OpenXRViveTracker.Build.cs
--- a/Engine/Plugins/Runtime/OpenXRViveTracker/Source/OpenXRViveTracker/OpenXRViveTracker.Build.cs
+++ b/Engine/Plugins/Runtime/OpenXRViveTracker/Source/OpenXRViveTracker/OpenXRViveTracker.Build.cs
@@ -5,49 +5,47 @@ using System.IO;
 
 namespace UnrealBuildTool.Rules
 {
-    public class OpenXRViveTracker: ModuleRules
-    {
-        public OpenXRViveTracker(ReadOnlyTargetRules Target) 
-				: base(Target)
-        {
+	public class OpenXRViveTracker : ModuleRules
+	{
+		public OpenXRViveTracker(ReadOnlyTargetRules Target) : base(Target)
+		{
 			PublicDependencyModuleNames.AddRange(
-			   new string[]
-			   {
+				new string[]
+				{
 					"InputDevice"
-			   }
-		   );
+				}
+			 );
 
-			var EngineDir = Path.GetFullPath(Target.RelativeEnginePath);
-            PrivateIncludePaths.AddRange(
-                new string[] {
-                    EngineDir + "/Source/ThirdParty/OpenXR/include",
-					EngineDir + "/Plugins/Runtime/OpenXR/Source/OpenXRHMD/Private",
+			PrivateIncludePaths.AddRange(
+				new string[] {
+					Path.Combine(GetModuleDirectory("OpenXRHMD"), "Private"), // TODO: Adding private include path from other module
 				}
-                );
+			);
 
-            PrivateDependencyModuleNames.AddRange(
-                new string[]
-                {
-                    "Core",
-                    "CoreUObject",
-                    "Engine",
+			PrivateDependencyModuleNames.AddRange(
+				new string[]
+				{
+					"Core",
+					"CoreUObject",
+					"Engine",
 					"HeadMountedDisplay",
 					"XRBase",
-                    "InputCore",
+					"InputCore",
 					"OpenXRHMD",
 					"OpenXRInput",
 					"Slate",
 					"SlateCore",
 					"ApplicationCore"
 				}
-				);
+			);
 
-            AddEngineThirdPartyPrivateStaticDependencies(Target, "OpenXR");
+			PrivateIncludePathModuleNames.Add("OpenXR");
+			AddEngineThirdPartyPrivateStaticDependencies(Target, "OpenXR");
 
-            if (Target.bBuildEditor == true)
-            {
-                PrivateDependencyModuleNames.Add("UnrealEd");
-            }
-        }
-    }
+			if (Target.bBuildEditor == true)
+			{
+				PrivateDependencyModuleNames.Add("UnrealEd");
+			}
+		}
+	}
 }
diff --git a/Engine/Source/Developer/Apple/MetalShaderFormat/MetalShaderFormat.Build.cs b/Engine/Source/Developer/Apple/MetalShaderFormat/MetalShaderFormat.Build.cs
--- a/Engine/Source/Developer/Apple/MetalShaderFormat/MetalShaderFormat.Build.cs
+++ b/Engine/Source/Developer/Apple/MetalShaderFormat/MetalShaderFormat.Build.cs
@@ -7,7 +7,7 @@ public class MetalShaderFormat : ModuleRules
 	public MetalShaderFormat(ReadOnlyTargetRules Target) : base(Target)
 	{
 		PrivateIncludePathModuleNames.Add("TargetPlatform");
-		PublicIncludePaths.Add("Runtime/Apple/MetalRHI/Public");
+		PublicIncludePathModuleNames.Add("MetalRHI");
 
 		PrivateDependencyModuleNames.AddRange(
 			new string[] {
diff --git a/Engine/Source/Developer/CrashDebugHelper/CrashDebugHelper.Build.cs b/Engine/Source/Developer/CrashDebugHelper/CrashDebugHelper.Build.cs
--- a/Engine/Source/Developer/CrashDebugHelper/CrashDebugHelper.Build.cs
+++ b/Engine/Source/Developer/CrashDebugHelper/CrashDebugHelper.Build.cs
@@ -4,10 +4,10 @@ using UnrealBuildTool;
 
 public class CrashDebugHelper : ModuleRules
 {
-	public CrashDebugHelper( ReadOnlyTargetRules Target ) : base(Target)
+	public CrashDebugHelper(ReadOnlyTargetRules Target) : base(Target)
 	{
 		PrivateIncludePaths.AddRange(
-		new string[] {
+			new string[] {
 				"Developer/CrashDebugHelper/Private/",
 				"Developer/CrashDebugHelper/Private/Linux",
 				"Developer/CrashDebugHelper/Private/Mac",
@@ -16,13 +16,11 @@ public class CrashDebugHelper : ModuleRules
 			}
 		);
 
-		PrivateIncludePaths.Add( "ThirdParty/PLCrashReporter/PLCrashReporter/Source" );
-
 		if (Target.Type != TargetType.Game && Target.Type != TargetType.Client)
 		{
 			PublicDependencyModuleNames.AddRange(
 				new string[] {
-				"Core",
+					"Core",
 				}
 			);
 		}
@@ -31,17 +29,18 @@ public class CrashDebugHelper : ModuleRules
 			IsRedistributableOverride = true;
 			PublicDependencyModuleNames.AddRange(
 				new string[] {
-				"Core",
+					"Core",
 				}
 			);
 		}
 
-		if(Target.Platform == UnrealTargetPlatform.Win64 && Target.WindowsPlatform.bUseBundledDbgHelp)
+		if (Target.Platform == UnrealTargetPlatform.Win64 && Target.WindowsPlatform.bUseBundledDbgHelp)
 		{
 			throw new System.Exception("CrashDebugHelper uses DBGENG.DLL at runtime, which depends on a matching version of DBGHELP.DLL but cannot be redistributed. Please set WindowsPlatform.bUseBundledDbgHelp = false for this target.");
 		}
 
 		// Note: PLCrashReporter is not supported on tvOS.
+		PrivateIncludePathModuleNames.Add("PLCrashReporter");
 		if (Target.Platform == UnrealTargetPlatform.Mac || Target.Platform == UnrealTargetPlatform.IOS)
 		{
 			AddEngineThirdPartyPrivateStaticDependencies(Target, new string[]
@@ -49,5 +48,5 @@ public class CrashDebugHelper : ModuleRules
 				"PLCrashReporter"
 			});
 		}
-    }
+	}
 }
diff --git a/Engine/Source/Developer/IOS/IOSTargetPlatform/IOSTargetPlatform.Build.cs b/Engine/Source/Developer/IOS/IOSTargetPlatform/IOSTargetPlatform.Build.cs
--- a/Engine/Source/Developer/IOS/IOSTargetPlatform/IOSTargetPlatform.Build.cs
+++ b/Engine/Source/Developer/IOS/IOSTargetPlatform/IOSTargetPlatform.Build.cs
@@ -7,7 +7,7 @@ public class IOSTargetPlatform : ModuleRules
 	public IOSTargetPlatform(ReadOnlyTargetRules Target) : base(Target)
 	{
 		BinariesSubFolder = "IOS";
-		
+
 		PrivateDependencyModuleNames.AddRange(
 			new string[]
 			{
@@ -25,21 +25,11 @@ public class IOSTargetPlatform : ModuleRules
 
 		PrivateIncludePathModuleNames.AddRange(
 			new string[] {
-                "MessagingCommon",
+				"MessagingCommon",
 				"TargetDeviceServices",
 			}
 		);
 
-		//This is somehow necessary for getting iOS to build on, at least, windows. It seems like the target platform is included for cooking, and thus it requirtes a bunch of other info.
-		PublicIncludePaths.AddRange(
-			new string[]
-			{
-				"Runtime/Core/Public/Apple",
-				"Runtime/Core/Public/IOS",
-				"Runtime/Networking/Public",
-			}
-		);
-
 		if (Target.bCompileAgainstEngine)
 		{
 			PrivateDependencyModuleNames.Add("Engine");
@@ -47,7 +37,7 @@ public class IOSTargetPlatform : ModuleRules
 
 		if (Target.Platform == UnrealTargetPlatform.Mac)
 		{
-            PublicAdditionalLibraries.Add("/System/Library/PrivateFrameworks/MobileDevice.framework/Versions/Current/MobileDevice");
+			PublicAdditionalLibraries.Add("/System/Library/PrivateFrameworks/MobileDevice.framework/Versions/Current/MobileDevice");
 		}
 	}
 }
diff --git a/Engine/Source/Runtime/Apple/AudioMixerCoreAudio/AudioMixerCoreAudio.Build.cs b/Engine/Source/Runtime/Apple/AudioMixerCoreAudio/AudioMixerCoreAudio.Build.cs
--- a/Engine/Source/Runtime/Apple/AudioMixerCoreAudio/AudioMixerCoreAudio.Build.cs
+++ b/Engine/Source/Runtime/Apple/AudioMixerCoreAudio/AudioMixerCoreAudio.Build.cs
@@ -1,18 +1,15 @@
 // Copyright Epic Games, Inc. All Rights Reserved.
 
-
-
+using System.IO;
 using UnrealBuildTool;
 
-
-
 public class AudioMixerCoreAudio : ModuleRules
 {
 	public AudioMixerCoreAudio(ReadOnlyTargetRules Target) : base(Target)
 	{
 		PrivateIncludePathModuleNames.Add("TargetPlatform");
-		PublicIncludePaths.Add("Runtime/AudioMixer/Public");
-		PrivateIncludePaths.Add("Runtime/AudioMixer/Private");
+		PrivateIncludePathModuleNames.Add("AudioMixer");
+		PrivateIncludePaths.Add(Path.Combine(GetModuleDirectory("AudioMixer"), "Private")); // TODO: Adding private include path from other module
 
 
 		PrivateDependencyModuleNames.AddRange(
@@ -22,23 +19,21 @@ public class AudioMixerCoreAudio : ModuleRules
 				"AudioMixerCore",
 				"BinkAudioDecoder"
 			}
-			);
+		);
 
 		PrecompileForTargets = PrecompileTargetsType.None;
 
-        if (Target.bCompileAgainstEngine)
-        {
-            // Engine module is required for CompressedAudioInfo implementations.
-            PrivateDependencyModuleNames.Add("Engine");
-
-            AddEngineThirdPartyPrivateStaticDependencies(Target,
-                "UEOgg",
-                "Vorbis",
-                "VorbisFile"
-            );
-        }
-
+		if (Target.bCompileAgainstEngine)
+		{
+			// Engine module is required for CompressedAudioInfo implementations.
+			PrivateDependencyModuleNames.Add("Engine");
 
+			AddEngineThirdPartyPrivateStaticDependencies(Target,
+				"UEOgg",
+				"Vorbis",
+				"VorbisFile"
+			);
+		}
 
 		PublicFrameworks.AddRange(new string[]
 		{
@@ -47,10 +42,9 @@ public class AudioMixerCoreAudio : ModuleRules
 		});
 
 		if (Target.Platform == UnrealTargetPlatform.Mac)
-        {
-            PublicFrameworks.Add("AudioUnit");
-        }
-
+		{
+			PublicFrameworks.Add("AudioUnit");
+		}
 
 		PublicDefinitions.Add("WITH_OGGVORBIS=1");
 
