USD: Fix ensure when resaving an asset cache with assets that were unloaded after being deserialized.

diff --git a/Engine/Plugins/Importers/USDImporter/Source/USDClasses/Private/USDAssetCache2.cpp b/Engine/Plugins/Importers/USDImporter/Source/USDClasses/Private/USDAssetCache2.cpp
--- a/Engine/Plugins/Importers/USDImporter/Source/USDClasses/Private/USDAssetCache2.cpp
+++ b/Engine/Plugins/Importers/USDImporter/Source/USDClasses/Private/USDAssetCache2.cpp
@@ -1666,8 +1666,15 @@ void UUsdAssetCache2::Serialize(FArchive& Ar)
 				// when serializing to disk (which will happen within the Info.Serialize call below)
 				if (Ar.IsPersistent())
 				{
+					// The asset should either be loaded and in AssetStorage, or have been unloaded back into persistent storage
 					UObject* Asset = AssetStorage.FindRef(Info.Hash);
-					if (ensure(Asset))
+					ensure(Asset || PendingPersistentStorage.Contains(Info.Hash));
+
+					// We're only going to serialize it back to disk if it is already loaded of course. We could prevent overwriting
+					// on disk an already serialized asset that happens to be loaded right now since we don't really support
+					// modifying our generated assets, but I suppose this could be useful if the user just changed a couple of
+					// settings on his StaticMesh or whatever and expects them to persist.
+					if (Asset)
 					{
 						// Serialize the referenced asset into a byte buffer
 						TArray<uint8> Buffer;
@@ -1685,6 +1692,7 @@ void UUsdAssetCache2::Serialize(FArchive& Ar)
 
 						// Now we know *exactly* how big the asset is on disk
 						Info.SizeOnDiskInBytes = Buffer.Num();
+						ensure(Info.SizeOnDiskInBytes > 0);
 
 						UE_LOG(
 							LogUsd,
