USD: Fix missing light property LevelSequence tracks when opening stages with animated light attributes.

diff --git a/Engine/Plugins/Importers/USDImporter/Source/USDStage/Private/USDLevelSequenceHelper.cpp b/Engine/Plugins/Importers/USDImporter/Source/USDStage/Private/USDLevelSequenceHelper.cpp
--- a/Engine/Plugins/Importers/USDImporter/Source/USDStage/Private/USDLevelSequenceHelper.cpp
+++ b/Engine/Plugins/Importers/USDImporter/Source/USDStage/Private/USDLevelSequenceHelper.cpp
@@ -2110,7 +2110,7 @@ void FUsdLevelSequenceHelperImpl::AddPrim(UUsdPrimTwin& PrimTwin, bool bForceVis
 	{
 		AddCameraTracks(PrimTwin, UsdPrim);
 	}
-	else if (UsdPrim.HasAPI(TEXT("UsdLuxLightAPI")))
+	else if (UsdPrim.HasAPI(TEXT("LightAPI")))
 	{
 		AddLightTracks(PrimTwin, UsdPrim);
 	}
