USD: Fix LevelSequence export not tracking spawnable selections.

diff --git a/Engine/Plugins/Importers/USDImporter/Source/USDExporter/Private/LevelSequenceExporterUSD.cpp b/Engine/Plugins/Importers/USDImporter/Source/USDExporter/Private/LevelSequenceExporterUSD.cpp
--- a/Engine/Plugins/Importers/USDImporter/Source/USDExporter/Private/LevelSequenceExporterUSD.cpp
+++ b/Engine/Plugins/Importers/USDImporter/Source/USDExporter/Private/LevelSequenceExporterUSD.cpp
@@ -26,6 +26,7 @@
 #include "Editor.h"
 #include "Engine/SphereReflectionCapture.h"
 #include "EngineAnalytics.h"
+#include "EngineUtils.h"
 #include "Evaluation/MovieSceneSequenceHierarchy.h"
 #include "ISequencer.h"
 #include "ISequencerModule.h"
@@ -242,6 +243,18 @@
 					return Object;
 				}
 
+				virtual void PreDestroyObject(UObject& Object, const FGuid& BindingId, FMovieSceneSequenceIDRef TemplateID) override
+				{
+					// Don't let the FLevelSequenceEditorSpawnRegister's overload run as it will mess with our editor selection
+					if (bDestroyingJustHides)
+					{
+						FLevelSequenceSpawnRegister::PreDestroyObject(Object, BindingId, TemplateID);
+						return;
+					}
+
+					FLevelSequenceEditorSpawnRegister::PreDestroyObject(Object, BindingId, TemplateID);
+				}
+
 				virtual void DestroySpawnedObject( UObject& Object ) override
 				{
 					if ( bDestroyingJustHides )
@@ -314,6 +327,11 @@
 					}
 				}
 
+				bool HasSpawnedObject(const FGuid& BindingGuid) const
+				{
+					return SpawnableInstances.Contains(BindingGuid);
+				}
+
 				void DeleteSpawns( IMovieScenePlayer& Player )
 				{
 					bDestroyingJustHides = false;
@@ -409,9 +427,6 @@
 
 				ULevelSequenceExporterUsdOptions* ExportOptions;
 
-				// If ExportOptions->bSelectionOnly is true, this specifies the components whose bindings we should export
-				TSet<UActorComponent*> SelectedComponents;
-
 				// If ExportOptions->bSelectionOnly is true, this specifies the actors whose bindings we should export
 				TSet<AActor*> SelectedActors;
 
@@ -526,6 +541,7 @@
 			{
 				TMap< UMovieSceneSequence*, TArray<FMovieSceneSequenceID> > SequenceInstances = GetSequenceHierarchyInstances( RootSequence, Context.Sequencer.Get() );
 
+				UMovieSceneSequence* OrigRootSequence = Context.Sequencer->GetRootMovieSceneSequence();
 				Context.Sequencer->ResetToNewRootSequence( RootSequence );
 
 				// Spawn everything for this Sequence hierarchy
@@ -574,9 +590,165 @@
 					}
 				}
 
+				// Put this back to what it was before this call
+				if (OrigRootSequence)
+				{
+					Context.Sequencer->ResetToNewRootSequence(*OrigRootSequence);
+				}
+
 				return SequenceInstances;
 			}
 
+			// Collect selected actors and components before we do anything.
+			// Note that this is especially needed for spawnables since we're going to have to replicate the selection on
+			// the analogous spawnables that belong to our TempSequencer
+			void CollectEditorSelection(
+				const FLevelSequenceHidingSpawnRegister& SpawnRegister,
+				TArray<AActor*>& OutSelectedSpawnableActors,
+				TArray<AActor*>& OutSelectedNonSpawnableActors,
+				TArray<UE::MovieScene::FFixedObjectBindingID>& OutSelectedSpawnableBindings
+			)
+			{
+				OutSelectedSpawnableActors.Reset();
+				OutSelectedNonSpawnableActors.Reset();
+				OutSelectedSpawnableBindings.Reset();
+
+				// We'll start off with all actors inside OutSelectedNonSpawnableActors, and then remove the
+				// spawnables and put them on OutSelectedSpawnableActors
+				USelection* ActorSelection = GEditor->GetSelectedActors();
+				ActorSelection->GetSelectedObjects(OutSelectedNonSpawnableActors);
+
+				OutSelectedSpawnableActors.Reserve(OutSelectedNonSpawnableActors.Num());
+
+				TArray<UE::MovieScene::FFixedObjectBindingID> SpawnableBindings;
+				for (const TWeakPtr<ISequencer>& Sequencer : FLevelEditorSequencerIntegration::Get().GetSequencers())
+				{
+					if (TSharedPtr<ISequencer> PinnedSequencer = Sequencer.Pin())
+					{
+						// Components are never "spawnable", so we only need to get these bindings for actors
+						for(int32 Index = OutSelectedNonSpawnableActors.Num() - 1; Index >= 0; --Index)
+						{
+							AActor* Actor = OutSelectedNonSpawnableActors[Index];
+
+							TArray<FMovieSceneObjectBindingID> Bindings;
+							PinnedSequencer->State.FilterObjectBindings(Actor, *PinnedSequencer, &Bindings);
+
+							OutSelectedSpawnableBindings.Reserve(Bindings.Num() + OutSelectedSpawnableBindings.Num());
+							for (const FMovieSceneObjectBindingID& Binding : Bindings)
+							{
+								// Using the spawn register is an easy way of telling if a Guid is a spawnable or not,
+								// but it's more appropriate because we really only ever care about the spawnables that
+								// we have spawned on our TempSequencer
+								const bool bIsSpawnable = SpawnRegister.HasSpawnedObject(Binding.GetGuid());
+								if (!bIsSpawnable)
+								{
+									continue;
+								}
+
+								// We must keep fixed bindings only as those are kind of like "absolute file paths" and can be used
+								// to identify any of the bindings within this LevelSequence's template hierarchy. The alternative
+								// would be a relative binding ID, which would only be useful within the context of a particular
+								// subsequence of the root sequence
+								UE::MovieScene::FFixedObjectBindingID NewBinding;
+								if (Binding.IsFixedBinding())
+								{
+									NewBinding = Binding.ReinterpretAsFixed();
+								}
+								else
+								{
+									NewBinding = Binding.ResolveToFixed(Binding.GetRelativeSequenceID(), *PinnedSequencer);
+								}
+								OutSelectedSpawnableBindings.Add(NewBinding);
+								OutSelectedSpawnableActors.Add(Actor);
+
+								const int32 Count = 1;
+								const bool bAllowShrinking = false;
+								OutSelectedNonSpawnableActors.RemoveAt(Index, Count, bAllowShrinking);
+							}
+						}
+					}
+				}
+			}
+
+			void TransferSpawnableSelection(
+				ISequencer& Sequencer,
+				const TArray<UE::MovieScene::FFixedObjectBindingID>& SelectedSpawnableBindings,
+				TArray<AActor*>& OutSelectedActors
+			)
+			{
+				USelection* ActorSelection = GEditor->GetSelectedActors();
+				if (!ensure(ActorSelection))
+				{
+					return;
+				}
+				ActorSelection->Modify();
+				ActorSelection->BeginBatchSelectOperation();
+
+				OutSelectedActors.Reset();
+
+				for (const UE::MovieScene::FFixedObjectBindingID& FixedBinding : SelectedSpawnableBindings)
+				{
+					FMovieSceneObjectBindingID BindingID{FixedBinding};
+					TArrayView<TWeakObjectPtr<UObject>> BoundObjects = BindingID.ResolveBoundObjects(FixedBinding.SequenceID, Sequencer);
+					for (const TWeakObjectPtr<UObject>& BoundObject : BoundObjects)
+					{
+						if (AActor* Actor = Cast<AActor>(BoundObject.Get()))
+						{
+							ActorSelection->Select(Actor);
+
+							OutSelectedActors.Add(Actor);
+						}
+					}
+				}
+
+				const bool bNotify = true;
+				ActorSelection->EndBatchSelectOperation(bNotify);
+			}
+
+			void RestoreEditorSelection(
+				const TArray<AActor*>& SelectedNonSpawnableActors,
+				const TArray<UE::MovieScene::FFixedObjectBindingID>& SelectedSpawnableBindings
+			)
+			{
+				USelection* ActorSelection = GEditor->GetSelectedActors();
+				if (!ensure(ActorSelection))
+				{
+					return;
+				}
+				ActorSelection->Modify();
+				ActorSelection->BeginBatchSelectOperation();
+				ActorSelection->DeselectAll();
+
+				for (AActor* Actor : SelectedNonSpawnableActors)
+				{
+					ActorSelection->Select(Actor);
+				}
+
+				// Note that we're not tracking *which* sequencer produced which spawnable originally, because realistically the
+				// user will only ever have a single sequencer editing a particular LevelSequence at a time, if that
+				for (const TWeakPtr<ISequencer>& Sequencer : FLevelEditorSequencerIntegration::Get().GetSequencers())
+				{
+					if (TSharedPtr<ISequencer> PinnedSequencer = Sequencer.Pin())
+					{
+						for (const UE::MovieScene::FFixedObjectBindingID& FixedBinding : SelectedSpawnableBindings)
+						{
+							FMovieSceneObjectBindingID BindingID{FixedBinding};
+							TArrayView<TWeakObjectPtr<UObject>> BoundObjects = BindingID.ResolveBoundObjects(FixedBinding.SequenceID, *PinnedSequencer);
+							for (const TWeakObjectPtr<UObject>& BoundObject : BoundObjects)
+							{
+								if (AActor* Actor = Cast<AActor>(BoundObject.Get()))
+								{
+									ActorSelection->Select(Actor);
+								}
+							}
+						}
+					}
+				}
+
+				const bool bNotify = true;
+				ActorSelection->EndBatchSelectOperation(bNotify);
+			}
+
 			// Appends to InOutComponentBakers all of the component bakers for all components bound to MovieSceneSequence.
 			// In the process it will generate the output prims for each of these components, and keep track of them
 			// within the bakers themselves
@@ -673,34 +845,24 @@
 
 					// We always use components here because when exporting actors and components to USD we basically
 					// just ignore actors altogether and export the component attachment hierarchy instead
-					USceneComponent* BoundComponent = Cast<USceneComponent>( BoundObject );
-					if ( BoundComponent )
+					USceneComponent* BoundComponent = nullptr;
+					AActor* BoundActor = nullptr;
+					if (USceneComponent* Component = Cast<USceneComponent>(BoundObject))
 					{
-						// Only ignore component bindings if we're actually selecting some components (including this one),
-						// and our parent actor is not selected. This so that if select some components of actor X but
-						// also actor Y directly (which has component bindings) we'll export all of Y's component bindings, but
-						// only the selected component bindings of X
-						if ( Context.ExportOptions->bSelectionOnly &&
-							 Context.SelectedComponents.Num() > 0 &&
-							 !Context.SelectedActors.Contains( BoundComponent->GetOwner() ) &&
-							 !Context.SelectedComponents.Contains( BoundComponent ) )
-						{
-							continue;
-						}
+						BoundComponent = Component;
+						BoundActor = BoundComponent->GetOwner();
 					}
-					else
+					else if (AActor* Actor = Cast<AActor>(BoundObject))
 					{
-						if ( const AActor* Actor = Cast<AActor>( BoundObject ) )
-						{
-							if ( Context.ExportOptions->bSelectionOnly && !Context.SelectedActors.Contains( Actor ) )
-							{
-								continue;
-							}
+						BoundActor = Actor;
+						BoundComponent = Actor->GetRootComponent();
+					}
 
-							BoundComponent = Actor->GetRootComponent();
-						}
-					}
-					if ( !BoundComponent )
+					if (!BoundComponent || !BoundActor)
 					{
 						continue;
 					}
+					if (Context.ExportOptions->bSelectionOnly && !Context.SelectedActors.Contains(BoundActor))
+					{
+						continue;
+					}
@@ -1151,7 +1313,7 @@
 
 #if USE_USD_SDK
 	ULevelSequence* LevelSequence = Cast< ULevelSequence >( Object );
-	if ( !LevelSequence )
+	if (!GEditor || !GIsEditor || !LevelSequence)
 	{
 		return false;
 	}
@@ -1247,6 +1409,30 @@
 	// would be the exact same between all instances of the child sequence (same FGuid)
 	LevelSequenceExporterImpl::PreSpawnSpawnables( Context, *LevelSequence );
 
+	TArray<AActor*> OriginalSelectedSpawnableActors;
+	TArray<AActor*> OriginalSelectedNonSpawnableActors;
+	TArray<UE::MovieScene::FFixedObjectBindingID> OriginalSelectedSpawnableBindings;
+	LevelSequenceExporterImpl::CollectEditorSelection(
+		*SpawnRegister,
+		OriginalSelectedSpawnableActors,
+		OriginalSelectedNonSpawnableActors,
+		OriginalSelectedSpawnableBindings
+	);
+
+	if (Options->bSelectionOnly || Options->LevelExportOptions.bSelectionOnly)
+	{
+		// If the user has any spawnable selected on their Sequencer, we need to transfer that selection to our temp sequencer's spawnables
+		TArray<AActor*> TransferredSelectedSpawnableActors;
+		LevelSequenceExporterImpl::TransferSpawnableSelection(*TempSequencer, OriginalSelectedSpawnableBindings, TransferredSelectedSpawnableActors);
+
+		if (Options->bSelectionOnly)
+		{
+			Context.SelectedActors.Reset();
+			Context.SelectedActors.Append(OriginalSelectedNonSpawnableActors);
+			Context.SelectedActors.Append(TransferredSelectedSpawnableActors);
+		}
+	}
+
 	// Hide all our spawns again so we can pretend they haven't actually spawned yet
 	Context.SpawnRegister->bDestroyingJustHides = true;
 	Context.SpawnRegister->CleanUp( *Context.Sequencer );
@@ -1254,52 +1440,16 @@
 	// Capture this first because when we launch UExporter::RunAssetExportTask the CurrentFileName will change
 	const FString TargetFileName = UExporter::CurrentFilename;
 
-	// Export level if we need to
-	if ( Options && Options->bExportLevel )
-	{
-		if ( UWorld* WorldToExport = Options->Level.Get() )
-		{
-			// Come up with a file path for the level
-			FString Directory;
-			FString Filename;
-			FString Extension;
-			FPaths::Split( TargetFileName, Directory, Filename, Extension );
-			Context.LevelFilePath = FPaths::Combine( Directory, FString::Printf( TEXT( "%s.%s" ), *WorldToExport->GetFName().ToString(), *Extension ) );
-
-			Context.UsedFilePaths.Add( Context.LevelFilePath );
-
-			ULevelExporterUSDOptions* LevelOptions = GetMutableDefault<ULevelExporterUSDOptions>();
-			LevelOptions->StageOptions = Options->StageOptions;
-			LevelOptions->Inner = Options->LevelExportOptions;
-			LevelOptions->bReExportIdenticalAssets = Options->bReExportIdenticalAssets;
-			LevelOptions->bReExportIdenticalLevelsAndSequences = Options->bReExportIdenticalLevelsAndSequences;
-
-			UAssetExportTask* LevelExportTask = NewObject<UAssetExportTask>();
-			FGCObjectScopeGuard ExportTaskGuard( LevelExportTask );
-			LevelExportTask->Object = WorldToExport;
-			LevelExportTask->Options = LevelOptions;
-			LevelExportTask->Exporter = nullptr;
-			LevelExportTask->Filename = Context.LevelFilePath;
-			LevelExportTask->bSelected = LevelOptions->Inner.bSelectionOnly; // Move this as the level exporter will favor bSelected
-			LevelExportTask->bReplaceIdentical = ExportTask->bReplaceIdentical;
-			LevelExportTask->bPrompt = false;
-			LevelExportTask->bUseFileArchive = false;
-			LevelExportTask->bWriteEmptyFiles = false;
-			LevelExportTask->bAutomated = true; // Pretend this is an automated task so it doesn't pop the options dialog
-
-			UExporter::RunAssetExportTask( LevelExportTask );
-		}
-	}
-
 	// Always close all opened sequencers since it doesn't look like we're supposed to have more than one opened at a time.
 	// Without this, the ResetToNewRootSequence call may actually evaluate our opened subsequence at the playhead position of
 	// *other* sequencers, for whatever reason. We also can't call RestorePreAnimatedState on the sequence we're exporting
 	// after that either, as it hasn't stored anything yet.
 	// Additionally, it seems the sequencer will also attempt to interpolate between all opened sequencers when evaluating a track,
 	// which could affect the sequence we're exporting.
+	// Closing the Sequencers before exporting the Level is important in order to despawn their spawnables. We'll spawn our own,
+	// so we would have ended up with duplicates otherwise
 	TArray<UObject*> AssetsToReopenEditorsFor;
 	UAssetEditorSubsystem* AssetEditorSubsystem = nullptr;
-	if ( GEditor && GIsEditor )
 	{
 		AssetEditorSubsystem = GEditor->GetEditorSubsystem<UAssetEditorSubsystem>();
 		if ( AssetEditorSubsystem )
@@ -1326,25 +1476,40 @@
 		}
 	}
 
-	// Get selected objects
-	if ( GEditor && Options->bSelectionOnly )
+	// Export level if we need to
+	if ( Options && Options->bExportLevel )
 	{
-		USelection* ComponentSelection= GEditor->GetSelectedComponents();
-		TArray<UActorComponent*> Components;
-		ComponentSelection->GetSelectedObjects( Components );
-		Context.SelectedComponents = TSet<UActorComponent*>{ Components };
-
-		USelection* ActorSelection = GEditor->GetSelectedActors();
-		TArray<AActor*> Actors;
-		ActorSelection->GetSelectedObjects( Actors );
-		Context.SelectedActors = TSet<AActor*>{ Actors };
-		for ( const UActorComponent* Component : Components )
+		if ( UWorld* WorldToExport = Options->Level.Get() )
 		{
-			// UE will ensure that the actor selection artificially includes all owners of all selected components, so
-			// we can't tell if an actor is intentionally selected or not. To provide some control, let's make it so that
-			// if components of an actor are selected we'll select only those particular components. If the user wants
-			// the whole actor they can just select the actor itself and none of its components
-			Context.SelectedActors.Remove( Component->GetOwner() );
+			// Come up with a file path for the level
+			FString Directory;
+			FString Filename;
+			FString Extension;
+			FPaths::Split( TargetFileName, Directory, Filename, Extension );
+			Context.LevelFilePath = FPaths::Combine( Directory, FString::Printf( TEXT( "%s.%s" ), *WorldToExport->GetFName().ToString(), *Extension ) );
+
+			Context.UsedFilePaths.Add( Context.LevelFilePath );
+
+			ULevelExporterUSDOptions* LevelOptions = GetMutableDefault<ULevelExporterUSDOptions>();
+			LevelOptions->StageOptions = Options->StageOptions;
+			LevelOptions->Inner = Options->LevelExportOptions;
+			LevelOptions->bReExportIdenticalAssets = Options->bReExportIdenticalAssets;
+			LevelOptions->bReExportIdenticalLevelsAndSequences = Options->bReExportIdenticalLevelsAndSequences;
+
+			UAssetExportTask* LevelExportTask = NewObject<UAssetExportTask>();
+			FGCObjectScopeGuard ExportTaskGuard( LevelExportTask );
+			LevelExportTask->Object = WorldToExport;
+			LevelExportTask->Options = LevelOptions;
+			LevelExportTask->Exporter = nullptr;
+			LevelExportTask->Filename = Context.LevelFilePath;
+			LevelExportTask->bSelected = LevelOptions->Inner.bSelectionOnly; // Move this as the level exporter will favor bSelected
+			LevelExportTask->bReplaceIdentical = ExportTask->bReplaceIdentical;
+			LevelExportTask->bPrompt = false;
+			LevelExportTask->bUseFileArchive = false;
+			LevelExportTask->bWriteEmptyFiles = false;
+			LevelExportTask->bAutomated = true; // Pretend this is an automated task so it doesn't pop the options dialog
+
+			UExporter::RunAssetExportTask( LevelExportTask );
 		}
 	}
 
@@ -1355,11 +1520,13 @@
 	TempSequencer->SetPlaybackStatus( EMovieScenePlayerStatus::Stopped );
 	TempSequencer->Close();
 
-	if ( GEditor && GIsEditor && AssetEditorSubsystem )
+	if (AssetEditorSubsystem)
 	{
-		AssetEditorSubsystem->OpenEditorForAssets( AssetsToReopenEditorsFor );
+		AssetEditorSubsystem->OpenEditorForAssets(AssetsToReopenEditorsFor);
 	}
 
+	LevelSequenceExporterImpl::RestoreEditorSelection(OriginalSelectedNonSpawnableActors, OriginalSelectedSpawnableBindings);
+
 	// Analytics
 	{
 		bool bAutomated = ExportTask ? ExportTask->bAutomated : false;
