USD: Always spawn skeletal mesh components as movable or else we get some warnings when building physics assets.

diff --git a/Engine/Plugins/Importers/USDImporter/Source/USDSchemas/Private/USDGeomXformableTranslator.cpp b/Engine/Plugins/Importers/USDImporter/Source/USDSchemas/Private/USDGeomXformableTranslator.cpp
--- a/Engine/Plugins/Importers/USDImporter/Source/USDSchemas/Private/USDGeomXformableTranslator.cpp
+++ b/Engine/Plugins/Importers/USDImporter/Source/USDSchemas/Private/USDGeomXformableTranslator.cpp
@@ -553,7 +553,8 @@ USceneComponent* FUsdGeomXformableTranslator::CreateComponentsEx( TOptional< TSu
 		const bool bParentIsStationary = Context->ParentComponent && Context->ParentComponent->Mobility == EComponentMobility::Stationary;
 
 		// Don't call SetMobility as it would trigger a reregister, queuing unnecessary rhi commands since this is a brand new component
-		if (bIsAssociating || bParentIsMovable || UsdUtils::IsAnimated(Prim))
+		// Always have movable Skeletal mesh components or else we get some warnings when building physics assets
+		if (bIsAssociating || bParentIsMovable || SceneComponent->IsA<USkeletalMeshComponent>() || UsdUtils::IsAnimated(Prim))
 		{
 			SceneComponent->Mobility = EComponentMobility::Movable;
 		}
