USD: Improve duplication behavior with instanceables.

Fix copying of instanceable prims not copying the full data;
Show warning when edit would have been refused by USD due to being on/inside instance proxies.

diff --git a/Engine/Plugins/Importers/USDImporter/Source/USDClasses/Public/USDProjectSettings.h b/Engine/Plugins/Importers/USDImporter/Source/USDClasses/Public/USDProjectSettings.h
--- a/Engine/Plugins/Importers/USDImporter/Source/USDClasses/Public/USDProjectSettings.h
+++ b/Engine/Plugins/Importers/USDImporter/Source/USDClasses/Public/USDProjectSettings.h
@@ -16,6 +16,14 @@ enum class EUsdSaveDialogBehavior : uint8
 	ShowPrompt
 };
 
+UENUM( BlueprintType )
+enum class EUsdEditInInstanceBehavior : uint8
+{
+	Ignore,
+	RemoveInstanceable,
+	ShowPrompt
+};
+
 // USDImporter and defaultconfig here so this ends up at DefaultUSDImporter.ini in the editor, and is sent to the
 // packaged game as well
 UCLASS(config=USDImporter, defaultconfig, meta=(DisplayName=USDImporter), MinimalAPI)
@@ -75,6 +83,10 @@ public:
 	UPROPERTY( config, EditAnywhere, Category = "USD|Dialogs" )
 	bool bShowOverriddenOpinionsWarning = true;
 
+	// Whether to show the warning dialog when authoring opinions inside an instance or instance proxy
+	UPROPERTY( config, EditAnywhere, Category = "USD|Dialogs" )
+	EUsdEditInInstanceBehavior EditInInstanceableBehavior = EUsdEditInInstanceBehavior::ShowPrompt;
+
 	// Whether to show a warning whenever the "Duplicate All Local Layer Specs" option is picked, and the duplicated
 	// prim has some specs outside the local layer stack that will not be duplicated.
 	UPROPERTY( config, EditAnywhere, Category = "USD|Dialogs" )
diff --git a/Engine/Plugins/Importers/USDImporter/Source/USDStageEditor/Private/SUSDStageTreeView.cpp b/Engine/Plugins/Importers/USDImporter/Source/USDStageEditor/Private/SUSDStageTreeView.cpp
--- a/Engine/Plugins/Importers/USDImporter/Source/USDStageEditor/Private/SUSDStageTreeView.cpp
+++ b/Engine/Plugins/Importers/USDImporter/Source/USDStageEditor/Private/SUSDStageTreeView.cpp
@@ -4,6 +4,7 @@
 
 #include "SUSDStageEditorStyle.h"
 #include "UnrealUSDWrapper.h"
+#include "USDAttributeUtils.h"
 #include "USDConversionUtils.h"
 #include "USDDuplicateType.h"
 #include "USDLayerUtils.h"
@@ -985,12 +986,19 @@ void SUsdStageTreeView::OnAddChildPrim()
 
 void SUsdStageTreeView::OnCutPrim()
 {
+	TArray<FUsdPrimViewModelRef> MySelectedItems = GetSelectedItems();
+	for (FUsdPrimViewModelRef SelectedItem : MySelectedItems)
+	{
+		if (UsdUtils::NotifyIfInstanceProxy(SelectedItem->UsdPrim))
+		{
+			return;
+		}
+	}
+
 	FScopedTransaction Transaction( LOCTEXT( "CutPrimTransaction", "Cut prims" ) );
 
 	UE::FSdfChangeBlock Block;
 
-	TArray< FUsdPrimViewModelRef > MySelectedItems = GetSelectedItems();
-
 	TArray<UE::FUsdPrim> Prims;
 	Prims.Reserve( MySelectedItems.Num() );
 
@@ -1030,12 +1038,19 @@ void SUsdStageTreeView::OnPastePrim()
 		return;
 	}
 
+	TArray<FUsdPrimViewModelRef> MySelectedItems = GetSelectedItems();
+	for (FUsdPrimViewModelRef SelectedItem : MySelectedItems)
+	{
+		if (UsdUtils::NotifyIfInstanceProxy(SelectedItem->UsdPrim))
+		{
+			return;
+		}
+	}
+
 	FScopedTransaction Transaction( LOCTEXT( "PastePrimTransaction", "Paste prims" ) );
 
 	UE::FSdfChangeBlock Block;
 
-	TArray< FUsdPrimViewModelRef > MySelectedItems = GetSelectedItems();
-
 	TArray<UE::FUsdPrim> ParentPrims;
 
 	// This happens when right-clicking the background area without selecting any prim
@@ -1071,11 +1086,16 @@ void SUsdStageTreeView::OnDuplicatePrim( EUsdDuplicateType DuplicateType )
 		return;
 	}
 
-	FScopedTransaction Transaction( LOCTEXT( "DuplicatePrimTransaction", "Duplicate prims" ) );
-
-	UE::FSdfChangeBlock Block;
+	TArray<FUsdPrimViewModelRef> MySelectedItems = GetSelectedItems();
+	for (FUsdPrimViewModelRef SelectedItem : MySelectedItems)
+	{
+		if (UsdUtils::NotifyIfInstanceProxy(SelectedItem->UsdPrim))
+		{
+			return;
+		}
+	}
 
-	TArray< FUsdPrimViewModelRef > MySelectedItems = GetSelectedItems();
+	FScopedTransaction Transaction( LOCTEXT( "DuplicatePrimTransaction", "Duplicate prims" ) );
 
 	TArray<UE::FUsdPrim> Prims;
 	Prims.Reserve( MySelectedItems.Num() );
@@ -1093,12 +1113,19 @@ void SUsdStageTreeView::OnDuplicatePrim( EUsdDuplicateType DuplicateType )
 
 void SUsdStageTreeView::OnDeletePrim()
 {
+	TArray<FUsdPrimViewModelRef> MySelectedItems = GetSelectedItems();
+	for (FUsdPrimViewModelRef SelectedItem : MySelectedItems)
+	{
+		if (UsdUtils::NotifyIfInstanceProxy(SelectedItem->UsdPrim))
+		{
+			return;
+		}
+	}
+
 	FScopedTransaction Transaction( LOCTEXT( "DeletePrimTransaction", "Delete prims" ) );
 
 	UE::FSdfChangeBlock Block;
 
-	TArray< FUsdPrimViewModelRef > MySelectedItems = GetSelectedItems();
-
 	for ( FUsdPrimViewModelRef SelectedItem : MySelectedItems )
 	{
 		UsdUtils::RemoveAllLocalPrimSpecs( SelectedItem->UsdPrim );
@@ -1149,6 +1176,10 @@ void SUsdStageTreeView::OnAddReference()
 		return;
 	}
 	UE::FUsdPrim Referencer = MySelectedItems[ 0 ]->UsdPrim;
+	if (UsdUtils::NotifyIfInstanceProxy(Referencer))
+	{
+		return;
+	}
 
 	// This transaction is important as adding a reference may trigger the creation of new unreal assets, which need to be
 	// destroyed if we spam undo afterwards. Undoing won't remove the actual reference from the stage yet though, sadly...
@@ -1183,9 +1214,16 @@ void SUsdStageTreeView::OnAddReference()
 
 void SUsdStageTreeView::OnClearReferences()
 {
-	FScopedTransaction Transaction( LOCTEXT( "ClearReferenceTransaction", "Clear references to USD layers" ) );
+	TArray<FUsdPrimViewModelRef> MySelectedItems = GetSelectedItems();
+	for (FUsdPrimViewModelRef SelectedItem : MySelectedItems)
+	{
+		if (UsdUtils::NotifyIfInstanceProxy(SelectedItem->UsdPrim))
+		{
+			return;
+		}
+	}
 
-	TArray< FUsdPrimViewModelRef > MySelectedItems = GetSelectedItems();
+	FScopedTransaction Transaction( LOCTEXT( "ClearReferenceTransaction", "Clear references to USD layers" ) );
 
 	for ( FUsdPrimViewModelRef SelectedItem : MySelectedItems )
 	{
@@ -1223,6 +1261,10 @@ void SUsdStageTreeView::OnAddPayload()
 		return;
 	}
 	UE::FUsdPrim Referencer = MySelectedItems[ 0 ]->UsdPrim;
+	if (UsdUtils::NotifyIfInstanceProxy(Referencer))
+	{
+		return;
+	}
 
 	// This transaction is important as adding a payload may trigger the creation of new unreal assets, which need to be
 	// destroyed if we spam undo afterwards. Undoing won't remove the actual payload from the stage yet though, sadly...
@@ -1257,9 +1299,16 @@ void SUsdStageTreeView::OnAddPayload()
 
 void SUsdStageTreeView::OnClearPayloads()
 {
-	FScopedTransaction Transaction( LOCTEXT( "ClearPayloadTransaction", "Clear payloads to USD layers" ) );
+	TArray<FUsdPrimViewModelRef> MySelectedItems = GetSelectedItems();
+	for (FUsdPrimViewModelRef SelectedItem : MySelectedItems)
+	{
+		if (UsdUtils::NotifyIfInstanceProxy(SelectedItem->UsdPrim))
+		{
+			return;
+		}
+	}
 
-	TArray< FUsdPrimViewModelRef > MySelectedItems = GetSelectedItems();
+	FScopedTransaction Transaction( LOCTEXT( "ClearPayloadTransaction", "Clear payloads to USD layers" ) );
 
 	for ( FUsdPrimViewModelRef SelectedItem : MySelectedItems )
 	{
@@ -1269,13 +1318,20 @@ void SUsdStageTreeView::OnClearPayloads()
 
 void SUsdStageTreeView::OnApplySchema( FName SchemaName )
 {
+	TArray<FUsdPrimViewModelRef> MySelectedItems = GetSelectedItems();
+	for (FUsdPrimViewModelRef SelectedItem : MySelectedItems)
+	{
+		if (UsdUtils::NotifyIfInstanceProxy(SelectedItem->UsdPrim))
+		{
+			return;
+		}
+	}
+
 	FScopedTransaction Transaction( FText::Format(
 		LOCTEXT( "ApplySchemaTransaction", "Apply the '{0}' schema onto selected prims" ),
 		FText::FromName( SchemaName )
 	));
 
-	TArray< FUsdPrimViewModelRef > MySelectedItems = GetSelectedItems();
-
 	UE::FSdfChangeBlock Block;
 
 	for ( FUsdPrimViewModelRef SelectedItem : MySelectedItems )
@@ -1286,13 +1342,20 @@ void SUsdStageTreeView::OnApplySchema( FName SchemaName )
 
 void SUsdStageTreeView::OnRemoveSchema( FName SchemaName )
 {
+	TArray<FUsdPrimViewModelRef> MySelectedItems = GetSelectedItems();
+	for (FUsdPrimViewModelRef SelectedItem : MySelectedItems)
+	{
+		if (UsdUtils::NotifyIfInstanceProxy(SelectedItem->UsdPrim))
+		{
+			return;
+		}
+	}
+
 	FScopedTransaction Transaction( FText::Format(
 		LOCTEXT( "RemoveSchemaTransaction", "Remove the '{0}' schema from selected prims" ),
 		FText::FromName( SchemaName )
 	) );
 
-	TArray< FUsdPrimViewModelRef > MySelectedItems = GetSelectedItems();
-
 	UE::FSdfChangeBlock Block;
 
 	for ( FUsdPrimViewModelRef SelectedItem : MySelectedItems )
@@ -1630,27 +1693,38 @@ void SUsdStageTreeView::OnPrimNameCommitted( const FUsdPrimViewModelRef& ViewMod
 	const bool bRenamingExistingPrim = ViewModel->bIsRenamingExistingPrim;
 	ViewModel->bIsRenamingExistingPrim = false;
 
-	if ( InPrimName.IsEmptyOrWhitespace() )
+	// Escaped out of initially setting a prim name
+	TFunction<void()> CancelInput = [&ViewModel, this]()
 	{
-		// Escaped out of initially setting a prim name
-		if ( !ViewModel->UsdPrim )
+		if (!ViewModel->UsdPrim)
 		{
 			if (FUsdPrimViewModel* Parent = ViewModel->ParentItem)
 			{
-				ViewModel->ParentItem->Children.Remove( ViewModel );
+				ViewModel->ParentItem->Children.Remove(ViewModel);
 			}
 			else
 			{
-				RootItems.Remove( ViewModel );
+				RootItems.Remove(ViewModel);
 			}
 
 			RequestTreeRefresh();
 		}
+	};
+
+	if ( InPrimName.IsEmptyOrWhitespace() )
+	{
+		CancelInput();
 		return;
 	}
 
 	if ( bRenamingExistingPrim )
 	{
+		if (UsdUtils::NotifyIfInstanceProxy(ViewModel->UsdPrim))
+		{
+			CancelInput();
+			return;
+		}
+
 		FScopedTransaction Transaction( LOCTEXT( "RenamePrimTransaction", "Rename a prim" ) );
 
 		// e.g. "/Root/OldPrim/"
@@ -1683,6 +1757,15 @@ void SUsdStageTreeView::OnPrimNameCommitted( const FUsdPrimViewModelRef& ViewMod
 	}
 	else
 	{
+		if (FUsdPrimViewModel* ParentModel = ViewModel->ParentItem)
+		{
+			if (UsdUtils::NotifyIfInstanceProxy(ParentModel->UsdPrim))
+			{
+				CancelInput();
+				return;
+			}
+		}
+
 		FScopedTransaction Transaction( LOCTEXT( "AddPrimTransaction", "Add a new prim" ) );
 
 		ViewModel->DefinePrim( *InPrimName.ToString() );
diff --git a/Engine/Plugins/Importers/USDImporter/Source/USDStageEditor/Private/Widgets/SUSDIntegrationsPanel.cpp b/Engine/Plugins/Importers/USDImporter/Source/USDStageEditor/Private/Widgets/SUSDIntegrationsPanel.cpp
--- a/Engine/Plugins/Importers/USDImporter/Source/USDStageEditor/Private/Widgets/SUSDIntegrationsPanel.cpp
+++ b/Engine/Plugins/Importers/USDImporter/Source/USDStageEditor/Private/Widgets/SUSDIntegrationsPanel.cpp
@@ -143,6 +143,12 @@ TSharedRef< SWidget > SUsdIntegrationsPanelRow::GenerateWidgetForColumn( const F
 
 				FSimpleDelegate PropertyChanged = FSimpleDelegate::CreateLambda( [AttributeCopy]()
 				{
+					UE::FUsdPrim Prim = AttributeCopy->GetPrim();
+					if (UsdUtils::NotifyIfInstanceProxy(Prim))
+					{
+						return;
+					}
+
 					UUsdIntegrationsPanelPropertyDummy* Dummy = GetMutableDefault<UUsdIntegrationsPanelPropertyDummy>();
 
 					FString Path = ( Dummy && Dummy->AnimBPProperty ) ? Dummy->AnimBPProperty->GetPathName() : FString{};
@@ -203,6 +209,12 @@ TSharedRef< SWidget > SUsdIntegrationsPanelRow::GenerateWidgetForColumn( const F
 
 				FSimpleDelegate PropertyChanged = FSimpleDelegate::CreateLambda( [AttributeCopy]()
 				{
+					UE::FUsdPrim Prim = AttributeCopy->GetPrim();
+					if (UsdUtils::NotifyIfInstanceProxy(Prim))
+					{
+						return;
+					}
+
 					UUsdIntegrationsPanelPropertyDummy* Dummy = GetMutableDefault<UUsdIntegrationsPanelPropertyDummy>();
 
 					FString Path = ( Dummy && Dummy->ControlRigProperty ) ? Dummy->ControlRigProperty->GetPathName() : FString{};
@@ -267,6 +279,12 @@ TSharedRef< SWidget > SUsdIntegrationsPanelRow::GenerateWidgetForColumn( const F
 				})
 				.OnCheckStateChanged_Lambda( [AttributeCopy]( ECheckBoxState NewValue )
 				{
+					UE::FUsdPrim Prim = AttributeCopy->GetPrim();
+					if (UsdUtils::NotifyIfInstanceProxy(Prim))
+					{
+						return;
+					}
+
 					FScopedTransaction Transaction(
 						FText::Format( LOCTEXT( "ControlRigReduceKeysTransaction", "Changed ControlRig reduce keys property to '{0}'" ),
 							FText::FromString( NewValue == ECheckBoxState::Checked ? TEXT( "true" ) : TEXT( "false" ) )
@@ -293,6 +311,12 @@ TSharedRef< SWidget > SUsdIntegrationsPanelRow::GenerateWidgetForColumn( const F
 				.Font( FAppStyle::GetFontStyle( UE::SUsdIntergrationsPanel::Private::NormalFont ) )
 				.OnValueCommitted_Lambda( [AttributeCopy]( float NewValue, ETextCommit::Type CommitType )
 				{
+					UE::FUsdPrim Prim = AttributeCopy->GetPrim();
+					if (UsdUtils::NotifyIfInstanceProxy(Prim))
+					{
+						return;
+					}
+
 					FScopedTransaction Transaction(
 						FText::Format( LOCTEXT( "ControlRigReduceToleranceTransaction", "Changed ControlRig reduce tolerance to '{0}'" ),
 							NewValue
@@ -358,6 +382,12 @@ TSharedRef< SWidget > SUsdIntegrationsPanelRow::GenerateWidgetForColumn( const F
 			})
 			.OnValueChanged_Lambda( [AttributeCopy]( SLiveLinkSubjectRepresentationPicker::FLiveLinkSourceSubjectRole NewValue )
 			{
+				UE::FUsdPrim Prim = AttributeCopy->GetPrim();
+				if (UsdUtils::NotifyIfInstanceProxy(Prim))
+				{
+					return;
+				}
+
 				FScopedUsdAllocs UsdAllocs;
 
 				std::string UsdValue = UnrealToUsd::ConvertString( *( NewValue.Subject.ToString() ) ).Get();
@@ -400,6 +430,12 @@ TSharedRef< SWidget > SUsdIntegrationsPanelRow::GenerateWidgetForColumn( const F
 				})
 				.OnCheckStateChanged_Lambda( [AttributeCopy, AttributeName]( ECheckBoxState NewState )
 				{
+					UE::FUsdPrim Prim = AttributeCopy->GetPrim();
+					if (UsdUtils::NotifyIfInstanceProxy(Prim))
+					{
+						return;
+					}
+
 					FScopedTransaction Transaction(
 						FText::Format( LOCTEXT( "IntegrationBoolTransaction", "Changed attribute '{0}' to " ),
 							FText::FromName( AttributeName ),
diff --git a/Engine/Plugins/Importers/USDImporter/Source/USDStageEditorViewModels/Private/USDObjectFieldViewModel.cpp b/Engine/Plugins/Importers/USDImporter/Source/USDStageEditorViewModels/Private/USDObjectFieldViewModel.cpp
--- a/Engine/Plugins/Importers/USDImporter/Source/USDStageEditorViewModels/Private/USDObjectFieldViewModel.cpp
+++ b/Engine/Plugins/Importers/USDImporter/Source/USDStageEditorViewModels/Private/USDObjectFieldViewModel.cpp
@@ -181,6 +181,11 @@ void FUsdObjectFieldsViewModel::SetFieldValue( const FString& FieldName, const U
 		}
 		else if (UE::FUsdPrim UsdPrim = UsdStage.GetPrimAtPath(UE::FSdfPath(*ObjectPath)))
 		{
+			if (UsdUtils::NotifyIfInstanceProxy(UsdPrim))
+			{
+				return;
+			}
+
 			FScopedUsdAllocs UsdAllocs;
 
 			pxr::UsdPrim PxrUsdPrim{UsdPrim};
@@ -231,6 +236,11 @@ void FUsdObjectFieldsViewModel::SetFieldValue( const FString& FieldName, const U
 				pxr::UsdStageRefPtr PxrUsdStage{UsdStage};
 				if (pxr::UsdPrim Prim = PxrUsdStage->GetPrimAtPath(PxrPrimPath))
 				{
+					if (UsdUtils::NotifyIfInstanceProxy(Prim))
+					{
+						return;
+					}
+
 					if (pxr::UsdProperty Property = Prim.GetProperty(PropertyName))
 					{
 						pxr::TfToken FieldNameToken = UnrealToUsd::ConvertToken(*FieldName).Get();
diff --git a/Engine/Plugins/Importers/USDImporter/Source/USDStageEditorViewModels/Private/USDPrimViewModel.cpp b/Engine/Plugins/Importers/USDImporter/Source/USDStageEditorViewModels/Private/USDPrimViewModel.cpp
--- a/Engine/Plugins/Importers/USDImporter/Source/USDStageEditorViewModels/Private/USDPrimViewModel.cpp
+++ b/Engine/Plugins/Importers/USDImporter/Source/USDStageEditorViewModels/Private/USDPrimViewModel.cpp
@@ -130,6 +130,15 @@ void FUsdPrimViewModel::FillChildren()
 void FUsdPrimViewModel::RefreshData( bool bRefreshChildren )
 {
 #if USE_USD_SDK
+	// Before we fully abort due to an invalid prim, first check the case that we just need to get a "refreshed prim"
+	// for the same path. This is important for example when setting/clearing the instanceable metadata: If this prim used
+	// to be an instance proxy and now our parent is not an instance anymore, the prim will become "invalid", but that's
+	// just because the instance doesn't exist anymore: The analogous non-instance-proxy prim still may exist on the stage
+	if (!UsdPrim && UsdStage)
+	{
+		UsdPrim = UsdStage.GetPrimAtPath(UsdPrim.GetPrimPath());
+	}
+
 	if ( !UsdPrim )
 	{
 		return;
diff --git a/Engine/Plugins/Importers/USDImporter/Source/USDUtilities/Private/USDAttributeUtils.cpp b/Engine/Plugins/Importers/USDImporter/Source/USDUtilities/Private/USDAttributeUtils.cpp
--- a/Engine/Plugins/Importers/USDImporter/Source/USDUtilities/Private/USDAttributeUtils.cpp
+++ b/Engine/Plugins/Importers/USDImporter/Source/USDUtilities/Private/USDAttributeUtils.cpp
@@ -24,14 +24,13 @@
 	#include "pxr/usd/pcp/cache.h"
 	#include "pxr/usd/pcp/primIndex.h"
 	#include "pxr/usd/pcp/propertyIndex.h"
-	#include "pxr/usd/usd/usdFileFormat.h"
-
 	#include "pxr/usd/sdf/layer.h"
 	#include "pxr/usd/sdf/layerUtils.h"
 	#include "pxr/usd/sdf/spec.h"
 	#include "pxr/usd/usd/attribute.h"
 	#include "pxr/usd/usd/editContext.h"
 	#include "pxr/usd/usd/stage.h"
+	#include "pxr/usd/usd/usdFileFormat.h"
 #include "USDIncludesEnd.h"
 #endif // #if USE_USD_SDK
 
@@ -284,6 +283,214 @@ void UsdUtils::NotifyIfOverriddenOpinion(const pxr::UsdProperty& Property)
 		}
 	}
 }
+
+bool UsdUtils::NotifyIfInstanceProxy(const pxr::UsdPrim& Prim)
+{
+	if (!Prim)
+	{
+		return false;
+	}
+
+	TFunction<void()> RemoveInstanceables = [Prim = UE::FUsdPrim{Prim}]()
+	{
+		UE::FUsdStage Stage = Prim.GetStage();
+		if (!Stage)
+		{
+			return;
+		}
+
+		// We have to track paths and not prims directly, because these will actually be the instance
+		// proxies in case we have an instance parent, and even if the parents are not instances anymore
+		// trying to author to these instance proxies would be an error. If we call GetPrimAtPath *after*
+		// clearing the parent instanceable, then we get a regular prim
+		TArray<UE::FSdfPath> LeafToRoot;
+
+		UE::FUsdPrim Iter = Prim;
+		while (Iter && !Iter.IsPseudoRoot())
+		{
+			LeafToRoot.Add(Iter.GetPrimPath());
+			Iter = Iter.GetParent();
+		}
+
+		// Annoyingly we have to break from root downwards, as otherwise we'd be trying
+		// to author inside instances ourselves!
+		// Note: We also can't use a change block here, because we could have nested instanceables,
+		// and we need USD to fully respond to the outer instanceable being cleared before it lets
+		// us clear the inner one
+		for (int32 Index = LeafToRoot.Num() - 1; Index >= 0; --Index)
+		{
+			const UE::FUsdPrim& SomePrim = Stage.GetPrimAtPath(LeafToRoot[Index]);
+			if (SomePrim.IsInstanceable())
+			{
+				// We force false here instead of just clearing the authored opinion because the instanceable=true
+				// opinion may come from a referenced layer or some other place we can't just clear on our current
+				// edit target
+				SomePrim.SetInstanceable(false);
+			}
+		}
+	};
+
+	if (Prim.IsInstanceProxy())
+	{
+		const UUsdProjectSettings* Settings = GetDefault<UUsdProjectSettings>();
+		if (Settings)
+		{
+			switch (Settings->EditInInstanceableBehavior)
+			{
+				case EUsdEditInInstanceBehavior::Ignore:
+				{
+					UE_LOG(
+						LogUsd,
+						Log,
+						TEXT("Ignoring edits to prim '%s' as it is an instance proxy"),
+						*UsdToUnreal::ConvertPath(Prim.GetPrimPath())
+					);
+					return true;
+					break;
+				}
+				case EUsdEditInInstanceBehavior::RemoveInstanceable:
+				{
+					UE_LOG(LogUsd, Log, TEXT("Removing all instanceable flags from ancestors of prim '%s'"), *UsdToUnreal::ConvertPath(Prim.GetPrimPath()));
+					RemoveInstanceables();
+
+					// We shouldn't be instanceable now, so we can probably author whatever we wanted to author
+					return false;
+					break;
+				}
+				default:
+				case EUsdEditInInstanceBehavior::ShowPrompt:
+				{
+					const FText Text = LOCTEXT("AuthoringInsideInstanceText", "USD: Authoring inside instance");
+
+					FString FirstInstancePath;
+					{
+						FScopedUsdAllocs Allocs;
+						pxr::UsdPrim Iter = Prim;
+						while (Iter && !Iter.IsPseudoRoot())
+						{
+							if (Iter.IsInstance())
+							{
+								break;
+							}
+							Iter = Iter.GetParent();
+						}
+						FirstInstancePath = UsdToUnreal::ConvertPath(Iter.GetPrimPath());
+					}
+
+					const FText SubText = FText::Format(
+						LOCTEXT(
+							"AuthoringInsideInstanceSubText",
+							"Trying to author an opinion at or below prim:\n\n{0}\n\nThis prim is an instance proxy, and so cannot be "
+							"edited directly.\n\nIf you wish to modify just this prim, you can remove the instanceable flag for this prim and try again. If you wish to "
+							"modify all instances, please edit the prims referenced by prim '{1}' directly."
+						),
+						FText::FromString(UsdToUnreal::ConvertPath(Prim.GetPrimPath())),
+						FText::FromString(FirstInstancePath)
+					);
+
+					UE_LOG(LogUsd, Warning, TEXT("%s"), *SubText.ToString().Replace(TEXT("\n\n"), TEXT(" ")));
+
+					static TWeakPtr<SNotificationItem> Notification;
+
+					FNotificationInfo Toast(Text);
+					Toast.SubText = SubText;
+					Toast.Image = FCoreStyle::Get().GetBrush(TEXT("MessageLog.Warning"));
+					Toast.CheckBoxText = LOCTEXT("DontAskAgain", "Don't prompt again");
+					Toast.bUseLargeFont = false;
+					Toast.bFireAndForget = false;
+					Toast.FadeOutDuration = 0.0f;
+					Toast.ExpireDuration = 0.0f;
+					Toast.bUseThrobber = false;
+					Toast.bUseSuccessFailIcons = false;
+					Toast.ButtonDetails.Emplace(
+						LOCTEXT("OverridenOpinionMessageRemove", "Remove instanceable flag"),
+						FText::GetEmpty(),
+						FSimpleDelegate::CreateLambda(
+							[RemoveInstanceables]()
+							{
+								RemoveInstanceables();
+
+								if (TSharedPtr<SNotificationItem> PinnedNotification = Notification.Pin())
+								{
+									PinnedNotification->SetCompletionState(SNotificationItem::CS_Success);
+									PinnedNotification->ExpireAndFadeout();
+								}
+
+								if (UUsdProjectSettings* Settings = GetMutableDefault<UUsdProjectSettings>())
+								{
+									// We'll only set this to "ShowPrompt" if the checkbox to "Don't prompt again" is unchecked
+									if (Settings->EditInInstanceableBehavior != EUsdEditInInstanceBehavior::ShowPrompt)
+									{
+										Settings->EditInInstanceableBehavior = EUsdEditInInstanceBehavior::RemoveInstanceable;
+									}
+									Settings->SaveConfig();
+								}
+							}
+						)
+					);
+					Toast.ButtonDetails.Emplace(
+						LOCTEXT("OverridenOpinionMessageCancel", "Cancel"),
+						FText::GetEmpty(),
+						FSimpleDelegate::CreateLambda(
+							[]()
+							{
+								if (TSharedPtr<SNotificationItem> PinnedNotification = Notification.Pin())
+								{
+									PinnedNotification->SetCompletionState(SNotificationItem::CS_Success);
+									PinnedNotification->ExpireAndFadeout();
+								}
+								if (UUsdProjectSettings* Settings = GetMutableDefault<UUsdProjectSettings>())
+								{
+									if (Settings->EditInInstanceableBehavior != EUsdEditInInstanceBehavior::ShowPrompt)
+									{
+										Settings->EditInInstanceableBehavior = EUsdEditInInstanceBehavior::Ignore;
+									}
+									Settings->SaveConfig();
+								}
+							}
+						)
+					);
+					// This is flipped because the default checkbox message is "Don't prompt again"
+					Toast.CheckBoxState = Settings->EditInInstanceableBehavior == EUsdEditInInstanceBehavior::ShowPrompt
+											  ? ECheckBoxState::Unchecked
+											  : ECheckBoxState::Checked;
+					Toast.CheckBoxStateChanged = FOnCheckStateChanged::CreateStatic(
+						[](ECheckBoxState NewState)
+						{
+							if (UUsdProjectSettings* Settings = GetMutableDefault<UUsdProjectSettings>())
+							{
+								// This is flipped because the default checkbox message is "Don't prompt again"
+								Settings->EditInInstanceableBehavior = (NewState == ECheckBoxState::Unchecked)
+																		   ? EUsdEditInInstanceBehavior::ShowPrompt
+																		   : EUsdEditInInstanceBehavior::Ignore;	// Either would do here, we have
+																													// to press one of the buttons to
+																													// close the prompt, which will
+																													// set the right one
+								Settings->SaveConfig();
+							}
+						}
+					);
+
+					// Only show one at a time
+					if (!Notification.IsValid())
+					{
+						Notification = FSlateNotificationManager::Get().AddNotification(Toast);
+					}
+
+					if (TSharedPtr<SNotificationItem> PinnedNotification = Notification.Pin())
+					{
+						PinnedNotification->SetCompletionState(SNotificationItem::CS_Pending);
+					}
+
+					return true;
+					break;
+				}
+			}
+		}
+	}
+
+	return false;
+}
 #endif // #if USE_USD_SDK
 
 #undef LOCTEXT_NAMESPACE
\ No newline at end of file
diff --git a/Engine/Plugins/Importers/USDImporter/Source/USDUtilities/Private/USDConversionUtils.cpp b/Engine/Plugins/Importers/USDImporter/Source/USDUtilities/Private/USDConversionUtils.cpp
--- a/Engine/Plugins/Importers/USDImporter/Source/USDUtilities/Private/USDConversionUtils.cpp
+++ b/Engine/Plugins/Importers/USDImporter/Source/USDUtilities/Private/USDConversionUtils.cpp
@@ -2122,6 +2122,31 @@ bool UsdUtils::CopyPrims( const TArray<UE::FUsdPrim>& Prims )
 		return false;
 	}
 
+	// USD will retain instances and prototypes even when flattening, which is not what we want
+	// so let's disable instancing on our temp stage before we ask it to flatten.
+	// Note how we traverse the entire masked stage here, because we also need to handle the case
+	// where the prim we're duplicating is not instanceable, but has instanceable children
+	TArray<pxr::SdfPath> OldInstanceablePrims;
+	if (TempStage->GetPrototypes().size() > 0)
+	{
+		pxr::UsdEditContext Context{TempStage, TempStage->GetSessionLayer()};
+
+		pxr::UsdPrimRange PrimRange(TempStage->GetPseudoRoot());
+		for (pxr::UsdPrimRange::iterator PrimRangeIt = PrimRange.begin(); PrimRangeIt != PrimRange.end(); ++PrimRangeIt)
+		{
+			if (PrimRangeIt->IsPseudoRoot())
+			{
+				continue;
+			}
+
+			if (PrimRangeIt->HasAuthoredInstanceable())
+			{
+				PrimRangeIt->SetInstanceable(false);
+				OldInstanceablePrims.Add(PrimRangeIt->GetPrimPath());
+			}
+		}
+	}
+
 	const bool bAddSourceFileComment = false;
 	pxr::SdfLayerRefPtr FlattenedLayer = TempStage->Flatten( bAddSourceFileComment );
 	if ( !FlattenedLayer )
@@ -2129,6 +2154,17 @@ bool UsdUtils::CopyPrims( const TArray<UE::FUsdPrim>& Prims )
 		return false;
 	}
 
+	// We may had to force instanceable=false on the prims we duplicated in order to get our session layer
+	// opinion to disable instancing. We don't want those prims to come out with "instanceable=false" on the
+	// flattened copy though, so here we clear that opinion
+	for (const pxr::SdfPath& Path : OldInstanceablePrims)
+	{
+		if (pxr::SdfPrimSpecHandle Spec = FlattenedLayer->GetPrimAtPath(Path))
+		{
+			Spec->ClearInstanceable();
+		}
+	}
+
 	ClipboardRoot->Clear();
 
 	TSet<FString> UsedNames;
@@ -2399,12 +2435,48 @@ TArray<UE::FSdfPath> UsdUtils::DuplicatePrims( const TArray<UE::FUsdPrim>& Prims
 			return Result;
 		}
 
+		// USD will retain instances and prototypes even when flattening, which is not what we want
+		// so let's disable instancing on our temp stage before we ask it to flatten.
+		// Note how we travere the entire masked stage here, because we also need to handle the case
+		// where the prim we're duplicating is not instanceable, but has instanceable children
+		TArray<pxr::SdfPath> OldInstanceablePrims;
+		if (TempStage->GetPrototypes().size() > 0)
+		{
+			pxr::UsdEditContext Context{TempStage, TempStage->GetSessionLayer()};
+
+			pxr::UsdPrimRange PrimRange(TempStage->GetPseudoRoot());
+			for (pxr::UsdPrimRange::iterator PrimRangeIt = PrimRange.begin(); PrimRangeIt != PrimRange.end(); ++PrimRangeIt)
+			{
+				if (PrimRangeIt->IsPseudoRoot())
+				{
+					continue;
+				}
+
+				if (PrimRangeIt->HasAuthoredInstanceable())
+				{
+					PrimRangeIt->SetInstanceable(false);
+					OldInstanceablePrims.Add(PrimRangeIt->GetPrimPath());
+				}
+			}
+		}
+
 		const bool bAddSourceFileComment = false;
 		FlattenedLayer = TempStage->Flatten( bAddSourceFileComment );
 		if ( !FlattenedLayer )
 		{
 			return Result;
 		}
+
+		// We may had to force instanceable=false on the prims we duplicated in order to get our session layer
+		// opinion to disable instancing. We don't want those prims to come out with "instanceable=false" on the
+		// flattened copy though, so here we clear that opinion
+		for (const pxr::SdfPath& Path : OldInstanceablePrims)
+		{
+			if (pxr::SdfPrimSpecHandle Spec = FlattenedLayer->GetPrimAtPath(Path))
+			{
+				Spec->ClearInstanceable();
+			}
+		}
 	}
 
 	for ( int32 Index = 0; Index < Prims.Num(); ++Index )
@@ -2425,8 +2497,11 @@ TArray<UE::FSdfPath> UsdUtils::DuplicatePrims( const TArray<UE::FUsdPrim>& Prims
 		std::unordered_set<pxr::SdfLayerHandle, pxr::TfHash> LayersThatWillBeAffected;
 		LayersThatWillBeAffected.reserve( PrimSpecs.size() );
 
-		for ( const pxr::SdfPrimSpecHandle& Spec : PrimSpecs )
+		pxr::SdfPath TargetPath = UsdPrim.GetPrimPath();
+		for (int32 SpecIndex = PrimSpecs.size() - 1; SpecIndex >= 0; --SpecIndex)
 		{
+			const pxr::SdfPrimSpecHandle& Spec = PrimSpecs[SpecIndex];
+
 			// For whatever reason sometimes there are invalid specs in the layer stack, so we need to be careful
 			if ( !Spec )
 			{
@@ -2434,7 +2509,13 @@ TArray<UE::FSdfPath> UsdUtils::DuplicatePrims( const TArray<UE::FUsdPrim>& Prims
 			}
 
 			pxr::SdfPath SpecPath = Spec->GetPath();
-			if ( !SpecPath.IsPrimPath() )
+
+			// Skip specs that have a different path than the actual prim path. The only way this could happen
+			// is if the prim is referencing this particular path, and if we were to duplicate this spec
+			// we'd essentially end up flattening the referenced prim over the new duplicate prim, which
+			// is not what we want. We'll already get the fact that "prim references this other prim" by copying
+			// the spec at the actual TargetPath however
+			if (!SpecPath.IsPrimPath() || SpecPath.StripAllVariantSelections() != TargetPath)
 			{
 				continue;
 			}
@@ -2536,12 +2617,47 @@ TArray<UE::FSdfPath> UsdUtils::DuplicatePrims( const TArray<UE::FUsdPrim>& Prims
 					continue;
 				}
 
-				if ( !pxr::SdfCopySpec( SpecLayerHandle, SpecPath, SpecLayerHandle, NewSpecPath ) )
+				pxr::SdfShouldCopyValueFn ShouldCopyValue =
+					[](pxr::SdfSpecType SpecType,
+					   const pxr::TfToken& Field,
+					   const pxr::SdfLayerHandle& SrcLayer,
+					   const pxr::SdfPath& SrcPath,
+					   bool FieldInSrc,
+					   const pxr::SdfLayerHandle& DstLayer,
+					   const pxr::SdfPath& DstPath,
+					   bool FieldInDst,
+					   boost::optional<pxr::VtValue>* ValueToCopy) -> bool
 				{
-					UE_LOG( LogUsd, Warning, TEXT( "Failed to copy spec from path '%s' onto path '%s' within layer '%s'" ),
-						*UsdToUnreal::ConvertPath( SpecPath ),
-						*UsdToUnreal::ConvertPath( NewSpecPath ),
-						*UsdToUnreal::ConvertString( SpecLayerHandle->GetIdentifier() )
+					// Only copy a field over if it has a value. Otherwise it seems to clear the destination spec
+					// for nothing
+					return FieldInSrc;
+				};
+
+				pxr::SdfShouldCopyChildrenFn ShouldCopyChildren =
+					[](const pxr::TfToken& ChildrenField,
+					   const pxr::SdfLayerHandle& SrcLayer,
+					   const pxr::SdfPath& SrcPath,
+					   bool FieldInSrc,
+					   const pxr::SdfLayerHandle& DstLayer,
+					   const pxr::SdfPath& DstPath,
+					   bool FieldInDst,
+					   boost::optional<pxr::VtValue>* SrcChildren,
+					   boost::optional<pxr::VtValue>* DstChildren) -> bool
+				{
+					return true;
+				};
+
+				// We use the advanced version of SdfCopySpec here as otherwise the default behavior is to fully clear
+				// the destination spec before copying stuff, and we may want to copy multiple specs overwriting each other
+				if (!pxr::SdfCopySpec(SpecLayerHandle, SpecPath, SpecLayerHandle, NewSpecPath, ShouldCopyValue, ShouldCopyChildren))
+				{
+					UE_LOG(
+						LogUsd,
+						Warning,
+						TEXT("Failed to copy spec from path '%s' onto path '%s' within layer '%s'"),
+						*UsdToUnreal::ConvertPath(SpecPath),
+						*UsdToUnreal::ConvertPath(NewSpecPath),
+						*UsdToUnreal::ConvertString(SpecLayerHandle->GetIdentifier())
 					);
 				}
 			}
diff --git a/Engine/Plugins/Importers/USDImporter/Source/USDUtilities/Private/USDIntegrationUtils.cpp b/Engine/Plugins/Importers/USDImporter/Source/USDUtilities/Private/USDIntegrationUtils.cpp
--- a/Engine/Plugins/Importers/USDImporter/Source/USDUtilities/Private/USDIntegrationUtils.cpp
+++ b/Engine/Plugins/Importers/USDImporter/Source/USDUtilities/Private/USDIntegrationUtils.cpp
@@ -161,6 +161,11 @@ void UnrealToUsd::ConvertLiveLinkProperties( const UActorComponent* InComponent,
 
 	UE::FSdfChangeBlock ChangeBlock;
 
+	if (UsdUtils::NotifyIfInstanceProxy(InOutPrim))
+	{
+		return;
+	}
+
 	// Skeletal LiveLink case
 	if ( const USkeletalMeshComponent* SkeletalComponent = Cast<USkeletalMeshComponent>( InComponent ) )
 	{
diff --git a/Engine/Plugins/Importers/USDImporter/Source/USDUtilities/Private/USDLightConversion.cpp b/Engine/Plugins/Importers/USDImporter/Source/USDUtilities/Private/USDLightConversion.cpp
--- a/Engine/Plugins/Importers/USDImporter/Source/USDUtilities/Private/USDLightConversion.cpp
+++ b/Engine/Plugins/Importers/USDImporter/Source/USDUtilities/Private/USDLightConversion.cpp
@@ -389,6 +389,11 @@ bool UnrealToUsd::ConvertLightComponent( const ULightComponentBase& LightCompone
 		return false;
 	}
 
+	if (UsdUtils::NotifyIfInstanceProxy(Prim))
+	{
+		return false;
+	}
+
 	if ( pxr::UsdAttribute Attr = LightAPI.CreateIntensityAttr() )
 	{
 		Attr.Set<float>( LightComponent.Intensity, UsdTimeCode );
@@ -443,6 +448,11 @@ bool UnrealToUsd::ConvertDirectionalLightComponent( const UDirectionalLightCompo
 		return false;
 	}
 
+	if (UsdUtils::NotifyIfInstanceProxy(Prim))
+	{
+		return false;
+	}
+
 	if ( pxr::UsdAttribute Attr = Light.CreateAngleAttr() )
 	{
 		Attr.Set<float>( LightComponent.LightSourceAngle, UsdTimeCode );
@@ -470,6 +480,11 @@ bool UnrealToUsd::ConvertRectLightComponent( const URectLightComponent& LightCom
 		return false;
 	}
 
+	if (UsdUtils::NotifyIfInstanceProxy(Prim))
+	{
+		return false;
+	}
+
 	FUsdStageInfo StageInfo( Prim.GetStage() );
 
 	// Disk light
@@ -571,6 +586,11 @@ bool UnrealToUsd::ConvertPointLightComponent( const UPointLightComponent& LightC
 		return false;
 	}
 
+	if (UsdUtils::NotifyIfInstanceProxy(Prim))
+	{
+		return false;
+	}
+
 	FUsdStageInfo StageInfo( Prim.GetStage() );
 
 	if ( pxr::UsdAttribute Attr = Light.CreateRadiusAttr() )
@@ -645,6 +665,11 @@ bool UnrealToUsd::ConvertSkyLightComponent( const USkyLightComponent& LightCompo
 		return false;
 	}
 
+	if (UsdUtils::NotifyIfInstanceProxy(Prim))
+	{
+		return false;
+	}
+
 #if WITH_EDITORONLY_DATA
 	FUsdStageInfo StageInfo( Prim.GetStage() );
 
@@ -685,6 +710,11 @@ bool UnrealToUsd::ConvertSpotLightComponent( const USpotLightComponent& LightCom
 		return false;
 	}
 
+	if (UsdUtils::NotifyIfInstanceProxy(Prim))
+	{
+		return false;
+	}
+
 	if ( pxr::UsdAttribute ConeAngleAttr = ShapingAPI.CreateShapingConeAngleAttr() )
 	{
 		ConeAngleAttr.Set<float>( LightComponent.OuterConeAngle, UsdTimeCode );
diff --git a/Engine/Plugins/Importers/USDImporter/Source/USDUtilities/Private/USDPrimConversion.cpp b/Engine/Plugins/Importers/USDImporter/Source/USDUtilities/Private/USDPrimConversion.cpp
--- a/Engine/Plugins/Importers/USDImporter/Source/USDUtilities/Private/USDPrimConversion.cpp
+++ b/Engine/Plugins/Importers/USDImporter/Source/USDUtilities/Private/USDPrimConversion.cpp
@@ -1441,6 +1441,11 @@ bool UnrealToUsd::ConvertCameraComponent( const UCineCameraComponent& CameraComp
 		return false;
 	}
 
+	if (UsdUtils::NotifyIfInstanceProxy(Prim))
+	{
+		return false;
+	}
+
 	FUsdStageInfo StageInfo( Prim.GetStage() );
 
 	if ( pxr::UsdAttribute Attr = GeomCamera.CreateFocalLengthAttr() )
@@ -1958,6 +1963,11 @@ bool UnrealToUsd::ConvertSceneComponent( const pxr::UsdStageRefPtr& Stage, const
 		return false;
 	}
 
+	if (UsdUtils::NotifyIfInstanceProxy(UsdPrim))
+	{
+		return false;
+	}
+
 	FScopedUsdAllocs UsdAllocs;
 
 	// Transform
@@ -2082,6 +2092,11 @@ bool UnrealToUsd::ConvertHierarchicalInstancedStaticMeshComponent(
 		return false;
 	}
 
+	if (UsdUtils::NotifyIfInstanceProxy(UsdPrim))
+	{
+		return false;
+	}
+
 	UsdStageRefPtr Stage = UsdPrim.GetStage();
 	FUsdStageInfo StageInfo{ Stage };
 
@@ -2483,6 +2498,11 @@ bool UnrealToUsd::ConvertXformable( const FTransform& RelativeTransform, pxr::Us
 		return false;
 	}
 
+	if (UsdUtils::NotifyIfInstanceProxy(UsdPrim))
+	{
+		return false;
+	}
+
 	FUsdStageInfo StageInfo( UsdPrim.GetStage() );
 	pxr::GfMatrix4d UsdTransform = UnrealToUsd::ConvertTransform( StageInfo, RelativeTransform );
 
@@ -2570,6 +2590,11 @@ bool UnrealToUsd::ConvertInstancedFoliageActor( const AInstancedFoliageActor& Ac
 		return false;
 	}
 
+	if (UsdUtils::NotifyIfInstanceProxy(UsdPrim))
+	{
+		return false;
+	}
+
 	UsdStageRefPtr Stage = UsdPrim.GetStage();
 	FUsdStageInfo StageInfo{ Stage };
 
@@ -3086,6 +3111,11 @@ UnrealToUsd::FPropertyTrackWriter UnrealToUsd::CreatePropertyTrackWriter( const
 	pxr::UsdStageRefPtr UsdStage = UsdPrim.GetStage();
 	FUsdStageInfo StageInfo{ UsdStage };
 
+	if (UsdUtils::NotifyIfInstanceProxy(Prim))
+	{
+		return Result;
+	}
+
 	pxr::UsdAttribute Attr;
 	{
 		pxr::SdfChangeBlock ChangeBlock;
diff --git a/Engine/Plugins/Importers/USDImporter/Source/USDUtilities/Private/USDShadeConversion.cpp b/Engine/Plugins/Importers/USDImporter/Source/USDUtilities/Private/USDShadeConversion.cpp
--- a/Engine/Plugins/Importers/USDImporter/Source/USDUtilities/Private/USDShadeConversion.cpp
+++ b/Engine/Plugins/Importers/USDImporter/Source/USDUtilities/Private/USDShadeConversion.cpp
@@ -2881,6 +2881,11 @@ bool UsdUtils::MarkMaterialPrimWithWorldSpaceNormals( const UE::FUsdPrim& Materi
 		return false;
 	}
 
+	if (UsdUtils::NotifyIfInstanceProxy(MaterialPrim))
+	{
+		return false;
+	}
+
 	const bool bCustom = true;
 	pxr::UsdAttribute Attr = UsdPrim.CreateAttribute( UnrealIdentifiers::WorldSpaceNormals, pxr::SdfValueTypeNames->Bool, bCustom );
 	if ( !Attr )
diff --git a/Engine/Plugins/Importers/USDImporter/Source/USDUtilities/Private/USDSkeletalDataConversion.cpp b/Engine/Plugins/Importers/USDImporter/Source/USDUtilities/Private/USDSkeletalDataConversion.cpp
--- a/Engine/Plugins/Importers/USDImporter/Source/USDUtilities/Private/USDSkeletalDataConversion.cpp
+++ b/Engine/Plugins/Importers/USDImporter/Source/USDUtilities/Private/USDSkeletalDataConversion.cpp
@@ -3290,6 +3290,11 @@ bool UnrealToUsd::ConvertControlRigSection(
 		return false;
 	}
 
+	if (UsdUtils::NotifyIfInstanceProxy(OutSkelAnimPrim))
+	{
+		return false;
+	}
+
 	FUsdStageInfo StageInfo( UsdStage );
 
 	double StartTime = FPlatformTime::Cycles64();
diff --git a/Engine/Plugins/Importers/USDImporter/Source/USDUtilities/Public/USDAttributeUtils.h b/Engine/Plugins/Importers/USDImporter/Source/USDUtilities/Public/USDAttributeUtils.h
--- a/Engine/Plugins/Importers/USDImporter/Source/USDUtilities/Public/USDAttributeUtils.h
+++ b/Engine/Plugins/Importers/USDImporter/Source/USDUtilities/Public/USDAttributeUtils.h
@@ -8,12 +8,14 @@
 
 #if USE_USD_SDK
 PXR_NAMESPACE_OPEN_SCOPE
+	class UsdPrim;
 	class UsdProperty;
 PXR_NAMESPACE_CLOSE_SCOPE
 #endif // USE_USD_SDK
 
 namespace UE
 {
+	class FUsdPrim;
 	class FUsdAttribute;
 }
 
@@ -38,6 +40,13 @@ namespace UsdUtils
 	 */
 #if USE_USD_SDK
 	USDUTILITIES_API void NotifyIfOverriddenOpinion( const pxr::UsdProperty& Property );
+
+	/**
+	 * Generally we are not allowed to author opinions inside instanceables. Our actual results when doing that will range from
+	 * USD coding errors to seeming like nothing has happened on the stage.
+	 * This will show a warning and return true if Prim is an instance, letting us avoid those scenarios
+	 */
+	USDUTILITIES_API bool NotifyIfInstanceProxy(const pxr::UsdPrim& Prim);
 #endif // USE_USD_SDK
 }
 
diff --git a/Engine/Plugins/Importers/USDImporter/Source/UnrealUSDWrapper/Private/UsdWrappers/UsdPrim.cpp b/Engine/Plugins/Importers/USDImporter/Source/UnrealUSDWrapper/Private/UsdWrappers/UsdPrim.cpp
--- a/Engine/Plugins/Importers/USDImporter/Source/UnrealUSDWrapper/Private/UsdWrappers/UsdPrim.cpp
+++ b/Engine/Plugins/Importers/USDImporter/Source/UnrealUSDWrapper/Private/UsdWrappers/UsdPrim.cpp
@@ -645,4 +645,144 @@ namespace UE
 		return false;
 #endif // #if USE_USD_SDK
 	}
-}
+
+	bool FUsdPrim::IsInstanceable() const
+	{
+#if USE_USD_SDK
+		FScopedUsdAllocs Allocs;
+		return Impl->PxrUsdPrim.Get().IsInstanceable();
+#else
+		return false;
+#endif	  // USE_USD_SDK
+	}
+
+	bool FUsdPrim::SetInstanceable(bool bInstanceable) const
+	{
+#if USE_USD_SDK
+		FScopedUsdAllocs Allocs;
+		return Impl->PxrUsdPrim.Get().SetInstanceable(bInstanceable);
+#else
+		return false;
+#endif	  // USE_USD_SDK
+	}
+
+	bool FUsdPrim::ClearInstanceable() const
+	{
+#if USE_USD_SDK
+		FScopedUsdAllocs Allocs;
+		return Impl->PxrUsdPrim.Get().ClearInstanceable();
+#else
+		return false;
+#endif	  // USE_USD_SDK
+	}
+
+	bool FUsdPrim::HasAuthoredInstanceable() const
+	{
+#if USE_USD_SDK
+		FScopedUsdAllocs Allocs;
+		return Impl->PxrUsdPrim.Get().HasAuthoredInstanceable();
+#else
+		return false;
+#endif	  // USE_USD_SDK
+	}
+
+	bool FUsdPrim::IsInstance() const
+	{
+#if USE_USD_SDK
+		FScopedUsdAllocs Allocs;
+		return Impl->PxrUsdPrim.Get().IsInstance();
+#else
+		return false;
+#endif	  // USE_USD_SDK
+	}
+
+	bool FUsdPrim::IsInstanceProxy() const
+	{
+#if USE_USD_SDK
+		FScopedUsdAllocs Allocs;
+		return Impl->PxrUsdPrim.Get().IsInstanceProxy();
+#else
+		return false;
+#endif	  // USE_USD_SDK
+	}
+
+	bool FUsdPrim::IsPrototype() const
+	{
+#if USE_USD_SDK
+		FScopedUsdAllocs Allocs;
+		return Impl->PxrUsdPrim.Get().IsPrototype();
+#else
+		return false;
+#endif	  // USE_USD_SDK
+	}
+
+	bool FUsdPrim::IsInPrototype() const
+	{
+#if USE_USD_SDK
+		FScopedUsdAllocs Allocs;
+		return Impl->PxrUsdPrim.Get().IsInPrototype();
+#else
+		return false;
+#endif	  // USE_USD_SDK
+	}
+
+	FUsdPrim FUsdPrim::GetPrototype() const
+	{
+#if USE_USD_SDK
+		FScopedUsdAllocs Allocs;
+		return UE::FUsdPrim{Impl->PxrUsdPrim.Get().GetPrototype()};
+#else
+		return {};
+#endif	  // USE_USD_SDK
+	}
+
+	FUsdPrim FUsdPrim::GetPrimInPrototype() const
+	{
+#if USE_USD_SDK
+		FScopedUsdAllocs Allocs;
+		return UE::FUsdPrim{Impl->PxrUsdPrim.Get().GetPrimInPrototype()};
+#else
+		return {};
+#endif	  // USE_USD_SDK
+	}
+
+	TArray<FUsdPrim> FUsdPrim::GetInstances() const
+	{
+		TArray<FUsdPrim> Result;
+
+#if USE_USD_SDK
+		FScopedUsdAllocs Allocs;
+
+		std::vector<pxr::UsdPrim> Instances = Impl->PxrUsdPrim.Get().GetInstances();
+		Result.Reserve(Instances.size());
+
+		for (const pxr::UsdPrim& Instance : Instances)
+		{
+			Result.Emplace(Instance);
+		}
+#endif	  // USE_USD_SDK
+
+		return Result;
+	}
+
+	bool FUsdPrim::IsPrototypePath(const FSdfPath& Path)
+	{
+#if USE_USD_SDK
+		FScopedUsdAllocs Allocs;
+		return pxr::UsdPrim::IsPrototypePath(Path);
+#else
+		return false;
+#endif	  // USE_USD_SDK
+	}
+
+	bool FUsdPrim::IsPathInPrototype(const FSdfPath& Path)
+	{
+#if USE_USD_SDK
+		FScopedUsdAllocs Allocs;
+		return pxr::UsdPrim::IsPathInPrototype(Path);
+#else
+		return false;
+#endif	  // USE_USD_SDK
+	}
+
+}	 // namespace UE
diff --git a/Engine/Plugins/Importers/USDImporter/Source/UnrealUSDWrapper/Public/UsdWrappers/UsdPrim.h b/Engine/Plugins/Importers/USDImporter/Source/UnrealUSDWrapper/Public/UsdWrappers/UsdPrim.h
--- a/Engine/Plugins/Importers/USDImporter/Source/UnrealUSDWrapper/Public/UsdWrappers/UsdPrim.h
+++ b/Engine/Plugins/Importers/USDImporter/Source/UnrealUSDWrapper/Public/UsdWrappers/UsdPrim.h
@@ -140,6 +140,20 @@ namespace UE
 		FUsdAttribute GetAttribute(const TCHAR* AttrName) const;
 		bool HasAttribute(const TCHAR* AttrName) const;
 
+		bool IsInstanceable() const;
+		bool SetInstanceable(bool bInstanceable) const;
+		bool ClearInstanceable() const;
+		bool HasAuthoredInstanceable() const;
+		bool IsInstance() const;
+		bool IsInstanceProxy() const;
+		bool IsPrototype() const;
+		bool IsInPrototype() const;
+		FUsdPrim GetPrototype() const;
+		FUsdPrim GetPrimInPrototype() const;
+		TArray<FUsdPrim> GetInstances() const;
+		static bool IsPrototypePath(const FSdfPath& Path);
+		static bool IsPathInPrototype(const FSdfPath& Path);
+
 	private:
 		TUniquePtr< Internal::FUsdPrimImpl > Impl;
 	};
