USD: Update instances when sources and prototype prims are modified.

diff --git a/Engine/Plugins/Importers/USDImporter/Source/USDSchemas/Private/USDInfoCache.cpp b/Engine/Plugins/Importers/USDImporter/Source/USDSchemas/Private/USDInfoCache.cpp
--- a/Engine/Plugins/Importers/USDImporter/Source/USDSchemas/Private/USDInfoCache.cpp
+++ b/Engine/Plugins/Importers/USDImporter/Source/USDSchemas/Private/USDInfoCache.cpp
@@ -20,6 +20,7 @@
 
 #if USE_USD_SDK
 #include "USDIncludesStart.h"
+	#include "pxr/usd/pcp/layerStack.h"
 	#include "pxr/usd/sdf/path.h"
 	#include "pxr/usd/usd/collectionAPI.h"
 	#include "pxr/usd/usd/prim.h"
@@ -31,6 +32,7 @@
 	#include "pxr/usd/usdGeom/xform.h"
 	#include "pxr/usd/usdShade/materialBindingAPI.h"
 	#include "pxr/usd/usdSkel/root.h"
+	#include "pxr/usd/usd/primRange.h"
 #include "USDIncludesEnd.h"
 #endif // USE_USD_SDK
 
@@ -1069,6 +1071,100 @@ namespace UE::USDInfoCacheImpl::Private
 		}
 	}
 
+	// Returns the paths to all prims on the same local layer stack, that are used as sources for composition
+	// arcs that are non-root (i.e. the arcs that are either reference, payload, inherits, etc.).
+	// In other words, "instanceable composition arcs from local prims"
+	TSet<UE::FSdfPath> GetLocalNonRootCompositionArcSourcePaths(const pxr::UsdPrim& UsdPrim)
+	{
+		TSet<UE::FSdfPath> Result;
+
+		if (!UsdPrim)
+		{
+			return Result;
+		}
+
+		pxr::PcpLayerStackRefPtr RootLayerStack;
+
+		pxr::UsdPrimCompositionQuery PrimCompositionQuery = pxr::UsdPrimCompositionQuery(UsdPrim);
+		std::vector<pxr::UsdPrimCompositionQueryArc> Arcs = PrimCompositionQuery.GetCompositionArcs();
+		Result.Reserve(Arcs.size());
+		for (const pxr::UsdPrimCompositionQueryArc& Arc : Arcs)
+		{
+			pxr::PcpNodeRef TargetNode = Arc.GetTargetNode();
+
+			if (Arc.GetArcType() == pxr::PcpArcTypeRoot)
+			{
+				RootLayerStack = TargetNode.GetLayerStack();
+			}
+			// We use this function to collect aux/main prim links for instanceables, and we don't have
+			// to track instanceable arcs to outside the local layer stack because those don't generate
+			// source prims on the stage that the user could edit anyway!
+			else if (TargetNode.GetLayerStack() == RootLayerStack)
+			{
+				Result.Add(UE::FSdfPath{Arc.GetTargetPrimPath()});
+			}
+		}
+
+		return Result;
+	}
+
+	void RegisterInstanceableAuxPrims(
+		const pxr::UsdPrim& UsdPrim,
+		FUsdSchemaTranslationContext& Context,
+		FUsdInfoCache::FUsdInfoCacheImpl& Impl
+	)
+	{
+		TRACE_CPUPROFILER_EVENT_SCOPE(UE::USDInfoCacheImpl::Private::RegisterInstanceableAuxPrims);
+		FScopedUsdAllocs Allocs;
+
+		pxr::UsdStageRefPtr Stage = UsdPrim.GetStage();
+		for(const pxr::UsdPrim& Prototype : Stage->GetPrototypes())
+		{
+			if(!Prototype)
+			{
+				continue;
+			}
+
+			// Step into every instance of this prototype on the stage
+			for (const pxr::UsdPrim& Instance : Prototype.GetInstances())
+			{
+				UE::FSdfPath InstancePath{Instance.GetPrimPath()};
+
+				// Adding a dependency on the prototype directly is interesting, even though we currently don't display those.
+				// and the prototype paths are mostly transient
+				Impl.RegisterAuxiliaryPrims(InstancePath, {UE::FSdfPath{Prototype.GetPrimPath()}});
+
+				// Really what we want is to find the source prim that generated this prototype though. Instances always work
+				// through some kind of composition arc, so here we collect all references/payloads/inherits/specializes/etc.
+				TSet<UE::FSdfPath> SourcePaths = GetLocalNonRootCompositionArcSourcePaths(Instance);
+				Impl.RegisterAuxiliaryPrims(InstancePath, SourcePaths);
+
+				// Here we'll traverse the entire subtree of the instance
+				pxr::UsdPrimRange PrimRange(Instance, pxr::UsdTraverseInstanceProxies());
+				for (pxr::UsdPrimRange::iterator InstanceChildIt = ++PrimRange.begin(); InstanceChildIt != PrimRange.end(); ++InstanceChildIt)
+				{
+					pxr::SdfPath SdfChildPrimPath = InstanceChildIt->GetPrimPath();
+					UE::FSdfPath ChildPrimPath{SdfChildPrimPath};
+
+					// Register a dependency from child prim to the analogue prim within the prototype itself
+					Impl.RegisterAuxiliaryPrims(ChildPrimPath, {UE::FSdfPath{InstanceChildIt->GetPrimInPrototype().GetPrimPath()}});
+
+					// Register a dependency from child prim to analogue prims on the sources used for the instance.
+					// We have to do some path surgery to discover what the analogue paths on the source prims are though
+					pxr::SdfPath RelativeChildPath = SdfChildPrimPath.MakeRelativePath(InstancePath);
+					for (const UE::FSdfPath& SourcePath : SourcePaths)
+					{
+						pxr::SdfPath ChildOnSourcePath = pxr::SdfPath{SourcePath}.AppendPath(RelativeChildPath);
+						if (pxr::UsdPrim ChildOnSource = Stage->GetPrimAtPath(ChildOnSourcePath))
+						{
+							Impl.RegisterAuxiliaryPrims(ChildPrimPath, {UE::FSdfPath{ChildOnSourcePath}});
+						}
+					}
+				}
+			}
+		}
+	}
+
 	void FindValidGeometryCacheRoot(const pxr::UsdPrim& UsdPrim, FUsdSchemaTranslationContext& Context, FUsdInfoCache::FUsdInfoCacheImpl& Impl, UE::UsdInfoCache::Private::EGeometryCachePrimState& OutState)
 	{
 		TRACE_CPUPROFILER_EVENT_SCOPE(FindValidGeometryCacheRoot);
@@ -1632,6 +1728,8 @@ void FUsdInfoCache::RebuildCacheForSubtree( const UE::FUsdPrim& Prim, FUsdSchema
 			Registry
 		);
 
+		UE::USDInfoCacheImpl::Private::RegisterInstanceableAuxPrims(UsdPrim, Context, *ImplPtr);
+
 		UE::USDInfoCacheImpl::Private::CollectMaterialSlotCounts(
 			*ImplPtr,
 			TempSubtreeSlots,
