USD: Fix material assignments when a geom subset is expected to fallback to its parent Mesh's material binding.

diff --git a/Engine/Plugins/Importers/USDImporter/Source/USDUtilities/Private/USDGeomMeshConversion.cpp b/Engine/Plugins/Importers/USDImporter/Source/USDUtilities/Private/USDGeomMeshConversion.cpp
--- a/Engine/Plugins/Importers/USDImporter/Source/USDUtilities/Private/USDGeomMeshConversion.cpp
+++ b/Engine/Plugins/Importers/USDImporter/Source/USDUtilities/Private/USDGeomMeshConversion.cpp
@@ -1816,6 +1816,7 @@ UsdUtils::FUsdPrimMaterialAssignmentInfo UsdUtils::GetPrimMaterialAssignments(
 
 		if ( bProvideMaterialIndices )
 		{
+			// Note how we're defaulting to slot zero here, which is our "main assignment"
 			Result.MaterialIndices.SetNumZeroed( NumFaces );
 		}
 	}
@@ -1835,222 +1836,276 @@ UsdUtils::FUsdPrimMaterialAssignmentInfo UsdUtils::GetPrimMaterialAssignments(
 
 	FString MeshPrimPath = UsdToUnreal::ConvertPath( UsdPrim.GetPath() );
 
-	// Priority 1: Material is an unreal asset
-	if ( RenderContext == UnrealIdentifiers::Unreal )
-	{
-		// Priority 1.1: unreal rendercontext material prim
-		pxr::UsdShadeMaterialBindingAPI BindingAPI( UsdPrim );
-		if ( BindingAPI )
-		{
-			if ( pxr::UsdShadeMaterial ShadeMaterial = BindingAPI.ComputeBoundMaterial( MaterialPurpose ) )
-			{
-				if ( TOptional<FString> UnrealMaterial = UsdUtils::GetUnrealSurfaceOutput( ShadeMaterial.GetPrim() ) )
-				{
-					FUsdPrimMaterialSlot& Slot = Result.Slots.Emplace_GetRef();
-					Slot.MaterialSource = UnrealMaterial.GetValue();
-					Slot.AssignmentType = UsdUtils::EPrimAssignmentType::UnrealMaterial;
-					Slot.bMeshIsDoubleSided = bIsDoubleSided;
-					Slot.PrimPaths.Add( MeshPrimPath );
-
-					return Result;
-				}
-			}
-		}
-
-		// Priority 1.2: unrealMaterial attribute directly on the prim
-		if ( TOptional<FString> UnrealMaterial = FetchFirstUEMaterialFromAttribute( UsdPrim, TimeCode ) )
-		{
-			FUsdPrimMaterialSlot& Slot = Result.Slots.Emplace_GetRef();
-			Slot.MaterialSource = UnrealMaterial.GetValue();
-			Slot.AssignmentType = UsdUtils::EPrimAssignmentType::UnrealMaterial;
-			Slot.bMeshIsDoubleSided = bIsDoubleSided;
-			Slot.PrimPaths.Add( MeshPrimPath );
+	bool bNeedsMainAssignment = true;
 
-			return Result;
-		}
-	}
-
-	// Priority 2: material binding directly on the prim
-	if ( TOptional<FString> BoundMaterial = FetchMaterialByComputingBoundMaterial( UsdPrim ) )
-	{
-		FUsdPrimMaterialSlot& Slot = Result.Slots.Emplace_GetRef();
-		Slot.MaterialSource = BoundMaterial.GetValue();
-		Slot.AssignmentType = UsdUtils::EPrimAssignmentType::MaterialPrim;
-		Slot.bMeshIsDoubleSided = bIsDoubleSided;
-		Slot.PrimPaths.Add( MeshPrimPath );
-
-		return Result;
-	}
-
-	// Priority 3: material:binding relationship directly on the prim (not sure why this is a separate step, but it came from IUsdPrim::GetGeometryMaterials. I bumped it in priority as the GeomSubsets do the same)
-	if ( TOptional<FString> TargetMaterial = FetchMaterialByMaterialRelationship( UsdPrim ) )
+	// Priority 0: GeomSubset partitions
+	std::vector<pxr::UsdGeomSubset> GeomSubsets = pxr::UsdShadeMaterialBindingAPI(UsdPrim).GetMaterialBindSubsets();
+	if (GeomSubsets.size() > 0)
 	{
-		FUsdPrimMaterialSlot& Slot = Result.Slots.Emplace_GetRef();
-		Slot.MaterialSource = TargetMaterial.GetValue();
-		Slot.AssignmentType = UsdUtils::EPrimAssignmentType::MaterialPrim;
-		Slot.bMeshIsDoubleSided = bIsDoubleSided;
-		Slot.PrimPaths.Add( MeshPrimPath );
-
-		return Result;
-	}
-
-	TOptional<FDisplayColorMaterial> DisplayColor = ExtractDisplayColorMaterial( pxr::UsdGeomMesh( UsdPrim ), TimeCode );
-
-	// Priority 4: GeomSubset partitions
-	std::vector<pxr::UsdGeomSubset> GeomSubsets = pxr::UsdShadeMaterialBindingAPI( UsdPrim ).GetMaterialBindSubsets();
-	if ( GeomSubsets.size() > 0 )
-	{
-		// We need to do this even if we won't provide indices because we may create an additional slot for unassigned polygons
-		pxr::VtIntArray UnassignedIndices;
 		std::string ReasonWhyNotPartition;
-		bool ValidPartition = pxr::UsdGeomSubset::ValidateSubsets( GeomSubsets, NumFaces, pxr::UsdGeomTokens->partition, &ReasonWhyNotPartition );
-		if ( !ValidPartition )
-		{
-			UE_LOG( LogUsd, Warning, TEXT( "Found an invalid GeomSubsets partition in prim '%s': %s" ),
-				*UsdToUnreal::ConvertPath( UsdPrim.GetPath() ), *UsdToUnreal::ConvertString( ReasonWhyNotPartition ));
-			UnassignedIndices = pxr::UsdGeomSubset::GetUnassignedIndices( GeomSubsets, NumFaces );
+		bool bHasValidPartitions = pxr::UsdGeomSubset::ValidateSubsets(GeomSubsets, NumFaces, pxr::UsdGeomTokens->partition, &ReasonWhyNotPartition);
+		if (!bHasValidPartitions)
+		{
+			UE_LOG(
+				LogUsd,
+				Warning,
+				TEXT("Found an invalid GeomSubsets partition in prim '%s': %s"),
+				*UsdToUnreal::ConvertPath(UsdPrim.GetPath()),
+				*UsdToUnreal::ConvertString(ReasonWhyNotPartition)
+			);
 		}
 
-		for ( uint32 GeomSubsetIndex = 0; GeomSubsetIndex < GeomSubsets.size(); ++GeomSubsetIndex )
+		for (uint32 GeomSubsetIndex = 0; GeomSubsetIndex < GeomSubsets.size(); ++GeomSubsetIndex)
 		{
-			const pxr::UsdGeomSubset& GeomSubset = GeomSubsets[ GeomSubsetIndex ];
-			bool bHasAssignment = false;
-
+			const pxr::UsdGeomSubset& GeomSubset = GeomSubsets[GeomSubsetIndex];
 			pxr::UsdPrim GeomSubsetPrim = GeomSubset.GetPrim();
-			FString GeomSubsetPath = UsdToUnreal::ConvertPath( GeomSubsetPrim.GetPath() );
+			FString GeomSubsetPath = UsdToUnreal::ConvertPath(GeomSubsetPrim.GetPath());
+			bool bHasAssignment = false;
 
-			// Priority 4.1: Material is an unreal asset
-			if ( RenderContext == UnrealIdentifiers::Unreal )
+			// Priority 0.1: Material is an unreal asset
+			if (RenderContext == UnrealIdentifiers::Unreal)
 			{
-				// Priority 4.1.1: Partition has an unreal rendercontext material prim binding
-				if ( !bHasAssignment )
+				// Priority 0.1.1: Partition has an unreal rendercontext material prim binding
+				if (!bHasAssignment)
 				{
-					pxr::UsdShadeMaterialBindingAPI BindingAPI( GeomSubsetPrim );
-					if ( BindingAPI )
+					pxr::UsdShadeMaterialBindingAPI BindingAPI(GeomSubsetPrim);
+					if (pxr::UsdShadeMaterial ShadeMaterial = BindingAPI.ComputeBoundMaterial(MaterialPurpose))
 					{
-						if ( pxr::UsdShadeMaterial ShadeMaterial = BindingAPI.ComputeBoundMaterial( MaterialPurpose ) )
+						if (TOptional<FString> UnrealMaterial = UsdUtils::GetUnrealSurfaceOutput(ShadeMaterial.GetPrim()))
 						{
-							if ( TOptional<FString> UnrealMaterial = UsdUtils::GetUnrealSurfaceOutput( ShadeMaterial.GetPrim() ) )
-							{
-								FUsdPrimMaterialSlot& Slot = Result.Slots.Emplace_GetRef();
-								Slot.MaterialSource = UnrealMaterial.GetValue();
-								Slot.AssignmentType = UsdUtils::EPrimAssignmentType::UnrealMaterial;
-								Slot.bMeshIsDoubleSided = bIsDoubleSided;
-								Slot.PrimPaths.Add( GeomSubsetPath );
-								bHasAssignment = true;
-							}
+							FUsdPrimMaterialSlot& Slot = Result.Slots.Emplace_GetRef();
+							Slot.MaterialSource = UnrealMaterial.GetValue();
+							Slot.AssignmentType = UsdUtils::EPrimAssignmentType::UnrealMaterial;
+							Slot.bMeshIsDoubleSided = bIsDoubleSided;
+							Slot.PrimPaths.Add(GeomSubsetPath);
+							bHasAssignment = true;
 						}
 					}
 				}
 
-				// Priority 4.1.2: Partitition has an unrealMaterial attribute directly on it
-				if ( !bHasAssignment )
+				// Priority 0.1.2: Partitition has an unrealMaterial attribute directly on it
+				if (!bHasAssignment)
 				{
-					if ( TOptional<FString> UnrealMaterial = FetchFirstUEMaterialFromAttribute( GeomSubsetPrim, TimeCode ) )
+					if (TOptional<FString> UnrealMaterial = FetchFirstUEMaterialFromAttribute(GeomSubsetPrim, TimeCode))
 					{
 						FUsdPrimMaterialSlot& Slot = Result.Slots.Emplace_GetRef();
 						Slot.MaterialSource = UnrealMaterial.GetValue();
 						Slot.AssignmentType = UsdUtils::EPrimAssignmentType::UnrealMaterial;
 						Slot.bMeshIsDoubleSided = bIsDoubleSided;
-						Slot.PrimPaths.Add( GeomSubsetPath );
+						Slot.PrimPaths.Add(GeomSubsetPath);
 						bHasAssignment = true;
 					}
 				}
 			}
 
-			// Priority 4.2: computing bound material
-			if ( !bHasAssignment )
+			// Priority 0.2: computing bound material
+			if (!bHasAssignment)
 			{
-				if ( TOptional<FString> BoundMaterial = FetchMaterialByComputingBoundMaterial( GeomSubsetPrim ) )
+				if (TOptional<FString> BoundMaterial = FetchMaterialByComputingBoundMaterial(GeomSubsetPrim))
 				{
 					FUsdPrimMaterialSlot& Slot = Result.Slots.Emplace_GetRef();
 					Slot.MaterialSource = BoundMaterial.GetValue();
 					Slot.AssignmentType = UsdUtils::EPrimAssignmentType::MaterialPrim;
 					Slot.bMeshIsDoubleSided = bIsDoubleSided;
-					Slot.PrimPaths.Add( GeomSubsetPath );
+					Slot.PrimPaths.Add(GeomSubsetPath);
 					bHasAssignment = true;
 				}
 			}
 
-			// Priority 4.3: material:binding relationship
-			if ( !bHasAssignment )
+			// Priority 0.3: material:binding relationship
+			if (!bHasAssignment)
 			{
-				if ( TOptional<FString> TargetMaterial = FetchMaterialByMaterialRelationship( GeomSubsetPrim ) )
+				if (TOptional<FString> TargetMaterial = FetchMaterialByMaterialRelationship(GeomSubsetPrim))
 				{
 					FUsdPrimMaterialSlot& Slot = Result.Slots.Emplace_GetRef();
 					Slot.MaterialSource = TargetMaterial.GetValue();
 					Slot.AssignmentType = UsdUtils::EPrimAssignmentType::MaterialPrim;
 					Slot.bMeshIsDoubleSided = bIsDoubleSided;
-					Slot.PrimPaths.Add( GeomSubsetPath );
+					Slot.PrimPaths.Add(GeomSubsetPath);
 					bHasAssignment = true;
 				}
 			}
 
-			// Priority 4.4: Create a section anyway so it becomes its own slot. Assign displayColor if we have one
-			if ( !bHasAssignment )
+			// Priority 0.4: Create a section anyway so that we always get a slot for each geom subset.
+			// We leave the assignment type cleared here, and will fill this in later with whatever we
+			// extract as a "main" material assignment.
+			// Note that we may have yet another "leftover" slot if our partition doesn't specify all faces,
+			// and that will be separate to this slot
+			if (!bHasAssignment)
 			{
 				FUsdPrimMaterialSlot& Slot = Result.Slots.Emplace_GetRef();
-				if ( DisplayColor )
-				{
-					Slot.MaterialSource = DisplayColor.GetValue().ToString();
-					Slot.AssignmentType = UsdUtils::EPrimAssignmentType::DisplayColor;
-					Slot.bMeshIsDoubleSided = bIsDoubleSided;
-				}
-				Slot.PrimPaths.Add( GeomSubsetPath );
+				Slot.PrimPaths.Add(GeomSubsetPath);
 				bHasAssignment = true;
 			}
 
-			pxr::VtIntArray PolygonIndicesInSubset;
-			GeomSubset.GetIndicesAttr().Get( &PolygonIndicesInSubset, TimeCode );
-
-			if ( bProvideMaterialIndices )
+			if (bProvideMaterialIndices)
 			{
+				pxr::VtIntArray PolygonIndicesInSubset;
+				GeomSubset.GetIndicesAttr().Get(&PolygonIndicesInSubset, TimeCode);
+
 				int32 LastAssignmentIndex = Result.Slots.Num() - 1;
-				for ( int PolygonIndex : PolygonIndicesInSubset )
+				for (int PolygonIndex : PolygonIndicesInSubset)
 				{
-					Result.MaterialIndices[ PolygonIndex ] = LastAssignmentIndex;
+					Result.MaterialIndices[PolygonIndex] = LastAssignmentIndex;
 				}
 			}
 		}
 
-		// Extra slot for unassigned polygons
-		if ( UnassignedIndices.size() > 0 )
+		// Extra slot for unspecified faces.
+		// We need to fetch this even if we won't provide indices because we may need to create an additional slot for unassigned polygons
+		pxr::VtIntArray UnassignedIndices = pxr::UsdGeomSubset::GetUnassignedIndices(GeomSubsets, NumFaces);
+		if (UnassignedIndices.size() == 0)
+		{
+			bNeedsMainAssignment = false;
+		}
+		else
 		{
-			FUsdPrimMaterialSlot& Slot = Result.Slots.Emplace_GetRef();
-			if ( DisplayColor )
+			// Assign these leftover indices to the *next* material slot we'll create (doesn't exist yet),
+			// which will be the "main" material assignment slot
+			int32 LeftoverSlotIndex = Result.Slots.Num();
+			if (bProvideMaterialIndices)
 			{
-				Slot.MaterialSource = DisplayColor.GetValue().ToString();
-				Slot.AssignmentType = UsdUtils::EPrimAssignmentType::DisplayColor;
-				Slot.bMeshIsDoubleSided = bIsDoubleSided;
+				for (int PolygonIndex : UnassignedIndices)
+				{
+					Result.MaterialIndices[PolygonIndex] = LeftoverSlotIndex;
+				}
 			}
+		}
+	}
 
-			if ( bProvideMaterialIndices )
+	TOptional<FDisplayColorMaterial> DisplayColor;
+
+	bool bHasMainAssignment = false;
+	if (bNeedsMainAssignment)
+	{
+		// Priority 1: Material is an unreal asset
+		if (RenderContext == UnrealIdentifiers::Unreal)
+		{
+			// Priority 1.1: unreal rendercontext material prim
+			// Note how we don't test this BindingAPI for truthiness: This allows us to compute a bound material
+			// even if this prim is just inheriting a material binding, but doesn't actually have the API itself
+			pxr::UsdShadeMaterialBindingAPI BindingAPI{UsdPrim};
+			if (pxr::UsdShadeMaterial ShadeMaterial = BindingAPI.ComputeBoundMaterial(MaterialPurpose))
 			{
-				int32 LastAssignmentIndex = Result.Slots.Num() - 1;
-				for ( int PolygonIndex : UnassignedIndices )
+				if (TOptional<FString> UnrealMaterial = UsdUtils::GetUnrealSurfaceOutput(ShadeMaterial.GetPrim()))
 				{
-					Result.MaterialIndices[ PolygonIndex ] = LastAssignmentIndex;
+					FUsdPrimMaterialSlot& Slot = Result.Slots.Emplace_GetRef();
+					Slot.MaterialSource = UnrealMaterial.GetValue();
+					Slot.AssignmentType = UsdUtils::EPrimAssignmentType::UnrealMaterial;
+					Slot.bMeshIsDoubleSided = bIsDoubleSided;
+					Slot.PrimPaths.Add(MeshPrimPath);
+
+					bHasMainAssignment = true;
+				}
+			}
+
+			// Priority 1.2: unrealMaterial attribute directly on the prim
+			if (!bHasMainAssignment)
+			{
+				if (TOptional<FString> UnrealMaterial = FetchFirstUEMaterialFromAttribute(UsdPrim, TimeCode))
+				{
+					FUsdPrimMaterialSlot& Slot = Result.Slots.Emplace_GetRef();
+					Slot.MaterialSource = UnrealMaterial.GetValue();
+					Slot.AssignmentType = UsdUtils::EPrimAssignmentType::UnrealMaterial;
+					Slot.bMeshIsDoubleSided = bIsDoubleSided;
+					Slot.PrimPaths.Add(MeshPrimPath);
+
+					bHasMainAssignment = true;
 				}
 			}
 		}
 
-		return Result;
+		// Priority 2: material binding directly on the prim
+		if (!bHasMainAssignment)
+		{
+			if (TOptional<FString> BoundMaterial = FetchMaterialByComputingBoundMaterial(UsdPrim))
+			{
+				FUsdPrimMaterialSlot& Slot = Result.Slots.Emplace_GetRef();
+				Slot.MaterialSource = BoundMaterial.GetValue();
+				Slot.AssignmentType = UsdUtils::EPrimAssignmentType::MaterialPrim;
+				Slot.bMeshIsDoubleSided = bIsDoubleSided;
+				Slot.PrimPaths.Add(MeshPrimPath);
+
+				bHasMainAssignment = true;
+			}
+		}
+
+		// Priority 3: material:binding relationship directly on the prim (not sure why this is a separate step, but it came from
+		// IUsdPrim::GetGeometryMaterials. I bumped it in priority as the GeomSubsets do the same)
+		if (!bHasMainAssignment)
+		{
+			if (TOptional<FString> TargetMaterial = FetchMaterialByMaterialRelationship(UsdPrim))
+			{
+				FUsdPrimMaterialSlot& Slot = Result.Slots.Emplace_GetRef();
+				Slot.MaterialSource = TargetMaterial.GetValue();
+				Slot.AssignmentType = UsdUtils::EPrimAssignmentType::MaterialPrim;
+				Slot.bMeshIsDoubleSided = bIsDoubleSided;
+				Slot.PrimPaths.Add(MeshPrimPath);
+
+				bHasMainAssignment = true;
+			}
+		}
+
+		// Priority 4: vertex color material using displayColor/displayOpacity information for the entire mesh
+		// Note: This will in general always succeed for any mesh prim, as the schema will provide fallback values
+		// for displayColor and displayOpacity
+		if (!bHasMainAssignment)
+		{
+			DisplayColor = ExtractDisplayColorMaterial(pxr::UsdGeomMesh(UsdPrim), TimeCode);
+			if (DisplayColor)
+			{
+				FUsdPrimMaterialSlot& Slot = Result.Slots.Emplace_GetRef();
+				Slot.MaterialSource = DisplayColor.GetValue().ToString();
+				Slot.AssignmentType = UsdUtils::EPrimAssignmentType::DisplayColor;
+				Slot.bMeshIsDoubleSided = bIsDoubleSided;
+				Slot.PrimPaths.Add(MeshPrimPath);
+
+				bHasMainAssignment = true;
+			}
+		}
 	}
+	ensure(bHasMainAssignment || !bNeedsMainAssignment);
 
-	// Priority 5: vertex color material using displayColor/displayOpacity information for the entire mesh
-	if ( GeomSubsets.size() == 0 && DisplayColor )
+	// If we have any slot without an actual material assignment yet, copy over the material assignment from
+	// the "main" slot, or fallback to displayColor. This is how we have unspecified faces or geomsubsets without
+	// assignmets "fallback" to using the main material assignment
+	if (Result.Slots.Num() > 1)
 	{
-		FUsdPrimMaterialSlot& Slot = Result.Slots.Emplace_GetRef();
-		Slot.MaterialSource = DisplayColor.GetValue().ToString();
-		Slot.AssignmentType = UsdUtils::EPrimAssignmentType::DisplayColor;
-		Slot.bMeshIsDoubleSided = bIsDoubleSided;
-		Slot.PrimPaths.Add( MeshPrimPath );
+		FString FallbackMaterialSource;
+		EPrimAssignmentType FallbackAssignmentType = EPrimAssignmentType::None;
 
-		return Result;
-	}
+		if (bHasMainAssignment)
+		{
+			// Our main slot is the last created one at this point
+			FUsdPrimMaterialSlot& MainSlot = Result.Slots[Result.Slots.Num() - 1];
+			FallbackAssignmentType = MainSlot.AssignmentType;
+			FallbackMaterialSource = MainSlot.MaterialSource;
+		}
+		else
+		{
+			if (!DisplayColor.IsSet())
+			{
+				DisplayColor = ExtractDisplayColorMaterial(pxr::UsdGeomMesh(UsdPrim), TimeCode);
+				if (ensure(DisplayColor.IsSet()))
+				{
+					FallbackAssignmentType = UsdUtils::EPrimAssignmentType::DisplayColor;
+					FallbackMaterialSource = DisplayColor.GetValue().ToString();
+				}
+			}
+		}
 
-	// Priority 6: Make sure there is always at least one slot, even if empty
-	if ( Result.Slots.Num() < 1 )
+		for (int32 Index = 0; Index < Result.Slots.Num(); ++Index)
+		{
+			FUsdPrimMaterialSlot& Slot = Result.Slots[Index];
+			if (Slot.AssignmentType == UsdUtils::EPrimAssignmentType::None)
+			{
+				Slot.AssignmentType = FallbackAssignmentType;
+				Slot.MaterialSource = FallbackMaterialSource;
+			}
+		}
+	}
+	// Priority 5: Make sure there is always at least one slot, even if empty
+	else if (Result.Slots.Num() < 1)
 	{
 		Result.Slots.Emplace();
 	}
