USD: Add analytics for usage of non-native schemas.

diff --git a/Engine/Plugins/Importers/USDImporter/Source/USDClasses/Private/USDClassesModule.cpp b/Engine/Plugins/Importers/USDImporter/Source/USDClasses/Private/USDClassesModule.cpp
--- a/Engine/Plugins/Importers/USDImporter/Source/USDClasses/Private/USDClassesModule.cpp
+++ b/Engine/Plugins/Importers/USDImporter/Source/USDClasses/Private/USDClassesModule.cpp
@@ -214,6 +214,21 @@ void IUsdClassesModule::SendAnalytics( TArray<FAnalyticsEventAttribute>&& InAttr
 	}
 }
 
+void IUsdClassesModule::SendAnalytics(TArray<FAnalyticsEventAttribute>&& InAttributes, const FString& EventName)
+{
+	if (FEngineAnalytics::IsAvailable())
+	{
+		TArray<FAnalyticsEventAttribute> Attributes(InAttributes);
+
+		Attributes.Emplace(TEXT("Platform"), FPlatformProperties::IniPlatformName());
+		Attributes.Emplace(TEXT("EngineVersion"), FEngineVersion::Current().ToString());
+		Attributes.Emplace(TEXT("EngineMode"), FPlatformMisc::GetEngineMode());
+
+		const FString EventText = FString::Printf(TEXT("Engine.Usage.USD.%s"), *EventName);
+		FEngineAnalytics::GetProvider().RecordEvent(EventText, Attributes);
+	}
+}
+
 bool IUsdClassesModule::HashObjectPackage( const UObject* Object, FSHA1& HashToUpdate )
 {
 #if WITH_EDITOR
diff --git a/Engine/Plugins/Importers/USDImporter/Source/USDClasses/Public/USDClassesModule.h b/Engine/Plugins/Importers/USDImporter/Source/USDClasses/Public/USDClassesModule.h
--- a/Engine/Plugins/Importers/USDImporter/Source/USDClasses/Public/USDClassesModule.h
+++ b/Engine/Plugins/Importers/USDImporter/Source/USDClasses/Public/USDClassesModule.h
@@ -30,6 +30,10 @@ public:
 		double NumberOfFrames,
 		const FString& Extension
 	);
+	USDCLASSES_API static void SendAnalytics(
+		TArray<FAnalyticsEventAttribute>&& InAttributes,
+		const FString& EventName
+	);
 
 	/**
 	 * Updates HashToUpdate with the Object's package's persistent guid, the corresponding file save
diff --git a/Engine/Plugins/Importers/USDImporter/Source/USDStage/Private/USDStageActor.cpp b/Engine/Plugins/Importers/USDImporter/Source/USDStage/Private/USDStageActor.cpp
--- a/Engine/Plugins/Importers/USDImporter/Source/USDStage/Private/USDStageActor.cpp
+++ b/Engine/Plugins/Importers/USDImporter/Source/USDStage/Private/USDStageActor.cpp
@@ -2162,6 +2162,10 @@ void AUsdStageActor::SetRootLayer(const FString& RootFilePath)
 	UnloadUsdStage();
 	RootLayer.FilePath = RelativeFilePath;
 	LoadUsdStage();
+
+	// Do this here instead of on OpenUsdStage/LoadUsdStage as those also get called when changing any of
+	// our properties, like render context, material purpose, etc.
+	UsdUtils::CollectSchemaAnalytics(UsdStage, TEXT("Open"));
 }
 
 void AUsdStageActor::SetAssetCache(UUsdAssetCache2* NewCache)
diff --git a/Engine/Plugins/Importers/USDImporter/Source/USDStageImporter/Private/USDStageImporter.cpp b/Engine/Plugins/Importers/USDImporter/Source/USDStageImporter/Private/USDStageImporter.cpp
--- a/Engine/Plugins/Importers/USDImporter/Source/USDStageImporter/Private/USDStageImporter.cpp
+++ b/Engine/Plugins/Importers/USDImporter/Source/USDStageImporter/Private/USDStageImporter.cpp
@@ -1876,6 +1876,8 @@ void UUsdStageImporter::ImportFromFile(FUsdStageImportContext& ImportContext)
 		double ElapsedSeconds = FPlatformTime::ToSeconds64( FPlatformTime::Cycles64() - StartTime );
 		UsdStageImporterImpl::SendAnalytics( ImportContext, nullptr, TEXT("Import"), PublishedAssetsAndDependencies, ElapsedSeconds);
 		UE_LOG(LogUsd, Log, TEXT("Imported '%s' in %.3f seconds."), *ImportContext.FilePath, ElapsedSeconds);
+
+		UsdUtils::CollectSchemaAnalytics(ImportContext.Stage, TEXT("Import"));
 	}
 
 	UsdStageImporterImpl::CloseStageIfNeeded( ImportContext );
@@ -2040,6 +2042,8 @@ bool UUsdStageImporter::ReimportSingleAsset(
 		double ElapsedSeconds = FPlatformTime::ToSeconds64( FPlatformTime::Cycles64() - StartTime );
 		UsdStageImporterImpl::SendAnalytics( ImportContext, ReimportedObject, TEXT( "Reimport" ), { ReimportedObject }, ElapsedSeconds );
 		UE_LOG(LogUsd, Log, TEXT("Re-imported '%s' in %.3f seconds."), *ImportContext.FilePath, ElapsedSeconds);
+
+		UsdUtils::CollectSchemaAnalytics(ImportContext.Stage, TEXT("Reimport"));
 	}
 
 	UsdStageImporterImpl::CloseStageIfNeeded( ImportContext );
diff --git a/Engine/Plugins/Importers/USDImporter/Source/USDUtilities/Private/USDConversionUtils.cpp b/Engine/Plugins/Importers/USDImporter/Source/USDUtilities/Private/USDConversionUtils.cpp
--- a/Engine/Plugins/Importers/USDImporter/Source/USDUtilities/Private/USDConversionUtils.cpp
+++ b/Engine/Plugins/Importers/USDImporter/Source/USDUtilities/Private/USDConversionUtils.cpp
@@ -18,7 +18,7 @@
 #include "UsdWrappers/UsdPrim.h"
 #include "UsdWrappers/UsdStage.h"
 
-#include "Algo/Copy.h"
+#include "AnalyticsEventAttribute.h"
 #include "Animation/AnimBlueprint.h"
 #include "Animation/AnimSequence.h"
 #include "Animation/Skeleton.h"
@@ -69,8 +69,6 @@
 	#include "pxr/usd/usdGeom/metrics.h"
 	#include "pxr/usd/usdGeom/primvar.h"
 	#include "pxr/usd/usdGeom/primvarsAPI.h"
-	#include "pxr/usd/usdGeom/scope.h"
-	#include "pxr/usd/usdGeom/xform.h"
 	#include "pxr/usd/usdLux/diskLight.h"
 	#include "pxr/usd/usdLux/distantLight.h"
 	#include "pxr/usd/usdLux/domeLight.h"
@@ -2820,6 +2818,165 @@ FUsdUnrealAssetInfo UsdUtils::GetPrimAssetInfo( const UE::FUsdPrim& Prim )
 	return Result;
 }
 
+void UsdUtils::CollectSchemaAnalytics(const UE::FUsdStage& Stage, const FString& EventName)
+{
+#if USE_USD_SDK
+	if (!Stage)
+	{
+		return;
+	}
+
+	TMap<FString, int32> Counts;
+
+	{
+		FScopedUsdAllocs Allocs;
+
+		pxr::UsdPrimRange PrimRange{Stage.GetPseudoRoot(), pxr::UsdTraverseInstanceProxies()};
+		for (pxr::UsdPrimRange::iterator PrimRangeIt = ++PrimRange.begin(); PrimRangeIt != PrimRange.end(); ++PrimRangeIt)
+		{
+			// It's perfectly fine to have a typeless prim (e.g. "def 'Cube'")
+			if (PrimRangeIt->HasAuthoredTypeName())
+			{
+				const pxr::TfToken& TypeName = PrimRangeIt->GetTypeName();
+				const FString TypeNameStr = UsdToUnreal::ConvertToken(TypeName);
+
+				if (!TypeNameStr.IsEmpty())
+				{
+					Counts.FindOrAdd(TypeNameStr) += 1;
+				}
+			}
+
+			for (const pxr::TfToken& AppliedSchema : PrimRangeIt->GetAppliedSchemas())
+			{
+				std::pair<pxr::TfToken, pxr::TfToken> Pair = pxr::UsdSchemaRegistry::GetTypeNameAndInstance(AppliedSchema);
+				FString TypeName = UsdToUnreal::ConvertToken(Pair.first);
+
+				// These applied schema names shouldn't ever end up as the empty string... but we don't really want to pop
+				// an ensure or show a warning when analytics fails
+				if (!TypeName.IsEmpty())
+				{
+					Counts.FindOrAdd(TypeName) += 1;
+				}
+			}
+		}
+	}
+
+	const static TSet<FString> NativeSchemaNames{
+		"AssetPreviewsAPI",
+		"Backdrop",
+		"BasisCurves",
+		"BlendShape",
+		"Camera",
+		"Capsule",
+		"ClipsAPI",
+		"CollectionAPI",
+		"Cone",
+		"ConnectableAPI",
+		"ControlRigAPI",
+		"CoordSysAPI",
+		"Cube",
+		"Cylinder",
+		"CylinderLight",
+		"DiskLight",
+		"DistantLight",
+		"DomeLight",
+		"Field3DAsset",
+		"GenerativeProcedural",
+		"GeomModelAPI",
+		"GeomSubset",
+		"GeometryLight",
+		"GroomAPI",
+		"GroomBindingAPI",
+		"HermiteCurves",
+		"HydraGenerativeProceduralAPI",
+		"LightAPI",
+		"LightFilter",
+		"LightListAPI",
+		"ListAPI",
+		"LiveLinkAPI",
+		"Material",
+		"MaterialBindingAPI",
+		"Mesh",
+		"MeshLightAPI",
+		"ModelAPI",
+		"MotionAPI",
+		"NodeDefAPI",
+		"NodeGraph",
+		"NodeGraphNodeAPI",
+		"NurbsCurves",
+		"NurbsPatch",
+		"OpenVDBAsset",
+		"PackedJointAnimation",
+		"PhysicsArticulationRootAPI",
+		"PhysicsCollisionAPI",
+		"PhysicsCollisionGroup",
+		"PhysicsDistanceJoint",
+		"PhysicsDriveAPI",
+		"PhysicsFilteredPairsAPI",
+		"PhysicsFixedJoint",
+		"PhysicsJoint",
+		"PhysicsLimitAPI",
+		"PhysicsMassAPI",
+		"PhysicsMaterialAPI",
+		"PhysicsMeshCollisionAPI",
+		"PhysicsPrismaticJoint",
+		"PhysicsRevoluteJoint",
+		"PhysicsRigidBodyAPI",
+		"PhysicsScene",
+		"PhysicsSphericalJoint",
+		"Plane",
+		"PluginLight",
+		"PluginLightFilter",
+		"PointInstancer",
+		"Points",
+		"PortalLight",
+		"PrimvarsAPI",
+		"RectLight",
+		"RenderDenoisePass",
+		"RenderPass",
+		"RenderProduct",
+		"RenderSettings",
+		"RenderVar",
+		"RiMaterialAPI",
+		"RiSplineAPI",
+		"SceneGraphPrimAPI",
+		"Scope",
+		"Shader",
+		"ShadowAPI",
+		"ShapingAPI",
+		"SkelAnimation",
+		"SkelBindingAPI",
+		"SkelRoot",
+		"Skeleton",
+		"SpatialAudio",
+		"Sphere",
+		"SphereLight",
+		"StatementsAPI",
+		"VisibilityAPI",
+		"Volume",
+		"VolumeLightAPI",
+		"Xform",
+		"XformCommonAPI"};
+
+	TArray<FAnalyticsEventAttribute> EventAttributes;
+	EventAttributes.Reserve(Counts.Num());
+
+	for (const TPair<FString, int32>& Pair : Counts)
+	{
+		// We only care about non-native schemas
+		if (!NativeSchemaNames.Contains(Pair.Key))
+		{
+			EventAttributes.Emplace(Pair.Key, Pair.Value);
+		}
+	}
+
+	if (EventAttributes.Num() > 0)
+	{
+		IUsdClassesModule::SendAnalytics(MoveTemp(EventAttributes), FString::Printf(TEXT("%s.NonNativeSchemaCounts"), *EventName));
+	}
+#endif	  // USE_USD_SDK
+}
+
 #if USE_USD_SDK
 #undef LOCTEXT_NAMESPACE
 #endif
diff --git a/Engine/Plugins/Importers/USDImporter/Source/USDUtilities/Public/USDConversionUtils.h b/Engine/Plugins/Importers/USDImporter/Source/USDUtilities/Public/USDConversionUtils.h
--- a/Engine/Plugins/Importers/USDImporter/Source/USDUtilities/Public/USDConversionUtils.h
+++ b/Engine/Plugins/Importers/USDImporter/Source/USDUtilities/Public/USDConversionUtils.h
@@ -448,5 +448,8 @@ namespace UsdUtils
 
 	/** Retrieves from Prim the assetInfo metadata values that we use as export metadata, when exporting Unreal assets */
 	USDUTILITIES_API FUsdUnrealAssetInfo GetPrimAssetInfo( const UE::FUsdPrim& Prim );
+
+	/** Collects how many times each schema shows up on the provided stage and send it as an analytics event */
+	USDUTILITIES_API void CollectSchemaAnalytics(const UE::FUsdStage& Stage, const FString& EventName);
 }
 
