USD: Refactor how the USD Stage Editor refreshes to naturally avoid spamming refreshes when USD emits too many notices.

diff --git a/Engine/Plugins/Importers/USDImporter/Source/USDStageEditor/Private/SUSDLayersTreeView.h b/Engine/Plugins/Importers/USDImporter/Source/USDStageEditor/Private/SUSDLayersTreeView.h
--- a/Engine/Plugins/Importers/USDImporter/Source/USDStageEditor/Private/SUSDLayersTreeView.h
+++ b/Engine/Plugins/Importers/USDImporter/Source/USDStageEditor/Private/SUSDLayersTreeView.h
@@ -72,6 +72,22 @@ private:
 
 	void RestoreExpansionStates();
 
+public:
+	// We update only on slate tick, but can receive a USD notice at any point, from any thread.
+	// These variables are used to control what needs to be refreshed on the next tick
+
+	enum class ELayersTreeViewState : uint8
+	{
+		NoRefreshNeeded = 0,
+		RefreshNeeded = 1,
+		ResyncNeeded = 2
+	};
+	ELayersTreeViewState bUpdateState = ELayersTreeViewState::NoRefreshNeeded;
+
+	// Use a lock because some of the notices that will update these refresh variables can be
+	// sent from TBB threads
+	mutable FRWLock RefreshStateLock;
+
 private:
 	// Should always be valid, we keep the one we're given on Refresh()
 	UE::FUsdStageWeak UsdStage;
diff --git a/Engine/Plugins/Importers/USDImporter/Source/USDStageEditor/Private/SUSDStage.cpp b/Engine/Plugins/Importers/USDImporter/Source/USDStageEditor/Private/SUSDStage.cpp
--- a/Engine/Plugins/Importers/USDImporter/Source/USDStageEditor/Private/SUSDStage.cpp
+++ b/Engine/Plugins/Importers/USDImporter/Source/USDStageEditor/Private/SUSDStage.cpp
@@ -395,10 +395,13 @@ void SUsdStage::Construct( const FArguments& InArgs )
 
 	SetupStageActorDelegates();
 
+	// Pre tick so that we setup/refresh our widgets before Slate renders on this frame
+	FSlateApplication::Get().OnPreTick().AddSP(this, &SUsdStage::OnSlateTick);
+
 	// We're opening the USD Stage editor for the first time and we already have a stage: Display it immediately
 	if ( UsdStage )
 	{
-		Refresh();
+		RequestFullRefresh();
 	}
 }
 
@@ -512,7 +515,7 @@ void SUsdStage::AttachToStageActor( AUsdStageActor* InUsdStageActor, bool bFlash
 
 	// Refresh here because we may be receiving an actor that has a stage already loaded,
 	// like during undo/redo
-	Refresh();
+	RequestFullRefresh();
 }
 
 void SUsdStage::SetupStageActorDelegates()
@@ -522,47 +525,14 @@ void SUsdStage::SetupStageActorDelegates()
 	if ( ViewModel.UsdStageActor.IsValid() )
 	{
 		OnPrimChangedHandle = ViewModel.UsdStageActor->OnPrimChanged.AddLambda(
-			[ this ]( const FString& PrimPath, bool bResync )
+			[this](const FString& PrimPath, bool bResync)
 			{
-				// The USD notices may come from a background USD TBB thread, but we should only update slate from the main/slate threads.
-				// We can't retrieve the FSlateApplication singleton here (because that can also only be used from the main/slate threads),
-				// so we must use core tickers here
-				FTSTicker::GetCoreTicker().AddTicker(FTickerDelegate::CreateLambda(
-					[this, PrimPath, bResync](float Time)
-					{
-						if (bEditorIsShuttingDown)
-						{
-							return false;
-						}
-
-						if ( this->UsdStageTreeView )
-						{
-							this->UsdStageTreeView->RefreshPrim( PrimPath, bResync );
-							this->UsdStageTreeView->RequestTreeRefresh();
-						}
-
-						const bool bViewingTheUpdatedPrim = SelectedPrimPath.Equals( PrimPath, ESearchCase::IgnoreCase );
-						const bool bViewingStageProperties = SelectedPrimPath.IsEmpty() || SelectedPrimPath == TEXT("/");
-						const bool bStageUpdated = PrimPath == TEXT("/");
-
-						if ( this->UsdPrimInfoWidget &&
-							 ViewModel.UsdStageActor.IsValid() &&
-							 ( bViewingTheUpdatedPrim || ( bViewingStageProperties && bStageUpdated ) ) )
-						{
-							this->UsdPrimInfoWidget->SetPrimPath( GetCurrentStage(), *PrimPath);
-						}
-
-						// If we resynced our selected prim or our ancestor and have selection sync enabled, try to refresh it so that we're
-						// still selecting the same actor/component that corresponds to the prim we have currently selected
-						if ( bResync && SelectedPrimPath.StartsWith( PrimPath ) )
-						{
-							OnPrimSelectionChanged( { SelectedPrimPath } );
-						}
+				if (UsdStageTreeView)
+				{
+					FWriteScopeLock StageTreeViewLock{UsdStageTreeView->RefreshStateLock};
 
-						// Returning false means this is a one-off, and won't repeat
-						return false;
-					}
-				));
+					UsdStageTreeView->PrimsToRefresh.Emplace(PrimPath, bResync);
+				}
 			}
 		);
 
@@ -570,39 +540,15 @@ void SUsdStage::SetupStageActorDelegates()
 		OnStageChangedHandle = ViewModel.UsdStageActor->OnStageChanged.AddLambda(
 			[this]()
 			{
-				FTSTicker::GetCoreTicker().AddTicker(FTickerDelegate::CreateLambda(
-					[this](float Time)
+				if (UsdStageTreeView)
+				{
 					{
-						if (bEditorIsShuttingDown)
-						{
-							return false;
-						}
-
-						// So we can reset even if our actor is being destroyed right now
-						const bool bEvenIfPendingKill = true;
-						if (ViewModel.UsdStageActor.IsValid(bEvenIfPendingKill))
-						{
-							// Reset our selection to the stage root
-							SelectedPrimPath = TEXT("/");
-
-							if (this->UsdPrimInfoWidget)
-							{
-								this->UsdPrimInfoWidget->SetPrimPath(GetCurrentStage(), TEXT("/"));
-							}
-
-							if (this->UsdStageTreeView)
-							{
-								this->UsdStageTreeView->ClearSelection();
-								this->UsdStageTreeView->RequestTreeRefresh();
-							}
-						}
-
-						this->Refresh();
-
-						// Returning false means this is a one-off, and won't repeat
-						return false;
+						FWriteScopeLock StageTreeViewLock{UsdStageTreeView->RefreshStateLock};
+						UsdStageTreeView->PrimsToRefresh.Reset();
 					}
-				));
+
+					RequestFullRefresh();
+				}
 			}
 		);
 
@@ -615,113 +561,26 @@ void SUsdStage::SetupStageActorDelegates()
 				// (and control flow returned to the game thread) which can lead to some weird results (and break automated tests).
 				// We could get around this on the Python script's side by just yielding, but there may be other scenarios
 				ClearStageActorDelegates();
-				this->ViewModel.CloseStage();
-
-				FTSTicker::GetCoreTicker().AddTicker(FTickerDelegate::CreateLambda(
-					[this](float Time)
-					{
-						if (bEditorIsShuttingDown)
-						{
-							return false;
-						}
-
-						this->Refresh();
-
-						// Returning false means this is a one-off, and won't repeat
-						return false;
-					}
-				));
+				ViewModel.CloseStage();
+				RequestFullRefresh();
 			}
 		);
 
-		OnStageEditTargetChangedHandle = ViewModel.UsdStageActor->GetUsdListener().GetOnStageEditTargetChanged().AddLambda(
-			[ this ]()
-			{
-				FTSTicker::GetCoreTicker().AddTicker(
-					FTickerDelegate::CreateLambda( [this]( float Time )
-					{
-						if (bEditorIsShuttingDown)
-						{
-							return false;
-						}
-
-						AUsdStageActor* StageActor = ViewModel.UsdStageActor.Get();
-						if (this->UsdLayersTreeView && StageActor)
-						{
-							constexpr bool bResync = false;
-							UsdLayersTreeView->Refresh(
-								StageActor->GetBaseUsdStage(),
-								StageActor->GetIsolatedUsdStage(),
-								bResync
-							);
-						}
-
-						// Returning false means this is a one-off, and won't repeat
-						return false;
-					})
-				);
-			}
+		OnStageEditTargetChangedHandle = ViewModel.UsdStageActor->GetUsdListener().GetOnStageEditTargetChanged().AddSP(
+			this,
+			&SUsdStage::RequestLayersTreeViewRefresh
 		);
 
-		TFunction<bool(float)> DelayedRefreshOfLayersTreeView = [this](float Time)
-		{
-			if (bEditorIsShuttingDown)
-			{
-				return false;
-			}
-
-			AUsdStageActor* StageActor = ViewModel.UsdStageActor.Get();
-			if (this->UsdLayersTreeView && StageActor)
-			{
-				constexpr bool bResync = false;
-				UsdLayersTreeView->Refresh(StageActor->GetBaseUsdStage(), StageActor->GetIsolatedUsdStage(), bResync);
-			}
-
-			// Returning false means this is a one-off, and won't repeat
-			return false;
-		};
-
 		OnSdfLayersChangedHandle = ViewModel.UsdStageActor->GetUsdListener().GetOnSdfLayersChanged().AddLambda(
-			[this, DelayedRefreshOfLayersTreeView](const UsdUtils::FLayerToSdfChangeList& LayersToChangeList)
+			[this](const UsdUtils::FLayerToSdfChangeList& LayersToChangeList)
 			{
-				// The layers changed events are global, and so fired for every stage actor when any layer of
-				// any stage is modified. Let's update only when a layer used by our stage is updated
-				bool bRelevantToCurrentStage = false;
-				if (AUsdStageActor* StageActor = ViewModel.UsdStageActor.Get())
-				{
-					if (const UE::FUsdStage& CurrentStage = StageActor->GetUsdStage())
-					{
-						TSet<UE::FSdfLayer> UsedLayers{CurrentStage.GetUsedLayers()};
-						for (const TPair<UE::FSdfLayerWeak, UsdUtils::FSdfChangeList>& LayerToChangeList : LayersToChangeList)
-						{
-							if (UsedLayers.Contains(LayerToChangeList.Key))
-							{
-								bRelevantToCurrentStage = true;
-								break;
-							}
-						}
-					}
-				}
-				if (!bRelevantToCurrentStage)
-				{
-					return;
-				}
-
-				FTSTicker::GetCoreTicker().AddTicker(FTickerDelegate::CreateLambda(DelayedRefreshOfLayersTreeView));
+				RequestLayersTreeViewRefresh();
 			}
 		);
 
-		OnSdfLayerDirtinessChangedHandle = ViewModel.UsdStageActor->GetUsdListener().GetOnSdfLayerDirtinessChanged().AddLambda(
-			[DelayedRefreshOfLayersTreeView]()
-			{
-				// The layer dirtiness notice as we're using here is global, which means any layer of any other stage
-				// going dirty will affect us. We *could* make this notice layer-specific on the listener, but it would force us to
-				// register/unregister to every single layer... since the updates done in DelayedRefreshOfLayersTreeView are
-				// relatively light, it's probably for the best to just refresh every time than to support the churn of
-				// a bunch of separate listeners and registration/unregistration happening every time, potentially also
-				// crashing if done incorrectly when a layer is created/added/removed/etc.
-				FTSTicker::GetCoreTicker().AddTicker(FTickerDelegate::CreateLambda(DelayedRefreshOfLayersTreeView));
-			}
+		OnSdfLayerDirtinessChangedHandle = ViewModel.UsdStageActor->GetUsdListener().GetOnSdfLayerDirtinessChanged().AddSP(
+			this,
+			&SUsdStage::RequestLayersTreeViewRefresh
 		);
 	}
 }
@@ -743,6 +602,8 @@ void SUsdStage::ClearStageActorDelegates()
 
 SUsdStage::~SUsdStage()
 {
+	FSlateApplication::Get().OnPreTick().RemoveAll(this);
+
 	FEditorDelegates::PostPIEStarted.Remove( PostPIEStartedHandle );
 	FEditorDelegates::EndPIE.Remove( EndPIEHandle );
 
@@ -2250,7 +2111,7 @@ void SUsdStage::FileClose()
 	}
 
 	ViewModel.CloseStage();
-	Refresh();
+	RequestFullRefresh();
 
 	if (GDiscardUndoBufferOnStageOpenClose && GEditor)
 	{
@@ -2268,7 +2129,7 @@ void SUsdStage::OnLayerIsolated( const UE::FSdfLayer& IsolatedLayer )
 	{
 		StageActor->IsolateLayer( IsolatedLayer );
 
-		Refresh();
+		RequestFullRefresh();
 	}
 }
 
@@ -2303,12 +2164,12 @@ void SUsdStage::OnPrimSelectionChanged( const TArray<FString>& PrimPaths )
 		return;
 	}
 
-	if ( UsdPrimInfoWidget )
-	{
-		UE::FUsdStage UsdStage = static_cast< const AUsdStageActor* >( StageActor )->GetUsdStage();
+	UE::FUsdStage UsdStage = static_cast< const AUsdStageActor* >( StageActor )->GetUsdStage();
+	SelectedPrimPath = PrimPaths.Num() == 1 ? PrimPaths[ 0 ] : TEXT( "" );
 
-		SelectedPrimPath = PrimPaths.Num() == 1 ? PrimPaths[ 0 ] : TEXT( "" );
-		UsdPrimInfoWidget->SetPrimPath( UsdStage, *SelectedPrimPath );
+	if (UsdPrimInfoWidget)
+	{
+		UsdPrimInfoWidget->SetPrimPath(UsdStage, *SelectedPrimPath);
 	}
 
 	// We don't need to check for actors and components in case we're just selecting the stage root
@@ -2362,39 +2223,169 @@ void SUsdStage::OpenStage( const TCHAR* FilePath )
 	}
 }
 
-void SUsdStage::Refresh()
+void SUsdStage::RequestLayersTreeViewRefresh()
 {
-	// May be nullptr, but that is ok. Its how the widgets are reset
-	AUsdStageActor* StageActor = ViewModel.UsdStageActor.Get();
+	if (UsdLayersTreeView)
+	{
+		FWriteScopeLock LayersTreeViewLock{UsdLayersTreeView->RefreshStateLock};
 
-	if ( UsdLayersTreeView )
+		UsdLayersTreeView->bUpdateState = (SUsdLayersTreeView::ELayersTreeViewState
+		)FMath::Max((uint8)SUsdLayersTreeView::ELayersTreeViewState::RefreshNeeded, (uint8)UsdLayersTreeView->bUpdateState);
+	}
+}
+
+void SUsdStage::RequestFullRefresh()
+{
+	if (UsdLayersTreeView)
+	{
+		FWriteScopeLock LayersTreeViewLock{UsdLayersTreeView->RefreshStateLock};
+
+		UsdLayersTreeView->bUpdateState = SUsdLayersTreeView::ELayersTreeViewState::ResyncNeeded;
+	}
+
+	if (UsdStageTreeView)
+	{
+		FWriteScopeLock StageTreeViewLock{UsdStageTreeView->RefreshStateLock};
+
+		UsdStageTreeView->bNeedsFullUpdate = true;
+
+		if (UsdStageTreeView->PrimPathToSelect.IsEmpty())
+		{
+			UsdStageTreeView->PrimPathToSelect = SelectedPrimPath;
+		}
+	}
+}
+
+void SUsdStage::OnSlateTick(float Time)
+{
+	if (bEditorIsShuttingDown)
+	{
+		return;
+	}
+
+	SUsdLayersTreeView::ELayersTreeViewState LayersTreeViewState = SUsdLayersTreeView::ELayersTreeViewState::NoRefreshNeeded;
+	bool bStageTreeViewNeedsFullUpdate = false;
+	FString PrimPathToSelect;
+	TArray<TPair<FString, bool>> PrimsToRefresh;
+
+	if (UsdLayersTreeView)
+	{
+		FWriteScopeLock LayersTreeViewLock{UsdLayersTreeView->RefreshStateLock};
+		Swap(LayersTreeViewState, UsdLayersTreeView->bUpdateState);
+	}
+
+	if (UsdStageTreeView)
 	{
+		FWriteScopeLock StageTreeViewLock{UsdStageTreeView->RefreshStateLock};
+		Swap(bStageTreeViewNeedsFullUpdate, UsdStageTreeView->bNeedsFullUpdate);
+		Swap(PrimPathToSelect, UsdStageTreeView->PrimPathToSelect);
+		Swap(PrimsToRefresh, UsdStageTreeView->PrimsToRefresh);
+	}
+
+	if (LayersTreeViewState != SUsdLayersTreeView::ELayersTreeViewState::NoRefreshNeeded)
+	{
+		// May be nullptr, but that is ok. Its how the widgets are reset
+		AUsdStageActor* StageActor = ViewModel.UsdStageActor.Get();
+
 		UE::FUsdStageWeak IsolatedStage = StageActor ? StageActor->GetIsolatedUsdStage() : UE::FUsdStage{};
 		UE::FUsdStageWeak BaseStage = StageActor ? StageActor->GetBaseUsdStage() : UE::FUsdStage{};
 
 		// The layers tree view always needs to receive both the full stage as well as the isolated
-		const bool bResync = true;
-		UsdLayersTreeView->Refresh( BaseStage, IsolatedStage, bResync );
+		const bool bResync = LayersTreeViewState == SUsdLayersTreeView::ELayersTreeViewState::ResyncNeeded;
+		UsdLayersTreeView->Refresh(BaseStage, IsolatedStage, bResync);
 	}
 
-	if ( UsdStageTreeView )
+	if (bStageTreeViewNeedsFullUpdate || !PrimPathToSelect.IsEmpty() || PrimsToRefresh.Num() > 0)
 	{
-		UE::FUsdStageWeak CurrentStage = StageActor
-			? static_cast< const AUsdStageActor* >( StageActor )->GetUsdStage()
-			: UE::FUsdStage{};
+		// May be nullptr, but that is ok. Its how the widgets are reset
+		AUsdStageActor* StageActor = ViewModel.UsdStageActor.Get();
+
+		// Don't bother refreshing individual prims if we'll do a full refresh anyway
+		if (!bStageTreeViewNeedsFullUpdate)
+		{
+			TSet<FString> RefreshedPrims;
+			TSet<FString> ResyncedPrims;
+
+			// Note how we go from oldest to newest update here, so after the loop we'll end up with LastResyncedPrim
+			// pointing at what it should, and PrimPathToSelect pointing at the last updated prim path to select
+			for (const TPair<FString, bool>& PrimAndResync : PrimsToRefresh)
+			{
+				const FString& PrimPath = PrimAndResync.Key;
+				const bool bResync = PrimAndResync.Value;
+
+				// Prevent from refreshing the same prim too many times on a single tick
+				if (ResyncedPrims.Contains(PrimPath))
+				{
+					continue;
+				}
+				if (!bResync && RefreshedPrims.Contains(PrimPath))
+				{
+					continue;
+				}
+				if (bResync)
+				{
+					ResyncedPrims.Add(PrimPath);
+				}
+				else
+				{
+					RefreshedPrims.Add(PrimPath);
+				}
+
+				UsdStageTreeView->RefreshPrim(PrimPath, bResync);
+				{
+					// RefreshPrim may internally set UsdStageTreeView->bNeedsFullUpdate to true
+					FWriteScopeLock StageTreeViewLock{UsdStageTreeView->RefreshStateLock};
+					bStageTreeViewNeedsFullUpdate |= UsdStageTreeView->bNeedsFullUpdate;
+					UsdStageTreeView->bNeedsFullUpdate = false;
+				}
+
+				// If we already have a PrimPathToSelect, we're going to already refresh the info panel already or end up
+				// looking at another prim anyway, so don't bother refreshing the info widget due to a prim update
+				if (PrimPathToSelect.IsEmpty())
+				{
+					const bool bViewingTheUpdatedPrim = SelectedPrimPath.Equals(PrimPath, ESearchCase::IgnoreCase);
+					const bool bViewingStageProperties = SelectedPrimPath.IsEmpty() || SelectedPrimPath == TEXT("/");
+					const bool bStageUpdated = PrimPath == TEXT("/");
+
+					if (UsdPrimInfoWidget &&
+						StageActor &&
+						(bViewingTheUpdatedPrim || (bViewingStageProperties && bStageUpdated))
+					)
+					{
+						PrimPathToSelect = PrimPath;
+					}
+				}
+			}
+		}
 
-		UsdStageTreeView->Refresh( CurrentStage );
+		if (bStageTreeViewNeedsFullUpdate)
+		{
+			UsdStageTreeView->Refresh(GetCurrentStage());
+		}
 
-		// Refresh will generate brand new RootItems, so let's immediately select the new item
-		// that corresponds to the prim we're supposed to be selecting
-		UsdStageTreeView->SetSelectedPrimPaths( { SelectedPrimPath } );
+		if (bStageTreeViewNeedsFullUpdate || !PrimPathToSelect.IsEmpty())
+		{
+			// Refresh will generate brand new RootItems, so let's immediately select the new item
+			// that corresponds to the prim we're supposed to be selecting
+			UsdStageTreeView->SetSelectedPrimPaths(
+				{PrimPathToSelect.IsEmpty() ? SelectedPrimPath : PrimPathToSelect}
+			);
+
+			// If we changed/redid the selection, let's make sure that our info panel will update too (right below this)
+			PrimPathToSelect = SelectedPrimPath;
+		}
 
-		UsdStageTreeView->RequestTreeRefresh();
+		// Refresh the actual slate widget if we changed any underlying data
+		if (bStageTreeViewNeedsFullUpdate || PrimsToRefresh.Num() > 0)
+		{
+			UsdStageTreeView->RequestTreeRefresh();
+		}
 	}
 
-	if ( UsdPrimInfoWidget )
+	if (UsdPrimInfoWidget && !PrimPathToSelect.IsEmpty())
 	{
-		UsdPrimInfoWidget->SetPrimPath( GetCurrentStage(), *SelectedPrimPath );
+		UsdPrimInfoWidget->SetPrimPath(GetCurrentStage(), *PrimPathToSelect);
+		PrimPathToSelect.Empty();
 	}
 }
 
diff --git a/Engine/Plugins/Importers/USDImporter/Source/USDStageEditor/Private/SUSDStage.h b/Engine/Plugins/Importers/USDImporter/Source/USDStageEditor/Private/SUSDStage.h
--- a/Engine/Plugins/Importers/USDImporter/Source/USDStageEditor/Private/SUSDStage.h
+++ b/Engine/Plugins/Importers/USDImporter/Source/USDStageEditor/Private/SUSDStage.h
@@ -21,6 +21,7 @@ struct FSlateBrush;
 namespace UE
 {
 	class FUsdPrim;
+	class FSdfPath;
 	class FUsdAttribute;
 }
 
@@ -98,7 +99,9 @@ protected:
 
 	void OpenStage( const TCHAR* FilePath );
 
-	void Refresh();
+	void RequestLayersTreeViewRefresh();
+	void RequestFullRefresh();
+	void OnSlateTick(float Time);
 
 	void OnViewportSelectionChanged( UObject* NewSelection );
 
diff --git a/Engine/Plugins/Importers/USDImporter/Source/USDStageEditor/Private/SUSDStageTreeView.cpp b/Engine/Plugins/Importers/USDImporter/Source/USDStageEditor/Private/SUSDStageTreeView.cpp
--- a/Engine/Plugins/Importers/USDImporter/Source/USDStageEditor/Private/SUSDStageTreeView.cpp
+++ b/Engine/Plugins/Importers/USDImporter/Source/USDStageEditor/Private/SUSDStageTreeView.cpp
@@ -452,18 +452,20 @@ void SUsdStageTreeView::RefreshPrim( const FString& PrimPath, bool bResync )
 	// We couldn't find the target prim, do a full refresh instead
 	else
 	{
-		Refresh( UsdStage );
+		FWriteScopeLock StageTreeViewLock{RefreshStateLock};
+		bNeedsFullUpdate = true;
 	}
 
 	if ( bResync )
 	{
-		RequestTreeRefresh();
+		FWriteScopeLock StageTreeViewLock{RefreshStateLock};
+		bNeedsFullUpdate = true;
 	}
 }
 
 FUsdPrimViewModelPtr SUsdStageTreeView::GetItemFromPrimPath( const FString& PrimPath )
 {
-	FScopedUnrealAllocs UnrealAllocs; // RefreshPrim can be called by a delegate for which we don't know the active allocator
+	FScopedUnrealAllocs UnrealAllocs;
 
 	UE::FSdfPath UsdPrimPath( *PrimPath );
 
diff --git a/Engine/Plugins/Importers/USDImporter/Source/USDStageEditor/Private/SUSDStageTreeView.h b/Engine/Plugins/Importers/USDImporter/Source/USDStageEditor/Private/SUSDStageTreeView.h
--- a/Engine/Plugins/Importers/USDImporter/Source/USDStageEditor/Private/SUSDStageTreeView.h
+++ b/Engine/Plugins/Importers/USDImporter/Source/USDStageEditor/Private/SUSDStageTreeView.h
@@ -94,6 +94,23 @@ private:
 
 	void SelectItemsInternal( const TArray< FUsdPrimViewModelRef >& ItemsToSelect );
 
+public:
+	// We update only on slate tick, but can receive a USD notice at any point, from any thread.
+	// These variables are used to control what needs to be refreshed on the next tick
+
+	FString PrimPathToSelect;
+
+	bool bNeedsFullUpdate = false;
+
+	// We append requests to refresh prims on the stage tree view to the end of this array as they arrive,
+	// and then perform them all at once during OnSlateTick.
+	// Elements map from prim path to whether the update neeeds to be a resync or not.
+	TArray<TPair<FString, bool>> PrimsToRefresh;
+
+	// Use a lock because some of the notices that will update these refresh variables can be
+	// sent from TBB threads
+	mutable FRWLock RefreshStateLock;
+
 private:
 	// Should always be valid, we keep the one we're given on Refresh()
 	UE::FUsdStageWeak UsdStage;
