Add 'SupportedTargetPlatforms' directive to SimCacheUSD plugin spec.

diff --git a/Engine/Plugins/Experimental/ChaosFlesh/ChaosFlesh.uplugin b/Engine/Plugins/Experimental/ChaosFlesh/ChaosFlesh.uplugin
--- a/Engine/Plugins/Experimental/ChaosFlesh/ChaosFlesh.uplugin
+++ b/Engine/Plugins/Experimental/ChaosFlesh/ChaosFlesh.uplugin
@@ -64,11 +64,9 @@
 		},
         {
             "Name": "USDImporter",
-            "Enabled": true,
+			"Enabled": true,
 			"SupportedTargetPlatforms": [
-				"Win64",
-				"Mac",
-				"Linux"
+				"Win64"
 			]
         }
 	]
diff --git a/Engine/Plugins/Importers/USDImporter/Source/SimCacheUSD/Resources/plugInfo.json b/Engine/Plugins/Importers/USDImporter/Source/SimCacheUSD/Resources/plugInfo.json
--- a/Engine/Plugins/Importers/USDImporter/Source/SimCacheUSD/Resources/plugInfo.json
+++ b/Engine/Plugins/Importers/USDImporter/Source/SimCacheUSD/Resources/plugInfo.json
@@ -19,7 +19,7 @@
                 }
             },
             "Name": "SimCacheUSD",
-            #"LibraryPath": "@PLUG_INFO_LIBRARY_PATH@",
+            "LibraryPath": "../../../../../../../../Plugins/Importers/USDImporter/Binaries/Win64/UnrealEditor-SimCacheUSD.dll",
             #"ResourcePath": "@PLUG_INFO_RESOURCE_PATH@",
             "ResourcePath": "Resources",
             #"Root": "@PLUG_INFO_ROOT@",
diff --git a/Engine/Plugins/Importers/USDImporter/USDImporter.uplugin b/Engine/Plugins/Importers/USDImporter/USDImporter.uplugin
--- a/Engine/Plugins/Importers/USDImporter/USDImporter.uplugin
+++ b/Engine/Plugins/Importers/USDImporter/USDImporter.uplugin
@@ -76,7 +76,10 @@
 		{
 			"Name": "SimCacheUSD",
 			"Type": "Editor",
-			"LoadingPhase": "Default"
+			"LoadingPhase": "Default",
+			"SupportedTargetPlatforms": [
+				"Win64"
+			]
 		},
 		{
 			"Name": "USDTests",
