USD: Fixed wrong subframe interpolation for USD geometry cache and wrong initialization with start offset.

diff --git a/Engine/Plugins/Importers/USDImporter/Source/GeometryCacheUSD/Private/GeometryCacheTrackUSD.cpp b/Engine/Plugins/Importers/USDImporter/Source/GeometryCacheUSD/Private/GeometryCacheTrackUSD.cpp
--- a/Engine/Plugins/Importers/USDImporter/Source/GeometryCacheUSD/Private/GeometryCacheTrackUSD.cpp
+++ b/Engine/Plugins/Importers/USDImporter/Source/GeometryCacheUSD/Private/GeometryCacheTrackUSD.cpp
@@ -101,7 +101,7 @@ const int32 UGeometryCacheTrackUsd::FindSampleIndexFromTime(const float Time, co
 		SampleTime = GeometyCacheHelpers::WrapAnimationTime(Time, Duration);
 	}
 	// Which is converted to an index
-	int32 FrameIndex = FMath::RoundToInt(Time * FramesPerSecond);
+	int32 FrameIndex = FMath::FloorToInt(Time * FramesPerSecond);
 
 	// The final computed frame index must be offset by the start frame index
 	return FMath::Clamp(FrameIndex + StartFrameIndex, StartFrameIndex, EndFrameIndex - 1);
@@ -125,8 +125,8 @@ void UGeometryCacheTrackUsd::GetFractionalFrameIndexFromTime(const float Time, c
 	// Time at ThisFrameIndex with index normalized to 0
 	const float FrameIndexTime = (OutFrameIndex - StartFrameIndex) / FramesPerSecond;
 	OutFraction = float((AdjustedTime - FrameIndexTime) * FramesPerSecond);
-	// The fractional part is clamped to (-0.5, 0.5) for subframe interpolation that is centered around FrameIndex
-	OutFraction = FMath::Clamp(OutFraction, -0.5f, 0.5f);
+	// The fractional part is clamped to (0, 1) for subframe interpolation since the FrameIndex is floored
+	OutFraction = FMath::Clamp(OutFraction, 0.0f, 1.0f);
 }
 
 const FGeometryCacheTrackSampleInfo& UGeometryCacheTrackUsd::GetSampleInfo(float Time, bool bLooping)
diff --git a/Engine/Plugins/Importers/USDImporter/Source/USDSchemas/Private/USDGeometryCacheTranslator.cpp b/Engine/Plugins/Importers/USDImporter/Source/USDSchemas/Private/USDGeometryCacheTranslator.cpp
--- a/Engine/Plugins/Importers/USDImporter/Source/USDSchemas/Private/USDGeometryCacheTranslator.cpp
+++ b/Engine/Plugins/Importers/USDImporter/Source/USDSchemas/Private/USDGeometryCacheTranslator.cpp
@@ -979,8 +979,13 @@ void FUsdGeometryCacheTranslator::UpdateComponents(USceneComponent* SceneCompone
 				}
 			}
 
+			// The Time from the stage has to be adjusted to be relative to the time range of the geometry cache
+			// by applying the start offset. Thus, the adjusted time has to be clamped between 0 and the duration.
+			const float Duration = GeometryCache->CalculateDuration();
 			const double FramesPerSecond = Context->Stage.GetTimeCodesPerSecond();
-			const float AdjustedTime = static_cast<float>(TimeCode / FramesPerSecond - LayerStartOffsetSeconds);
+			float AdjustedTime = static_cast<float>(TimeCode / FramesPerSecond - LayerStartOffsetSeconds);
+			AdjustedTime = FMath::Clamp(AdjustedTime, 0.0f, Duration);
+
 			const bool bIsRunning = true;
 			const bool bIsBackwards = false;
 			GeometryCacheComponent->TickAtThisTime(AdjustedTime, bIsRunning, bIsBackwards, bIsLooping);
