USD: Add proper support for exporting LevelSequences with Dynamic Bindings to USD (being our Dynamic Bindings or any other);

Process exported prim paths when exporting LevelSequences so that they match the new pattern of exporting stage actors with payloads to their currently loaded stages, if that is applicable.

diff --git a/Engine/Plugins/Importers/USDImporter/Source/USDExporter/Private/LevelSequenceExporterUSD.cpp b/Engine/Plugins/Importers/USDImporter/Source/USDExporter/Private/LevelSequenceExporterUSD.cpp
--- a/Engine/Plugins/Importers/USDImporter/Source/USDExporter/Private/LevelSequenceExporterUSD.cpp
+++ b/Engine/Plugins/Importers/USDImporter/Source/USDExporter/Private/LevelSequenceExporterUSD.cpp
@@ -7,11 +7,13 @@
 #include "UnrealUSDWrapper.h"
 #include "USDClassesModule.h"
 #include "USDConversionUtils.h"
+#include "USDErrorUtils.h"
+#include "USDExporterModule.h"
 #include "USDLayerUtils.h"
 #include "USDLog.h"
-#include "USDExporterModule.h"
 #include "USDOptionsWindow.h"
 #include "USDPrimConversion.h"
+#include "USDStageActor.h"
 #include "USDUnrealAssetInfo.h"
 
 #include "UsdWrappers/SdfLayer.h"
@@ -48,6 +50,8 @@
 #include "Tracks/MovieSceneSubTrack.h"
 #include "UObject/UObjectGlobals.h"
 
+#define LOCTEXT_NAMESPACE "LevelSequenceExporterUSD"
+
 namespace UE
 {
 	namespace LevelSequenceExporterUSD
@@ -460,43 +464,6 @@ namespace UE
 				}
 			};
 
-			// Returns the UObject possessed by a possessable from a MovieScene
-			UObject* LocateBoundObject(
-				const UMovieSceneSequence& RootMovieSceneSequence,
-				const UMovieSceneSequence& CurrentMovieSceneSequence,
-				FMovieSceneSequenceIDRef TemplateID,
-				const FMovieScenePossessable& Possessable,
-				const TSharedRef<FLevelSequenceHidingSpawnRegister>& SpawnRegister
-			)
-			{
-				UMovieScene* MovieScene = CurrentMovieSceneSequence.GetMovieScene();
-				const FGuid& Guid = Possessable.GetGuid();
-				const FGuid& ParentGuid = Possessable.GetParent();
-
-				// If we have a parent guid, we must provide the object as a context because really the binding path
-				// will just contain the component name
-				UObject* ParentContext = nullptr;
-				if ( ParentGuid.IsValid() )
-				{
-					if ( FMovieScenePossessable* ParentPossessable = MovieScene->FindPossessable( ParentGuid ) )
-					{
-						ParentContext = LocateBoundObject( RootMovieSceneSequence, CurrentMovieSceneSequence, TemplateID, *ParentPossessable, SpawnRegister );
-					}
-					else if ( FMovieSceneSpawnable* ParentSpawnable = MovieScene->FindSpawnable( ParentGuid ) )
-					{
-						ParentContext = SpawnRegister->GetExistingSpawn( RootMovieSceneSequence, TemplateID, ParentSpawnable->GetGuid() );
-					}
-				}
-
-				TArray<UObject*, TInlineAllocator<1>> Objects = CurrentMovieSceneSequence.LocateBoundObjects( Guid, ParentContext );
-				if ( Objects.Num() > 0 )
-				{
-					return Objects[ 0 ];
-				}
-
-				return nullptr;
-			}
-
 			bool IsTrackAnimated( const UMovieSceneTrack* Track )
 			{
 				for ( const UMovieSceneSection* Section : Track->GetAllSections() )
@@ -772,6 +739,13 @@ namespace UE
 				// Index from UObject to FGuid because we may have multiple spawned objects for a given spawnable Guid
 				TMap<UObject*, FGuid> BoundObjects;
 
+				// Collect any USD-related DynamicBinding. The idea being that if we find any, we're likely looking at a
+				// loaded USD Stage that's going to be exported, and the possessable is one of the transient actors and
+				// components. It that's the case, we don't want to just come up with a random name for the prim based on
+				// the actor/component path, but instead want to use the prim path that it has been given on the dynamic binding,
+				// if any
+				TMap<FGuid, const FMovieSceneDynamicBinding*> DynamicBindings;
+
 				const TArray<FMovieSceneSequenceID>* InstancesOfThisSequence = SequenceInstances.Find( &MovieSceneSequence );
 				if ( !InstancesOfThisSequence )
 				{
@@ -784,8 +758,14 @@ namespace UE
 					return;
 				}
 
+				// Force spawn spawnables again here so that they exist on the Register map when we rely on FindBoundObjects
+				// to resolve bindings
+				PreSpawnSpawnables(Context, MovieSceneSequence);
+
 				for ( FMovieSceneSequenceID SequenceInstance : *InstancesOfThisSequence )
 				{
+					FMovieSceneObjectCache& ObjectCache = Context.Sequencer->State.GetObjectCache(SequenceInstance);
+
 					// Possessables
 					int32 NumPossessables = MovieScene->GetPossessableCount();
 					for ( int32 Index = 0; Index < NumPossessables; ++Index )
@@ -793,17 +773,34 @@ namespace UE
 						const FMovieScenePossessable& Possessable = MovieScene->GetPossessable( Index );
 						const FGuid& Guid = Possessable.GetGuid();
 
-						UObject* BoundObject = LocateBoundObject(
-							*RootSequence,
-							MovieSceneSequence,
-							SequenceInstance,
-							Possessable,
-							Context.SpawnRegister
-						);
+						UObject* BoundObject = nullptr;
+
+						// Go through FMovieSceneObjectCache and FindBoundObjects because that will also evaluate DynamicBindings.
+						// Note that we need to make sure that PreSpawnSpawnables has been called above this (at all, but also at
+						// least once *after* Context.SpawnRegister->CleanUp(), if that has been called). The idea here is that FindBoundObjects
+						// will manage to find the binding even it their "parent context" is a spawnable (e.g. if it's a possessable component of
+						// a spawnable) and also uses DynamicBindings, which is great! For it to be able to find our spawns however, the
+						// spawnables must be *currently* spawned.
+						// It doesn't help at all that our custom spawn register only hides stuff instead of destroying them, because even if the
+						// UObject itself still exists and is just hidden, having been "despawned" means the spawnable has been removed from the
+						// "Register" member of FMovieSceneSpawnRegister, and so the base part of the spawn register "doesn't know about it".
+						// For reference, check how FMovieSceneObjectCache::UpdateBindings (called by FindBoundObjects) will end up calling
+						// "Player.GetSpawnRegister().FindSpawnedObject", and observe how that in turn just checks the "Register" member...
+						// Ideally we could tweak a bit how the "Register" member is used, but that is part of the base FMovieSceneSpawnRegister.
+						TArrayView<TWeakObjectPtr<UObject>> ObjectWeakPtrs = ObjectCache.FindBoundObjects(Guid, *Context.Sequencer);
+						if (ObjectWeakPtrs.Num() > 0)
+						{
+							BoundObject = ObjectWeakPtrs[0].Get();
+						}
 
 						if ( BoundObject && ( BoundObject->IsA<USceneComponent>() || BoundObject->IsA<AActor>() ) )
 						{
 							BoundObjects.Add( BoundObject, Guid );
+
+							if (Possessable.DynamicBinding.Function)
+							{
+								DynamicBindings.Add(Guid, &Possessable.DynamicBinding);
+							}
 						}
 					}
 
@@ -833,6 +830,11 @@ namespace UE
 						if ( BoundObject && ( BoundObject->IsA<USceneComponent>() || BoundObject->IsA<AActor>() ) )
 						{
 							BoundObjects.Add( BoundObject, Guid );
+
+							if (Spawnable.DynamicBinding.Function)
+							{
+								DynamicBindings.Add(Guid, &Spawnable.DynamicBinding);
+							}
 						}
 					}
 				}
@@ -877,6 +879,78 @@ namespace UE
 						continue;
 					}
 
+					// If this binding has one of our dynamic bindings set up pointing to a valid prim path, let's use that path
+					// instead of using our generated PrimPath, as that one will better match the prim paths that we'll get when
+					// opening a referenced stage via an exported UsdStageActor
+					if (const FMovieSceneDynamicBinding* DynamicBinding = DynamicBindings.FindRef(Guid))
+					{
+						if (const FMovieSceneDynamicBindingPayloadVariable* FoundPrimPathPayload = DynamicBinding->PayloadVariables.Find(TEXT("PrimPath")))
+						{
+							FString PrimPathInSourceStage = FoundPrimPathPayload->Value;
+							if (!PrimPathInSourceStage.IsEmpty())
+							{
+								AActor* PossibleParentStageActor = BoundActor;
+								while (PossibleParentStageActor && !PossibleParentStageActor->IsA<AUsdStageActor>())
+								{
+									PossibleParentStageActor = PossibleParentStageActor->GetAttachParentActor();
+								}
+
+								// Our possessable has a dynamic binding with a "PrimPath" payload variable and is a child of a stage actor,
+								// for now let's consider this enough to consider this is one of our dynamic bindings
+								if (AUsdStageActor* StageActor = Cast<AUsdStageActor>(PossibleParentStageActor))
+								{
+									FString ParentStageActorPrimPathOnExport = UsdUtils::GetPrimPathForObject(
+										StageActor,
+										TEXT(""),
+										Context.ExportOptions && Context.ExportOptions->LevelExportOptions.bExportActorFolders
+									);
+									if (!ParentStageActorPrimPathOnExport.IsEmpty())
+									{
+										UE::FSdfPath SdfPrimPathInSourceStage{*PrimPathInSourceStage};
+
+										UE::FUsdStage LoadedStage = StageActor->GetBaseUsdStage();
+										UE::FSdfPath DefaultPrimPath = LoadedStage.GetDefaultPrim().GetPrimPath();
+
+										if (!LoadedStage.GetRootLayer().IsAnonymous() &&
+											!DefaultPrimPath.IsEmpty() &&
+											SdfPrimPathInSourceStage.HasPrefix(DefaultPrimPath))
+										{
+											// Note that it's perfectly fine if this ends up being just "."
+											UE::FSdfPath RelativePrimPathInSourceStage = SdfPrimPathInSourceStage.MakeRelativePath(DefaultPrimPath);
+
+											FString PrimPathRelativeToParentStageActor = UE::FSdfPath{*ParentStageActorPrimPathOnExport}
+												.AppendPath(RelativePrimPathInSourceStage)
+												.GetString();
+
+											if (!PrimPathRelativeToParentStageActor.IsEmpty())
+											{
+												PrimPath = PrimPathRelativeToParentStageActor;
+											}
+										}
+										else
+										{
+											FUsdLogManager::LogMessage(
+												EMessageSeverity::Warning,
+												FText::Format(
+													LOCTEXT(
+														"NonIdealComposition",
+														"Exported animation for prim '{0}' may not compose correctly with the prims from referenced "
+														"layer '{1}' on the exported stage for the LevelSequence '{2}'. For best results, make sure "
+														"the referenced layer is saved to disk (i.e. not anonymous), has a defaultPrim setup, and "
+														"that the animation tracks are only bound to prims that are descendents of the defaultPrim."
+													),
+													FText::FromString(PrimPathInSourceStage),
+													FText::FromString(LoadedStage.GetRootLayer().GetIdentifier()),
+													FText::FromString(Context.RootSequence.GetPathName())
+												)
+											);
+										}
+									}
+								}
+							}
+						}
+					}
+
 					FString SchemaName = UsdUtils::GetSchemaNameForComponent( *BoundComponent );
 					if ( SchemaName.IsEmpty() )
 					{
@@ -992,6 +1066,11 @@ namespace UE
 					return;
 				}
 
+				// Hide all our spawns again so we can pretend they haven't actually spawned and let the sequence spawn them
+				// as it plays
+				Context.SpawnRegister->bDestroyingJustHides = true;
+				Context.SpawnRegister->CleanUp( *Context.Sequencer );
+
 				const TRange< FFrameNumber > PlaybackRange = MovieScene->GetPlaybackRange();
 				const FFrameRate Resolution = MovieScene->GetTickResolution();
 				const FFrameRate DisplayRate = MovieScene->GetDisplayRate();
@@ -1362,6 +1441,8 @@ bool ULevelSequenceExporterUsd::ExportBinary( UObject* Object, const TCHAR* Type
 		return false;
 	}
 
+	FScopedUsdMessageLog UsdMessageLog;
+
 	FSequencerInitParams Params;
 	Params.RootSequence = LevelSequence;
 	Params.SpawnRegister = SpawnRegister;
@@ -1433,10 +1514,6 @@ bool ULevelSequenceExporterUsd::ExportBinary( UObject* Object, const TCHAR* Type
 		}
 	}
 
-	// Hide all our spawns again so we can pretend they haven't actually spawned yet
-	Context.SpawnRegister->bDestroyingJustHides = true;
-	Context.SpawnRegister->CleanUp( *Context.Sequencer );
-
 	// Capture this first because when we launch UExporter::RunAssetExportTask the CurrentFileName will change
 	const FString TargetFileName = UExporter::CurrentFilename;
 
@@ -1541,3 +1618,5 @@ bool ULevelSequenceExporterUsd::ExportBinary( UObject* Object, const TCHAR* Type
 	return false;
 #endif // #if USE_USD_SDK
 }
+
+#undef LOCTEXT_NAMESPACE
\ No newline at end of file
diff --git a/Engine/Plugins/Importers/USDImporter/Source/USDExporter/USDExporter.Build.cs b/Engine/Plugins/Importers/USDImporter/Source/USDExporter/USDExporter.Build.cs
--- a/Engine/Plugins/Importers/USDImporter/Source/USDExporter/USDExporter.Build.cs
+++ b/Engine/Plugins/Importers/USDImporter/Source/USDExporter/USDExporter.Build.cs
@@ -53,6 +53,7 @@ namespace UnrealBuildTool.Rules
 					"StaticMeshDescription",
 					"UnrealEd",
 					"USDClasses",
+					"USDStage",
 					"USDStageImporter", // For USDOptionsWindow
 					"USDUtilities",
 				}
diff --git a/Engine/Plugins/Importers/USDImporter/Source/UnrealUSDWrapper/Private/UsdWrappers/SdfPath.cpp b/Engine/Plugins/Importers/USDImporter/Source/UnrealUSDWrapper/Private/UsdWrappers/SdfPath.cpp
--- a/Engine/Plugins/Importers/USDImporter/Source/UnrealUSDWrapper/Private/UsdWrappers/SdfPath.cpp
+++ b/Engine/Plugins/Importers/USDImporter/Source/UnrealUSDWrapper/Private/UsdWrappers/SdfPath.cpp
@@ -277,6 +277,15 @@ namespace UE
 #endif // #if USE_USD_SDK
 	}
 
+	FSdfPath FSdfPath::AppendPath(const UE::FSdfPath& NewRelativeSuffix) const
+	{
+#if USE_USD_SDK
+		return FSdfPath(Impl->PxrSdfPath.Get().AppendPath(NewRelativeSuffix.Impl->PxrSdfPath.Get()));
+#else
+		return FSdfPath();
+#endif // #if USE_USD_SDK
+	}
+
 	FSdfPath FSdfPath::AppendChild( const TCHAR* ChildName ) const
 	{
 #if USE_USD_SDK
@@ -336,6 +345,24 @@ namespace UE
 		return Impl->PxrSdfPath.Get().HasPrefix(Prefix);
 #else
 		return false;
+#endif // #if USE_USD_SDK
+	}
+
+	FSdfPath FSdfPath::MakeAbsolutePath(const UE::FSdfPath& Anchor) const
+	{
+#if USE_USD_SDK
+		return FSdfPath(Impl->PxrSdfPath.Get().MakeAbsolutePath(Anchor.Impl->PxrSdfPath.Get()));
+#else
+		return FSdfPath();
+#endif // #if USE_USD_SDK
+	}
+
+	FSdfPath FSdfPath::MakeRelativePath(const UE::FSdfPath& Anchor) const
+	{
+#if USE_USD_SDK
+		return FSdfPath(Impl->PxrSdfPath.Get().MakeRelativePath(Anchor.Impl->PxrSdfPath.Get()));
+#else
+		return FSdfPath();
 #endif // #if USE_USD_SDK
 	}
 }
diff --git a/Engine/Plugins/Importers/USDImporter/Source/UnrealUSDWrapper/Public/UsdWrappers/SdfPath.h b/Engine/Plugins/Importers/USDImporter/Source/UnrealUSDWrapper/Public/UsdWrappers/SdfPath.h
--- a/Engine/Plugins/Importers/USDImporter/Source/UnrealUSDWrapper/Public/UsdWrappers/SdfPath.h
+++ b/Engine/Plugins/Importers/USDImporter/Source/UnrealUSDWrapper/Public/UsdWrappers/SdfPath.h
@@ -75,6 +75,7 @@ namespace UE
 		FSdfPath ReplaceName( const TCHAR* NewLeafName ) const;
 
 		FSdfPath GetParentPath() const;
+		FSdfPath AppendPath(const UE::FSdfPath& NewRelativeSuffix) const;
 		FSdfPath AppendChild( const TCHAR* ChildName ) const;
 		FSdfPath AppendProperty( FName PropertyName ) const;
 
@@ -86,6 +87,9 @@ namespace UE
 
 		bool HasPrefix(const UE::FSdfPath& Prefix) const;
 
+		FSdfPath MakeAbsolutePath(const FSdfPath& Anchor) const;
+		FSdfPath MakeRelativePath(const FSdfPath& Anchor) const;
+
 	private:
 		TUniquePtr< Internal::FSdfPathImpl > Impl;
 	};
