USD: Take care to also mark UMaterialInterface's EditorOnlyData as non-transient when importing, as that is not a traditional default subobject.

Any material triggers this issue.

It is caused by the fact that UMaterialInterfaces have this annoying UMaterialInterfaceEditorOnlyData "subobject" that they create on-demand. It has caused a bunch of issues in the past as well because they rename this subobject to follow the material's name, etc.

In this issue, the subobject was causing trouble because we create the materials with the transient flag, and so the subobject also becomes transient. When we PublishAsset for import, we clear the transient flag of the UMaterialInterface, but not of the UMaterialInterfaceEditorOnlyData.

I first tried fixing this by changing all asset NewObjects to remove the transient flag, and then add it back later (similar to how we handled the actor externalpackage thing). This didn't quite work, I think because this subobject is created sort of "on-demand" and not along with NewObject. At any case it's likely not the best idea anyway because we shouldn't rely on the subobject never being internally marked as transient again for some reason (e.g. when the UObject is renamed or something).

That's when I tried to clear the transient flag on the subobject on PublishAssets instead... only to discover we already do that. It just wasn't affecting the UMaterialInterfaceEditorOnlyData because it's not flagged as a "default subobject". Like the comment says I don't think we should mark *all* subobjects as non-transient as we don't know what may be there (maybe it is meant to be some transient object generated on-demand or something), and I don't think we should touch UMaterialInterface, so the only option left is to specially handle this case. I guess if we get this happening again with other asset types we can look into other, more general solutions?

diff --git a/Engine/Plugins/Importers/USDImporter/Source/USDStageImporter/Private/USDStageImporter.cpp b/Engine/Plugins/Importers/USDImporter/Source/USDStageImporter/Private/USDStageImporter.cpp
--- a/Engine/Plugins/Importers/USDImporter/Source/USDStageImporter/Private/USDStageImporter.cpp
+++ b/Engine/Plugins/Importers/USDImporter/Source/USDStageImporter/Private/USDStageImporter.cpp
@@ -787,6 +787,16 @@ namespace UsdStageImporterImpl
 		// We probably never want to make them public or standalone if they aren't already though
 		TArray<UObject*> Subobjects;
 		MovedAsset->GetDefaultSubobjects( Subobjects );
+		if (UMaterialInterface* Material = Cast<UMaterialInterface>(MovedAsset))
+		{
+			// Materials in particular have EditorOnlyData which behaves like a default subobject but kind of isn't flagged as one...
+			// Not sure at this point whether we should force non-transient on all subobjects of the asset, as that could also have edge cases that
+			// don't do what we expect.
+			if (UMaterialInterfaceEditorOnlyData* EditorOnlyData = Material->GetEditorOnlyData())
+			{
+				Subobjects.Add(EditorOnlyData);
+			}
+		}
 		for ( UObject* Subobject : Subobjects )
 		{
 			Subobject->ClearFlags( EObjectFlags::RF_Transient | EObjectFlags::RF_DuplicateTransient | EObjectFlags::RF_NonPIEDuplicateTransient );
