USD: Fix unnecessary UI update every frame of USD animation.

diff --git a/Engine/Plugins/Importers/USDImporter/Source/USDStage/Private/USDStageActor.cpp b/Engine/Plugins/Importers/USDImporter/Source/USDStage/Private/USDStageActor.cpp
--- a/Engine/Plugins/Importers/USDImporter/Source/USDStage/Private/USDStageActor.cpp
+++ b/Engine/Plugins/Importers/USDImporter/Source/USDStage/Private/USDStageActor.cpp
@@ -1802,6 +1802,19 @@ void AUsdStageActor::OnUsdObjectsChanged(const UsdUtils::FObjectChangesByPath& I
 	{
 		UsdAssetCache->RefreshStorage();
 	}
+
+#if WITH_EDITOR
+	if (GIsEditor && GEditor && !IsGarbageCollecting()) // Make sure we're not in standalone either
+	{
+		if (bHasResync)
+		{
+			GEditor->BroadcastLevelActorListChanged();
+		}
+
+		GEditor->RedrawLevelEditingViewports();
+	}
+#endif // WITH_EDITOR
+
 #endif // USE_USD_SDK
 }
 
@@ -2062,14 +2075,6 @@ void AUsdStageActor::UpdatePrim(const UE::FSdfPath& InUsdPrimPath, bool bResync,
 	{
 		UE::FUsdPrim PrimToExpand = GetOrLoadUsdStage().GetPrimAtPath(UsdPrimPath);
 		ExpandPrim(PrimToExpand, bResync, TranslationContext);
-
-#if WITH_EDITOR
-		if (GIsEditor && GEditor && !IsGarbageCollecting()) // Make sure we're not in standalone either
-		{
-			GEditor->BroadcastLevelActorListChanged();
-			GEditor->RedrawLevelEditingViewports();
-		}
-#endif // WITH_EDITOR
 	}
 }
 
@@ -2765,9 +2770,10 @@ void AUsdStageActor::LoadUsdStage()
 	FUsdStageActorImpl::RepairExternalSequencerBindings();
 
 #if WITH_EDITOR
-	if (GIsEditor && GEditor && !IsGarbageCollecting())
+	if (GIsEditor && GEditor && !IsGarbageCollecting()) // Make sure we're not in standalone either
 	{
 		GEditor->BroadcastLevelActorListChanged();
+		GEditor->RedrawLevelEditingViewports();
 	}
 #endif // WITH_EDITOR
 
@@ -2830,9 +2836,10 @@ void AUsdStageActor::UnloadUsdStage()
 #if WITH_EDITOR
 	// We can't emit this when garbage collecting as it may lead to objects being created
 	// (we may unload stage when going into PIE or other sensitive transitions)
-	if (GEditor && !IsGarbageCollecting())
+	if (GIsEditor && GEditor && !IsGarbageCollecting())
 	{
 		GEditor->BroadcastLevelActorListChanged();
+		GEditor->RedrawLevelEditingViewports();
 	}
 #endif // WITH_EDITOR
 
@@ -4395,14 +4402,6 @@ void AUsdStageActor::AnimatePrims()
 	}
 
 	TranslationContext->CompleteTasks();
-
-#if WITH_EDITOR
-	if (GIsEditor && GEditor && !IsGarbageCollecting())
-	{
-		GEditor->BroadcastLevelActorListChanged();
-		GEditor->RedrawLevelEditingViewports();
-	}
-#endif // WITH_EDITOR
 }
 
 FScopedBlockNoticeListening::FScopedBlockNoticeListening(AUsdStageActor* InStageActor)
