USD: Fixed possible crash when loading a USD file with geometry cache.

diff --git a/Engine/Plugins/Importers/USDImporter/Source/USDSchemas/Private/USDGeometryCacheTranslator.cpp b/Engine/Plugins/Importers/USDImporter/Source/USDSchemas/Private/USDGeometryCacheTranslator.cpp
--- a/Engine/Plugins/Importers/USDImporter/Source/USDSchemas/Private/USDGeometryCacheTranslator.cpp
+++ b/Engine/Plugins/Importers/USDImporter/Source/USDSchemas/Private/USDGeometryCacheTranslator.cpp
@@ -981,10 +981,13 @@ void FUsdGeometryCacheTranslator::UpdateComponents(USceneComponent* SceneCompone
 
 			// The Time from the stage has to be adjusted to be relative to the time range of the geometry cache
 			// by applying the start offset. Thus, the adjusted time has to be clamped between 0 and the duration.
-			const float Duration = GeometryCache->CalculateDuration();
 			const double FramesPerSecond = Context->Stage.GetTimeCodesPerSecond();
 			float AdjustedTime = static_cast<float>(TimeCode / FramesPerSecond - LayerStartOffsetSeconds);
-			AdjustedTime = FMath::Clamp(AdjustedTime, 0.0f, Duration);
+			if (GeometryCache)
+			{
+				const float Duration = GeometryCache->CalculateDuration();
+				AdjustedTime = FMath::Clamp(AdjustedTime, 0.0f, Duration);
+			}
 
 			const bool bIsRunning = true;
 			const bool bIsBackwards = false;
