USD: Fix crash when trying to parse SkeletalMeshes from Skeleton prims with no bones.

diff --git a/Engine/Plugins/Importers/USDImporter/Source/USDSchemas/Private/USDSkelRootTranslator.cpp b/Engine/Plugins/Importers/USDImporter/Source/USDSchemas/Private/USDSkelRootTranslator.cpp
--- a/Engine/Plugins/Importers/USDImporter/Source/USDSchemas/Private/USDSkelRootTranslator.cpp
+++ b/Engine/Plugins/Importers/USDImporter/Source/USDSchemas/Private/USDSkelRootTranslator.cpp
@@ -1327,7 +1327,7 @@ namespace UsdSkelRootTranslatorImpl
 				int32 NumRootBones = 0;
 				for (const SkeletalMeshImportData::FBone& Bone : SkeletonBones)
 				{
-					if (Bone.ParentIndex == -1)
+					if (Bone.ParentIndex == INDEX_NONE)
 					{
 						NumRootBones++;
 					}
diff --git a/Engine/Plugins/Importers/USDImporter/Source/USDUtilities/Private/USDSkeletalDataConversion.cpp b/Engine/Plugins/Importers/USDImporter/Source/USDUtilities/Private/USDSkeletalDataConversion.cpp
--- a/Engine/Plugins/Importers/USDImporter/Source/USDUtilities/Private/USDSkeletalDataConversion.cpp
+++ b/Engine/Plugins/Importers/USDImporter/Source/USDUtilities/Private/USDSkeletalDataConversion.cpp
@@ -2018,7 +2018,10 @@ bool UsdToUnreal::ConvertSkelAnim(
 	const FReferenceSkeleton& RefSkeleton = Skeleton->GetReferenceSkeleton();
 	const TArray<FMeshBoneInfo>& BoneInfo = RefSkeleton.GetRawRefBoneInfo();
 	int32 NumBones = BoneInfo.Num();
-	if ( InUsdSkeletonQuery.GetJointOrder().size() != NumBones )
+	int32 NumSkelQueryJoints = static_cast<int32>(InUsdSkeletonQuery.GetJointOrder().size());
+	// If we have zero bones on our skeleton we'll generate a dummy "Root" bone just so that Unreal can have a USkeleton asset,
+	// so we have to check for that case
+	if ((NumSkelQueryJoints != NumBones) && !(NumSkelQueryJoints == 0 && NumBones == 1 && BoneInfo[0].Name == TEXT("Root")))
 	{
 		return false;
 	}
