FUsdLevelSequenceHelperImpl: Replace cast by compile time check
#rnx

diff --git a/Engine/Plugins/Importers/USDImporter/Source/USDStage/Private/USDLevelSequenceHelper.cpp b/Engine/Plugins/Importers/USDImporter/Source/USDStage/Private/USDLevelSequenceHelper.cpp
--- a/Engine/Plugins/Importers/USDImporter/Source/USDStage/Private/USDLevelSequenceHelper.cpp
+++ b/Engine/Plugins/Importers/USDImporter/Source/USDStage/Private/USDLevelSequenceHelper.cpp
@@ -2161,14 +2161,14 @@ TrackType* FUsdLevelSequenceHelperImpl::AddTrack(const FName& TrackName, const U
 			return nullptr;
 		}
 
-		if (UMovieScenePropertyTrack* PropertyTrack = Cast<UMovieScenePropertyTrack>(Track))
+		if constexpr (std::is_base_of_v<UMovieScenePropertyTrack, TrackType>)
 		{
-			PropertyTrack->SetPropertyNameAndPath(TrackName, TrackName.ToString());
+			Track->SetPropertyNameAndPath(TrackName, TrackName.ToString());
 		}
 #if WITH_EDITOR
-		else if (UMovieSceneSkeletalAnimationTrack* SkeletalTrack = Cast<UMovieSceneSkeletalAnimationTrack>(Track))
+		else if constexpr (std::is_base_of_v<UMovieSceneSkeletalAnimationTrack, TrackType>)
 		{
-			SkeletalTrack->SetDisplayName(FText::FromName(TrackName));
+			Track->SetDisplayName(FText::FromName(TrackName));
 		}
 #endif // WITH_EDITOR
 	}
