USDImporter: remove redundant "fetch material by relationship" in UsdUtils::GetPrimMaterialAssignments()

Invocations of the FetchMaterialByComputingBoundMaterial function object were
always followed by invocations of the FetchMaterialByMaterialRelationship function
object. It turns out the latter is fully redundant though because it is effectively doing
the same work as the former but looking for an allPurpose material binding rather
than one of the purpose passed into GetPrimMaterialAssignments().

FetchMaterialByComputingBoundMaterial ends up calling
UsdShadeMaterialBindingAPI::ComputeBoundMaterial(), which will always look for
an allPurpose binding if one of the given purpose is not found first.
FetchMaterialByComputingBoundMaterial was therefore already doing the work of
FetchMaterialByMaterialRelationship, so if the former did not find a bound material,
the latter never would either.

diff --git a/Engine/Plugins/Importers/USDImporter/Source/USDUtilities/Private/USDGeomMeshConversion.cpp b/Engine/Plugins/Importers/USDImporter/Source/USDUtilities/Private/USDGeomMeshConversion.cpp
--- a/Engine/Plugins/Importers/USDImporter/Source/USDUtilities/Private/USDGeomMeshConversion.cpp
+++ b/Engine/Plugins/Importers/USDImporter/Source/USDUtilities/Private/USDGeomMeshConversion.cpp
@@ -1714,81 +1714,10 @@ UsdUtils::FUsdPrimMaterialAssignmentInfo UsdUtils::GetPrimMaterialAssignments(
 		pxr::UsdPrim ShadeMaterialPrim = ShadeMaterial.GetPrim();
 		if ( ShadeMaterialPrim )
 		{
-			pxr::SdfPath Path = ShadeMaterialPrim.GetPath();
-			std::string ShadingEngineName = ( ShadeMaterialPrim ? ShadeMaterialPrim.GetPrim() : UsdPrim.GetPrim() ).GetPrimPath().GetString();
-			if(ShadingEngineName.size() > 0 )
+			const std::string ShadingEngineName = ShadeMaterialPrim.GetPrimPath().GetString();
+			if (!ShadingEngineName.empty())
 			{
-				return UsdToUnreal::ConvertString( ShadingEngineName );
-			}
-		}
-
-		return {};
-	};
-
-	auto FetchMaterialByMaterialRelationship = [ &RenderContext ]( const pxr::UsdPrim& UsdPrim ) -> TOptional<FString>
-	{
-		if ( pxr::UsdRelationship Relationship = UsdPrim.GetRelationship( pxr::UsdShadeTokens->materialBinding ) )
-		{
-			pxr::SdfPathVector Targets;
-			Relationship.GetTargets( &Targets );
-
-			if ( Targets.size() > 0 )
-			{
-				const pxr::SdfPath& TargetMaterialPrimPath = Targets[0];
-				pxr::UsdPrim MaterialPrim = UsdPrim.GetStage()->GetPrimAtPath( TargetMaterialPrimPath );
-				if (!MaterialPrim)
-				{
-					FUsdLogManager::LogMessage(
-						EMessageSeverity::Warning,
-						FText::Format(LOCTEXT("IgnoringMaterialInvalid", "Ignoring target material '{0}' bound to prim '{1}' as that prim doesn't exist on this stage"),
-							FText::FromString(UsdToUnreal::ConvertPath(TargetMaterialPrimPath)),
-							FText::FromString(UsdToUnreal::ConvertPath(UsdPrim.GetPath()))
-						)
-					);
-					return {};
-				}
-
-				pxr::UsdShadeMaterial UsdShadeMaterial{ MaterialPrim };
-				if ( !UsdShadeMaterial )
-				{
-					FUsdLogManager::LogMessage(
-						EMessageSeverity::Warning,
-						FText::Format( LOCTEXT( "IgnoringMaterialNoSchema", "Ignoring material '{0}' bound to prim '{1}' as it does not possess the UsdShadeMaterial schema" ),
-							FText::FromString( UsdToUnreal::ConvertPath( TargetMaterialPrimPath ) ),
-							FText::FromString( UsdToUnreal::ConvertPath( UsdPrim.GetPath() ) )
-						)
-					);
-					return {};
-				}
-
-				// Ignore this material if UsdToUnreal::ConvertMaterial would as well
-				pxr::UsdShadeShader SurfaceShader = UsdShadeMaterial.ComputeSurfaceSource( RenderContext );
-				if ( !SurfaceShader )
-				{
-					FUsdLogManager::LogMessage(
-						EMessageSeverity::Warning,
-						FText::Format( LOCTEXT( "IgnoringMaterialSurface", "Ignoring material '{0}' bound to prim '{1}' as it contains no valid surface shader source for render context '{2}'" ),
-							FText::FromString( UsdToUnreal::ConvertPath( TargetMaterialPrimPath ) ),
-							FText::FromString( UsdToUnreal::ConvertPath( UsdPrim.GetPath() ) ),
-							FText::FromString( RenderContext == pxr::UsdShadeTokens->universalRenderContext ? TEXT( "universal" ) : UsdToUnreal::ConvertToken( RenderContext ) )
-						)
-					);
-					return {};
-				}
-
-				FString MaterialPrimPath = UsdToUnreal::ConvertPath( TargetMaterialPrimPath );
-				if ( Targets.size() > 1 )
-				{
-					FUsdLogManager::LogMessage(
-						EMessageSeverity::Warning,
-						FText::Format( LOCTEXT( "MoreThanOneMaterialBinding", "Found more than on material:binding targets on prim '{0}'. The first material ('{1}') will be used, and the rest ignored." ),
-							FText::FromString( UsdToUnreal::ConvertPath( UsdPrim.GetPath() ) ),
-							FText::FromString( MaterialPrimPath )
-						)
-					);
-				}
-
-				return MaterialPrimPath;
+				return UsdToUnreal::ConvertString(ShadingEngineName);
 			}
 		}
 
@@ -1912,21 +1841,7 @@ UsdUtils::FUsdPrimMaterialAssignmentInfo UsdUtils::GetPrimMaterialAssignments(
 				}
 			}
 
-			// Priority 0.3: material:binding relationship
-			if (!bHasAssignment)
-			{
-				if (TOptional<FString> TargetMaterial = FetchMaterialByMaterialRelationship(GeomSubsetPrim))
-				{
-					FUsdPrimMaterialSlot& Slot = Result.Slots.Emplace_GetRef();
-					Slot.MaterialSource = TargetMaterial.GetValue();
-					Slot.AssignmentType = UsdUtils::EPrimAssignmentType::MaterialPrim;
-					Slot.bMeshIsDoubleSided = bIsDoubleSided;
-					Slot.PrimPaths.Add(GeomSubsetPath);
-					bHasAssignment = true;
-				}
-			}
-
-			// Priority 0.4: Create a section anyway so that we always get a slot for each geom subset.
+			// Priority 0.3: Create a section anyway so that we always get a slot for each geom subset.
 			// We leave the assignment type cleared here, and will fill this in later with whatever we
 			// extract as a "main" material assignment.
 			// Note that we may have yet another "leftover" slot if our partition doesn't specify all faces,
@@ -2030,23 +1945,7 @@ UsdUtils::FUsdPrimMaterialAssignmentInfo UsdUtils::GetPrimMaterialAssignments(
 			}
 		}
 
-		// Priority 3: material:binding relationship directly on the prim (not sure why this is a separate step, but it came from
-		// IUsdPrim::GetGeometryMaterials. I bumped it in priority as the GeomSubsets do the same)
-		if (!bHasMainAssignment)
-		{
-			if (TOptional<FString> TargetMaterial = FetchMaterialByMaterialRelationship(UsdPrim))
-			{
-				FUsdPrimMaterialSlot& Slot = Result.Slots.Emplace_GetRef();
-				Slot.MaterialSource = TargetMaterial.GetValue();
-				Slot.AssignmentType = UsdUtils::EPrimAssignmentType::MaterialPrim;
-				Slot.bMeshIsDoubleSided = bIsDoubleSided;
-				Slot.PrimPaths.Add(MeshPrimPath);
-
-				bHasMainAssignment = true;
-			}
-		}
-
-		// Priority 4: vertex color material using displayColor/displayOpacity information for the entire mesh
+		// Priority 3: vertex color material using displayColor/displayOpacity information for the entire mesh
 		// Note: This will in general always succeed for any mesh prim, as the schema will provide fallback values
 		// for displayColor and displayOpacity
 		if (!bHasMainAssignment)
