USDImporter: align sublayer manipulation API on FSdfLayer wrapper with USD API

The utility functions in UsdUtils for creating and inserting sublayers perform additional
work that may not always be desirable such as authoring a root/default prim, so wrapping
the sublayer manipulation functions from the USD API allows direct access to edit sublayer
paths.

diff --git a/Engine/Plugins/Importers/USDImporter/Source/UnrealUSDWrapper/Private/UsdWrappers/SdfLayer.cpp b/Engine/Plugins/Importers/USDImporter/Source/UnrealUSDWrapper/Private/UsdWrappers/SdfLayer.cpp
--- a/Engine/Plugins/Importers/USDImporter/Source/UnrealUSDWrapper/Private/UsdWrappers/SdfLayer.cpp
+++ b/Engine/Plugins/Importers/USDImporter/Source/UnrealUSDWrapper/Private/UsdWrappers/SdfLayer.cpp
@@ -497,19 +497,6 @@ namespace UE
 		return false;
 	}
 
-	template<typename PtrType>
-	int64 FSdfLayerBase<PtrType>::GetNumSubLayerPaths() const
-	{
-#if USE_USD_SDK
-		if ( const PtrType& Ptr = Impl->GetInner() )
-		{
-			return ( int64 ) Ptr->GetNumSubLayerPaths();
-		}
-#endif // #if USE_USD_SDK
-
-		return 0;
-	}
-
 	template<typename PtrType>
 	double FSdfLayerBase<PtrType>::GetStartTimeCode() const
 	{
@@ -628,62 +615,114 @@ namespace UE
 	}
 
 	template<typename PtrType>
-	void FSdfLayerBase<PtrType>::RemoveSubLayerPath( int32 Index )
+	void FSdfLayerBase<PtrType>::SetFramesPerSecond( double FramesPerSecond )
 	{
 #if USE_USD_SDK
-		if ( const PtrType& Ptr = Impl->GetInner() )
+		if ( PtrType& Ptr = Impl->GetInner() )
 		{
-			Ptr->RemoveSubLayerPath( Index );
+			return Ptr->SetFramesPerSecond( FramesPerSecond );
 		}
 #endif // #if USE_USD_SDK
 	}
 
 	template<typename PtrType>
-	void FSdfLayerBase<PtrType>::SetFramesPerSecond( double FramesPerSecond )
+	TArray<FString> FSdfLayerBase<PtrType>::GetSubLayerPaths() const
 	{
+		TArray<FString> SubLayerPaths;
+
 #if USE_USD_SDK
-		if ( PtrType& Ptr = Impl->GetInner() )
+		if (const PtrType& Ptr = Impl->GetInner())
 		{
-			return Ptr->SetFramesPerSecond( FramesPerSecond );
+			FScopedUsdAllocs UsdAllocs;
+
+			for (const std::string& SubLayerPath : Ptr->GetSubLayerPaths())
+			{
+				const auto ConvertedPath = StringCast<TCHAR>(reinterpret_cast<const UTF8CHAR*>(SubLayerPath.c_str()));
+				SubLayerPaths.Emplace(ConvertedPath.Get());
+			}
 		}
 #endif // #if USE_USD_SDK
+
+		return SubLayerPaths;
 	}
 
 	template<typename PtrType>
-	TArray< FString > FSdfLayerBase<PtrType>::GetSubLayerPaths() const
+	void FSdfLayerBase<PtrType>::SetSubLayerPaths(const TArray<FString>& NewPaths)
 	{
-		TArray< FString > SubLayerPaths;
-
 #if USE_USD_SDK
-		if ( const PtrType& Ptr = Impl->GetInner() )
+		if (const PtrType& Ptr = Impl->GetInner())
 		{
 			FScopedUsdAllocs UsdAllocs;
 
-			for ( const std::string& SubLayerPath : Ptr->GetSubLayerPaths() )
+			std::vector<std::string> SubLayerPaths;
+			SubLayerPaths.reserve(NewPaths.Num());
+
+			for (const FString& NewPath : NewPaths)
 			{
-				SubLayerPaths.Emplace( ANSI_TO_TCHAR( SubLayerPath.c_str() ) );
+				const auto ConvertedPath = StringCast<UTF8CHAR>(*NewPath);
+				std::string SubLayerPath = reinterpret_cast<const char*>(ConvertedPath.Get());
+				SubLayerPaths.push_back(MoveTemp(SubLayerPath));
 			}
+
+			Ptr->SetSubLayerPaths(SubLayerPaths);
 		}
 #endif // #if USE_USD_SDK
+	}
 
-		return SubLayerPaths;
+	template<typename PtrType>
+	int64 FSdfLayerBase<PtrType>::GetNumSubLayerPaths() const
+	{
+#if USE_USD_SDK
+		if (const PtrType& Ptr = Impl->GetInner())
+		{
+			return (int64)Ptr->GetNumSubLayerPaths();
+		}
+#endif // #if USE_USD_SDK
+
+		return 0;
 	}
 
 	template<typename PtrType>
-	TArray< FSdfLayerOffset > FSdfLayerBase<PtrType>::GetSubLayerOffsets() const
+	void FSdfLayerBase<PtrType>::InsertSubLayerPath(const FString& Path, int32 Index)
 	{
-		TArray< FSdfLayerOffset > SubLayerOffsets;
+#if USE_USD_SDK
+		if (const PtrType& Ptr = Impl->GetInner())
+		{
+			FScopedUsdAllocs UsdAllocs;
 
+			const auto ConvertedPath = StringCast<UTF8CHAR>(*Path);
+			std::string SubLayerPath = reinterpret_cast<const char*>(ConvertedPath.Get());
+			Ptr->InsertSubLayerPath(MoveTemp(SubLayerPath), Index);
+		}
+#endif // #if USE_USD_SDK
+	}
+
+	template<typename PtrType>
+	void FSdfLayerBase<PtrType>::RemoveSubLayerPath(int32 Index)
+	{
 #if USE_USD_SDK
-		if ( const PtrType& Ptr = Impl->GetInner() )
+		if (const PtrType& Ptr = Impl->GetInner())
+		{
+			Ptr->RemoveSubLayerPath(Index);
+		}
+#endif // #if USE_USD_SDK
+	}
+
+	template<typename PtrType>
+	TArray<FSdfLayerOffset> FSdfLayerBase<PtrType>::GetSubLayerOffsets() const
+	{
+		TArray<FSdfLayerOffset> SubLayerOffsets;
+
+#if USE_USD_SDK
+		if (const PtrType& Ptr = Impl->GetInner())
 		{
 			FScopedUsdAllocs UsdAllocs;
 
-			for ( const pxr::SdfLayerOffset& SubLayerOffset : Ptr->GetSubLayerOffsets() )
+			for (const pxr::SdfLayerOffset& SubLayerOffset : Ptr->GetSubLayerOffsets())
 			{
-				if ( SubLayerOffset.IsValid() )
+				if (SubLayerOffset.IsValid())
 				{
-					SubLayerOffsets.Emplace( SubLayerOffset.GetOffset(), SubLayerOffset.GetScale() );
+					SubLayerOffsets.Emplace(SubLayerOffset.GetOffset(), SubLayerOffset.GetScale());
 				}
 				else
 				{
@@ -697,15 +736,34 @@ namespace UE
 	}
 
 	template<typename PtrType>
-	void FSdfLayerBase<PtrType>::SetSubLayerOffset( const FSdfLayerOffset& LayerOffset, int32 Index )
+	FSdfLayerOffset FSdfLayerBase<PtrType>::GetSubLayerOffset(int32 Index) const
 	{
 #if USE_USD_SDK
-		if ( PtrType& Ptr = Impl->GetInner() )
+		if (const PtrType& Ptr = Impl->GetInner())
+		{
+			FScopedUsdAllocs UsdAllocs;
+
+			const pxr::SdfLayerOffset& SubLayerOffset = Ptr->GetSubLayerOffset(Index);
+			if (SubLayerOffset.IsValid())
+			{
+				return FSdfLayerOffset(SubLayerOffset.GetOffset(), SubLayerOffset.GetScale());
+			}
+		}
+#endif // #if USE_USD_SDK
+
+		return FSdfLayerOffset{};
+	}
+
+	template<typename PtrType>
+	void FSdfLayerBase<PtrType>::SetSubLayerOffset(const FSdfLayerOffset& Offset, int32 Index)
+	{
+#if USE_USD_SDK
+		if (PtrType& Ptr = Impl->GetInner())
 		{
 			FScopedUsdAllocs UsdAllocs;
 
-			pxr::SdfLayerOffset UsdLayerOffset( LayerOffset.Offset, LayerOffset.Scale );
-			Ptr->SetSubLayerOffset( MoveTemp( UsdLayerOffset ), Index );
+			pxr::SdfLayerOffset UsdLayerOffset(Offset.Offset, Offset.Scale);
+			Ptr->SetSubLayerOffset(MoveTemp(UsdLayerOffset), Index);
 		}
 #endif // #if USE_USD_SDK
 	}
diff --git a/Engine/Plugins/Importers/USDImporter/Source/UnrealUSDWrapper/Public/UsdWrappers/SdfLayer.h b/Engine/Plugins/Importers/USDImporter/Source/UnrealUSDWrapper/Public/UsdWrappers/SdfLayer.h
--- a/Engine/Plugins/Importers/USDImporter/Source/UnrealUSDWrapper/Public/UsdWrappers/SdfLayer.h
+++ b/Engine/Plugins/Importers/USDImporter/Source/UnrealUSDWrapper/Public/UsdWrappers/SdfLayer.h
@@ -182,14 +182,16 @@ namespace UE
 
 		bool HasFramesPerSecond() const;
 		double GetFramesPerSecond() const;
-		void RemoveSubLayerPath( int32 Index );
 		void SetFramesPerSecond( double FramesPerSecond );
 
+		TArray<FString> GetSubLayerPaths() const;
+		void SetSubLayerPaths(const TArray<FString>& NewPaths);
 		int64 GetNumSubLayerPaths() const;
-		TArray< FString > GetSubLayerPaths() const;
-		TArray< FSdfLayerOffset > GetSubLayerOffsets() const;
-
-		void SetSubLayerOffset( const FSdfLayerOffset& LayerOffset, int32 Index );
+		void InsertSubLayerPath(const FString& Path, int32 Index = -1);
+		void RemoveSubLayerPath(int32 Index);
+		TArray<FSdfLayerOffset> GetSubLayerOffsets() const;
+		FSdfLayerOffset GetSubLayerOffset(int32 Index) const;
+		void SetSubLayerOffset(const FSdfLayerOffset& Offset, int32 Index);
 
 		bool HasSpec( const FSdfPath& Path ) const;
 
