USD: Fix crash when trying to check destroyed prim twins after GC runs.

diff --git a/Engine/Plugins/Importers/USDImporter/Source/USDStage/Private/USDStageActor.cpp b/Engine/Plugins/Importers/USDImporter/Source/USDStage/Private/USDStageActor.cpp
--- a/Engine/Plugins/Importers/USDImporter/Source/USDStage/Private/USDStageActor.cpp
+++ b/Engine/Plugins/Importers/USDImporter/Source/USDStage/Private/USDStageActor.cpp
@@ -175,7 +175,9 @@ struct FUsdStageActorImpl
 	static void DeselectActorsAndComponents(AUsdStageActor* StageActor)
 	{
 #if WITH_EDITOR
-		if (!StageActor)
+		// If we're being BeginDestroyed (by GC) then it's not really safe to even *check* our prim twins because
+		// they may have been fully destroyed before us, and could be just garbage memory at this point
+		if (!StageActor || StageActor->HasAnyFlags(RF_BeginDestroyed))
 		{
 			return;
 		}
@@ -184,7 +186,7 @@ struct FUsdStageActorImpl
 		// Don't do this during garbage collecting if we need to delay-create the root twin (can't NewObject during garbage collection).
 		// If we have no root twin we don't have any tracked spawned actors and components, so we don't need to deselect anything in the first place
 		bool bDeselected = false;
-		if (GEditor && !IsGarbageCollecting() && StageActor->RootUsdTwin)
+		if (GEditor && !IsGarbageCollecting() && StageActor->RootUsdTwin && !StageActor->RootUsdTwin->HasAnyFlags(RF_BeginDestroyed))
 		{
 			TArray<UObject*> ActorsToDeselect;
 			TArray<UObject*> ComponentsToDeselect;
