USD: Fix MaterialX reference composition arcs being ignored if they were ancestral.

diff --git a/Engine/Plugins/Importers/USDImporter/Source/USDSchemas/Private/Custom/MaterialXUSDShadeMaterialTranslator.cpp b/Engine/Plugins/Importers/USDImporter/Source/USDSchemas/Private/Custom/MaterialXUSDShadeMaterialTranslator.cpp
--- a/Engine/Plugins/Importers/USDImporter/Source/USDSchemas/Private/Custom/MaterialXUSDShadeMaterialTranslator.cpp
+++ b/Engine/Plugins/Importers/USDImporter/Source/USDSchemas/Private/Custom/MaterialXUSDShadeMaterialTranslator.cpp
@@ -56,7 +56,9 @@ namespace UE::USDMaterialXTranslator::Private
 	{
 		TArray<FString> Result;
 
-		pxr::UsdPrimCompositionQuery PrimCompositionQuery = pxr::UsdPrimCompositionQuery::GetDirectReferences(Prim);
+		// We used to just fetch "direct references" here, but stages may compose the .mtlx file reference onto the prim
+		// via another sublayer or reference, so it will be marked as an "ancestral arc" and not be included in the "direct references" filter
+		pxr::UsdPrimCompositionQuery PrimCompositionQuery = pxr::UsdPrimCompositionQuery{Prim};
 		for (const pxr::UsdPrimCompositionQueryArc& CompositionArc : PrimCompositionQuery.GetCompositionArcs())
 		{
 			if (CompositionArc.GetArcType() == pxr::PcpArcTypeReference)
