USD: New translator for USD mesh primitives (cube, cylinder, etc.).

diff --git a/Engine/Plugins/Importers/USDImporter/Source/USDSchemas/Private/USDGeomMeshTranslator.cpp b/Engine/Plugins/Importers/USDImporter/Source/USDSchemas/Private/USDGeomMeshTranslator.cpp
--- a/Engine/Plugins/Importers/USDImporter/Source/USDSchemas/Private/USDGeomMeshTranslator.cpp
+++ b/Engine/Plugins/Importers/USDImporter/Source/USDSchemas/Private/USDGeomMeshTranslator.cpp
@@ -316,7 +316,7 @@ namespace UsdGeomMeshTranslatorImpl
 			(const pxr::UsdGeomMesh& LODMesh, int32 LODIndex)
 			{
 				TArray<TUsdStore<pxr::UsdGeomPrimvar>> MeshPrimvars = UsdUtils::GetUVSetPrimvars(
-					LODMesh,
+					LODMesh.GetPrim(),
 					TNumericLimits<int32>::Max()
 				);
 
diff --git a/Engine/Plugins/Importers/USDImporter/Source/USDSchemas/Private/USDGeomPrimitiveTranslator.cpp b/Engine/Plugins/Importers/USDImporter/Source/USDSchemas/Private/USDGeomPrimitiveTranslator.cpp
new file mode 100644
--- /dev/null
+++ b/Engine/Plugins/Importers/USDImporter/Source/USDSchemas/Private/USDGeomPrimitiveTranslator.cpp
@@ -0,0 +1,207 @@
+// Copyright Epic Games, Inc. All Rights Reserved.
+
+#include "USDGeomPrimitiveTranslator.h"
+
+#if USE_USD_SDK
+
+#include "UnrealUSDWrapper.h"
+#include "USDAssetImportData.h"
+#include "USDConversionUtils.h"
+#include "USDGeomMeshConversion.h"
+#include "USDGeomMeshTranslator.h"
+#include "USDInfoCache.h"
+#include "USDLog.h"
+#include "USDTypesConversion.h"
+
+#include "UsdWrappers/SdfPath.h"
+#include "UsdWrappers/UsdPrim.h"
+
+#include "Components/StaticMeshComponent.h"
+#include "Engine/StaticMesh.h"
+#include "Materials/MaterialInterface.h"
+#include "MeshTranslationImpl.h"
+#include "StaticMeshAttributes.h"
+
+#include "USDIncludesStart.h"
+	#include "pxr/usd/usd/prim.h"
+#include "USDIncludesEnd.h"
+
+namespace UsdGeomPrimitiveTranslatorImpl
+{
+	void LoadMeshDescriptions(
+		UE::FUsdPrim Prim,
+		TArray<FMeshDescription>& OutLODIndexToMeshDescription,
+		TArray<UsdUtils::FUsdPrimMaterialAssignmentInfo>& OutLODIndexToMaterialInfo,
+		const UsdToUnreal::FUsdMeshConversionOptions& Options
+	)
+	{
+		if (!Prim)
+		{
+			return;
+		}
+
+		FMeshDescription TempMeshDescription;
+		UsdUtils::FUsdPrimMaterialAssignmentInfo TempMaterialInfo;
+
+		FStaticMeshAttributes StaticMeshAttributes(TempMeshDescription);
+		StaticMeshAttributes.Register();
+
+		const bool bSuccess = UsdToUnreal::ConvertGeomPrimitive(Prim, TempMeshDescription, TempMaterialInfo, Options);
+		if (bSuccess)
+		{
+			OutLODIndexToMeshDescription = {MoveTemp(TempMeshDescription)};
+			OutLODIndexToMaterialInfo = {MoveTemp(TempMaterialInfo)};
+		}
+	}
+
+	class FGeomPrimitiveCreateAssetsTaskChain : public FBuildStaticMeshTaskChain
+	{
+	public:
+		explicit FGeomPrimitiveCreateAssetsTaskChain(
+			const TSharedRef<FUsdSchemaTranslationContext>& InContext,
+			const UE::FSdfPath& InPrimPath,
+			const TOptional<UE::FSdfPath>& AlternativePrimToLinkAssetsTo = {}
+		)
+			: FBuildStaticMeshTaskChain(InContext, InPrimPath, AlternativePrimToLinkAssetsTo)
+		{
+			SetupTasks();
+		}
+
+	protected:
+		virtual void SetupTasks() override
+		{
+			FScopedUnrealAllocs UnrealAllocs;
+
+			// Create mesh descriptions (Async)
+			Do(ESchemaTranslationLaunchPolicy::Async,
+			   [this]() -> bool
+			   {
+				   pxr::TfToken RenderContextToken = pxr::UsdShadeTokens->universalRenderContext;
+				   if (!Context->RenderContext.IsNone())
+				   {
+					   RenderContextToken = UnrealToUsd::ConvertToken(*Context->RenderContext.ToString()).Get();
+				   }
+
+				   pxr::TfToken MaterialPurposeToken = pxr::UsdShadeTokens->allPurpose;
+				   if (!Context->MaterialPurpose.IsNone())
+				   {
+					   MaterialPurposeToken = UnrealToUsd::ConvertToken(*Context->MaterialPurpose.ToString()).Get();
+				   }
+
+				   UsdToUnreal::FUsdMeshConversionOptions Options;
+				   Options.TimeCode = Context->Time;
+				   Options.PurposesToLoad = Context->PurposesToLoad;
+				   Options.RenderContext = RenderContextToken;
+				   Options.MaterialPurpose = MaterialPurposeToken;
+				   Options.bMergeIdenticalMaterialSlots = Context->bMergeIdenticalMaterialSlots;
+
+				   UsdGeomPrimitiveTranslatorImpl::LoadMeshDescriptions(
+					   GetPrim(),
+					   LODIndexToMeshDescription,
+					   LODIndexToMaterialInfo,
+					   Options
+				   );
+
+				   // If we have at least one valid LOD, we should keep going
+				   for (const FMeshDescription& MeshDescription : LODIndexToMeshDescription)
+				   {
+					   if (!MeshDescription.IsEmpty())
+					   {
+						   return true;
+					   }
+				   }
+				   return false;
+			   });
+
+			FBuildStaticMeshTaskChain::SetupTasks();
+		}
+	};
+}	 // namespace UsdGeomPrimitiveTranslatorImpl
+
+void FUsdGeomPrimitiveTranslator::CreateAssets()
+{
+	TRACE_CPUPROFILER_EVENT_SCOPE(FUsdGeomPrimitiveTranslator::CreateAssets);
+
+	using namespace UsdGeomPrimitiveTranslatorImpl;
+	TSharedRef<FGeomPrimitiveCreateAssetsTaskChain> AssetsTaskChain = MakeShared<FGeomPrimitiveCreateAssetsTaskChain>(Context, PrimPath);
+
+	Context->TranslatorTasks.Add(MoveTemp(AssetsTaskChain));
+}
+
+USceneComponent* FUsdGeomPrimitiveTranslator::CreateComponents()
+{
+	USceneComponent* SceneComponent = CreateComponentsEx({}, {});
+	UpdateComponents(SceneComponent);
+
+	// Handle material overrides
+	if (UStaticMeshComponent* StaticMeshComponent = Cast<UStaticMeshComponent>(SceneComponent))
+	{
+		if (Context->InfoCache)
+		{
+			if (UStaticMesh* StaticMesh = Context->InfoCache->GetSingleAssetForPrim<UStaticMesh>(
+				PrimPath
+			))
+			{
+				TArray<UMaterialInterface*> ExistingAssignments;
+				for (FStaticMaterial& StaticMaterial : StaticMesh->GetStaticMaterials())
+				{
+					ExistingAssignments.Add(StaticMaterial.MaterialInterface);
+				}
+
+				MeshTranslationImpl::SetMaterialOverrides(
+					GetPrim(),
+					ExistingAssignments,
+					*StaticMeshComponent,
+					*Context->AssetCache.Get(),
+					*Context->InfoCache.Get(),
+					Context->Time,
+					Context->ObjectFlags,
+					Context->bAllowInterpretingLODs,
+					Context->RenderContext,
+					Context->MaterialPurpose
+				);
+			}
+		}
+	}
+
+	return SceneComponent;
+}
+
+void FUsdGeomPrimitiveTranslator::UpdateComponents(USceneComponent* SceneComponent)
+{
+	if (UsdUtils::IsAnimated(GetPrim()))
+	{
+		// The assets might have changed since our attributes are animated
+		// Note that we must wait for these to complete as they make take a while and we want to
+		// reassign our new static meshes when we get to FUsdGeomXformableTranslator::UpdateComponents
+		CreateAssets();
+		Context->CompleteTasks();
+	}
+
+	Super::UpdateComponents(SceneComponent);
+}
+
+bool FUsdGeomPrimitiveTranslator::CollapsesChildren(ECollapsingType CollapsingType) const
+{
+	// Gprims shouldn't really have any children so this is not very well defined.
+	// We're going with 'false' here to match FUsdGeomMeshTranslator::CollapsesChildren
+	return false;
+}
+
+bool FUsdGeomPrimitiveTranslator::CanBeCollapsed(ECollapsingType CollapsingType) const
+{
+	return Super::CanBeCollapsed(CollapsingType);
+}
+
+TSet<UE::FSdfPath> FUsdGeomPrimitiveTranslator::CollectAuxiliaryPrims() const
+{
+	if (!Context->bIsBuildingInfoCache)
+	{
+		return Context->InfoCache->GetAuxiliaryPrims(PrimPath);
+	}
+
+	// Let's assume we can't specify UsdGeomSubsets on primitives for now
+	return {};
+}
+
+#endif	  // #if USE_USD_SDK
diff --git a/Engine/Plugins/Importers/USDImporter/Source/USDSchemas/Private/USDGeomXformableTranslator.cpp b/Engine/Plugins/Importers/USDImporter/Source/USDSchemas/Private/USDGeomXformableTranslator.cpp
--- a/Engine/Plugins/Importers/USDImporter/Source/USDSchemas/Private/USDGeomXformableTranslator.cpp
+++ b/Engine/Plugins/Importers/USDImporter/Source/USDSchemas/Private/USDGeomXformableTranslator.cpp
@@ -489,8 +489,8 @@ USceneComponent* FUsdGeomXformableTranslator::CreateComponentsEx( TOptional< TSu
 					// We can't just check if there's a static mesh for this prim on the cache, because the prims with meshes could be potentially invisible (and so
 					// we don't have parsed their meshes yet), so here we traverse our child hierarchy and if we have any chance of ever generating a Mesh, we go
 					// for a static mesh component
-					TArray< UE::FUsdPrim > ChildMeshPrims = UsdUtils::GetAllPrimsOfType( Prim, TEXT( "UsdGeomMesh" ) );
-					if ( ChildMeshPrims.Num() > 0 )
+					TArray< UE::FUsdPrim > ChildGprims = UsdUtils::GetAllPrimsOfType( Prim, TEXT( "UsdGeomGprim" ) );
+					if ( ChildGprims.Num() > 0 )
 					{
 						ComponentType = UStaticMeshComponent::StaticClass();
 					}
diff --git a/Engine/Plugins/Importers/USDImporter/Source/USDSchemas/Private/USDInfoCache.cpp b/Engine/Plugins/Importers/USDImporter/Source/USDSchemas/Private/USDInfoCache.cpp
--- a/Engine/Plugins/Importers/USDImporter/Source/USDSchemas/Private/USDInfoCache.cpp
+++ b/Engine/Plugins/Importers/USDImporter/Source/USDSchemas/Private/USDInfoCache.cpp
@@ -385,17 +385,9 @@ namespace UE::USDInfoCacheImpl::Private
 	{
 		TRACE_CPUPROFILER_EVENT_SCOPE( UE::USDInfoCacheImpl::Private::GetPrimVertexCountAndSlots );
 
-		if ( UsdPrim.IsA<pxr::UsdGeomMesh>() || UsdPrim.IsA<pxr::UsdGeomSubset>() )
+		if (UsdPrim.IsA<pxr::UsdGeomGprim>() || UsdPrim.IsA<pxr::UsdGeomSubset>())
 		{
-			if ( pxr::UsdGeomMesh Mesh{ UsdPrim } )
-			{
-				if ( pxr::UsdAttribute Points = Mesh.GetPointsAttr() )
-				{
-					pxr::VtArray< pxr::GfVec3f > PointsArray;
-					Points.Get( &PointsArray, pxr::UsdTimeCode( Context.Time ) );
-					OutVertexCount = PointsArray.size();
-				}
-			}
+			OutVertexCount = UsdUtils::GetGprimVertexCount(pxr::UsdGeomGprim{UsdPrim}, Context.Time);
 
 			pxr::TfToken RenderContextToken = pxr::UsdShadeTokens->universalRenderContext;
 			if ( !Context.RenderContext.IsNone() )
diff --git a/Engine/Plugins/Importers/USDImporter/Source/USDSchemas/Private/USDSchemasModule.cpp b/Engine/Plugins/Importers/USDImporter/Source/USDSchemas/Private/USDSchemasModule.cpp
--- a/Engine/Plugins/Importers/USDImporter/Source/USDSchemas/Private/USDSchemasModule.cpp
+++ b/Engine/Plugins/Importers/USDImporter/Source/USDSchemas/Private/USDSchemasModule.cpp
@@ -18,6 +18,7 @@
 #include "USDGeomPointInstancerTranslator.h"
 #include "USDGeomXformableTranslator.h"
 #include "USDGroomTranslator.h"
+#include "USDGeomPrimitiveTranslator.h"
 #include "USDLuxLightTranslator.h"
 #include "USDMemory.h"
 #include "USDShadeMaterialTranslator.h"
@@ -34,32 +35,45 @@ public:
 #if USE_USD_SDK
 		LLM_SCOPE_BYTAG(Usd);
 
+		FUsdSchemaTranslatorRegistry& Registry = GetTranslatorRegistry();
+		FUsdRenderContextRegistry& ShaderRegistry = GetRenderContextRegistry();
+
 		// Register the default translators
-		UsdGeomCameraTranslatorHandle = GetTranslatorRegistry().Register< FUsdGeomCameraTranslator >( TEXT("UsdGeomCamera") );
-		UsdGeomMeshTranslatorHandle = GetTranslatorRegistry().Register< FUsdGeomMeshTranslator >( TEXT("UsdGeomMesh") );
-		UsdGeomPointInstancerTranslatorHandle = GetTranslatorRegistry().Register< FUsdGeomPointInstancerTranslator >( TEXT("UsdGeomPointInstancer") );
-		UsdGeomXformableTranslatorHandle = GetTranslatorRegistry().Register< FUsdGeomXformableTranslator >( TEXT("UsdGeomXformable") );
-		UsdShadeMaterialTranslatorHandle = GetTranslatorRegistry().Register< FUsdShadeMaterialTranslator >( TEXT("UsdShadeMaterial") );
-		UsdLuxBoundableLightBaseTranslatorHandle = GetTranslatorRegistry().Register< FUsdLuxLightTranslator >( TEXT("UsdLuxBoundableLightBase") );
-		UsdLuxNonboundableLightBaseTranslatorHandle = GetTranslatorRegistry().Register< FUsdLuxLightTranslator >( TEXT("UsdLuxNonboundableLightBase") );
+		TranslatorHandles = {
+			Registry.Register<FUsdGeomCameraTranslator>(TEXT("UsdGeomCamera")),
+			Registry.Register<FUsdGeomMeshTranslator>(TEXT("UsdGeomMesh")),
+			Registry.Register<FUsdGeomPrimitiveTranslator>(TEXT("UsdGeomCapsule")),
+			Registry.Register<FUsdGeomPrimitiveTranslator>(TEXT("UsdGeomCone")),
+			Registry.Register<FUsdGeomPrimitiveTranslator>(TEXT("UsdGeomCube")),
+			Registry.Register<FUsdGeomPrimitiveTranslator>(TEXT("UsdGeomCylinder")),
+			Registry.Register<FUsdGeomPrimitiveTranslator>(TEXT("UsdGeomPlane")),
+			Registry.Register<FUsdGeomPrimitiveTranslator>(TEXT("UsdGeomSphere")),
+			Registry.Register<FUsdGeomPointInstancerTranslator>(TEXT("UsdGeomPointInstancer")),
+			Registry.Register<FUsdGeomXformableTranslator>(TEXT("UsdGeomXformable")),
+			Registry.Register<FUsdShadeMaterialTranslator>(TEXT("UsdShadeMaterial")),
+			Registry.Register<FUsdLuxLightTranslator>(TEXT("UsdLuxBoundableLightBase")),
+			Registry.Register<FUsdLuxLightTranslator>(TEXT("UsdLuxNonboundableLightBase"))};
 
 #if WITH_EDITOR
-		GetRenderContextRegistry().Register(FMaterialXUsdShadeMaterialTranslator::MaterialXRenderContext);
-		MaterialXUsdShadeMaterialTranslatorHandle = GetTranslatorRegistry().Register<FMaterialXUsdShadeMaterialTranslator>(TEXT("UsdShadeMaterial"));
+		ShaderRegistry.Register(FMaterialXUsdShadeMaterialTranslator::MaterialXRenderContext);
+		TranslatorHandles.Add(Registry.Register<FMaterialXUsdShadeMaterialTranslator>(TEXT("UsdShadeMaterial")));
 
 		// Creating skeletal meshes technically works in Standalone mode, but by checking for this we artificially block it
-		// to not confuse users as to why it doesn't work at runtime. Not registering the actual translators lets the inner meshes get parsed as static meshes, at least.
-		if ( GIsEditor )
+		// to not confuse users as to why it doesn't work at runtime. Not registering the actual translators lets the inner meshes get parsed as
+		// static meshes, at least.
+		if (GIsEditor)
 		{
-			UsdSkelRootTranslatorHandle = GetTranslatorRegistry().Register< FUsdSkelRootTranslator >( TEXT("UsdSkelRoot") );
-			UsdGroomTranslatorHandle = GetTranslatorRegistry().Register< FUsdGroomTranslator >( TEXT("UsdGeomXformable") );
-			// The GeometryCacheTranslator also works on UsdGeomXformable through the GroomTranslator
-			UsdGeometryCacheTranslatorHandle = GetTranslatorRegistry().Register< FUsdGeometryCacheTranslator >( TEXT("UsdGeomMesh") );
-
-			if ( IMDLImporterModule* MDLImporterModule = FModuleManager::Get().LoadModulePtr< IMDLImporterModule >( TEXT("MDLImporter") ) )
+			TranslatorHandles.Append({
+				Registry.Register<FUsdSkelRootTranslator>(TEXT("UsdSkelRoot")),
+				Registry.Register<FUsdGroomTranslator>(TEXT("UsdGeomXformable")),
+				// The GeometryCacheTranslator also works on UsdGeomXformable through the GroomTranslator
+				Registry.Register<FUsdGeometryCacheTranslator>(TEXT("UsdGeomMesh")),
+			});
+
+			if (IMDLImporterModule* MDLImporterModule = FModuleManager::Get().LoadModulePtr<IMDLImporterModule>(TEXT("MDLImporter")))
 			{
-				GetRenderContextRegistry().Register(FMdlUsdShadeMaterialTranslator::MdlRenderContext);
-				MdlUsdShadeMaterialTranslatorHandle = GetTranslatorRegistry().Register< FMdlUsdShadeMaterialTranslator >( TEXT("UsdShadeMaterial") );
+				ShaderRegistry.Register(FMdlUsdShadeMaterialTranslator::MdlRenderContext);
+				TranslatorHandles.Add(Registry.Register<FMdlUsdShadeMaterialTranslator>(TEXT("UsdShadeMaterial")));
 			}
 		}
 #endif // WITH_EDITOR
@@ -69,31 +83,18 @@ public:
 
 	virtual void ShutdownModule() override
 	{
-#if USE_USD_SDK
-		GetTranslatorRegistry().Unregister( UsdGeomCameraTranslatorHandle );
-		GetTranslatorRegistry().Unregister( UsdGeomMeshTranslatorHandle );
-		GetTranslatorRegistry().Unregister( UsdGeomPointInstancerTranslatorHandle );
+		FUsdSchemaTranslatorRegistry& Registry = GetTranslatorRegistry();
+		FUsdRenderContextRegistry& ShaderRegistry = GetRenderContextRegistry();
 
-#if WITH_EDITOR
-		if ( GIsEditor )
+		for (const FRegisteredSchemaTranslatorHandle& TranslatorHandle : TranslatorHandles)
 		{
-			GetTranslatorRegistry().Unregister( UsdSkelRootTranslatorHandle );
-			GetTranslatorRegistry().Unregister( UsdGroomTranslatorHandle );
-			GetTranslatorRegistry().Unregister( UsdGeometryCacheTranslatorHandle );
-
-			GetTranslatorRegistry().Unregister( MdlUsdShadeMaterialTranslatorHandle );
-			GetRenderContextRegistry().Unregister(FMdlUsdShadeMaterialTranslator::MdlRenderContext);
+			Registry.Unregister(TranslatorHandle);
 		}
 
-		GetTranslatorRegistry().Unregister(MaterialXUsdShadeMaterialTranslatorHandle);
-		GetRenderContextRegistry().Unregister(FMaterialXUsdShadeMaterialTranslator::MaterialXRenderContext);
+#if USE_USD_SDK && WITH_EDITOR
+		ShaderRegistry.Unregister(FMdlUsdShadeMaterialTranslator::MdlRenderContext);
+		ShaderRegistry.Unregister(FMaterialXUsdShadeMaterialTranslator::MaterialXRenderContext);
 #endif // WITH_EDITOR
-
-		GetTranslatorRegistry().Unregister( UsdGeomXformableTranslatorHandle );
-		GetTranslatorRegistry().Unregister( UsdShadeMaterialTranslatorHandle );
-		GetTranslatorRegistry().Unregister( UsdLuxBoundableLightBaseTranslatorHandle );
-		GetTranslatorRegistry().Unregister( UsdLuxNonboundableLightBaseTranslatorHandle );
-#endif // #if USE_USD_SDK
 	}
 
 	virtual FUsdSchemaTranslatorRegistry& GetTranslatorRegistry() override
@@ -110,20 +111,7 @@ protected:
 	FUsdSchemaTranslatorRegistry UsdSchemaTranslatorRegistry;
 	FUsdRenderContextRegistry UsdRenderContextRegistry;
 
-	FRegisteredSchemaTranslatorHandle UsdGeomCameraTranslatorHandle;
-	FRegisteredSchemaTranslatorHandle UsdGeometryCacheTranslatorHandle;
-	FRegisteredSchemaTranslatorHandle UsdGeomMeshTranslatorHandle;
-	FRegisteredSchemaTranslatorHandle UsdGeomPointInstancerTranslatorHandle;
-	FRegisteredSchemaTranslatorHandle UsdSkelRootTranslatorHandle;
-	FRegisteredSchemaTranslatorHandle UsdGeomXformableTranslatorHandle;
-	FRegisteredSchemaTranslatorHandle UsdShadeMaterialTranslatorHandle;
-	FRegisteredSchemaTranslatorHandle UsdGroomTranslatorHandle;
-	FRegisteredSchemaTranslatorHandle UsdLuxBoundableLightBaseTranslatorHandle;
-	FRegisteredSchemaTranslatorHandle UsdLuxNonboundableLightBaseTranslatorHandle;
-
-	// Custom schemas
-	FRegisteredSchemaTranslatorHandle MdlUsdShadeMaterialTranslatorHandle;
-	FRegisteredSchemaTranslatorHandle MaterialXUsdShadeMaterialTranslatorHandle;
+	TArray< FRegisteredSchemaTranslatorHandle > TranslatorHandles;
 };
 
 IMPLEMENT_MODULE_USD( FUsdSchemasModule, USDSchemas );
diff --git a/Engine/Plugins/Importers/USDImporter/Source/USDSchemas/Public/USDGeomPrimitiveTranslator.h b/Engine/Plugins/Importers/USDImporter/Source/USDSchemas/Public/USDGeomPrimitiveTranslator.h
new file mode 100644
--- /dev/null
+++ b/Engine/Plugins/Importers/USDImporter/Source/USDSchemas/Public/USDGeomPrimitiveTranslator.h
@@ -0,0 +1,30 @@
+// Copyright Epic Games, Inc. All Rights Reserved.
+
+#pragma once
+
+#include "USDGeomXformableTranslator.h"
+
+#include "UsdWrappers/SdfPath.h"
+
+#if USE_USD_SDK
+
+class USDSCHEMAS_API FUsdGeomPrimitiveTranslator : public FUsdGeomXformableTranslator
+{
+public:
+	using Super = FUsdGeomXformableTranslator;
+	using FUsdGeomXformableTranslator::FUsdGeomXformableTranslator;
+
+	FUsdGeomPrimitiveTranslator(const FUsdGeomPrimitiveTranslator& Other) = delete;
+	FUsdGeomPrimitiveTranslator& operator=(const FUsdGeomPrimitiveTranslator& Other) = delete;
+
+	virtual void CreateAssets() override;
+	virtual USceneComponent* CreateComponents() override;
+	virtual void UpdateComponents(USceneComponent* SceneComponent) override;
+
+	virtual bool CollapsesChildren(ECollapsingType CollapsingType) const override;
+	virtual bool CanBeCollapsed(ECollapsingType CollapsingType) const override;
+
+	virtual TSet<UE::FSdfPath> CollectAuxiliaryPrims() const override;
+};
+
+#endif	  // #if USE_USD_SDK
diff --git a/Engine/Plugins/Importers/USDImporter/Source/USDSchemas/Public/USDGeomXformableTranslator.h b/Engine/Plugins/Importers/USDImporter/Source/USDSchemas/Public/USDGeomXformableTranslator.h
--- a/Engine/Plugins/Importers/USDImporter/Source/USDSchemas/Public/USDGeomXformableTranslator.h
+++ b/Engine/Plugins/Importers/USDImporter/Source/USDSchemas/Public/USDGeomXformableTranslator.h
@@ -27,7 +27,7 @@ public:
 
 	virtual void CreateAssets() override;
 	virtual USceneComponent* CreateComponents() override;
-	virtual void UpdateComponents( USceneComponent* SceneComponent );
+	virtual void UpdateComponents( USceneComponent* SceneComponent ) override;
 
 	virtual bool CollapsesChildren( ECollapsingType CollapsingType ) const override;
 	virtual bool CanBeCollapsed( ECollapsingType CollapsingType ) const override;
diff --git a/Engine/Plugins/Importers/USDImporter/Source/USDStage/Private/USDLevelSequenceHelper.cpp b/Engine/Plugins/Importers/USDImporter/Source/USDStage/Private/USDLevelSequenceHelper.cpp
--- a/Engine/Plugins/Importers/USDImporter/Source/USDStage/Private/USDLevelSequenceHelper.cpp
+++ b/Engine/Plugins/Importers/USDImporter/Source/USDStage/Private/USDLevelSequenceHelper.cpp
@@ -1485,46 +1485,88 @@ void FUsdLevelSequenceHelperImpl::AddCommonTracks(const UUsdPrimTwin& PrimTwin,
 		}
 	}
 
-	TArray<double> TransformTimeSamples;
-	if ((Xformable.GetTimeSamples(&TransformTimeSamples) && TransformTimeSamples.Num() > 0) || bNeedTrackToCompensateResetXformOp)
+	// Check if we need to add Transform tracks.
+	// In case we're e.g. a Cube with animated "size", which needs to become animated transforms
+	if (Xformable.TransformMightBeTimeVarying() || bNeedTrackToCompensateResetXformOp || Prim.IsA(TEXT("Gprim")))
 	{
+		TArray<double> TimeSampleUnion;
+
+		// Get all *animated* attributes that may contribute to the transform
+		bool bAreAllMuted = true;
 		TArray<UE::FUsdAttribute> Attrs = UnrealToUsd::GetAttributesForProperty(Prim, UnrealIdentifiers::TransformPropertyName);
-		if (Attrs.Num() > 0)
+		for (int32 Index = Attrs.Num() - 1; Index >= 0; --Index)
 		{
-			if (UE::FUsdAttribute& TransformAttribute = Attrs[0])
+			const UE::FUsdAttribute& Attr = Attrs[Index];
+
+			TArray<double> TimeSamplesForAttr;
+			if (!Attr.GetTimeSamples(TimeSamplesForAttr) || TimeSamplesForAttr.Num() == 0)
 			{
-				if (ULevelSequence* AttributeSequence = FindOrAddSequenceForAttribute(TransformAttribute))
-				{
-					const bool bIsMuted = UsdUtils::IsAttributeMuted(TransformAttribute, UsdStage);
+				const int32 Count = 1;
+				const bool bAllowShrinking = false;
+				Attrs.RemoveAt(Index, Count, bAllowShrinking);
+				continue;
+			}
 
-					FMovieSceneSequenceTransform SequenceTransform;
-					FMovieSceneSequenceID SequenceID = SequencesID.FindRef(AttributeSequence);
-					if (FMovieSceneSubSequenceData* SubSequenceData = SequenceHierarchyCache.FindSubData(SequenceID))
-					{
-						SequenceTransform = SubSequenceData->RootToSequenceTransform;
-					}
+			if (!UsdUtils::IsAttributeMuted(Attr, UsdStage))
+			{
+				bAreAllMuted = false;
 
-					if (UMovieScene* MovieScene = AttributeSequence->GetMovieScene())
-					{
-						TArray<double> TimeSamples;
-						if (Xformable.GetTimeSamples(&TimeSamples))
-						{
-							if (bNeedTrackToCompensateResetXformOp)
-							{
-								TimeSamples.Append(AncestorTimeSamples);
-								TimeSamples.Sort();
-							}
+				// Union the time samples so we know to always sample where we have a value for a relevant attribute
+				TimeSampleUnion.Append(TimeSamplesForAttr);
+			}
+		}
 
-							if (UMovieScene3DTransformTrack* TransformTrack = AddTrack<UMovieScene3DTransformTrack>(UnrealIdentifiers::TransformPropertyName, PrimTwin, *ComponentToBind, *AttributeSequence, bIsMuted))
-							{
-								UsdToUnreal::FPropertyTrackReader Reader = UsdToUnreal::CreatePropertyTrackReader(Prim, UnrealIdentifiers::TransformPropertyName, bIgnorePrimLocalTransform);
-								UsdToUnreal::ConvertTransformTimeSamples(UsdStage, TimeSamples, Reader.TransformReader, *TransformTrack, SequenceTransform);
-							}
+		// Find the strongest layer where any of these is authored. The TimeCode here is only for handling Value Clips, which
+		// we largely don't support anyway
+		const double TimeCode = 0.0;
+		const bool bIncludeSessionLayers = false;
+		UE::FSdfLayer Layer = UsdUtils::FindLayerForAttributes(Attrs, TimeCode, bIncludeSessionLayers);
 
-							PrimPathByLevelSequenceName.AddUnique(AttributeSequence->GetFName(), Prim.GetPrimPath().GetString());
-						}
-					}
+		// If we're creating a brand new transform track to compensate resetXformOp we may not have any animated attribute
+		// already, but we still need to do this
+		if (!Layer && bNeedTrackToCompensateResetXformOp)
+		{
+			Layer = UsdUtils::FindLayerForPrim(Prim);
+		}
+
+		// Get the Subsequence where we should create our track according to that Layer
+		if (ULevelSequence* AttributeSequence = FindOrAddSequenceForLayer(Layer, Layer.GetIdentifier(), Layer.GetDisplayName()))
+		{
+			FMovieSceneSequenceTransform SequenceTransform;
+			FMovieSceneSequenceID SequenceID = SequencesID.FindRef(AttributeSequence);
+			if (FMovieSceneSubSequenceData* SubSequenceData = SequenceHierarchyCache.FindSubData(SequenceID))
+			{
+				SequenceTransform = SubSequenceData->RootToSequenceTransform;
+			}
+
+			if (UMovieScene* MovieScene = AttributeSequence->GetMovieScene())
+			{
+				if (bNeedTrackToCompensateResetXformOp)
+				{
+					TimeSampleUnion.Append(AncestorTimeSamples);
 				}
+
+				// Note that since we sort, we can cheaply handle duplicate timeSamples on this array because
+				// UsdToUnreal::ConvertTransformTimeSamples ignores consecutive duplicates anyway (using FMath::IsNearlyEqual too)
+				TimeSampleUnion.Sort();
+
+				if (UMovieScene3DTransformTrack* TransformTrack = AddTrack<UMovieScene3DTransformTrack>(
+						UnrealIdentifiers::TransformPropertyName,
+						PrimTwin,
+						*ComponentToBind,
+						*AttributeSequence,
+						bAreAllMuted
+					))
+				{
+					UsdToUnreal::FPropertyTrackReader Reader = UsdToUnreal::CreatePropertyTrackReader(
+						Prim,
+						UnrealIdentifiers::TransformPropertyName,
+						bIgnorePrimLocalTransform
+					);
+					UsdToUnreal::ConvertTransformTimeSamples(UsdStage, TimeSampleUnion, Reader.TransformReader, *TransformTrack, SequenceTransform);
+				}
+
+				PrimPathByLevelSequenceName.AddUnique(AttributeSequence->GetFName(), Prim.GetPrimPath().GetString());
 			}
 		}
 	}
@@ -3486,6 +3528,13 @@ void FUsdLevelSequenceHelperImpl::HandleTrackChange(const UMovieSceneTrack& Trac
 				else if (const UMovieScene3DTransformTrack* TransformTrack = Cast<const UMovieScene3DTransformTrack>(&Track))
 				{
 					UnrealToUsd::Convert3DTransformTrack(*TransformTrack, SequenceTransform, Writer.TransformWriter, UsdPrim);
+
+					// If we're a Cylinder, Cube, etc. clear the animation of the primitive attributes that can affect the
+					// primitive transform ("height", "radius", etc.) as we'll be writing the full combined primitive+local
+					// transform directly to the Xform animation instead
+					const bool bDefaultValues = false;
+					const bool bTimeSampleValues = true;
+					UsdUtils::AuthorIdentityTransformGprimAttributes(UsdPrim, bDefaultValues, bTimeSampleValues);
 				}
 
 				// Refresh tracks that needed to be updated in USD (e.g. we wrote out a new keyframe to a RectLight's width -> that
diff --git a/Engine/Plugins/Importers/USDImporter/Source/USDStage/Private/USDStageActor.cpp b/Engine/Plugins/Importers/USDImporter/Source/USDStage/Private/USDStageActor.cpp
--- a/Engine/Plugins/Importers/USDImporter/Source/USDStage/Private/USDStageActor.cpp
+++ b/Engine/Plugins/Importers/USDImporter/Source/USDStage/Private/USDStageActor.cpp
@@ -3995,6 +3995,16 @@ void AUsdStageActor::OnObjectPropertyChanged(UObject* ObjectBeingModified, FProp
 
 				UnrealToUsd::ConvertSceneComponent(CurrentStage, PrimSceneComponent, UsdPrim);
 
+				// When we parse a Gprim like a Cube or a Cylinder, we'll always generate some "default" meshes (e.g. Cylinder with
+				// height always equal 1), and combine the Xform and the effect of the prim's attributes (e.g. height/width) into
+				// a SINGLE transform, and put that on the component (this approach allows attribute animation purely with Sequencer tracks).
+				// When we modify any property and want to write back out to USD however, we'll write that combined transform as the prim's
+				// transform. This means we must also "reset" the (e.g. height/width) attributes, so that the combined transform stays
+				// consistent
+				const bool bDefaultValues = true;
+				const bool bTimeSampleValues = false;
+				UsdUtils::AuthorIdentityTransformGprimAttributes(UsdPrim, bDefaultValues, bTimeSampleValues);
+
 				if (UMeshComponent* MeshComponent = Cast< UMeshComponent >(PrimSceneComponent))
 				{
 					UnrealToUsd::ConvertMeshComponent(CurrentStage, MeshComponent, UsdPrim);
diff --git a/Engine/Plugins/Importers/USDImporter/Source/USDUtilities/Private/USDAttributeUtils.cpp b/Engine/Plugins/Importers/USDImporter/Source/USDUtilities/Private/USDAttributeUtils.cpp
--- a/Engine/Plugins/Importers/USDImporter/Source/USDUtilities/Private/USDAttributeUtils.cpp
+++ b/Engine/Plugins/Importers/USDImporter/Source/USDUtilities/Private/USDAttributeUtils.cpp
@@ -175,6 +175,24 @@ bool UsdUtils::IsAttributeMuted( const UE::FUsdAttribute& Attribute, const UE::F
 }
 
 #if USE_USD_SDK
+bool UsdUtils::ClearAllTimeSamples(const pxr::UsdAttribute& Attribute)
+{
+	FScopedUsdAllocs Allocs;
+
+	std::vector<double> Times;
+	if (Attribute.GetTimeSamples(&Times))
+	{
+		for (double Time : Times)
+		{
+			Attribute.ClearAtTime(Time);
+		}
+
+		return true;
+	}
+
+	return false;
+}
+
 void UsdUtils::NotifyIfOverriddenOpinion(const pxr::UsdProperty& Property)
 {
 	FScopedUsdAllocs Allocs;
diff --git a/Engine/Plugins/Importers/USDImporter/Source/USDUtilities/Private/USDConversionUtils.cpp b/Engine/Plugins/Importers/USDImporter/Source/USDUtilities/Private/USDConversionUtils.cpp
--- a/Engine/Plugins/Importers/USDImporter/Source/USDUtilities/Private/USDConversionUtils.cpp
+++ b/Engine/Plugins/Importers/USDImporter/Source/USDUtilities/Private/USDConversionUtils.cpp
@@ -64,6 +64,7 @@
 	#include "pxr/usd/usd/stage.h"
 	#include "pxr/usd/usd/variantSets.h"
 	#include "pxr/usd/usdGeom/camera.h"
+	#include "pxr/usd/usdGeom/gprim.h"
 	#include "pxr/usd/usdGeom/imageable.h"
 	#include "pxr/usd/usdGeom/mesh.h"
 	#include "pxr/usd/usdGeom/metrics.h"
@@ -308,6 +309,8 @@ ValueType UsdUtils::GetUsdValue( const pxr::UsdAttribute& Attribute, pxr::UsdTim
 // Explicit template instantiation
 template USDUTILITIES_API bool							UsdUtils::GetUsdValue< bool >( const pxr::UsdAttribute& Attribute, pxr::UsdTimeCode TimeCode );
 template USDUTILITIES_API float							UsdUtils::GetUsdValue< float >( const pxr::UsdAttribute& Attribute, pxr::UsdTimeCode TimeCode );
+template USDUTILITIES_API double						UsdUtils::GetUsdValue< double >( const pxr::UsdAttribute& Attribute, pxr::UsdTimeCode TimeCode );
+template USDUTILITIES_API pxr::TfToken					UsdUtils::GetUsdValue< pxr::TfToken >( const pxr::UsdAttribute& Attribute, pxr::UsdTimeCode TimeCode );
 template USDUTILITIES_API pxr::GfVec3f					UsdUtils::GetUsdValue< pxr::GfVec3f >( const pxr::UsdAttribute& Attribute, pxr::UsdTimeCode TimeCode );
 template USDUTILITIES_API pxr::GfMatrix4d				UsdUtils::GetUsdValue< pxr::GfMatrix4d >( const pxr::UsdAttribute& Attribute, pxr::UsdTimeCode TimeCode );
 template USDUTILITIES_API pxr::SdfAssetPath				UsdUtils::GetUsdValue< pxr::SdfAssetPath >( const pxr::UsdAttribute& Attribute, pxr::UsdTimeCode TimeCode );
@@ -438,7 +441,7 @@ UClass* UsdUtils::GetComponentTypeForPrim( const pxr::UsdPrim& Prim )
 	{
 		return USkeletalMeshComponent::StaticClass();
 	}
-	else if ( Prim.IsA< pxr::UsdGeomMesh >() )
+	else if (Prim.IsA<pxr::UsdGeomGprim>())
 	{
 		return UStaticMeshComponent::StaticClass();
 	}
@@ -663,11 +666,11 @@ int32 UsdUtils::GetPrimvarUVIndex( FString PrimvarName )
 }
 
 TArray<TUsdStore<pxr::UsdGeomPrimvar>> UsdUtils::GetUVSetPrimvars(
-	const pxr::UsdGeomMesh& UsdMesh,
+	const pxr::UsdPrim& UsdPrim,
 	int32 MaxNumPrimvars
 )
 {
-	if (!UsdMesh)
+	if (!UsdPrim)
 	{
 		return {};
 	}
@@ -678,7 +681,7 @@ TArray<TUsdStore<pxr::UsdGeomPrimvar>> UsdUtils::GetUVSetPrimvars(
 	TArray<TUsdStore<pxr::UsdGeomPrimvar>> Float2Primvars;
 
 	// Collect all primvars that could be used as UV sets
-	pxr::UsdGeomPrimvarsAPI PrimvarsAPI{UsdMesh};
+	pxr::UsdGeomPrimvarsAPI PrimvarsAPI{UsdPrim};
 	for (const pxr::UsdGeomPrimvar& Primvar : PrimvarsAPI.GetPrimvars())
 	{
 		if (!Primvar || !Primvar.HasValue())
@@ -737,7 +740,7 @@ TArray<TUsdStore<pxr::UsdGeomPrimvar>> UsdUtils::GetUVSetPrimvars(
 	const pxr::TfToken& MaterialPurpose
 )
 {
-	return UsdUtils::GetUVSetPrimvars(UsdMesh);
+	return UsdUtils::GetUVSetPrimvars(UsdMesh.GetPrim());
 }
 
 TArray<TUsdStore<pxr::UsdGeomPrimvar>> UsdUtils::GetUVSetPrimvars(
@@ -746,7 +749,7 @@ TArray<TUsdStore<pxr::UsdGeomPrimvar>> UsdUtils::GetUVSetPrimvars(
 	const UsdUtils::FUsdPrimMaterialAssignmentInfo& UsdMeshMaterialAssignmentInfo
 )
 {
-	return UsdUtils::GetUVSetPrimvars(UsdMesh);
+	return UsdUtils::GetUVSetPrimvars(UsdMesh.GetPrim());
 }
 
 TArray<TUsdStore<pxr::UsdGeomPrimvar>> UsdUtils::AssemblePrimvarsIntoUVSets(
diff --git a/Engine/Plugins/Importers/USDImporter/Source/USDUtilities/Private/USDGeomMeshConversion.cpp b/Engine/Plugins/Importers/USDImporter/Source/USDUtilities/Private/USDGeomMeshConversion.cpp
--- a/Engine/Plugins/Importers/USDImporter/Source/USDUtilities/Private/USDGeomMeshConversion.cpp
+++ b/Engine/Plugins/Importers/USDImporter/Source/USDUtilities/Private/USDGeomMeshConversion.cpp
@@ -39,6 +39,7 @@
 #endif // WITH_EDITOR
 
 #include "USDIncludesStart.h"
+	#include "pxr/imaging/pxOsd/meshTopology.h"
 	#include "pxr/usd/ar/resolver.h"
 	#include "pxr/usd/ar/resolverScopedCache.h"
 	#include "pxr/usd/sdf/layer.h"
@@ -47,14 +48,22 @@
 	#include "pxr/usd/usd/editContext.h"
 	#include "pxr/usd/usd/prim.h"
 	#include "pxr/usd/usd/primRange.h"
+	#include "pxr/usd/usdGeom/capsule.h"
+	#include "pxr/usd/usdGeom/cone.h"
+	#include "pxr/usd/usdGeom/cube.h"
+	#include "pxr/usd/usdGeom/cylinder.h"
+	#include "pxr/usd/usdGeom/gprim.h"
 	#include "pxr/usd/usdGeom/mesh.h"
+	#include "pxr/usd/usdGeom/plane.h"
 	#include "pxr/usd/usdGeom/pointInstancer.h"
 	#include "pxr/usd/usdGeom/primvarsAPI.h"
+	#include "pxr/usd/usdGeom/sphere.h"
 	#include "pxr/usd/usdGeom/subset.h"
 	#include "pxr/usd/usdGeom/tokens.h"
 	#include "pxr/usd/usdShade/material.h"
 	#include "pxr/usd/usdShade/materialBindingAPI.h"
 	#include "pxr/usd/usdShade/tokens.h"
+	#include "pxr/usdImaging/usdImaging/implicitSurfaceMeshUtils.h"
 #include "USDIncludesEnd.h"
 
 #define LOCTEXT_NAMESPACE "USDGeomMeshConversion"
@@ -71,6 +80,10 @@ namespace UE::UsdGeomMeshConversion::Private
 {
 	static const FString DisplayColorID = TEXT( "!DisplayColor" );
 
+	// Dimensions used when generating Capsule meshes
+	static const float DefaultCapsuleMeshRadius = 0.25;
+	static const float DefaultCapsuleMeshHeight = 0.50;
+
 	int32 GetPrimValueIndex( const pxr::TfToken& InterpType, const int32 VertexIndex, const int32 VertexInstanceIndex, const int32 PolygonIndex )
 	{
 		if ( InterpType == pxr::UsdGeomTokens->vertex )
@@ -95,6 +108,187 @@ namespace UE::UsdGeomMeshConversion::Private
 		}
 	}
 
+	pxr::TfToken GetAttrInterpolation(const pxr::UsdAttribute& Attr, const pxr::TfToken& DefaultValue = pxr::UsdGeomTokens->constant)
+	{
+		pxr::TfToken RetrievedValue;
+		if (Attr && Attr.GetMetadata(pxr::UsdGeomTokens->interpolation, &RetrievedValue))
+		{
+			return RetrievedValue;
+		}
+		return DefaultValue;
+	}
+
+	pxr::TfToken GetGprimOrientation(const pxr::UsdGeomGprim& Gprim, pxr::UsdTimeCode TimeCode)
+	{
+		if (pxr::UsdAttribute Attr = Gprim.GetOrientationAttr())
+		{
+			pxr::TfToken Orientation;
+			if (Attr.Get(&Orientation, TimeCode))
+			{
+				return Orientation;
+			}
+		}
+		return pxr::UsdGeomTokens->rightHanded;
+	}
+
+	pxr::VtArray<int> GetFaceVertexCounts(const pxr::UsdPrim& UsdPrim, pxr::UsdTimeCode TimeCode)
+	{
+		if (pxr::UsdGeomMesh Mesh = pxr::UsdGeomMesh{UsdPrim})
+		{
+			pxr::UsdAttribute Attr = Mesh.GetFaceVertexCountsAttr();
+
+			pxr::VtArray<int> Result;
+			if (Attr && Attr.Get(&Result, TimeCode))
+			{
+				return Result;
+			}
+		}
+
+		const pxr::PxOsdMeshTopology* Topology = nullptr;
+		if (pxr::UsdGeomCapsule Capsule = pxr::UsdGeomCapsule{UsdPrim})
+		{
+			Topology = &pxr::UsdImagingGetCapsuleMeshTopology();
+		}
+		else if (pxr::UsdGeomCone Cone = pxr::UsdGeomCone{UsdPrim})
+		{
+			Topology = &pxr::UsdImagingGetUnitConeMeshTopology();
+		}
+		else if (pxr::UsdGeomCube Cube = pxr::UsdGeomCube{UsdPrim})
+		{
+			Topology = &pxr::UsdImagingGetUnitCubeMeshTopology();
+		}
+		else if (pxr::UsdGeomCylinder Cylinder = pxr::UsdGeomCylinder{UsdPrim})
+		{
+			Topology = &pxr::UsdImagingGetUnitCylinderMeshTopology();
+		}
+		else if (pxr::UsdGeomSphere Sphere = pxr::UsdGeomSphere{UsdPrim})
+		{
+			Topology = &pxr::UsdImagingGetUnitSphereMeshTopology();
+		}
+		else if (pxr::UsdGeomPlane Plane = pxr::UsdGeomPlane{UsdPrim})
+		{
+			Topology = &pxr::UsdImagingGetPlaneTopology();
+		}
+		if (Topology)
+		{
+			return Topology->GetFaceVertexCounts();
+		}
+
+		return {};
+	}
+
+	pxr::VtArray<pxr::GfVec3f> GetUnitCylinderMeshPoints(pxr::TfToken Axis)
+	{
+		if (Axis == pxr::UsdGeomTokens->x)
+		{
+			static const pxr::VtArray<pxr::GfVec3f> XCylinder = []()
+			{
+				// The USD cylinder is aligned to the z axis by default
+				pxr::VtArray<pxr::GfVec3f> Points = pxr::UsdImagingGetUnitCylinderMeshPoints();
+
+				pxr::GfMatrix4d ZToXAxis{
+					0.0, 1.0, 0.0, 0.0,
+					0.0, 0.0, 1.0, 0.0,
+					1.0, 0.0, 0.0, 0.0,
+					0.0, 0.0, 0.0, 1.0
+				};
+
+				for (pxr::GfVec3f& Point : Points)
+				{
+					Point = ZToXAxis.Transform(Point);
+				}
+
+				return Points;
+			}();
+			return XCylinder;
+		}
+		else if (Axis == pxr::UsdGeomTokens->y)
+		{
+			static const pxr::VtArray<pxr::GfVec3f> YCylinder = []()
+			{
+				// The USD cylinder is aligned to the z axis by default
+				pxr::VtArray<pxr::GfVec3f> Points = pxr::UsdImagingGetUnitCylinderMeshPoints();
+
+				pxr::GfMatrix4d ZToYAxis{
+					0.0, 0.0, 1.0, 0.0,
+					1.0, 0.0, 0.0, 0.0,
+					0.0, 1.0, 0.0, 0.0,
+					0.0, 0.0, 0.0, 1.0
+				};
+
+				for (pxr::GfVec3f& Point : Points)
+				{
+					Point = ZToYAxis.Transform(Point);
+				}
+
+				return Points;
+			}();
+			return YCylinder;
+		}
+		else if (Axis == pxr::UsdGeomTokens->z)
+		{
+			return pxr::UsdImagingGetUnitCylinderMeshPoints();
+		}
+
+		return {};
+	}
+
+	pxr::VtArray<pxr::GfVec3f> GetUnitConeMeshPoints(pxr::TfToken Axis)
+	{
+		if (Axis == pxr::UsdGeomTokens->x)
+		{
+			static const pxr::VtArray<pxr::GfVec3f> XCylinder = []()
+			{
+				// The USD cone is aligned to the z axis by default
+				pxr::VtArray<pxr::GfVec3f> Points = pxr::UsdImagingGetUnitConeMeshPoints();
+
+				pxr::GfMatrix4d ZToXAxis{
+					0.0, 1.0, 0.0, 0.0,
+					0.0, 0.0, 1.0, 0.0,
+					1.0, 0.0, 0.0, 0.0,
+					0.0, 0.0, 0.0, 1.0
+				};
+
+				for (pxr::GfVec3f& Point : Points)
+				{
+					Point = ZToXAxis.Transform(Point);
+				}
+
+				return Points;
+			}();
+			return XCylinder;
+		}
+		else if (Axis == pxr::UsdGeomTokens->y)
+		{
+			static const pxr::VtArray<pxr::GfVec3f> YCylinder = []()
+			{
+				// The USD cone is aligned to the z axis by default
+				pxr::VtArray<pxr::GfVec3f> Points = pxr::UsdImagingGetUnitConeMeshPoints();
+
+				pxr::GfMatrix4d ZToYAxis{
+					0.0, 0.0, 1.0, 0.0,
+					1.0, 0.0, 0.0, 0.0,
+					0.0, 1.0, 0.0, 0.0,
+					0.0, 0.0, 0.0, 1.0
+				};
+
+				for (pxr::GfVec3f& Point : Points)
+				{
+					Point = ZToYAxis.Transform(Point);
+				}
+
+				return Points;
+			}();
+			return YCylinder;
+		}
+		else if (Axis == pxr::UsdGeomTokens->z)
+		{
+			return pxr::UsdImagingGetUnitConeMeshPoints();
+		}
+
+		return {};
+	}
+
 	int32 GetLODIndexFromName( const std::string& Name )
 	{
 		const std::string LODString = UnrealIdentifiers::LOD.GetString();
@@ -583,6 +777,11 @@ namespace UE::UsdGeomMeshConversion::Private
 		{
 			bSuccess = UsdToUnreal::ConvertGeomMesh( Mesh, OutMeshDescription, OutMaterialAssignments, Options );
 		}
+		// Check for a cube/capsule/etc. ConvertGeomPrimitive will internally check for all specific types
+		else if (pxr::UsdGeomGprim Gprim = pxr::UsdGeomGprim{Prim})
+		{
+			bSuccess = UsdToUnreal::ConvertGeomPrimitive(Prim, OutMeshDescription, OutMaterialAssignments, Options);
+		}
 		else if ( pxr::UsdGeomPointInstancer PointInstancer = pxr::UsdGeomPointInstancer{ Prim } )
 		{
 			bSuccess = UsdToUnreal::ConvertPointInstancerToMesh(
@@ -654,28 +853,7 @@ namespace UE::UsdGeomMeshConversion::Private
 
 		bool bTraverseChildren = true;
 
-		if (pxr::UsdGeomMesh Mesh = pxr::UsdGeomMesh(Prim))
-		{
-			TArray<TUsdStore<pxr::UsdGeomPrimvar>> MeshPrimvars = UsdUtils::GetUVSetPrimvars(
-				Mesh,
-				TNumericLimits<int32>::Max()
-			);
-
-			for (const TUsdStore<pxr::UsdGeomPrimvar>& MeshPrimvar : MeshPrimvars)
-			{
-				FString PrimvarName = UsdToUnreal::ConvertToken(MeshPrimvar.Get().GetName());
-				PrimvarName.RemoveFromStart(TEXT("primvars:"));
-
-				InOutAllPrimvars.Add(PrimvarName);
-
-				// Keep track of which primvars are texCoord2f as we always want to prefer these over other float2s
-				if (MeshPrimvar.Get().GetTypeName().GetRole() == pxr::SdfValueTypeNames->TexCoord2f.GetRole())
-				{
-					InOutPreferredPrimvars.Add(PrimvarName);
-				}
-			}
-		}
-		else if (pxr::UsdGeomPointInstancer PointInstancer = pxr::UsdGeomPointInstancer{Prim})
+		if (pxr::UsdGeomPointInstancer PointInstancer = pxr::UsdGeomPointInstancer{Prim})
 		{
 			pxr::SdfPathVector PrototypePaths;
 			if (!PointInstancer.GetPrototypesRel().GetTargets(&PrototypePaths))
@@ -705,6 +883,27 @@ namespace UE::UsdGeomMeshConversion::Private
 			// We never want to step into point instancers when fetching prims for drawing
 			bTraverseChildren = false;
 		}
+		if (Prim)
+		{
+			TArray<TUsdStore<pxr::UsdGeomPrimvar>> Primvars = UsdUtils::GetUVSetPrimvars(
+				Prim,
+				TNumericLimits<int32>::Max()
+			);
+
+			for (const TUsdStore<pxr::UsdGeomPrimvar>& Primvar : Primvars)
+			{
+				FString PrimvarName = UsdToUnreal::ConvertToken(Primvar.Get().GetName());
+				PrimvarName.RemoveFromStart(TEXT("primvars:"));
+
+				InOutAllPrimvars.Add(PrimvarName);
+
+				// Keep track of which primvars are texCoord2f as we always want to prefer these over other float2s
+				if (Primvar.Get().GetTypeName().GetRole() == pxr::SdfValueTypeNames->TexCoord2f.GetRole())
+				{
+					InOutPreferredPrimvars.Add(PrimvarName);
+				}
+			}
+		}
 
 		if (bTraverseChildren)
 		{
@@ -748,560 +947,606 @@ namespace UE::UsdGeomMeshConversion::Private
 			PreferredPrimvars
 		);
 	}
-}
-namespace UsdGeomMeshImpl = UE::UsdGeomMeshConversion::Private;
 
-namespace UsdToUnreal
-{
-	const FUsdMeshConversionOptions FUsdMeshConversionOptions::DefaultOptions;
-
-	FUsdMeshConversionOptions::FUsdMeshConversionOptions()
-		: AdditionalTransform( FTransform::Identity )
-		, PurposesToLoad( EUsdPurpose::Render )
-		, RenderContext( pxr::UsdShadeTokens->universalRenderContext )
-		, MaterialPurpose( pxr::UsdShadeTokens->allPurpose )
-		, TimeCode( pxr::UsdTimeCode::EarliestTime() )
-		, bMergeIdenticalMaterialSlots( true )
+	// Unconverted, raw USD mesh data to convert into a MeshDescription
+	struct FUsdMeshData
 	{
-	}
-}
+		// So that we can reference the prim on error messages
+		FString SourcePrimPath;
 
-bool UsdToUnreal::ConvertGeomMesh(
-	const pxr::UsdGeomMesh& UsdMesh,
-	FMeshDescription& OutMeshDescription,
-	UsdUtils::FUsdPrimMaterialAssignmentInfo& OutMaterialAssignments,
-	const FUsdMeshConversionOptions& Options
-)
-{
-	TRACE_CPUPROFILER_EVENT_SCOPE( UsdToUnreal::ConvertGeomMesh );
+		pxr::VtArray<int> FaceVertexCounts;
+		pxr::VtArray<int> FaceIndices;
+		pxr::VtArray<pxr::GfVec3f> Points;
+		pxr::VtArray<pxr::GfVec3f> Normals;
+		pxr::VtArray<pxr::GfVec3f> Velocities;
+		pxr::VtArray<pxr::GfVec3f> DisplayColors;
+		pxr::VtArray<float> DisplayOpacities;
 
-	if ( !UsdMesh )
-	{
-		return false;
-	}
+		pxr::TfToken Orientation = pxr::UsdGeomTokens->rightHanded;
+		pxr::TfToken NormalInterpolation = pxr::UsdGeomTokens->vertex;
+		pxr::TfToken VelocityInterpolation = pxr::UsdGeomTokens->vertex;
+		pxr::TfToken DisplayColorInterpolation = pxr::UsdGeomTokens->constant;
+		pxr::TfToken DisplayOpacityInterpolation = pxr::UsdGeomTokens->constant;
 
-	FScopedUsdAllocs Allocs;
+		TOptional<int32> ProvidedNumUVSets;
+		TArray<TUsdStore<pxr::UsdGeomPrimvar>> PrimvarsByUVIndex;
 
-	pxr::UsdPrim UsdPrim = UsdMesh.GetPrim();
-	pxr::UsdStageRefPtr Stage = UsdPrim.GetStage();
-	const FUsdStageInfo StageInfo( Stage );
+		int32 MaterialIndexOffset = 0;
+		UsdUtils::FUsdPrimMaterialAssignmentInfo LocalMaterialInfo;
+	};
 
-	const double TimeCodeValue = Options.TimeCode.GetValue();
+	bool ConvertMeshData(
+		const FUsdMeshData& InMeshData,
+		const FUsdStageInfo& InStageInfo,
+		const UsdToUnreal::FUsdMeshConversionOptions& InOptions,
+		FMeshDescription& OutMeshDescription,
+		UsdUtils::FUsdPrimMaterialAssignmentInfo& OutMaterialAssignments
+	)
+	{
+		TRACE_CPUPROFILER_EVENT_SCOPE(UsdToUnreal::ConvertMeshData);
 
-	const int32 MaterialIndexOffset = OutMaterialAssignments.Slots.Num();
+		FScopedUsdAllocs Allocs;
 
-	// Material assignments
-	const bool bProvideMaterialIndices = true;
-	UsdUtils::FUsdPrimMaterialAssignmentInfo LocalInfo = UsdUtils::GetPrimMaterialAssignments(
-		UsdPrim,
-		Options.TimeCode,
-		bProvideMaterialIndices,
-		Options.RenderContext,
-		Options.MaterialPurpose
-	);
-	TArray< UsdUtils::FUsdPrimMaterialSlot >& LocalMaterialSlots = LocalInfo.Slots;
-	TArray< int32 >& FaceMaterialIndices = LocalInfo.MaterialIndices;
+		// Material assignments
+		const TArray<UsdUtils::FUsdPrimMaterialSlot>& LocalMaterialSlots = InMeshData.LocalMaterialInfo.Slots;
+		const TArray<int32>& FaceMaterialIndices = InMeshData.LocalMaterialInfo.MaterialIndices;
 
-	// Position 3 in this has the value 6 --> Local material slot #3 is actually the combined material slot #6
-	TArray<int32> LocalToCombinedMaterialSlotIndices;
-	LocalToCombinedMaterialSlotIndices.SetNumZeroed( LocalInfo.Slots.Num() );
+		// Position 3 in this has the value 6 --> Local material slot #3 is actually the combined material slot #6
+		TArray<int32> LocalToCombinedMaterialSlotIndices;
+		LocalToCombinedMaterialSlotIndices.SetNumZeroed(InMeshData.LocalMaterialInfo.Slots.Num());
 
-	if ( Options.bMergeIdenticalMaterialSlots )
-	{
-		// Build a map of our existing slots since we can hash the entire slot, and our incoming mesh may have an arbitrary number of new slots
-		TMap< UsdUtils::FUsdPrimMaterialSlot, int32 > CombinedMaterialSlotsToIndex;
-		for ( int32 Index = 0; Index < OutMaterialAssignments.Slots.Num(); ++Index )
+		if (InOptions.bMergeIdenticalMaterialSlots)
 		{
-			const UsdUtils::FUsdPrimMaterialSlot& Slot = OutMaterialAssignments.Slots[ Index ];
-
-			// Combine entries in this way so that we can append PrimPaths
-			TMap< UsdUtils::FUsdPrimMaterialSlot, int32 >::TKeyIterator KeyIt = CombinedMaterialSlotsToIndex.CreateKeyIterator( Slot );
-			if ( KeyIt )
-			{
-				KeyIt.Key().PrimPaths.Append( Slot.PrimPaths );
-				KeyIt.Value() = Index;
-			}
-			else
+			// Build a map of our existing slots since we can hash the entire slot, and our incoming mesh may have an arbitrary number of new slots
+			TMap<UsdUtils::FUsdPrimMaterialSlot, int32> CombinedMaterialSlotsToIndex;
+			for (int32 Index = 0; Index < OutMaterialAssignments.Slots.Num(); ++Index)
 			{
-				CombinedMaterialSlotsToIndex.Add( Slot, Index );
-			}
-		}
+				const UsdUtils::FUsdPrimMaterialSlot& Slot = OutMaterialAssignments.Slots[Index];
 
-		// Combine our LocalSlots into CombinedMaterialSlotsToIndex
-		for ( int32 LocalIndex = 0; LocalIndex < LocalInfo.Slots.Num(); ++LocalIndex )
-		{
-			const UsdUtils::FUsdPrimMaterialSlot& LocalSlot = LocalInfo.Slots[ LocalIndex ];
+				// Combine entries in this way so that we can append PrimPaths
+				TMap<UsdUtils::FUsdPrimMaterialSlot, int32>::TKeyIterator KeyIt = CombinedMaterialSlotsToIndex.CreateKeyIterator(Slot);
+				if (KeyIt)
+				{
+					KeyIt.Key().PrimPaths.Append(Slot.PrimPaths);
+					KeyIt.Value() = Index;
+				}
+				else
+				{
+					CombinedMaterialSlotsToIndex.Add(Slot, Index);
+				}
+			}
 
-			// Combine entries in this way so that we can append PrimPaths
-			TMap< UsdUtils::FUsdPrimMaterialSlot, int32 >::TKeyIterator KeyIt = CombinedMaterialSlotsToIndex.CreateKeyIterator( LocalSlot );
-			if ( KeyIt )
+			// Combine our LocalSlots into CombinedMaterialSlotsToIndex
+			for (int32 LocalIndex = 0; LocalIndex < InMeshData.LocalMaterialInfo.Slots.Num(); ++LocalIndex)
 			{
-				KeyIt.Key().PrimPaths.Append( LocalSlot.PrimPaths );
+				const UsdUtils::FUsdPrimMaterialSlot& LocalSlot = InMeshData.LocalMaterialInfo.Slots[LocalIndex];
 
-				const int32 ExistingCombinedIndex = KeyIt.Value();
-				LocalToCombinedMaterialSlotIndices[ LocalIndex ] = ExistingCombinedIndex;
+				// Combine entries in this way so that we can append PrimPaths
+				TMap<UsdUtils::FUsdPrimMaterialSlot, int32>::TKeyIterator KeyIt = CombinedMaterialSlotsToIndex.CreateKeyIterator(LocalSlot);
+				if (KeyIt)
+				{
+					KeyIt.Key().PrimPaths.Append(LocalSlot.PrimPaths);
+
+					const int32 ExistingCombinedIndex = KeyIt.Value();
+					LocalToCombinedMaterialSlotIndices[LocalIndex] = ExistingCombinedIndex;
+				}
+				else
+				{
+					int32 NewIndex = OutMaterialAssignments.Slots.Add(LocalSlot);
+					CombinedMaterialSlotsToIndex.Add(LocalSlot, NewIndex);
+					LocalToCombinedMaterialSlotIndices[LocalIndex] = NewIndex;
+				}
 			}
-			else
+
+			// Now that we merged all prim paths into they keys of CombinedMaterialSlotsToIndex, let's copy them back into
+			// our output
+			for (UsdUtils::FUsdPrimMaterialSlot& Slot : OutMaterialAssignments.Slots)
 			{
-				int32 NewIndex = OutMaterialAssignments.Slots.Add( LocalSlot );
-				CombinedMaterialSlotsToIndex.Add( LocalSlot, NewIndex );
-				LocalToCombinedMaterialSlotIndices[ LocalIndex ] = NewIndex;
+				TMap<UsdUtils::FUsdPrimMaterialSlot, int32>::TKeyIterator KeyIt = CombinedMaterialSlotsToIndex.CreateKeyIterator(Slot);
+				ensure(KeyIt);
+				Slot.PrimPaths = KeyIt.Key().PrimPaths;
 			}
 		}
-
-		// Now that we merged all prim paths into they keys of CombinedMaterialSlotsToIndex, let's copy them back into
-		// our output
-		for ( UsdUtils::FUsdPrimMaterialSlot& Slot : OutMaterialAssignments.Slots )
-		{
-			TMap< UsdUtils::FUsdPrimMaterialSlot, int32 >::TKeyIterator KeyIt = CombinedMaterialSlotsToIndex.CreateKeyIterator( Slot );
-			ensure(KeyIt);
-			Slot.PrimPaths = KeyIt.Key().PrimPaths;
-		}
-	}
-	else
-	{
-		// Just append our new local material slots at the end of MaterialAssignments
-		OutMaterialAssignments.Slots.Append( LocalInfo.Slots );
-		for ( int32 LocalIndex = 0; LocalIndex < LocalInfo.Slots.Num(); ++LocalIndex )
+		else
 		{
-			LocalToCombinedMaterialSlotIndices[ LocalIndex ] = LocalIndex + MaterialIndexOffset;
+			// Just append our new local material slots at the end of MaterialAssignments
+			OutMaterialAssignments.Slots.Append(InMeshData.LocalMaterialInfo.Slots);
+			for (int32 LocalIndex = 0; LocalIndex < InMeshData.LocalMaterialInfo.Slots.Num(); ++LocalIndex)
+			{
+				LocalToCombinedMaterialSlotIndices[LocalIndex] = LocalIndex + InMeshData.MaterialIndexOffset;
+			}
 		}
-	}
 
-	const int32 VertexOffset = OutMeshDescription.Vertices().Num();
-	const int32 VertexInstanceOffset = OutMeshDescription.VertexInstances().Num();
-	const int32 PolygonOffset = OutMeshDescription.Polygons().Num();
+		const int32 VertexOffset = OutMeshDescription.Vertices().Num();
+		const int32 VertexInstanceOffset = OutMeshDescription.VertexInstances().Num();
 
-	FStaticMeshAttributes StaticMeshAttributes( OutMeshDescription );
+		FStaticMeshAttributes StaticMeshAttributes(OutMeshDescription);
 
-	// Vertex positions
-	TVertexAttributesRef< FVector3f > MeshDescriptionVertexPositions = StaticMeshAttributes.GetVertexPositions();
-	{
-		pxr::UsdAttribute Points = UsdMesh.GetPointsAttr();
-		if ( Points )
+		// Vertex positions
+		TVertexAttributesRef<FVector3f> MeshDescriptionVertexPositions = StaticMeshAttributes.GetVertexPositions();
 		{
-			pxr::VtArray< pxr::GfVec3f > PointsArray;
-			Points.Get( &PointsArray, TimeCodeValue );
-			if (PointsArray.size() < 3)
+			if (InMeshData.Points.size() < 3)
 			{
 				return false;
 			}
 
-			OutMeshDescription.ReserveNewVertices( PointsArray.size() );
+			OutMeshDescription.ReserveNewVertices(InMeshData.Points.size());
 
-			for ( int32 LocalPointIndex = 0; LocalPointIndex < PointsArray.size(); ++LocalPointIndex )
+			for (size_t LocalPointIndex = 0; LocalPointIndex < InMeshData.Points.size(); ++LocalPointIndex)
 			{
-				const pxr::GfVec3f& Point = PointsArray[ LocalPointIndex ];
+				const pxr::GfVec3f& Point = InMeshData.Points.cdata()[LocalPointIndex];
 
-				FVector Position = Options.AdditionalTransform.TransformPosition( UsdToUnreal::ConvertVector( StageInfo, Point ) );
+				FVector Position = InOptions.AdditionalTransform.TransformPosition(UsdToUnreal::ConvertVector(InStageInfo, Point));
 
 				FVertexID AddedVertexId = OutMeshDescription.CreateVertex();
-				MeshDescriptionVertexPositions[ AddedVertexId ] = ( FVector3f ) Position;
+				MeshDescriptionVertexPositions[AddedVertexId] = (FVector3f)Position;
 			}
 		}
-	}
-
-	uint32 NumSkippedPolygons = 0;
-	uint32 NumPolygons = 0;
 
-	// Polygons
-	{
-		TMap<int32, FPolygonGroupID> PolygonGroupMapping;
-		TArray<FVertexInstanceID> CornerInstanceIDs;
-		TArray<FVertexID> CornerVerticesIDs;
-		int32 CurrentVertexInstanceIndex = 0;
-		TPolygonGroupAttributesRef<FName> MaterialSlotNames = StaticMeshAttributes.GetPolygonGroupMaterialSlotNames();
-
-		// If we're going to share existing material slots, we'll need to share existing PolygonGroups in the mesh description,
-		// so we need to traverse it and prefill PolygonGroupMapping
-		if ( Options.bMergeIdenticalMaterialSlots )
+		uint32 NumSkippedPolygons = 0;
+		uint32 NumPolygons = InMeshData.FaceVertexCounts.size();
+		if (NumPolygons < 1)
 		{
-			for ( FPolygonGroupID PolygonGroupID : OutMeshDescription.PolygonGroups().GetElementIDs() )
-			{
-				// We always create our polygon groups in order with our combined material slots, so its easy to reconstruct this mapping here
-				PolygonGroupMapping.Add( PolygonGroupID.GetValue(), PolygonGroupID );
-			}
+			return false;
 		}
-
-		// Material slots
-		// Note that we always create these in the order they show up in LocalInfo: If we created these on-demand when parsing polygons (like before)
-		// we could run into polygons in a different order than the material slots and end up with different material assignments.
-		// We could use the StaticMesh's SectionInfoMap to unswitch things, but that's not available at runtime so we better do this here
-		for ( int32 LocalMaterialIndex = 0; LocalMaterialIndex < LocalInfo.Slots.Num(); ++LocalMaterialIndex )
+		if (InMeshData.FaceIndices.size() < 1)
 		{
-			const int32 CombinedMaterialIndex = LocalToCombinedMaterialSlotIndices[ LocalMaterialIndex ];
-			if ( !PolygonGroupMapping.Contains( CombinedMaterialIndex ) )
-			{
-				FPolygonGroupID NewPolygonGroup = OutMeshDescription.CreatePolygonGroup();
-				PolygonGroupMapping.Add( CombinedMaterialIndex, NewPolygonGroup );
-
-				// This is important for runtime, where the material slots are matched to LOD sections based on their material slot name
-				MaterialSlotNames[ NewPolygonGroup ] = *LexToString( NewPolygonGroup.GetValue() );
-			}
+			return false;
 		}
 
-		bool bFlipThisGeometry = false;
-
-		if ( IUsdPrim::GetGeometryOrientation( UsdMesh ) == EUsdGeomOrientation::LeftHanded )
+		// Polygons
 		{
-			bFlipThisGeometry = !bFlipThisGeometry;
-		}
+			TMap<int32, FPolygonGroupID> PolygonGroupMapping;
+			TArray<FVertexInstanceID> CornerInstanceIDs;
+			TArray<FVertexID> CornerVerticesIDs;
+			int32 CurrentVertexInstanceIndex = 0;
+			TPolygonGroupAttributesRef<FName> MaterialSlotNames = StaticMeshAttributes.GetPolygonGroupMaterialSlotNames();
 
-		// Face counts
-		pxr::UsdAttribute FaceCountsAttribute = UsdMesh.GetFaceVertexCountsAttr();
-		pxr::VtArray< int > FaceCounts;
-
-		if ( FaceCountsAttribute )
-		{
-			FaceCountsAttribute.Get( &FaceCounts, TimeCodeValue );
-			NumPolygons = FaceCounts.size();
-			if (NumPolygons < 1)
+			// If we're going to share existing material slots, we'll need to share existing PolygonGroups in the mesh description,
+			// so we need to traverse it and prefill PolygonGroupMapping
+			if (InOptions.bMergeIdenticalMaterialSlots)
 			{
-				return false;
+				for (FPolygonGroupID PolygonGroupID : OutMeshDescription.PolygonGroups().GetElementIDs())
+				{
+					// We always create our polygon groups in order with our combined material slots, so its easy to reconstruct this mapping here
+					PolygonGroupMapping.Add(PolygonGroupID.GetValue(), PolygonGroupID);
+				}
 			}
-		}
 
-		// Face indices
-		pxr::UsdAttribute FaceIndicesAttribute = UsdMesh.GetFaceVertexIndicesAttr();
-		pxr::VtArray< int > FaceIndices;
-
-		if ( FaceIndicesAttribute )
-		{
-			FaceIndicesAttribute.Get( &FaceIndices, TimeCodeValue );
-			if (FaceIndices.size() < 1)
+			// Material slots
+			// Note that we always create these in the order they show up in LocalInfo: If we created these on-demand when parsing polygons (like
+			// before) we could run into polygons in a different order than the material slots and end up with different material assignments. We
+			// could use the StaticMesh's SectionInfoMap to unswitch things, but that's not available at runtime so we better do this here
+			for (int32 LocalMaterialIndex = 0; LocalMaterialIndex < InMeshData.LocalMaterialInfo.Slots.Num(); ++LocalMaterialIndex)
 			{
-				return false;
-			}
-		}
-
-		// Normals
-		pxr::UsdAttribute NormalsAttribute = UsdMesh.GetNormalsAttr();
-		pxr::VtArray< pxr::GfVec3f > Normals;
-
-		if ( NormalsAttribute )
-		{
-			NormalsAttribute.Get( &Normals, TimeCodeValue );
-		}
-
-		pxr::TfToken NormalsInterpType = UsdMesh.GetNormalsInterpolation();
+				const int32 CombinedMaterialIndex = LocalToCombinedMaterialSlotIndices[LocalMaterialIndex];
+				if (!PolygonGroupMapping.Contains(CombinedMaterialIndex))
+				{
+					FPolygonGroupID NewPolygonGroup = OutMeshDescription.CreatePolygonGroup();
+					PolygonGroupMapping.Add(CombinedMaterialIndex, NewPolygonGroup);
 
-		// Velocities
-		pxr::UsdAttribute VelocitiesAttribute = UsdMesh.GetVelocitiesAttr();
-		pxr::TfToken VelocitiesInterpType = pxr::UsdGeomTokens->vertex;
-		pxr::VtArray<pxr::GfVec3f> Velocities;
-		if (VelocitiesAttribute)
-		{
-			VelocitiesAttribute.Get(&Velocities, TimeCodeValue);
+					// This is important for runtime, where the material slots are matched to LOD sections based on their material slot name
+					MaterialSlotNames[NewPolygonGroup] = *LexToString(NewPolygonGroup.GetValue());
+				}
+			}
 
-			if (Velocities.size() > 0)
+			// Velocities
+			if (InMeshData.Velocities.size() > 0)
 			{
 				if (!OutMeshDescription.VertexInstanceAttributes().HasAttribute(MeshAttribute::VertexInstance::Velocity))
 				{
-					OutMeshDescription.VertexInstanceAttributes().RegisterAttribute<FVector3f>(MeshAttribute::VertexInstance::Velocity, 1, FVector3f::ZeroVector, EMeshAttributeFlags::Lerpable);
-				}
-
-				pxr::TfToken InterpType;
-				if (VelocitiesAttribute.GetMetadata(pxr::UsdGeomTokens->interpolation, &InterpType))
-				{
-					VelocitiesInterpType = InterpType;
+					OutMeshDescription.VertexInstanceAttributes().RegisterAttribute<FVector3f>(
+						MeshAttribute::VertexInstance::Velocity,
+						1,
+						FVector3f::ZeroVector,
+						EMeshAttributeFlags::Lerpable
+					);
 				}
 			}
-		}
-
-		// UVs
-		TVertexInstanceAttributesRef< FVector2f > MeshDescriptionUVs = StaticMeshAttributes.GetVertexInstanceUVs();
 
-		struct FUVSet
-		{
-			int32 UVSetIndexUE; // The user may only have 'uv4' and 'uv5', so we can't just use array indices to find the target UV channel
-			TOptional< pxr::VtIntArray > UVIndices; // UVs might be indexed or they might be flat (one per vertex)
-			pxr::VtVec2fArray UVs;
-
-			pxr::TfToken InterpType = pxr::UsdGeomTokens->faceVarying;
-		};
-
-		TArray< FUVSet > UVSets;
-
-		TOptional<int32> ProvidedNumUVSets;
+			// UVs
+			TVertexInstanceAttributesRef<FVector2f> MeshDescriptionUVs = StaticMeshAttributes.GetVertexInstanceUVs();
 
-		// If we already have a primvar to UV index assignment, let's just use that.
-		// When collapsing, we'll do a pre-pass on all meshes to translate and determine this beforehand.
-		TArray<TUsdStore<pxr::UsdGeomPrimvar>> PrimvarsByUVIndex;
-		if (OutMaterialAssignments.PrimvarToUVIndex.Num() > 0)
-		{
-			int32 HighestProvidedUVIndex = 0;
-			for (const TPair<FString, int32>& Pair : OutMaterialAssignments.PrimvarToUVIndex)
+			struct FUVSet
 			{
-				HighestProvidedUVIndex = FMath::Max(HighestProvidedUVIndex, Pair.Value);
-			}
-			ProvidedNumUVSets = HighestProvidedUVIndex + 1;
+				int32 UVSetIndexUE;	   // The user may only have 'uv4' and 'uv5', so we can't just use array indices to find the target UV channel
+				TOptional<pxr::VtIntArray> UVIndices;	 // UVs might be indexed or they might be flat (one per vertex)
+				pxr::VtVec2fArray UVs;
 
-			TArray<TUsdStore<pxr::UsdGeomPrimvar>> AllMeshUVPrimvars =
-				UsdUtils::GetUVSetPrimvars(UsdMesh, TNumericLimits<int32>::Max());
+				pxr::TfToken InterpType = pxr::UsdGeomTokens->faceVarying;
+			};
 
-			PrimvarsByUVIndex =
-				UsdUtils::AssemblePrimvarsIntoUVSets(AllMeshUVPrimvars, OutMaterialAssignments.PrimvarToUVIndex);
-		}
-		// Let's use the best primvar assignment for this particular mesh instead
-		else
-		{
-			PrimvarsByUVIndex = UsdUtils::GetUVSetPrimvars(UsdMesh);
-
-			OutMaterialAssignments.PrimvarToUVIndex =
-				UsdUtils::AssemblePrimvarsIntoPrimvarToUVIndexMap(PrimvarsByUVIndex);
-		}
+			TArray<FUVSet> UVSets;
 
-		int32 HighestAddedUVChannel = 0;
-		for (int32 UVChannelIndex = 0; UVChannelIndex < PrimvarsByUVIndex.Num(); ++UVChannelIndex)
-		{
-			pxr::UsdGeomPrimvar& Primvar = PrimvarsByUVIndex[UVChannelIndex].Get();
-			if (!Primvar)
+			int32 HighestAddedUVChannel = 0;
+			for (int32 UVChannelIndex = 0; UVChannelIndex < InMeshData.PrimvarsByUVIndex.Num(); ++UVChannelIndex)
 			{
-				continue;
-			}
-
-			FUVSet UVSet;
-			UVSet.InterpType = Primvar.GetInterpolation();
-			UVSet.UVSetIndexUE = UVChannelIndex;
+				const pxr::UsdGeomPrimvar& Primvar = InMeshData.PrimvarsByUVIndex[UVChannelIndex].Get();
+				if (!Primvar)
+				{
+					continue;
+				}
 
-			if ( Primvar.IsIndexed() )
-			{
-				UVSet.UVIndices.Emplace();
+				FUVSet UVSet;
+				UVSet.InterpType = Primvar.GetInterpolation();
+				UVSet.UVSetIndexUE = UVChannelIndex;
 
-				if ( Primvar.GetIndices( &UVSet.UVIndices.GetValue(), Options.TimeCode ) && Primvar.Get( &UVSet.UVs, Options.TimeCode ) )
+				if (Primvar.IsIndexed())
 				{
-					if ( UVSet.UVs.size() > 0 )
+					UVSet.UVIndices.Emplace();
+
+					if (Primvar.GetIndices(&UVSet.UVIndices.GetValue(), InOptions.TimeCode) && Primvar.Get(&UVSet.UVs, InOptions.TimeCode))
 					{
-						UVSets.Add( MoveTemp( UVSet ) );
-						HighestAddedUVChannel = UVSet.UVSetIndexUE;
+						if (UVSet.UVs.size() > 0)
+						{
+							UVSets.Add(MoveTemp(UVSet));
+							HighestAddedUVChannel = UVSet.UVSetIndexUE;
+						}
 					}
 				}
-			}
-			else
-			{
-				if ( Primvar.Get( &UVSet.UVs ) )
+				else
 				{
-					if ( UVSet.UVs.size() > 0 )
+					if (Primvar.Get(&UVSet.UVs))
 					{
-						UVSets.Add( MoveTemp( UVSet ) );
-						HighestAddedUVChannel = UVSet.UVSetIndexUE;
+						if (UVSet.UVs.size() > 0)
+						{
+							UVSets.Add(MoveTemp(UVSet));
+							HighestAddedUVChannel = UVSet.UVSetIndexUE;
+						}
 					}
 				}
 			}
-		}
-
-		// When importing multiple mesh pieces to the same static mesh.  Ensure each mesh piece has the same number of UVs
-		{
-			int32 ExistingUVCount = MeshDescriptionUVs.GetNumChannels();
-			int32 NumUVs = FMath::Max( HighestAddedUVChannel + 1, ExistingUVCount );
 
-			// When we provide a PrimvarToUVIndex map to this function it means we'll end up combining this
-			// MeshDescription with others later (e.g. due to collapsing or multiple-LOD meshes).
-			// In that case we can get better results by making sure all of the individual MeshDescriptions have the
-			// same total number of UV sets, even if the unused ones are empty.
-			// Otherwise, if we e.g. have a material reading UVIndex3 when we only have a single UV set, UE seems to
-			// just read that one UV set anyway, which is somewhat unexpected and can be misleading
-			if (ProvidedNumUVSets.IsSet())
+			// When importing multiple mesh pieces to the same static mesh.  Ensure each mesh piece has the same number of UVs
 			{
-				NumUVs = FMath::Max<int32>(ProvidedNumUVSets.GetValue(), NumUVs);
-			}
-
-			NumUVs = FMath::Min<int32>( USD_PREVIEW_SURFACE_MAX_UV_SETS, NumUVs );
-			// At least one UV set must exist.
-			NumUVs = FMath::Max<int32>(1, NumUVs);
+				int32 ExistingUVCount = MeshDescriptionUVs.GetNumChannels();
+				int32 NumUVs = FMath::Max(HighestAddedUVChannel + 1, ExistingUVCount);
 
-			//Make sure all Vertex instance have the correct number of UVs
-			MeshDescriptionUVs.SetNumChannels( NumUVs );
-		}
-
-		TVertexInstanceAttributesRef< FVector3f > MeshDescriptionNormals = StaticMeshAttributes.GetVertexInstanceNormals();
-		TVertexInstanceAttributesRef<FVector3f> MeshDescriptionVelocities = OutMeshDescription.VertexInstanceAttributes().GetAttributesRef<FVector3f>(MeshAttribute::VertexInstance::Velocity);
-
-		OutMeshDescription.ReserveNewVertexInstances( FaceCounts.size() * 3 );
-		OutMeshDescription.ReserveNewPolygons( FaceCounts.size() );
-		OutMeshDescription.ReserveNewEdges( FaceCounts.size() * 2 );
-
-		// Vertex color
-		TVertexInstanceAttributesRef< FVector4f > MeshDescriptionColors = StaticMeshAttributes.GetVertexInstanceColors();
+				// When we provide a PrimvarToUVIndex map to this function it means we'll end up combining this
+				// MeshDescription with others later (e.g. due to collapsing or multiple-LOD meshes).
+				// In that case we can get better results by making sure all of the individual MeshDescriptions have the
+				// same total number of UV sets, even if the unused ones are empty.
+				// Otherwise, if we e.g. have a material reading UVIndex3 when we only have a single UV set, UE seems to
+				// just read that one UV set anyway, which is somewhat unexpected and can be misleading
+				if (InMeshData.ProvidedNumUVSets.IsSet())
+				{
+					NumUVs = FMath::Max<int32>(InMeshData.ProvidedNumUVSets.GetValue(), NumUVs);
+				}
 
-		pxr::UsdGeomPrimvar ColorPrimvar = UsdMesh.GetDisplayColorPrimvar();
-		pxr::TfToken ColorInterpolation = pxr::UsdGeomTokens->constant;
-		pxr::VtArray< pxr::GfVec3f > UsdColors;
+				NumUVs = FMath::Min<int32>(USD_PREVIEW_SURFACE_MAX_UV_SETS, NumUVs);
+				// At least one UV set must exist.
+				NumUVs = FMath::Max<int32>(1, NumUVs);
 
-		if ( ColorPrimvar )
-		{
-			ColorPrimvar.ComputeFlattened( &UsdColors, Options.TimeCode );
-			ColorInterpolation = ColorPrimvar.GetInterpolation();
-		}
+				// Make sure all Vertex instance have the correct number of UVs
+				MeshDescriptionUVs.SetNumChannels(NumUVs);
+			}
 
-		// Vertex opacity
-		pxr::UsdGeomPrimvar OpacityPrimvar = UsdMesh.GetDisplayOpacityPrimvar();
-		pxr::TfToken OpacityInterpolation = pxr::UsdGeomTokens->constant;
-		pxr::VtArray< float > UsdOpacities;
+			TVertexInstanceAttributesRef<FVector3f> MeshDescriptionNormals = StaticMeshAttributes.GetVertexInstanceNormals();
+			TVertexInstanceAttributesRef<FVector3f>
+				MeshDescriptionVelocities = OutMeshDescription.VertexInstanceAttributes().GetAttributesRef<FVector3f>(
+					MeshAttribute::VertexInstance::Velocity
+				);
 
-		if ( OpacityPrimvar )
-		{
-			OpacityPrimvar.ComputeFlattened( &UsdOpacities );
-			OpacityInterpolation = OpacityPrimvar.GetInterpolation();
-		}
+			OutMeshDescription.ReserveNewVertexInstances(InMeshData.FaceVertexCounts.size() * 3);
+			OutMeshDescription.ReserveNewPolygons(InMeshData.FaceVertexCounts.size());
+			OutMeshDescription.ReserveNewEdges(InMeshData.FaceVertexCounts.size() * 2);
 
-		for ( int32 PolygonIndex = 0; PolygonIndex < FaceCounts.size(); ++PolygonIndex )
-		{
-			int32 PolygonVertexCount = FaceCounts[ PolygonIndex ];
-			CornerInstanceIDs.Reset( PolygonVertexCount );
-			CornerVerticesIDs.Reset( PolygonVertexCount );
+			// Vertex color
+			TVertexInstanceAttributesRef<FVector4f> MeshDescriptionColors = StaticMeshAttributes.GetVertexInstanceColors();
 
-			for ( int32 CornerIndex = 0; CornerIndex < PolygonVertexCount; ++CornerIndex, ++CurrentVertexInstanceIndex )
+			for (size_t PolygonIndex = 0; PolygonIndex < InMeshData.FaceVertexCounts.size(); ++PolygonIndex)
 			{
-				int32 VertexInstanceIndex = VertexInstanceOffset + CurrentVertexInstanceIndex;
-				const FVertexInstanceID VertexInstanceID( VertexInstanceIndex );
-				const int32 ControlPointIndex = FaceIndices[ CurrentVertexInstanceIndex ];
-				const FVertexID VertexID( VertexOffset + ControlPointIndex );
+				int32 PolygonVertexCount = InMeshData.FaceVertexCounts.cdata()[PolygonIndex];
+				CornerInstanceIDs.Reset(PolygonVertexCount);
+				CornerVerticesIDs.Reset(PolygonVertexCount);
 
-				// This data is read straight from USD so there's nothing guaranteeing we have as many positions as we need
-				if (VertexID.GetValue() >= MeshDescriptionVertexPositions.GetNumElements() || VertexID.GetValue() < 0)
+				for (int32 CornerIndex = 0; CornerIndex < PolygonVertexCount; ++CornerIndex, ++CurrentVertexInstanceIndex)
 				{
-					continue;
-				}
-
-				// Make sure a face doesn't use the same vertex twice as MeshDescription doesn't like that
-				if ( CornerVerticesIDs.Contains( VertexID ) )
-				{
-					continue;
-				}
+					int32 VertexInstanceIndex = VertexInstanceOffset + CurrentVertexInstanceIndex;
+					const FVertexInstanceID VertexInstanceID(VertexInstanceIndex);
+					const int32 ControlPointIndex = InMeshData.FaceIndices.cdata()[CurrentVertexInstanceIndex];
+					const FVertexID VertexID(VertexOffset + ControlPointIndex);
 
-				const FVector VertexPosition = (FVector)MeshDescriptionVertexPositions[VertexID];
+					// This data is read straight from USD so there's nothing guaranteeing we have as many positions as we need
+					if (VertexID.GetValue() >= MeshDescriptionVertexPositions.GetNumElements() || VertexID.GetValue() < 0)
+					{
+						continue;
+					}
 
-				CornerVerticesIDs.Add( VertexID );
+					// Make sure a face doesn't use the same vertex twice as MeshDescription doesn't like that
+					if (CornerVerticesIDs.Contains(VertexID))
+					{
+						continue;
+					}
 
-				FVertexInstanceID AddedVertexInstanceId = OutMeshDescription.CreateVertexInstance( VertexID );
-				CornerInstanceIDs.Add( AddedVertexInstanceId );
+					CornerVerticesIDs.Add(VertexID);
 
-				if ( Normals.size() > 0 )
-				{
-					const int32 NormalIndex = UsdGeomMeshImpl::GetPrimValueIndex( NormalsInterpType, ControlPointIndex, CurrentVertexInstanceIndex, PolygonIndex );
+					FVertexInstanceID AddedVertexInstanceId = OutMeshDescription.CreateVertexInstance(VertexID);
+					CornerInstanceIDs.Add(AddedVertexInstanceId);
 
-					if ( NormalIndex < Normals.size() )
+					if (InMeshData.Normals.size() > 0)
 					{
-						const pxr::GfVec3f& Normal = Normals[ NormalIndex ];
-						FVector TransformedNormal = Options.AdditionalTransform.TransformVector( UsdToUnreal::ConvertVector( StageInfo, Normal ) ).GetSafeNormal();
+						const size_t NormalIndex = GetPrimValueIndex(
+							InMeshData.NormalInterpolation,
+							ControlPointIndex,
+							CurrentVertexInstanceIndex,
+							PolygonIndex
+						);
 
-						MeshDescriptionNormals[ AddedVertexInstanceId ] = ( FVector3f ) TransformedNormal.GetSafeNormal();
-					}
-				}
+						if (NormalIndex < InMeshData.Normals.size())
+						{
+							const pxr::GfVec3f& Normal = InMeshData.Normals.cdata()[NormalIndex];
+							FVector TransformedNormal = InOptions.AdditionalTransform.TransformVector(UsdToUnreal::ConvertVector(InStageInfo, Normal))
+															.GetSafeNormal();
 
-				if (Velocities.size() > 0)
-				{
-					const int32 VelocityIndex = UsdGeomMeshImpl::GetPrimValueIndex(VelocitiesInterpType, ControlPointIndex, CurrentVertexInstanceIndex, PolygonIndex);
+							MeshDescriptionNormals[AddedVertexInstanceId] = (FVector3f)TransformedNormal.GetSafeNormal();
+						}
+					}
 
-					if (VelocityIndex < Velocities.size())
+					if (InMeshData.Velocities.size() > 0)
 					{
-						const pxr::GfVec3f& Velocity = Velocities[VelocityIndex];
-						FVector TransformedVelocity = Options.AdditionalTransform.TransformVector(UsdToUnreal::ConvertVector(StageInfo, Velocity));
-
-						MeshDescriptionVelocities[AddedVertexInstanceId] = (FVector3f) TransformedVelocity;
-					}
-				}
+						const size_t VelocityIndex = GetPrimValueIndex(
+							InMeshData.VelocityInterpolation,
+							ControlPointIndex,
+							CurrentVertexInstanceIndex,
+							PolygonIndex
+						);
 
-				for ( const FUVSet& UVSet : UVSets )
-				{
-					const int32 ValueIndex = UsdGeomMeshImpl::GetPrimValueIndex( UVSet.InterpType, ControlPointIndex, CurrentVertexInstanceIndex, PolygonIndex );
+						if (VelocityIndex < InMeshData.Velocities.size())
+						{
+							const pxr::GfVec3f& Velocity = InMeshData.Velocities.cdata()[VelocityIndex];
+							FVector TransformedVelocity = InOptions.AdditionalTransform.TransformVector(
+								UsdToUnreal::ConvertVector(InStageInfo, Velocity)
+							);
 
-					pxr::GfVec2f UV( 0.f, 0.f );
+							MeshDescriptionVelocities[AddedVertexInstanceId] = (FVector3f)TransformedVelocity;
+						}
+					}
 
-					if ( UVSet.UVIndices.IsSet() )
+					for (const FUVSet& UVSet : UVSets)
 					{
-						const pxr::VtIntArray& UVIndices = UVSet.UVIndices.GetValue();
+						const size_t ValueIndex = GetPrimValueIndex(UVSet.InterpType, ControlPointIndex, CurrentVertexInstanceIndex, PolygonIndex);
+
+						pxr::GfVec2f UV(0.f, 0.f);
 
-						if (ensure(ValueIndex < UVIndices.size()))
+						if (UVSet.UVIndices.IsSet())
 						{
-							int UVIndex = UVIndices[ValueIndex];
+							const pxr::VtIntArray& UVIndices = UVSet.UVIndices.GetValue();
 
-							if (ensure(UVIndex < UVSet.UVs.size()))
+							if (ensure(ValueIndex < UVIndices.size()))
 							{
-								UV = UVSet.UVs[UVIndex];
+								size_t UVIndex = UVIndices[ValueIndex];
+
+								if (ensure(UVIndex < UVSet.UVs.size()))
+								{
+									UV = UVSet.UVs[UVIndex];
+								}
 							}
 						}
+						else if (ensure(UVSet.UVs.size() > ValueIndex))
+						{
+							UV = UVSet.UVs[ValueIndex];
+						}
+
+						// Flip V for Unreal uv's which match directx
+						FVector2f FinalUVVector(UV[0], 1.f - UV[1]);
+						MeshDescriptionUVs.Set(AddedVertexInstanceId, UVSet.UVSetIndexUE, FinalUVVector);
 					}
-					else if ( ensure( UVSet.UVs.size() > ValueIndex ) )
+
+					// Vertex color
 					{
-						UV = UVSet.UVs[ ValueIndex ];
-					}
+						const size_t ValueIndex = GetPrimValueIndex(
+							InMeshData.DisplayColorInterpolation,
+							ControlPointIndex,
+							CurrentVertexInstanceIndex,
+							PolygonIndex
+						);
 
-					// Flip V for Unreal uv's which match directx
-					FVector2f FinalUVVector( UV[ 0 ], 1.f - UV[ 1 ] );
-					MeshDescriptionUVs.Set( AddedVertexInstanceId, UVSet.UVSetIndexUE, FinalUVVector );
-				}
+						pxr::GfVec3f UsdColor(1.f, 1.f, 1.f);
 
-				// Vertex color
-				{
-					const int32 ValueIndex = UsdGeomMeshImpl::GetPrimValueIndex( ColorInterpolation, ControlPointIndex, CurrentVertexInstanceIndex, PolygonIndex );
+						if (!InMeshData.DisplayColors.empty() && ensure(InMeshData.DisplayColors.size() > ValueIndex))
+						{
+							UsdColor = InMeshData.DisplayColors.cdata()[ValueIndex];
+						}
 
-					pxr::GfVec3f UsdColor( 1.f, 1.f, 1.f );
+						MeshDescriptionColors[AddedVertexInstanceId] = UsdToUnreal::ConvertColor(UsdColor);
+					}
 
-					if ( !UsdColors.empty() && ensure( UsdColors.size() > ValueIndex ) )
+					// Vertex opacity
 					{
-						UsdColor = UsdColors[ ValueIndex ];
+						const size_t ValueIndex = GetPrimValueIndex(
+							InMeshData.DisplayOpacityInterpolation,
+							ControlPointIndex,
+							CurrentVertexInstanceIndex,
+							PolygonIndex
+						);
+
+						if (!InMeshData.DisplayOpacities.empty() && ensure(InMeshData.DisplayOpacities.size() > ValueIndex))
+						{
+							MeshDescriptionColors[AddedVertexInstanceId][3] = InMeshData.DisplayOpacities.cdata()[ValueIndex];
+						}
 					}
+				}
 
-					MeshDescriptionColors[ AddedVertexInstanceId ] = UsdToUnreal::ConvertColor( UsdColor );
+				// This polygon was using the same vertex instance more than once and we removed too many
+				// vertex indices, so now we're forced to skip the whole polygon. We'll show a warning about it though
+				if (CornerVerticesIDs.Num() < 3)
+				{
+					++NumSkippedPolygons;
+					continue;
 				}
 
-				// Vertex opacity
+				// Polygon groups
+				int32 LocalMaterialIndex = 0;
+				if (FaceMaterialIndices.IsValidIndex(PolygonIndex))
 				{
-					const int32 ValueIndex = UsdGeomMeshImpl::GetPrimValueIndex( OpacityInterpolation, ControlPointIndex, CurrentVertexInstanceIndex, PolygonIndex );
+					LocalMaterialIndex = FaceMaterialIndices[PolygonIndex];
+					if (!LocalMaterialSlots.IsValidIndex(LocalMaterialIndex))
+					{
+						LocalMaterialIndex = 0;
+					}
+				}
+
+				const int32 CombinedMaterialIndex = LocalToCombinedMaterialSlotIndices[LocalMaterialIndex];
 
-					if ( !UsdOpacities.empty() && ensure( UsdOpacities.size() > ValueIndex ) )
+				// Flip geometry if needed
+				if (InMeshData.Orientation == pxr::UsdGeomTokens->leftHanded)
+				{
+					for (int32 i = 0; i < CornerInstanceIDs.Num() / 2; ++i)
 					{
-						MeshDescriptionColors[ AddedVertexInstanceId ][ 3 ] = UsdOpacities[ ValueIndex ];
+						Swap(CornerInstanceIDs[i], CornerInstanceIDs[CornerInstanceIDs.Num() - i - 1]);
 					}
 				}
-			}
 
-			// This polygon was using the same vertex instance more than once and we removed too many
-			// vertex indices, so now we're forced to skip the whole polygon. We'll show a warning about it though
-			if ( CornerVerticesIDs.Num() < 3 )
-			{
-				++NumSkippedPolygons;
-				continue;
+				// Insert a polygon into the mesh
+				FPolygonGroupID PolygonGroupID = PolygonGroupMapping[CombinedMaterialIndex];
+				OutMeshDescription.CreatePolygon(PolygonGroupID, CornerInstanceIDs);
 			}
+		}
+
+		if (NumPolygons > 0 && NumSkippedPolygons > 0)
+		{
+			UE_LOG(
+				LogUsd,
+				Warning,
+				TEXT("Skipped %d out of %d faces when parsing the mesh for prim '%s', as those faces contained too many repeated vertex indices"),
+				NumSkippedPolygons,
+				NumPolygons,
+				*InMeshData.SourcePrimPath
+			);
+		}
+
+		return true;
+	}
+}
+namespace UsdGeomMeshImpl = UE::UsdGeomMeshConversion::Private;
+
+namespace UsdToUnreal
+{
+	const FUsdMeshConversionOptions FUsdMeshConversionOptions::DefaultOptions;
+
+	FUsdMeshConversionOptions::FUsdMeshConversionOptions()
+		: AdditionalTransform( FTransform::Identity )
+		, PurposesToLoad( EUsdPurpose::Render )
+		, RenderContext( pxr::UsdShadeTokens->universalRenderContext )
+		, MaterialPurpose( pxr::UsdShadeTokens->allPurpose )
+		, TimeCode( pxr::UsdTimeCode::EarliestTime() )
+		, bMergeIdenticalMaterialSlots( true )
+	{
+	}
+}
+
+bool UsdToUnreal::ConvertGeomMesh(
+	const pxr::UsdGeomMesh& UsdMesh,
+	FMeshDescription& OutMeshDescription,
+	UsdUtils::FUsdPrimMaterialAssignmentInfo& OutMaterialAssignments,
+	const FUsdMeshConversionOptions& Options
+)
+{
+	TRACE_CPUPROFILER_EVENT_SCOPE( UsdToUnreal::ConvertGeomMesh );
+
+	if ( !UsdMesh )
+	{
+		return false;
+	}
+
+	FScopedUsdAllocs Allocs;
+
+	pxr::UsdPrim UsdPrim = UsdMesh.GetPrim();
+	pxr::UsdStageRefPtr Stage = UsdPrim.GetStage();
+	const FUsdStageInfo StageInfo(Stage);
+
+	UsdGeomMeshImpl::FUsdMeshData MeshData;
+	MeshData.SourcePrimPath = UsdToUnreal::ConvertPath(UsdPrim.GetPrimPath());
+
+	// Face counts
+	if (pxr::UsdAttribute Attr = UsdMesh.GetFaceVertexCountsAttr())
+	{
+		Attr.Get(&MeshData.FaceVertexCounts, Options.TimeCode);
+	}
+
+	// Face indices
+	if (pxr::UsdAttribute Attr = UsdMesh.GetFaceVertexIndicesAttr())
+	{
+		Attr.Get(&MeshData.FaceIndices, Options.TimeCode);
+	}
+
+	// Points
+	if (pxr::UsdAttribute Attr = UsdMesh.GetPointsAttr())
+	{
+		Attr.Get(&MeshData.Points, Options.TimeCode);
+	}
+
+	// Normals
+	if (pxr::UsdAttribute Attr = UsdMesh.GetNormalsAttr())
+	{
+		Attr.Get(&MeshData.Normals, Options.TimeCode);
+		MeshData.NormalInterpolation = UsdGeomMeshImpl::GetAttrInterpolation(Attr, pxr::UsdGeomTokens->vertex);
+	}
+
+	// Velocities
+	if (pxr::UsdAttribute Attr = UsdMesh.GetVelocitiesAttr())
+	{
+		Attr.Get(&MeshData.Velocities, Options.TimeCode);
+		MeshData.VelocityInterpolation = UsdGeomMeshImpl::GetAttrInterpolation(Attr, pxr::UsdGeomTokens->vertex);
+	}
 
-			// Polygon groups
-			int32 LocalMaterialIndex = 0;
-			if ( FaceMaterialIndices.IsValidIndex( PolygonIndex ) )
-			{
-				LocalMaterialIndex = FaceMaterialIndices[ PolygonIndex ];
-				if ( !LocalMaterialSlots.IsValidIndex( LocalMaterialIndex ) )
-				{
-					LocalMaterialIndex = 0;
-				}
-			}
+	// Vertex colors
+	if (pxr::UsdGeomPrimvar Primvar = pxr::UsdGeomPrimvar(UsdPrim.GetAttribute(pxr::UsdGeomTokens->primvarsDisplayColor)))
+	{
+		Primvar.ComputeFlattened(&MeshData.DisplayColors, Options.TimeCode);
+		MeshData.DisplayColorInterpolation = UsdGeomMeshImpl::GetAttrInterpolation(Primvar);
+	}
 
-			const int32 CombinedMaterialIndex = LocalToCombinedMaterialSlotIndices[ LocalMaterialIndex ];
+	// Vertex opacity
+	if (pxr::UsdGeomPrimvar Primvar = pxr::UsdGeomPrimvar(UsdPrim.GetAttribute(pxr::UsdGeomTokens->primvarsDisplayOpacity)))
+	{
+		Primvar.ComputeFlattened(&MeshData.DisplayOpacities, Options.TimeCode);
+		MeshData.DisplayOpacityInterpolation = UsdGeomMeshImpl::GetAttrInterpolation(Primvar);
+	}
 
-			if ( bFlipThisGeometry )
+	// UVs
+	{
+		// If we already have a primvar to UV index assignment, let's just use that.
+		// When collapsing, we'll do a pre-pass on all meshes to translate and determine this beforehand.
+		if (OutMaterialAssignments.PrimvarToUVIndex.Num() > 0)
+		{
+			int32 HighestProvidedUVIndex = 0;
+			for (const TPair<FString, int32>& Pair : OutMaterialAssignments.PrimvarToUVIndex)
 			{
-				for ( int32 i = 0; i < CornerInstanceIDs.Num() / 2; ++i )
-				{
-					Swap( CornerInstanceIDs[ i ], CornerInstanceIDs[ CornerInstanceIDs.Num() - i - 1 ] );
-				}
+				HighestProvidedUVIndex = FMath::Max(HighestProvidedUVIndex, Pair.Value);
 			}
-			// Insert a polygon into the mesh
-			FPolygonGroupID PolygonGroupID = PolygonGroupMapping[ CombinedMaterialIndex ];
-			const FPolygonID NewPolygonID = OutMeshDescription.CreatePolygon( PolygonGroupID, CornerInstanceIDs );
+			MeshData.ProvidedNumUVSets = HighestProvidedUVIndex + 1;
+
+			TArray<TUsdStore<pxr::UsdGeomPrimvar>> AllMeshUVPrimvars = UsdUtils::GetUVSetPrimvars(UsdPrim, TNumericLimits<int32>::Max());
+			MeshData.PrimvarsByUVIndex = UsdUtils::AssemblePrimvarsIntoUVSets(AllMeshUVPrimvars, OutMaterialAssignments.PrimvarToUVIndex);
+		}
+		// Let's use the best primvar assignment for this particular mesh instead
+		else
+		{
+			MeshData.PrimvarsByUVIndex = UsdUtils::GetUVSetPrimvars(UsdPrim);
+			OutMaterialAssignments.PrimvarToUVIndex = UsdUtils::AssemblePrimvarsIntoPrimvarToUVIndexMap(MeshData.PrimvarsByUVIndex);
 		}
 	}
 
-	if ( NumPolygons > 0 && NumSkippedPolygons > 0 )
+	// Orientation
+	MeshData.Orientation = UsdGeomMeshImpl::GetGprimOrientation(UsdMesh, Options.TimeCode);
+
+	// Material assignments
 	{
-		UE_LOG( LogUsd, Warning, TEXT( "Skipped %d out of %d faces when parsing the mesh for prim '%s', as those faces contained too many repeated vertex indices" ),
-			NumSkippedPolygons,
-			NumPolygons,
-			*UsdToUnreal::ConvertPath( UsdPrim.GetPath() )
+		const bool bProvideMaterialIndices = true;
+		MeshData.LocalMaterialInfo = UsdUtils::GetPrimMaterialAssignments(
+			UsdPrim,
+			Options.TimeCode,
+			bProvideMaterialIndices,
+			Options.RenderContext,
+			Options.MaterialPurpose
 		);
+
+		MeshData.MaterialIndexOffset = OutMaterialAssignments.Slots.Num();
 	}
 
-	return true;
+	return UsdGeomMeshImpl::ConvertMeshData(MeshData, StageInfo, Options, OutMeshDescription, OutMaterialAssignments);
 }
 
 bool UsdToUnreal::ConvertPointInstancerToMesh(
@@ -1542,6 +1787,376 @@ bool UsdToUnreal::ConvertGeomMeshHierarchy(
 	);
 }
 
+bool UsdToUnreal::ConvertGeomPrimitive(
+	const pxr::UsdPrim& InPrim,
+	FMeshDescription& InOutMeshDescription,
+	UsdUtils::FUsdPrimMaterialAssignmentInfo& InOutMaterialAssignments,
+	const FUsdMeshConversionOptions& InOptions
+)
+{
+	TRACE_CPUPROFILER_EVENT_SCOPE(UsdToUnreal::ConvertGeomPrimitive);
+
+	if (!InPrim)
+	{
+		return false;
+	}
+
+	FScopedUsdAllocs UsdAllocs;
+
+	pxr::UsdStageRefPtr Stage = InPrim.GetStage();
+	const FUsdStageInfo StageInfo(Stage);
+
+	// Remember that USD arrays are copy-on-write, so these are both "pointers", as long as we
+	// don't try writing (or using non-const operator[]) from PrimitivePoints
+	pxr::VtVec3fArray PrimitivePoints;
+	const pxr::PxOsdMeshTopology* PrimitiveTopology = nullptr;
+
+	if (pxr::UsdGeomCapsule Capsule = pxr::UsdGeomCapsule{InPrim})
+	{
+		pxr::TfToken Axis = pxr::UsdGeomTokens->z;
+		if (pxr::UsdAttribute Attr = Capsule.GetAxisAttr())
+		{
+			Axis = UsdUtils::GetUsdValue<pxr::TfToken>(Attr, pxr::UsdTimeCode::Default());
+		}
+
+		PrimitivePoints = pxr::UsdImagingGenerateCapsuleMeshPoints(
+			UsdGeomMeshImpl::DefaultCapsuleMeshHeight,
+			UsdGeomMeshImpl::DefaultCapsuleMeshRadius,
+			Axis
+		);
+		PrimitiveTopology = &pxr::UsdImagingGetCapsuleMeshTopology();
+	}
+	else if (pxr::UsdGeomCone Cone = pxr::UsdGeomCone{InPrim})
+	{
+		pxr::TfToken Axis = pxr::UsdGeomTokens->z;
+		if (pxr::UsdAttribute Attr = Capsule.GetAxisAttr())
+		{
+			Axis = UsdUtils::GetUsdValue<pxr::TfToken>(Attr, pxr::UsdTimeCode::Default());
+		}
+
+		PrimitivePoints = UsdGeomMeshImpl::GetUnitConeMeshPoints(Axis);
+		PrimitiveTopology = &pxr::UsdImagingGetUnitConeMeshTopology();
+	}
+	else if (pxr::UsdGeomCube Cube = pxr::UsdGeomCube{InPrim})
+	{
+		PrimitivePoints = pxr::UsdImagingGetUnitCubeMeshPoints();
+		PrimitiveTopology = &pxr::UsdImagingGetUnitCubeMeshTopology();
+	}
+	else if (pxr::UsdGeomCylinder Cylinder = pxr::UsdGeomCylinder{InPrim})
+	{
+		pxr::TfToken Axis = pxr::UsdGeomTokens->z;
+		if (pxr::UsdAttribute Attr = Capsule.GetAxisAttr())
+		{
+			Axis = UsdUtils::GetUsdValue<pxr::TfToken>(Attr, pxr::UsdTimeCode::Default());
+		}
+
+		PrimitivePoints = UsdGeomMeshImpl::GetUnitCylinderMeshPoints(Axis);
+		PrimitiveTopology = &pxr::UsdImagingGetUnitCylinderMeshTopology();
+	}
+	else if (pxr::UsdGeomPlane Plane = pxr::UsdGeomPlane{InPrim})
+	{
+		const double Width = 1.0f;
+		const double Length = 1.0f;
+
+		pxr::TfToken Axis = pxr::UsdGeomTokens->z;
+		if (pxr::UsdAttribute Attr = Capsule.GetAxisAttr())
+		{
+			Axis = UsdUtils::GetUsdValue<pxr::TfToken>(Attr, pxr::UsdTimeCode::Default());
+		}
+
+		PrimitivePoints = pxr::UsdImagingGeneratePlaneMeshPoints(Width, Length, Axis);
+		PrimitiveTopology = &pxr::UsdImagingGetPlaneTopology();
+	}
+	else if (pxr::UsdGeomSphere Sphere = pxr::UsdGeomSphere{InPrim})
+	{
+		PrimitivePoints = pxr::UsdImagingGetUnitSphereMeshPoints();
+		PrimitiveTopology = &pxr::UsdImagingGetUnitSphereMeshTopology();
+	}
+
+	if (!PrimitiveTopology || PrimitivePoints.empty())
+	{
+		return false;
+	}
+
+	UsdGeomMeshImpl::FUsdMeshData MeshData;
+	MeshData.SourcePrimPath = UsdToUnreal::ConvertPath(InPrim.GetPrimPath());
+	MeshData.FaceVertexCounts = PrimitiveTopology->GetFaceVertexCounts();
+	MeshData.FaceIndices = PrimitiveTopology->GetFaceVertexIndices();
+	MeshData.Points = PrimitivePoints;
+
+	// Normals
+	if (pxr::UsdAttribute Attr = InPrim.GetAttribute(pxr::UsdGeomTokens->normals))
+	{
+		Attr.Get(&MeshData.Normals, InOptions.TimeCode);
+		MeshData.NormalInterpolation = UsdGeomMeshImpl::GetAttrInterpolation(Attr, pxr::UsdGeomTokens->vertex);
+	}
+
+	// Velocities
+	if (pxr::UsdGeomPointBased PointBased{InPrim})
+	{
+		if (pxr::UsdAttribute Attr = PointBased.GetVelocitiesAttr())
+		{
+			Attr.Get(&MeshData.Velocities, InOptions.TimeCode);
+			MeshData.VelocityInterpolation = UsdGeomMeshImpl::GetAttrInterpolation(Attr, pxr::UsdGeomTokens->vertex);
+		}
+	}
+
+	// Vertex colors
+	if (pxr::UsdGeomPrimvar Primvar = pxr::UsdGeomPrimvar(InPrim.GetAttribute(pxr::UsdGeomTokens->primvarsDisplayColor)))
+	{
+		Primvar.ComputeFlattened(&MeshData.DisplayColors, InOptions.TimeCode);
+		MeshData.DisplayColorInterpolation = UsdGeomMeshImpl::GetAttrInterpolation(Primvar);
+	}
+
+	// Vertex opacity
+	if (pxr::UsdGeomPrimvar Primvar = pxr::UsdGeomPrimvar(InPrim.GetAttribute(pxr::UsdGeomTokens->primvarsDisplayOpacity)))
+	{
+		Primvar.ComputeFlattened(&MeshData.DisplayOpacities, InOptions.TimeCode);
+		MeshData.DisplayOpacityInterpolation = UsdGeomMeshImpl::GetAttrInterpolation(Primvar);
+	}
+
+	// UVs
+	{
+		// If we already have a primvar to UV index assignment, let's just use that.
+		// When collapsing, we'll do a pre-pass on all meshes to translate and determine this beforehand.
+		if (InOutMaterialAssignments.PrimvarToUVIndex.Num() > 0)
+		{
+			int32 HighestProvidedUVIndex = 0;
+			for (const TPair<FString, int32>& Pair : InOutMaterialAssignments.PrimvarToUVIndex)
+			{
+				HighestProvidedUVIndex = FMath::Max(HighestProvidedUVIndex, Pair.Value);
+			}
+			MeshData.ProvidedNumUVSets = HighestProvidedUVIndex + 1;
+
+			TArray<TUsdStore<pxr::UsdGeomPrimvar>> AllMeshUVPrimvars = UsdUtils::GetUVSetPrimvars(InPrim, TNumericLimits<int32>::Max());
+			MeshData.PrimvarsByUVIndex = UsdUtils::AssemblePrimvarsIntoUVSets(AllMeshUVPrimvars, InOutMaterialAssignments.PrimvarToUVIndex);
+		}
+		// Let's use the best primvar assignment for this particular mesh instead
+		else
+		{
+			MeshData.PrimvarsByUVIndex = UsdUtils::GetUVSetPrimvars(InPrim);
+			InOutMaterialAssignments.PrimvarToUVIndex = UsdUtils::AssemblePrimvarsIntoPrimvarToUVIndexMap(MeshData.PrimvarsByUVIndex);
+		}
+	}
+
+	// Orientation
+	MeshData.Orientation = UsdGeomMeshImpl::GetGprimOrientation(pxr::UsdGeomGprim{InPrim}, InOptions.TimeCode);
+
+	// Material assignments
+	{
+		const bool bProvideMaterialIndices = true;
+		MeshData.LocalMaterialInfo = UsdUtils::GetPrimMaterialAssignments(
+			InPrim,
+			InOptions.TimeCode,
+			bProvideMaterialIndices,
+			InOptions.RenderContext,
+			InOptions.MaterialPurpose
+		);
+
+		MeshData.MaterialIndexOffset = InOutMaterialAssignments.Slots.Num();
+	}
+
+	return ConvertMeshData(MeshData, StageInfo, InOptions, InOutMeshDescription, InOutMaterialAssignments);
+}
+
+bool UsdToUnreal::ConvertGeomPrimitiveTransform(
+	const pxr::UsdPrim& InPrim,
+	const pxr::UsdTimeCode& InTimeCode,
+	FTransform& OutTransform
+)
+{
+	TRACE_CPUPROFILER_EVENT_SCOPE(UsdToUnreal::ConvertGeomPrimitive);
+
+	if (!InPrim || !InPrim.IsA<pxr::UsdGeomGprim>())
+	{
+		return false;
+	}
+
+	FScopedUsdAllocs UsdAllocs;
+
+	pxr::UsdStageRefPtr Stage = InPrim.GetStage();
+	const FUsdStageInfo StageInfo(Stage);
+
+	TFunction<pxr::GfMatrix4d(pxr::TfToken, double, double)> GetScalingTransform =
+		[](pxr::TfToken Axis, double Longitudinal, double Transversal)
+	{
+		if (Axis == pxr::UsdGeomTokens->x)
+		{
+			return pxr::GfMatrix4d(
+				Longitudinal, 		  0.0, 		   0.0, 0.0,
+						 0.0, Transversal, 		   0.0, 0.0,
+						 0.0, 		  0.0, Transversal, 0.0,
+						 0.0, 		  0.0, 		   0.0, 1.0
+			);
+		}
+		else if (Axis == pxr::UsdGeomTokens->y)
+		{
+			return pxr::GfMatrix4d(
+				Transversal, 		  0.0, 		   0.0, 0.0,
+						0.0, Longitudinal, 		   0.0, 0.0,
+						0.0, 		  0.0, Transversal, 0.0,
+						0.0, 		  0.0, 		   0.0, 1.0
+			);
+		}
+		else
+		{
+			return pxr::GfMatrix4d(
+				Transversal, 		 0.0, 		   0.0, 0.0,
+						0.0, Transversal, 		   0.0, 0.0,
+						0.0, 		 0.0, Longitudinal, 0.0,
+						0.0, 		 0.0, 		   0.0, 1.0
+			);
+		}
+	};
+
+	if (pxr::UsdGeomCapsule Capsule = pxr::UsdGeomCapsule{InPrim})
+	{
+		const double Radius = UsdUtils::GetUsdValue<double>(Capsule.GetRadiusAttr(), InTimeCode);
+		const double Height = UsdUtils::GetUsdValue<double>(Capsule.GetHeightAttr(), InTimeCode);
+
+		pxr::TfToken Axis = pxr::UsdGeomTokens->z;
+		if (pxr::UsdAttribute Attr = Capsule.GetAxisAttr())
+		{
+			Axis = UsdUtils::GetUsdValue<pxr::TfToken>(Attr, pxr::UsdTimeCode::Default());
+			if (Attr.ValueMightBeTimeVarying())
+			{
+				UE_LOG(
+					LogUsd,
+					Warning,
+					TEXT("Animation of the 'axis' attribute for prim '%s' is not supported!"),
+					*UsdToUnreal::ConvertPath(InPrim.GetPrimPath())
+				);
+			}
+		}
+
+		// We use these "Scaling" factors instead of direct height/radius because we're assuming
+		// we'll have generated this capsule mesh using ConvertGeomPrimitive, where we provide
+		// UsdImagingGenerateCapsuleMeshPoints with DefaultCapsuleMeshHeight and
+		// DefaultCapsuleMeshRadius. If our current height/radius match those, we need to create
+		// an identity transform. If our height is twice as that, our axis direction needs to have
+		// a scaling of 2.0, etc.
+		// Also keep in mind that the capsule total height is (Radius + Height + Radius).
+		double HeightScaling = (Height + 2.0f * Radius)
+							   / (UsdGeomMeshImpl::DefaultCapsuleMeshHeight + 2.0f * UsdGeomMeshImpl::DefaultCapsuleMeshRadius);
+		double RadiusScaling = (Radius / UsdGeomMeshImpl::DefaultCapsuleMeshRadius);
+		pxr::GfMatrix4d PrimitiveTransform = GetScalingTransform(Axis, HeightScaling, RadiusScaling);
+
+		OutTransform = UsdToUnreal::ConvertMatrix(StageInfo, PrimitiveTransform);
+		return true;
+	}
+	else if (pxr::UsdGeomCone Cone = pxr::UsdGeomCone{InPrim})
+	{
+		const double Radius = UsdUtils::GetUsdValue<double>(Cone.GetRadiusAttr(), InTimeCode);
+		const double Height = UsdUtils::GetUsdValue<double>(Cone.GetHeightAttr(), InTimeCode);
+
+		pxr::TfToken Axis = pxr::UsdGeomTokens->z;
+		if (pxr::UsdAttribute Attr = Cone.GetAxisAttr())
+		{
+			Axis = UsdUtils::GetUsdValue<pxr::TfToken>(Attr, pxr::UsdTimeCode::Default());
+			if (Attr.ValueMightBeTimeVarying())
+			{
+				UE_LOG(
+					LogUsd,
+					Warning,
+					TEXT("Animation of the 'axis' attribute for prim '%s' is not supported!"),
+					*UsdToUnreal::ConvertPath(InPrim.GetPrimPath())
+				);
+			}
+		}
+
+		const double Diameter = 2.0 * Radius;
+		pxr::GfMatrix4d PrimitiveTransform = GetScalingTransform(Axis, Height, Diameter);
+
+		OutTransform = UsdToUnreal::ConvertMatrix(StageInfo, PrimitiveTransform);
+		return true;
+	}
+	else if (pxr::UsdGeomCube Cube = pxr::UsdGeomCube{InPrim})
+	{
+		const double Size = UsdUtils::GetUsdValue<double>(Cube.GetSizeAttr(), InTimeCode);
+		const pxr::GfMatrix4d UsdTransform = pxr::UsdImagingGenerateSphereOrCubeTransform(Size);
+
+		OutTransform = UsdToUnreal::ConvertMatrix(StageInfo, UsdTransform);
+		return true;
+	}
+	else if (pxr::UsdGeomCylinder Cylinder = pxr::UsdGeomCylinder{InPrim})
+	{
+		const double Radius = UsdUtils::GetUsdValue<double>(Cone.GetRadiusAttr(), InTimeCode);
+		const double Height = UsdUtils::GetUsdValue<double>(Cone.GetHeightAttr(), InTimeCode);
+
+		pxr::TfToken Axis = pxr::UsdGeomTokens->z;
+		if (pxr::UsdAttribute Attr = Cone.GetAxisAttr())
+		{
+			Axis = UsdUtils::GetUsdValue<pxr::TfToken>(Attr, pxr::UsdTimeCode::Default());
+			if (Attr.ValueMightBeTimeVarying())
+			{
+				UE_LOG(
+					LogUsd,
+					Warning,
+					TEXT("Animation of the 'axis' attribute for prim '%s' is not supported!"),
+					*UsdToUnreal::ConvertPath(InPrim.GetPrimPath())
+				);
+			}
+		}
+
+		const double Diameter = 2.0 * Radius;
+		pxr::GfMatrix4d PrimitiveTransform = GetScalingTransform(Axis, Height, Diameter);
+
+		OutTransform = UsdToUnreal::ConvertMatrix(StageInfo, PrimitiveTransform);
+		return true;
+	}
+	else if (pxr::UsdGeomPlane Plane = pxr::UsdGeomPlane{InPrim})
+	{
+		const double Width = UsdUtils::GetUsdValue<double>(Plane.GetWidthAttr(), InTimeCode);
+		const double Length = UsdUtils::GetUsdValue<double>(Plane.GetLengthAttr(), InTimeCode);
+
+		pxr::TfToken Axis = pxr::UsdGeomTokens->z;
+		if (pxr::UsdAttribute Attr = Plane.GetAxisAttr())
+		{
+			Axis = UsdUtils::GetUsdValue<pxr::TfToken>(Attr, pxr::UsdTimeCode::Default());
+			if (Attr.ValueMightBeTimeVarying())
+			{
+				UE_LOG(
+					LogUsd,
+					Warning,
+					TEXT("Animation of the 'axis' attribute for prim '%s' is not supported!"),
+					*UsdToUnreal::ConvertPath(InPrim.GetPrimPath())
+				);
+			}
+		}
+
+		OutTransform = FTransform::Identity;
+
+		// Generate a scaling transform in USD coordinate system
+		if (Axis == pxr::UsdGeomTokens->x)
+		{
+			OutTransform.SetScale3D(FVector{1.0f, Length, Width});
+		}
+		else if (Axis == pxr::UsdGeomTokens->y)
+		{
+			OutTransform.SetScale3D(FVector{Width, 1.0f, Length});
+		}
+		else if (Axis == pxr::UsdGeomTokens->z)
+		{
+			OutTransform.SetScale3D(FVector{Width, Length, 1.0f});
+		}
+
+		// Convert that transform to the UE coordinate system
+		OutTransform = UsdUtils::ConvertAxes(StageInfo.UpAxis == EUsdUpAxis::ZAxis, OutTransform);
+		return true;
+	}
+	else if (pxr::UsdGeomSphere Sphere = pxr::UsdGeomSphere{InPrim})
+	{
+		const double Radius = UsdUtils::GetUsdValue<double>(Sphere.GetRadiusAttr(), InTimeCode);
+		const double Diameter = Radius * 2.0;
+		const pxr::GfMatrix4d UsdTransform = pxr::UsdImagingGenerateSphereOrCubeTransform(Diameter);
+
+		OutTransform = UsdToUnreal::ConvertMatrix(StageInfo, UsdTransform);
+		return true;
+	}
+
+	return false;
+}
+
 UMaterialInstanceDynamic* UsdUtils::CreateDisplayColorMaterialInstanceDynamic( const UsdUtils::FDisplayColorMaterial& DisplayColorDescription )
 {
 	const UUsdProjectSettings* Settings = GetDefault<UUsdProjectSettings>();
@@ -1728,15 +2343,7 @@ UsdUtils::FUsdPrimMaterialAssignmentInfo UsdUtils::GetPrimMaterialAssignments(
 
 	uint64 NumFaces = 0;
 	{
-		pxr::UsdGeomMesh Mesh = pxr::UsdGeomMesh( UsdPrim );
-		pxr::UsdAttribute FaceCounts = Mesh.GetFaceVertexCountsAttr();
-		if ( !Mesh || !FaceCounts )
-		{
-			return Result;
-		}
-
-		pxr::VtArray<int> FaceVertexCounts;
-		FaceCounts.Get( &FaceVertexCounts, TimeCode );
+		pxr::VtArray<int> FaceVertexCounts = UsdGeomMeshImpl::GetFaceVertexCounts(UsdPrim, TimeCode);
 		NumFaces = FaceVertexCounts.size();
 		if ( NumFaces < 1 )
 		{
@@ -1950,7 +2557,7 @@ UsdUtils::FUsdPrimMaterialAssignmentInfo UsdUtils::GetPrimMaterialAssignments(
 		// for displayColor and displayOpacity
 		if (!bHasMainAssignment)
 		{
-			DisplayColor = ExtractDisplayColorMaterial(pxr::UsdGeomMesh(UsdPrim), TimeCode);
+			DisplayColor = ExtractDisplayColorMaterial(pxr::UsdGeomGprim{UsdPrim}, TimeCode);
 			if (DisplayColor)
 			{
 				FUsdPrimMaterialSlot& Slot = Result.Slots.Emplace_GetRef();
@@ -2346,14 +2953,14 @@ TOptional<UsdUtils::FDisplayColorMaterial> UsdUtils::FDisplayColorMaterial::From
 	return Result;
 }
 
-TOptional<UsdUtils::FDisplayColorMaterial> UsdUtils::ExtractDisplayColorMaterial( const pxr::UsdGeomMesh& UsdMesh, const pxr::UsdTimeCode TimeCode )
+TOptional<UsdUtils::FDisplayColorMaterial> UsdUtils::ExtractDisplayColorMaterial(const pxr::UsdGeomGprim& Gprim, const pxr::UsdTimeCode TimeCode)
 {
-	if ( !UsdMesh )
+	if (!Gprim)
 	{
 		return {};
 	}
 
-	if ( !UsdMesh.GetDisplayOpacityAttr().IsDefined() && !UsdMesh.GetDisplayColorAttr().IsDefined() )
+	if (!Gprim.GetDisplayOpacityAttr().IsDefined() && !Gprim.GetDisplayColorAttr().IsDefined())
 	{
 		return {};
 	}
@@ -2361,20 +2968,20 @@ TOptional<UsdUtils::FDisplayColorMaterial> UsdUtils::ExtractDisplayColorMaterial
 	UsdUtils::FDisplayColorMaterial Desc;
 
 	// Opacity
-	pxr::VtArray< float > UsdOpacities = UsdUtils::GetUsdValue< pxr::VtArray< float > >( UsdMesh.GetDisplayOpacityAttr(), TimeCode );
-	for ( float Opacity : UsdOpacities )
+	pxr::VtArray<float> UsdOpacities = UsdUtils::GetUsdValue<pxr::VtArray<float>>(Gprim.GetDisplayOpacityAttr(), TimeCode);
+	for (float Opacity : UsdOpacities)
 	{
-		Desc.bHasOpacity = !FMath::IsNearlyEqual( Opacity, 1.f );
-		if ( Desc.bHasOpacity )
+		Desc.bHasOpacity = !FMath::IsNearlyEqual(Opacity, 1.f);
+		if (Desc.bHasOpacity)
 		{
 			break;
 		}
 	}
 
 	// Double-sided
-	if ( UsdMesh.GetDoubleSidedAttr().IsDefined() )
+	if (Gprim.GetDoubleSidedAttr().IsDefined())
 	{
-		Desc.bIsDoubleSided = UsdUtils::GetUsdValue< bool >( UsdMesh.GetDoubleSidedAttr(), TimeCode );
+		Desc.bIsDoubleSided = UsdUtils::GetUsdValue<bool>(Gprim.GetDoubleSidedAttr(), TimeCode);
 	}
 
 	return Desc;
@@ -2951,7 +3558,7 @@ FString UsdUtils::HashGeomMeshPrim( const UE::FUsdStage& Stage, const FString& P
 
 	// TODO: This is not providing render context or material purpose, so it will never consider float2f primvars
 	// for the hash, which could be an issue in very exotic cases
-	TArray< TUsdStore< UsdGeomPrimvar > > PrimvarsByUVIndex = UsdUtils::GetUVSetPrimvars( UsdMesh );
+	TArray< TUsdStore< UsdGeomPrimvar > > PrimvarsByUVIndex = UsdUtils::GetUVSetPrimvars( UsdPrim );
 	for ( int32 UVChannelIndex = 0; UVChannelIndex < PrimvarsByUVIndex.Num(); ++UVChannelIndex )
 	{
 		if ( !PrimvarsByUVIndex.IsValidIndex( UVChannelIndex ) )
@@ -3107,6 +3714,126 @@ UsdUtils::EMeshTopologyVariance UsdUtils::GetMeshTopologyVariance(const pxr::Usd
 	return EMeshTopologyVariance::Constant;
 }
 
+uint64 UsdUtils::GetGprimVertexCount(const pxr::UsdGeomGprim& Gprim, double TimeCode)
+{
+	if (pxr::UsdGeomMesh Mesh{Gprim})
+	{
+		if (pxr::UsdAttribute Points = Mesh.GetPointsAttr())
+		{
+			pxr::VtArray<pxr::GfVec3f> PointsArray;
+			Points.Get(&PointsArray, pxr::UsdTimeCode(TimeCode));
+			return PointsArray.size();
+		}
+	}
+	else if (pxr::UsdGeomCapsule Capsule{Gprim})
+	{
+		// These numbers come from inspecting USD's implicitSurfaceMeshUtils.cpp
+		// and comparing with the generated UStaticMesh vertex counts.
+		// In practice it doesn't matter much though: These small Gprims are likely
+		// never going to significantly affect whether a subtree should collapse or not
+		return 82;
+	}
+	else if (pxr::UsdGeomCone Cone{Gprim})
+	{
+		return 31;
+	}
+	else if (pxr::UsdGeomCube Cube{Gprim})
+	{
+		return 8;
+	}
+	else if (pxr::UsdGeomCylinder Cylinder{Gprim})
+	{
+		return 42;
+	}
+	else if (pxr::UsdGeomSphere Sphere{Gprim})
+	{
+		return 92;
+	}
+	else if (pxr::UsdGeomPlane Plane{Gprim})
+	{
+		return 4;
+	}
+
+	return 0;
+}
+
+void UsdUtils::AuthorIdentityTransformGprimAttributes(const pxr::UsdPrim& UsdPrim, bool bDefaultValues, bool bTimeSampleValues)
+{
+	pxr::UsdGeomGprim Gprim{UsdPrim};
+	if (!Gprim)
+	{
+		return;
+	}
+
+	FScopedUsdAllocs Allocs;
+
+	// We can't just "clear" these opinions because we may cause some weaker opinion
+	// to pop up, and the caller will likely be relying on this function to make sure our
+	// prim has attributes in such a way that its "primitive transform" is the identity.
+	// In other words, after we call this function on UsdPrim, calling
+	// ConvertGeomPrimitiveTransform on the same prim should generate the identity transform.
+	auto SetAttrValue = [bDefaultValues, bTimeSampleValues](const pxr::UsdAttribute& Attr, auto Value)
+	{
+		if (!Attr)
+		{
+			return;
+		}
+
+		if (bDefaultValues)
+		{
+			Attr.Set(Value);
+		}
+
+		if (bTimeSampleValues)
+		{
+			UsdUtils::ClearAllTimeSamples(Attr);
+
+			// Ideally we'd use pxr::UsdTimeCode::EarliestTime() but that seems to be -DBL_MAX, which
+			// could look weird to a user when written on the USD file. Since this is going to be the
+			// only timeSample it doesn't really matter anyway
+			const pxr::UsdTimeCode TimeCode{0.0};
+			Attr.Set(Value, TimeCode);
+		}
+	};
+
+	// In here we must author the attribute values that cause ConvertGeomPrimitive
+	// to generate meshes in the [-0.5, 0.5] bounding box, as that will correspond to the
+	// identity "primitive transform".
+	// Note that these values *do not* correspond to the fallback values for the attributes.
+	// For whatever reason the attribute fallback values all lead to a scaling factor of 2 instead.
+	// If we want our primitives to end up with a scale of 2 when writing out to USD however,
+	// we will put the scale of 2 directly on the Xform/component transform instead, and with this
+	// function have the attributes generate a scale of 1 instead.
+	if (pxr::UsdGeomCapsule Capsule = pxr::UsdGeomCapsule{UsdPrim})
+	{
+		SetAttrValue(Capsule.CreateRadiusAttr(), 0.25);
+		SetAttrValue(Capsule.CreateHeightAttr(), 0.5);
+	}
+	else if (pxr::UsdGeomCone Cone = pxr::UsdGeomCone{UsdPrim})
+	{
+		SetAttrValue(Cone.CreateRadiusAttr(), 0.5);
+		SetAttrValue(Cone.CreateHeightAttr(), 1.0);
+	}
+	else if (pxr::UsdGeomCube Cube = pxr::UsdGeomCube{UsdPrim})
+	{
+		SetAttrValue(Cube.CreateSizeAttr(), 1.0);
+	}
+	else if (pxr::UsdGeomCylinder Cylinder = pxr::UsdGeomCylinder{UsdPrim})
+	{
+		SetAttrValue(Cylinder.CreateRadiusAttr(), 0.5);
+		SetAttrValue(Cylinder.CreateHeightAttr(), 1.0);
+	}
+	else if (pxr::UsdGeomSphere Sphere = pxr::UsdGeomSphere{UsdPrim})
+	{
+		SetAttrValue(Sphere.CreateRadiusAttr(), 0.5);
+	}
+	else if (pxr::UsdGeomPlane Plane = pxr::UsdGeomPlane{UsdPrim})
+	{
+		SetAttrValue(Plane.CreateWidthAttr(), 1.0);
+		SetAttrValue(Plane.CreateLengthAttr(), 1.0);
+	}
+}
+
 #undef LOCTEXT_NAMESPACE
 
 #endif // #if USE_USD_SDK
diff --git a/Engine/Plugins/Importers/USDImporter/Source/USDUtilities/Private/USDPrimConversion.cpp b/Engine/Plugins/Importers/USDImporter/Source/USDUtilities/Private/USDPrimConversion.cpp
--- a/Engine/Plugins/Importers/USDImporter/Source/USDUtilities/Private/USDPrimConversion.cpp
+++ b/Engine/Plugins/Importers/USDImporter/Source/USDUtilities/Private/USDPrimConversion.cpp
@@ -57,36 +57,40 @@
 #if USE_USD_SDK
 
 #include "USDIncludesStart.h"
-
-#include "pxr/usd/sdf/changeBlock.h"
-#include "pxr/usd/usd/attribute.h"
-#include "pxr/usd/usd/prim.h"
-#include "pxr/usd/usd/stage.h"
-#include "pxr/usd/usd/timeCode.h"
-#include "pxr/usd/usdGeom/camera.h"
-#include "pxr/usd/usdGeom/imageable.h"
-#include "pxr/usd/usdGeom/mesh.h"
-#include "pxr/usd/usdGeom/pointInstancer.h"
-#include "pxr/usd/usdGeom/scope.h"
-#include "pxr/usd/usdGeom/tokens.h"
-#include "pxr/usd/usdGeom/xform.h"
-#include "pxr/usd/usdGeom/xformable.h"
-#include "pxr/usd/usdGeom/xformCommonAPI.h"
-#include "pxr/usd/usdLux/diskLight.h"
-#include "pxr/usd/usdLux/distantLight.h"
-#include "pxr/usd/usdLux/lightAPI.h"
-#include "pxr/usd/usdLux/rectLight.h"
-#include "pxr/usd/usdLux/shapingAPI.h"
-#include "pxr/usd/usdLux/sphereLight.h"
-#include "pxr/usd/usdLux/tokens.h"
-#include "pxr/usd/usdShade/connectableAPI.h"
-#include "pxr/usd/usdShade/material.h"
-#include "pxr/usd/usdShade/materialBindingAPI.h"
-#include "pxr/usd/usdShade/shader.h"
-#include "pxr/usd/usdShade/tokens.h"
-#include "pxr/usd/usdSkel/animation.h"
-#include "pxr/usd/usdSkel/root.h"
-
+	#include "pxr/usd/sdf/changeBlock.h"
+	#include "pxr/usd/usd/attribute.h"
+	#include "pxr/usd/usd/prim.h"
+	#include "pxr/usd/usd/stage.h"
+	#include "pxr/usd/usd/timeCode.h"
+	#include "pxr/usd/usdGeom/camera.h"
+	#include "pxr/usd/usdGeom/capsule.h"
+	#include "pxr/usd/usdGeom/cone.h"
+	#include "pxr/usd/usdGeom/cube.h"
+	#include "pxr/usd/usdGeom/cylinder.h"
+	#include "pxr/usd/usdGeom/imageable.h"
+	#include "pxr/usd/usdGeom/mesh.h"
+	#include "pxr/usd/usdGeom/plane.h"
+	#include "pxr/usd/usdGeom/pointInstancer.h"
+	#include "pxr/usd/usdGeom/scope.h"
+	#include "pxr/usd/usdGeom/sphere.h"
+	#include "pxr/usd/usdGeom/tokens.h"
+	#include "pxr/usd/usdGeom/xform.h"
+	#include "pxr/usd/usdGeom/xformable.h"
+	#include "pxr/usd/usdGeom/xformCommonAPI.h"
+	#include "pxr/usd/usdLux/diskLight.h"
+	#include "pxr/usd/usdLux/distantLight.h"
+	#include "pxr/usd/usdLux/lightAPI.h"
+	#include "pxr/usd/usdLux/rectLight.h"
+	#include "pxr/usd/usdLux/shapingAPI.h"
+	#include "pxr/usd/usdLux/sphereLight.h"
+	#include "pxr/usd/usdLux/tokens.h"
+	#include "pxr/usd/usdShade/connectableAPI.h"
+	#include "pxr/usd/usdShade/material.h"
+	#include "pxr/usd/usdShade/materialBindingAPI.h"
+	#include "pxr/usd/usdShade/shader.h"
+	#include "pxr/usd/usdShade/tokens.h"
+	#include "pxr/usd/usdSkel/animation.h"
+	#include "pxr/usd/usdSkel/root.h"
 #include "USDIncludesEnd.h"
 
 namespace UE
@@ -336,6 +340,11 @@ bool UsdToUnreal::ConvertXformable( const pxr::UsdStageRefPtr& Stage, const pxr:
 
 	FScopedUsdAllocs UsdAllocs;
 
+	OutTransform = FTransform::Identity;
+
+	// If we're a primitive try extracting its transform as well, given that we'll always reuse the default, 1.0 size procedural meshes
+	UsdToUnreal::ConvertGeomPrimitiveTransform(Xformable.GetPrim(), EvalTime, OutTransform);
+
 	// Transform
 	pxr::GfMatrix4d UsdMatrix;
 	bool bResetXformStack = false;
@@ -343,7 +352,7 @@ bool UsdToUnreal::ConvertXformable( const pxr::UsdStageRefPtr& Stage, const pxr:
 	Xformable.GetLocalTransformation( &UsdMatrix, bResetXformStackPtr, EvalTime );
 
 	FUsdStageInfo StageInfo( Stage );
-	OutTransform = UsdToUnreal::ConvertMatrix( StageInfo, UsdMatrix );
+	OutTransform = OutTransform * UsdToUnreal::ConvertMatrix( StageInfo, UsdMatrix );
 
 	const bool bPrimIsLight = Xformable.GetPrim().HasAPI< pxr::UsdLuxLightAPI >();
 
@@ -3788,15 +3797,52 @@ TArray<UE::FUsdAttribute> UnrealToUsd::GetAttributesForProperty( const UE::FUsdP
 	// Common attributes
 	if ( PropertyPath == TransformPropertyName )
 	{
-		if ( pxr::UsdAttribute Attr = UsdPrim.GetAttribute( UnrealToUsd::ConvertToken( TEXT( "xformOp:transform" ) ).Get() ) )
+		TArray<UE::FUsdAttribute> Attrs;
+
+		if (pxr::UsdGeomXformable Xformable{UsdPrim})
 		{
-			return { UE::FUsdAttribute{ Attr } };
+			Attrs.Emplace(Xformable.GetXformOpOrderAttr());
+
+			bool bResetsXformStack = false;
+			std::vector<pxr::UsdGeomXformOp> Ops = Xformable.GetOrderedXformOps(&bResetsXformStack);
+			for (const pxr::UsdGeomXformOp& Op : Ops)
+			{
+				Attrs.Emplace(Op.GetAttr());
+			}
 		}
 
-		if ( pxr::UsdGeomXformable Xformable{ UsdPrim } )
+		// This function returns all attributes that can affect a property. For the Gprim primitives and the Transform property
+		// this will include their heights, widths, etc. too, as we also handle those with just the component transform.
+		if (pxr::UsdGeomCapsule Capsule{UsdPrim})
+		{
+			Attrs.Emplace(Capsule.GetHeightAttr());
+			Attrs.Emplace(Capsule.GetRadiusAttr());
+		}
+		else if (pxr::UsdGeomCone Cone{UsdPrim})
 		{
-			return { UE::FUsdAttribute{ Xformable.GetXformOpOrderAttr() } };
+			Attrs.Emplace(Cone.GetHeightAttr());
+			Attrs.Emplace(Cone.GetRadiusAttr());
 		}
+		else if (pxr::UsdGeomCube Cube{UsdPrim})
+		{
+			Attrs.Emplace(Cube.GetSizeAttr());
+		}
+		else if (pxr::UsdGeomCylinder Cylinder{UsdPrim})
+		{
+			Attrs.Emplace(Cylinder.GetHeightAttr());
+			Attrs.Emplace(Cylinder.GetRadiusAttr());
+		}
+		else if (pxr::UsdGeomPlane Plane{UsdPrim})
+		{
+			Attrs.Emplace(Plane.GetLengthAttr());
+			Attrs.Emplace(Plane.GetWidthAttr());
+		}
+		else if (pxr::UsdGeomSphere Sphere{UsdPrim})
+		{
+			Attrs.Emplace(Sphere.GetRadiusAttr());
+		}
+
+		return Attrs;
 	}
 	if ( PropertyPath == HiddenInGamePropertyName )
 	{
diff --git a/Engine/Plugins/Importers/USDImporter/Source/USDUtilities/Private/USDSkeletalDataConversion.cpp b/Engine/Plugins/Importers/USDImporter/Source/USDUtilities/Private/USDSkeletalDataConversion.cpp
--- a/Engine/Plugins/Importers/USDImporter/Source/USDUtilities/Private/USDSkeletalDataConversion.cpp
+++ b/Engine/Plugins/Importers/USDImporter/Source/USDUtilities/Private/USDSkeletalDataConversion.cpp
@@ -1645,7 +1645,7 @@ bool UsdToUnreal::ConvertSkinnedMesh(
 	if (MaterialAssignments.PrimvarToUVIndex.Num() > 0)
 	{
 		TArray<TUsdStore<pxr::UsdGeomPrimvar>> AllMeshUVPrimvars =
-			UsdUtils::GetUVSetPrimvars(UsdMesh, TNumericLimits<int32>::Max());
+			UsdUtils::GetUVSetPrimvars(SkinningPrim, TNumericLimits<int32>::Max());
 
 		PrimvarsByUVIndex =
 			UsdUtils::AssemblePrimvarsIntoUVSets(AllMeshUVPrimvars, MaterialAssignments.PrimvarToUVIndex);
@@ -1653,7 +1653,7 @@ bool UsdToUnreal::ConvertSkinnedMesh(
 	// Let's use the best primvar assignment for this particular mesh instead
 	else
 	{
-		PrimvarsByUVIndex = UsdUtils::GetUVSetPrimvars(UsdMesh);
+		PrimvarsByUVIndex = UsdUtils::GetUVSetPrimvars(SkinningPrim);
 
 		MaterialAssignments.PrimvarToUVIndex = UsdUtils::AssemblePrimvarsIntoPrimvarToUVIndexMap(PrimvarsByUVIndex);
 	}
diff --git a/Engine/Plugins/Importers/USDImporter/Source/USDUtilities/Public/USDAttributeUtils.h b/Engine/Plugins/Importers/USDImporter/Source/USDUtilities/Public/USDAttributeUtils.h
--- a/Engine/Plugins/Importers/USDImporter/Source/USDUtilities/Public/USDAttributeUtils.h
+++ b/Engine/Plugins/Importers/USDImporter/Source/USDUtilities/Public/USDAttributeUtils.h
@@ -8,6 +8,7 @@
 
 #if USE_USD_SDK
 PXR_NAMESPACE_OPEN_SCOPE
+	class UsdAttribute;
 	class UsdPrim;
 	class UsdProperty;
 PXR_NAMESPACE_CLOSE_SCOPE
@@ -30,6 +31,13 @@ namespace UsdUtils
 	/** Returns whether this attribute has the "Muted" CustomData on the stage's UE state sublayer, meaning it shouldn't be animated when loaded into UE */
 	USDUTILITIES_API bool IsAttributeMuted( const UE::FUsdAttribute& Attribute, const UE::FUsdStage& Stage );
 
+#if USE_USD_SDK
+	/**
+	 * Calls pxr::UsdAttribute::GetTimeSamples(), and then pxr::UsdAttribute::ClearAtTime() for each timeSample value, clearing
+	 * all timeSample opinions at the current edit target.
+	 */
+	USDUTILITIES_API bool ClearAllTimeSamples(const pxr::UsdAttribute& Attribute);
+
 	/**
 	 * If Attribute has opinions authored on layers stronger than the current edit target this will emit a warning to
 	 * the user, indicating that the new value may not be visible on the USD Stage, and may cease to be visible on
@@ -38,7 +46,6 @@ namespace UsdUtils
 	 * Use this after setting any attribute: This function does nothing in case the Stage's current edit target has
 	 * the strongest opinion for the attribute already
 	 */
-#if USE_USD_SDK
 	USDUTILITIES_API void NotifyIfOverriddenOpinion( const pxr::UsdProperty& Property );
 
 	/**
diff --git a/Engine/Plugins/Importers/USDImporter/Source/USDUtilities/Public/USDConversionUtils.h b/Engine/Plugins/Importers/USDImporter/Source/USDUtilities/Public/USDConversionUtils.h
--- a/Engine/Plugins/Importers/USDImporter/Source/USDUtilities/Public/USDConversionUtils.h
+++ b/Engine/Plugins/Importers/USDImporter/Source/USDUtilities/Public/USDConversionUtils.h
@@ -141,19 +141,19 @@ namespace UsdUtils
 	USDUTILITIES_API int32 GetPrimvarUVIndex( FString PrimvarName );
 
 	/**
-	 * Gets the primvars that should be used as UV sets, per index, for this UsdMesh.
+	 * Gets the primvars that should be used as UV sets, per index, for this UsdPrim.
 	 * This will return between 0 and MaxNumPrimvars primvars, where the first item of array is the primvar that should be used for
 	 * UV index 0 (if any), the second the primvar that should be used for UV set 1, etc.
 	 *
 	 * In case there are more than MaxNumPrimvars available primvars, the returned list will contain all the 'texCoord2f' role
 	 * primvars (sorted lexicographically), followed by the regular float2 primvars (also separately sorted lexicographically).
 	 *
-	 * @param UsdMesh - Mesh that contains primvars that can be used as texture coordinates.
+	 * @param UsdPrim - UsdPrim that contains primvars that can be used as texture coordinates.
 	 * @param MaxNumPrimvars - Maximum number of primvars to return from this function.
 	 * @return Array of up to MaxNumPrimvars primvars sorted by priority (most important come first)
 	 */
 	USDUTILITIES_API TArray<TUsdStore<pxr::UsdGeomPrimvar>> GetUVSetPrimvars(
-		const pxr::UsdGeomMesh& UsdMesh,
+		const pxr::UsdPrim& UsdPrim,
 		int32 MaxNumPrimvars = USD_PREVIEW_SURFACE_MAX_UV_SETS
 	);
 
diff --git a/Engine/Plugins/Importers/USDImporter/Source/USDUtilities/Public/USDGeomMeshConversion.h b/Engine/Plugins/Importers/USDImporter/Source/USDUtilities/Public/USDGeomMeshConversion.h
--- a/Engine/Plugins/Importers/USDImporter/Source/USDUtilities/Public/USDGeomMeshConversion.h
+++ b/Engine/Plugins/Importers/USDImporter/Source/USDUtilities/Public/USDGeomMeshConversion.h
@@ -17,6 +17,7 @@
 PXR_NAMESPACE_OPEN_SCOPE
 	class TfToken;
 	class UsdGeomMesh;
+	class UsdGeomGprim;
 	class UsdGeomPointInstancer;
 	class UsdPrim;
 	class UsdShadeMaterial;
@@ -146,6 +147,46 @@ namespace UsdToUnreal
 		const FUsdMeshConversionOptions& CommonOptions = FUsdMeshConversionOptions::DefaultOptions,
 		bool bSkipRootPrimTransformAndVisibility = false
 	);
+
+	/**
+	 * Converts an UsdGeomCapsule, UsdGeomCone, UsdGeomCube, UsdGeomCylinder, UsdGeomPlane or UsdGeomSphere and appends
+	 * the converted data to InOutMeshDescription and InOutMaterialAssignments.
+	 *
+	 * Note that all generated meshes will have extents within the [-0.5, 0.5] box centered on the origin.
+	 * You can pass additional transforms on CommonOptions.AdditionalTransform though.
+	 *
+	 * @param Prim - USD Prim prim to convert
+	 * @param InOutMeshDescription - Input/Output parameter that will be filled with the converted data
+	 * @param InOutMaterialAssignments - Input/Output parameter that will be filled with the material data extracted from Mesh
+	 * @param CommonOptions - Shared options used for the conversion
+	 * @return Whether the conversion was successful or not.
+	 */
+	USDUTILITIES_API bool ConvertGeomPrimitive(
+		const pxr::UsdPrim& Prim,
+		FMeshDescription& InOutMeshDescription,
+		UsdUtils::FUsdPrimMaterialAssignmentInfo& InOutMaterialAssignments,
+		const FUsdMeshConversionOptions& CommonOptions = FUsdMeshConversionOptions::DefaultOptions
+	);
+
+	/**
+	 * Extracts a "primitive transform" from the provided UsdGeomCapsule, UsdGeomCone, UsdGeomCube, UsdGeomCylinder,
+	 * UsdGeomPlane or UsdGeomSphere and places the converted, UE coordinate system transform on OutTransform.
+	 *
+	 * This is intended to be applied to the primitive data generated by ConvertGeomPrimitive, and concatenated with
+	 * any actual Xform on the prim (Xform applied after). The idea is that the meshes generated there are standard
+	 * size and fill the [-0.5, 0.5] box centered on the origin (length 1). If a Cylinder has "height=2.5", this function
+	 * will produce a scaling of 2.5 on the longitudinal axis for that cylinder.
+	 *
+	 * @param Prim - USD Prim prim to convert
+	 * @param InTimeCode - TimeCode to sample the
+	 * @param OutTransform - The converted transform on the UE coordinate system
+	 * @return Whether the conversion was successful or not.
+	 */
+	USDUTILITIES_API bool ConvertGeomPrimitiveTransform(
+		const pxr::UsdPrim& InPrim,
+		const pxr::UsdTimeCode& InTimeCode,
+		FTransform& OutTransform
+	);
 }
 
 namespace UnrealToUsd
@@ -252,13 +293,13 @@ namespace UsdUtils
 		/** Describes which primvars should be assigned to each UV index for this mesh. */
 		TMap<FString, int32> PrimvarToUVIndex;
 	};
-}
 
 #if USE_USD_SDK
-namespace UsdUtils
-{
-	/** Creates a FDisplayColorMaterial object describing the vertex color/opacity data from UsdMesh at time TimeCode */
-	USDUTILITIES_API TOptional<FDisplayColorMaterial> ExtractDisplayColorMaterial( const pxr::UsdGeomMesh& UsdMesh, const pxr::UsdTimeCode TimeCode = pxr::UsdTimeCode::EarliestTime() );
+	/** Creates a FDisplayColorMaterial object describing the vertex color/opacity data from Gprim at time TimeCode */
+	USDUTILITIES_API TOptional<FDisplayColorMaterial> ExtractDisplayColorMaterial(
+		const pxr::UsdGeomGprim& Gprim,
+		const pxr::UsdTimeCode TimeCode = pxr::UsdTimeCode::EarliestTime()
+	);
 
 	/** Creates a dynamic material instance using the right reference material depending on the given description */
 	USDUTILITIES_API UMaterialInstanceDynamic* CreateDisplayColorMaterialInstanceDynamic( const UsdUtils::FDisplayColorMaterial& DisplayColorDescription );
@@ -361,6 +402,10 @@ namespace UsdUtils
 	/** Returns the type of topology variance UsdMesh has */
 	USDUTILITIES_API EMeshTopologyVariance GetMeshTopologyVariance(const pxr::UsdGeomMesh& UsdMesh);
 
-}
+	USDUTILITIES_API uint64 GetGprimVertexCount(const pxr::UsdGeomGprim& Gprim, double TimeCode);
+
+	USDUTILITIES_API void AuthorIdentityTransformGprimAttributes(const pxr::UsdPrim& Gprim, bool bDefaultValues, bool bTimeSampleValues);
 
 #endif // #if USE_USD_SDK
+}
+
