USD: Added support for headless mode (commandlet) that prevents initializing render resources for StaticMesh.

diff --git a/Engine/Plugins/Importers/USDImporter/Source/USDClasses/Private/USDAssetCache2.cpp b/Engine/Plugins/Importers/USDImporter/Source/USDClasses/Private/USDAssetCache2.cpp
--- a/Engine/Plugins/Importers/USDImporter/Source/USDClasses/Private/USDAssetCache2.cpp
+++ b/Engine/Plugins/Importers/USDImporter/Source/USDClasses/Private/USDAssetCache2.cpp
@@ -14,6 +14,7 @@
 #include "Materials/MaterialInterface.h"
 #include "MaterialShared.h"
 #include "MeshDescription.h"
+#include "Misc/App.h"
 #include "Misc/ScopeRWLock.h"
 #include "PhysicsEngine/PhysicsAsset.h"
 #include "Serialization/ObjectAndNameAsStringProxyArchive.h"
@@ -518,12 +519,15 @@ namespace UE::AssetCache::Private
 			// We only need these in the editor: When we're cooking we'll just save the cooked mesh data into the
 			// asset itself, so at runtime there is no need for an additional "build" step that we have to handle here
 
-			StaticMesh->SetRenderData(MakeUnique< FStaticMeshRenderData >());
-			ITargetPlatformManagerModule& TargetPlatformManager = GetTargetPlatformManagerRef();
-			ITargetPlatform* RunningPlatform = TargetPlatformManager.GetRunningTargetPlatform();
-			const FStaticMeshLODSettings& LODSettings = RunningPlatform->GetStaticMeshLODSettings();
-			StaticMesh->GetRenderData()->Cache(RunningPlatform, StaticMesh, LODSettings);
-			StaticMesh->InitResources();
+			if (FApp::CanEverRender())
+			{
+				StaticMesh->SetRenderData(MakeUnique< FStaticMeshRenderData >());
+				ITargetPlatformManagerModule& TargetPlatformManager = GetTargetPlatformManagerRef();
+				ITargetPlatform* RunningPlatform = TargetPlatformManager.GetRunningTargetPlatform();
+				const FStaticMeshLODSettings& LODSettings = RunningPlatform->GetStaticMeshLODSettings();
+				StaticMesh->GetRenderData()->Cache(RunningPlatform, StaticMesh, LODSettings);
+				StaticMesh->InitResources();
+			}
 		}
 		else if (USkeletalMesh* SkeletalMesh = Cast<USkeletalMesh>(Object))
 		{
diff --git a/Engine/Plugins/Importers/USDImporter/Source/USDSchemas/Private/USDGeomMeshTranslator.cpp b/Engine/Plugins/Importers/USDImporter/Source/USDSchemas/Private/USDGeomMeshTranslator.cpp
--- a/Engine/Plugins/Importers/USDImporter/Source/USDSchemas/Private/USDGeomMeshTranslator.cpp
+++ b/Engine/Plugins/Importers/USDImporter/Source/USDSchemas/Private/USDGeomMeshTranslator.cpp
@@ -33,6 +33,7 @@
 #include "Materials/MaterialInstanceConstant.h"
 #include "Materials/MaterialInstanceDynamic.h"
 #include "Materials/MaterialInterface.h"
+#include "Misc/App.h"
 #include "Misc/SecureHash.h"
 #include "Modules/ModuleManager.h"
 #include "PhysicsEngine/BodySetup.h"
@@ -723,6 +724,11 @@ namespace UsdGeomMeshTranslatorImpl
 		// For runtime builds, the analogue for this stuff is already done from within BuildFromMeshDescriptions
 		TRACE_CPUPROFILER_EVENT_SCOPE(UsdGeomMeshTranslatorImpl::PostBuildStaticMesh);
 
+		if (!FApp::CanEverRender())
+		{
+			return;
+		}
+
 		StaticMesh.InitResources();
 
 #if WITH_EDITOR
diff --git a/Engine/Plugins/Importers/USDImporter/Source/USDSchemas/Private/USDGeomXformableTranslator.cpp b/Engine/Plugins/Importers/USDImporter/Source/USDSchemas/Private/USDGeomXformableTranslator.cpp
--- a/Engine/Plugins/Importers/USDImporter/Source/USDSchemas/Private/USDGeomXformableTranslator.cpp
+++ b/Engine/Plugins/Importers/USDImporter/Source/USDSchemas/Private/USDGeomXformableTranslator.cpp
@@ -27,6 +27,7 @@
 #include "Engine/World.h"
 #include "LiveLinkComponentController.h"
 #include "LiveLinkRole.h"
+#include "Misc/App.h"
 #include "Modules/ModuleManager.h"
 #include "Roles/LiveLinkTransformRole.h"
 #include "StaticMeshAttributes.h"
@@ -616,7 +617,7 @@ void FUsdGeomXformableTranslator::UpdateComponents( USceneComponent* SceneCompon
 			{
 				// Need to make sure the mesh's resources are initialized here as it may have just been built in another thread
 				// Only do this if required though, as this mesh could using these resources currently (e.g. PIE and editor world sharing the mesh)
-				if ( PrimStaticMesh && !PrimStaticMesh->AreRenderingResourcesInitialized() )
+				if ( PrimStaticMesh && !PrimStaticMesh->AreRenderingResourcesInitialized() && FApp::CanEverRender() )
 				{
 					PrimStaticMesh->InitResources();
 				}
