USD: Fix texture AssetUserData being Outer'd to materials by mistake.

#rnx

diff --git a/Engine/Plugins/Importers/USDImporter/Source/USDSchemas/Private/USDShadeMaterialTranslator.cpp b/Engine/Plugins/Importers/USDImporter/Source/USDSchemas/Private/USDShadeMaterialTranslator.cpp
--- a/Engine/Plugins/Importers/USDImporter/Source/USDSchemas/Private/USDShadeMaterialTranslator.cpp
+++ b/Engine/Plugins/Importers/USDImporter/Source/USDSchemas/Private/USDShadeMaterialTranslator.cpp
@@ -539,7 +539,7 @@ void FUsdShadeMaterialTranslator::PostImportMaterial(const FString& MaterialHash
 				UUsdAssetUserData* TextureUserData = Texture->GetAssetUserData<UUsdAssetUserData>();
 				if (!TextureUserData)
 				{
-					TextureUserData = NewObject<UUsdAssetUserData>(ImportedMaterial, TEXT("USDAssetUserData"));
+					TextureUserData = NewObject<UUsdAssetUserData>(Texture, TEXT("USDAssetUserData"));
 					Texture->AddAssetUserData(TextureUserData);
 				}
 				TextureUserData->PrimPaths.AddUnique(PrimPath.GetString());
