USD: Fix missing call to parent PostEditChangeProperty.

#rnx

diff --git a/Engine/Plugins/Importers/USDImporter/Source/USDClasses/Private/USDAssetCache2.cpp b/Engine/Plugins/Importers/USDImporter/Source/USDClasses/Private/USDAssetCache2.cpp
--- a/Engine/Plugins/Importers/USDImporter/Source/USDClasses/Private/USDAssetCache2.cpp
+++ b/Engine/Plugins/Importers/USDImporter/Source/USDClasses/Private/USDAssetCache2.cpp
@@ -630,6 +630,8 @@ void UUsdAssetCache2::PostEditChangeProperty(FPropertyChangedEvent& PropertyChan
 	{
 		RefreshStorage();
 	}
+
+	Super::PostEditChangeProperty(PropertyChangedEvent);
 }
 #endif // #if WITH_EDITOR
 
