USD: Add support for subdividing meshes via OpenSubdiv and pxOsd.

diff --git a/Engine/Plugins/Importers/USDImporter/Source/GeometryCacheUSD/GeometryCacheUSD.build.cs b/Engine/Plugins/Importers/USDImporter/Source/GeometryCacheUSD/GeometryCacheUSD.build.cs
--- a/Engine/Plugins/Importers/USDImporter/Source/GeometryCacheUSD/GeometryCacheUSD.build.cs
+++ b/Engine/Plugins/Importers/USDImporter/Source/GeometryCacheUSD/GeometryCacheUSD.build.cs
@@ -14,15 +14,15 @@ public class GeometryCacheUSD : ModuleRules
 				"GeometryCache",
 				"GeometryCacheStreamer",
 				"RenderCore",
-				"RHI"
+				"RHI",
+				"UnrealUSDWrapper",
+				"USDUtilities"
 			}
 		);
 
 		PrivateDependencyModuleNames.AddRange(
 			new string[] {
-				"UnrealUSDWrapper",
 				"USDClasses",
-				"USDUtilities"
             }
         );
 
diff --git a/Engine/Plugins/Importers/USDImporter/Source/GeometryCacheUSD/Private/GeometryCacheUSDStream.cpp b/Engine/Plugins/Importers/USDImporter/Source/GeometryCacheUSD/Private/GeometryCacheUSDStream.cpp
--- a/Engine/Plugins/Importers/USDImporter/Source/GeometryCacheUSD/Private/GeometryCacheUSDStream.cpp
+++ b/Engine/Plugins/Importers/USDImporter/Source/GeometryCacheUSD/Private/GeometryCacheUSDStream.cpp
@@ -1,6 +1,9 @@
 // Copyright Epic Games, Inc. All Rights Reserved.
 
 #include "GeometryCacheUSDStream.h"
+
+#include "USDGeomMeshConversion.h"
+
 #include "DerivedDataCacheInterface.h"
 #include "GeometryCacheMeshData.h"
 #include "GeometryCacheTrackUSD.h"
@@ -8,7 +11,6 @@
 #include "Misc/CoreMisc.h"
 #include "Serialization/MemoryReader.h"
 #include "Serialization/MemoryWriter.h"
-#include "USDGeomMeshConversion.h"
 
 static bool GUsdStreamCacheInDDC = true;
 static FAutoConsoleVariableRef CVarUsdStreamCacheInDDC(
@@ -22,7 +24,7 @@ static int32 kUsdReadConcurrency = 10;
 // differences, etc.) replace the version GUID below with a new one.
 // In case of merge conflicts with DDC versions, you MUST generate a new GUID
 // and set this new GUID as the version.
-#define USDSTREAM_DERIVED_DATA_VERSION TEXT("AB2B7CC003C54AEBBCC5ABDC1B0BFFD8")
+#define USDSTREAM_DERIVED_DATA_VERSION TEXT("BB2B7CC023C54AEBBCC5ABDC1B0BFFD2")
 
 class FUsdStreamDDCUtils
 {
@@ -39,18 +41,41 @@ private:
 	}
 
 public:
-	static FString GetUsdStreamDDCKey(const UE::FUsdStage& Stage, const FString& PrimPath, int32 FrameIndex)
-	{
 #if USE_USD_SDK
-		FString PrimHash = UsdUtils::HashGeomMeshPrim(Stage, PrimPath, FrameIndex);
-
-		if (!PrimHash.IsEmpty())
+	static FString GetUsdStreamDDCKey(
+		const UE::FUsdStage& Stage,
+		const FString& PrimPath,
+		int32 FrameIndex,
+		const UsdToUnreal::FUsdMeshConversionOptions& MeshConversionOptions
+	)
+	{
+		FString Hash;
 		{
-			return BuildDerivedDataKey(PrimHash);
+			FMD5 MD5;
+
+			UsdUtils::HashGeomMeshPrim(Stage, PrimPath, FrameIndex, MD5);
+
+			// Also hash these options because they may affect how the mesh conversion functions end up parsing
+			// the mesh data. This may seem a bit awkward, but it would be weirder still to provide MeshConversionOptions
+			// to HashGeomMeshPrim directly, as it already has a TimeCode parameter.
+			// To remove that and just use the TimeCode within the FUsdMeshConversionOptions instead would mean we would
+			// need to take a local copy of MeshConversionOptions to set FrameIndex as its TimeCode in here (as we can't
+			// modify MeshConversionOptions directly as it's owned by the UsdTrack and we may be called from a thread)
+			MD5.Update((uint8*)&MeshConversionOptions.SubdivisionLevel, sizeof(MeshConversionOptions.SubdivisionLevel));
+			MD5.Update((uint8*)&MeshConversionOptions.PurposesToLoad, sizeof(MeshConversionOptions.PurposesToLoad));
+
+			uint8 Digest[16];
+			MD5.Final(Digest);
+
+			for (int32 i = 0; i < 16; ++i)
+			{
+				Hash += FString::Printf(TEXT("%02x"), Digest[i]);
+			}
 		}
-#endif
-		return {};
+
+		return BuildDerivedDataKey(Hash);
 	}
+#endif
 };
 
 FGeometryCacheUsdStream::FGeometryCacheUsdStream(UGeometryCacheTrackUsd* InUsdTrack, FReadUsdMeshFunction InReadFunc)
@@ -127,7 +152,12 @@ void FGeometryCacheUsdStream::GetMeshData(int32 FrameIndex, int32 ConcurrencyInd
 	if (GUsdStreamCacheInDDC)
 	{
 		const FString& UsdPrimPath = UsdTrack->PrimPath;
-		const FString DerivedDataKey = FUsdStreamDDCUtils::GetUsdStreamDDCKey(UsdTrack->CurrentStagePinned, UsdPrimPath, FrameIndex);
+		const FString DerivedDataKey = FUsdStreamDDCUtils::GetUsdStreamDDCKey(
+			UsdTrack->CurrentStagePinned,
+			UsdPrimPath,
+			FrameIndex,
+			UsdTrack->MeshConversionOptions
+		);
 
 		if (!DerivedDataKey.IsEmpty())
 		{
diff --git a/Engine/Plugins/Importers/USDImporter/Source/GeometryCacheUSD/Public/GeometryCacheTrackUSD.h b/Engine/Plugins/Importers/USDImporter/Source/GeometryCacheUSD/Public/GeometryCacheTrackUSD.h
--- a/Engine/Plugins/Importers/USDImporter/Source/GeometryCacheUSD/Public/GeometryCacheTrackUSD.h
+++ b/Engine/Plugins/Importers/USDImporter/Source/GeometryCacheUSD/Public/GeometryCacheTrackUSD.h
@@ -2,12 +2,15 @@
 
 #pragma once
 
+#include "USDGeomMeshConversion.h"
+
+#include "UsdWrappers/UsdStage.h"
+
 #include "CoreMinimal.h"
 #include "GeometryCacheTrack.h"
 #include "GeometryCacheMeshData.h"
 #include "GeometryCacheTrackUSDTypes.h"
 #include "GeometryCacheUSDStream.h"
-#include "UsdWrappers/UsdStage.h"
 
 #include "GeometryCacheTrackUSD.generated.h"
 
@@ -75,6 +78,10 @@ public:
 
 	FString PrimPath;
 
+#if USE_USD_SDK
+	UsdToUnreal::FUsdMeshConversionOptions MeshConversionOptions;
+#endif
+
 	UE::FUsdStage CurrentStagePinned;
 	UE::FUsdStageWeak CurrentStageWeak;
 	FString StageRootLayerPath;
diff --git a/Engine/Plugins/Importers/USDImporter/Source/USDSchemas/Private/USDGeomMeshTranslator.cpp b/Engine/Plugins/Importers/USDImporter/Source/USDSchemas/Private/USDGeomMeshTranslator.cpp
--- a/Engine/Plugins/Importers/USDImporter/Source/USDSchemas/Private/USDGeomMeshTranslator.cpp
+++ b/Engine/Plugins/Importers/USDImporter/Source/USDSchemas/Private/USDGeomMeshTranslator.cpp
@@ -526,7 +526,7 @@ namespace UsdGeomMeshTranslatorImpl
 		if (InvalidNormalFraction >= GMeshNormalRepairThreshold)
 		{
 			Options |= EComputeNTBsFlags::Normals;
-			UE_LOG(LogUsd, Warning, TEXT("%f%% of the normals from Mesh prim '%s' are invalid. This is at or above the threshold of '%f%%' (configurable via the cvar '%s'), so normals will be discarded and fully recomputed."),
+			UE_LOG(LogUsd, Log, TEXT("%f%% of the normals from Mesh prim '%s' are invalid or unusable. This is at or above the threshold of '%f%%' (configurable via the cvar '%s'), so normals will be discarded and fully recomputed. Note that when the cvar 'USD.Subdiv.IgnoreNormalsWhenSubdividing' is true it is expected for subdivision meshes to have their normals discarded."),
 				InvalidNormalFraction * 100.0f,
 				*PrimPath,
 				GMeshNormalRepairThreshold * 100.0f,
@@ -535,7 +535,7 @@ namespace UsdGeomMeshTranslatorImpl
 		}
 		else if (InvalidNormalFraction > 0)
 		{
-			UE_LOG(LogUsd, Warning, TEXT("%f%% of the normals from Mesh prim '%s' are invalid. This is below the threshold of '%f%%' (configurable via the cvar '%s'), so the invalid normals will be repaired."),
+			UE_LOG(LogUsd, Log, TEXT("%f%% of the normals from Mesh prim '%s' are invalid or unusable. This is below the threshold of '%f%%' (configurable via the cvar '%s'), so the invalid normals will be repaired."),
 				InvalidNormalFraction * 100.0f,
 				*PrimPath,
 				GMeshNormalRepairThreshold * 100.0f,
@@ -967,6 +967,7 @@ void FGeomMeshCreateAssetsTaskChain::SetupTasks()
 			Options.MaterialPurpose = MaterialPurposeToken;
 			Options.bMergeIdenticalMaterialSlots = Context->bMergeIdenticalMaterialSlots;
 			Options.AdditionalTransform = AdditionalTransform;
+			Options.SubdivisionLevel = Context->SubdivisionLevel;
 
 			UsdGeomMeshTranslatorImpl::LoadMeshDescriptions(
 				GetPrim(),
diff --git a/Engine/Plugins/Importers/USDImporter/Source/USDSchemas/Private/USDGeomPointInstancerTranslator.cpp b/Engine/Plugins/Importers/USDImporter/Source/USDSchemas/Private/USDGeomPointInstancerTranslator.cpp
--- a/Engine/Plugins/Importers/USDImporter/Source/USDSchemas/Private/USDGeomPointInstancerTranslator.cpp
+++ b/Engine/Plugins/Importers/USDImporter/Source/USDSchemas/Private/USDGeomPointInstancerTranslator.cpp
@@ -133,6 +133,7 @@ void FUsdGeomPointInstancerCreateAssetsTaskChain::SetupTasks()
 			Options.RenderContext = RenderContextToken;
 			Options.MaterialPurpose = MaterialPurposeToken;
 			Options.bMergeIdenticalMaterialSlots = Context->bMergeIdenticalMaterialSlots;
+			Options.SubdivisionLevel = Context->SubdivisionLevel;
 
 			UsdToUnreal::ConvertGeomMeshHierarchy(
 				GetPrim(),
diff --git a/Engine/Plugins/Importers/USDImporter/Source/USDSchemas/Private/USDGeomPrimitiveTranslator.cpp b/Engine/Plugins/Importers/USDImporter/Source/USDSchemas/Private/USDGeomPrimitiveTranslator.cpp
--- a/Engine/Plugins/Importers/USDImporter/Source/USDSchemas/Private/USDGeomPrimitiveTranslator.cpp
+++ b/Engine/Plugins/Importers/USDImporter/Source/USDSchemas/Private/USDGeomPrimitiveTranslator.cpp
@@ -94,6 +94,7 @@ namespace UsdGeomPrimitiveTranslatorImpl
 				   Options.RenderContext = RenderContextToken;
 				   Options.MaterialPurpose = MaterialPurposeToken;
 				   Options.bMergeIdenticalMaterialSlots = Context->bMergeIdenticalMaterialSlots;
+				   Options.SubdivisionLevel = Context->SubdivisionLevel;
 
 				   UsdGeomPrimitiveTranslatorImpl::LoadMeshDescriptions(
 					   GetPrim(),
diff --git a/Engine/Plugins/Importers/USDImporter/Source/USDSchemas/Private/USDGeomXformableTranslator.cpp b/Engine/Plugins/Importers/USDImporter/Source/USDSchemas/Private/USDGeomXformableTranslator.cpp
--- a/Engine/Plugins/Importers/USDImporter/Source/USDSchemas/Private/USDGeomXformableTranslator.cpp
+++ b/Engine/Plugins/Importers/USDImporter/Source/USDSchemas/Private/USDGeomXformableTranslator.cpp
@@ -244,6 +244,7 @@ void FUsdGeomXformableCreateAssetsTaskChain::SetupTasks()
 			Options.RenderContext = RenderContextToken;
 			Options.MaterialPurpose = MaterialPurposeToken;
 			Options.bMergeIdenticalMaterialSlots = Context->bMergeIdenticalMaterialSlots;
+			Options.SubdivisionLevel = Context->SubdivisionLevel;
 
 			UsdToUnreal::ConvertGeomMeshHierarchy(
 				GetPrim(),
diff --git a/Engine/Plugins/Importers/USDImporter/Source/USDSchemas/Private/USDGeometryCacheTranslator.cpp b/Engine/Plugins/Importers/USDImporter/Source/USDSchemas/Private/USDGeometryCacheTranslator.cpp
--- a/Engine/Plugins/Importers/USDImporter/Source/USDSchemas/Private/USDGeometryCacheTranslator.cpp
+++ b/Engine/Plugins/Importers/USDImporter/Source/USDSchemas/Private/USDGeometryCacheTranslator.cpp
@@ -56,6 +56,12 @@ static FAutoConsoleVariableRef CVarUsdGeometryCacheParallelFrameReads(
 	GUsdGeometryCacheParallelFrameReads,
 	TEXT("Maximum number of mesh frames to read in parallel"));
 
+static bool GEnableSubdiv = false;
+static FAutoConsoleVariableRef CVarEnableSubdiv(
+	TEXT("USD.GeometryCache.EnableSubdiv"),
+	GEnableSubdiv,
+	TEXT("Whether to subdivide Mesh prim data when parsing GeometryCaches via OpenSubdiv, the same way we subdivide the Mesh data that ends up in StaticMeshes"));
+
 namespace UsdGeometryCacheTranslatorImpl
 {
 	bool ProcessGeometryCacheMaterials(const pxr::UsdPrim& UsdPrim, const TArray<UsdUtils::FUsdPrimMaterialAssignmentInfo>& LODIndexToMaterialInfo,
@@ -226,6 +232,7 @@ namespace UsdGeometryCacheTranslatorImpl
 		Args.Options.RenderContext = RenderContextToken;
 		Args.Options.MaterialPurpose = MaterialPurposeToken;
 		Args.Options.bMergeIdenticalMaterialSlots = false; // Don't merge because the GeometryCache is processed as unflattened (ie. one track per mesh)
+		Args.Options.SubdivisionLevel = GEnableSubdiv ? Context->SubdivisionLevel : 0;
 
 		return Args;
 	}
@@ -275,6 +282,7 @@ namespace UsdGeometryCacheTranslatorImpl
 	{
 		// Create and configure a new USDTrack to be added to the GeometryCache
 		UGeometryCacheTrackUsd* UsdTrack = NewObject<UGeometryCacheTrackUsd>(GeometryCache);
+		UsdTrack->MeshConversionOptions = Args.Options;  // Also pass along the options we'll use for mesh conversion so that we can properly hash the prim
 		UsdTrack->Initialize(
 			Args.Stage,
 			PrimPath,
@@ -687,6 +695,7 @@ void FGeometryCacheCreateAssetsTaskChain::SetupTasks()
 			Options.RenderContext = RenderContextToken;
 			Options.MaterialPurpose = MaterialPurposeToken;
 			Options.bMergeIdenticalMaterialSlots = false; // Don't merge because the GeometryCache is processed as unflattened (ie. one track per mesh)
+			Options.SubdivisionLevel = GEnableSubdiv ? Context->SubdivisionLevel : 0;
 
 			// GeometryCache has only one LOD so add just one MeshDescription and MaterialAssignmentInfo
 			FMeshDescription& AddedMeshDescription = LODIndexToMeshDescription.Emplace_GetRef();
diff --git a/Engine/Plugins/Importers/USDImporter/Source/USDSchemas/Private/USDSkelRootTranslator.cpp b/Engine/Plugins/Importers/USDImporter/Source/USDSchemas/Private/USDSkelRootTranslator.cpp
--- a/Engine/Plugins/Importers/USDImporter/Source/USDSchemas/Private/USDSkelRootTranslator.cpp
+++ b/Engine/Plugins/Importers/USDImporter/Source/USDSchemas/Private/USDSkelRootTranslator.cpp
@@ -1309,6 +1309,7 @@ namespace UsdSkelRootTranslatorImpl
 				Options.TimeCode = Context->Time;
 				Options.RenderContext = RenderContextToken;
 				Options.MaterialPurpose = MaterialPurposeToken;
+				Options.SubdivisionLevel = Context->SubdivisionLevel;
 
 				bool bContinueTaskChain = UsdSkelRootTranslatorImpl::LoadAllSkeletalData(
 					SkeletonCache.Get(),
diff --git a/Engine/Plugins/Importers/USDImporter/Source/USDSchemas/Public/USDSchemaTranslator.h b/Engine/Plugins/Importers/USDImporter/Source/USDSchemas/Public/USDSchemaTranslator.h
--- a/Engine/Plugins/Importers/USDImporter/Source/USDSchemas/Public/USDSchemaTranslator.h
+++ b/Engine/Plugins/Importers/USDImporter/Source/USDSchemas/Public/USDSchemaTranslator.h
@@ -173,6 +173,9 @@ struct USDSCHEMAS_API FUsdSchemaTranslationContext : public TSharedFromThis< FUs
 	/** Describes what to add to the root bone animation within generated AnimSequences, if anything */
 	EUsdRootMotionHandling RootMotionHandling = EUsdRootMotionHandling::NoAdditionalRootMotion;
 
+	/** Subdivision level to use for all subdivision meshes on the opened stage. 0 means "don't subdivide" */
+	int32 SubdivisionLevel = 0;
+
 	/** If a generated UStaticMesh has at least this many triangles we will attempt to enable Nanite */
 	int32 NaniteTriangleThreshold;
 
diff --git a/Engine/Plugins/Importers/USDImporter/Source/USDStage/Private/USDStageActor.cpp b/Engine/Plugins/Importers/USDImporter/Source/USDStage/Private/USDStageActor.cpp
--- a/Engine/Plugins/Importers/USDImporter/Source/USDStage/Private/USDStageActor.cpp
+++ b/Engine/Plugins/Importers/USDImporter/Source/USDStage/Private/USDStageActor.cpp
@@ -140,6 +140,7 @@ struct FUsdStageActorImpl
 		TranslationContext->RenderContext = StageActor->RenderContext;
 		TranslationContext->MaterialPurpose = StageActor->MaterialPurpose;
 		TranslationContext->RootMotionHandling = StageActor->RootMotionHandling;
+		TranslationContext->SubdivisionLevel = StageActor->SubdivisionLevel;
 		TranslationContext->BlendShapesByPath = &StageActor->BlendShapesByPath;
 		TranslationContext->InfoCache = StageActor->InfoCache;
 		TranslationContext->bTranslateOnlyUsedMaterials = GTranslateOnlyUsedMaterialsWhenOpeningStage;
@@ -571,6 +572,7 @@ struct FUsdStageActorImpl
 			EventAttributes.Emplace(TEXT("RenderContext"), StageActor->RenderContext.ToString());
 			EventAttributes.Emplace(TEXT("MaterialPurpose"), StageActor->MaterialPurpose.ToString());
 			EventAttributes.Emplace(TEXT("RootMotionHandling"), LexToString((uint8)StageActor->RootMotionHandling));
+			EventAttributes.Emplace(TEXT("SubdivisionLevel"), LexToString(StageActor->SubdivisionLevel));
 
 			const bool bAutomated = false;
 			IUsdClassesModule::SendAnalytics(MoveTemp(EventAttributes), TEXT("Open"), bAutomated, ElapsedSeconds, NumberOfFrames, Extension);
@@ -966,6 +968,7 @@ AUsdStageActor::AUsdStageActor()
 	, NaniteTriangleThreshold((uint64)1000000)
 	, MaterialPurpose(*UnrealIdentifiers::MaterialPreviewPurpose)
 	, RootMotionHandling(EUsdRootMotionHandling::NoAdditionalRootMotion)
+	, SubdivisionLevel(0)
 	, Time(0.0f)
 	, bIsTransitioningIntoPIE(false)
 	, bIsModifyingAProperty(false)
@@ -2320,6 +2323,15 @@ void AUsdStageActor::SetRootMotionHandling( EUsdRootMotionHandling NewHandlingSt
 	LoadUsdStage();
 }
 
+void AUsdStageActor::SetSubdivisionLevel(int32 NewSubdivisionLevel)
+{
+	const bool bMarkDirty = false;
+	Modify(bMarkDirty);
+
+	SubdivisionLevel = NewSubdivisionLevel;
+	LoadUsdStage();
+}
+
 float AUsdStageActor::GetTime() const
 {
 	return Time;
@@ -4141,6 +4153,10 @@ void AUsdStageActor::HandlePropertyChangedEvent(FPropertyChangedEvent& PropertyC
 	{
 		SetRootMotionHandling(RootMotionHandling);
 	}
+	else if (PropertyName == GET_MEMBER_NAME_CHECKED(AUsdStageActor, SubdivisionLevel))
+	{
+		SetSubdivisionLevel(SubdivisionLevel);
+	}
 	else if (PropertyName == GET_MEMBER_NAME_CHECKED(AUsdStageActor, UsdAssetCache))
 	{
 		SetAssetCache(UsdAssetCache);
diff --git a/Engine/Plugins/Importers/USDImporter/Source/USDStage/Public/USDStageActor.h b/Engine/Plugins/Importers/USDImporter/Source/USDStage/Public/USDStageActor.h
--- a/Engine/Plugins/Importers/USDImporter/Source/USDStage/Public/USDStageActor.h
+++ b/Engine/Plugins/Importers/USDImporter/Source/USDStage/Public/USDStageActor.h
@@ -92,6 +92,13 @@ public:
 	UPROPERTY( EditAnywhere, BlueprintReadOnly, Category = "USD", config )
 	EUsdRootMotionHandling RootMotionHandling = EUsdRootMotionHandling::NoAdditionalRootMotion;
 
+	/**
+	 * Subdivision level to use for all subdivision meshes on the opened stage. 0 means "don't subdivide".
+	 * The maximum level of subdivision allowed can be configured via the 'USD.Subdiv.MaxSubdivLevel' cvar.
+	 */
+	UPROPERTY(EditAnywhere, BlueprintReadOnly, Category = "USD", config)
+	int32 SubdivisionLevel;
+
 public:
 	DECLARE_EVENT_OneParam(AUsdStageActor, FOnActorLoaded, AUsdStageActor*);
 	USDSTAGE_API static FOnActorLoaded OnActorLoaded;
@@ -150,6 +157,9 @@ public:
 	UFUNCTION( BlueprintCallable, Category = "USD", meta = ( CallInEditor = "true" ) )
 	USDSTAGE_API void SetRootMotionHandling( EUsdRootMotionHandling NewHandlingStrategy );
 
+	UFUNCTION(BlueprintCallable, Category = "USD", meta = (CallInEditor = "true"))
+	USDSTAGE_API void SetSubdivisionLevel(int32 NewLevel);
+
 	UFUNCTION(BlueprintCallable, Category = "USD", meta = (CallInEditor = "true"))
 	USDSTAGE_API float GetTime() const;
 
diff --git a/Engine/Plugins/Importers/USDImporter/Source/USDStageEditor/Private/SUSDStage.cpp b/Engine/Plugins/Importers/USDImporter/Source/USDStageEditor/Private/SUSDStage.cpp
--- a/Engine/Plugins/Importers/USDImporter/Source/USDStageEditor/Private/SUSDStage.cpp
+++ b/Engine/Plugins/Importers/USDImporter/Source/USDStageEditor/Private/SUSDStage.cpp
@@ -1030,6 +1030,11 @@ void SUsdStage::FillOptionsMenu(FMenuBuilder& MenuBuilder)
 			LOCTEXT( "RootMotionHandling_ToolTip", "Choose how to handle root motion for generated AnimSequences and LevelSequences" ),
 			FNewMenuDelegate::CreateSP( this, &SUsdStage::FillRootMotionSubMenu ) );
 
+		MenuBuilder.AddSubMenu(
+			LOCTEXT("SubdivLevel", "Subdivision level"),
+			LOCTEXT("SubdivLevel_ToolTip", "Subdivision level to use for all subdivision meshes on the opened stage. 0 means 'don't subdivide'"),
+			FNewMenuDelegate::CreateSP(this, &SUsdStage::FillSubdivisionLevelSubMenu));
+
 		MenuBuilder.AddSubMenu(
 			LOCTEXT( "Collapsing", "Collapsing" ),
 			LOCTEXT( "Collapsing_ToolTip", "Whether to try to combine individual assets and components of the same type on a Kind-per-Kind basis, like multiple Mesh prims into a single Static Mesh" ),
@@ -1504,6 +1509,32 @@ void SUsdStage::FillRootMotionSubMenu( FMenuBuilder& MenuBuilder )
 	AddRootMotionEntry( EUsdRootMotionHandling::UseMotionFromSkeleton, LOCTEXT("UseMotionFromSkeletonText", "Use motion from Skeleton"));
 }
 
+void SUsdStage::FillSubdivisionLevelSubMenu(FMenuBuilder& MenuBuilder)
+{
+	if (AUsdStageActor* StageActor = GetStageActorOrCDO())
+	{
+		CurrentSubdivisionLevel = StageActor->SubdivisionLevel;
+	}
+
+	static IConsoleVariable* MaxSubdivCvar = IConsoleManager::Get().FindConsoleVariable(TEXT("USD.Subdiv.MaxSubdivLevel"));
+	const int32 MaxSubdivLevel = MaxSubdivCvar ? MaxSubdivCvar->GetInt() : 0;
+
+	TSharedRef<SSpinBox<int32>> Slider = SNew(SSpinBox<int32>)
+		.MinValue(0)
+		.MaxValue(MaxSubdivLevel)
+		.ToolTipText(LOCTEXT(
+			"SubdivLevelSliderTooltip",
+			"Subdivision level to use for all subdivision meshes on the opened stage. 0 means 'don't subdivide'. The maximum level of subdivision allowed can be configured via the 'USD.Subdiv.MaxSubdivLevel' cvar."
+		))
+		.Value(this, &SUsdStage::GetSubdivisionLevelValue)
+		.OnValueChanged(this, &SUsdStage::OnSubdivisionLevelValueChanged)
+		.SupportDynamicSliderMaxValue(true)
+		.OnValueCommitted(this, &SUsdStage::OnSubdivisionLevelValueCommitted);
+
+	const bool bNoIndent = true;
+	MenuBuilder.AddWidget(Slider, FText::FromString(TEXT("Subdivision level: ")), bNoIndent);
+}
+
 void SUsdStage::FillCollapsingSubMenu( FMenuBuilder& MenuBuilder )
 {
 	MenuBuilder.AddMenuEntry(
@@ -2558,6 +2589,41 @@ void SUsdStage::OnNaniteTriangleThresholdValueCommitted( int32 InValue, ETextCom
 	}
 }
 
+int32 SUsdStage::GetSubdivisionLevelValue() const
+{
+	return CurrentSubdivisionLevel;
+}
+
+void SUsdStage::OnSubdivisionLevelValueChanged(int32 InValue)
+{
+	CurrentSubdivisionLevel = InValue;
+}
+
+void SUsdStage::OnSubdivisionLevelValueCommitted(int32 InValue, ETextCommit::Type InCommitType)
+{
+	AUsdStageActor* StageActor = GetStageActorOrCDO();
+	if (!StageActor)
+	{
+		return;
+	}
+
+	FScopedTransaction Transaction(FText::Format(
+		LOCTEXT("SubdivisionLevelCommittedTransaction", "Change Subdivision level for USD stage actor '{0}'"),
+		FText::FromString(StageActor->GetActorLabel())
+	));
+
+	// c.f. comment in SUsdStage::FillCollapsingSubMenu
+	TGuardValue<bool> MaintainSelectionGuard(bUpdatingViewportSelection, true);
+
+	StageActor->SetSubdivisionLevel(InValue);
+	CurrentSubdivisionLevel = InValue;
+
+	if (StageActor->IsTemplate())
+	{
+		StageActor->SaveConfig();
+	}
+}
+
 UE::FUsdStageWeak SUsdStage::GetCurrentStage() const
 {
 	const AUsdStageActor* StageActor = ViewModel.UsdStageActor.Get();
diff --git a/Engine/Plugins/Importers/USDImporter/Source/USDStageEditor/Private/SUSDStage.h b/Engine/Plugins/Importers/USDImporter/Source/USDStageEditor/Private/SUSDStage.h
--- a/Engine/Plugins/Importers/USDImporter/Source/USDStageEditor/Private/SUSDStage.h
+++ b/Engine/Plugins/Importers/USDImporter/Source/USDStageEditor/Private/SUSDStage.h
@@ -88,6 +88,7 @@ protected:
 	void FillRenderContextSubMenu( FMenuBuilder& MenuBuilder );
 	void FillMaterialPurposeSubMenu( FMenuBuilder& MenuBuilder );
 	void FillRootMotionSubMenu( FMenuBuilder& MenuBuilder );
+	void FillSubdivisionLevelSubMenu(FMenuBuilder& MenuBuilder);
 	void FillCollapsingSubMenu( FMenuBuilder& MenuBuilder );
 	void FillInterpolationTypeSubMenu( FMenuBuilder& MenuBuilder );
 	void FillSelectionSubMenu( FMenuBuilder& MenuBuilder );
@@ -112,6 +113,10 @@ protected:
 	void OnNaniteTriangleThresholdValueChanged( int32 InValue );
 	void OnNaniteTriangleThresholdValueCommitted( int32 InValue, ETextCommit::Type InCommitType );
 
+	int32 GetSubdivisionLevelValue() const;
+	void OnSubdivisionLevelValueChanged(int32 InValue);
+	void OnSubdivisionLevelValueCommitted(int32 InValue, ETextCommit::Type InCommitType);
+
 	UE::FUsdStageWeak GetCurrentStage() const;
 
 	AUsdStageActor* GetStageActorOrCDO() const;
@@ -152,6 +157,8 @@ protected:
 
 	int32 CurrentNaniteThreshold = INT32_MAX;
 
+	int32 CurrentSubdivisionLevel = 0;
+
 	TArray<TSharedPtr<FString>> MaterialPurposes;
 
 	// We use our own bool to track engine exit because we use some tickers to delay work onto the next frame,
diff --git a/Engine/Plugins/Importers/USDImporter/Source/USDStageEditorViewModels/Private/USDStageViewModel.cpp b/Engine/Plugins/Importers/USDImporter/Source/USDStageEditorViewModels/Private/USDStageViewModel.cpp
--- a/Engine/Plugins/Importers/USDImporter/Source/USDStageEditorViewModels/Private/USDStageViewModel.cpp
+++ b/Engine/Plugins/Importers/USDImporter/Source/USDStageEditorViewModels/Private/USDStageViewModel.cpp
@@ -316,6 +316,7 @@ void FUsdStageViewModel::ImportStage( const TCHAR* TargetContentFolder, UUsdStag
 		ImportContext.ImportOptions->ImportTimeCode = StageActor->GetTime();
 		ImportContext.ImportOptions->NaniteTriangleThreshold = StageActor->NaniteTriangleThreshold;
 		ImportContext.ImportOptions->RootMotionHandling = StageActor->RootMotionHandling;
+		ImportContext.ImportOptions->SubdivisionLevel = StageActor->SubdivisionLevel;
 		ImportContext.ImportOptions->KindsToCollapse = StageActor->KindsToCollapse;
 		ImportContext.ImportOptions->bMergeIdenticalMaterialSlots = StageActor->bMergeIdenticalMaterialSlots;
 
diff --git a/Engine/Plugins/Importers/USDImporter/Source/USDStageImporter/Private/USDStageImportOptions.cpp b/Engine/Plugins/Importers/USDImporter/Source/USDStageImporter/Private/USDStageImportOptions.cpp
--- a/Engine/Plugins/Importers/USDImporter/Source/USDStageImporter/Private/USDStageImportOptions.cpp
+++ b/Engine/Plugins/Importers/USDImporter/Source/USDStageImporter/Private/USDStageImportOptions.cpp
@@ -26,6 +26,7 @@ UUsdStageImportOptions::UUsdStageImportOptions(const FObjectInitializer& ObjectI
 	IUsdSchemasModule& UsdSchemasModule = FModuleManager::Get().LoadModuleChecked< IUsdSchemasModule >( TEXT("USDSchemas") );
 	RenderContextToImport = UsdSchemasModule.GetRenderContextRegistry().GetUnrealRenderContext();
 	MaterialPurpose = *UnrealIdentifiers::MaterialPreviewPurpose;
+	SubdivisionLevel = 0;
 	bOverrideStageOptions = false;
 	StageOptions.MetersPerUnit = 0.01f;
 	StageOptions.UpAxis = EUsdUpAxis::ZAxis;
@@ -95,6 +96,7 @@ void UsdUtils::AddAnalyticsAttributes(
 	InOutAttributes.Emplace( TEXT( "RenderContextToImport" ), Options.RenderContextToImport.ToString() );
 	InOutAttributes.Emplace( TEXT( "MaterialPurpose" ), Options.MaterialPurpose.ToString() );
 	InOutAttributes.Emplace( TEXT( "RootMotionHandling" ), LexToString( ( uint8 ) Options.RootMotionHandling ) );
+	InOutAttributes.Emplace(TEXT("SubdivisionLevel"), LexToString(Options.SubdivisionLevel));
 	InOutAttributes.Emplace( TEXT( "OverrideStageOptions" ), Options.bOverrideStageOptions );
 	if ( Options.bOverrideStageOptions )
 	{
diff --git a/Engine/Plugins/Importers/USDImporter/Source/USDStageImporter/Private/USDStageImporter.cpp b/Engine/Plugins/Importers/USDImporter/Source/USDStageImporter/Private/USDStageImporter.cpp
--- a/Engine/Plugins/Importers/USDImporter/Source/USDStageImporter/Private/USDStageImporter.cpp
+++ b/Engine/Plugins/Importers/USDImporter/Source/USDStageImporter/Private/USDStageImporter.cpp
@@ -1848,6 +1848,7 @@ void UUsdStageImporter::ImportFromFile(FUsdStageImportContext& ImportContext)
 	TranslationContext->RenderContext = ImportContext.ImportOptions->RenderContextToImport;
 	TranslationContext->MaterialPurpose = ImportContext.ImportOptions->MaterialPurpose;
 	TranslationContext->RootMotionHandling = ImportContext.ImportOptions->RootMotionHandling;
+	TranslationContext->SubdivisionLevel = ImportContext.ImportOptions->SubdivisionLevel;
 	TranslationContext->ParentComponent = ImportContext.SceneActor ? ImportContext.SceneActor->GetRootComponent() : nullptr;
 	TranslationContext->KindsToCollapse = ( EUsdDefaultKind ) ImportContext.ImportOptions->KindsToCollapse;
 	TranslationContext->bMergeIdenticalMaterialSlots = ImportContext.ImportOptions->bMergeIdenticalMaterialSlots;
@@ -1968,6 +1969,7 @@ bool UUsdStageImporter::ReimportSingleAsset(
 	TranslationContext->RenderContext = ImportContext.ImportOptions->RenderContextToImport;
 	TranslationContext->MaterialPurpose = ImportContext.ImportOptions->MaterialPurpose;
 	TranslationContext->RootMotionHandling = ImportContext.ImportOptions->RootMotionHandling;
+	TranslationContext->SubdivisionLevel = ImportContext.ImportOptions->SubdivisionLevel;
 	TranslationContext->KindsToCollapse = ( EUsdDefaultKind ) ImportContext.ImportOptions->KindsToCollapse;
 	TranslationContext->bMergeIdenticalMaterialSlots = ImportContext.ImportOptions->bMergeIdenticalMaterialSlots;
 	TranslationContext->bAllowInterpretingLODs = ImportContext.ImportOptions->bInterpretLODs;
diff --git a/Engine/Plugins/Importers/USDImporter/Source/USDStageImporter/Public/USDStageImportOptions.h b/Engine/Plugins/Importers/USDImporter/Source/USDStageImporter/Public/USDStageImportOptions.h
--- a/Engine/Plugins/Importers/USDImporter/Source/USDStageImporter/Public/USDStageImportOptions.h
+++ b/Engine/Plugins/Importers/USDImporter/Source/USDStageImporter/Public/USDStageImportOptions.h
@@ -105,6 +105,10 @@ public:
 	UPROPERTY( BlueprintReadWrite, config, EditAnywhere, Category = "USD options", meta = ( EditCondition = bImportSkeletalAnimations ) )
 	EUsdRootMotionHandling RootMotionHandling = EUsdRootMotionHandling::NoAdditionalRootMotion;
 
+	/** Subdivision level to use for all subdivision meshes on the opened stage. 0 means "don't subdivide" */
+	UPROPERTY(BlueprintReadWrite, config, EditAnywhere, Category = "USD options")
+	int32 SubdivisionLevel;
+
 	/** Whether to use the specified StageOptions instead of the stage's own settings */
 	UPROPERTY( BlueprintReadWrite, config, EditAnywhere, Category = "USD options" )
 	bool bOverrideStageOptions;
diff --git a/Engine/Plugins/Importers/USDImporter/Source/USDUtilities/Private/USDGeomMeshConversion.cpp b/Engine/Plugins/Importers/USDImporter/Source/USDUtilities/Private/USDGeomMeshConversion.cpp
--- a/Engine/Plugins/Importers/USDImporter/Source/USDUtilities/Private/USDGeomMeshConversion.cpp
+++ b/Engine/Plugins/Importers/USDImporter/Source/USDUtilities/Private/USDGeomMeshConversion.cpp
@@ -5,10 +5,8 @@
 #include "USDGeomMeshConversion.h"
 
 #include "UnrealUSDWrapper.h"
-#include "USDAssetImportData.h"
 #include "USDAttributeUtils.h"
 #include "USDConversionUtils.h"
-#include "USDErrorUtils.h"
 #include "USDLog.h"
 #include "USDMemory.h"
 #include "USDPrimConversion.h"
@@ -16,20 +14,15 @@
 #include "USDShadeConversion.h"
 #include "USDTypesConversion.h"
 
-#include "UsdWrappers/SdfLayer.h"
 #include "UsdWrappers/SdfPath.h"
-#include "UsdWrappers/UsdAttribute.h"
 #include "UsdWrappers/UsdPrim.h"
 #include "UsdWrappers/UsdStage.h"
 
-#include "AssetRegistry/AssetData.h"
-#include "AssetRegistry/AssetRegistryModule.h"
 #include "Engine/StaticMesh.h"
 #include "Materials/MaterialInstanceConstant.h"
 #include "Materials/MaterialInstanceDynamic.h"
 #include "MeshDescription.h"
 #include "Misc/Paths.h"
-#include "Modules/ModuleManager.h"
 #include "StaticMeshAttributes.h"
 #include "StaticMeshOperations.h"
 #include "StaticMeshResources.h"
@@ -39,9 +32,10 @@
 #endif // WITH_EDITOR
 
 #include "USDIncludesStart.h"
+	#include "opensubdiv/far/primvarRefiner.h"
+	#include "opensubdiv/far/topologyRefiner.h"
 	#include "pxr/imaging/pxOsd/meshTopology.h"
-	#include "pxr/usd/ar/resolver.h"
-	#include "pxr/usd/ar/resolverScopedCache.h"
+	#include "pxr/imaging/pxOsd/refinerFactory.h"
 	#include "pxr/usd/sdf/layer.h"
 	#include "pxr/usd/sdf/layerUtils.h"
 	#include "pxr/usd/sdf/types.h"
@@ -74,6 +68,21 @@ static FAutoConsoleVariableRef CVarMaxInstancesPerPointInstancer(
 	GMaxInstancesPerPointInstancer,
 	TEXT( "We will only parse up to this many instances from any point instancer when reading from USD to UE. Set this to -1 to disable this limit." ) );
 
+static bool GIgnoreNormalsWhenSubdividing = true;
+static FAutoConsoleVariableRef CVarIgnoreNormalsWhenSubdividing(
+	TEXT("USD.Subdiv.IgnoreNormalsWhenSubdividing"),
+	GIgnoreNormalsWhenSubdividing,
+	TEXT("This being true means that whenever we subdivide a mesh we fully ignore the authored normals (if any) and recompute new normals. If this is false we will try interpolating the normals during subdivision like a regular primvar")
+);
+
+static const FString MaxUsdSubdivLevelCvarName = TEXT("USD.Subdiv.MaxSubdivLevel");
+static int32 GMaxSubdivLevel = 6;
+static FAutoConsoleVariableRef CVarMaxSubdivLevel(
+	*MaxUsdSubdivLevelCvarName,
+	GMaxSubdivLevel,
+	TEXT("Maximum allowed level of subdivision (1 means a single iteration of subdivision)")
+);
+
 const FName MeshAttribute::VertexInstance::Velocity("Velocity");
 
 namespace UE::UsdGeomMeshConversion::Private
@@ -110,11 +119,47 @@ namespace UE::UsdGeomMeshConversion::Private
 
 	pxr::TfToken GetAttrInterpolation(const pxr::UsdAttribute& Attr, const pxr::TfToken& DefaultValue = pxr::UsdGeomTokens->constant)
 	{
+		if (!Attr)
+		{
+			return DefaultValue;
+		}
+
 		pxr::TfToken RetrievedValue;
-		if (Attr && Attr.GetMetadata(pxr::UsdGeomTokens->interpolation, &RetrievedValue))
+		const bool bGotInterpolationValue = Attr.GetMetadata(pxr::UsdGeomTokens->interpolation, &RetrievedValue);
+
+		// If we have an authored value just go ahead and use that
+		if (Attr.HasAuthoredMetadata(pxr::UsdGeomTokens->interpolation) && bGotInterpolationValue)
 		{
 			return RetrievedValue;
 		}
+
+		// Otherwise if our attribute describes an array with a single element and has no authored interpolation assume "constant", as
+		// it's impossible for any other interpolation type to be valid. usdview does this too.
+		// Note we try our best to get anything here and also check timeSampled values in case our default time Get() fails
+		pxr::VtValue TypeErasedValue;
+		if (Attr.Get(&TypeErasedValue) || Attr.Get(&TypeErasedValue, pxr::UsdTimeCode::EarliestTime()))
+		{
+			if (TypeErasedValue.IsArrayValued() && TypeErasedValue.GetArraySize() == 1)
+			{
+				return pxr::UsdGeomTokens->constant;
+			}
+		}
+		// If we couldn't get any actual value for the attribute whatsoever then pretend it doesn't have a valid value for interpolation
+		// either. We need use this because if SubdivideMeshData sees that an attribute has e.g. "vertex" interpolation, it will allocate
+		// and try generating one value for it for every vertex... if we don't have any value to begin with we'll just end up with a zero
+		// value for each instead
+		else
+		{
+			return {};
+		}
+
+		// Otherwise if we don't have an authored value but did manage to get value for interpolation somehow
+		// (maybe as a fallback?) then return that
+		if (bGotInterpolationValue)
+		{
+			return RetrievedValue;
+		}
+
 		return DefaultValue;
 	}
 
@@ -905,75 +950,1041 @@ namespace UE::UsdGeomMeshConversion::Private
 			}
 		}
 
-		if (bTraverseChildren)
-		{
-			for (const pxr::UsdPrim& ChildPrim : Prim.GetFilteredChildren(pxr::UsdTraverseInstanceProxies()))
+		if (bTraverseChildren)
+		{
+			for (const pxr::UsdPrim& ChildPrim : Prim.GetFilteredChildren(pxr::UsdTraverseInstanceProxies()))
+			{
+				const bool bChildIsFirstPrim = false;
+				RecursivelyCollectPrimvars(
+					ChildPrim,
+					Options,
+					InOutAllPrimvars,
+					InOutPreferredPrimvars,
+					bChildIsFirstPrim
+				);
+			}
+		}
+	}
+
+	/**
+	 * Returns the set of primvar names that can be used for each UV index for a mesh collapsed from the subtree
+	 * starting at RootPrim.
+	 */
+	TMap<FString, int32> CollectSubtreePrimvars(
+		const pxr::UsdPrim& RootPrim,
+		const UsdToUnreal::FUsdMeshConversionOptions& Options,
+		bool bIsFirstPrim
+	)
+	{
+		TSet<FString> AllPrimvars;
+		TSet<FString> PreferredPrimvars;
+
+		RecursivelyCollectPrimvars(
+			RootPrim,
+			Options,
+			AllPrimvars,
+			PreferredPrimvars,
+			bIsFirstPrim
+		);
+
+		return UsdUtils::CombinePrimvarsIntoUVSets(
+			AllPrimvars,
+			PreferredPrimvars
+		);
+	}
+
+	// Unconverted, raw USD mesh data to convert into a MeshDescription
+	struct FUsdMeshData
+	{
+		// So that we can reference the prim on error messages
+		FString SourcePrimPath;
+
+		pxr::TfToken Orientation = pxr::UsdGeomTokens->rightHanded;
+
+		pxr::VtArray<int> FaceVertexCounts;
+		pxr::VtArray<int> FaceVertexIndices;
+
+		// Main attributes, which could have come from primvars
+		pxr::VtArray<pxr::GfVec3f> Points;
+		pxr::VtArray<pxr::GfVec3f> Normals;
+		pxr::VtArray<pxr::GfVec3f> Velocities;
+		pxr::VtArray<pxr::GfVec3f> DisplayColors;
+		pxr::VtArray<float> DisplayOpacities;
+		TArray<pxr::VtArray<pxr::GfVec2f>> UVSets;
+
+		pxr::TfToken PointInterpolation;
+		pxr::TfToken NormalInterpolation;
+		pxr::TfToken VelocityInterpolation;
+		pxr::TfToken DisplayColorInterpolation;
+		pxr::TfToken DisplayOpacityInterpolation;
+		TArray<pxr::TfToken> UVSetInterpolations;
+
+		// In case those are indexed primvars, these will contain the indices
+		// Note: Velocities is not a primvar, so it can't have indices
+		pxr::VtArray<int> PointIndices;
+		pxr::VtArray<int> NormalIndices;
+		pxr::VtArray<int> DisplayColorIndices;
+		pxr::VtArray<int> DisplayOpacityIndices;
+		TArray<pxr::VtArray<int>> UVSetIndices;
+
+		// Attributes used for subdivision
+		pxr::TfToken SubdivScheme;
+		pxr::TfToken InterpolateBoundary;
+		pxr::TfToken FaceVaryingInterpolation;
+		pxr::TfToken TriangleSubdivision;
+		pxr::VtArray<int> CornerIndices;
+		pxr::VtArray<float> CornerSharpnesses;
+		pxr::VtArray<int> CreaseIndices;
+		pxr::VtArray<int> CreaseLengths;
+		pxr::VtArray<float> CreaseSharpnesses;
+		pxr::TfToken CreaseMethod = pxr::PxOsdOpenSubdivTokens->uniform;
+		pxr::VtArray<int> HoleIndices;
+
+		UsdUtils::FUsdPrimMaterialAssignmentInfo LocalMaterialInfo;
+
+		TOptional<int32> ProvidedNumUVSets;
+		int32 MaterialIndexOffset = 0;
+	};
+
+	bool CollectMeshData(
+		const pxr::UsdPrim& Prim,
+		const UsdToUnreal::FUsdMeshConversionOptions& Options,
+		FUsdMeshData& InOutMeshData,
+		UsdUtils::FUsdPrimMaterialAssignmentInfo& InOutCombinedMaterialAssignments
+	)
+	{
+		FScopedUsdAllocs Allocs;
+
+		pxr::UsdGeomGprim Gprim{Prim};
+		if (!Gprim)
+		{
+			return false;
+		}
+
+		// All pointsBased/Gprim attributes we'll retrieve happen to have default varying interpolation
+		const pxr::TfToken DefaultInterpolation = pxr::UsdGeomTokens->varying;
+
+		InOutMeshData.SourcePrimPath = UsdToUnreal::ConvertPath(Prim.GetPrimPath());
+
+		InOutMeshData.Orientation = GetGprimOrientation(Gprim, Options.TimeCode);
+
+		// DisplayColors
+		if (pxr::UsdGeomPrimvar DisplayColorsPrimvar = Gprim.GetDisplayColorPrimvar())
+		{
+			DisplayColorsPrimvar.Get(&InOutMeshData.DisplayColors, Options.TimeCode);
+			InOutMeshData.DisplayColorInterpolation = GetAttrInterpolation(DisplayColorsPrimvar, DefaultInterpolation);
+			DisplayColorsPrimvar.GetIndices(&InOutMeshData.DisplayColorIndices, Options.TimeCode);
+		}
+
+		// DisplayOpacities
+		if (pxr::UsdGeomPrimvar DisplayOpacitiesPrimvar = Gprim.GetDisplayOpacityPrimvar())
+		{
+			DisplayOpacitiesPrimvar.Get(&InOutMeshData.DisplayOpacities, Options.TimeCode);
+			InOutMeshData.DisplayOpacityInterpolation = GetAttrInterpolation(DisplayOpacitiesPrimvar, DefaultInterpolation);
+			DisplayOpacitiesPrimvar.GetIndices(&InOutMeshData.DisplayOpacityIndices, Options.TimeCode);
+		}
+
+		if (pxr::UsdGeomMesh UsdMesh{Prim})
+		{
+			// Faces
+			if (pxr::UsdAttribute FaceVertexCountsAttr = UsdMesh.GetFaceVertexCountsAttr())
+			{
+				FaceVertexCountsAttr.Get(&InOutMeshData.FaceVertexCounts, Options.TimeCode);
+			}
+
+			// Vertex indices
+			if (pxr::UsdAttribute FaceVertexIndicesAttr = UsdMesh.GetFaceVertexIndicesAttr())
+			{
+				FaceVertexIndicesAttr.Get(&InOutMeshData.FaceVertexIndices, Options.TimeCode);
+			}
+
+			// Points
+			if (pxr::UsdGeomPrimvar PointsPrimvar = pxr::UsdGeomPrimvar(Prim.GetAttribute(UnrealIdentifiers::PrimvarsPoints)))
+			{
+				// Should points always have "vertex" interpolation? Having "varying" forces it to just tessellate instead,
+				// and all OpenSubdiv tutorials use the vertex interpolation type for it
+				PointsPrimvar.Get(&InOutMeshData.Points, Options.TimeCode);
+				InOutMeshData.PointInterpolation = GetAttrInterpolation(PointsPrimvar, pxr::UsdGeomTokens->vertex);
+				PointsPrimvar.GetIndices(&InOutMeshData.PointIndices, Options.TimeCode);
+			}
+			else if (pxr::UsdAttribute PointsAttr = UsdMesh.GetPointsAttr())
+			{
+				PointsAttr.Get(&InOutMeshData.Points, Options.TimeCode);
+				InOutMeshData.PointInterpolation = GetAttrInterpolation(PointsAttr, pxr::UsdGeomTokens->vertex);
+			}
+
+			// Normals
+			if (pxr::UsdGeomPrimvar NormalsPrimvar = pxr::UsdGeomPrimvar(Prim.GetAttribute(UnrealIdentifiers::PrimvarsNormals)))
+			{
+				NormalsPrimvar.Get(&InOutMeshData.Normals, Options.TimeCode);
+				InOutMeshData.NormalInterpolation = GetAttrInterpolation(NormalsPrimvar, DefaultInterpolation);
+				NormalsPrimvar.GetIndices(&InOutMeshData.NormalIndices, Options.TimeCode);
+			}
+			else if (pxr::UsdAttribute NormalsAttr = UsdMesh.GetNormalsAttr())
+			{
+				NormalsAttr.Get(&InOutMeshData.Normals, Options.TimeCode);
+				InOutMeshData.NormalInterpolation = GetAttrInterpolation(NormalsAttr, DefaultInterpolation);
+			}
+
+			// Velocities
+			if (pxr::UsdAttribute VelocitiesAttr = UsdMesh.GetVelocitiesAttr())
+			{
+				VelocitiesAttr.Get(&InOutMeshData.Velocities, Options.TimeCode);
+				InOutMeshData.VelocityInterpolation = GetAttrInterpolation(VelocitiesAttr, DefaultInterpolation);
+			}
+
+			// Collect the subdivision attributes only if we plan on subdividing
+			if (Options.SubdivisionLevel > 0)
+			{
+				if (pxr::UsdAttribute SubdivSchemeAttr = UsdMesh.GetSubdivisionSchemeAttr())
+				{
+					SubdivSchemeAttr.Get(&InOutMeshData.SubdivScheme, Options.TimeCode);
+				}
+
+				if (InOutMeshData.SubdivScheme != pxr::UsdGeomTokens->none)
+				{
+					if (pxr::UsdAttribute InterpolateBoundaryAttr = UsdMesh.GetInterpolateBoundaryAttr())
+					{
+						InterpolateBoundaryAttr.Get(&InOutMeshData.InterpolateBoundary, Options.TimeCode);
+					}
+
+					if (pxr::UsdAttribute FaceVaryingInterpolationAttr = UsdMesh.GetFaceVaryingLinearInterpolationAttr())
+					{
+						FaceVaryingInterpolationAttr.Get(&InOutMeshData.FaceVaryingInterpolation, Options.TimeCode);
+					}
+
+					if (pxr::UsdAttribute TriangleSubdivisionAttr = UsdMesh.GetTriangleSubdivisionRuleAttr())
+					{
+						TriangleSubdivisionAttr.Get(&InOutMeshData.TriangleSubdivision, Options.TimeCode);
+					}
+
+					if (pxr::UsdAttribute CornerIndicesAttr = UsdMesh.GetCornerIndicesAttr())
+					{
+						CornerIndicesAttr.Get(&InOutMeshData.CornerIndices, Options.TimeCode);
+					}
+
+					if (pxr::UsdAttribute CornerSharpnessesAttr = UsdMesh.GetCornerSharpnessesAttr())
+					{
+						CornerSharpnessesAttr.Get(&InOutMeshData.CornerSharpnesses, Options.TimeCode);
+					}
+
+					if (pxr::UsdAttribute CreaseIndicesAttr = UsdMesh.GetCreaseIndicesAttr())
+					{
+						CreaseIndicesAttr.Get(&InOutMeshData.CreaseIndices, Options.TimeCode);
+					}
+
+					if (pxr::UsdAttribute CreaseLengthsAttr = UsdMesh.GetCreaseLengthsAttr())
+					{
+						CreaseLengthsAttr.Get(&InOutMeshData.CreaseLengths, Options.TimeCode);
+					}
+
+					if (pxr::UsdAttribute CreaseSharpnessesAttr = UsdMesh.GetCreaseSharpnessesAttr())
+					{
+						CreaseSharpnessesAttr.Get(&InOutMeshData.CreaseSharpnesses, Options.TimeCode);
+					}
+
+					if (pxr::UsdAttribute HoleIndicesAttr = UsdMesh.GetHoleIndicesAttr())
+					{
+						HoleIndicesAttr.Get(&InOutMeshData.HoleIndices, Options.TimeCode);
+					}
+
+					// For some reason this is not part of the USD schema so just pick the first valid token.
+					// See UsdImagingMeshAdapter::GetSubdivTags
+					InOutMeshData.CreaseMethod = pxr::PxOsdOpenSubdivTokens->uniform;
+				}
+			}
+		}
+
+		// UVs
+		{
+			TArray<TUsdStore<pxr::UsdGeomPrimvar>> PrimvarsToUse;
+
+			// If we already have a primvar to UV index assignment, let's just use that.
+			// When collapsing, we'll do a pre-pass on all meshes to translate and determine this beforehand.
+			if (InOutCombinedMaterialAssignments.PrimvarToUVIndex.Num() > 0)
+			{
+				int32 HighestProvidedUVIndex = 0;
+				for (const TPair<FString, int32>& Pair : InOutCombinedMaterialAssignments.PrimvarToUVIndex)
+				{
+					HighestProvidedUVIndex = FMath::Max(HighestProvidedUVIndex, Pair.Value);
+				}
+				InOutMeshData.ProvidedNumUVSets = HighestProvidedUVIndex + 1;
+
+				TArray<TUsdStore<pxr::UsdGeomPrimvar>> AllMeshUVPrimvars = UsdUtils::GetUVSetPrimvars(Prim, TNumericLimits<int32>::Max());
+				PrimvarsToUse = UsdUtils::AssemblePrimvarsIntoUVSets(AllMeshUVPrimvars, InOutCombinedMaterialAssignments.PrimvarToUVIndex);
+			}
+			// Let's use the best primvar assignment for this particular mesh instead
+			else
+			{
+				PrimvarsToUse = UsdUtils::GetUVSetPrimvars(Prim);
+				InOutCombinedMaterialAssignments.PrimvarToUVIndex = UsdUtils::AssemblePrimvarsIntoPrimvarToUVIndexMap(PrimvarsToUse);
+			}
+
+			// Unpack the primvars we'll be using into simple arrays so that if we want to subdivide this mesh
+			// we can just update those arrays with new data
+			InOutMeshData.UVSets.Reset(PrimvarsToUse.Num());
+			InOutMeshData.UVSetIndices.Reset(PrimvarsToUse.Num());
+			InOutMeshData.UVSetInterpolations.Reset(PrimvarsToUse.Num());
+			for (const TUsdStore<pxr::UsdGeomPrimvar>& PrimvarPtr : PrimvarsToUse)
+			{
+				pxr::VtArray<pxr::GfVec2f>& UVs = InOutMeshData.UVSets.Emplace_GetRef();
+				pxr::VtArray<int>& Indices = InOutMeshData.UVSetIndices.Emplace_GetRef();
+				pxr::TfToken& Interpolation = InOutMeshData.UVSetInterpolations.Emplace_GetRef();
+
+				// There are some code paths where it's OK to end up with an invalid primvar here:
+				// For example when collapsing two cubes and only one of them has the e.g. "st1" primvar:
+				// We will allocate a UV index for it and try reading it on both cubes, and end up with
+				// an invalid primvar in one of them, although it is important to retain the UV set
+				// ordering between the cubes
+				if (pxr::UsdGeomPrimvar Primvar = PrimvarPtr.Get())
+				{
+					ensure(Primvar.Get(&UVs, Options.TimeCode));
+					ensure(Primvar.GetIndices(&Indices, Options.TimeCode) || !Primvar.IsIndexed());
+					Interpolation = GetAttrInterpolation(Primvar, DefaultInterpolation);
+				}
+			}
+		}
+
+		// Material assignments
+		{
+			const bool bProvideMaterialIndices = true;
+			InOutMeshData.LocalMaterialInfo = UsdUtils::GetPrimMaterialAssignments(
+				Prim,
+				Options.TimeCode,
+				bProvideMaterialIndices,
+				Options.RenderContext,
+				Options.MaterialPurpose
+			);
+
+			InOutMeshData.MaterialIndexOffset = InOutCombinedMaterialAssignments.Slots.Num();
+		}
+
+		return true;
+	}
+
+
+	// OpenSubdiv expects the data elements of its buffers to implement a simple interface,
+	// so here we wrap the datatypes we'll be interpolating with that interface
+	struct FSubdivVec2f
+	{
+		pxr::GfVec2f Data;
+
+		void Clear()
+		{
+			Data = {0, 0};
+		}
+
+		void AddWithWeight(const pxr::GfVec2f& Src, float Weight)
+		{
+			Data += Src * Weight;
+		}
+		void AddWithWeight(const FSubdivVec2f& Src, float Weight)
+		{
+			Data += Src.Data * Weight;
+		}
+	};
+
+	struct FSubdivVec3f
+	{
+		pxr::GfVec3f Data;
+
+		void Clear()
+		{
+			Data = {0, 0, 0};
+		}
+
+		void AddWithWeight(const pxr::GfVec3f& Src, float Weight)
+		{
+			Data += Src * Weight;
+		}
+		void AddWithWeight(const FSubdivVec3f& Src, float Weight)
+		{
+			Data += Src.Data * Weight;
+		}
+	};
+
+	struct FSubdivInt
+	{
+		int Data;
+
+		void Clear()
+		{
+			Data = 0;
+		}
+
+		void AddWithWeight(const int& Src, float Weight)
+		{
+			Data += Src * Weight;
+		}
+		void AddWithWeight(const FSubdivInt& Src, float Weight)
+		{
+			Data += Src.Data * Weight;
+		}
+	};
+
+	struct FSubdivFloat
+	{
+		float Data;
+
+		void Clear()
+		{
+			Data = 0;
+		}
+
+		void AddWithWeight(const float& Src, float Weight)
+		{
+			Data += Src * Weight;
+		}
+		void AddWithWeight(const FSubdivFloat& Src, float Weight)
+		{
+			Data += Src.Data * Weight;
+		}
+	};
+
+	// We're going to be doing some reinterpret casting between these, so let's try our best to make sure we're safe
+	static_assert(sizeof(FSubdivVec3f) == sizeof(pxr::GfVec3f) && alignof(FSubdivVec3f) == alignof(pxr::GfVec3f));
+	static_assert(sizeof(FSubdivInt) == sizeof(int) && alignof(FSubdivInt) == alignof(int));
+	static_assert(sizeof(FSubdivFloat) == sizeof(float) && alignof(FSubdivFloat) == alignof(float));
+
+	template<typename T>
+	pxr::VtArray<T> ComputeFlattened(const pxr::VtArray<T>& Values, const pxr::VtArray<int>& Indices)
+	{
+		// Adapted from USD's _ComputeFlattened within pxr\imaging\hd\primvarSchema.cpp
+
+		FScopedUsdAllocs Allocs;
+
+		if (Indices.empty())
+		{
+			return Values;
+		}
+
+		pxr::VtArray<T> Result = pxr::VtArray<T>(Indices.size());
+
+		bool bInvalidIndices = false;
+		for (size_t Index = 0; Index < Indices.size(); ++Index)
+		{
+			int ValueIndex = Indices[Index];
+			if (ValueIndex >= 0 && (size_t)ValueIndex < Values.size())
+			{
+				Result[Index] = Values[ValueIndex];
+			}
+			else
+			{
+				Result[Index] = T();
+				bInvalidIndices = true;
+			}
+		}
+
+		if (bInvalidIndices)
+		{
+			UE_LOG(LogUsd, Warning, TEXT("Invalid primvar indices encountered in ComputeFlattened"));
+		}
+
+		return Result;
+	};
+
+	// In-place converts SharedValuesArray from an array of values that are shared according to the topology described in
+	// 'Level' into a flattened array that has a single value for each face vertex.
+	template<typename T>
+	void FlattenFaceVaryingValues(T& SharedValuesArray, int32 FaceVaryingChannel, const OpenSubdiv::Far::TopologyLevel& Level)
+	{
+		if (SharedValuesArray.size() == 0)
+		{
+			return;
+		}
+
+		T FlattenedValues;
+		FlattenedValues.reserve(Level.GetNumFaceVertices());
+		for (int32 FaceIndex = 0; FaceIndex < Level.GetNumFaces(); ++FaceIndex)
+		{
+			OpenSubdiv::Far::ConstIndexArray Face = Level.GetFaceVertices(FaceIndex);
+			OpenSubdiv::Far::ConstIndexArray FaceNormalsFaceVaryingIndices = Level.GetFaceFVarValues(FaceIndex, FaceVaryingChannel);
+
+			for (int32 FaceVertexIndex = 0; FaceVertexIndex < Face.size(); ++FaceVertexIndex)
+			{
+				int32 FaceVertexNormalsIndex = FaceNormalsFaceVaryingIndices[FaceVertexIndex];
+				FlattenedValues.push_back(SharedValuesArray[FaceVertexNormalsIndex]);
+			}
+		}
+		Swap(FlattenedValues, SharedValuesArray);
+	};
+
+	bool SubdivideMeshData(
+		const pxr::UsdPrim& Prim,
+		const UsdToUnreal::FUsdMeshConversionOptions& Options,
+		FUsdMeshData& InOutMeshData
+	)
+	{
+		// References:
+		// - USD's HdSt_OsdTopologyComputation::Resolve
+		// - UE's FSubdividePoly::ComputeTopologySubdivision
+		// - OpenSubdiv's [REDACTED_PROFILE_URL]/OpenSubdiv/blob/release/tutorials/far/tutorial_2_2/far_tutorial_2_2.cpp
+
+		FScopedUsdAllocs Allocs;
+
+		pxr::UsdGeomMesh UsdMesh{Prim};
+		if (!UsdMesh)
+		{
+			return false;
+		}
+
+		int32 TargetSubdivLevel = FMath::Max(0, FMath::Min(GMaxSubdivLevel, Options.SubdivisionLevel));
+		if (TargetSubdivLevel < Options.SubdivisionLevel)
+		{
+			UE_LOG(
+				LogUsd,
+				Warning,
+				TEXT("Max subdivision level was clamped to %d (controlled by the cvar '%s')"),
+				GMaxSubdivLevel,
+				*MaxUsdSubdivLevelCvarName
+			);
+		}
+		if (TargetSubdivLevel < 1)
+		{
+			UE_LOG(
+				LogUsd,
+				Log,
+				TEXT("Cancelling out of subdividing mesh '%s' due to target subdivision level being %d after clamping (it needs to be at least 1 for a round of subdivision)"),
+				*InOutMeshData.SourcePrimPath,
+				TargetSubdivLevel
+			);
+			return false;
+		}
+
+		UE_LOG(LogUsd, Log, TEXT("Subdividing mesh '%s' to subdivision level %d"), *InOutMeshData.SourcePrimPath, TargetSubdivLevel);
+
+		// We need to track our faceVarying attributes when subdividing, so we'll give each one an unique
+		// index into that FaceVaryingTopologies array down below. We'll use this to track how many entries that array
+		// will need, and which attribute has which index
+		int32 FaceVaryingChannelCounter = 0;
+
+		// It is very likely that if a primvar is faceVarying it will be indexed, and we should use those indices as the
+		// FaceVarying topology (https://openusd.org/release/api/class_usd_geom_primvar.html#UsdGeomPrimvar_Indexed_primvars).
+		// In case the user provided a faceVarying attribute/primvar *without* indexing however, it means each
+		// face vertex gets its own dedicated value and the topology is for them to be all "disconnected" and never share
+		// vertices, which we can represent with an index array with increasing values. We can reuse that array for all un-indexed
+		// attributes and primvars though, which is what we'll track here
+		int32 IotaFaceVaryingChannel = INDEX_NONE;
+		pxr::VtArray<int> IotaIndices;
+		const int NumFaceVertices = InOutMeshData.FaceVertexIndices.size();
+		TFunction<void()> CreateIotaIndicesIfNeeded = [&IotaIndices, &IotaFaceVaryingChannel, &FaceVaryingChannelCounter, NumFaceVertices]()
+		{
+			if (IotaFaceVaryingChannel != INDEX_NONE)
+			{
+				return;
+			}
+			IotaFaceVaryingChannel = FaceVaryingChannelCounter++;
+
+			IotaIndices.resize(NumFaceVertices);
+			for (uint32 Index = 0; Index < IotaIndices.size(); ++Index)
+			{
+				IotaIndices[Index] = Index;
+			}
+		};
+
+		// All pointsBased/Gprim attributes we'll retrieve happen to have varying default interpolation
+		const pxr::TfToken DefaultInterpolation = pxr::UsdGeomTokens->varying;
+
+		// Points
+		int32 PointsFaceVaryingChannel = INDEX_NONE;
+		if (InOutMeshData.PointIndices.size() > 0)
+		{
+			// faceVarying indices are important for the topology and OpenSubdiv needs them, so we need to keep our
+			// array indexed and flatten only after refining the mesh
+			if (InOutMeshData.PointInterpolation == pxr::UsdGeomTokens->faceVarying)
+			{
+				PointsFaceVaryingChannel = FaceVaryingChannelCounter++;
+			}
+			// Indexing on vertex, varying and uniform interpolation are just to allow reusing of the values.
+			// As far as I know there is no way to get these indices handled by OpenSubdiv (at least not through the pxOsd
+			// wrapper), and we'll end up flattening all indexing later anyway, so we might as well flatten now
+			else
+			{
+				InOutMeshData.Points = ComputeFlattened(InOutMeshData.Points, InOutMeshData.PointIndices);
+				InOutMeshData.PointIndices = {};
+			}
+		}
+		if (InOutMeshData.PointInterpolation == pxr::UsdGeomTokens->faceVarying && PointsFaceVaryingChannel == INDEX_NONE)
+		{
+			// If we're faceVarying we will need *some* indices, so create the iota indices here and use that
+			CreateIotaIndicesIfNeeded();
+			PointsFaceVaryingChannel = IotaFaceVaryingChannel;
+		}
+
+		// Normals
+		int32 NormalsFaceVaryingChannel = INDEX_NONE;
+		if (GIgnoreNormalsWhenSubdividing)
+		{
+			// According to https://openusd.org/release/api/class_usd_geom_mesh.html#UsdGeom_Mesh_Primvars,
+			// "Normals should not be authored on a subdivision mesh, since subdivision algorithms define their own normals.
+			// They should only be authored for polygonal meshes (subdivisionScheme = "none")."
+			// There is no free normal computation to be had from OpenSubdiv subdivision algoriths as far as I can tell however.
+			// We'd have to compute them manually as in [REDACTED_PROFILE_URL]/OpenSubdiv/blob/release/tutorials/far/tutorial_2_3/far_tutorial_2_3.cpp
+			// If that is the case, we may as well just ignore normals here and let RepairNormalsAndTangents fix it,
+			// since it will need to run it to compute tangents anyway
+
+			InOutMeshData.Normals = {};
+			InOutMeshData.NormalIndices = {};
+			InOutMeshData.NormalInterpolation = {};
+		}
+		else
+		{
+			if (InOutMeshData.NormalIndices.size() > 0)
+			{
+				if (InOutMeshData.NormalInterpolation == pxr::UsdGeomTokens->faceVarying)
+				{
+					NormalsFaceVaryingChannel = FaceVaryingChannelCounter++;
+				}
+				else
+				{
+					InOutMeshData.Normals = ComputeFlattened(InOutMeshData.Normals, InOutMeshData.NormalIndices);
+					InOutMeshData.NormalIndices = {};
+				}
+			}
+			if (InOutMeshData.NormalInterpolation == pxr::UsdGeomTokens->faceVarying && NormalsFaceVaryingChannel == INDEX_NONE)
+			{
+				CreateIotaIndicesIfNeeded();
+				NormalsFaceVaryingChannel = IotaFaceVaryingChannel;
+			}
+		}
+
+		// Velocities
+		int32 VelocitiesFaceVaryingChannel = INDEX_NONE;
+		if (InOutMeshData.VelocityInterpolation == pxr::UsdGeomTokens->faceVarying)
+		{
+			// Simple attributes can't be indexed, so if this is faceVarying then we know we need the iota indices
+			CreateIotaIndicesIfNeeded();
+			VelocitiesFaceVaryingChannel = IotaFaceVaryingChannel;
+		}
+
+		// DisplayColors
+		int32 DisplayColorsFaceVaryingChannel = INDEX_NONE;
+		if (InOutMeshData.DisplayColorIndices.size() > 0)
+		{
+			if (InOutMeshData.DisplayColorInterpolation == pxr::UsdGeomTokens->faceVarying)
+			{
+				DisplayColorsFaceVaryingChannel = FaceVaryingChannelCounter++;
+			}
+			else
+			{
+				InOutMeshData.DisplayColors = ComputeFlattened(InOutMeshData.DisplayColors, InOutMeshData.DisplayColorIndices);
+				InOutMeshData.DisplayColorIndices = {};
+			}
+		}
+		if (InOutMeshData.DisplayColorInterpolation == pxr::UsdGeomTokens->faceVarying && DisplayColorsFaceVaryingChannel == INDEX_NONE)
+		{
+			CreateIotaIndicesIfNeeded();
+			DisplayColorsFaceVaryingChannel = IotaFaceVaryingChannel;
+		}
+
+		// DisplayOpacities
+		int32 DisplayOpacitiesFaceVaryingChannel = INDEX_NONE;
+		if (InOutMeshData.DisplayOpacityIndices.size() > 0)
+		{
+			if (InOutMeshData.DisplayOpacityInterpolation == pxr::UsdGeomTokens->faceVarying)
+			{
+				DisplayOpacitiesFaceVaryingChannel = FaceVaryingChannelCounter++;
+			}
+			else
+			{
+				InOutMeshData.DisplayOpacities = ComputeFlattened(InOutMeshData.DisplayOpacities, InOutMeshData.DisplayOpacityIndices);
+				InOutMeshData.DisplayOpacityIndices = {};
+			}
+		}
+		if (InOutMeshData.DisplayOpacityInterpolation == pxr::UsdGeomTokens->faceVarying && DisplayOpacitiesFaceVaryingChannel == INDEX_NONE)
+		{
+			CreateIotaIndicesIfNeeded();
+			DisplayOpacitiesFaceVaryingChannel = IotaFaceVaryingChannel;
+		}
+
+		// UVs
+		const int32 NumUVSets = InOutMeshData.UVSets.Num();
+		if (!ensure(InOutMeshData.UVSetIndices.Num() == NumUVSets && InOutMeshData.UVSetInterpolations.Num() == NumUVSets))
+		{
+			return false;
+		}
+		TArray<int32> UVFaceVaryingChannels;
+		UVFaceVaryingChannels.SetNumUninitialized(NumUVSets);
+		for (int32 UVSetIndex = 0; UVSetIndex < NumUVSets; ++UVSetIndex)
+		{
+			int32& UVFaceVaryingChannel = UVFaceVaryingChannels[UVSetIndex];
+			pxr::VtArray<pxr::GfVec2f>& UVSet = InOutMeshData.UVSets[UVSetIndex];
+			pxr::TfToken& UVSetInterpolation = InOutMeshData.UVSetInterpolations[UVSetIndex];
+			pxr::VtArray<int>& UVSetIndices = InOutMeshData.UVSetIndices[UVSetIndex];
+
+			UVFaceVaryingChannel = INDEX_NONE;
+			if (UVSetIndices.size() > 0)
+			{
+				if (UVSetInterpolation == pxr::UsdGeomTokens->faceVarying)
+				{
+					UVFaceVaryingChannel = FaceVaryingChannelCounter++;
+				}
+				else
+				{
+					UVSet = ComputeFlattened(UVSet, UVSetIndices);
+					UVSetIndices = {};
+				}
+			}
+			if (UVSetInterpolation == pxr::UsdGeomTokens->faceVarying && UVFaceVaryingChannel == INDEX_NONE)
+			{
+				CreateIotaIndicesIfNeeded();
+				UVFaceVaryingChannel = IotaFaceVaryingChannel;
+			}
+		}
+
+		pxr::TfToken MaterialIndicesInterpolation = pxr::UsdGeomTokens->uniform;
+
+		pxr::PxOsdSubdivTags SubdivTags{
+			InOutMeshData.InterpolateBoundary,
+			InOutMeshData.FaceVaryingInterpolation,
+			InOutMeshData.CreaseMethod,
+			InOutMeshData.TriangleSubdivision,
+			InOutMeshData.CreaseIndices,
+			InOutMeshData.CreaseLengths,
+			InOutMeshData.CreaseSharpnesses,
+			InOutMeshData.CornerIndices,
+			InOutMeshData.CornerSharpnesses};
+
+		pxr::PxOsdMeshTopology Topology{
+			InOutMeshData.SubdivScheme,
+			InOutMeshData.Orientation,
+			InOutMeshData.FaceVertexCounts,
+			InOutMeshData.FaceVertexIndices,
+			InOutMeshData.HoleIndices,
+			SubdivTags};
+
+		std::vector<pxr::VtArray<int>> FaceVaryingTopologies;
+		FaceVaryingTopologies.resize(FaceVaryingChannelCounter);
+		if (PointsFaceVaryingChannel != INDEX_NONE)
+		{
+			FaceVaryingTopologies[PointsFaceVaryingChannel] = InOutMeshData.PointIndices;
+		}
+		if (NormalsFaceVaryingChannel != INDEX_NONE)
+		{
+			FaceVaryingTopologies[NormalsFaceVaryingChannel] = InOutMeshData.NormalIndices;
+		}
+		// No need to check Velocities here as there's no way it has custom indices
+		if (DisplayColorsFaceVaryingChannel != INDEX_NONE)
+		{
+			FaceVaryingTopologies[DisplayColorsFaceVaryingChannel] = InOutMeshData.DisplayColorIndices;
+		}
+		if (DisplayOpacitiesFaceVaryingChannel != INDEX_NONE)
+		{
+			FaceVaryingTopologies[DisplayOpacitiesFaceVaryingChannel] = InOutMeshData.DisplayOpacityIndices;
+		}
+		for (int32 UVSetIndex = 0; UVSetIndex < NumUVSets; ++UVSetIndex)
+		{
+			int32 UVFaceVaryingChannel = UVFaceVaryingChannels[UVSetIndex];
+			if (UVFaceVaryingChannel != INDEX_NONE)
+			{
+				FaceVaryingTopologies[UVFaceVaryingChannel] = InOutMeshData.UVSetIndices[UVSetIndex];
+			}
+		}
+		// Iota being last means we replace whatever else may have been placed at the iota channel with the actual indices
+		if (IotaFaceVaryingChannel != INDEX_NONE)
+		{
+			FaceVaryingTopologies[IotaFaceVaryingChannel] = IotaIndices;
+		}
+
+		std::shared_ptr<OpenSubdiv::Far::TopologyRefiner> TopologyRefiner = pxr::PxOsdRefinerFactory::Create(Topology, FaceVaryingTopologies);
+		if (!TopologyRefiner)
+		{
+			return false;
+		}
+
+		// Refine our topology (we only use uniform subdivision for now)
+		OpenSubdiv::Far::TopologyRefiner::UniformOptions UniformOptions{TargetSubdivLevel};
+		// From tutorial_2_2: "fullTopologyInLastLevel must be true to work with faceVarying data"
+		UniformOptions.fullTopologyInLastLevel = true;
+		TopologyRefiner->RefineUniform(UniformOptions);
+
+		// We're using primvar refiners here as that is the simplest method of getting our primvars
+		// subdivided, but if performance becomes an issue we could try using stencil/patch tables instead
+		OpenSubdiv::Far::PrimvarRefiner PrimvarRefiner{*TopologyRefiner};
+
+		// Temp buffers where we'll store the iterative subdivision values
+		pxr::VtArray<int> TempFaceVertexCounts;
+		pxr::VtArray<int> TempFaceVertexIndices;
+		pxr::VtArray<pxr::GfVec3f> TempPoints;
+		pxr::VtArray<pxr::GfVec3f> TempNormals;
+		pxr::VtArray<pxr::GfVec3f> TempVelocities;
+		pxr::VtArray<pxr::GfVec3f> TempDisplayColors;
+		pxr::VtArray<float> TempDisplayOpacities;
+		TArray<int32> TempMaterialIndices; // Using a TArray saves us a memcpy when outputting results
+		TArray<pxr::VtArray<pxr::GfVec2f>> TempUVSets;
+
+		// Resize the target buffers to be large enough to hold all refinements *simultaneously* (one next to the other).
+		// This is great because we can just read/write to the same buffer as we iteratively refine
+		TFunction<size_t(pxr::TfToken, int32)> GetTotalNumElements = [&TopologyRefiner](pxr::TfToken InterpolationType, int32 FaceVaryingChannel) -> size_t
+		{
+			// The "GetNumXTotal()" functions also include space for the source data as well.
+			// In our case we'll keep the source data on the actual source arrays so we don't have to
+			// copy them over, meaning our buffers can be a bit smaller too
+			const OpenSubdiv::Far::TopologyLevel& UnsubdividedLevel = TopologyRefiner->GetLevel(0);
+
+			if (InterpolationType == pxr::UsdGeomTokens->vertex)
+			{
+				return TopologyRefiner->GetNumVerticesTotal() - UnsubdividedLevel.GetNumVertices();
+			}
+			else if (InterpolationType == pxr::UsdGeomTokens->varying)
+			{
+				return TopologyRefiner->GetNumVerticesTotal() - UnsubdividedLevel.GetNumVertices();
+			}
+			else if (InterpolationType == pxr::UsdGeomTokens->faceVarying)
+			{
+				return TopologyRefiner->GetNumFVarValuesTotal(FaceVaryingChannel) - UnsubdividedLevel.GetNumFVarValues(FaceVaryingChannel);
+			}
+			else if (InterpolationType == pxr::UsdGeomTokens->uniform)
+			{
+				return TopologyRefiner->GetNumFacesTotal() - UnsubdividedLevel.GetNumFaces();
+			}
+			else if (InterpolationType == pxr::UsdGeomTokens->constant)
+			{
+				return 1;
+			}
+			else
+			{
+				return 0;
+			}
+		};
+		TempPoints.resize(GetTotalNumElements(InOutMeshData.PointInterpolation, PointsFaceVaryingChannel));
+		TempNormals.resize(GetTotalNumElements(InOutMeshData.NormalInterpolation, NormalsFaceVaryingChannel));
+		TempVelocities.resize(GetTotalNumElements(InOutMeshData.VelocityInterpolation, VelocitiesFaceVaryingChannel));
+		TempDisplayColors.resize(GetTotalNumElements(InOutMeshData.DisplayColorInterpolation, DisplayColorsFaceVaryingChannel));
+		TempDisplayOpacities.resize(GetTotalNumElements(InOutMeshData.DisplayOpacityInterpolation, DisplayOpacitiesFaceVaryingChannel));
+		TempMaterialIndices.SetNum(GetTotalNumElements(MaterialIndicesInterpolation, /*FaceVaryingChannel*/ 0));	 // Always 'uniform'
+		TempUVSets.SetNum(NumUVSets);
+		for (int32 UVSetIndex = 0; UVSetIndex < NumUVSets; ++UVSetIndex)
+		{
+			pxr::VtArray<pxr::GfVec2f>& UVSet = TempUVSets[UVSetIndex];
+			UVSet.resize(GetTotalNumElements(InOutMeshData.UVSetInterpolations[UVSetIndex], UVFaceVaryingChannels[UVSetIndex]));
+		}
+
+		// Use the right function from PrimvarRefiner depending on InterpolationType.
+		// Has to be auto as we reuse this for FSubdivVec3f*, FSubdivFloat* and FSubdivInt* SrcPtr and DstPtr
+		auto InterpolateAttribute =
+			[&PrimvarRefiner](auto SrcPtr, auto DstPtr, pxr::TfToken InterpolationType, int32 CurrentRefinementLevel, int32 FaceVaryingChannel)
+		{
+			// If the mesh doesn't have any values for an attribute, its SrcPtr will be nullptr
+			if (!SrcPtr || !DstPtr)
+			{
+				return;
+			}
+
+			if (InterpolationType == pxr::UsdGeomTokens->vertex)
+			{
+				PrimvarRefiner.Interpolate(CurrentRefinementLevel, SrcPtr, DstPtr);
+			}
+			else if (InterpolationType == pxr::UsdGeomTokens->varying)
+			{
+				PrimvarRefiner.InterpolateVarying(CurrentRefinementLevel, SrcPtr, DstPtr);
+			}
+			else if (InterpolationType == pxr::UsdGeomTokens->faceVarying)
+			{
+				PrimvarRefiner.InterpolateFaceVarying(CurrentRefinementLevel, SrcPtr, DstPtr, FaceVaryingChannel);
+			}
+			else if (InterpolationType == pxr::UsdGeomTokens->uniform)
+			{
+				PrimvarRefiner.InterpolateFaceUniform(CurrentRefinementLevel, SrcPtr, DstPtr);
+			}
+			else if (InterpolationType == pxr::UsdGeomTokens->constant)
+			{
+				// We need to move this because for the very first refinement we rely on InterpolateAttribute to
+				// move the value from the source array to the dest array
+				*DstPtr = *SrcPtr;
+			}
+		};
+
+		// Note how these start by pointing at the actual source data. After the first refinement iteration
+		// these (as well as the dst pointers) will all point at different locations within the Temp buffers
+		FSubdivVec3f* SrcPointsPtr = reinterpret_cast<FSubdivVec3f*>(InOutMeshData.Points.data());
+		FSubdivVec3f* SrcNormalsPtr = reinterpret_cast<FSubdivVec3f*>(InOutMeshData.Normals.data());
+		FSubdivVec3f* SrcVelocitiesPtr = reinterpret_cast<FSubdivVec3f*>(InOutMeshData.Velocities.data());
+		FSubdivVec3f* SrcDisplayColorsPtr = reinterpret_cast<FSubdivVec3f*>(InOutMeshData.DisplayColors.data());
+		FSubdivFloat* SrcDisplayOpacitiesPtr = reinterpret_cast<FSubdivFloat*>(InOutMeshData.DisplayOpacities.data());
+		FSubdivInt* SrcMaterialIndicesPtr = reinterpret_cast<FSubdivInt*>(InOutMeshData.LocalMaterialInfo.MaterialIndices.GetData());
+		TArray<FSubdivVec2f*> SrcUVSetPtrs;
+		SrcUVSetPtrs.SetNum(NumUVSets);
+		for (int32 UVSetIndex = 0; UVSetIndex < NumUVSets; ++UVSetIndex)
+		{
+			SrcUVSetPtrs[UVSetIndex] = reinterpret_cast<FSubdivVec2f*>(InOutMeshData.UVSets[UVSetIndex].data());
+		}
+
+		FSubdivVec3f* DstPointsPtr = reinterpret_cast<FSubdivVec3f*>(TempPoints.data());
+		FSubdivVec3f* DstNormalsPtr = reinterpret_cast<FSubdivVec3f*>(TempNormals.data());
+		FSubdivVec3f* DstVelocitiesPtr = reinterpret_cast<FSubdivVec3f*>(TempVelocities.data());
+		FSubdivVec3f* DstDisplayColorsPtr = reinterpret_cast<FSubdivVec3f*>(TempDisplayColors.data());
+		FSubdivFloat* DstDisplayOpacitiesPtr = reinterpret_cast<FSubdivFloat*>(TempDisplayOpacities.data());
+		FSubdivInt* DstMaterialIndicesPtr = reinterpret_cast<FSubdivInt*>(TempMaterialIndices.GetData());
+		TArray<FSubdivVec2f*> DstUVSetPtrs;
+		DstUVSetPtrs.SetNum(NumUVSets);
+		for (int32 UVSetIndex = 0; UVSetIndex < NumUVSets; ++UVSetIndex)
+		{
+			DstUVSetPtrs[UVSetIndex] = reinterpret_cast<FSubdivVec2f*>(TempUVSets[UVSetIndex].data());
+		}
+
+		TFunction<int(const OpenSubdiv::Far::TopologyLevel&, pxr::TfToken, int32)> GetPtrIncrement =
+			[](const OpenSubdiv::Far::TopologyLevel& Level, pxr::TfToken InterpolationType, int32 FaceVaryingChannel) -> int
+		{
+			if (InterpolationType == pxr::UsdGeomTokens->vertex)
+			{
+				return Level.GetNumVertices();
+			}
+			else if (InterpolationType == pxr::UsdGeomTokens->varying)
+			{
+				return Level.GetNumVertices();
+			}
+			else if (InterpolationType == pxr::UsdGeomTokens->faceVarying)
+			{
+				return Level.GetNumFVarValues(FaceVaryingChannel);
+			}
+			else if (InterpolationType == pxr::UsdGeomTokens->uniform)
+			{
+				return Level.GetNumFaces();
+			}
+			else
+			{
+				// For pxr::UsdGeomTokens->constant don't increment anything as we'll really just have a single value throughout
+				return 0;
+			}
+		};
+
+		// Actually refine all of our attributes/primvars
+		// Inspired by [REDACTED_PROFILE_URL]/OpenSubdiv/blob/7d0ab5530feef693ac0a920585b5c663b80773b3/tutorials/far/tutorial_2_2/far_tutorial_2_2.cpp#L293
+		// but avoiding the initial copy from the source data arrays
+		for (int32 CurrentLevel = 1; CurrentLevel <= TargetSubdivLevel; ++CurrentLevel)
+		{
+			InterpolateAttribute(SrcPointsPtr, DstPointsPtr, InOutMeshData.PointInterpolation, CurrentLevel, PointsFaceVaryingChannel);
+			InterpolateAttribute(SrcNormalsPtr, DstNormalsPtr, InOutMeshData.NormalInterpolation, CurrentLevel, NormalsFaceVaryingChannel);
+			InterpolateAttribute(SrcVelocitiesPtr, DstVelocitiesPtr, InOutMeshData.VelocityInterpolation, CurrentLevel, VelocitiesFaceVaryingChannel);
+			InterpolateAttribute(SrcDisplayColorsPtr, DstDisplayColorsPtr, InOutMeshData.DisplayColorInterpolation, CurrentLevel, DisplayColorsFaceVaryingChannel);
+			InterpolateAttribute(SrcDisplayOpacitiesPtr, DstDisplayOpacitiesPtr, InOutMeshData.DisplayOpacityInterpolation, CurrentLevel, DisplayOpacitiesFaceVaryingChannel);
+			InterpolateAttribute(SrcMaterialIndicesPtr, DstMaterialIndicesPtr, MaterialIndicesInterpolation, CurrentLevel, /*FaceVaryingChannel*/ 0); // Always 'uniform'
+			for (int32 UVSetIndex = 0; UVSetIndex < NumUVSets; ++UVSetIndex)
+			{
+				InterpolateAttribute(
+					SrcUVSetPtrs[UVSetIndex],
+					DstUVSetPtrs[UVSetIndex],
+					InOutMeshData.UVSetInterpolations[UVSetIndex],
+					CurrentLevel,
+					UVFaceVaryingChannels[UVSetIndex]
+				);
+			}
+
+			SrcPointsPtr = DstPointsPtr;
+			SrcNormalsPtr = DstNormalsPtr;
+			SrcVelocitiesPtr = DstVelocitiesPtr;
+			SrcDisplayColorsPtr = DstDisplayColorsPtr;
+			SrcDisplayOpacitiesPtr = DstDisplayOpacitiesPtr;
+			SrcMaterialIndicesPtr = DstMaterialIndicesPtr;
+			SrcUVSetPtrs = DstUVSetPtrs;
+
+			const OpenSubdiv::Far::TopologyLevel& AfterSubdiv = TopologyRefiner->GetLevel(CurrentLevel);
+			DstPointsPtr += GetPtrIncrement(AfterSubdiv, InOutMeshData.PointInterpolation, PointsFaceVaryingChannel);
+			DstNormalsPtr += GetPtrIncrement(AfterSubdiv, InOutMeshData.NormalInterpolation, NormalsFaceVaryingChannel);
+			DstVelocitiesPtr += GetPtrIncrement(AfterSubdiv, InOutMeshData.VelocityInterpolation, VelocitiesFaceVaryingChannel);
+			DstDisplayColorsPtr += GetPtrIncrement(AfterSubdiv, InOutMeshData.DisplayColorInterpolation, DisplayColorsFaceVaryingChannel);
+			DstDisplayOpacitiesPtr += GetPtrIncrement(AfterSubdiv, InOutMeshData.DisplayOpacityInterpolation, DisplayOpacitiesFaceVaryingChannel);
+			DstMaterialIndicesPtr += GetPtrIncrement(AfterSubdiv, MaterialIndicesInterpolation, /*FaceVaryingChannel*/ 0);
+			for (int32 UVSetIndex = 0; UVSetIndex < NumUVSets; ++UVSetIndex)
 			{
-				const bool bChildIsFirstPrim = false;
-				RecursivelyCollectPrimvars(
-					ChildPrim,
-					Options,
-					InOutAllPrimvars,
-					InOutPreferredPrimvars,
-					bChildIsFirstPrim
+				DstUVSetPtrs[UVSetIndex] += GetPtrIncrement(
+					AfterSubdiv,
+					InOutMeshData.UVSetInterpolations[UVSetIndex],
+					UVFaceVaryingChannels[UVSetIndex]
 				);
 			}
 		}
-	}
 
-	/**
-	 * Returns the set of primvar names that can be used for each UV index for a mesh collapsed from the subtree
-	 * starting at RootPrim.
-	 */
-	TMap<FString, int32> CollectSubtreePrimvars(
-		const pxr::UsdPrim& RootPrim,
-		const UsdToUnreal::FUsdMeshConversionOptions& Options,
-		bool bIsFirstPrim
-	)
-	{
-		TSet<FString> AllPrimvars;
-		TSet<FString> PreferredPrimvars;
+		// Shrink down the Result buffers to just contain the values from the last refinement.
+		// We use SrcPtrs here because they were left at the start of the last refinement section of each array
+		TempPoints.erase(TempPoints.begin(), reinterpret_cast<pxr::GfVec3f*>(SrcPointsPtr));
+		TempNormals.erase(TempNormals.begin(), reinterpret_cast<pxr::GfVec3f*>(SrcNormalsPtr));
+		TempVelocities.erase(TempVelocities.begin(), reinterpret_cast<pxr::GfVec3f*>(SrcVelocitiesPtr));
+		TempDisplayColors.erase(TempDisplayColors.begin(), reinterpret_cast<pxr::GfVec3f*>(SrcDisplayColorsPtr));
+		TempDisplayOpacities.erase(TempDisplayOpacities.begin(), reinterpret_cast<float*>(SrcDisplayOpacitiesPtr));
+		TempMaterialIndices.RemoveAt(0, reinterpret_cast<int32*>(SrcMaterialIndicesPtr) - &TempMaterialIndices[0]);
+		for (int32 UVSetIndex = 0; UVSetIndex < NumUVSets; ++UVSetIndex)
+		{
+			pxr::VtArray<pxr::GfVec2f>& TempUVSet = TempUVSets[UVSetIndex];
+			TempUVSet.erase(TempUVSet.begin(), reinterpret_cast<pxr::GfVec2f*>(SrcUVSetPtrs[UVSetIndex]));
+		}
 
-		RecursivelyCollectPrimvars(
-			RootPrim,
-			Options,
-			AllPrimvars,
-			PreferredPrimvars,
-			bIsFirstPrim
-		);
+		// If we're interpolating normals we have to take a pass to actually normalize them, as the primvar interpolation won't ensure that
+		if (!GIgnoreNormalsWhenSubdividing)
+		{
+			for (pxr::GfVec3f& Normal : TempNormals)
+			{
+				Normal.Normalize();
+			}
+		}
 
-		return UsdUtils::CombinePrimvarsIntoUVSets(
-			AllPrimvars,
-			PreferredPrimvars
-		);
-	}
+		// Face vertex counts and indices
+		const OpenSubdiv::Far::TopologyLevel& FinalLevel = TopologyRefiner->GetLevel(TargetSubdivLevel);
+		TempFaceVertexCounts.resize(FinalLevel.GetNumFaces());
+		TempFaceVertexIndices.reserve(FinalLevel.GetNumFaceVertices());
+		for (int32 FaceIndex = 0; FaceIndex < FinalLevel.GetNumFaces(); ++FaceIndex)
+		{
+			OpenSubdiv::Far::ConstIndexArray Face = FinalLevel.GetFaceVertices(FaceIndex);
+			TempFaceVertexCounts[FaceIndex] = Face.size();
 
-	// Unconverted, raw USD mesh data to convert into a MeshDescription
-	struct FUsdMeshData
-	{
-		// So that we can reference the prim on error messages
-		FString SourcePrimPath;
+			for(int32 FaceVertexIndex = 0; FaceVertexIndex < Face.size(); ++FaceVertexIndex)
+			{
+				TempFaceVertexIndices.push_back(Face[FaceVertexIndex]);
+			}
+		}
 
-		pxr::VtArray<int> FaceVertexCounts;
-		pxr::VtArray<int> FaceIndices;
-		pxr::VtArray<pxr::GfVec3f> Points;
-		pxr::VtArray<pxr::GfVec3f> Normals;
-		pxr::VtArray<pxr::GfVec3f> Velocities;
-		pxr::VtArray<pxr::GfVec3f> DisplayColors;
-		pxr::VtArray<float> DisplayOpacities;
+		// All faceVarying primvars can be arbitrarily indexed after subdivision (e.g. we may have 96 face vertices but end up
+		// with only 54 values for a particular primvar, because the topology allowed them to be shared). Our downstream code
+		// can't generally consume indexed stuff though, so here we flatten those primvar values to always be one per face vertex.
+		// The other interpolation types never have this issue however (e.g. 'vertex' interpolation will always output one
+		// value for each vertex)
+		if (InOutMeshData.PointInterpolation == pxr::UsdGeomTokens->faceVarying)
+		{
+			FlattenFaceVaryingValues(TempPoints, PointsFaceVaryingChannel, FinalLevel);
+		}
+		if (InOutMeshData.NormalInterpolation == pxr::UsdGeomTokens->faceVarying)
+		{
+			FlattenFaceVaryingValues(TempNormals, NormalsFaceVaryingChannel, FinalLevel);
+		}
+		if (InOutMeshData.DisplayColorInterpolation == pxr::UsdGeomTokens->faceVarying)
+		{
+			FlattenFaceVaryingValues(TempDisplayColors, DisplayColorsFaceVaryingChannel, FinalLevel);
+		}
+		if (InOutMeshData.DisplayOpacityInterpolation == pxr::UsdGeomTokens->faceVarying)
+		{
+			FlattenFaceVaryingValues(TempDisplayOpacities, DisplayOpacitiesFaceVaryingChannel, FinalLevel);
+		}
+		for (int32 UVSetIndex = 0; UVSetIndex < NumUVSets; ++UVSetIndex)
+		{
+			if (InOutMeshData.UVSetInterpolations[UVSetIndex] == pxr::UsdGeomTokens->faceVarying)
+			{
+				// Note: Our downstream ConvertMeshData code *can* handle indexed UV sets, but it's simpler
+				// to just flatten discard them here. We could revisit this later and evaluate the impact on
+				// performance of keeping indices around though
+				FlattenFaceVaryingValues(TempUVSets[UVSetIndex], UVFaceVaryingChannels[UVSetIndex], FinalLevel);
+			}
 
-		pxr::TfToken Orientation = pxr::UsdGeomTokens->rightHanded;
-		pxr::TfToken NormalInterpolation = pxr::UsdGeomTokens->vertex;
-		pxr::TfToken VelocityInterpolation = pxr::UsdGeomTokens->vertex;
-		pxr::TfToken DisplayColorInterpolation = pxr::UsdGeomTokens->constant;
-		pxr::TfToken DisplayOpacityInterpolation = pxr::UsdGeomTokens->constant;
+			InOutMeshData.UVSetIndices[UVSetIndex] = {};
+		}
+		InOutMeshData.PointIndices = {};
+		InOutMeshData.NormalIndices = {};
+		InOutMeshData.DisplayColorIndices = {};
+		InOutMeshData.DisplayOpacityIndices = {};
 
-		TOptional<int32> ProvidedNumUVSets;
-		TArray<TUsdStore<pxr::UsdGeomPrimvar>> PrimvarsByUVIndex;
+		// Output results
+		Swap(TempFaceVertexCounts, InOutMeshData.FaceVertexCounts);
+		Swap(TempFaceVertexIndices, InOutMeshData.FaceVertexIndices);
+		Swap(TempPoints, InOutMeshData.Points);
+		Swap(TempNormals, InOutMeshData.Normals);
+		Swap(TempVelocities, InOutMeshData.Velocities);
+		Swap(TempDisplayColors, InOutMeshData.DisplayColors);
+		Swap(TempDisplayOpacities, InOutMeshData.DisplayOpacities);
+		Swap(TempMaterialIndices, InOutMeshData.LocalMaterialInfo.MaterialIndices);
+		Swap(TempUVSets, InOutMeshData.UVSets);
 
-		int32 MaterialIndexOffset = 0;
-		UsdUtils::FUsdPrimMaterialAssignmentInfo LocalMaterialInfo;
-	};
+		return true;
+	}
 
 	bool ConvertMeshData(
 		const FUsdMeshData& InMeshData,
@@ -1089,7 +2100,7 @@ namespace UE::UsdGeomMeshConversion::Private
 		{
 			return false;
 		}
-		if (InMeshData.FaceIndices.size() < 1)
+		if (InMeshData.FaceVertexIndices.size() < 1)
 		{
 			return false;
 		}
@@ -1159,40 +2170,30 @@ namespace UE::UsdGeomMeshConversion::Private
 			TArray<FUVSet> UVSets;
 
 			int32 HighestAddedUVChannel = 0;
-			for (int32 UVChannelIndex = 0; UVChannelIndex < InMeshData.PrimvarsByUVIndex.Num(); ++UVChannelIndex)
+			for (int32 UVChannelIndex = 0; UVChannelIndex < InMeshData.UVSets.Num(); ++UVChannelIndex)
 			{
-				const pxr::UsdGeomPrimvar& Primvar = InMeshData.PrimvarsByUVIndex[UVChannelIndex].Get();
-				if (!Primvar)
-				{
-					continue;
-				}
-
 				FUVSet UVSet;
-				UVSet.InterpType = Primvar.GetInterpolation();
+				UVSet.InterpType = InMeshData.UVSetInterpolations[UVChannelIndex];
 				UVSet.UVSetIndexUE = UVChannelIndex;
 
-				if (Primvar.IsIndexed())
+				if (InMeshData.UVSetIndices[UVChannelIndex].size() > 0)
 				{
-					UVSet.UVIndices.Emplace();
+					UVSet.UVIndices = InMeshData.UVSetIndices[UVChannelIndex];
+					UVSet.UVs = InMeshData.UVSets[UVChannelIndex];
 
-					if (Primvar.GetIndices(&UVSet.UVIndices.GetValue(), InOptions.TimeCode) && Primvar.Get(&UVSet.UVs, InOptions.TimeCode))
+					if (UVSet.UVs.size() > 0)
 					{
-						if (UVSet.UVs.size() > 0)
-						{
-							UVSets.Add(MoveTemp(UVSet));
-							HighestAddedUVChannel = UVSet.UVSetIndexUE;
-						}
+						UVSets.Add(MoveTemp(UVSet));
+						HighestAddedUVChannel = UVSet.UVSetIndexUE;
 					}
 				}
 				else
 				{
-					if (Primvar.Get(&UVSet.UVs))
+					UVSet.UVs = InMeshData.UVSets[UVChannelIndex];
+					if (UVSet.UVs.size() > 0)
 					{
-						if (UVSet.UVs.size() > 0)
-						{
-							UVSets.Add(MoveTemp(UVSet));
-							HighestAddedUVChannel = UVSet.UVSetIndexUE;
-						}
+						UVSets.Add(MoveTemp(UVSet));
+						HighestAddedUVChannel = UVSet.UVSetIndexUE;
 					}
 				}
 			}
@@ -1244,7 +2245,7 @@ namespace UE::UsdGeomMeshConversion::Private
 				{
 					int32 VertexInstanceIndex = VertexInstanceOffset + CurrentVertexInstanceIndex;
 					const FVertexInstanceID VertexInstanceID(VertexInstanceIndex);
-					const int32 ControlPointIndex = InMeshData.FaceIndices.cdata()[CurrentVertexInstanceIndex];
+					const int32 ControlPointIndex = InMeshData.FaceVertexIndices.cdata()[CurrentVertexInstanceIndex];
 					const FVertexID VertexID(VertexOffset + ControlPointIndex);
 
 					// This data is read straight from USD so there's nothing guaranteeing we have as many positions as we need
@@ -1432,6 +2433,7 @@ namespace UsdToUnreal
 		, MaterialPurpose( pxr::UsdShadeTokens->allPurpose )
 		, TimeCode( pxr::UsdTimeCode::EarliestTime() )
 		, bMergeIdenticalMaterialSlots( true )
+		, SubdivisionLevel(0)
 	{
 	}
 }
@@ -1443,9 +2445,9 @@ bool UsdToUnreal::ConvertGeomMesh(
 	const FUsdMeshConversionOptions& Options
 )
 {
-	TRACE_CPUPROFILER_EVENT_SCOPE( UsdToUnreal::ConvertGeomMesh );
+	TRACE_CPUPROFILER_EVENT_SCOPE(UsdToUnreal::ConvertGeomMesh);
 
-	if ( !UsdMesh )
+	if (!UsdMesh)
 	{
 		return false;
 	}
@@ -1453,99 +2455,18 @@ bool UsdToUnreal::ConvertGeomMesh(
 	FScopedUsdAllocs Allocs;
 
 	pxr::UsdPrim UsdPrim = UsdMesh.GetPrim();
-	pxr::UsdStageRefPtr Stage = UsdPrim.GetStage();
-	const FUsdStageInfo StageInfo(Stage);
 
 	UsdGeomMeshImpl::FUsdMeshData MeshData;
-	MeshData.SourcePrimPath = UsdToUnreal::ConvertPath(UsdPrim.GetPrimPath());
-
-	// Face counts
-	if (pxr::UsdAttribute Attr = UsdMesh.GetFaceVertexCountsAttr())
-	{
-		Attr.Get(&MeshData.FaceVertexCounts, Options.TimeCode);
-	}
-
-	// Face indices
-	if (pxr::UsdAttribute Attr = UsdMesh.GetFaceVertexIndicesAttr())
-	{
-		Attr.Get(&MeshData.FaceIndices, Options.TimeCode);
-	}
-
-	// Points
-	if (pxr::UsdAttribute Attr = UsdMesh.GetPointsAttr())
-	{
-		Attr.Get(&MeshData.Points, Options.TimeCode);
-	}
-
-	// Normals
-	if (pxr::UsdAttribute Attr = UsdMesh.GetNormalsAttr())
-	{
-		Attr.Get(&MeshData.Normals, Options.TimeCode);
-		MeshData.NormalInterpolation = UsdGeomMeshImpl::GetAttrInterpolation(Attr, pxr::UsdGeomTokens->vertex);
-	}
-
-	// Velocities
-	if (pxr::UsdAttribute Attr = UsdMesh.GetVelocitiesAttr())
-	{
-		Attr.Get(&MeshData.Velocities, Options.TimeCode);
-		MeshData.VelocityInterpolation = UsdGeomMeshImpl::GetAttrInterpolation(Attr, pxr::UsdGeomTokens->vertex);
-	}
-
-	// Vertex colors
-	if (pxr::UsdGeomPrimvar Primvar = pxr::UsdGeomPrimvar(UsdPrim.GetAttribute(pxr::UsdGeomTokens->primvarsDisplayColor)))
-	{
-		Primvar.ComputeFlattened(&MeshData.DisplayColors, Options.TimeCode);
-		MeshData.DisplayColorInterpolation = UsdGeomMeshImpl::GetAttrInterpolation(Primvar);
-	}
-
-	// Vertex opacity
-	if (pxr::UsdGeomPrimvar Primvar = pxr::UsdGeomPrimvar(UsdPrim.GetAttribute(pxr::UsdGeomTokens->primvarsDisplayOpacity)))
-	{
-		Primvar.ComputeFlattened(&MeshData.DisplayOpacities, Options.TimeCode);
-		MeshData.DisplayOpacityInterpolation = UsdGeomMeshImpl::GetAttrInterpolation(Primvar);
-	}
-
-	// UVs
-	{
-		// If we already have a primvar to UV index assignment, let's just use that.
-		// When collapsing, we'll do a pre-pass on all meshes to translate and determine this beforehand.
-		if (OutMaterialAssignments.PrimvarToUVIndex.Num() > 0)
-		{
-			int32 HighestProvidedUVIndex = 0;
-			for (const TPair<FString, int32>& Pair : OutMaterialAssignments.PrimvarToUVIndex)
-			{
-				HighestProvidedUVIndex = FMath::Max(HighestProvidedUVIndex, Pair.Value);
-			}
-			MeshData.ProvidedNumUVSets = HighestProvidedUVIndex + 1;
 
-			TArray<TUsdStore<pxr::UsdGeomPrimvar>> AllMeshUVPrimvars = UsdUtils::GetUVSetPrimvars(UsdPrim, TNumericLimits<int32>::Max());
-			MeshData.PrimvarsByUVIndex = UsdUtils::AssemblePrimvarsIntoUVSets(AllMeshUVPrimvars, OutMaterialAssignments.PrimvarToUVIndex);
-		}
-		// Let's use the best primvar assignment for this particular mesh instead
-		else
-		{
-			MeshData.PrimvarsByUVIndex = UsdUtils::GetUVSetPrimvars(UsdPrim);
-			OutMaterialAssignments.PrimvarToUVIndex = UsdUtils::AssemblePrimvarsIntoPrimvarToUVIndexMap(MeshData.PrimvarsByUVIndex);
-		}
-	}
-
-	// Orientation
-	MeshData.Orientation = UsdGeomMeshImpl::GetGprimOrientation(UsdMesh, Options.TimeCode);
+	UsdGeomMeshImpl::CollectMeshData(UsdPrim, Options, MeshData, OutMaterialAssignments);
 
-	// Material assignments
+	if (Options.SubdivisionLevel > 0 && MeshData.SubdivScheme != pxr::UsdGeomTokens->none)
 	{
-		const bool bProvideMaterialIndices = true;
-		MeshData.LocalMaterialInfo = UsdUtils::GetPrimMaterialAssignments(
-			UsdPrim,
-			Options.TimeCode,
-			bProvideMaterialIndices,
-			Options.RenderContext,
-			Options.MaterialPurpose
-		);
-
-		MeshData.MaterialIndexOffset = OutMaterialAssignments.Slots.Num();
+		UsdGeomMeshImpl::SubdivideMeshData(UsdPrim, Options, MeshData);
 	}
 
+	pxr::UsdStageRefPtr Stage = UsdPrim.GetStage();
+	const FUsdStageInfo StageInfo(Stage);
 	return UsdGeomMeshImpl::ConvertMeshData(MeshData, StageInfo, Options, OutMeshDescription, OutMaterialAssignments);
 }
 
@@ -1803,159 +2724,93 @@ bool UsdToUnreal::ConvertGeomPrimitive(
 
 	FScopedUsdAllocs UsdAllocs;
 
-	pxr::UsdStageRefPtr Stage = InPrim.GetStage();
-	const FUsdStageInfo StageInfo(Stage);
+	UsdGeomMeshImpl::FUsdMeshData MeshData;
 
-	// Remember that USD arrays are copy-on-write, so these are both "pointers", as long as we
-	// don't try writing (or using non-const operator[]) from PrimitivePoints
-	pxr::VtVec3fArray PrimitivePoints;
-	const pxr::PxOsdMeshTopology* PrimitiveTopology = nullptr;
+	// Collect all attributes authored as usual
+	UsdGeomMeshImpl::CollectMeshData(InPrim, InOptions, MeshData, InOutMaterialAssignments);
 
-	if (pxr::UsdGeomCapsule Capsule = pxr::UsdGeomCapsule{InPrim})
+	// Generate primitive points and topology on-demand
 	{
-		pxr::TfToken Axis = pxr::UsdGeomTokens->z;
-		if (pxr::UsdAttribute Attr = Capsule.GetAxisAttr())
-		{
-			Axis = UsdUtils::GetUsdValue<pxr::TfToken>(Attr, pxr::UsdTimeCode::Default());
-		}
+		// Remember that USD arrays are copy-on-write, so these are both "pointers", as long as we
+		// don't try writing (or using non-const operator[]) from PrimitivePoints
+		pxr::VtVec3fArray PrimitivePoints;
+		const pxr::PxOsdMeshTopology* PrimitiveTopology = nullptr;
 
-		PrimitivePoints = pxr::UsdImagingGenerateCapsuleMeshPoints(
-			UsdGeomMeshImpl::DefaultCapsuleMeshHeight,
-			UsdGeomMeshImpl::DefaultCapsuleMeshRadius,
-			Axis
-		);
-		PrimitiveTopology = &pxr::UsdImagingGetCapsuleMeshTopology();
-	}
-	else if (pxr::UsdGeomCone Cone = pxr::UsdGeomCone{InPrim})
-	{
-		pxr::TfToken Axis = pxr::UsdGeomTokens->z;
-		if (pxr::UsdAttribute Attr = Capsule.GetAxisAttr())
+		if (pxr::UsdGeomCapsule Capsule = pxr::UsdGeomCapsule{InPrim})
 		{
-			Axis = UsdUtils::GetUsdValue<pxr::TfToken>(Attr, pxr::UsdTimeCode::Default());
-		}
+			pxr::TfToken Axis = pxr::UsdGeomTokens->z;
+			if (pxr::UsdAttribute Attr = Capsule.GetAxisAttr())
+			{
+				Axis = UsdUtils::GetUsdValue<pxr::TfToken>(Attr, pxr::UsdTimeCode::Default());
+			}
 
-		PrimitivePoints = UsdGeomMeshImpl::GetUnitConeMeshPoints(Axis);
-		PrimitiveTopology = &pxr::UsdImagingGetUnitConeMeshTopology();
-	}
-	else if (pxr::UsdGeomCube Cube = pxr::UsdGeomCube{InPrim})
-	{
-		PrimitivePoints = pxr::UsdImagingGetUnitCubeMeshPoints();
-		PrimitiveTopology = &pxr::UsdImagingGetUnitCubeMeshTopology();
-	}
-	else if (pxr::UsdGeomCylinder Cylinder = pxr::UsdGeomCylinder{InPrim})
-	{
-		pxr::TfToken Axis = pxr::UsdGeomTokens->z;
-		if (pxr::UsdAttribute Attr = Capsule.GetAxisAttr())
-		{
-			Axis = UsdUtils::GetUsdValue<pxr::TfToken>(Attr, pxr::UsdTimeCode::Default());
+			PrimitivePoints = pxr::UsdImagingGenerateCapsuleMeshPoints(
+				UsdGeomMeshImpl::DefaultCapsuleMeshHeight,
+				UsdGeomMeshImpl::DefaultCapsuleMeshRadius,
+				Axis
+			);
+			PrimitiveTopology = &pxr::UsdImagingGetCapsuleMeshTopology();
 		}
-
-		PrimitivePoints = UsdGeomMeshImpl::GetUnitCylinderMeshPoints(Axis);
-		PrimitiveTopology = &pxr::UsdImagingGetUnitCylinderMeshTopology();
-	}
-	else if (pxr::UsdGeomPlane Plane = pxr::UsdGeomPlane{InPrim})
-	{
-		const double Width = 1.0f;
-		const double Length = 1.0f;
-
-		pxr::TfToken Axis = pxr::UsdGeomTokens->z;
-		if (pxr::UsdAttribute Attr = Capsule.GetAxisAttr())
+		else if (pxr::UsdGeomCone Cone = pxr::UsdGeomCone{InPrim})
 		{
-			Axis = UsdUtils::GetUsdValue<pxr::TfToken>(Attr, pxr::UsdTimeCode::Default());
-		}
-
-		PrimitivePoints = pxr::UsdImagingGeneratePlaneMeshPoints(Width, Length, Axis);
-		PrimitiveTopology = &pxr::UsdImagingGetPlaneTopology();
-	}
-	else if (pxr::UsdGeomSphere Sphere = pxr::UsdGeomSphere{InPrim})
-	{
-		PrimitivePoints = pxr::UsdImagingGetUnitSphereMeshPoints();
-		PrimitiveTopology = &pxr::UsdImagingGetUnitSphereMeshTopology();
-	}
-
-	if (!PrimitiveTopology || PrimitivePoints.empty())
-	{
-		return false;
-	}
-
-	UsdGeomMeshImpl::FUsdMeshData MeshData;
-	MeshData.SourcePrimPath = UsdToUnreal::ConvertPath(InPrim.GetPrimPath());
-	MeshData.FaceVertexCounts = PrimitiveTopology->GetFaceVertexCounts();
-	MeshData.FaceIndices = PrimitiveTopology->GetFaceVertexIndices();
-	MeshData.Points = PrimitivePoints;
-
-	// Normals
-	if (pxr::UsdAttribute Attr = InPrim.GetAttribute(pxr::UsdGeomTokens->normals))
-	{
-		Attr.Get(&MeshData.Normals, InOptions.TimeCode);
-		MeshData.NormalInterpolation = UsdGeomMeshImpl::GetAttrInterpolation(Attr, pxr::UsdGeomTokens->vertex);
-	}
+			pxr::TfToken Axis = pxr::UsdGeomTokens->z;
+			if (pxr::UsdAttribute Attr = Capsule.GetAxisAttr())
+			{
+				Axis = UsdUtils::GetUsdValue<pxr::TfToken>(Attr, pxr::UsdTimeCode::Default());
+			}
 
-	// Velocities
-	if (pxr::UsdGeomPointBased PointBased{InPrim})
-	{
-		if (pxr::UsdAttribute Attr = PointBased.GetVelocitiesAttr())
+			PrimitivePoints = UsdGeomMeshImpl::GetUnitConeMeshPoints(Axis);
+			PrimitiveTopology = &pxr::UsdImagingGetUnitConeMeshTopology();
+		}
+		else if (pxr::UsdGeomCube Cube = pxr::UsdGeomCube{InPrim})
 		{
-			Attr.Get(&MeshData.Velocities, InOptions.TimeCode);
-			MeshData.VelocityInterpolation = UsdGeomMeshImpl::GetAttrInterpolation(Attr, pxr::UsdGeomTokens->vertex);
+			PrimitivePoints = pxr::UsdImagingGetUnitCubeMeshPoints();
+			PrimitiveTopology = &pxr::UsdImagingGetUnitCubeMeshTopology();
 		}
-	}
-
-	// Vertex colors
-	if (pxr::UsdGeomPrimvar Primvar = pxr::UsdGeomPrimvar(InPrim.GetAttribute(pxr::UsdGeomTokens->primvarsDisplayColor)))
-	{
-		Primvar.ComputeFlattened(&MeshData.DisplayColors, InOptions.TimeCode);
-		MeshData.DisplayColorInterpolation = UsdGeomMeshImpl::GetAttrInterpolation(Primvar);
-	}
-
-	// Vertex opacity
-	if (pxr::UsdGeomPrimvar Primvar = pxr::UsdGeomPrimvar(InPrim.GetAttribute(pxr::UsdGeomTokens->primvarsDisplayOpacity)))
-	{
-		Primvar.ComputeFlattened(&MeshData.DisplayOpacities, InOptions.TimeCode);
-		MeshData.DisplayOpacityInterpolation = UsdGeomMeshImpl::GetAttrInterpolation(Primvar);
-	}
+		else if (pxr::UsdGeomCylinder Cylinder = pxr::UsdGeomCylinder{InPrim})
+		{
+			pxr::TfToken Axis = pxr::UsdGeomTokens->z;
+			if (pxr::UsdAttribute Attr = Capsule.GetAxisAttr())
+			{
+				Axis = UsdUtils::GetUsdValue<pxr::TfToken>(Attr, pxr::UsdTimeCode::Default());
+			}
 
-	// UVs
-	{
-		// If we already have a primvar to UV index assignment, let's just use that.
-		// When collapsing, we'll do a pre-pass on all meshes to translate and determine this beforehand.
-		if (InOutMaterialAssignments.PrimvarToUVIndex.Num() > 0)
+			PrimitivePoints = UsdGeomMeshImpl::GetUnitCylinderMeshPoints(Axis);
+			PrimitiveTopology = &pxr::UsdImagingGetUnitCylinderMeshTopology();
+		}
+		else if (pxr::UsdGeomPlane Plane = pxr::UsdGeomPlane{InPrim})
 		{
-			int32 HighestProvidedUVIndex = 0;
-			for (const TPair<FString, int32>& Pair : InOutMaterialAssignments.PrimvarToUVIndex)
+			const double Width = 1.0f;
+			const double Length = 1.0f;
+
+			pxr::TfToken Axis = pxr::UsdGeomTokens->z;
+			if (pxr::UsdAttribute Attr = Capsule.GetAxisAttr())
 			{
-				HighestProvidedUVIndex = FMath::Max(HighestProvidedUVIndex, Pair.Value);
+				Axis = UsdUtils::GetUsdValue<pxr::TfToken>(Attr, pxr::UsdTimeCode::Default());
 			}
-			MeshData.ProvidedNumUVSets = HighestProvidedUVIndex + 1;
 
-			TArray<TUsdStore<pxr::UsdGeomPrimvar>> AllMeshUVPrimvars = UsdUtils::GetUVSetPrimvars(InPrim, TNumericLimits<int32>::Max());
-			MeshData.PrimvarsByUVIndex = UsdUtils::AssemblePrimvarsIntoUVSets(AllMeshUVPrimvars, InOutMaterialAssignments.PrimvarToUVIndex);
+			PrimitivePoints = pxr::UsdImagingGeneratePlaneMeshPoints(Width, Length, Axis);
+			PrimitiveTopology = &pxr::UsdImagingGetPlaneTopology();
 		}
-		// Let's use the best primvar assignment for this particular mesh instead
-		else
+		else if (pxr::UsdGeomSphere Sphere = pxr::UsdGeomSphere{InPrim})
 		{
-			MeshData.PrimvarsByUVIndex = UsdUtils::GetUVSetPrimvars(InPrim);
-			InOutMaterialAssignments.PrimvarToUVIndex = UsdUtils::AssemblePrimvarsIntoPrimvarToUVIndexMap(MeshData.PrimvarsByUVIndex);
+			PrimitivePoints = pxr::UsdImagingGetUnitSphereMeshPoints();
+			PrimitiveTopology = &pxr::UsdImagingGetUnitSphereMeshTopology();
 		}
-	}
 
-	// Orientation
-	MeshData.Orientation = UsdGeomMeshImpl::GetGprimOrientation(pxr::UsdGeomGprim{InPrim}, InOptions.TimeCode);
-
-	// Material assignments
-	{
-		const bool bProvideMaterialIndices = true;
-		MeshData.LocalMaterialInfo = UsdUtils::GetPrimMaterialAssignments(
-			InPrim,
-			InOptions.TimeCode,
-			bProvideMaterialIndices,
-			InOptions.RenderContext,
-			InOptions.MaterialPurpose
-		);
+		if (!PrimitiveTopology || PrimitivePoints.empty())
+		{
+			return false;
+		}
 
-		MeshData.MaterialIndexOffset = InOutMaterialAssignments.Slots.Num();
+		MeshData.FaceVertexCounts = PrimitiveTopology->GetFaceVertexCounts();
+		MeshData.FaceVertexIndices = PrimitiveTopology->GetFaceVertexIndices();
+		MeshData.Points = PrimitivePoints;
+		MeshData.PointInterpolation = pxr::UsdGeomTokens->vertex;
 	}
 
+	pxr::UsdStageRefPtr Stage = InPrim.GetStage();
+	const FUsdStageInfo StageInfo(Stage);
 	return ConvertMeshData(MeshData, StageInfo, InOptions, InOutMeshDescription, InOutMaterialAssignments);
 }
 
@@ -3506,54 +4361,110 @@ void UsdUtils::ReplaceUnrealMaterialsWithBaked(
 	);
 }
 
-FString UsdUtils::HashGeomMeshPrim( const UE::FUsdStage& Stage, const FString& PrimPath, double TimeCode )
+namespace UE::UsdGeomMeshConversion::Private
 {
-	TRACE_CPUPROFILER_EVENT_SCOPE( UsdUtils::HashGeomMeshPrim );
-
-	using namespace pxr;
-
-	FScopedUsdAllocs Allocs;
+	template<typename T>
+	inline void HashArrayAttribute(FMD5& MD5, const pxr::UsdAttribute& Attribute, double TimeCode)
+	{
+		if (Attribute)
+		{
+			pxr::VtArray<T> Value;
+			Attribute.Get(&Value, TimeCode);
+			MD5.Update((uint8*)Value.cdata(), Value.size() * sizeof(T));
+		}
+	}
 
-	UsdPrim UsdPrim = pxr::UsdPrim( Stage.GetPrimAtPath( UE::FSdfPath( *PrimPath ) ) );
-	if ( !UsdPrim )
+	template<typename T>
+	inline void HashArrayPrimvar(FMD5& MD5, const pxr::UsdGeomPrimvar& Primvar, double TimeCode)
 	{
-		return {};
+		if (Primvar)
+		{
+			pxr::VtArray<T> Value;
+			Primvar.Get(&Value, TimeCode);
+			MD5.Update((uint8*)Value.cdata(), Value.size() * sizeof(T));
+
+			pxr::VtArray<int> Indices;
+			if (Primvar.GetIndices(&Indices, TimeCode))
+			{
+				MD5.Update((uint8*)Indices.cdata(), Indices.size() * sizeof(int));
+			}
+		}
 	}
 
-	UsdGeomMesh UsdMesh( UsdPrim );
-	if ( !UsdMesh )
+	inline void HashTokenAttribute(FMD5& MD5, const pxr::UsdAttribute& Attribute, double TimeCode)
 	{
-		return {};
+		if (Attribute)
+		{
+			pxr::TfToken Token;
+			Attribute.Get(&Token, TimeCode);
+			MD5.Update(reinterpret_cast<const uint8*>(Token.data()), Token.size());
+		}
 	}
+}
+
+FString UsdUtils::HashGeomMeshPrim(const UE::FUsdStage& Stage, const FString& PrimPath, double TimeCode)
+{
+	TRACE_CPUPROFILER_EVENT_SCOPE(UsdUtils::HashGeomMeshPrim);
 
 	FMD5 MD5;
 
-	if ( UsdAttribute Points = UsdMesh.GetPointsAttr() )
+	HashGeomMeshPrim(Stage, PrimPath, TimeCode, MD5);
+
+	uint8 Digest[16];
+	MD5.Final(Digest);
+
+	FString Hash;
+	for (int32 i = 0; i < 16; ++i)
 	{
-		VtArray< GfVec3f > PointsArray;
-		Points.Get( &PointsArray, TimeCode );
-		MD5.Update( ( uint8* ) PointsArray.cdata(), PointsArray.size() * sizeof( GfVec3f ) );
+		Hash += FString::Printf(TEXT("%02x"), Digest[i]);
 	}
+	return Hash;
+}
 
-	if ( UsdAttribute NormalsAttribute = UsdMesh.GetNormalsAttr() )
+void UsdUtils::HashGeomMeshPrim(const UE::FUsdStage& Stage, const FString& PrimPath, double TimeCode, FMD5& InOutHashState)
+{
+	TRACE_CPUPROFILER_EVENT_SCOPE(UsdUtils::HashGeomMeshPrim);
+
+	using namespace pxr;
+	using namespace UsdGeomMeshImpl;
+
+	FScopedUsdAllocs Allocs;
+
+	UsdPrim UsdPrim = pxr::UsdPrim( Stage.GetPrimAtPath( UE::FSdfPath( *PrimPath ) ) );
+	if ( !UsdPrim )
 	{
-		VtArray< GfVec3f > Normals;
-		NormalsAttribute.Get( &Normals, TimeCode );
-		MD5.Update( ( uint8* ) Normals.cdata(), Normals.size() * sizeof( GfVec3f ) );
+		return;
 	}
 
-	if ( UsdGeomPrimvar ColorPrimvar = UsdMesh.GetDisplayColorPrimvar() )
+	UsdGeomMesh UsdMesh( UsdPrim );
+	if ( !UsdMesh )
 	{
-		VtArray< GfVec3f > UsdColors;
-		ColorPrimvar.ComputeFlattened( &UsdColors, TimeCode );
-		MD5.Update( ( uint8* ) UsdColors.cdata(), UsdColors.size() * sizeof( GfVec3f ) );
+		return;
 	}
 
-	if ( UsdGeomPrimvar OpacityPrimvar = UsdMesh.GetDisplayOpacityPrimvar() )
+	HashArrayAttribute<GfVec3f>(InOutHashState, UsdMesh.GetPointsAttr(), TimeCode);
+	HashArrayAttribute<GfVec3f>(InOutHashState, UsdMesh.GetNormalsAttr(), TimeCode);
+	HashArrayPrimvar<GfVec3f>(InOutHashState, UsdMesh.GetDisplayColorPrimvar(), TimeCode);
+	HashArrayPrimvar<float>(InOutHashState, UsdMesh.GetDisplayOpacityPrimvar(), TimeCode);
+
+	// Note: The actual subdivision level used is not factored in here because currently the single caller of this function is
+	// GetUsdStreamDDCKey, which hashes it directly. The Static/Skeletal mesh code paths won't need to currently hash it directly
+	// because the generated FMeshDescription or FSkeletalMeshImportData is hashed directly, and by then we already have subdivided
+	// the mesh data, and what we end up with will naturally depend on the level of subdivision
+
+	pxr::UsdAttribute SubdivSchemeAttr = UsdMesh.GetSubdivisionSchemeAttr();
+	pxr::TfToken SubdivScheme;
+	if (SubdivSchemeAttr && SubdivSchemeAttr.Get(&SubdivScheme, TimeCode) && SubdivScheme != pxr::UsdGeomTokens->none)
 	{
-		VtArray< float > UsdOpacities;
-		OpacityPrimvar.ComputeFlattened( &UsdOpacities );
-		MD5.Update( ( uint8* ) UsdOpacities.cdata(), UsdOpacities.size() * sizeof( float ) );
+		HashTokenAttribute(InOutHashState, UsdMesh.GetSubdivisionSchemeAttr(), TimeCode);
+		HashTokenAttribute(InOutHashState, UsdMesh.GetFaceVaryingLinearInterpolationAttr(), TimeCode);
+		HashTokenAttribute(InOutHashState, UsdMesh.GetTriangleSubdivisionRuleAttr(), TimeCode);
+		HashArrayAttribute<int>(InOutHashState, UsdMesh.GetCornerIndicesAttr(), TimeCode);
+		HashArrayAttribute<float>(InOutHashState, UsdMesh.GetCornerSharpnessesAttr(), TimeCode);
+		HashArrayAttribute<int>(InOutHashState, UsdMesh.GetCreaseIndicesAttr(), TimeCode);
+		HashArrayAttribute<int>(InOutHashState, UsdMesh.GetCreaseLengthsAttr(), TimeCode);
+		HashArrayAttribute<float>(InOutHashState, UsdMesh.GetCreaseSharpnessesAttr(), TimeCode);
+		HashArrayAttribute<int>(InOutHashState, UsdMesh.GetHoleIndicesAttr(), TimeCode);
 	}
 
 	// TODO: This is not providing render context or material purpose, so it will never consider float2f primvars
@@ -3566,26 +4477,8 @@ FString UsdUtils::HashGeomMeshPrim( const UE::FUsdStage& Stage, const FString& P
 			break;
 		}
 
-		UsdGeomPrimvar& Primvar = PrimvarsByUVIndex[ UVChannelIndex ].Get();
-		if ( !Primvar )
-		{
-			continue;
-		}
-
-		VtArray< GfVec2f > UsdUVs;
-		Primvar.Get( &UsdUVs, TimeCode );
-		MD5.Update( ( uint8* ) UsdUVs.cdata(), UsdUVs.size() * sizeof( GfVec2f ) );
-	}
-
-	uint8 Digest[ 16 ];
-	MD5.Final( Digest );
-
-	FString Hash;
-	for ( int32 i = 0; i < 16; ++i )
-	{
-		Hash += FString::Printf( TEXT( "%02x" ), Digest[ i ] );
+		HashArrayPrimvar<GfVec2f>(InOutHashState, PrimvarsByUVIndex[UVChannelIndex].Get(), TimeCode);
 	}
-	return Hash;
 }
 
 bool UsdUtils::GetPointInstancerTransforms(const FUsdStageInfo& StageInfo, const pxr::UsdGeomPointInstancer& PointInstancer, const int32 ProtoIndex, pxr::UsdTimeCode EvalTime, TArray<FTransform>& OutInstanceTransforms)
diff --git a/Engine/Plugins/Importers/USDImporter/Source/USDUtilities/Public/USDGeomMeshConversion.h b/Engine/Plugins/Importers/USDImporter/Source/USDUtilities/Public/USDGeomMeshConversion.h
--- a/Engine/Plugins/Importers/USDImporter/Source/USDUtilities/Public/USDGeomMeshConversion.h
+++ b/Engine/Plugins/Importers/USDImporter/Source/USDUtilities/Public/USDGeomMeshConversion.h
@@ -85,6 +85,10 @@ namespace UsdToUnreal
 		// Whether to try reusing material slots (both local and the ones already in MaterialAssignments) when
 		// converting the material assignments from the mesh
 		bool bMergeIdenticalMaterialSlots;
+
+		// What subdivision level to subdivide to, if the converted mesh is a subdivision mesh.
+		// 0 means no subdivision should take place (default).
+		int32 SubdivisionLevel;
 	};
 
 	/**
@@ -380,7 +384,8 @@ namespace UsdUtils
 	 * Hashes the attributes of the GeomMesh at the given prim path on the Stage at TimeCode
 	 * If the prim path is invalid or not a GeomMesh, it will return an empty hash
 	 */
-	USDUTILITIES_API FString HashGeomMeshPrim( const UE::FUsdStage& Stage, const FString& PrimPath, double TimeCode );
+	USDUTILITIES_API FString HashGeomMeshPrim(const UE::FUsdStage& Stage, const FString& PrimPath, double TimeCode);
+	USDUTILITIES_API void HashGeomMeshPrim(const UE::FUsdStage& Stage, const FString& PrimPath, double TimeCode, FMD5& InOutHashState);
 
 	/**
 	 * Places in OutInstanceTransforms the UE-space instance transforms for a given point instancer prototype index.
diff --git a/Engine/Plugins/Importers/USDImporter/Source/USDUtilities/USDUtilities.Build.cs b/Engine/Plugins/Importers/USDImporter/Source/USDUtilities/USDUtilities.Build.cs
--- a/Engine/Plugins/Importers/USDImporter/Source/USDUtilities/USDUtilities.Build.cs
+++ b/Engine/Plugins/Importers/USDImporter/Source/USDUtilities/USDUtilities.Build.cs
@@ -44,6 +44,7 @@ namespace UnrealBuildTool.Rules
 					"MeshDescription",
 					"MovieScene",
 					"MovieSceneTracks",
+					"OpenSubdiv",
 					"RenderCore",
 					"RHI", // So that we can use GMaxRHIFeatureLevel when force-loading textures before baking materials
 					"Slate",
diff --git a/Engine/Plugins/Importers/USDImporter/Source/UnrealUSDWrapper/Private/UnrealUSDWrapper.cpp b/Engine/Plugins/Importers/USDImporter/Source/UnrealUSDWrapper/Private/UnrealUSDWrapper.cpp
--- a/Engine/Plugins/Importers/USDImporter/Source/UnrealUSDWrapper/Private/UnrealUSDWrapper.cpp
+++ b/Engine/Plugins/Importers/USDImporter/Source/UnrealUSDWrapper/Private/UnrealUSDWrapper.cpp
@@ -144,6 +144,8 @@ namespace UnrealIdentifiers
 	const TfToken UsdUVTexture = TfToken( "UsdUVTexture" );
 
 	const TfToken WorldSpaceNormals = TfToken( "worldSpaceNormals" );
+	const TfToken PrimvarsNormals = TfToken("primvars:normals");
+	const TfToken PrimvarsPoints = TfToken("primvars:points");
 
 	const TfToken GroomAPI = TfToken( "GroomAPI" );
 	const TfToken GroomBindingAPI = TfToken( "GroomBindingAPI" );
diff --git a/Engine/Plugins/Importers/USDImporter/Source/UnrealUSDWrapper/Public/UnrealUSDWrapper.h b/Engine/Plugins/Importers/USDImporter/Source/UnrealUSDWrapper/Public/UnrealUSDWrapper.h
--- a/Engine/Plugins/Importers/USDImporter/Source/UnrealUSDWrapper/Public/UnrealUSDWrapper.h
+++ b/Engine/Plugins/Importers/USDImporter/Source/UnrealUSDWrapper/Public/UnrealUSDWrapper.h
@@ -388,6 +388,10 @@ namespace UnrealIdentifiers
 	// Token used to indicate that a material parsed from a material prim should use world space normals
 	extern UNREALUSDWRAPPER_API const pxr::TfToken WorldSpaceNormals;
 
+	// Normals and points can also be primvars and have indices, but there is no defined tokens for them
+	extern UNREALUSDWRAPPER_API const pxr::TfToken PrimvarsNormals;
+	extern UNREALUSDWRAPPER_API const pxr::TfToken PrimvarsPoints;
+
 	extern UNREALUSDWRAPPER_API const pxr::TfToken GroomAPI;
 	extern UNREALUSDWRAPPER_API const pxr::TfToken GroomBindingAPI;
 	extern UNREALUSDWRAPPER_API const pxr::TfToken UnrealGroomToBind;
