USD: Add support for alternate draw modes.

diff --git a/Engine/Plugins/Importers/USDImporter/Source/USDClasses/Private/USDClassesModule.cpp b/Engine/Plugins/Importers/USDImporter/Source/USDClasses/Private/USDClassesModule.cpp
--- a/Engine/Plugins/Importers/USDImporter/Source/USDClasses/Private/USDClassesModule.cpp
+++ b/Engine/Plugins/Importers/USDImporter/Source/USDClasses/Private/USDClassesModule.cpp
@@ -3,6 +3,7 @@
 #include "USDClassesModule.h"
 
 #include "USDLog.h"
+#include "USDProjectSettings.h"
 
 #include "AnalyticsEventAttribute.h"
 #include "Animation/AnimBlueprint.h"
@@ -21,6 +22,8 @@
 #include "HAL/FileManager.h"
 #include "Materials/Material.h"
 #include "Materials/MaterialInstance.h"
+#include "Materials/MaterialInstanceConstant.h"
+#include "Materials/MaterialInstanceDynamic.h"
 #include "Misc/EngineVersion.h"
 #include "Misc/FileHelper.h"
 #include "Misc/PackageName.h"
@@ -31,11 +34,19 @@
 #include "UObject/NameTypes.h"
 #include "UObject/ObjectSaveContext.h"
 #include "UObject/Package.h"
+#include "UObject/Package.h"
+#include "UObject/UObjectGlobals.h"
+
+#if WITH_EDITOR
+#include "Editor/MaterialEditor/Public/MaterialEditingLibrary.h"
+#endif // WITH_EDITOR
 
 DEFINE_LOG_CATEGORY( LogUsd );
 
 namespace UE::USDClasses::Private
 {
+	static const FString DisplayColorID = TEXT( "!DisplayColor" );
+
 	TSharedPtr<FJsonObject> ParseJSON( const FString& FileContents )
 	{
 		if ( FileContents.IsEmpty() )
@@ -424,6 +435,139 @@ FString IUsdClassesModule::SanitizeObjectName(const FString& InObjectName)
 	return SanitizedText;
 }
 
+FString IUsdClassesModule::FDisplayColorMaterial::ToString()
+{
+	return FString::Printf(TEXT("%s_%d_%d"), *UE::USDClasses::Private::DisplayColorID, bHasOpacity, bIsDoubleSided);
+}
+
+TOptional<IUsdClassesModule::FDisplayColorMaterial> IUsdClassesModule::FDisplayColorMaterial::FromString(const FString& DisplayColorString)
+{
+	TArray<FString> Tokens;
+	DisplayColorString.ParseIntoArray(Tokens, TEXT("_"));
+
+	if (Tokens.Num() != 3 || Tokens[0] != UE::USDClasses::Private::DisplayColorID)
+	{
+		return {};
+	}
+
+	IUsdClassesModule::FDisplayColorMaterial Result;
+	Result.bHasOpacity = static_cast<bool>(FCString::Atoi(*Tokens[1]));
+	Result.bIsDoubleSided = static_cast<bool>(FCString::Atoi(*Tokens[2]));
+	return Result;
+}
+
+UMaterialInstanceDynamic* IUsdClassesModule::CreateDisplayColorMaterialInstanceDynamic(const FDisplayColorMaterial& DisplayColorDescription)
+{
+	const UUsdProjectSettings* Settings = GetDefault<UUsdProjectSettings>();
+	if (!Settings)
+	{
+		return nullptr;
+	}
+
+	const FSoftObjectPath* ParentPathPtr = nullptr;
+	if (DisplayColorDescription.bHasOpacity)
+	{
+		if (DisplayColorDescription.bIsDoubleSided)
+		{
+			ParentPathPtr = &Settings->ReferenceDisplayColorAndOpacityTwoSidedMaterial;
+		}
+		else
+		{
+			ParentPathPtr = &Settings->ReferenceDisplayColorAndOpacityMaterial;
+		}
+	}
+	else
+	{
+		if (DisplayColorDescription.bIsDoubleSided)
+		{
+			ParentPathPtr = &Settings->ReferenceDisplayColorTwoSidedMaterial;
+		}
+		else
+		{
+			ParentPathPtr = &Settings->ReferenceDisplayColorMaterial;
+		}
+	}
+	if (!ParentPathPtr)
+	{
+		return nullptr;
+	}
+
+	if (UMaterialInterface* ParentMaterial = Cast<UMaterialInterface>(ParentPathPtr->TryLoad()))
+	{
+		FName AssetName = MakeUniqueObjectName(
+			GetTransientPackage(),
+			UMaterialInstanceConstant::StaticClass(),
+			*FString::Printf(
+				TEXT("DisplayColor_%s_%s"),
+				DisplayColorDescription.bHasOpacity ? TEXT("Opacity") : TEXT("NoOpacity"),
+				DisplayColorDescription.bIsDoubleSided ? TEXT("DoubleSided") : TEXT("SingleSided")
+			)
+		);
+
+		if (UMaterialInstanceDynamic* NewMaterial = UMaterialInstanceDynamic::Create(ParentMaterial, GetTransientPackage(), AssetName))
+		{
+			return NewMaterial;
+		}
+	}
+
+	return nullptr;
+}
+
+UMaterialInstanceConstant* IUsdClassesModule::CreateDisplayColorMaterialInstanceConstant(const FDisplayColorMaterial& DisplayColorDescription)
+{
+#if WITH_EDITOR
+	const UUsdProjectSettings* Settings = GetDefault<UUsdProjectSettings>();
+	if (!Settings)
+	{
+		return nullptr;
+	}
+
+	const FSoftObjectPath* ParentPathPtr = nullptr;
+	if (DisplayColorDescription.bHasOpacity)
+	{
+		if (DisplayColorDescription.bIsDoubleSided)
+		{
+			ParentPathPtr = &Settings->ReferenceDisplayColorAndOpacityTwoSidedMaterial;
+		}
+		else
+		{
+			ParentPathPtr = &Settings->ReferenceDisplayColorAndOpacityMaterial;
+		}
+	}
+	else
+	{
+		if (DisplayColorDescription.bIsDoubleSided)
+		{
+			ParentPathPtr = &Settings->ReferenceDisplayColorTwoSidedMaterial;
+		}
+		else
+		{
+			ParentPathPtr = &Settings->ReferenceDisplayColorMaterial;
+		}
+	}
+
+	if (UMaterialInterface* ParentMaterial = Cast<UMaterialInterface>(ParentPathPtr->TryLoad()))
+	{
+		FName AssetName = MakeUniqueObjectName(
+			GetTransientPackage(),
+			UMaterialInstanceConstant::StaticClass(),
+			*FString::Printf(
+				TEXT("DisplayColor_%s_%s"),
+				DisplayColorDescription.bHasOpacity ? TEXT("Opacity") : TEXT("NoOpacity"),
+				DisplayColorDescription.bIsDoubleSided ? TEXT("DoubleSided") : TEXT("SingleSided")
+			)
+		);
+
+		if (UMaterialInstanceConstant* MaterialInstance = NewObject<UMaterialInstanceConstant>(GetTransientPackage(), AssetName, RF_NoFlags))
+		{
+			UMaterialEditingLibrary::SetMaterialInstanceParent(MaterialInstance, ParentMaterial);
+			return MaterialInstance;
+		}
+	}
+#endif	  // WITH_EDITOR
+	return nullptr;
+}
+
 class FUsdClassesModule : public IUsdClassesModule
 {
 
diff --git a/Engine/Plugins/Importers/USDImporter/Source/USDClasses/Private/USDDrawModeComponent.cpp b/Engine/Plugins/Importers/USDImporter/Source/USDClasses/Private/USDDrawModeComponent.cpp
new file mode 100644
--- /dev/null
+++ b/Engine/Plugins/Importers/USDImporter/Source/USDClasses/Private/USDDrawModeComponent.cpp
@@ -0,0 +1,1250 @@
+// Copyright Epic Games, Inc. All Rights Reserved.
+
+#include "USDDrawModeComponent.h"
+
+#include "USDAssetUserData.h"
+#include "USDClassesModule.h"
+#include "USDProjectSettings.h"
+
+#include "Components/BaseDynamicMeshSceneProxy.h"
+#include "DynamicMeshBuilder.h"
+#include "Engine/Engine.h"
+#include "Engine/Texture2D.h"
+#include "HitProxies.h"
+#include "MaterialDomain.h"
+#include "Materials/Material.h"
+#include "Materials/MaterialInstance.h"
+#include "Materials/MaterialInstanceConstant.h"
+#include "Materials/MaterialInstanceDynamic.h"
+#include "Materials/MaterialInterface.h"
+#include "Materials/MaterialRenderProxy.h"
+#include "PrimitiveSceneInfo.h"
+#include "PrimitiveSceneProxy.h"
+#include "PrimitiveViewRelevance.h"
+#include "SceneInterface.h"
+#include "SceneManagement.h"
+#include "UObject/Package.h"
+
+#if WITH_EDITOR
+#include "Editor/MaterialEditor/Public/MaterialEditingLibrary.h"
+#endif // WITH_EDITOR
+
+static bool GUseWholeExtentsToDetectSelection = true;
+static FAutoConsoleVariableRef CVarUseWholeExtentsToDetectSelection(
+	TEXT("USD.Bounds.UseWholeExtentsToDetectSelection"),
+	GUseWholeExtentsToDetectSelection,
+	TEXT("When true it means we can select the UUsdDrawModeComponents in 'bounds', 'cards' or 'origin' mode by clicking anywhere on the extent's 'volume' itself, and not only by clicking directly on the drawn lines")
+);
+
+static float GBoundsLineThickness = 2.0f;
+static FAutoConsoleVariableRef CVarBoundsLineThickness(
+	TEXT("USD.Bounds.BoundsLineThickness"),
+	GBoundsLineThickness,
+	TEXT("Line thickness to use when drawing 'bounds' or 'origin' alternate prim draw modes from UsdGeomModelAPI")
+);
+
+namespace UE::UsdDrawModeComponentImpl::Private
+{
+	// Simple proxy to draw just bounding box and coordinate axes lines
+	class FUsdLinesSceneProxy : public FPrimitiveSceneProxy
+	{
+	public:
+		FBox BoundsToDraw;
+		FLinearColor BoundsColor;
+
+		// Matrix to get a [0.5, 0.5] box centered on the origin (which is our mesh data)
+		// to the position described by BoundsToDraw. Used to draw the hit proxy box
+		FMatrix BoundsTransform;
+
+		FMeshRenderBufferSet HitProxyBuffer;
+
+		FMaterialRelevance MaterialRelevance;
+
+	public:
+		FUsdLinesSceneProxy(const UUsdDrawModeComponent* InComponent)
+			: FPrimitiveSceneProxy(InComponent)
+			, BoundsToDraw(FBox{InComponent->BoundsMin, InComponent->BoundsMax})
+			, BoundsColor(InComponent->BoundsColor)
+			, HitProxyBuffer(GetScene().GetFeatureLevel())
+			, MaterialRelevance(UMaterial::GetDefaultMaterial(MD_Surface)->GetRelevance_Concurrent(GetScene().GetFeatureLevel()))
+		{
+			bWillEverBeLit = false;
+			SetWireframeColor(FLinearColor(0.0f, 0.5f, 1.0f));
+
+			const FVector Translation = BoundsToDraw.GetCenter();
+			const FVector Scale = BoundsToDraw.GetSize();
+			FTransform Trans{FQuat::Identity, Translation, Scale};
+			BoundsTransform = Trans.ToMatrixWithScale();
+
+			// Create a simple box that we can present during the hit testing pass to make sure that
+			// our bounding boxes are easy to click
+			{
+				// Reference: PrimitiveDrawingUtils.cpp, DrawBox function
+
+				// Calculate verts for a face pointing down Z
+				FVector3f Positions[4] = {
+					FVector3f(-0.5, -0.5, +0.5),
+					FVector3f(-0.5, +0.5, +0.5),
+					FVector3f(+0.5, +0.5, +0.5),
+					FVector3f(+0.5, -0.5, +0.5)
+				};
+
+				FVector2f UVs[4] = {
+					FVector2f(0, 0),
+					FVector2f(0, 1),
+					FVector2f(1, 1),
+					FVector2f(1, 0),
+				};
+
+				// Then rotate this face 6 times
+				FRotator3f FaceRotations[6];
+				FaceRotations[0] = FRotator3f(0, 0, 0);
+				FaceRotations[1] = FRotator3f(90.f, 0, 0);
+				FaceRotations[2] = FRotator3f(-90.f, 0, 0);
+				FaceRotations[3] = FRotator3f(0, 0, 90.f);
+				FaceRotations[4] = FRotator3f(0, 0, -90.f);
+				FaceRotations[5] = FRotator3f(180.f, 0, 0);
+
+				const int NumTriangles = 2 * 6;
+				const int NumVertices = 4 * 6;
+				const int NumTexCoords = 1;
+				HitProxyBuffer.PositionVertexBuffer.Init(NumVertices);
+				HitProxyBuffer.StaticMeshVertexBuffer.Init(NumVertices, NumTexCoords);
+				HitProxyBuffer.ColorVertexBuffer.Init(NumVertices);
+				HitProxyBuffer.TriangleCount = NumTriangles;
+
+				for (uint32 FaceIndex = 0; FaceIndex < 6; FaceIndex++)
+				{
+					FMatrix44f FaceTransform = FRotationMatrix44f(FaceRotations[FaceIndex]);
+
+					uint32 NumVertsBeforeThisFace = FaceIndex * 4;
+
+					for (uint32 VertexIndex = 0; VertexIndex < 4; VertexIndex++)
+					{
+						uint32 GlobalVertIndex = VertexIndex + NumVertsBeforeThisFace;
+
+						HitProxyBuffer.PositionVertexBuffer.VertexPosition(GlobalVertIndex) = FaceTransform.TransformPosition(Positions[VertexIndex]);
+
+						HitProxyBuffer.StaticMeshVertexBuffer.SetVertexTangents(
+							GlobalVertIndex,
+							FaceTransform.TransformVector(FVector3f(1, 0, 0)),
+							FaceTransform.TransformVector(FVector3f(0, 1, 0)),
+							FaceTransform.TransformVector(FVector3f(0, 0, 1))
+						);
+
+						const uint32 UVIndex = 0;
+						HitProxyBuffer.StaticMeshVertexBuffer.SetVertexUV(GlobalVertIndex, UVIndex, UVs[VertexIndex]);
+
+						HitProxyBuffer.ColorVertexBuffer.VertexColor(GlobalVertIndex) = FColor::White;
+					}
+
+					HitProxyBuffer.IndexBuffer.Indices.Append({NumVertsBeforeThisFace, NumVertsBeforeThisFace + 1, NumVertsBeforeThisFace + 2});
+					HitProxyBuffer.IndexBuffer.Indices.Append({NumVertsBeforeThisFace, NumVertsBeforeThisFace + 2, NumVertsBeforeThisFace + 3});
+				}
+
+				HitProxyBuffer.Material = UMaterial::GetDefaultMaterial(MD_Surface);
+
+				ENQUEUE_RENDER_COMMAND(FUsdCardsSceneProxyUploadBuffer)
+				(
+					[this](FRHICommandListImmediate& RHICmdList)
+					{
+						HitProxyBuffer.Upload();
+					}
+				);
+
+			}
+		}
+
+		virtual void GetDynamicMeshElements(
+			const TArray<const FSceneView*>& Views,
+			const FSceneViewFamily& ViewFamily,
+			uint32 VisibilityMap,
+			FMeshElementCollector& Collector
+		) const override
+		{
+			const FMatrix& Matrix = GetLocalToWorld();
+
+			for (int32 ViewIndex = 0; ViewIndex < Views.Num(); ViewIndex++)
+			{
+				if (VisibilityMap & (1 << ViewIndex))
+				{
+					if (FPrimitiveDrawInterface* PDI = Collector.GetPDI(ViewIndex))
+					{
+						bool bIsWireframe = false;
+						if (AllowDebugViewmodes())
+						{
+							bIsWireframe = ViewFamily.EngineShowFlags.Wireframe;
+
+							// Render bounds for our bounds
+							if (ViewFamily.EngineShowFlags.Bounds)
+							{
+								RenderBounds(PDI, ViewFamily.EngineShowFlags, GetBounds(), IsSelected());
+							}
+						}
+
+						const uint8 DepthPriority = (uint8)GetDepthPriorityGroup(Views[ViewIndex]);
+
+						// When hit testing we just draw a regular solid box so that we can select the component by
+						// clicking anywhere within it
+						const bool bIsHitTesting = ViewFamily.EngineShowFlags.HitProxies;
+						if (bIsHitTesting && GUseWholeExtentsToDetectSelection)
+						{
+							bool bHasPrecomputedVolumetricLightmap;
+							FMatrix PreviousLocalToWorld;
+							int32 SingleCaptureIndex;
+							bool bOutputVelocity;
+							GetScene().GetPrimitiveUniformShaderParameters_RenderThread(
+								GetPrimitiveSceneInfo(),
+								bHasPrecomputedVolumetricLightmap,
+								PreviousLocalToWorld,
+								SingleCaptureIndex,
+								bOutputVelocity
+							);
+							bOutputVelocity |= AlwaysHasVelocity();
+
+							FDynamicPrimitiveUniformBuffer& DynamicPrimitiveUniformBuffer = Collector.AllocateOneFrameResource<FDynamicPrimitiveUniformBuffer>();
+
+							DynamicPrimitiveUniformBuffer.Set(
+								Collector.GetRHICommandList(),
+								BoundsTransform * GetLocalToWorld(),
+								PreviousLocalToWorld,
+								GetBounds(),
+								GetLocalBounds(),
+								GetLocalBounds(),
+								true,
+								bHasPrecomputedVolumetricLightmap,
+								bOutputVelocity,
+								GetCustomPrimitiveData()
+							);
+
+							FMeshBatch& Mesh = Collector.AllocateMesh();
+							Mesh.VertexFactory = &HitProxyBuffer.VertexFactory;
+							Mesh.MaterialRenderProxy = HitProxyBuffer.Material->GetRenderProxy();
+							Mesh.ReverseCulling = IsLocalToWorldDeterminantNegative();
+							Mesh.Type = PT_TriangleList;
+							Mesh.DepthPriorityGroup = DepthPriority;
+
+							FMeshBatchElement& BatchElement = Mesh.Elements[0];
+							BatchElement.IndexBuffer = &HitProxyBuffer.IndexBuffer;
+							BatchElement.PrimitiveUniformBufferResource = &DynamicPrimitiveUniformBuffer.UniformBuffer;
+							BatchElement.FirstIndex = 0;
+							BatchElement.NumPrimitives = HitProxyBuffer.IndexBuffer.Indices.Num() / 3;
+							BatchElement.MinVertexIndex = 0;
+							BatchElement.MaxVertexIndex = HitProxyBuffer.PositionVertexBuffer.GetNumVertices() - 1;
+							Collector.AddMesh(ViewIndex, Mesh);
+						}
+						else
+						{
+							DrawLines(
+								PDI,
+								Matrix,
+								bIsWireframe ? GetWireframeColor() : BoundsColor,
+								DepthPriority,
+								bIsWireframe ? 0.0f : GBoundsLineThickness
+							);
+						}
+					}
+				}
+			}
+		}
+
+		// Actually does the line drawing
+		virtual void DrawLines(
+			FPrimitiveDrawInterface* PDI,
+			const FMatrix& LocalToWorld,
+			const FLinearColor& LineColor,
+			uint8 DepthPriority,
+			float LineThickness
+		) const = 0;
+
+		virtual FPrimitiveViewRelevance GetViewRelevance(const FSceneView* View) const override
+		{
+			FPrimitiveViewRelevance ViewRelevance;
+			ViewRelevance.bDrawRelevance = IsShown(View);
+			ViewRelevance.bDynamicRelevance = true;
+			ViewRelevance.bNormalTranslucency = true;
+			return ViewRelevance;
+		}
+
+		virtual uint32 GetMemoryFootprint(void) const override
+		{
+			return (sizeof(*this) + GetAllocatedSize());
+		}
+
+		SIZE_T GetTypeHash() const override
+		{
+			static size_t UniquePointer;
+			return reinterpret_cast<size_t>(&UniquePointer);
+		}
+	};
+
+	// Version of the FUsdLinesSceneProxy proxy that draws a bounding box
+	class FUsdDrawModeLinesSceneProxy : public FUsdLinesSceneProxy
+	{
+	public:
+		using FUsdLinesSceneProxy::FUsdLinesSceneProxy;
+
+		virtual void DrawLines(
+			FPrimitiveDrawInterface* PDI,
+			const FMatrix& LocalToWorldMatrix,
+			const FLinearColor& LineColor,
+			uint8 DepthPriority,
+			float LineThickness
+		) const override
+		{
+			DrawWireBox(PDI, LocalToWorldMatrix, BoundsToDraw, LineColor, DepthPriority, LineThickness);
+		}
+	};
+
+	// Version of the FUsdLinesSceneProxy proxy that draws XYZ coordinate axes (always in the UE coordinate system)
+	class FUsdOriginLinesSceneProxy : public FUsdLinesSceneProxy
+	{
+	public:
+		using FUsdLinesSceneProxy::FUsdLinesSceneProxy;
+
+		virtual void DrawLines(
+			FPrimitiveDrawInterface* PDI,
+			const FMatrix& LocalToWorldMatrix,
+			const FLinearColor& LineColor,
+			uint8 DepthPriority,
+			float LineThickness
+		) const override
+		{
+			const FTransform World{LocalToWorldMatrix};
+			const FVector AxesLengths = BoundsToDraw.Max;
+
+			// In theory we should draw this with the wireframe color in wireframe mode but maybe it's more useful to have RGB XYZ colors still?
+			// Note: The DrawCoordinateSystem function does exist, but we don't use it here as it has a uniform scaling for all axes
+			// and it may be useful to show the axes scaled as usual given the prim's transform (i.e. doing this if the prims is
+			// squashed on the Z we will see a squashed Z line, while DrawCoordinateSystem would give the same line lengths for all
+			// axes)
+			PDI->DrawLine(
+				World.GetLocation(),
+				World.GetLocation() + World.TransformVector(FVector{AxesLengths.X, 0.0, 0.0}),
+				FLinearColor::Red,
+				DepthPriority,
+				LineThickness
+			);
+			PDI->DrawLine(
+				World.GetLocation(),
+				World.GetLocation() + World.TransformVector(FVector{0.0, AxesLengths.Y, 0.0}),
+				FLinearColor::Green,
+				DepthPriority,
+				LineThickness
+			);
+			PDI->DrawLine(
+				World.GetLocation(),
+				World.GetLocation() + World.TransformVector(FVector{0.0, 0.0, AxesLengths.Z}),
+				FLinearColor::Blue,
+				DepthPriority,
+				LineThickness
+			);
+		}
+	};
+
+	// Proxy capable of drawing meshes with materials and textures.
+	// Dynamic meshes via GetDynamicMeshElements can only have flat colors, so we need a slightly more complicated proxy.
+	// Reference: FDynamicMeshSceneProxy,
+	class FUsdCardsSceneProxy final : public FPrimitiveSceneProxy
+	{
+	public:
+
+		// Matrix to get a [0.5, 0.5] box centered on the origin (which is our mesh data)
+		// to the position described by BoundsToDraw
+		FMatrix BoundsTransform;
+
+		TArray<FMeshRenderBufferSet, TInlineAllocator<6>> Buffers;
+		TArray<FMaterialRelevance, TInlineAllocator<6>> MaterialRelevances;
+
+	public:
+		FUsdCardsSceneProxy(const UUsdDrawModeComponent* InComponent, const TArray<UMaterialInstance*>& MaterialInstances)
+			: FPrimitiveSceneProxy(InComponent)
+		{
+			bWillEverBeLit = true;
+			SetWireframeColor(FLinearColor(0.0f, 0.5f, 1.0f));
+
+			FBox BoundsToDraw = FBox{InComponent->BoundsMin, InComponent->BoundsMax};
+			const FVector Translation = BoundsToDraw.GetCenter();
+			const FVector Scale = BoundsToDraw.GetSize();
+			FTransform Trans{FQuat::Identity, Translation, Scale};
+			BoundsTransform = Trans.ToMatrixWithScale();
+
+			float OffsetAlongAxis = InComponent->CardGeometry == EUsdModelCardGeometry::Cross ? 0.0f : 0.5f;
+
+			// Weirdly enough given the drawMode rules when we have EUsdModelCardFace::None here we draw *all* faces, except we will just be using
+			// DisplayColor materials on all of them
+			const bool bDrawingAllWithDisplayColor = InComponent->GetAuthoredFaces() == EUsdModelCardFace::None;
+
+			for (int32 FaceIndex = 0; FaceIndex < 6; ++FaceIndex)
+			{
+				EUsdModelCardFace Face = (EUsdModelCardFace)(1 << FaceIndex);
+				EUsdModelCardFace OppositeFace = UsdUtils::GetOppositeFaceOnSameAxis(Face);
+
+				// If either face on an axis is enabled, we need to draw both
+				const bool bFaceIsAuthored = EnumHasAllFlags(InComponent->GetAuthoredFaces(), Face);
+				const bool bOppositeFaceIsAuthored = EnumHasAllFlags(InComponent->GetAuthoredFaces(), OppositeFace);
+				if (!bFaceIsAuthored && !bOppositeFaceIsAuthored && !bDrawingAllWithDisplayColor)
+				{
+					continue;
+				}
+
+				FMeshRenderBufferSet& Buffer = Buffers.Emplace_GetRef(GetScene().GetFeatureLevel());
+
+				const int NumTriangles = 2;
+				const int NumVertices = 4;
+				const int NumTexCoords = 1;
+				Buffer.PositionVertexBuffer.Init(NumVertices);
+				Buffer.StaticMeshVertexBuffer.Init(NumVertices, NumTexCoords);
+
+				switch (Face)
+				{
+					case EUsdModelCardFace::XPos:
+					{
+						Buffer.PositionVertexBuffer.VertexPosition(0) = {OffsetAlongAxis, -0.5, -0.5};
+						Buffer.PositionVertexBuffer.VertexPosition(1) = {OffsetAlongAxis, -0.5, +0.5};
+						Buffer.PositionVertexBuffer.VertexPosition(2) = {OffsetAlongAxis, +0.5, +0.5};
+						Buffer.PositionVertexBuffer.VertexPosition(3) = {OffsetAlongAxis, +0.5, -0.5};
+						Buffer.StaticMeshVertexBuffer.SetVertexTangents(0, {0, 1, 0}, {0, 0, 1}, {1, 0, 0});
+						Buffer.StaticMeshVertexBuffer.SetVertexTangents(1, {0, 1, 0}, {0, 0, 1}, {1, 0, 0});
+						Buffer.StaticMeshVertexBuffer.SetVertexTangents(2, {0, 1, 0}, {0, 0, 1}, {1, 0, 0});
+						Buffer.StaticMeshVertexBuffer.SetVertexTangents(3, {0, 1, 0}, {0, 0, 1}, {1, 0, 0});
+						break;
+					}
+					case EUsdModelCardFace::YPos:
+					{
+						Buffer.PositionVertexBuffer.VertexPosition(0) = {+0.5, OffsetAlongAxis, -0.5};
+						Buffer.PositionVertexBuffer.VertexPosition(1) = {+0.5, OffsetAlongAxis, +0.5};
+						Buffer.PositionVertexBuffer.VertexPosition(2) = {-0.5, OffsetAlongAxis, +0.5};
+						Buffer.PositionVertexBuffer.VertexPosition(3) = {-0.5, OffsetAlongAxis, -0.5};
+						Buffer.StaticMeshVertexBuffer.SetVertexTangents(0, {0, 0, 1}, {1, 0, 0}, {0, 1, 0});
+						Buffer.StaticMeshVertexBuffer.SetVertexTangents(1, {0, 0, 1}, {1, 0, 0}, {0, 1, 0});
+						Buffer.StaticMeshVertexBuffer.SetVertexTangents(2, {0, 0, 1}, {1, 0, 0}, {0, 1, 0});
+						Buffer.StaticMeshVertexBuffer.SetVertexTangents(3, {0, 0, 1}, {1, 0, 0}, {0, 1, 0});
+						break;
+					}
+					case EUsdModelCardFace::ZPos:
+					{
+						Buffer.PositionVertexBuffer.VertexPosition(0) = {+0.5, +0.5, OffsetAlongAxis};
+						Buffer.PositionVertexBuffer.VertexPosition(1) = {+0.5, -0.5, OffsetAlongAxis};
+						Buffer.PositionVertexBuffer.VertexPosition(2) = {-0.5, -0.5, OffsetAlongAxis};
+						Buffer.PositionVertexBuffer.VertexPosition(3) = {-0.5, +0.5, OffsetAlongAxis};
+						Buffer.StaticMeshVertexBuffer.SetVertexTangents(0, {1, 0, 0}, {0, 1, 0}, {0, 0, 1});
+						Buffer.StaticMeshVertexBuffer.SetVertexTangents(1, {1, 0, 0}, {0, 1, 0}, {0, 0, 1});
+						Buffer.StaticMeshVertexBuffer.SetVertexTangents(2, {1, 0, 0}, {0, 1, 0}, {0, 0, 1});
+						Buffer.StaticMeshVertexBuffer.SetVertexTangents(3, {1, 0, 0}, {0, 1, 0}, {0, 0, 1});
+						break;
+					}
+					case EUsdModelCardFace::XNeg:
+					{
+						Buffer.PositionVertexBuffer.VertexPosition(0) = {-OffsetAlongAxis, +0.5, -0.5};
+						Buffer.PositionVertexBuffer.VertexPosition(1) = {-OffsetAlongAxis, +0.5, +0.5};
+						Buffer.PositionVertexBuffer.VertexPosition(2) = {-OffsetAlongAxis, -0.5, +0.5};
+						Buffer.PositionVertexBuffer.VertexPosition(3) = {-OffsetAlongAxis, -0.5, -0.5};
+						Buffer.StaticMeshVertexBuffer.SetVertexTangents(0, {0, 0, 1}, {0, 1, 1}, {-1, 0, 0});
+						Buffer.StaticMeshVertexBuffer.SetVertexTangents(1, {0, 0, 1}, {0, 1, 1}, {-1, 0, 0});
+						Buffer.StaticMeshVertexBuffer.SetVertexTangents(2, {0, 0, 1}, {0, 1, 1}, {-1, 0, 0});
+						Buffer.StaticMeshVertexBuffer.SetVertexTangents(3, {0, 0, 1}, {0, 1, 1}, {-1, 0, 0});
+						break;
+					}
+					case EUsdModelCardFace::YNeg:
+					{
+						Buffer.PositionVertexBuffer.VertexPosition(0) = {-0.5, -OffsetAlongAxis, -0.5};
+						Buffer.PositionVertexBuffer.VertexPosition(1) = {-0.5, -OffsetAlongAxis, +0.5};
+						Buffer.PositionVertexBuffer.VertexPosition(2) = {+0.5, -OffsetAlongAxis, +0.5};
+						Buffer.PositionVertexBuffer.VertexPosition(3) = {+0.5, -OffsetAlongAxis, -0.5};
+						Buffer.StaticMeshVertexBuffer.SetVertexTangents(0, {0, 0, 1}, {-1, 0, 0}, {0, -1, 0});
+						Buffer.StaticMeshVertexBuffer.SetVertexTangents(1, {0, 0, 1}, {-1, 0, 0}, {0, -1, 0});
+						Buffer.StaticMeshVertexBuffer.SetVertexTangents(2, {0, 0, 1}, {-1, 0, 0}, {0, -1, 0});
+						Buffer.StaticMeshVertexBuffer.SetVertexTangents(3, {0, 0, 1}, {-1, 0, 0}, {0, -1, 0});
+						break;
+					}
+					case EUsdModelCardFace::ZNeg:
+					{
+						Buffer.PositionVertexBuffer.VertexPosition(0) = {+0.5, -0.5, -OffsetAlongAxis};
+						Buffer.PositionVertexBuffer.VertexPosition(1) = {+0.5, +0.5, -OffsetAlongAxis};
+						Buffer.PositionVertexBuffer.VertexPosition(2) = {-0.5, +0.5, -OffsetAlongAxis};
+						Buffer.PositionVertexBuffer.VertexPosition(3) = {-0.5, -0.5, -OffsetAlongAxis};
+						Buffer.StaticMeshVertexBuffer.SetVertexTangents(0, {1, 0, 0}, {0, -1, 0}, {0, 0, -1});
+						Buffer.StaticMeshVertexBuffer.SetVertexTangents(1, {1, 0, 0}, {0, -1, 0}, {0, 0, -1});
+						Buffer.StaticMeshVertexBuffer.SetVertexTangents(2, {1, 0, 0}, {0, -1, 0}, {0, 0, -1});
+						Buffer.StaticMeshVertexBuffer.SetVertexTangents(3, {1, 0, 0}, {0, -1, 0}, {0, 0, -1});
+						break;
+					}
+					default:
+					{
+						ensure(false);
+						break;
+					}
+				};
+
+				// This face is not authored but we're going to be drawing something for it --> Flip UVs horizontally.
+				// We don't know if our material is a texture material from the opposite face or a DisplayColor,
+				// but if it's the former we're doing the right thing by flipping the UV, and if it's the latter it doesn't matter!
+				const bool bFlipHorizontally = !bFaceIsAuthored;
+				if (bFlipHorizontally)
+				{
+					Buffer.StaticMeshVertexBuffer.SetVertexUV(0, 0, {0, 1});
+					Buffer.StaticMeshVertexBuffer.SetVertexUV(1, 0, {0, 0});
+					Buffer.StaticMeshVertexBuffer.SetVertexUV(2, 0, {1, 0});
+					Buffer.StaticMeshVertexBuffer.SetVertexUV(3, 0, {1, 1});
+				}
+				else
+				{
+					Buffer.StaticMeshVertexBuffer.SetVertexUV(0, 0, {1, 1});
+					Buffer.StaticMeshVertexBuffer.SetVertexUV(1, 0, {1, 0});
+					Buffer.StaticMeshVertexBuffer.SetVertexUV(2, 0, {0, 0});
+					Buffer.StaticMeshVertexBuffer.SetVertexUV(3, 0, {0, 1});
+				}
+
+				const bool bSRGB = true;
+				Buffer.ColorVertexBuffer.Init(NumVertices);
+				Buffer.ColorVertexBuffer.VertexColor(0) = InComponent->BoundsColor.ToFColor(bSRGB);
+				Buffer.ColorVertexBuffer.VertexColor(1) = InComponent->BoundsColor.ToFColor(bSRGB);
+				Buffer.ColorVertexBuffer.VertexColor(2) = InComponent->BoundsColor.ToFColor(bSRGB);
+				Buffer.ColorVertexBuffer.VertexColor(3) = InComponent->BoundsColor.ToFColor(bSRGB);
+
+				Buffer.IndexBuffer.Indices = {0, 1, 2, 0, 2, 3};
+				Buffer.TriangleCount = NumTriangles;
+
+				Buffer.Material = MaterialInstances[FaceIndex];
+				ensure(Buffer.Material);	// The component should always have a material for us, even if just the default material
+				MaterialRelevances.Emplace(Buffer.Material->GetRelevance_Concurrent(GetScene().GetFeatureLevel()));
+
+				ENQUEUE_RENDER_COMMAND(FUsdCardsSceneProxyUploadBuffer)
+				(
+					[&Buffer](FRHICommandListImmediate& RHICmdList)
+					{
+						Buffer.Upload();
+					}
+				);
+			}
+		}
+
+		virtual void GetDynamicMeshElements(
+			const TArray<const FSceneView*>& Views,
+			const FSceneViewFamily& ViewFamily,
+			uint32 VisibilityMap,
+			FMeshElementCollector& Collector
+		) const override
+		{
+			bool bHasPrecomputedVolumetricLightmap;
+			FMatrix PreviousLocalToWorld;
+			int32 SingleCaptureIndex;
+			bool bOutputVelocity;
+			GetScene().GetPrimitiveUniformShaderParameters_RenderThread(
+				GetPrimitiveSceneInfo(),
+				bHasPrecomputedVolumetricLightmap,
+				PreviousLocalToWorld,
+				SingleCaptureIndex,
+				bOutputVelocity
+			);
+			bOutputVelocity |= AlwaysHasVelocity();
+
+			FDynamicPrimitiveUniformBuffer& DynamicPrimitiveUniformBuffer = Collector.AllocateOneFrameResource<FDynamicPrimitiveUniformBuffer>();
+
+			DynamicPrimitiveUniformBuffer.Set(
+				Collector.GetRHICommandList(),
+				BoundsTransform * GetLocalToWorld(),
+				PreviousLocalToWorld,
+				GetBounds(),
+				GetLocalBounds(),
+				GetLocalBounds(),
+				true,
+				bHasPrecomputedVolumetricLightmap,
+				bOutputVelocity,
+				GetCustomPrimitiveData()
+			);
+
+			for (int32 ViewIndex = 0; ViewIndex < Views.Num(); ViewIndex++)
+			{
+				if (VisibilityMap & (1 << ViewIndex))
+				{
+					if (FPrimitiveDrawInterface* PDI = Collector.GetPDI(ViewIndex))
+					{
+						FMaterialRenderProxy* OverrideMaterialProxy = nullptr;
+
+						if (AllowDebugViewmodes())
+						{
+							if (ViewFamily.EngineShowFlags.Wireframe)
+							{
+								OverrideMaterialProxy = new FColoredMaterialRenderProxy(
+									GEngine->WireframeMaterial ? GEngine->WireframeMaterial->GetRenderProxy() : nullptr,
+									GetWireframeColor()
+								);
+								Collector.RegisterOneFrameMaterialProxy(OverrideMaterialProxy);
+							}
+
+							if (ViewFamily.EngineShowFlags.Bounds)
+							{
+								RenderBounds(PDI, ViewFamily.EngineShowFlags, GetBounds(), IsSelected());
+							}
+						}
+
+						for (const FMeshRenderBufferSet& Buffer : Buffers)
+						{
+							FMeshBatch& Mesh = Collector.AllocateMesh();
+							Mesh.VertexFactory = &Buffer.VertexFactory;
+							Mesh.MaterialRenderProxy = OverrideMaterialProxy ? OverrideMaterialProxy : Buffer.Material->GetRenderProxy();
+							Mesh.ReverseCulling = IsLocalToWorldDeterminantNegative();
+							Mesh.Type = PT_TriangleList;
+							Mesh.DepthPriorityGroup = (uint8)GetDepthPriorityGroup(Views[ViewIndex]);
+
+							FMeshBatchElement& BatchElement = Mesh.Elements[0];
+							BatchElement.IndexBuffer = &Buffer.IndexBuffer;
+							BatchElement.PrimitiveUniformBufferResource = &DynamicPrimitiveUniformBuffer.UniformBuffer;
+							BatchElement.FirstIndex = 0;
+							BatchElement.NumPrimitives = Buffer.IndexBuffer.Indices.Num() / 3;
+							BatchElement.MinVertexIndex = 0;
+							BatchElement.MaxVertexIndex = Buffer.PositionVertexBuffer.GetNumVertices() - 1;
+							Collector.AddMesh(ViewIndex, Mesh);
+						}
+					}
+				}
+			}
+		}
+
+		virtual FPrimitiveViewRelevance GetViewRelevance(const FSceneView* View) const override
+		{
+			// Reference: DynamicMeshSceneProxy::GetViewRelevance
+
+			FPrimitiveViewRelevance Result;
+			Result.bDrawRelevance = IsShown(View);
+			Result.bShadowRelevance = IsShadowCast(View);
+			Result.bDynamicRelevance = true;
+			Result.bRenderInMainPass = ShouldRenderInMainPass();
+			Result.bUsesLightingChannels = GetLightingChannelMask() != GetDefaultLightingChannelMask();
+			Result.bTranslucentSelfShadow = bCastVolumetricTranslucentShadow;
+			Result.bRenderCustomDepth = ShouldRenderCustomDepth();
+			Result.bVelocityRelevance = DrawsVelocity() && Result.bOpaque && Result.bRenderInMainPass;
+
+			if (MaterialRelevances.IsValidIndex(0))
+			{
+				MaterialRelevances[0].SetPrimitiveViewRelevance(Result);
+			}
+			return Result;
+		}
+
+		virtual bool CanBeOccluded() const override
+		{
+			if (MaterialRelevances.IsValidIndex(0))
+			{
+				return !MaterialRelevances[0].bDisableDepthTest;
+			}
+
+			return true;
+		}
+
+		virtual uint32 GetMemoryFootprint(void) const override
+		{
+			return (sizeof(*this) + GetAllocatedSize());
+		}
+
+		SIZE_T GetTypeHash() const override
+		{
+			static size_t UniquePointer;
+			return reinterpret_cast<size_t>(&UniquePointer);
+		}
+	};
+
+	int32 FaceToFaceIndex(EUsdModelCardFace Face)
+	{
+		ensure(Face != EUsdModelCardFace::None);
+		return FMath::CountTrailingZeros((int32)Face);
+	}
+
+	void SetTextureParameter(UMaterialInstance* Instance, UTexture2D* Texture)
+	{
+		FMaterialParameterInfo Info;
+		Info.Name = TEXT("Texture");
+
+#if WITH_EDITOR
+		if (GIsEditor)
+		{
+			if (UMaterialInstanceConstant* Constant = Cast<UMaterialInstanceConstant>(Instance))
+			{
+				Constant->SetTextureParameterValueEditorOnly(Info, Texture);
+			}
+		}
+		else
+#endif	  // WITH_EDITOR
+		{
+			if (UMaterialInstanceDynamic* Dynamic = Cast<UMaterialInstanceDynamic>(Instance))
+			{
+				Dynamic->SetTextureParameterValueByInfo(Info, Texture);
+			}
+		}
+	}
+};	  // namespace UE::UsdDrawModeComponentImpl::Private
+
+EUsdModelCardFace UsdUtils::GetOppositeFaceOnSameAxis(EUsdModelCardFace Face)
+{
+	switch (Face)
+	{
+		case EUsdModelCardFace::XPos:
+		{
+			return EUsdModelCardFace::XNeg;
+			break;
+		}
+		case EUsdModelCardFace::YPos:
+		{
+			return EUsdModelCardFace::YNeg;
+			break;
+		}
+		case EUsdModelCardFace::ZPos:
+		{
+			return EUsdModelCardFace::ZNeg;
+			break;
+		}
+		case EUsdModelCardFace::XNeg:
+		{
+			return EUsdModelCardFace::XPos;
+			break;
+		}
+		case EUsdModelCardFace::YNeg:
+		{
+			return EUsdModelCardFace::YPos;
+			break;
+		}
+		case EUsdModelCardFace::ZNeg:
+		{
+			return EUsdModelCardFace::ZPos;
+			break;
+		}
+		default:
+		{
+			ensure(false);
+			break;
+		}
+	}
+
+	return EUsdModelCardFace::None;
+}
+
+void UUsdDrawModeComponent::SetBoundsMin(const FVector& NewMin)
+{
+	BoundsMin = NewMin;
+	MarkRenderStateDirty();
+}
+
+void UUsdDrawModeComponent::SetBoundsMax(const FVector& NewMax)
+{
+	BoundsMax = NewMax;
+	MarkRenderStateDirty();
+}
+
+void UUsdDrawModeComponent::SetDrawMode(EUsdDrawMode NewDrawMode)
+{
+	DrawMode = NewDrawMode;
+	MarkRenderStateDirty();
+}
+
+void UUsdDrawModeComponent::SetBoundsColor(FLinearColor NewColor)
+{
+	BoundsColor = NewColor;
+	MarkRenderStateDirty();
+}
+
+void UUsdDrawModeComponent::SetCardGeometry(EUsdModelCardGeometry NewGeometry)
+{
+	CardGeometry = NewGeometry;
+	MarkRenderStateDirty();
+}
+
+void UUsdDrawModeComponent::SetCardTextureXPos(UTexture2D* NewTexture)
+{
+	// We have this logic so that clearing the texture on the details panel means
+	// removing the opinion of that face texture in the USD stage
+	// (when we convert the component back to USD we'll clear faces that aren't in AuthoredFaces).
+	if (NewTexture)
+	{
+		AuthoredFaces |= EUsdModelCardFace::XPos;
+	}
+	else
+	{
+		AuthoredFaces &= ~EUsdModelCardFace::XPos;
+	}
+
+	CardTextureXPos = NewTexture;
+	MarkRenderStateDirty();
+}
+
+void UUsdDrawModeComponent::SetCardTextureYPos(UTexture2D* NewTexture)
+{
+	if (NewTexture)
+	{
+		AuthoredFaces |= EUsdModelCardFace::YPos;
+	}
+	else
+	{
+		AuthoredFaces &= ~EUsdModelCardFace::YPos;
+	}
+
+	CardTextureYPos = NewTexture;
+	MarkRenderStateDirty();
+}
+
+void UUsdDrawModeComponent::SetCardTextureZPos(UTexture2D* NewTexture)
+{
+	if (NewTexture)
+	{
+		AuthoredFaces |= EUsdModelCardFace::ZPos;
+	}
+	else
+	{
+		AuthoredFaces &= ~EUsdModelCardFace::ZPos;
+	}
+
+	CardTextureZPos = NewTexture;
+	MarkRenderStateDirty();
+}
+
+void UUsdDrawModeComponent::SetCardTextureXNeg(UTexture2D* NewTexture)
+{
+	if (NewTexture)
+	{
+		AuthoredFaces |= EUsdModelCardFace::XNeg;
+	}
+	else
+	{
+		AuthoredFaces &= ~EUsdModelCardFace::XNeg;
+	}
+
+	CardTextureXNeg = NewTexture;
+	MarkRenderStateDirty();
+}
+
+void UUsdDrawModeComponent::SetCardTextureYNeg(UTexture2D* NewTexture)
+{
+	if (NewTexture)
+	{
+		AuthoredFaces |= EUsdModelCardFace::YNeg;
+	}
+	else
+	{
+		AuthoredFaces &= ~EUsdModelCardFace::YNeg;
+	}
+
+	CardTextureYNeg = NewTexture;
+	MarkRenderStateDirty();
+}
+
+void UUsdDrawModeComponent::SetCardTextureZNeg(UTexture2D* NewTexture)
+{
+	if (NewTexture)
+	{
+		AuthoredFaces |= EUsdModelCardFace::ZNeg;
+	}
+	else
+	{
+		AuthoredFaces &= ~EUsdModelCardFace::ZNeg;
+	}
+
+	CardTextureZNeg = NewTexture;
+	MarkRenderStateDirty();
+}
+
+UTexture2D* UUsdDrawModeComponent::GetTextureForFace(EUsdModelCardFace Face) const
+{
+	switch (Face)
+	{
+		case EUsdModelCardFace::XPos:
+		{
+			return CardTextureXPos;
+			break;
+		}
+		case EUsdModelCardFace::YPos:
+		{
+			return CardTextureYPos;
+			break;
+		}
+		case EUsdModelCardFace::ZPos:
+		{
+			return CardTextureZPos;
+			break;
+		}
+		case EUsdModelCardFace::XNeg:
+		{
+			return CardTextureXNeg;
+			break;
+		}
+		case EUsdModelCardFace::YNeg:
+		{
+			return CardTextureYNeg;
+			break;
+		}
+		case EUsdModelCardFace::ZNeg:
+		{
+			return CardTextureZNeg;
+			break;
+		}
+		default:
+		{
+			break;
+		}
+	}
+
+	return nullptr;
+}
+
+void UUsdDrawModeComponent::SetTextureForFace(EUsdModelCardFace Face, UTexture2D* Texture)
+{
+	switch (Face)
+	{
+		case EUsdModelCardFace::XPos:
+		{
+			SetCardTextureXPos(Texture);
+			break;
+		}
+		case EUsdModelCardFace::YPos:
+		{
+			SetCardTextureYPos(Texture);
+			break;
+		}
+		case EUsdModelCardFace::ZPos:
+		{
+			SetCardTextureZPos(Texture);
+			break;
+		}
+		case EUsdModelCardFace::XNeg:
+		{
+			SetCardTextureXNeg(Texture);
+			break;
+		}
+		case EUsdModelCardFace::YNeg:
+		{
+			SetCardTextureYNeg(Texture);
+			break;
+		}
+		case EUsdModelCardFace::ZNeg:
+		{
+			SetCardTextureZNeg(Texture);
+			break;
+		}
+		default:
+		{
+			break;
+		}
+	}
+}
+
+UUsdDrawModeComponent::UUsdDrawModeComponent()
+	: BoundsMin{FVector{-0.5}}
+	, BoundsMax{FVector{0.5}}
+	, DrawMode{EUsdDrawMode::Bounds}
+	, BoundsColor{FLinearColor::White}
+	, CardGeometry{EUsdModelCardGeometry::Cross}
+	, CardTextureXPos{nullptr}
+	, CardTextureYPos{nullptr}
+	, CardTextureZPos{nullptr}
+	, CardTextureXNeg{nullptr}
+	, CardTextureYNeg{nullptr}
+	, CardTextureZNeg{nullptr}
+	, MaterialInstances({nullptr, nullptr, nullptr, nullptr, nullptr, nullptr})
+	, AuthoredFaces{EUsdModelCardFace::None}
+{
+	RefreshMaterialInstances();
+}
+
+FPrimitiveSceneProxy* UUsdDrawModeComponent::CreateSceneProxy()
+{
+	using namespace UE::UsdDrawModeComponentImpl::Private;
+
+	// We call this here because changing a property on the editor doesn't call our BlueprintSetter
+	// functions, but nevertheless still triggers a PostEditChange event that invalidates our render
+	// proxy anyway. This means that if we delay refreshing the material instances in here it ends up
+	// called in all cases and we always get textures/materials that are up-to-date
+	RefreshMaterialInstances();
+
+	// The render thread owns and deletes these, so we can just reallocate every time.
+	// CreateSceneProxy gets called when we call MarkRenderStateDirty()
+	switch (DrawMode)
+	{
+		case EUsdDrawMode::Origin:
+		{
+			return new FUsdOriginLinesSceneProxy(this);
+			break;
+		}
+		case EUsdDrawMode::Bounds:
+		{
+			return new FUsdDrawModeLinesSceneProxy(this);
+			break;
+		}
+		case EUsdDrawMode::Cards:
+		{
+			return new FUsdCardsSceneProxy(this, MaterialInstances);
+			break;
+		}
+		case EUsdDrawMode::Default:
+		case EUsdDrawMode::Inherited:
+		default:
+		{
+			// We shouldn't have an UUsdDrawModeComponent if we have these draw modes...
+			ensure(false);
+			break;
+		}
+	}
+
+	return nullptr;
+}
+
+FBoxSphereBounds UUsdDrawModeComponent::CalcBounds(const FTransform& LocalToWorld) const
+{
+	FBox LocalSpaceBox{BoundsMin, BoundsMax};
+	LocalSpaceBox = LocalSpaceBox.TransformBy(LocalToWorld);
+	return FBoxSphereBounds(LocalSpaceBox);
+}
+
+void UUsdDrawModeComponent::GetUsedMaterials(TArray<UMaterialInterface*>& OutMaterials, bool bGetDebugMaterials) const
+{
+	OutMaterials.Reset(MaterialInstances.Num());
+	for (UMaterialInstance* Instance : MaterialInstances)
+	{
+		if (Instance)
+		{
+			OutMaterials.Add(Instance);
+		}
+	}
+}
+
+#if WITH_EDITOR
+void UUsdDrawModeComponent::PostEditChangeProperty(FPropertyChangedEvent& PropertyChangedEvent)
+{
+	// Redirect our property change events (e.g. editing the value directly on the Details panel) to
+	// the blueprint setters, so that they can also e.g. trigger changes to AuthoredFaces.
+	// In theory we only need this for the card face texture properties, but let's just do it for all
+	// of them anyway for consistency
+
+	const FName PropertyName = PropertyChangedEvent.GetMemberPropertyName();
+
+	if (PropertyName == GET_MEMBER_NAME_CHECKED(UUsdDrawModeComponent, BoundsMin))
+	{
+		SetBoundsMin(BoundsMin);
+	}
+	else if (PropertyName == GET_MEMBER_NAME_CHECKED(UUsdDrawModeComponent, BoundsMax))
+	{
+		SetBoundsMax(BoundsMax);
+	}
+	else if (PropertyName == GET_MEMBER_NAME_CHECKED(UUsdDrawModeComponent, DrawMode))
+	{
+		SetDrawMode(DrawMode);
+	}
+	else if (PropertyName == GET_MEMBER_NAME_CHECKED(UUsdDrawModeComponent, BoundsColor))
+	{
+		SetBoundsColor(BoundsColor);
+	}
+	else if (PropertyName == GET_MEMBER_NAME_CHECKED(UUsdDrawModeComponent, CardGeometry))
+	{
+		SetCardGeometry(CardGeometry);
+	}
+	else if (PropertyName == GET_MEMBER_NAME_CHECKED(UUsdDrawModeComponent, CardTextureXPos))
+	{
+		SetCardTextureXPos(CardTextureXPos);
+	}
+	else if (PropertyName == GET_MEMBER_NAME_CHECKED(UUsdDrawModeComponent, CardTextureYPos))
+	{
+		SetCardTextureYPos(CardTextureYPos);
+	}
+	else if (PropertyName == GET_MEMBER_NAME_CHECKED(UUsdDrawModeComponent, CardTextureZPos))
+	{
+		SetCardTextureZPos(CardTextureZPos);
+	}
+	else if (PropertyName == GET_MEMBER_NAME_CHECKED(UUsdDrawModeComponent, CardTextureXNeg))
+	{
+		SetCardTextureXNeg(CardTextureXNeg);
+	}
+	else if (PropertyName == GET_MEMBER_NAME_CHECKED(UUsdDrawModeComponent, CardTextureYNeg))
+	{
+		SetCardTextureYNeg(CardTextureYNeg);
+	}
+	else if (PropertyName == GET_MEMBER_NAME_CHECKED(UUsdDrawModeComponent, CardTextureZNeg))
+	{
+		SetCardTextureZNeg(CardTextureZNeg);
+	}
+
+	// Propagate this only after our changes, because that is what will cause AUsdStageActor::OnObjectPropertyChanged
+	// to respond, and we want to write out to USD only after our blueprint setters have had the chance to tweak
+	// AuthoredFaces
+	Super::PostEditChangeProperty(PropertyChangedEvent);
+}
+
+HHitProxy* UUsdDrawModeComponent::CreateMeshHitProxy(int32 SectionIndex, int32 MaterialIndex) const
+{
+	return nullptr;
+}
+#endif // WITH_EDITOR
+
+EUsdModelCardFace UUsdDrawModeComponent::GetAuthoredFaces() const
+{
+	return AuthoredFaces;
+}
+
+void UUsdDrawModeComponent::SetAuthoredFaces(EUsdModelCardFace NewAuthoredFaces)
+{
+	AuthoredFaces = NewAuthoredFaces;
+	MarkRenderStateDirty();
+}
+
+UMaterialInstance* UUsdDrawModeComponent::GetOrCreateTextureMaterial(EUsdModelCardFace Face)
+{
+	if (Face == EUsdModelCardFace::None)
+	{
+		return nullptr;
+	}
+
+	int32 FaceIndex = UE::UsdDrawModeComponentImpl::Private::FaceToFaceIndex(Face);
+	if (!ensure(MaterialInstances.IsValidIndex(FaceIndex)))
+	{
+		return nullptr;
+	}
+
+	const UUsdProjectSettings* ProjectSettings = GetDefault<UUsdProjectSettings>();
+	if (!ProjectSettings)
+	{
+		return nullptr;
+	}
+
+	TObjectPtr<UMaterialInstance>& Instance = MaterialInstances[FaceIndex];
+
+	// Need to create a new instance of the texture material for this face
+	// TODO: If the texture is really large UE may automatically make it into a VT texture, and our material with a
+	// regular texture samplers won't really work. Should we have a "VT upgrade" mechanism? Is that even relevant for these
+	// "card" textures? (who is using 4K+ texture cards?)
+	if (!Instance || FSoftObjectPath{Instance->Parent} != ProjectSettings->ReferenceModelCardTextureMaterial)
+	{
+		if (UMaterialInterface* ParentMaterial = Cast<UMaterialInterface>(ProjectSettings->ReferenceModelCardTextureMaterial.TryLoad()))
+		{
+#if WITH_EDITOR
+			if (GIsEditor)
+			{
+				if (UMaterialInstanceConstant* NewMIC = NewObject<UMaterialInstanceConstant>(GetTransientPackage(), NAME_None, RF_Transient))
+				{
+					UMaterialEditingLibrary::SetMaterialInstanceParent(NewMIC, ParentMaterial);
+					Instance = NewMIC;
+				}
+			}
+			else
+#endif
+			{
+				if (UMaterialInstanceDynamic* NewMID = UMaterialInstanceDynamic::Create(ParentMaterial, GetTransientPackage(), NAME_None))
+				{
+					Instance = NewMID;
+				}
+			}
+		}
+	}
+
+	return Instance;
+}
+
+void UUsdDrawModeComponent::RefreshMaterialInstances()
+{
+	// We must follow the following rules (some of these from
+	// https://openusd.org/release/api/class_usd_geom_model_a_p_i.html#UsdGeomModelAPI_drawMode):
+	// - If no texture is set (i.e. if AuthoredFaces is None) then draw all faces with BoundsColor. Otherwise just draw the axes set with textures;
+	// - If a texture is present for an axis, but the back face texture is missing, use the same texture flipped horizontally
+	//   (the proxy will handle that);
+	// - If a texture is set but fails to resolve (i.e. AuthoredFaces is set but the texture is nullptr) then draw that face with bounds color
+	//   (we will do this by using a DisplayColor material instead);
+
+	using namespace UE::UsdDrawModeComponentImpl::Private;
+
+	const bool bDrawingAllWithDisplayColor = (EUsdModelCardFace)AuthoredFaces == EUsdModelCardFace::None;
+
+	TSet<int32, DefaultKeyFuncs<int32>, TInlineSetAllocator<8>> IndicesToSetDrawModeColor;
+	for (int32 FaceIndex = 0; FaceIndex < 6; ++FaceIndex)
+	{
+		if (bDrawingAllWithDisplayColor)
+		{
+			IndicesToSetDrawModeColor.Add(FaceIndex);
+			continue;
+		}
+
+		EUsdModelCardFace Face = (EUsdModelCardFace)(1 << FaceIndex);
+		UTexture2D* Texture = GetTextureForFace(Face);
+		const bool bDrawFace = EnumHasAllFlags((EUsdModelCardFace)AuthoredFaces, Face);
+
+		EUsdModelCardFace OppositeFace = UsdUtils::GetOppositeFaceOnSameAxis(Face);
+		UTexture2D* OppositeTexture = GetTextureForFace(OppositeFace);
+		const bool bDrawOppositeFace = EnumHasAllFlags((EUsdModelCardFace)AuthoredFaces, OppositeFace);
+
+		// Face is authored and we have a texture --> Just draw that texture
+		if (bDrawFace && Texture)
+		{
+			SetTextureParameter(GetOrCreateTextureMaterial(Face), Texture);
+		}
+		// Face is authored but we don't have a texture --> Failed to resolve our texture --> Draw DisplayColor
+		else if (bDrawFace)
+		{
+			IndicesToSetDrawModeColor.Add(FaceIndex);
+		}
+		// Face is not authored, but the opposite face is and it has a texture --> Draw that texture (proxy will flip it horizontally)
+		else if (!bDrawFace && bDrawOppositeFace && OppositeTexture)
+		{
+			SetTextureParameter(GetOrCreateTextureMaterial(Face), OppositeTexture);
+		}
+		// Face is not authored, but the opposite face is. It doesn't have a texture though. (+ other cases) --> Draw DisplyColor
+		else
+		{
+			IndicesToSetDrawModeColor.Add(FaceIndex);
+		}
+	}
+
+	UMaterialInstance* DisplayColorInstance = nullptr;
+	if (IndicesToSetDrawModeColor.Num() > 0)
+	{
+		// Check if we already have any DisplayColor materials we can reuse. These don't
+		// have any parameters so there's no reason at all to have instances for each.
+		// Actually we don't even need any material instance at all really... although using the
+		// plugin material directly on the component is probably not the best idea
+		if (const UUsdProjectSettings* Settings = GetDefault<UUsdProjectSettings>())
+		{
+			for (UMaterialInstance* Instance : MaterialInstances)
+			{
+				if (!Instance)
+				{
+					continue;
+				}
+
+				// We only ever spawn opaque one-sided display colors for the cards
+				if (FSoftObjectPath{Instance->Parent} == Settings->ReferenceDisplayColorMaterial)
+				{
+					DisplayColorInstance = Instance;
+					break;
+				}
+			}
+		}
+
+		// Create a new display color instance
+		if (!DisplayColorInstance)
+		{
+			IUsdClassesModule::FDisplayColorMaterial Desc;
+			Desc.bHasOpacity = false;
+			Desc.bIsDoubleSided = false;
+
+#if WITH_EDITOR
+			if (GIsEditor)
+			{
+				DisplayColorInstance = IUsdClassesModule::CreateDisplayColorMaterialInstanceConstant(Desc);
+			}
+			else
+#endif	  // WITH_EDITOR
+			{
+				DisplayColorInstance = IUsdClassesModule::CreateDisplayColorMaterialInstanceDynamic(Desc);
+			}
+
+			if (DisplayColorInstance)
+			{
+				// Leave PrimPath as empty as it likely will be reused by many prims
+				UUsdAssetUserData* UserData = NewObject<UUsdAssetUserData>(DisplayColorInstance, TEXT("USDAssetUserData"));
+				DisplayColorInstance->AddAssetUserData(UserData);
+
+				DisplayColorInstance->SetFlags(RF_Transient);
+			}
+		}
+
+		for (int32 Index : IndicesToSetDrawModeColor)
+		{
+			MaterialInstances[Index] = DisplayColorInstance;
+		}
+	}
+}
diff --git a/Engine/Plugins/Importers/USDImporter/Source/USDClasses/Public/USDClassesModule.h b/Engine/Plugins/Importers/USDImporter/Source/USDClasses/Public/USDClassesModule.h
--- a/Engine/Plugins/Importers/USDImporter/Source/USDClasses/Public/USDClassesModule.h
+++ b/Engine/Plugins/Importers/USDImporter/Source/USDClasses/Public/USDClassesModule.h
@@ -5,6 +5,8 @@
 #include "CoreMinimal.h"
 #include "Modules/ModuleInterface.h"
 
+class UMaterialInstanceConstant;
+class UMaterialInstanceDynamic;
 struct FAnalyticsEventAttribute;
 
 class IUsdClassesModule : public IModuleInterface
@@ -53,4 +55,17 @@ public:
 
 	// Adapted from ObjectTools as it is within an Editor-only module
 	USDCLASSES_API static FString SanitizeObjectName(const FString& InObjectName);
+
+	/** Describes the type of vertex color/DisplayColor material that we would need in order to render a prim's displayColor data as intended */
+	struct USDCLASSES_API FDisplayColorMaterial
+	{
+		bool bHasOpacity = false;
+		bool bIsDoubleSided = false;
+
+		FString ToString();
+		static TOptional<FDisplayColorMaterial> FromString(const FString& DisplayColorString);
+	};
+
+	USDCLASSES_API static UMaterialInstanceDynamic* CreateDisplayColorMaterialInstanceDynamic(const FDisplayColorMaterial& DisplayColorDescription);
+	USDCLASSES_API static UMaterialInstanceConstant* CreateDisplayColorMaterialInstanceConstant(const FDisplayColorMaterial& DisplayColorDescription);
 };
\ No newline at end of file
diff --git a/Engine/Plugins/Importers/USDImporter/Source/USDClasses/Public/USDDrawModeComponent.h b/Engine/Plugins/Importers/USDImporter/Source/USDClasses/Public/USDDrawModeComponent.h
new file mode 100644
--- /dev/null
+++ b/Engine/Plugins/Importers/USDImporter/Source/USDClasses/Public/USDDrawModeComponent.h
@@ -0,0 +1,171 @@
+// Copyright Epic Games, Inc. All Rights Reserved.
+
+#pragma once
+
+#include "Components/PrimitiveComponent.h"
+
+#include "USDDrawModeComponent.generated.h"
+
+class UMaterialInstance;
+class UUsdAssetCache2;
+enum class EUsdUpAxis : uint8;
+
+UENUM()
+enum class EUsdDrawMode : int32
+{
+	Origin,
+	Bounds,
+	Cards,
+	// For now we just hide these, so that the user can't pick them on the component details panel.
+	// This because we don't yet have a good way of having component changes trigger resyncs: OnObjectPropertyChanged
+	// is fired immediately for each property change, and resyncing right away may cause some trouble in case it
+	// deletes components (which it will in this case) and the previous callstack of the OnObjectPropertyChanged still
+	// wants to interact with the component, or fire OnObjectPropertyChanged for other properties, etc.
+	Default UMETA(Hidden),
+	Inherited UMETA(Hidden)
+};
+
+// When we're on EUsdDrawMode::Cards, describes how to draw our cards
+UENUM()
+enum class EUsdModelCardGeometry : int32
+{
+	Cross,
+	Box,
+	FromTexture UMETA(Hidden)// We don't really support this one, and will fallback to "Box"
+};
+
+UENUM(meta = (Bitflags, UseEnumValuesAsMaskValuesInEditor = "true"))
+enum class EUsdModelCardFace : int32
+{
+	None = 0 UMETA(Hidden),
+	XPos = 1,
+	YPos = 2,
+	ZPos = 4,
+	XNeg = 8,
+	YNeg = 16,
+	ZNeg = 32
+};
+ENUM_CLASS_FLAGS(EUsdModelCardFace);
+
+namespace UsdUtils
+{
+	USDCLASSES_API EUsdModelCardFace GetOppositeFaceOnSameAxis(EUsdModelCardFace Face);
+}
+
+/**
+ * Component type that is used to draw bounds, cards and origin axes when translating a prim that has the UsdGeomModelAPI schema
+ * applied.
+ */
+UCLASS()
+class USDCLASSES_API UUsdDrawModeComponent : public UPrimitiveComponent
+{
+	GENERATED_BODY()
+
+public:
+	// Separate vectors instead of an FBox so we can get Sequencer tracks for these as is
+	UPROPERTY(EditAnywhere, BlueprintReadWrite, BlueprintSetter = SetBoundsMin, Category = "USD")
+	FVector BoundsMin;
+
+	UPROPERTY(EditAnywhere, BlueprintReadWrite, BlueprintSetter = SetBoundsMax, Category = "USD")
+	FVector BoundsMax;
+
+	UPROPERTY(EditAnywhere, BlueprintReadWrite, BlueprintSetter = SetDrawMode, Category = "USD")
+	EUsdDrawMode DrawMode;
+
+	UPROPERTY(EditAnywhere, BlueprintReadWrite, BlueprintSetter = SetBoundsColor, Category = "USD")
+	FLinearColor BoundsColor;
+
+	UPROPERTY(EditAnywhere, BlueprintReadWrite, BlueprintSetter = SetCardGeometry, Category = "USD|Cards")
+	EUsdModelCardGeometry CardGeometry;
+
+	UPROPERTY(EditAnywhere, BlueprintReadWrite, BlueprintSetter = SetCardTextureXPos, Category = "USD|Cards")
+	TObjectPtr<UTexture2D> CardTextureXPos;
+
+	UPROPERTY(EditAnywhere, BlueprintReadWrite, BlueprintSetter = SetCardTextureYPos, Category = "USD|Cards")
+	TObjectPtr<UTexture2D> CardTextureYPos;
+
+	UPROPERTY(EditAnywhere, BlueprintReadWrite, BlueprintSetter = SetCardTextureZPos, Category = "USD|Cards")
+	TObjectPtr<UTexture2D> CardTextureZPos;
+
+	UPROPERTY(EditAnywhere, BlueprintReadWrite, BlueprintSetter = SetCardTextureXNeg, Category = "USD|Cards")
+	TObjectPtr<UTexture2D> CardTextureXNeg;
+
+	UPROPERTY(EditAnywhere, BlueprintReadWrite, BlueprintSetter = SetCardTextureYNeg, Category = "USD|Cards")
+	TObjectPtr<UTexture2D> CardTextureYNeg;
+
+	UPROPERTY(EditAnywhere, BlueprintReadWrite, BlueprintSetter = SetCardTextureZNeg, Category = "USD|Cards")
+	TObjectPtr<UTexture2D> CardTextureZNeg;
+
+private:
+	UPROPERTY()
+	TArray<TObjectPtr<UMaterialInstance>> MaterialInstances;
+
+	// Flags to enable the card faces to draw.
+	// We need these because the required USD behavior specifies a difference between not having a texture because
+	// one wasn't authored (at which point we flip the opposite face), and not having a texture because it failed
+	// to resolve (at which point we draw vertex color). In both of those cases we'll end up with nullptr on the
+	// CardTexture property, but this will help us tell which case we're talking about
+	UPROPERTY()
+	EUsdModelCardFace AuthoredFaces;
+
+public:
+	UFUNCTION(BlueprintSetter, Category = "USD")
+	void SetBoundsMin(const FVector& NewMin);
+
+	UFUNCTION(BlueprintSetter, Category = "USD")
+	void SetBoundsMax(const FVector& NewMax);
+
+	UFUNCTION(BlueprintSetter, Category = "USD")
+	void SetDrawMode(EUsdDrawMode NewDrawMode);
+
+	UFUNCTION(BlueprintSetter, Category = "USD")
+	void SetBoundsColor(FLinearColor NewColor);
+
+	UFUNCTION(BlueprintSetter, Category = "USD")
+	void SetCardGeometry(EUsdModelCardGeometry NewGeometry);
+
+	UFUNCTION(BlueprintSetter, Category = "USD|Cards")
+	void SetCardTextureXPos(UTexture2D* NewTexture);
+
+	UFUNCTION(BlueprintSetter, Category = "USD|Cards")
+	void SetCardTextureYPos(UTexture2D* NewTexture);
+
+	UFUNCTION(BlueprintSetter, Category = "USD|Cards")
+	void SetCardTextureZPos(UTexture2D* NewTexture);
+
+	UFUNCTION(BlueprintSetter, Category = "USD|Cards")
+	void SetCardTextureXNeg(UTexture2D* NewTexture);
+
+	UFUNCTION(BlueprintSetter, Category = "USD|Cards")
+	void SetCardTextureYNeg(UTexture2D* NewTexture);
+
+	UFUNCTION(BlueprintSetter, Category = "USD|Cards")
+	void SetCardTextureZNeg(UTexture2D* NewTexture);
+
+	UFUNCTION(BlueprintPure, Category = "USD|Cards")
+	UTexture2D* GetTextureForFace(EUsdModelCardFace Face) const;
+
+	UFUNCTION(BlueprintCallable, Category = "USD|Cards")
+	void SetTextureForFace(EUsdModelCardFace Face, UTexture2D* Texture);
+
+public:
+	UUsdDrawModeComponent();
+	virtual ~UUsdDrawModeComponent() override {};
+
+	// Begin UPrimitiveComponent interface
+	virtual FPrimitiveSceneProxy* CreateSceneProxy() override;
+	virtual FBoxSphereBounds CalcBounds(const FTransform& LocalToWorld) const override;
+	virtual void GetUsedMaterials(TArray<UMaterialInterface*>& OutMaterials, bool bGetDebugMaterials = false) const override;
+#if WITH_EDITOR
+	virtual void PostEditChangeProperty(FPropertyChangedEvent& PropertyChangedEvent) override;
+	virtual HHitProxy* CreateMeshHitProxy(int32 SectionIndex, int32 MaterialIndex) const override;
+#endif
+	// End UPrimitiveComponent interface
+
+	EUsdModelCardFace GetAuthoredFaces() const;
+	void SetAuthoredFaces(EUsdModelCardFace NewAuthoredFaces);
+
+private:
+	UMaterialInstance* GetOrCreateTextureMaterial(EUsdModelCardFace Face);
+	void RefreshMaterialInstances();
+};
diff --git a/Engine/Plugins/Importers/USDImporter/Source/USDClasses/Public/USDProjectSettings.h b/Engine/Plugins/Importers/USDImporter/Source/USDClasses/Public/USDProjectSettings.h
--- a/Engine/Plugins/Importers/USDImporter/Source/USDClasses/Public/USDProjectSettings.h
+++ b/Engine/Plugins/Importers/USDImporter/Source/USDClasses/Public/USDProjectSettings.h
@@ -110,6 +110,16 @@ public:
 
 	// Note that the below properties being FSoftObjectPath ensure that these materials are cooked into packaged games
 
+	/**
+	 * What material to use for UUsdDrawModeComponents with "Cards" draw mode and provided textures (corresponding to
+	 * UsdGeomModelAPI with the "cards" drawMode).
+	 * Each face of the card geometry will use a separate texture material instance, and the UTexture2D will be set
+	 * as a material parameter named "Texture".
+	 * You can swap this with your own material, but make sure the replacement material has a "Texture" parameter
+	 */
+	UPROPERTY( config, EditAnywhere, Category = "USD|Reference Materials", meta = ( AllowedClasses = "/Script/Engine.MaterialInterface" ) )
+	FSoftObjectPath ReferenceModelCardTextureMaterial = FSoftObjectPath{ TEXT( "/USDImporter/Materials/CardTextureMaterial.CardTextureMaterial" ) };
+
 	/**
 	 * What material to use as reference material when creating material instances from USD materials.
 	 * You can swap these with your own materials, but make sure that the replacement materials have parameters with
diff --git a/Engine/Plugins/Importers/USDImporter/Source/USDClasses/USDClasses.Build.cs b/Engine/Plugins/Importers/USDImporter/Source/USDClasses/USDClasses.Build.cs
--- a/Engine/Plugins/Importers/USDImporter/Source/USDClasses/USDClasses.Build.cs
+++ b/Engine/Plugins/Importers/USDImporter/Source/USDClasses/USDClasses.Build.cs
@@ -26,12 +26,24 @@ namespace UnrealBuildTool.Rules
 				{
 					"DeveloperSettings",
 					"GeometryCache",
+					"GeometryFramework", // For BaseDynamicMeshSceneProxy.h
 					"HairStrandsCore",
 					"Json", // To read/write plugInfo.json files from UnrealUSDWrapper.cpp
 					"MeshDescription",
+					"RenderCore", // For UUsdDrawModeComponent implementation
 					"RHI",
 				}
 			);
+
+			if (Target.bBuildEditor)
+			{
+				PrivateDependencyModuleNames.AddRange(
+					new string[]
+					{
+						"MaterialEditor", // For SetMaterialInstanceParent, when managing the DisplayColor materials
+					}
+				);
+			}
 		}
 	}
 }
diff --git a/Engine/Plugins/Importers/USDImporter/Source/USDExporter/Private/USDConversionBlueprintContext.cpp b/Engine/Plugins/Importers/USDImporter/Source/USDExporter/Private/USDConversionBlueprintContext.cpp
--- a/Engine/Plugins/Importers/USDImporter/Source/USDExporter/Private/USDConversionBlueprintContext.cpp
+++ b/Engine/Plugins/Importers/USDImporter/Source/USDExporter/Private/USDConversionBlueprintContext.cpp
@@ -228,6 +228,22 @@ bool UUsdConversionBlueprintContext::ConvertSpotLightComponent( const USpotLight
 #endif // USE_USD_SDK
 }
 
+bool UUsdConversionBlueprintContext::ConvertBoundsComponent(const UUsdDrawModeComponent* Component, const FString& PrimPath, float TimeCode)
+{
+#if USE_USD_SDK
+	UE::FUsdPrim Prim = UnrealToUsdImpl::GetPrim(Stage, PrimPath);
+	if (!Prim || !Component)
+	{
+		return false;
+	}
+
+	const bool bWriteExtents = true;
+	return UnrealToUsd::ConvertBoundsComponent(*Component, Prim, bWriteExtents, TimeCode == FLT_MAX ? UsdUtils::GetDefaultTimeCode() : TimeCode);
+#else
+	return false;
+#endif	  // USE_USD_SDK
+}
+
 bool UUsdConversionBlueprintContext::ConvertSceneComponent( const USceneComponent* Component, const FString& PrimPath )
 {
 #if USE_USD_SDK
diff --git a/Engine/Plugins/Importers/USDImporter/Source/USDExporter/Public/USDConversionBlueprintContext.h b/Engine/Plugins/Importers/USDImporter/Source/USDExporter/Public/USDConversionBlueprintContext.h
--- a/Engine/Plugins/Importers/USDImporter/Source/USDExporter/Public/USDConversionBlueprintContext.h
+++ b/Engine/Plugins/Importers/USDImporter/Source/USDExporter/Public/USDConversionBlueprintContext.h
@@ -26,6 +26,7 @@ class URectLightComponent;
 class USceneComponent;
 class USkyLightComponent;
 class USpotLightComponent;
+class UUsdDrawModeComponent;
 
 /**
  * Wraps the UnrealToUsd component conversion functions from the USDUtilities module so that they can be used by
@@ -115,6 +116,9 @@ public:
 	UFUNCTION( BlueprintCallable, Category = "Component conversion" )
 	bool ConvertSpotLightComponent( const USpotLightComponent* Component, const FString& PrimPath, float TimeCode = 3.402823466e+38F );
 
+	UFUNCTION(BlueprintCallable, Category = "Component conversion")
+	bool ConvertBoundsComponent(const UUsdDrawModeComponent* Component, const FString& PrimPath, float TimeCode = 3.402823466e+38F);
+
 	UFUNCTION( BlueprintCallable, Category = "Component conversion" )
 	bool ConvertSceneComponent( const USceneComponent* Component, const FString& PrimPath );
 
diff --git a/Engine/Plugins/Importers/USDImporter/Source/USDSchemas/Private/MeshTranslationImpl.cpp b/Engine/Plugins/Importers/USDImporter/Source/USDSchemas/Private/MeshTranslationImpl.cpp
--- a/Engine/Plugins/Importers/USDImporter/Source/USDSchemas/Private/MeshTranslationImpl.cpp
+++ b/Engine/Plugins/Importers/USDImporter/Source/USDSchemas/Private/MeshTranslationImpl.cpp
@@ -3,8 +3,9 @@
 #include "MeshTranslationImpl.h"
 
 #include "USDAssetCache2.h"
-#include "USDConversionUtils.h"
 #include "USDAssetUserData.h"
+#include "USDClassesModule.h"
+#include "USDConversionUtils.h"
 #include "USDGeomMeshConversion.h"
 #include "USDInfoCache.h"
 #include "USDLog.h"
@@ -458,17 +459,19 @@ TMap<const UsdUtils::FUsdPrimMaterialSlot*, UMaterialInterface*> MeshTranslation
 				// Need to create a new DisplayColor material
 				if (Material == nullptr)
 				{
-					if (TOptional< UsdUtils::FDisplayColorMaterial > DisplayColorDesc = UsdUtils::FDisplayColorMaterial::FromString(Slot.MaterialSource))
+					if (TOptional<IUsdClassesModule::FDisplayColorMaterial> DisplayColorDesc = IUsdClassesModule::FDisplayColorMaterial::FromString(
+							Slot.MaterialSource
+						))
 					{
 						UMaterialInstance* MaterialInstance = nullptr;
 
 						if (GIsEditor)  // Editor, PIE => true; Standlone, packaged => false
 						{
-							MaterialInstance = UsdUtils::CreateDisplayColorMaterialInstanceConstant(DisplayColorDesc.GetValue());
+							MaterialInstance = IUsdClassesModule::CreateDisplayColorMaterialInstanceConstant(DisplayColorDesc.GetValue());
 						}
 						else
 						{
-							MaterialInstance = UsdUtils::CreateDisplayColorMaterialInstanceDynamic(DisplayColorDesc.GetValue());
+							MaterialInstance = IUsdClassesModule::CreateDisplayColorMaterialInstanceDynamic(DisplayColorDesc.GetValue());
 						}
 
 						if(MaterialInstance)
diff --git a/Engine/Plugins/Importers/USDImporter/Source/USDSchemas/Private/USDGeomCameraTranslator.cpp b/Engine/Plugins/Importers/USDImporter/Source/USDSchemas/Private/USDGeomCameraTranslator.cpp
--- a/Engine/Plugins/Importers/USDImporter/Source/USDSchemas/Private/USDGeomCameraTranslator.cpp
+++ b/Engine/Plugins/Importers/USDImporter/Source/USDSchemas/Private/USDGeomCameraTranslator.cpp
@@ -2,6 +2,7 @@
 
 #include "USDGeomCameraTranslator.h"
 
+#include "USDDrawModeComponent.h"
 #include "USDMemory.h"
 #include "USDPrimConversion.h"
 #include "USDTypesConversion.h"
@@ -22,10 +23,17 @@ USceneComponent* FUsdGeomCameraTranslator::CreateComponents()
 {
 	USceneComponent* Component = nullptr;
 
+	EUsdDrawMode DrawMode = UsdUtils::GetAppliedDrawMode(GetPrim());
+	if (DrawMode != EUsdDrawMode::Default)
+	{
+		Component = CreateAlternativeDrawModeComponents(DrawMode);
+	}
+
 	// Check if this prim actually originated from a CineCameraComponent that was the main camera component
 	// of an ACineCameraActor. If so, then the USDGeomXformableTranslator that ran on our parent prim
 	// likely already created an ACineCameraActor for us, and we can just take its already created
 	// main camera component instead
+	if (!Component)
 	{
 		FScopedUsdAllocs UsdAllocs;
 
@@ -61,13 +69,16 @@ USceneComponent* FUsdGeomCameraTranslator::CreateComponents()
 		}
 	}
 
-	if ( !Component )
+	if (!Component)
 	{
-		bool bNeedsActor = true;
-		Component = CreateComponentsEx( TSubclassOf< USceneComponent >( UCineCameraComponent::StaticClass() ), bNeedsActor );
+		const bool bNeedsActor = true;
+		Component = CreateComponentsEx({UCineCameraComponent::StaticClass()}, bNeedsActor);
 	}
 
-	UpdateComponents( Component );
+	// We pulled UpdateComponents outside CreateComponentsEx as in some cases we don't want to do it
+	// right away (like on FUsdGeomPointInstancerTranslator::CreateComponents)
+	UpdateComponents(Component);
+
 	return Component;
 }
 
@@ -107,6 +118,14 @@ void FUsdGeomCameraTranslator::UpdateComponents( USceneComponent* SceneComponent
 
 bool FUsdGeomCameraTranslator::CollapsesChildren( ECollapsingType CollapsingType ) const
 {
+	// If we have a custom draw mode, it means we should draw bounds/cards/etc. instead
+	// of our entire subtree, which is basically the same thing as collapsing
+	EUsdDrawMode DrawMode = UsdUtils::GetAppliedDrawMode(GetPrim());
+	if (DrawMode != EUsdDrawMode::Default)
+	{
+		return true;
+	}
+
 	return false;
 }
 
diff --git a/Engine/Plugins/Importers/USDImporter/Source/USDSchemas/Private/USDGeomMeshTranslator.cpp b/Engine/Plugins/Importers/USDImporter/Source/USDSchemas/Private/USDGeomMeshTranslator.cpp
--- a/Engine/Plugins/Importers/USDImporter/Source/USDSchemas/Private/USDGeomMeshTranslator.cpp
+++ b/Engine/Plugins/Importers/USDImporter/Source/USDSchemas/Private/USDGeomMeshTranslator.cpp
@@ -7,6 +7,7 @@
 #include "MeshTranslationImpl.h"
 #include "UnrealUSDWrapper.h"
 #include "USDAssetUserData.h"
+#include "USDDrawModeComponent.h"
 #include "USDClassesModule.h"
 #include "USDConversionUtils.h"
 #include "USDGeomMeshConversion.h"
@@ -1000,6 +1001,14 @@ void FUsdGeomMeshTranslator::CreateAssets()
 		return Super::CreateAssets();
 	}
 
+	// Don't bother generating assets if we're going to just draw some bounds for this prim instead
+	EUsdDrawMode DrawMode = UsdUtils::GetAppliedDrawMode(GetPrim());
+	if (DrawMode != EUsdDrawMode::Default)
+	{
+		CreateAlternativeDrawModeAssets(DrawMode);
+		return;
+	}
+
 	TSharedRef< FGeomMeshCreateAssetsTaskChain > AssetsTaskChain = MakeShared< FGeomMeshCreateAssetsTaskChain >(Context, PrimPath);
 
 	Context->TranslatorTasks.Add(MoveTemp(AssetsTaskChain));
@@ -1007,70 +1016,26 @@ void FUsdGeomMeshTranslator::CreateAssets()
 
 USceneComponent* FUsdGeomMeshTranslator::CreateComponents()
 {
-	if (!IsMeshPrim())
-	{
-		return Super::CreateComponents();
-	}
-
-	TOptional< TSubclassOf< USceneComponent > > ComponentType;
-
-	USceneComponent* SceneComponent = CreateComponentsEx(ComponentType, {});
-	UpdateComponents(SceneComponent);
-
-	// Handle material overrides
-	// Note: This can be here and not in USDGeomXformableTranslator because there is no way that a collapsed mesh prim could end up with a material override
-	if (UStaticMeshComponent* StaticMeshComponent = Cast<UStaticMeshComponent>(SceneComponent))
-	{
-		if (Context->InfoCache)
-		{
-			if (UStaticMesh* StaticMesh = Context->InfoCache->GetSingleAssetForPrim<UStaticMesh>(
-				PrimPath
-			))
-			{
-				TArray<UMaterialInterface*> ExistingAssignments;
-				for (FStaticMaterial& StaticMaterial : StaticMesh->GetStaticMaterials())
-				{
-					ExistingAssignments.Add(StaticMaterial.MaterialInterface);
-				}
-
-				MeshTranslationImpl::SetMaterialOverrides(
-					GetPrim(),
-					ExistingAssignments,
-					*StaticMeshComponent,
-					*Context->AssetCache.Get(),
-					*Context->InfoCache.Get(),
-					Context->Time,
-					Context->ObjectFlags,
-					Context->bAllowInterpretingLODs,
-					Context->RenderContext,
-					Context->MaterialPurpose
-				);
-			}
-		}
-	}
-
-	return SceneComponent;
+	return Super::CreateComponents();
 }
 
 void FUsdGeomMeshTranslator::UpdateComponents(USceneComponent* SceneComponent)
 {
-	if (!IsMeshPrim())
+	if (IsMeshPrim())
 	{
-		return Super::UpdateComponents(SceneComponent);
-	}
-
-	if (SceneComponent)
-	{
-		SceneComponent->Modify();
-	}
+		if (SceneComponent)
+		{
+			SceneComponent->Modify();
+		}
 
-	if (UsdUtils::IsAnimatedMesh(GetPrim()))
-	{
-		// The assets might have changed since our attributes are animated
-		// Note that we must wait for these to complete as they make take a while and we want to
-		// reassign our new static meshes when we get to FUsdGeomXformableTranslator::UpdateComponents
-		CreateAssets();
-		Context->CompleteTasks();
+		if (UsdUtils::IsAnimatedMesh(GetPrim()))
+		{
+			// The assets might have changed since our attributes are animated
+			// Note that we must wait for these to complete as they make take a while and we want to
+			// reassign our new static meshes when we get to FUsdGeomXformableTranslator::UpdateComponents
+			CreateAssets();
+			Context->CompleteTasks();
+		}
 	}
 
 	Super::UpdateComponents(SceneComponent);
@@ -1083,6 +1048,14 @@ bool FUsdGeomMeshTranslator::CollapsesChildren(ECollapsingType CollapsingType) c
 		return Super::CollapsesChildren(CollapsingType);
 	}
 
+	// If we have a custom draw mode, it means we should draw bounds/cards/etc. instead
+	// of our entire subtree, which is basically the same thing as collapsing
+	EUsdDrawMode DrawMode = UsdUtils::GetAppliedDrawMode(GetPrim());
+	if (DrawMode != EUsdDrawMode::Default)
+	{
+		return true;
+	}
+
 	// We can't claim we collapse anything here since we'll just parse the mesh for this prim and that's it,
 	// otherwise the translation context wouldn't spawn translators for our child prims.
 	// Another approach would be to actually recursively collapse our child mesh prims, but that leads to a few
diff --git a/Engine/Plugins/Importers/USDImporter/Source/USDSchemas/Private/USDGeomPointInstancerTranslator.cpp b/Engine/Plugins/Importers/USDImporter/Source/USDSchemas/Private/USDGeomPointInstancerTranslator.cpp
--- a/Engine/Plugins/Importers/USDImporter/Source/USDSchemas/Private/USDGeomPointInstancerTranslator.cpp
+++ b/Engine/Plugins/Importers/USDImporter/Source/USDSchemas/Private/USDGeomPointInstancerTranslator.cpp
@@ -5,6 +5,7 @@
 #include "MeshTranslationImpl.h"
 #include "USDAssetCache.h"
 #include "USDAssetUserData.h"
+#include "USDDrawModeComponent.h"
 #include "USDConversionUtils.h"
 #include "USDGeomMeshConversion.h"
 #include "USDGeomMeshTranslator.h"
@@ -153,6 +154,14 @@ void FUsdGeomPointInstancerTranslator::CreateAssets()
 {
 	TRACE_CPUPROFILER_EVENT_SCOPE( FUsdGeomPointInstancerTranslator::CreateAssets );
 
+	// Don't bother generating assets if we're going to just draw some bounds for this prim instead
+	EUsdDrawMode DrawMode = UsdUtils::GetAppliedDrawMode(GetPrim());
+	if (DrawMode != EUsdDrawMode::Default)
+	{
+		CreateAlternativeDrawModeAssets(DrawMode);
+		return;
+	}
+
 	FScopedUsdAllocs Allocs;
 
 	pxr::UsdPrim Prim = GetPrim();
@@ -272,13 +281,21 @@ USceneComponent* FUsdGeomPointInstancerTranslator::CreateComponents()
 	// HISM component for each prototype.
 	// We always request a scene component here explicitly or else we'll be upgraded to a static mesh component by the mechanism that
 	// handles collapsed meshes/static mesh components for the GeomXFormable translator.
-	USceneComponent* MainSceneComponent = CreateComponentsEx( { USceneComponent::StaticClass() }, {} );
-	if ( !MainSceneComponent )
+	bool bCreateChildHISMs = false;
+	USceneComponent* MainSceneComponent = nullptr;
+	EUsdDrawMode DrawMode = UsdUtils::GetAppliedDrawMode(GetPrim());
+	if (DrawMode == EUsdDrawMode::Default)
+	{
+		bCreateChildHISMs = true;
+		MainSceneComponent = CreateComponentsEx({USceneComponent::StaticClass()}, {});
+	}
+	else
 	{
-		return MainSceneComponent;
+		MainSceneComponent = CreateAlternativeDrawModeComponents(DrawMode);
 	}
 
 	// Actually create the child HISM components for each point instancer prototype
+	if (bCreateChildHISMs)
 	{
 		FScopedUsdAllocs UsdAllocs;
 
@@ -360,175 +377,191 @@ USceneComponent* FUsdGeomPointInstancerTranslator::CreateComponents()
 
 void FUsdGeomPointInstancerTranslator::UpdateComponents(USceneComponent* PointInstancerRootComponent)
 {
-	FScopedUsdAllocs UsdAllocs;
-
-	pxr::UsdPrim Prim = GetPrim();
-	pxr::UsdGeomPointInstancer PointInstancer( Prim );
-	if ( !PointInstancer )
+	if (!PointInstancerRootComponent)
 	{
 		return;
 	}
 
-	pxr::SdfPathVector PrototypePaths;
-	if ( !PointInstancer.GetPrototypesRel().GetTargets( &PrototypePaths ) )
-	{
-		return;
-	}
-
-	if (!Context->AssetCache.IsValid() || !Context->InfoCache.IsValid())
+	// We always spawn exactly an USceneComponent for the "parent" component of the point instancer, so early out if
+	// we dont' have one. This can happen now if we have an alternative draw mode for this point instancer, at which
+	// point this could be an UUsdDrawModeComponent
+	if (PointInstancerRootComponent->GetClass() == USceneComponent::StaticClass())
 	{
-		return;
-	}
-	UUsdAssetCache2& AssetCache = *Context->AssetCache.Get();
-	FUsdInfoCache& InfoCache = *Context->InfoCache.Get();
-
-	// Lets pretend ParentComponent is pointing to the parent USceneComponent while we create the child HISMs, so they get
-	// automatically attached to it as children
-	TGuardValue< USceneComponent* > ParentComponentGuard{Context->ParentComponent, PointInstancerRootComponent};
+		FScopedUsdAllocs UsdAllocs;
 
-	const TArray<TObjectPtr<USceneComponent>>& AttachedChildren = PointInstancerRootComponent->GetAttachChildren();
-	TArray<UHierarchicalInstancedStaticMeshComponent*> AttachedHISMs;
-	AttachedHISMs.Reserve(AttachedChildren.Num());
-	for (const TObjectPtr<USceneComponent>& AttachedChild : AttachedChildren)
-	{
-		if (UHierarchicalInstancedStaticMeshComponent* HISM = Cast<UHierarchicalInstancedStaticMeshComponent>(AttachedChild))
+		pxr::UsdPrim Prim = GetPrim();
+		pxr::UsdGeomPointInstancer PointInstancer( Prim );
+		if ( !PointInstancer )
 		{
-			AttachedHISMs.Add(HISM);
+			return;
 		}
-	}
 
-	// We "link" the prototype meshes to the point instancer, but we don't know which mesh corresponds to each
-	// prototype, as we translate these with task pools and some of those prototypes may have generated nullptr.
-	// We always put the prototype path on the asset import data though, so here we use that to figure out where
-	// each mesh should go
-	TArray<UStaticMesh*> PrototypeMeshArr = InfoCache.GetAssetsForPrim<UStaticMesh>(PrimPath);
-	std::unordered_map<pxr::SdfPath, UStaticMesh*, pxr::SdfPath::Hash> PrototypeMeshes;
-	PrototypeMeshes.reserve(PrototypeMeshArr.Num());
-	for (UStaticMesh* PrototypeMesh : PrototypeMeshArr)
-	{
-		if (UUsdAssetUserData* UserData = PrototypeMesh->GetAssetUserData<UUsdAssetUserData>())
+		pxr::SdfPathVector PrototypePaths;
+		if ( !PointInstancer.GetPrototypesRel().GetTargets( &PrototypePaths ) )
 		{
-			for (const FString& SourcePrimPath : UserData->PrimPaths)
-			{
-				pxr::SdfPath PrototypePath = UnrealToUsd::ConvertPath(*SourcePrimPath).Get();
-				PrototypeMeshes[PrototypePath] = PrototypeMesh;
-			}
+			return;
 		}
-	}
-
-	TArray<TFuture<TTuple<UHierarchicalInstancedStaticMeshComponent*, TArray<FTransform>>>> Tasks;
-	FScopedSlowTask PrototypePathsSlowTask( ( float ) PrototypePaths.size(), LOCTEXT( "GeomPointUpdateComponents", "Updating HierarchicalInstancedStaticMeshComponents for point instancers" ) );
-	for ( int32 PrototypeIndex = 0; PrototypeIndex < PrototypePaths.size(); ++PrototypeIndex )
-	{
-		PrototypePathsSlowTask.EnterProgressFrame();
 
-		pxr::SdfPath PrototypePath = PrototypePaths[PrototypeIndex];
-
-		pxr::UsdPrim PrototypeUsdPrim = Prim.GetStage()->GetPrimAtPath( PrototypePath );
-		if ( !PrototypeUsdPrim )
+		if (!Context->AssetCache.IsValid() || !Context->InfoCache.IsValid())
 		{
-			UE_LOG(
-				LogUsd,
-				Warning,
-				TEXT("Failed to find prototype '%s' for PointInstancer '%s' when updating components"),
-				*UsdToUnreal::ConvertPath(PrototypePath),
-				*PrimPath.GetString()
-			);
-			continue;
+			return;
 		}
+		UUsdAssetCache2& AssetCache = *Context->AssetCache.Get();
+		FUsdInfoCache& InfoCache = *Context->InfoCache.Get();
 
-		UHierarchicalInstancedStaticMeshComponent* HISMComponent = nullptr;
+		// Lets pretend ParentComponent is pointing to the parent USceneComponent while we create the child HISMs, so they get
+		// automatically attached to it as children
+		TGuardValue< USceneComponent* > ParentComponentGuard{Context->ParentComponent, PointInstancerRootComponent};
 
-		// The user could have just manually deleted the component, so we must check
-		if (!AttachedHISMs.IsValidIndex(PrototypeIndex))
+		const TArray<TObjectPtr<USceneComponent>>& AttachedChildren = PointInstancerRootComponent->GetAttachChildren();
+		TArray<UHierarchicalInstancedStaticMeshComponent*> AttachedHISMs;
+		AttachedHISMs.Reserve(AttachedChildren.Num());
+		for (const TObjectPtr<USceneComponent>& AttachedChild : AttachedChildren)
 		{
-			UE_LOG(
-				LogUsd,
-				Warning,
-				TEXT("Failed to find corresponding HISM component for prototype '%s' of PointInstancer '%s'. Cancelling component update"),
-				*UsdToUnreal::ConvertPath(PrototypePath),
-				*PrimPath.GetString()
-			);
-			break;
+			if (UHierarchicalInstancedStaticMeshComponent* HISM = Cast<UHierarchicalInstancedStaticMeshComponent>(AttachedChild))
+			{
+				AttachedHISMs.Add(HISM);
+			}
 		}
-		HISMComponent = AttachedHISMs[PrototypeIndex];
 
-		// If our prototype was a LOD mesh we will have used the path of one of the actual LOD meshes to start the FGeomMeshCreateAssetsTaskChain,
-		// so we have to look for our resulting mesh with the same path
-		if ( Context->bAllowInterpretingLODs && UsdUtils::DoesPrimContainMeshLODs( PrototypeUsdPrim ) )
+		// We "link" the prototype meshes to the point instancer, but we don't know which mesh corresponds to each
+		// prototype, as we translate these with task pools and some of those prototypes may have generated nullptr.
+		// We always put the prototype path on the asset import data though, so here we use that to figure out where
+		// each mesh should go
+		TArray<UStaticMesh*> PrototypeMeshArr = InfoCache.GetAssetsForPrim<UStaticMesh>(PrimPath);
+		std::unordered_map<pxr::SdfPath, UStaticMesh*, pxr::SdfPath::Hash> PrototypeMeshes;
+		PrototypeMeshes.reserve(PrototypeMeshArr.Num());
+		for (UStaticMesh* PrototypeMesh : PrototypeMeshArr)
 		{
-			pxr::UsdPrimSiblingRange PrimRange = PrototypeUsdPrim.GetChildren();
-			for ( pxr::UsdPrimSiblingRange::iterator PrimRangeIt = PrimRange.begin(); PrimRangeIt != PrimRange.end(); ++PrimRangeIt )
+			if (UUsdAssetUserData* UserData = PrototypeMesh->GetAssetUserData<UUsdAssetUserData>())
 			{
-				const pxr::UsdPrim& Child = *PrimRangeIt;
-				if ( pxr::UsdGeomMesh ChildMesh{ Child } )
+				for (const FString& SourcePrimPath : UserData->PrimPaths)
 				{
-					PrototypeUsdPrim = Child;
-					PrototypePath = Child.GetPrimPath();
-					break;
+					pxr::SdfPath PrototypePath = UnrealToUsd::ConvertPath(*SourcePrimPath).Get();
+					PrototypeMeshes[PrototypePath] = PrototypeMesh;
 				}
 			}
 		}
 
-		// This mesh could be nullptr, but that's OK
-		UStaticMesh* StaticMesh = nullptr;
-		std::unordered_map<pxr::SdfPath, UStaticMesh*, pxr::SdfPath::Hash>::iterator Iter = PrototypeMeshes.find(PrototypePath);
-		if (Iter != PrototypeMeshes.end())
+		TArray<TFuture<TTuple<UHierarchicalInstancedStaticMeshComponent*, TArray<FTransform>>>> Tasks;
+		FScopedSlowTask PrototypePathsSlowTask( ( float ) PrototypePaths.size(), LOCTEXT( "GeomPointUpdateComponents", "Updating HierarchicalInstancedStaticMeshComponents for point instancers" ) );
+		for ( int32 PrototypeIndex = 0; PrototypeIndex < PrototypePaths.size(); ++PrototypeIndex )
 		{
-			StaticMesh = Iter->second;
-		}
-		UsdGeomPointInstancerTranslatorImpl::SetStaticMesh( StaticMesh, *HISMComponent );
-
-		// Evaluating point instancer can take a long time and is thread-safe. Move to async task while we work on something else.
-		pxr::UsdTimeCode TimeCode{ Context->Time };
-		FUsdStageInfo StageInfo{ Prim.GetStage() };
-		Tasks.Emplace(
-			Async(
-				EAsyncExecution::ThreadPool,
-				[ TimeCode, StageInfo, PointInstancer, PrototypeIndex, HISMComponent ]()
-				{
-					TArray<FTransform> InstanceTransforms;
-					UsdUtils::GetPointInstancerTransforms( StageInfo, PointInstancer, PrototypeIndex, TimeCode, InstanceTransforms );
+			PrototypePathsSlowTask.EnterProgressFrame();
+
+			pxr::SdfPath PrototypePath = PrototypePaths[PrototypeIndex];
+
+			pxr::UsdPrim PrototypeUsdPrim = Prim.GetStage()->GetPrimAtPath( PrototypePath );
+			if ( !PrototypeUsdPrim )
+			{
+				UE_LOG(
+					LogUsd,
+					Warning,
+					TEXT("Failed to find prototype '%s' for PointInstancer '%s' when updating components"),
+					*UsdToUnreal::ConvertPath(PrototypePath),
+					*PrimPath.GetString()
+				);
+				continue;
+			}
 
-					return MakeTuple( HISMComponent, MoveTemp( InstanceTransforms ) );
+			UHierarchicalInstancedStaticMeshComponent* HISMComponent = nullptr;
+
+			// The user could have just manually deleted the component, so we must check
+			if (!AttachedHISMs.IsValidIndex(PrototypeIndex))
+			{
+				UE_LOG(
+					LogUsd,
+					Warning,
+					TEXT("Failed to find corresponding HISM component for prototype '%s' of PointInstancer '%s'. Cancelling component update"),
+					*UsdToUnreal::ConvertPath(PrototypePath),
+					*PrimPath.GetString()
+				);
+				break;
+			}
+			HISMComponent = AttachedHISMs[PrototypeIndex];
+
+			// If our prototype was a LOD mesh we will have used the path of one of the actual LOD meshes to start the FGeomMeshCreateAssetsTaskChain,
+			// so we have to look for our resulting mesh with the same path
+			if ( Context->bAllowInterpretingLODs && UsdUtils::DoesPrimContainMeshLODs( PrototypeUsdPrim ) )
+			{
+				pxr::UsdPrimSiblingRange PrimRange = PrototypeUsdPrim.GetChildren();
+				for ( pxr::UsdPrimSiblingRange::iterator PrimRangeIt = PrimRange.begin(); PrimRangeIt != PrimRange.end(); ++PrimRangeIt )
+				{
+					const pxr::UsdPrim& Child = *PrimRangeIt;
+					if ( pxr::UsdGeomMesh ChildMesh{ Child } )
+					{
+						PrototypeUsdPrim = Child;
+						PrototypePath = Child.GetPrimPath();
+						break;
+					}
 				}
-			)
-		);
+			}
 
-		// Handle material overrides
-		if ( StaticMesh )
-		{
-			TArray<UMaterialInterface*> ExistingAssignments;
-			for ( FStaticMaterial& StaticMaterial : StaticMesh->GetStaticMaterials() )
+			// This mesh could be nullptr, but that's OK
+			UStaticMesh* StaticMesh = nullptr;
+			std::unordered_map<pxr::SdfPath, UStaticMesh*, pxr::SdfPath::Hash>::iterator Iter = PrototypeMeshes.find(PrototypePath);
+			if (Iter != PrototypeMeshes.end())
 			{
-				ExistingAssignments.Add( StaticMaterial.MaterialInterface );
+				StaticMesh = Iter->second;
 			}
+			UsdGeomPointInstancerTranslatorImpl::SetStaticMesh( StaticMesh, *HISMComponent );
+
+			// Evaluating point instancer can take a long time and is thread-safe. Move to async task while we work on something else.
+			pxr::UsdTimeCode TimeCode{ Context->Time };
+			FUsdStageInfo StageInfo{ Prim.GetStage() };
+			Tasks.Emplace(
+				Async(
+					EAsyncExecution::ThreadPool,
+					[ TimeCode, StageInfo, PointInstancer, PrototypeIndex, HISMComponent ]()
+					{
+						TArray<FTransform> InstanceTransforms;
+						UsdUtils::GetPointInstancerTransforms( StageInfo, PointInstancer, PrototypeIndex, TimeCode, InstanceTransforms );
 
-			MeshTranslationImpl::SetMaterialOverrides(
-				PrototypeUsdPrim,
-				ExistingAssignments,
-				*HISMComponent,
-				AssetCache,
-				InfoCache,
-				Context->Time,
-				Context->ObjectFlags,
-				Context->bAllowInterpretingLODs,
-				Context->RenderContext,
-				Context->MaterialPurpose
+						return MakeTuple( HISMComponent, MoveTemp( InstanceTransforms ) );
+					}
+				)
 			);
+
+			// Handle material overrides
+			if ( StaticMesh )
+			{
+				TArray<UMaterialInterface*> ExistingAssignments;
+				for ( FStaticMaterial& StaticMaterial : StaticMesh->GetStaticMaterials() )
+				{
+					ExistingAssignments.Add( StaticMaterial.MaterialInterface );
+				}
+
+				MeshTranslationImpl::SetMaterialOverrides(
+					PrototypeUsdPrim,
+					ExistingAssignments,
+					*HISMComponent,
+					AssetCache,
+					InfoCache,
+					Context->Time,
+					Context->ObjectFlags,
+					Context->bAllowInterpretingLODs,
+					Context->RenderContext,
+					Context->MaterialPurpose
+				);
+			}
 		}
-	}
 
-	// Wait on and assign results of the point instancer.
-	for ( auto& Future : Tasks )
-	{
-		TTuple<UHierarchicalInstancedStaticMeshComponent*, TArray<FTransform>> Result{ Future.Get() };
-		UsdGeomPointInstancerTranslatorImpl::ApplyPointInstanceTransforms( Result.Key, Result.Value );
+		// Wait on and assign results of the point instancer.
+		for ( auto& Future : Tasks )
+		{
+			TTuple<UHierarchicalInstancedStaticMeshComponent*, TArray<FTransform>> Result{ Future.Get() };
+			UsdGeomPointInstancerTranslatorImpl::ApplyPointInstanceTransforms( Result.Key, Result.Value );
+		}
 	}
 
 	Super::UpdateComponents(PointInstancerRootComponent);
 }
 
+bool FUsdGeomPointInstancerTranslator::CollapsesChildren(ECollapsingType CollapsingType) const
+{
+	return true;
+}
+
 bool FUsdGeomPointInstancerTranslator::CanBeCollapsed( ECollapsingType CollapsingType ) const
 {
 	return GCollapseTopLevelPointInstancers;
diff --git a/Engine/Plugins/Importers/USDImporter/Source/USDSchemas/Private/USDGeomPrimitiveTranslator.cpp b/Engine/Plugins/Importers/USDImporter/Source/USDSchemas/Private/USDGeomPrimitiveTranslator.cpp
--- a/Engine/Plugins/Importers/USDImporter/Source/USDSchemas/Private/USDGeomPrimitiveTranslator.cpp
+++ b/Engine/Plugins/Importers/USDImporter/Source/USDSchemas/Private/USDGeomPrimitiveTranslator.cpp
@@ -7,6 +7,7 @@
 #include "UnrealUSDWrapper.h"
 #include "USDAssetImportData.h"
 #include "USDConversionUtils.h"
+#include "USDDrawModeComponent.h"
 #include "USDGeomMeshConversion.h"
 #include "USDGeomMeshTranslator.h"
 #include "USDInfoCache.h"
@@ -123,6 +124,14 @@ void FUsdGeomPrimitiveTranslator::CreateAssets()
 {
 	TRACE_CPUPROFILER_EVENT_SCOPE(FUsdGeomPrimitiveTranslator::CreateAssets);
 
+	// Don't bother generating assets if we're going to just draw some bounds for this prim instead
+	EUsdDrawMode DrawMode = UsdUtils::GetAppliedDrawMode(GetPrim());
+	if (DrawMode != EUsdDrawMode::Default)
+	{
+		CreateAlternativeDrawModeAssets(DrawMode);
+		return;
+	}
+
 	using namespace UsdGeomPrimitiveTranslatorImpl;
 	TSharedRef<FGeomPrimitiveCreateAssetsTaskChain> AssetsTaskChain = MakeShared<FGeomPrimitiveCreateAssetsTaskChain>(Context, PrimPath);
 
@@ -131,7 +140,18 @@ void FUsdGeomPrimitiveTranslator::CreateAssets()
 
 USceneComponent* FUsdGeomPrimitiveTranslator::CreateComponents()
 {
-	USceneComponent* SceneComponent = CreateComponentsEx({}, {});
+	USceneComponent* SceneComponent = nullptr;
+
+	EUsdDrawMode DrawMode = UsdUtils::GetAppliedDrawMode(GetPrim());
+	if (DrawMode == EUsdDrawMode::Default)
+	{
+		SceneComponent = CreateComponentsEx({}, {});
+	}
+	else
+	{
+		SceneComponent = CreateAlternativeDrawModeComponents(DrawMode);
+	}
+
 	UpdateComponents(SceneComponent);
 
 	// Handle material overrides
@@ -184,6 +204,14 @@ void FUsdGeomPrimitiveTranslator::UpdateComponents(USceneComponent* SceneCompone
 
 bool FUsdGeomPrimitiveTranslator::CollapsesChildren(ECollapsingType CollapsingType) const
 {
+	// If we have a custom draw mode, it means we should draw bounds/cards/etc. instead
+	// of our entire subtree, which is basically the same thing as collapsing
+	EUsdDrawMode DrawMode = UsdUtils::GetAppliedDrawMode(GetPrim());
+	if (DrawMode != EUsdDrawMode::Default)
+	{
+		return true;
+	}
+
 	// Gprims shouldn't really have any children so this is not very well defined.
 	// We're going with 'false' here to match FUsdGeomMeshTranslator::CollapsesChildren
 	return false;
diff --git a/Engine/Plugins/Importers/USDImporter/Source/USDSchemas/Private/USDGeomXformableTranslator.cpp b/Engine/Plugins/Importers/USDImporter/Source/USDSchemas/Private/USDGeomXformableTranslator.cpp
--- a/Engine/Plugins/Importers/USDImporter/Source/USDSchemas/Private/USDGeomXformableTranslator.cpp
+++ b/Engine/Plugins/Importers/USDImporter/Source/USDSchemas/Private/USDGeomXformableTranslator.cpp
@@ -5,6 +5,8 @@
 #include "MeshTranslationImpl.h"
 #include "UnrealUSDWrapper.h"
 #include "USDAssetCache.h"
+#include "USDAssetUserData.h"
+#include "USDDrawModeComponent.h"
 #include "USDClassesModule.h"
 #include "USDConversionUtils.h"
 #include "USDErrorUtils.h"
@@ -15,6 +17,7 @@
 #include "USDMemory.h"
 #include "USDPrimConversion.h"
 #include "USDSchemasModule.h"
+#include "USDShadeConversion.h"
 #include "USDTypesConversion.h"
 
 #include "Components/HierarchicalInstancedStaticMeshComponent.h"
@@ -24,6 +27,7 @@
 #include "Components/StaticMeshComponent.h"
 #include "Engine/Level.h"
 #include "Engine/StaticMesh.h"
+#include "Engine/Texture2D.h"
 #include "Engine/World.h"
 #include "LiveLinkComponentController.h"
 #include "LiveLinkRole.h"
@@ -33,6 +37,7 @@
 #include "StaticMeshAttributes.h"
 
 #include "UsdWrappers/SdfPath.h"
+#include "UsdWrappers/UsdGeomBBoxCache.h"
 #include "UsdWrappers/UsdGeomXformable.h"
 #include "UsdWrappers/UsdPrim.h"
 #include "UsdWrappers/UsdStage.h"
@@ -46,6 +51,7 @@
 	#include "pxr/usd/usd/primRange.h"
 	#include "pxr/usd/usd/variantSets.h"
 	#include "pxr/usd/usdGeom/mesh.h"
+	#include "pxr/usd/usdGeom/modelAPI.h"
 	#include "pxr/usd/usdGeom/pointInstancer.h"
 	#include "pxr/usd/usdGeom/subset.h"
 	#include "pxr/usd/usdGeom/xformable.h"
@@ -270,6 +276,14 @@ void FUsdGeomXformableTranslator::CreateAssets()
 		return;
 	}
 
+	// Don't bother generating assets if we're going to just draw some bounds for this prim instead
+	EUsdDrawMode DrawMode = UsdUtils::GetAppliedDrawMode(GetPrim());
+	if (DrawMode != EUsdDrawMode::Default)
+	{
+		CreateAlternativeDrawModeAssets(DrawMode);
+		return;
+	}
+
 	Context->TranslatorTasks.Add( MakeShared< FUsdGeomXformableCreateAssetsTaskChain >( Context, PrimPath ) );
 }
 
@@ -281,7 +295,17 @@ FUsdGeomXformableTranslator::FUsdGeomXformableTranslator( TSubclassOf< USceneCom
 
 USceneComponent* FUsdGeomXformableTranslator::CreateComponents()
 {
-	USceneComponent* SceneComponent = CreateComponentsEx( {}, {} );
+	USceneComponent* SceneComponent = nullptr;
+
+	EUsdDrawMode DrawMode = UsdUtils::GetAppliedDrawMode(GetPrim());
+	if (DrawMode == EUsdDrawMode::Default)
+	{
+		SceneComponent = CreateComponentsEx({}, {});
+	}
+	else
+	{
+		SceneComponent = CreateAlternativeDrawModeComponents(DrawMode);
+	}
 
 	// We pulled UpdateComponents outside CreateComponentsEx as in some cases we don't want to do it
 	// right away (like on FUsdGeomPointInstancerTranslator::CreateComponents)
@@ -433,7 +457,18 @@ USceneComponent* FUsdGeomXformableTranslator::CreateComponentsEx( TOptional< TSu
 		SpawnParameters.Name = Prim.GetName();
 		SpawnParameters.NameMode = FActorSpawnParameters::ESpawnActorNameMode::Requested; // Will generate a unique name in case of a conflict
 
-		UClass* ActorClass = UsdUtils::GetActorTypeForPrim( Prim );
+		UClass* ActorClass = nullptr;
+		if (ComponentType.Get({}) == UUsdDrawModeComponent::StaticClass())
+		{
+			// If we've been told to spawn a bounds component, we never want to create a light or camera actor or etc., as those
+			// come with their own specific root components
+			ActorClass = AActor::StaticClass();
+		}
+		else
+		{
+			ActorClass = UsdUtils::GetActorTypeForPrim(Prim);
+		}
+
 		AActor* SpawnedActor = Context->Level->GetWorld()->SpawnActor( ActorClass, nullptr, SpawnParameters );
 
 		if ( SpawnedActor )
@@ -600,6 +635,8 @@ void FUsdGeomXformableTranslator::UpdateComponents( USceneComponent* SceneCompon
 			SceneComponent->UnregisterComponent();
 		}
 
+		UE::FUsdPrim Prim = GetPrim();
+
 		// If the user modified a mesh parameter (e.g. vertex color), the hash will be different and it will become a separate asset
 		// so we must check for this and assign the new StaticMesh
 		bool bHasMultipleLODs = false;
@@ -634,8 +671,17 @@ void FUsdGeomXformableTranslator::UpdateComponents( USceneComponent* SceneCompon
 				// We'll always re-register when needed below, though.
 			}
 		}
-
-		UE::FUsdPrim Prim = GetPrim();
+		else if (UUsdDrawModeComponent* BoundsComponent = Cast<UUsdDrawModeComponent>(SceneComponent))
+		{
+			TOptional<FWriteScopeLock> BBoxLock;
+			pxr::UsdGeomBBoxCache* PxrBBoxCache = nullptr;
+			if (UE::FUsdGeomBBoxCache* UEBBoxCache = Context->BBoxCache.Get())
+			{
+				BBoxLock.Emplace(UEBBoxCache->Lock);
+				PxrBBoxCache = &static_cast<pxr::UsdGeomBBoxCache&>(*UEBBoxCache);
+			}
+			UsdToUnreal::ConvertBounds(Prim, BoundsComponent, Context->Time, PxrBBoxCache);
+		}
 
 		// Handle LiveLink, but only i we're not a skeletal mesh component: The SkelRootTranslator will deal with the
 		// skeletal version of the LiveLink configuration, we only handle setting up LiveLink for simple transforms
@@ -675,7 +721,232 @@ void FUsdGeomXformableTranslator::UpdateComponents( USceneComponent* SceneCompon
 	}
 }
 
-bool FUsdGeomXformableTranslator::CollapsesChildren( ECollapsingType CollapsingType ) const
+namespace UE::UsdXformableTranslatorImpl::Private
+{
+	void AssignBoundsComponentTextures(UE::FUsdPrim Prim, UUsdDrawModeComponent* BoundsComponent, UUsdAssetCache2& AssetCache, FUsdInfoCache& InfoCache)
+	{
+		if (!Prim)
+		{
+			return;
+		}
+
+		FScopedUsdAllocs Allocs;
+
+		pxr::UsdPrim UsdPrim{Prim};
+		pxr::UsdStageRefPtr Stage = UsdPrim.GetStage();
+		FUsdStageInfo StageInfo{Stage};
+
+		pxr::UsdGeomModelAPI GeomModelAPI{UsdPrim};
+		if (!GeomModelAPI)
+		{
+			return;
+		}
+
+		// Collect textures from the info cache (they will all be linked to this Prim but will have within their AssetUserData the attribute
+		// that they originated from)
+		TArray<UTexture2D*> Textures = InfoCache.GetAssetsForPrim<UTexture2D>(Prim.GetPrimPath());
+		std::unordered_map<pxr::SdfPath, UTexture2D*, pxr::SdfPath::Hash> AttrPathToTextures;
+		for (UTexture2D* Texture : Textures)
+		{
+			if (UUsdAssetUserData* AssetUserData = Texture->GetAssetUserData<UUsdAssetUserData>())
+			{
+				for (const FString& TexturePath : AssetUserData->PrimPaths)
+				{
+					if (TexturePath.IsEmpty())
+					{
+						continue;
+					}
+
+					pxr::SdfPath SdfPath = UnrealToUsd::ConvertPath(*TexturePath).Get();
+					if (SdfPath.IsPropertyPath())
+					{
+						AttrPathToTextures.insert({SdfPath, Texture});
+					}
+				}
+			}
+		}
+
+		// Switch up the faces depending on stage up axis. The effect of metersPerUnit is already baked in the size of the bounds,
+		// but here we "convert the faces" to swap the axes so that the UUsdDrawModeComponent component properties can reference faces in
+		// the UE coordinate system (e.g. PosY in the USD stage will become PosZ in UE coordinate system if the stage is Y up, but then
+		// you will actually see the +Z face in UE, and the PosZ property on the component will be set to match it).
+		pxr::UsdAttribute XPosAttr = GeomModelAPI.GetModelCardTextureXPosAttr();
+		pxr::UsdAttribute YPosAttr = GeomModelAPI.GetModelCardTextureYPosAttr();
+		pxr::UsdAttribute ZPosAttr = GeomModelAPI.GetModelCardTextureZPosAttr();
+		pxr::UsdAttribute XNegAttr = GeomModelAPI.GetModelCardTextureXNegAttr();
+		pxr::UsdAttribute YNegAttr = GeomModelAPI.GetModelCardTextureYNegAttr();
+		pxr::UsdAttribute ZNegAttr = GeomModelAPI.GetModelCardTextureZNegAttr();
+		if (StageInfo.UpAxis == EUsdUpAxis::ZAxis)
+		{
+			Swap(YPosAttr, YNegAttr);
+		}
+		else
+		{
+			Swap(YPosAttr, ZPosAttr);
+			Swap(YNegAttr, ZNegAttr);
+		}
+
+		EUsdModelCardFace AuthoredFaces = EUsdModelCardFace::None;
+
+		using TextureSetterFunc = void (UUsdDrawModeComponent::*)(UTexture2D*);
+
+		TFunction<void(const pxr::UsdAttribute&, TextureSetterFunc, EUsdModelCardFace)> HandleCardFace =
+			[&AuthoredFaces,
+			 BoundsComponent,
+			 &AttrPathToTextures](const pxr::UsdAttribute& Attr, TextureSetterFunc TextureSetter, EUsdModelCardFace Face)
+		{
+			if (Attr && Attr.HasAuthoredValue())
+			{
+				AuthoredFaces |= Face;
+
+				if (TextureSetter)
+				{
+					std::unordered_map<pxr::SdfPath, UTexture2D*, pxr::SdfPath::Hash>::iterator iter = AttrPathToTextures.find(Attr.GetPath());
+					if (iter != AttrPathToTextures.end())
+					{
+						if (UTexture2D* Texture = iter->second)
+						{
+							(BoundsComponent->*TextureSetter)(Texture);
+						}
+					}
+				}
+			}
+		};
+		HandleCardFace(XPosAttr, &UUsdDrawModeComponent::SetCardTextureXPos, EUsdModelCardFace::XPos);
+		HandleCardFace(YPosAttr, &UUsdDrawModeComponent::SetCardTextureYPos, EUsdModelCardFace::YPos);
+		HandleCardFace(ZPosAttr, &UUsdDrawModeComponent::SetCardTextureZPos, EUsdModelCardFace::ZPos);
+		HandleCardFace(XNegAttr, &UUsdDrawModeComponent::SetCardTextureXNeg, EUsdModelCardFace::XNeg);
+		HandleCardFace(YNegAttr, &UUsdDrawModeComponent::SetCardTextureYNeg, EUsdModelCardFace::YNeg);
+		HandleCardFace(ZNegAttr, &UUsdDrawModeComponent::SetCardTextureZNeg, EUsdModelCardFace::ZNeg);
+
+		// Override the AuthoredFaces with the correct value. The texture setter functions will all set the authored faces
+		// when we set any texture in the component, but we also want to set as authored the faces where there *was* some
+		// texture authored in USD but we failed to resolve it, so that we can show that face with vertex color like USD specifies
+		BoundsComponent->SetAuthoredFaces(AuthoredFaces);
+	}
+}	 // namespace UE::UsdXformableTranslatorImpl::Private
+
+USceneComponent* FUsdGeomXformableTranslator::CreateAlternativeDrawModeComponents(EUsdDrawMode DrawMode)
+{
+	// If we're in here, our prim is a model, and we always need actors for model prims anyway
+	const bool bNeedsActor = true;
+
+	switch (DrawMode)
+	{
+		case EUsdDrawMode::Origin:
+		case EUsdDrawMode::Bounds:
+		{
+			return CreateComponentsEx({UUsdDrawModeComponent::StaticClass()}, bNeedsActor);
+			break;
+		}
+		case EUsdDrawMode::Cards:
+		{
+			UUsdDrawModeComponent* Component = Cast<UUsdDrawModeComponent>(CreateComponentsEx({UUsdDrawModeComponent::StaticClass()}, bNeedsActor));
+			if(ensure(Component) && Context->AssetCache && Context->InfoCache)
+			{
+				// For now we only assign textures when creating components, not when updating. Maybe in the future we can
+				// add support for "texture animations"
+				UE::UsdXformableTranslatorImpl::Private::AssignBoundsComponentTextures(
+					GetPrim(),
+					Component,
+					*Context->AssetCache,
+					*Context->InfoCache
+				);
+			}
+			return Component;
+			break;
+		}
+		case EUsdDrawMode::Default:
+		case EUsdDrawMode::Inherited:
+		{
+			ensure(false);
+			break;
+		}
+	}
+
+	return nullptr;
+}
+
+void FUsdGeomXformableTranslator::CreateAlternativeDrawModeAssets(EUsdDrawMode DrawMode)
+{
+	// Currently we just use this function to create the textures that we're going to use on the bounds components,
+	// if applicable
+	if (DrawMode != EUsdDrawMode::Cards || !Context->AssetCache || !Context->InfoCache)
+	{
+		return;
+	}
+
+	FScopedUsdAllocs Allocs;
+
+	pxr::UsdPrim Prim = GetPrim();
+	pxr::UsdGeomModelAPI GeomModelAPI{Prim};
+	if (!GeomModelAPI)
+	{
+		return;
+	}
+
+	TFunction<void(const pxr::UsdAttribute&)> HandleCardFace = [this](const pxr::UsdAttribute& Attr)
+	{
+		if (Attr && Attr.HasAuthoredValue())
+		{
+			const FString ResolvedPath = UsdUtils::GetResolvedTexturePath(Attr);
+			if (ResolvedPath.IsEmpty())
+			{
+				pxr::SdfAssetPath TextureAssetPath;
+				Attr.Get(&TextureAssetPath, Context->Time);
+				FString TargetAssetPath = UsdToUnreal::ConvertString(TextureAssetPath.GetAssetPath());
+
+				UE_LOG(
+					LogUsd,
+					Warning,
+					TEXT("Failed to resolve texture path '%s' specified as a card geometry texture on prim '%s' at time '%f'"),
+					*TargetAssetPath,
+					*UsdToUnreal::ConvertPath(Attr.GetPrim().GetPrimPath()),
+					Context->Time
+				);
+			}
+			else
+			{
+				const FString TextureHash = LexToString(FMD5Hash::HashFile(*ResolvedPath));
+				UTexture2D* Texture = Cast<UTexture2D>(Context->AssetCache->GetCachedAsset(TextureHash));
+
+				if (!Texture)
+				{
+					Texture = Cast<UTexture2D>(UsdUtils::CreateTexture(
+						Attr,
+						UsdToUnreal::ConvertPath(Attr.GetPrim().GetPath()),
+						TEXTUREGROUP_World,
+						Context->AssetCache.Get()
+					));
+
+					Context->AssetCache->CacheAsset(TextureHash, Texture);
+				}
+
+				// We link the textures to the prim, so that if the prim is reloaded the AUsdStageActor knows to potentially
+				// drop the textures. However we put the full attribute path on AssetUserData, so that when we're filling in
+				// our UUsdDrawModeComponent later, we know which texture came from which attribute
+
+				UUsdAssetUserData* TextureUserData = Texture->GetAssetUserData<UUsdAssetUserData>();
+				if (!TextureUserData)
+				{
+					TextureUserData = NewObject<UUsdAssetUserData>(Texture, TEXT("USDAssetUserData"));
+					Texture->AddAssetUserData(TextureUserData);
+				}
+				TextureUserData->PrimPaths.AddUnique(UsdToUnreal::ConvertPath(Attr.GetPath()));
+
+				Context->InfoCache->LinkAssetToPrim(PrimPath, Texture);
+			}
+		}
+	};
+	HandleCardFace(GeomModelAPI.GetModelCardTextureXPosAttr());
+	HandleCardFace(GeomModelAPI.GetModelCardTextureYPosAttr());
+	HandleCardFace(GeomModelAPI.GetModelCardTextureZPosAttr());
+	HandleCardFace(GeomModelAPI.GetModelCardTextureXNegAttr());
+	HandleCardFace(GeomModelAPI.GetModelCardTextureYNegAttr());
+	HandleCardFace(GeomModelAPI.GetModelCardTextureZNegAttr());
+}
+
+bool FUsdGeomXformableTranslator::CollapsesChildren(ECollapsingType CollapsingType) const
 {
 	TRACE_CPUPROFILER_EVENT_SCOPE( FUsdGeomXformableTranslator::CollapsesChildren );
 
@@ -684,6 +955,14 @@ bool FUsdGeomXformableTranslator::CollapsesChildren( ECollapsingType CollapsingT
 		return Context->InfoCache->DoesPathCollapseChildren(PrimPath, CollapsingType);
 	}
 
+	// If we have a custom draw mode, it means we should draw bounds/cards/etc. instead
+	// of our entire subtree, which is basically the same thing as collapsing
+	EUsdDrawMode DrawMode = UsdUtils::GetAppliedDrawMode(GetPrim());
+	if (DrawMode != EUsdDrawMode::Default)
+	{
+		return true;
+	}
+
 	bool bCollapsesChildren = false;
 
 	FScopedUsdAllocs UsdAllocs;
diff --git a/Engine/Plugins/Importers/USDImporter/Source/USDSchemas/Private/USDGeometryCacheTranslator.cpp b/Engine/Plugins/Importers/USDImporter/Source/USDSchemas/Private/USDGeometryCacheTranslator.cpp
--- a/Engine/Plugins/Importers/USDImporter/Source/USDSchemas/Private/USDGeometryCacheTranslator.cpp
+++ b/Engine/Plugins/Importers/USDImporter/Source/USDSchemas/Private/USDGeometryCacheTranslator.cpp
@@ -7,6 +7,7 @@
 
 #include "MeshTranslationImpl.h"
 #include "USDAssetUserData.h"
+#include "USDDrawModeComponent.h"
 #include "USDClassesModule.h"
 #include "USDConversionUtils.h"
 #include "USDGroomTranslatorUtils.h"
@@ -841,6 +842,14 @@ void FUsdGeometryCacheTranslator::CreateAssets()
 		return;
 	}
 
+	// Don't bother generating assets if we're going to just draw some bounds for this prim instead
+	EUsdDrawMode DrawMode = UsdUtils::GetAppliedDrawMode(GetPrim());
+	if (DrawMode != EUsdDrawMode::Default)
+	{
+		CreateAlternativeDrawModeAssets(DrawMode);
+		return;
+	}
+
 	// Create the GeometryCache TaskChain
 	TSharedRef<FGeometryCacheCreateAssetsTaskChain> AssetsTaskChain = MakeShared<FGeometryCacheCreateAssetsTaskChain>(Context, PrimPath);
 
@@ -849,23 +858,26 @@ void FUsdGeometryCacheTranslator::CreateAssets()
 
 USceneComponent* FUsdGeometryCacheTranslator::CreateComponents()
 {
-	TOptional<TSubclassOf<USceneComponent>> ComponentType;
-
 	if (!IsPotentialGeometryCacheRoot())
 	{
 		return Super::CreateComponents();
 	}
 
-	if (!Context->bIsImporting)
+	USceneComponent* SceneComponent = nullptr;
+
+	EUsdDrawMode DrawMode = UsdUtils::GetAppliedDrawMode(GetPrim());
+	if (DrawMode == EUsdDrawMode::Default)
 	{
-		ComponentType = UGeometryCacheUsdComponent::StaticClass();
+		SceneComponent = CreateComponentsEx(
+			{Context->bIsImporting ? UGeometryCacheComponent::StaticClass() : UGeometryCacheUsdComponent::StaticClass()},
+			{}
+		);
 	}
 	else
 	{
-		ComponentType = UGeometryCacheComponent::StaticClass();
+		SceneComponent = CreateAlternativeDrawModeComponents(DrawMode);
 	}
 
-	USceneComponent* SceneComponent = CreateComponentsEx(ComponentType, {});
 	UpdateComponents(SceneComponent);
 
 	if (UGeometryCacheComponent* Component = Cast<UGeometryCacheComponent>(SceneComponent))
@@ -919,7 +931,9 @@ USceneComponent* FUsdGeometryCacheTranslator::CreateComponents()
 
 void FUsdGeometryCacheTranslator::UpdateComponents(USceneComponent* SceneComponent)
 {
-	if (!IsPotentialGeometryCacheRoot())
+	UGeometryCacheComponent* GeometryCacheComponent = Cast<UGeometryCacheComponent>(SceneComponent);
+
+	if (!IsPotentialGeometryCacheRoot() || !GeometryCacheComponent)
 	{
 		Super::UpdateComponents(SceneComponent);
 		return;
@@ -931,7 +945,7 @@ void FUsdGeometryCacheTranslator::UpdateComponents(USceneComponent* SceneCompone
 	}
 
 	// Set the initial GeometryCache on the GeometryCacheComponent
-	if (UGeometryCacheComponent* GeometryCacheComponent = Cast<UGeometryCacheComponent>(SceneComponent))
+	if (GeometryCacheComponent)
 	{
 		UGeometryCache* GeometryCache = nullptr;
 		if (Context->InfoCache)
@@ -1029,6 +1043,14 @@ void FUsdGeometryCacheTranslator::UpdateComponents(USceneComponent* SceneCompone
 
 bool FUsdGeometryCacheTranslator::CollapsesChildren(ECollapsingType CollapsingType) const
 {
+	// If we have a custom draw mode, it means we should draw bounds/cards/etc. instead
+	// of our entire subtree, which is basically the same thing as collapsing
+	EUsdDrawMode DrawMode = UsdUtils::GetAppliedDrawMode(GetPrim());
+	if (DrawMode != EUsdDrawMode::Default)
+	{
+		return true;
+	}
+
 	return IsPotentialGeometryCacheRoot() ? true : Super::CollapsesChildren(CollapsingType);
 }
 
diff --git a/Engine/Plugins/Importers/USDImporter/Source/USDSchemas/Private/USDGroomTranslator.cpp b/Engine/Plugins/Importers/USDImporter/Source/USDSchemas/Private/USDGroomTranslator.cpp
--- a/Engine/Plugins/Importers/USDImporter/Source/USDSchemas/Private/USDGroomTranslator.cpp
+++ b/Engine/Plugins/Importers/USDImporter/Source/USDSchemas/Private/USDGroomTranslator.cpp
@@ -19,6 +19,7 @@
 
 #include "UnrealUSDWrapper.h"
 #include "USDAssetUserData.h"
+#include "USDDrawModeComponent.h"
 #include "USDClassesModule.h"
 #include "USDGroomConversion.h"
 #include "USDGroomTranslatorUtils.h"
@@ -365,6 +366,14 @@ void FUsdGroomTranslator::CreateAssets()
 		return Super::CreateAssets();
 	}
 
+	// Don't bother generating assets if we're going to just draw some bounds for this prim instead
+	EUsdDrawMode DrawMode = UsdUtils::GetAppliedDrawMode(GetPrim());
+	if (DrawMode != EUsdDrawMode::Default)
+	{
+		CreateAlternativeDrawModeAssets(DrawMode);
+		return;
+	}
+
 	Context->TranslatorTasks.Add(MakeShared<FUsdGroomCreateAssetsTaskChain>(Context, PrimPath));
 }
 
@@ -382,8 +391,19 @@ USceneComponent* FUsdGroomTranslator::CreateComponents()
 		return nullptr;
 	}
 
-	bool bNeedsActor = true;
-	USceneComponent* Component = CreateComponentsEx(TSubclassOf<USceneComponent>(UGroomComponent::StaticClass()), bNeedsActor);
+	USceneComponent* Component = nullptr;
+
+	EUsdDrawMode DrawMode = UsdUtils::GetAppliedDrawMode(GetPrim());
+	if (DrawMode == EUsdDrawMode::Default)
+	{
+		const bool bNeedsActor = true;
+		Component = CreateComponentsEx({UGroomComponent::StaticClass()}, bNeedsActor);
+	}
+	else
+	{
+		Component = CreateAlternativeDrawModeComponents(DrawMode);
+	}
+
 	UpdateComponents(Component);
 
 	return Component;
@@ -391,55 +411,57 @@ USceneComponent* FUsdGroomTranslator::CreateComponents()
 
 void FUsdGroomTranslator::UpdateComponents(USceneComponent* SceneComponent)
 {
-	if (!Context->bAllowParsingGroomAssets || !IsGroomPrim())
-	{
-		Super::UpdateComponents(SceneComponent);
-	}
-
-	if (UGroomComponent* GroomComponent = Cast<UGroomComponent>(SceneComponent))
+	if (Context->bAllowParsingGroomAssets && IsGroomPrim())
 	{
-		GroomComponent->Modify();
-
-		UGroomAsset* Groom = nullptr;
-		if(Context->InfoCache)
-		{
-			Groom = Context->InfoCache->GetSingleAssetForPrim<UGroomAsset>(
-				PrimPath
-			);
-		}
-
-		bool bShouldRegister = false;
-		if (Groom != GroomComponent->GroomAsset.Get())
+		if (UGroomComponent* GroomComponent = Cast<UGroomComponent>(SceneComponent))
 		{
-			bShouldRegister = true;
+			GroomComponent->Modify();
 
-			if (GroomComponent->IsRegistered())
+			UGroomAsset* Groom = nullptr;
+			if(Context->InfoCache)
 			{
-				GroomComponent->UnregisterComponent();
+				Groom = Context->InfoCache->GetSingleAssetForPrim<UGroomAsset>(
+					PrimPath
+				);
 			}
 
-			GroomComponent->SetGroomAsset(Groom);
-
-			if (Groom)
+			bool bShouldRegister = false;
+			if (Groom != GroomComponent->GroomAsset.Get())
 			{
-				UGroomCache* GroomCache = Context->InfoCache->GetSingleAssetForPrim<UGroomCache>(PrimPath);
-				if (GroomCache != GroomComponent->GroomCache.Get())
+				bShouldRegister = true;
+
+				if (GroomComponent->IsRegistered())
+				{
+					GroomComponent->UnregisterComponent();
+				}
+
+				GroomComponent->SetGroomAsset(Groom);
+
+				if (Groom)
 				{
-					GroomComponent->SetGroomCache(GroomCache);
+					UGroomCache* GroomCache = Context->InfoCache->GetSingleAssetForPrim<UGroomCache>(PrimPath);
+					if (GroomCache != GroomComponent->GroomCache.Get())
+					{
+						GroomComponent->SetGroomCache(GroomCache);
+					}
 				}
 			}
-		}
 
-		// Use the prim purpose in conjunction with the prim's computed visibility to toggle the visibility of the groom component
-		// since the component itself cannot be removed if the groom shouldn't be displayed
-		const bool bShouldRender = UsdUtils::IsVisible(GetPrim()) && EnumHasAllFlags(Context->PurposesToLoad, IUsdPrim::GetPurpose(GetPrim()));
-		GroomComponent->SetVisibility(bShouldRender);
+			// Use the prim purpose in conjunction with the prim's computed visibility to toggle the visibility of the groom component
+			// since the component itself cannot be removed if the groom shouldn't be displayed
+			const bool bShouldRender = UsdUtils::IsVisible(GetPrim()) && EnumHasAllFlags(Context->PurposesToLoad, IUsdPrim::GetPurpose(GetPrim()));
+			GroomComponent->SetVisibility(bShouldRender);
 
-		if (bShouldRegister && !GroomComponent->IsRegistered())
-		{
-			GroomComponent->RegisterComponent();
+			if (bShouldRegister && !GroomComponent->IsRegistered())
+			{
+				GroomComponent->RegisterComponent();
+			}
+
+			return;
 		}
 	}
+
+	Super::UpdateComponents(SceneComponent);
 }
 
 bool FUsdGroomTranslator::CollapsesChildren(ECollapsingType CollapsingType) const
diff --git a/Engine/Plugins/Importers/USDImporter/Source/USDSchemas/Private/USDLuxLightTranslator.cpp b/Engine/Plugins/Importers/USDImporter/Source/USDSchemas/Private/USDLuxLightTranslator.cpp
--- a/Engine/Plugins/Importers/USDImporter/Source/USDSchemas/Private/USDLuxLightTranslator.cpp
+++ b/Engine/Plugins/Importers/USDImporter/Source/USDSchemas/Private/USDLuxLightTranslator.cpp
@@ -2,9 +2,11 @@
 
 #include "USDLuxLightTranslator.h"
 
+#include "USDDrawModeComponent.h"
 #include "USDConversionUtils.h"
 #include "USDLightConversion.h"
 #include "USDTypesConversion.h"
+
 #include "UsdWrappers/UsdPrim.h"
 #include "UsdWrappers/UsdStage.h"
 
@@ -30,10 +32,21 @@
 
 USceneComponent* FUsdLuxLightTranslator::CreateComponents()
 {
-	const bool bNeedsActor = true;
-	USceneComponent* LightComponent = CreateComponentsEx( {}, bNeedsActor );
-	UpdateComponents( LightComponent );
-	return LightComponent;
+	USceneComponent* SceneComponent = nullptr;
+
+	EUsdDrawMode DrawMode = UsdUtils::GetAppliedDrawMode(GetPrim());
+	if (DrawMode == EUsdDrawMode::Default)
+	{
+		const bool bNeedsActor = true;
+		SceneComponent = CreateComponentsEx({}, bNeedsActor);
+	}
+	else
+	{
+		SceneComponent = CreateAlternativeDrawModeComponents(DrawMode);
+	}
+
+	UpdateComponents(SceneComponent);
+	return SceneComponent;
 }
 
 void FUsdLuxLightTranslator::UpdateComponents( USceneComponent* SceneComponent )
@@ -100,6 +113,14 @@ void FUsdLuxLightTranslator::UpdateComponents( USceneComponent* SceneComponent )
 
 bool FUsdLuxLightTranslator::CollapsesChildren( ECollapsingType CollapsingType ) const
 {
+	// If we have a custom draw mode, it means we should draw bounds/cards/etc. instead
+	// of our entire subtree, which is basically the same thing as collapsing
+	EUsdDrawMode DrawMode = UsdUtils::GetAppliedDrawMode(GetPrim());
+	if (DrawMode != EUsdDrawMode::Default)
+	{
+		return true;
+	}
+
 	return false;
 }
 
diff --git a/Engine/Plugins/Importers/USDImporter/Source/USDSchemas/Private/USDSkelRootTranslator.cpp b/Engine/Plugins/Importers/USDImporter/Source/USDSchemas/Private/USDSkelRootTranslator.cpp
--- a/Engine/Plugins/Importers/USDImporter/Source/USDSchemas/Private/USDSkelRootTranslator.cpp
+++ b/Engine/Plugins/Importers/USDImporter/Source/USDSchemas/Private/USDSkelRootTranslator.cpp
@@ -7,6 +7,7 @@
 #include "MeshTranslationImpl.h"
 #include "USDAssetCache.h"
 #include "USDAssetUserData.h"
+#include "USDDrawModeComponent.h"
 #include "USDClassesModule.h"
 #include "USDConversionUtils.h"
 #include "USDErrorUtils.h"
@@ -1668,6 +1669,14 @@ void FUsdSkelRootTranslator::CreateAssets()
 {
 	TRACE_CPUPROFILER_EVENT_SCOPE(FUsdSkelRootTranslator::CreateAssets);
 
+	// Don't bother generating assets if we're going to just draw some bounds for this prim instead
+	EUsdDrawMode DrawMode = UsdUtils::GetAppliedDrawMode(GetPrim());
+	if (DrawMode != EUsdDrawMode::Default)
+	{
+		CreateAlternativeDrawModeAssets(DrawMode);
+		return;
+	}
+
 #if WITH_EDITOR
 	// Importing skeletal meshes actually works in Standalone mode, but we intentionally block it here
 	// to not confuse users as to why it doesn't work at runtime
@@ -1686,7 +1695,13 @@ USceneComponent* FUsdSkelRootTranslator::CreateComponents()
 
 #if WITH_EDITOR
 	// Check if the prim has the GroomBinding schema and setup the component and assets necessary to bind the groom to the SkeletalMesh
-	if (UsdUtils::PrimHasSchema(GetPrim(), UnrealIdentifiers::GroomBindingAPI) && Context->AssetCache && Context->InfoCache && Context->bAllowParsingGroomAssets)
+	if (SceneComponent &&
+		SceneComponent->IsA<USkeletalMeshComponent>() &&
+		UsdUtils::PrimHasSchema(GetPrim(), UnrealIdentifiers::GroomBindingAPI) &&
+		Context->AssetCache &&
+		Context->InfoCache &&
+		Context->bAllowParsingGroomAssets
+	)
 	{
 		UsdGroomTranslatorUtils::CreateGroomBindingAsset(
 			GetPrim(),
@@ -1699,7 +1714,7 @@ USceneComponent* FUsdSkelRootTranslator::CreateComponents()
 		// so the Context ParentComponent is set to the SceneComponent temporarily
 		TGuardValue< USceneComponent* > ParentComponentGuard{ Context->ParentComponent, SceneComponent };
 		const bool bNeedsActor = false;
-		UGroomComponent* GroomComponent = Cast< UGroomComponent >( CreateComponentsEx( TSubclassOf< USceneComponent >( UGroomComponent::StaticClass() ), bNeedsActor ) );
+		UGroomComponent* GroomComponent = Cast<UGroomComponent>(CreateComponentsEx({UGroomComponent::StaticClass()}, bNeedsActor));
 		if ( GroomComponent )
 		{
 			UpdateComponents( SceneComponent );
@@ -1712,9 +1727,11 @@ USceneComponent* FUsdSkelRootTranslator::CreateComponents()
 
 void FUsdSkelRootTranslator::UpdateComponents( USceneComponent* SceneComponent )
 {
-	USkeletalMeshComponent* SkeletalMeshComponent = Cast< USkeletalMeshComponent >( SceneComponent );
-	if (!SkeletalMeshComponent || !Context->InfoCache)
+	USkeletalMeshComponent* SkeletalMeshComponent = Cast<USkeletalMeshComponent>(SceneComponent);
+	if (!SkeletalMeshComponent)
 	{
+		// In case we have an alternate draw mode
+		Super::UpdateComponents(SceneComponent);
 		return;
 	}
 
diff --git a/Engine/Plugins/Importers/USDImporter/Source/USDSchemas/Public/USDGeomPointInstancerTranslator.h b/Engine/Plugins/Importers/USDImporter/Source/USDSchemas/Public/USDGeomPointInstancerTranslator.h
--- a/Engine/Plugins/Importers/USDImporter/Source/USDSchemas/Public/USDGeomPointInstancerTranslator.h
+++ b/Engine/Plugins/Importers/USDImporter/Source/USDSchemas/Public/USDGeomPointInstancerTranslator.h
@@ -41,7 +41,7 @@ public:
 	// below it, so returning true to "CanBeCollapsed" in that case won't really help anything.
 	// We only return true from CanBeCollapsed if e.g. a generic Xform prim should be able to collapse us too (like when we're
 	// collapsing even simple, not-nested point instancers).
-	virtual bool CollapsesChildren( ECollapsingType CollapsingType ) const override { return true; }
+	virtual bool CollapsesChildren( ECollapsingType CollapsingType ) const override;
 	virtual bool CanBeCollapsed( ECollapsingType CollapsingType ) const override;
 
 	virtual TSet<UE::FSdfPath> CollectAuxiliaryPrims() const override;
diff --git a/Engine/Plugins/Importers/USDImporter/Source/USDSchemas/Public/USDGeomXformableTranslator.h b/Engine/Plugins/Importers/USDImporter/Source/USDSchemas/Public/USDGeomXformableTranslator.h
--- a/Engine/Plugins/Importers/USDImporter/Source/USDSchemas/Public/USDGeomXformableTranslator.h
+++ b/Engine/Plugins/Importers/USDImporter/Source/USDSchemas/Public/USDGeomXformableTranslator.h
@@ -13,6 +13,7 @@
 #include "USDIncludesEnd.h"
 
 class USceneComponent;
+enum class EUsdDrawMode : int32;
 
 PXR_NAMESPACE_OPEN_SCOPE
 	class UsdGeomXformable;
@@ -37,6 +38,15 @@ public:
 	// If the optional parameters are not set, we'll figure them out automatically.
 	USceneComponent* CreateComponentsEx( TOptional< TSubclassOf< USceneComponent > > ComponentType, TOptional< bool > bNeedsActor );
 
+protected:
+	// Creates actors, components and assets in case we have an alt draw mode like bounds or cards.
+	// In theory *any* prim can have these, but we're placing these on FUsdGeomXformableTranslator as we're assuming only Xformables
+	// (something drawable in the first place) can realistically use an alternative draw mode (i.e. we're not going to do much
+	// placing these on a Material prim, that can't even have an Xform)
+	USceneComponent* CreateAlternativeDrawModeComponents(EUsdDrawMode DrawMode);
+	void UpdateAlternativeDrawModeComponents(USceneComponent* SceneComponent, EUsdDrawMode DrawMode);
+	void CreateAlternativeDrawModeAssets(EUsdDrawMode DrawMode);
+
 private:
 	TOptional< TSubclassOf< USceneComponent > > ComponentTypeOverride;
 };
diff --git a/Engine/Plugins/Importers/USDImporter/Source/USDSchemas/Public/USDSchemaTranslator.h b/Engine/Plugins/Importers/USDImporter/Source/USDSchemas/Public/USDSchemaTranslator.h
--- a/Engine/Plugins/Importers/USDImporter/Source/USDSchemas/Public/USDSchemaTranslator.h
+++ b/Engine/Plugins/Importers/USDImporter/Source/USDSchemas/Public/USDSchemaTranslator.h
@@ -33,6 +33,10 @@ class UStaticMesh;
 class UTexture;
 struct FUsdBlendShape;
 struct FUsdSchemaTranslationContext;
+namespace UE
+{
+	class FUsdGeomBBoxCache;
+}
 
 class USDSCHEMAS_API FRegisteredSchemaTranslatorHandle
 {
@@ -185,6 +189,9 @@ struct USDSCHEMAS_API FUsdSchemaTranslationContext : public TSharedFromThis< FUs
 	/** Caches various information about prims that are expensive to query */
 	TSharedPtr<FUsdInfoCache> InfoCache;
 
+	/** Bounding box cache used for the USD stage in case we have to spawn bounds components */
+	TSharedPtr<UE::FUsdGeomBBoxCache> BBoxCache;
+
 	/** Where we place imported blend shapes, if available */
 	UsdUtils::FBlendShapeMap* BlendShapesByPath = nullptr;
 
diff --git a/Engine/Plugins/Importers/USDImporter/Source/USDStage/Private/USDLevelSequenceHelper.cpp b/Engine/Plugins/Importers/USDImporter/Source/USDStage/Private/USDLevelSequenceHelper.cpp
--- a/Engine/Plugins/Importers/USDImporter/Source/USDStage/Private/USDLevelSequenceHelper.cpp
+++ b/Engine/Plugins/Importers/USDImporter/Source/USDStage/Private/USDLevelSequenceHelper.cpp
@@ -2,9 +2,9 @@
 
 #include "USDLevelSequenceHelper.h"
 
-#include "GeometryCache.h"
 #include "USDAssetUserData.h"
 #include "USDAttributeUtils.h"
+#include "USDDrawModeComponent.h"
 #include "USDClassesModule.h"
 #include "USDConversionUtils.h"
 #include "USDInfoCache.h"
@@ -19,10 +19,12 @@
 #include "USDStageActor.h"
 #include "USDTypesConversion.h"
 #include "USDValueConversion.h"
+
 #include "UsdWrappers/SdfChangeBlock.h"
 #include "UsdWrappers/SdfLayer.h"
 #include "UsdWrappers/UsdAttribute.h"
 #include "UsdWrappers/UsdEditContext.h"
+#include "UsdWrappers/UsdGeomBBoxCache.h"
 #include "UsdWrappers/UsdGeomXformable.h"
 #include "UsdWrappers/UsdPrim.h"
 #include "UsdWrappers/UsdStage.h"
@@ -44,6 +46,7 @@
 #include "ControlRigObjectBinding.h"
 #include "CoreMinimal.h"
 #include "Framework/Notifications/NotificationManager.h"
+#include "GeometryCache.h"
 #include "GeometryCacheComponent.h"
 #include "GeometryCacheTrackUSD.h"
 #include "GroomCache.h"
@@ -62,6 +65,7 @@
 #include "MovieSceneTimeHelpers.h"
 #include "MovieSceneTrack.h"
 #include "Rigs/FKControlRig.h"
+#include "RigVMBlueprintGeneratedClass.h"
 #include "Sections/MovieSceneFloatSection.h"
 #include "Sections/MovieSceneSubSection.h"
 #include "Sequencer/MovieSceneControlRigParameterSection.h"
@@ -74,11 +78,11 @@
 #include "Tracks/MovieScenePropertyTrack.h"
 #include "Tracks/MovieSceneSkeletalAnimationTrack.h"
 #include "Tracks/MovieSceneSubTrack.h"
+#include "Tracks/MovieSceneVectorTrack.h"
 #include "Tracks/MovieSceneVisibilityTrack.h"
 #include "UObject/Package.h"
 #include "UObject/UObjectGlobals.h"
 #include "Widgets/Notifications/SNotificationList.h"
-#include "RigVMBlueprintGeneratedClass.h"
 
 #if WITH_EDITOR
 #include "ControlRigBlueprint.h"
@@ -495,6 +499,7 @@ public:
 
 	ULevelSequence* Init(const UE::FUsdStage& InUsdStage);
 	void SetInfoCache(TSharedPtr<FUsdInfoCache> InfoCache);
+	void SetBBoxCache(TSharedPtr<UE::FUsdGeomBBoxCache> InBBoxCache);
 	bool HasData() const;
 	void Clear();
 
@@ -590,6 +595,7 @@ private:
 	void RemoveTimeTrack(const FUsdLevelSequenceHelperImpl::FLayerTimeInfo* Info);
 
 	void AddCommonTracks(const UUsdPrimTwin& PrimTwin, const UE::FUsdPrim& Prim, bool bForceVisibilityTracks = false);
+	void AddBoundsTracks(const UUsdPrimTwin& PrimTwin, const UE::FUsdPrim& Prim, TOptional<bool> HasAnimatedBounds = {});
 	void AddCameraTracks(const UUsdPrimTwin& PrimTwin, const UE::FUsdPrim& Prim);
 	void AddLightTracks(const UUsdPrimTwin& PrimTwin, const UE::FUsdPrim& Prim, const TSet<FName>& PropertyPathsToRead = {});
 	void AddSkeletalTracks(const UUsdPrimTwin& PrimTwin, const UE::FUsdPrim& Prim);
@@ -606,7 +612,7 @@ public:
 	// If bForceVisibilityTracks is true, we'll add and bake the visibility tracks for this prim even if the
 	// prim itself doesn't have animated visibility (so that we can handle its visibility in case one of its
 	// parents does have visibility animations)
-	void AddPrim(UUsdPrimTwin& PrimTwin, bool bForceVisibilityTracks = false);
+	void AddPrim(UUsdPrimTwin& PrimTwin, bool bForceVisibilityTracks = false, TOptional<bool> HasAnimatedBounds = {});
 	void RemovePrim(const UUsdPrimTwin& PrimTwin);
 
 	// These functions assume the skeletal animation tracks (if any) were already added to the level sequence
@@ -687,6 +693,7 @@ private:
 
 	TWeakObjectPtr<AUsdStageActor> StageActor = nullptr;
 	TSharedPtr<FUsdInfoCache> InfoCache = nullptr;  // We keep a pointer to this directly because we may be called via the USDStageImporter directly, when we don't have an available actor
+	TSharedPtr<UE::FUsdGeomBBoxCache> BBoxCache = nullptr; // Same as for the info cache
 	FGuid StageActorBinding;
 
 	// Only when this is zero we write LevelSequence object (tracks, moviescene, sections, etc.) transactions back to the USD stage
@@ -756,6 +763,11 @@ void FUsdLevelSequenceHelperImpl::SetInfoCache(TSharedPtr<FUsdInfoCache> InInfoC
 	InfoCache = InInfoCache;
 }
 
+void FUsdLevelSequenceHelperImpl::SetBBoxCache(TSharedPtr<UE::FUsdGeomBBoxCache> InBBoxCache)
+{
+	BBoxCache = InBBoxCache;
+}
+
 bool FUsdLevelSequenceHelperImpl::HasData() const
 {
 	if (!MainLevelSequence)
@@ -886,6 +898,7 @@ void FUsdLevelSequenceHelperImpl::BindToUsdStageActor(AUsdStageActor* InStageAct
 
 	StageActor = InStageActor;
 	SetInfoCache(InStageActor ? InStageActor->GetInfoCache() : nullptr);
+	SetBBoxCache(InStageActor ? InStageActor->GetBBoxCache() : nullptr);
 
 	if (!StageActor.IsValid() || !MainLevelSequence || !MainLevelSequence->GetMovieScene())
 	{
@@ -1663,6 +1676,138 @@ void FUsdLevelSequenceHelperImpl::AddCommonTracks(const UUsdPrimTwin& PrimTwin,
 	}
 }
 
+void FUsdLevelSequenceHelperImpl::AddBoundsTracks(const UUsdPrimTwin& PrimTwin, const UE::FUsdPrim& Prim, TOptional<bool> HasAnimatedBounds)
+{
+	if (HasAnimatedBounds.IsSet())
+	{
+		if (!HasAnimatedBounds.GetValue())
+		{
+			return;
+		}
+	}
+
+	USceneComponent* ComponentToBind = PrimTwin.GetSceneComponent();
+	if (!ComponentToBind || !Prim)
+	{
+		return;
+	}
+
+	if (!ensure(BBoxCache))
+	{
+		return;
+	}
+
+	// We only actually use bounds when drawing alternative draw modes
+	EUsdDrawMode DrawMode = UsdUtils::GetAppliedDrawMode(Prim);
+	if (DrawMode == EUsdDrawMode::Default)
+	{
+		return;
+	}
+
+	// We need to manually check if the prim has animated bounds now.
+	// We tried earlying out due to other reasons first as this can be expensive
+	if (!HasAnimatedBounds.IsSet())
+	{
+		if (!UsdUtils::HasAnimatedBounds(Prim, BBoxCache->GetIncludedPurposes(), BBoxCache->GetUseExtentsHint(), BBoxCache->GetIgnoreVisibility()))
+		{
+			return;
+		}
+	}
+
+	// Find the Sequence where we'll author the tracks
+	ULevelSequence* TargetSequence = nullptr;
+	bool bIsMuted = false;
+	if (Prim.IsA(TEXT("Boundable")))
+	{
+		UE::FUsdAttribute ExtentAttr = Prim.GetAttribute(TEXT("extent"));
+		if (ExtentAttr && ExtentAttr.HasAuthoredValue())
+		{
+			bIsMuted = UsdUtils::IsAttributeMuted(ExtentAttr, Prim.GetStage());
+			TargetSequence = FindOrAddSequenceForAttribute(ExtentAttr);
+		}
+	}
+	if (!TargetSequence)
+	{
+		if (Prim.HasAPI(TEXT("GeomModelAPI")))
+		{
+			UE::FUsdAttribute ExtentsHintAttr = Prim.GetAttribute(TEXT("extentsHint"));
+			if (ExtentsHintAttr && ExtentsHintAttr.HasAuthoredValue())
+			{
+				bIsMuted = UsdUtils::IsAttributeMuted(ExtentsHintAttr, Prim.GetStage());
+				TargetSequence = FindOrAddSequenceForAttribute(ExtentsHintAttr);
+			}
+		}
+	}
+	if (!TargetSequence)
+	{
+		// For the other track types we mostly look for a correspondence between one or more USD attributes and UE properties.
+		// For these tracks however we may not have any authored `extent` or `extentsHint` yet (and would only author them on-demand), so
+		// we may need to create tracks that correspond purely to computed, non-authored values in USD.
+		// If the user manually modifies these, we'll author these as `extent` or `extentsHint` depending on the prim, but only on-demand.
+		UE::FSdfLayer PrimLayer = UsdUtils::FindLayerForPrim(Prim);
+		TargetSequence = FindSequenceForIdentifier(PrimLayer.GetIdentifier());
+	}
+	if (!TargetSequence)
+	{
+		return;
+	}
+
+	FMovieSceneSequenceTransform SequenceTransform;
+	FMovieSceneSequenceID SequenceID = SequencesID.FindRef(TargetSequence);
+	if (FMovieSceneSubSequenceData* SubSequenceData = SequenceHierarchyCache.FindSubData(SequenceID))
+	{
+		SequenceTransform = SubSequenceData->RootToSequenceTransform;
+	}
+
+	UMovieScene* MovieScene = TargetSequence->GetMovieScene();
+	if (!MovieScene)
+	{
+		return;
+	}
+
+	UMovieSceneDoubleVectorTrack* MinTrack = AddTrack<UMovieSceneDoubleVectorTrack>(
+		GET_MEMBER_NAME_CHECKED(UUsdDrawModeComponent, BoundsMin),
+		PrimTwin,
+		*ComponentToBind,
+		*TargetSequence,
+		bIsMuted
+	);
+	if (!ensure(MinTrack))
+	{
+		return;
+	}
+	MinTrack->SetNumChannelsUsed(3);
+
+	UMovieSceneDoubleVectorTrack* MaxTrack = AddTrack<UMovieSceneDoubleVectorTrack>(
+		GET_MEMBER_NAME_CHECKED(UUsdDrawModeComponent, BoundsMax),
+		PrimTwin,
+		*ComponentToBind,
+		*TargetSequence,
+		bIsMuted
+	);
+	if (!ensure(MaxTrack))
+	{
+		return;
+	}
+	MaxTrack->SetNumChannelsUsed(3);
+
+	TArray<double> BoundsTimeSamples;
+	const bool bHasAnimatedBounds = UsdUtils::GetAnimatedBoundsTimeSamples(
+		Prim,
+		BoundsTimeSamples,
+		BBoxCache->GetIncludedPurposes(),
+		BBoxCache->GetUseExtentsHint(),
+		BBoxCache->GetIgnoreVisibility()
+	);
+	// GetAnimatedBoundsTimeSamples uses the same underlying code to check if a prim has animated bounds or not,
+	// so if we're this deep in creating bounds tracks it better agree that the bounds are animated
+	ensure(bHasAnimatedBounds);
+
+	UsdToUnreal::ConvertBoundsTimeSamples(Prim, BoundsTimeSamples, SequenceTransform, *MinTrack, *MaxTrack, BBoxCache.Get());
+
+	PrimPathByLevelSequenceName.AddUnique(TargetSequence->GetFName(), Prim.GetPrimPath().GetString());
+}
+
 void FUsdLevelSequenceHelperImpl::AddCameraTracks(const UUsdPrimTwin& PrimTwin, const UE::FUsdPrim& Prim)
 {
 	TArray<FName> TrackedProperties = {
@@ -2099,7 +2244,7 @@ void FUsdLevelSequenceHelperImpl::AddGroomTracks(const UUsdPrimTwin& PrimTwin, c
 	PrimPathByLevelSequenceName.AddUnique(GroomAnimationSequence->GetFName(), PrimPath);
 }
 
-void FUsdLevelSequenceHelperImpl::AddPrim(UUsdPrimTwin& PrimTwin, bool bForceVisibilityTracks)
+void FUsdLevelSequenceHelperImpl::AddPrim(UUsdPrimTwin& PrimTwin, bool bForceVisibilityTracks, TOptional<bool> HasAnimatedBounds)
 {
 	if (!UsdStage)
 	{
@@ -2166,6 +2311,8 @@ void FUsdLevelSequenceHelperImpl::AddPrim(UUsdPrimTwin& PrimTwin, bool bForceVis
 
 	AddCommonTracks(PrimTwin, UsdPrim, bForceVisibilityTracks);
 
+	AddBoundsTracks(PrimTwin, UsdPrim, HasAnimatedBounds);
+
 	RefreshSequencer();
 }
 
@@ -3537,6 +3684,39 @@ void FUsdLevelSequenceHelperImpl::HandleTrackChange(const UMovieSceneTrack& Trac
 					UsdUtils::AuthorIdentityTransformGprimAttributes(UsdPrim, bDefaultValues, bTimeSampleValues);
 				}
 
+				// For the bounds tracks alone we have two separate tracks we must read from at once, and write to a single
+				// USD attribute. We'll have one of those already (the Track itself), but we need to find the other, if any.
+				// This could be somewhat cleaned up if we had FBox tracks in the Sequencer, but it should work just fine for now
+				else if (Writer.TwoVectorWriter)
+				{
+					const UMovieSceneDoubleVectorTrack* MinTrack = nullptr;
+					const UMovieSceneDoubleVectorTrack* MaxTrack = nullptr;
+					if (Track.GetTrackName() == GET_MEMBER_NAME_CHECKED(UUsdDrawModeComponent, BoundsMin))
+					{
+						MinTrack = Cast<UMovieSceneDoubleVectorTrack>(&Track);
+						MaxTrack = Cast<UMovieSceneDoubleVectorTrack>(MovieScene->FindTrack(
+							UMovieSceneDoubleVectorTrack::StaticClass(),
+							PossessableGuid,
+							GET_MEMBER_NAME_CHECKED(UUsdDrawModeComponent, BoundsMax)
+						));
+					}
+					else
+					{
+						MaxTrack = Cast<UMovieSceneDoubleVectorTrack>(&Track);
+						MinTrack = Cast<UMovieSceneDoubleVectorTrack>(MovieScene->FindTrack(
+							UMovieSceneDoubleVectorTrack::StaticClass(),
+							PossessableGuid,
+							GET_MEMBER_NAME_CHECKED(UUsdDrawModeComponent, BoundsMin)
+						));
+					}
+
+					// Realistically we'll have both of them, but we *need* at least one
+					if (ensure(MinTrack || MaxTrack))
+					{
+						UnrealToUsd::ConvertBoundsVectorTracks(MinTrack, MaxTrack, SequenceTransform, Writer.TwoVectorWriter, UsdPrim);
+					}
+				}
+
 				// Refresh tracks that needed to be updated in USD (e.g. we wrote out a new keyframe to a RectLight's width -> that
 				// should become a new keyframe on our intensity track, because we use the RectLight's width for calculating intensity in UE).
 				if (PropertyPathsToRefresh.Num() > 0)
@@ -3641,6 +3821,7 @@ public:
 
 	ULevelSequence* Init(const UE::FUsdStage& InUsdStage) { return nullptr; }
 	void SetInfoCache(TSharedPtr<FUsdInfoCache> InfoCache) {};
+	void SetBBoxCache(TSharedPtr<UE::FUsdGeomBBoxCache> InBBoxCache) {};
 	bool HasData() const { return false; };
 	void Clear() {};
 
@@ -3650,7 +3831,7 @@ public:
 	void UnbindFromUsdStageActor() {}
 	void OnStageActorRenamed() {};
 
-	void AddPrim(UUsdPrimTwin& PrimTwin, bool bForceVisibilityTracks) {}
+	void AddPrim(UUsdPrimTwin& PrimTwin, bool bForceVisibilityTracks, TOptional<bool> HasAnimatedBounds) {}
 	void RemovePrim(const UUsdPrimTwin& PrimTwin) {}
 
 	void UpdateControlRigTracks(UUsdPrimTwin& PrimTwin) {}
@@ -3714,6 +3895,14 @@ void FUsdLevelSequenceHelper::SetInfoCache(TSharedPtr<FUsdInfoCache> InInfoCache
 	}
 }
 
+void FUsdLevelSequenceHelper::SetBBoxCache(TSharedPtr<UE::FUsdGeomBBoxCache> InBBoxCache)
+{
+	if (UsdSequencerImpl.IsValid())
+	{
+		UsdSequencerImpl->SetBBoxCache(InBBoxCache);
+	}
+}
+
 bool FUsdLevelSequenceHelper::HasData() const
 {
 	if (UsdSequencerImpl.IsValid())
@@ -3748,11 +3937,11 @@ void FUsdLevelSequenceHelper::UnbindFromUsdStageActor()
 	}
 }
 
-void FUsdLevelSequenceHelper::AddPrim(UUsdPrimTwin & PrimTwin, bool bForceVisibilityTracks)
+void FUsdLevelSequenceHelper::AddPrim(UUsdPrimTwin & PrimTwin, bool bForceVisibilityTracks, TOptional<bool> HasAnimatedBounds)
 {
 	if (UsdSequencerImpl.IsValid())
 	{
-		UsdSequencerImpl->AddPrim(PrimTwin, bForceVisibilityTracks);
+		UsdSequencerImpl->AddPrim(PrimTwin, bForceVisibilityTracks, HasAnimatedBounds);
 	}
 }
 
diff --git a/Engine/Plugins/Importers/USDImporter/Source/USDStage/Private/USDStageActor.cpp b/Engine/Plugins/Importers/USDImporter/Source/USDStage/Private/USDStageActor.cpp
--- a/Engine/Plugins/Importers/USDImporter/Source/USDStage/Private/USDStageActor.cpp
+++ b/Engine/Plugins/Importers/USDImporter/Source/USDStage/Private/USDStageActor.cpp
@@ -7,6 +7,7 @@
 #include "USDAssetCache.h"
 #include "USDAssetImportData.h"
 #include "USDAssetUserData.h"
+#include "USDDrawModeComponent.h"
 #include "USDClassesModule.h"
 #include "USDConversionUtils.h"
 #include "USDDynamicBindingResolverLibrary.h"
@@ -33,6 +34,7 @@
 #include "UsdWrappers/SdfChangeBlock.h"
 #include "UsdWrappers/SdfLayer.h"
 #include "UsdWrappers/UsdAttribute.h"
+#include "UsdWrappers/UsdGeomBBoxCache.h"
 #include "UsdWrappers/UsdGeomXformable.h"
 #include "UsdWrappers/UsdStage.h"
 
@@ -143,6 +145,7 @@ struct FUsdStageActorImpl
 		TranslationContext->SubdivisionLevel = StageActor->SubdivisionLevel;
 		TranslationContext->BlendShapesByPath = &StageActor->BlendShapesByPath;
 		TranslationContext->InfoCache = StageActor->InfoCache;
+		TranslationContext->BBoxCache = StageActor->BBoxCache;
 		TranslationContext->bTranslateOnlyUsedMaterials = GTranslateOnlyUsedMaterialsWhenOpeningStage;
 
 		// Its more convenient to toggle between variants using the USDStage window, as opposed to parsing LODs
@@ -1392,7 +1395,11 @@ void AUsdStageActor::OnUsdObjectsChanged(const UsdUtils::FObjectChangesByPath& I
 					// TODO: More nuanced info cache updates
 					UnrealIdentifiers::PrimvarsDisplayColor,
 					UnrealIdentifiers::PrimvarsDisplayOpacity,
-					UnrealIdentifiers::DoubleSided
+					UnrealIdentifiers::DoubleSided,
+					// When we change these UsdGeomModelAPI attributes we may need to create a new component type
+					// for the prim (as it may now need/stop needing an alternate draw mode component)
+					UnrealIdentifiers::ModelDrawMode,
+					UnrealIdentifiers::ModelApplyDrawMode
 				};
 
 				// Some stage info should trigger some resyncs because they should trigger reparsing of geometry
@@ -1461,6 +1468,11 @@ void AUsdStageActor::OnUsdObjectsChanged(const UsdUtils::FObjectChangesByPath& I
 		InfoCache->RebuildCacheForSubtree(Stage.GetPseudoRoot(), TranslationContext.Get());
 	}
 
+	if (BBoxCache.IsValid() && bHasResync)
+	{
+		BBoxCache->Clear();
+	}
+
 	TFunction<void(TMap<UE::FSdfPath, bool>&)> SortAndCleanPrimsToUpdate = [](TMap<UE::FSdfPath, bool>& InOutMap)
 	{
 		// Force traversal of parent before child (not needed for cleaning up, but very much needed later for reloading assets/components)
@@ -1986,7 +1998,8 @@ UUsdPrimTwin* AUsdStageActor::ExpandPrim(const UE::FUsdPrim& Prim, bool bResync,
 	}
 
 	// Update the prim animated status
-	const bool bIsAnimated = UsdUtils::IsAnimated(Prim);
+	bool bHasAnimatedBounds = false;
+	bool bIsAnimated = UsdUtils::IsAnimated(Prim);
 	if (bIsAnimated)
 	{
 		if (!PrimsToAnimate.Contains(UsdPrimTwin->PrimPath))
@@ -2014,25 +2027,40 @@ UUsdPrimTwin* AUsdStageActor::ExpandPrim(const UE::FUsdPrim& Prim, bool bResync,
 						}
 					},
 					bRecursive
-					);
+				);
 			}
 
-			UE::FSdfPath PrimPath(*UsdPrimTwin->PrimPath);
-			UE::FUsdPrim UsdPrim(UsdStage.GetPrimAtPath(PrimPath));
-
 			PrimsToAnimate.Add(UsdPrimTwin->PrimPath);
 			LevelSequenceHelper.AddPrim(*UsdPrimTwin);
 		}
 	}
-	else
+	else if (EUsdDrawMode DrawMode = UsdUtils::GetAppliedDrawMode(Prim); DrawMode != EUsdDrawMode::Default)
 	{
-		if (PrimsToAnimate.Contains(UsdPrimTwin->PrimPath))
+		bHasAnimatedBounds = UsdUtils::HasAnimatedBounds(
+			Prim,
+			BBoxCache->GetIncludedPurposes(),
+			BBoxCache->GetUseExtentsHint(),
+			BBoxCache->GetIgnoreVisibility()
+		);
+
+		if (bHasAnimatedBounds)
 		{
-			PrimsToAnimate.Remove(UsdPrimTwin->PrimPath);
-			LevelSequenceHelper.RemovePrim(*UsdPrimTwin);
+			const bool bForceVisibilityTrack = false;
+			LevelSequenceHelper.AddPrim(*UsdPrimTwin, bForceVisibilityTrack, bHasAnimatedBounds);
+			PrimsToAnimate.Add(UsdPrimTwin->PrimPath);
+
+			// Mark the component as animated right away because HasAnimatedBounds is expensive to call and
+			// we don't want to have to re-do it when creating the component
+			UsdPrimTwin->SceneComponent->SetMobility(EComponentMobility::Movable);
 		}
 	}
 
+	if (!bIsAnimated && !bHasAnimatedBounds && PrimsToAnimate.Contains(UsdPrimTwin->PrimPath))
+	{
+		PrimsToAnimate.Remove(UsdPrimTwin->PrimPath);
+		LevelSequenceHelper.RemovePrim(*UsdPrimTwin);
+	}
+
 	// Setup Control Rig tracks if we need to. This must be done after adding regular skeletal animation tracks
 	// if we have any as if will properly deactivate them like the usual "Bake to Control Rig" workflow.
 	if (Prim.IsA(TEXT("SkelRoot")))
@@ -2284,6 +2312,10 @@ void AUsdStageActor::SetPurposesToLoad(int32 NewPurposesToLoad)
 	Modify(bMarkDirty);
 
 	PurposesToLoad = NewPurposesToLoad;
+	if (BBoxCache.IsValid())
+	{
+		BBoxCache->SetIncludedPurposes(static_cast<EUsdPurpose>(PurposesToLoad));
+	}
 	LoadUsdStage();
 }
 
@@ -2343,6 +2375,11 @@ void AUsdStageActor::SetTime(float InTime)
 	Modify(bMarkDirty);
 
 	Time = InTime;
+	if (BBoxCache.IsValid())
+	{
+		BBoxCache->SetTime(Time);
+	}
+
 	Refresh();
 }
 
@@ -2619,6 +2656,9 @@ void AUsdStageActor::OnObjectsReplaced(const TMap<UObject*, UObject*>& ObjectRep
 			NewActor->InfoCache = InfoCache;
 			InfoCache = nullptr;
 
+			NewActor->BBoxCache = BBoxCache;
+			BBoxCache = nullptr;
+
 			// It could be that we're automatically recompiling when going into PIE because our blueprint was dirty.
 			// In that case we also need bIsTransitioningIntoPIE to be true to prevent us from calling LoadUsdStage from PostRegisterAllComponents
 			NewActor->bIsTransitioningIntoPIE = bIsTransitioningIntoPIE;
@@ -2724,6 +2764,19 @@ void AUsdStageActor::LoadUsdStage()
 
 	StageToLoad.SetInterpolationType(InterpolationType);
 
+	// Create Info and BBoxCache before calling ReloadAnimations as that is when the LevelSequenceHelper will also take
+	// a reference to them
+	if (!InfoCache.IsValid())
+	{
+		InfoCache = MakeShared<FUsdInfoCache>();
+	}
+	if (!BBoxCache.IsValid())
+	{
+		const bool bUseExtentsHint = true;
+		const bool bIgnoreVisibility = false;
+		BBoxCache = MakeShared<UE::FUsdGeomBBoxCache>(Time, static_cast<EUsdPurpose>(PurposesToLoad), bUseExtentsHint, bIgnoreVisibility);
+	}
+
 	ReloadAnimations();
 
 	// Make sure our PrimsToAnimate and the LevelSequenceHelper are kept in sync, because we'll use PrimsToAnimate to
@@ -2732,11 +2785,6 @@ void AUsdStageActor::LoadUsdStage()
 	// the tracks
 	PrimsToAnimate.Reset();
 
-	if (!InfoCache.IsValid())
-	{
-		InfoCache = MakeShared<FUsdInfoCache>();
-	}
-
 	TSharedRef< FUsdSchemaTranslationContext > TranslationContext = FUsdStageActorImpl::CreateUsdSchemaTranslationContext(this, RootTwin->PrimPath);
 
 	SlowTask.EnterProgressFrame(0.1f);
@@ -2868,6 +2916,11 @@ void AUsdStageActor::UnloadUsdStage()
 		InfoCache->Clear();
 	}
 
+	if (BBoxCache)
+	{
+		BBoxCache->Clear();
+	}
+
 	OnStageUnloaded.Broadcast();
 
 	CloseUsdStage();
@@ -3022,6 +3075,11 @@ TSharedPtr<FUsdInfoCache> AUsdStageActor::GetInfoCache()
 	return InfoCache;
 }
 
+TSharedPtr<UE::FUsdGeomBBoxCache> AUsdStageActor::GetBBoxCache()
+{
+	return BBoxCache;
+}
+
 TMap<FString, TMap<FString, int32>> AUsdStageActor::GetMaterialToPrimvarToUVIndex()
 {
 	return MaterialToPrimvarToUVIndex;
@@ -4021,6 +4079,18 @@ void AUsdStageActor::OnObjectPropertyChanged(UObject* ObjectBeingModified, FProp
 				{
 					UnrealToUsd::ConvertMeshComponent(CurrentStage, MeshComponent, UsdPrim);
 				}
+				else if (UUsdDrawModeComponent* BoundsComponent = Cast<UUsdDrawModeComponent>(PrimSceneComponent))
+				{
+					const static TSet<FName> BoundsProperties = {
+						GET_MEMBER_NAME_CHECKED(UUsdDrawModeComponent, BoundsMin),
+						GET_MEMBER_NAME_CHECKED(UUsdDrawModeComponent, BoundsMax),
+					};
+
+					// If we just manually tweaked the extents, also author those back out to USD as extents opinions
+					const bool bWriteExtents = BoundsProperties.Contains(PropertyChangedEvent.GetMemberPropertyName());
+					const double UsdTimeCode = UsdUtils::GetDefaultTimeCode();
+					UnrealToUsd::ConvertBoundsComponent(*BoundsComponent, UsdPrim, bWriteExtents, UsdTimeCode);
+				}
 				else if (UsdPrim && UsdPrim.IsA(TEXT("Camera")))
 				{
 					// Our component may be pointing directly at a camera component in case we recreated an exported
diff --git a/Engine/Plugins/Importers/USDImporter/Source/USDStage/Public/USDLevelSequenceHelper.h b/Engine/Plugins/Importers/USDImporter/Source/USDStage/Public/USDLevelSequenceHelper.h
--- a/Engine/Plugins/Importers/USDImporter/Source/USDStage/Public/USDLevelSequenceHelper.h
+++ b/Engine/Plugins/Importers/USDImporter/Source/USDStage/Public/USDLevelSequenceHelper.h
@@ -14,6 +14,11 @@ class FUsdLevelSequenceHelperImpl;
 class ULevelSequence;
 class UUsdPrimTwin;
 
+namespace UE
+{
+	class FUsdGeomBBoxCache;
+}
+
 /**
  * Builds and maintains the level sequence and subsequences for a Usd Stage
  */
@@ -37,6 +42,9 @@ public:
 	/** Sets the asset cache to use when fetching assets and asset info required for the level sequence animation, like UAnimSequences */
 	void SetInfoCache(TSharedPtr<FUsdInfoCache> InInfoCache);
 
+	/** Sets the BBoxCache to use when importing bound animations for prims. Needed for importing, where we don't have a stage actor to take the BBoxCache from */
+	void SetBBoxCache(TSharedPtr<UE::FUsdGeomBBoxCache> InBBoxCache);
+
 	/* Returns true if we have at least one possessable or a reference to a subsequence */
 	bool HasData() const;
 
@@ -55,9 +63,11 @@ public:
 	 * If bForceVisibilityTracks is true, will add visibility tracks even if this prim
 	 * doesn't actually have timeSamples on its visibility attribute (use this when
 	 * a parent does have animated visibility, and we need to "bake" that out to a dedicated
-	 * visibility track so that the standalone LevelSequence asset behaves as expected)
+	 * visibility track so that the standalone LevelSequence asset behaves as expected).
+	 * We'll check the prim for animated bounds and add bounds tracks, unless bHasAnimatedBounds
+	 * already provides whether the prim has animated bounds or not.
 	 */
-	void AddPrim(UUsdPrimTwin& PrimTwin, bool bForceVisibilityTracks = false);
+	void AddPrim(UUsdPrimTwin& PrimTwin, bool bForceVisibilityTracks = false, TOptional<bool> HasAnimatedBounds = {});
 
 	/** Removes any track associated with this prim */
 	void RemovePrim(const UUsdPrimTwin& PrimTwin);
diff --git a/Engine/Plugins/Importers/USDImporter/Source/USDStage/Public/USDStageActor.h b/Engine/Plugins/Importers/USDImporter/Source/USDStage/Public/USDStageActor.h
--- a/Engine/Plugins/Importers/USDImporter/Source/USDStage/Public/USDStageActor.h
+++ b/Engine/Plugins/Importers/USDImporter/Source/USDStage/Public/USDStageActor.h
@@ -24,6 +24,7 @@ enum class EMovieSceneDataChangeType;
 struct FUsdSchemaTranslationContext;
 namespace UE
 {
+	class FUsdGeomBBoxCache;
 	class FUsdPrim;
 }
 
@@ -236,6 +237,7 @@ public:
 	USDSTAGE_API void ReloadAnimations();
 
 	USDSTAGE_API TSharedPtr<FUsdInfoCache> GetInfoCache();
+	USDSTAGE_API TSharedPtr<UE::FUsdGeomBBoxCache> GetBBoxCache();
 	USDSTAGE_API TMap< FString, TMap< FString, int32 > > GetMaterialToPrimvarToUVIndex();
 	USDSTAGE_API const UsdUtils::FBlendShapeMap& GetBlendShapeMap();
 	USDSTAGE_API FUsdListener& GetUsdListener();
@@ -358,6 +360,9 @@ protected:
 	/** Caches various information about prims that are expensive to query */
 	TSharedPtr<FUsdInfoCache> InfoCache;
 
+	/** USD bounding box cache for the stage. Constructed on-demand */
+	TSharedPtr<UE::FUsdGeomBBoxCache> BBoxCache;
+
 	FUsdListener UsdListener;
 
 	FUsdLevelSequenceHelper LevelSequenceHelper;
diff --git a/Engine/Plugins/Importers/USDImporter/Source/USDStageEditorViewModels/Private/USDStageViewModel.cpp b/Engine/Plugins/Importers/USDImporter/Source/USDStageEditorViewModels/Private/USDStageViewModel.cpp
--- a/Engine/Plugins/Importers/USDImporter/Source/USDStageEditorViewModels/Private/USDStageViewModel.cpp
+++ b/Engine/Plugins/Importers/USDImporter/Source/USDStageEditorViewModels/Private/USDStageViewModel.cpp
@@ -350,6 +350,7 @@ void FUsdStageViewModel::ImportStage( const TCHAR* TargetContentFolder, UUsdStag
 			// Let the importer reuse our assets, but force it to spawn new actors and components always
 			// This allows a different setting for asset/component collapsing, and doesn't require modifying the PrimTwins
 			ImportContext.AssetCache = StageActor->UsdAssetCache;
+			ImportContext.BBoxCache = StageActor->GetBBoxCache();
 
 			ImportContext.TargetSceneActorAttachParent = StageActor->GetRootComponent()->GetAttachParent();
 			ImportContext.TargetSceneActorTargetTransform = StageActor->GetActorTransform();
diff --git a/Engine/Plugins/Importers/USDImporter/Source/USDStageImporter/Private/USDStageImportContext.cpp b/Engine/Plugins/Importers/USDImporter/Source/USDStageImporter/Private/USDStageImportContext.cpp
--- a/Engine/Plugins/Importers/USDImporter/Source/USDStageImporter/Private/USDStageImportContext.cpp
+++ b/Engine/Plugins/Importers/USDImporter/Source/USDStageImporter/Private/USDStageImportContext.cpp
@@ -89,6 +89,7 @@ void FUsdStageImportContext::Reset()
 	ImportedAsset = nullptr;
 	LevelSequenceHelper.Clear();
 	AssetCache = nullptr;
+	BBoxCache = nullptr;
 	MaterialToPrimvarToUVIndex.Empty();
 	Stage = UE::FUsdStage{};
 	ImportObjectFlags = EObjectFlags::RF_NoFlags;
diff --git a/Engine/Plugins/Importers/USDImporter/Source/USDStageImporter/Private/USDStageImporter.cpp b/Engine/Plugins/Importers/USDImporter/Source/USDStageImporter/Private/USDStageImporter.cpp
--- a/Engine/Plugins/Importers/USDImporter/Source/USDStageImporter/Private/USDStageImporter.cpp
+++ b/Engine/Plugins/Importers/USDImporter/Source/USDStageImporter/Private/USDStageImporter.cpp
@@ -6,6 +6,7 @@
 #include "USDAssetCache2.h"
 #include "USDAssetImportData.h"
 #include "USDAssetUserData.h"
+#include "USDDrawModeComponent.h"
 #include "USDClassesModule.h"
 #include "USDConversionUtils.h"
 #include "USDErrorUtils.h"
@@ -17,7 +18,9 @@
 #include "USDStageImportContext.h"
 #include "USDStageImportOptions.h"
 #include "USDTypesConversion.h"
+
 #include "UsdWrappers/SdfLayer.h"
+#include "UsdWrappers/UsdGeomBBoxCache.h"
 #include "UsdWrappers/UsdStage.h"
 #include "UsdWrappers/UsdTyped.h"
 
@@ -328,20 +331,27 @@ namespace UsdStageImporterImpl
 #endif // #if USE_USD_SDK
 	}
 
-	void ImportAnimation(FUsdStageImportContext& ImportContext, UE::FUsdPrim& Prim, bool bAnimatedVisibility, USceneComponent* SceneComponent)
+	void ImportAnimation(
+		FUsdStageImportContext& ImportContext,
+		UE::FUsdPrim& Prim,
+		bool bAnimatedVisibility,
+		TOptional<bool> bHasAnimatedBounds,
+		USceneComponent* SceneComponent
+	)
 	{
-		TRACE_CPUPROFILER_EVENT_SCOPE( ImportAnimation );
+		TRACE_CPUPROFILER_EVENT_SCOPE(ImportAnimation);
 
-		if ( !ImportContext.ImportOptions->bImportLevelSequences )
+		if (!ImportContext.ImportOptions->bImportLevelSequences)
 		{
 			return;
 		}
 
-		UUsdPrimTwin* UsdPrimTwin = NewObject< UUsdPrimTwin >();
+		UUsdPrimTwin* UsdPrimTwin = NewObject<UUsdPrimTwin>();
 		UsdPrimTwin->PrimPath = Prim.GetPrimPath().GetString();
 		UsdPrimTwin->SceneComponent = SceneComponent;
+		UsdPrimTwin->SceneComponent->SetMobility(EComponentMobility::Movable);
 
-		ImportContext.LevelSequenceHelper.AddPrim( *UsdPrimTwin, bAnimatedVisibility );
+		ImportContext.LevelSequenceHelper.AddPrim(*UsdPrimTwin, bAnimatedVisibility, bHasAnimatedBounds);
 	}
 
 	void ImportActor(FUsdStageImportContext& ImportContext, UE::FUsdPrim& Prim, bool bForceVisibilityAnimationTracks, FUsdSchemaTranslationContext& TranslationContext)
@@ -398,9 +408,23 @@ namespace UsdStageImporterImpl
 			}
 
 #if USE_USD_SDK
-			if ( bAnimatedVisibility || UsdUtils::IsAnimated( Prim ) )
+			if (bAnimatedVisibility || UsdUtils::IsAnimated(Prim))
 			{
-				ImportAnimation(ImportContext, Prim, bAnimatedVisibility, Component);
+				const TOptional<bool> HasAnimatedBounds = {};
+				ImportAnimation(ImportContext, Prim, bAnimatedVisibility, HasAnimatedBounds, Component);
+			}
+			else if (EUsdDrawMode DrawMode = UsdUtils::GetAppliedDrawMode(Prim); DrawMode != EUsdDrawMode::Default)
+			{
+				if (UsdUtils::HasAnimatedBounds(
+						Prim,
+						ImportContext.BBoxCache->GetIncludedPurposes(),
+						ImportContext.BBoxCache->GetUseExtentsHint(),
+						ImportContext.BBoxCache->GetIgnoreVisibility()
+					))
+				{
+					const TOptional<bool> HasAnimatedBounds = true;
+					ImportAnimation(ImportContext, Prim, bAnimatedVisibility, HasAnimatedBounds, Component);
+				}
 			}
 #endif // USE_USD_SDK
 		}
@@ -1833,6 +1857,33 @@ void UUsdStageImporter::ImportFromFile(FUsdStageImportContext& ImportContext)
 	ImportContext.LevelSequenceHelper.SetInfoCache(InfoCache);
 	ImportContext.LevelSequenceHelper.Init( ImportContext.Stage );  // Must happen after the context gets an InfoCache!
 
+	EUsdPurpose PurposesToImport = static_cast<EUsdPurpose>(ImportContext.ImportOptions->PurposesToImport);
+
+	float ImportTime = ImportContext.ImportOptions->bImportAtSpecificTimeCode
+						   ? ImportContext.ImportOptions->ImportTimeCode
+						   : static_cast<float>(UsdUtils::GetDefaultTimeCode());
+
+	if (!ImportContext.BBoxCache)
+	{
+		const bool bUseExtentsHint = true;
+		const bool bIgnoreVisibility = false;
+		ImportContext.BBoxCache = MakeShared<UE::FUsdGeomBBoxCache>(ImportTime, PurposesToImport, bUseExtentsHint, bIgnoreVisibility);
+	}
+
+	TOptional<EUsdPurpose> IncludedPurposesToRevertBBoxCacheTo;
+	if (ImportContext.BBoxCache->GetIncludedPurposes() != PurposesToImport)
+	{
+		IncludedPurposesToRevertBBoxCacheTo = ImportContext.BBoxCache->GetIncludedPurposes();
+		ImportContext.BBoxCache->SetIncludedPurposes(PurposesToImport);
+	}
+	TOptional<float> TimeToRevertBBoxCacheTo;
+	if (ImportContext.BBoxCache->GetTime() != ImportTime)
+	{
+		TimeToRevertBBoxCacheTo = ImportContext.BBoxCache->GetTime();
+		ImportContext.BBoxCache->SetTime(ImportTime);
+	}
+	ImportContext.LevelSequenceHelper.SetBBoxCache(ImportContext.BBoxCache);
+
 	// Shotgun approach to recreate all render states because we may want to reimport/delete/reassign a material/static/skeletalmesh while it is currently being drawn
 	FGlobalComponentRecreateRenderStateContext RecreateRenderStateContext;
 
@@ -1840,10 +1891,8 @@ void UUsdStageImporter::ImportFromFile(FUsdStageImportContext& ImportContext)
 	TranslationContext->bIsImporting = true;
 	TranslationContext->Level = ImportContext.World->GetCurrentLevel();
 	TranslationContext->ObjectFlags = ImportContext.ImportObjectFlags;
-	TranslationContext->Time = ImportContext.ImportOptions->bImportAtSpecificTimeCode
-		? ImportContext.ImportOptions->ImportTimeCode
-		: static_cast< float >( UsdUtils::GetDefaultTimeCode() );
-	TranslationContext->PurposesToLoad = ( EUsdPurpose ) ImportContext.ImportOptions->PurposesToImport;
+	TranslationContext->Time = ImportTime;
+	TranslationContext->PurposesToLoad = PurposesToImport;
 	TranslationContext->NaniteTriangleThreshold = ImportContext.ImportOptions->NaniteTriangleThreshold;
 	TranslationContext->RenderContext = ImportContext.ImportOptions->RenderContextToImport;
 	TranslationContext->MaterialPurpose = ImportContext.ImportOptions->MaterialPurpose;
@@ -1857,6 +1906,7 @@ void UUsdStageImporter::ImportFromFile(FUsdStageImportContext& ImportContext)
 	TranslationContext->bAllowParsingGroomAssets = ImportContext.ImportOptions->bImportGroomAssets;
 	TranslationContext->bTranslateOnlyUsedMaterials = ImportContext.ImportOptions->bImportOnlyUsedMaterials;
 	TranslationContext->InfoCache = InfoCache;
+	TranslationContext->BBoxCache = ImportContext.BBoxCache;
 	TranslationContext->BlendShapesByPath = &BlendShapesByPath;
 	TranslationContext->GroomInterpolationSettings = ImportContext.ImportOptions->GroomInterpolationSettings;
 	{
@@ -1880,6 +1930,15 @@ void UUsdStageImporter::ImportFromFile(FUsdStageImportContext& ImportContext)
 	UsdStageImporterImpl::NotifyAssetRegistry(PublishedAssetsAndDependencies);
 	UsdStageImporterImpl::RefreshComponents( ImportContext.SceneActor, ImportContext.ImportOptions->bImportAtSpecificTimeCode );
 
+	if (IncludedPurposesToRevertBBoxCacheTo.IsSet())
+	{
+		ImportContext.BBoxCache->SetIncludedPurposes(IncludedPurposesToRevertBBoxCacheTo.GetValue());
+	}
+	if (TimeToRevertBBoxCacheTo.IsSet())
+	{
+		ImportContext.BBoxCache->SetTime(TimeToRevertBBoxCacheTo.GetValue());
+	}
+
 	FUsdDelegates::OnPostUsdImport.Broadcast( ImportContext.FilePath );
 
 	// Analytics
@@ -1954,6 +2013,32 @@ bool UUsdStageImporter::ReimportSingleAsset(
 	ImportContext.LevelSequenceHelper.SetInfoCache(InfoCache);
 	ImportContext.LevelSequenceHelper.Init(ImportContext.Stage);  // Must happen after the context gets an InfoCache!
 
+	EUsdPurpose PurposesToImport = static_cast<EUsdPurpose>(ImportContext.ImportOptions->PurposesToImport);
+
+	float ImportTime = ImportContext.ImportOptions->bImportAtSpecificTimeCode
+						   ? ImportContext.ImportOptions->ImportTimeCode
+						   : static_cast<float>(UsdUtils::GetDefaultTimeCode());
+
+	if (!ImportContext.BBoxCache)
+	{
+		const bool bUseExtentsHint = true;
+		const bool bIgnoreVisibility = false;
+		ImportContext.BBoxCache = MakeShared<UE::FUsdGeomBBoxCache>(ImportTime, PurposesToImport, bUseExtentsHint, bIgnoreVisibility);
+	}
+	TOptional<EUsdPurpose> IncludedPurposesToRevertTo;
+	if (ImportContext.BBoxCache->GetIncludedPurposes() != PurposesToImport)
+	{
+		IncludedPurposesToRevertTo = ImportContext.BBoxCache->GetIncludedPurposes();
+		ImportContext.BBoxCache->SetIncludedPurposes(PurposesToImport);
+	}
+	TOptional<float> TimeToRevertBBoxCacheTo;
+	if (ImportContext.BBoxCache->GetTime() != ImportTime)
+	{
+		TimeToRevertBBoxCacheTo = ImportContext.BBoxCache->GetTime();
+		ImportContext.BBoxCache->SetTime(ImportTime);
+	}
+	ImportContext.LevelSequenceHelper.SetBBoxCache(ImportContext.BBoxCache);
+
 	// Shotgun approach to recreate all render states because we may want to reimport/delete/reassign a material/static/skeletalmesh while it is currently being drawn
 	FGlobalComponentRecreateRenderStateContext RecreateRenderStateContext;
 
@@ -1961,10 +2046,8 @@ bool UUsdStageImporter::ReimportSingleAsset(
 	TranslationContext->bIsImporting = true;
 	TranslationContext->Level = ImportContext.World->GetCurrentLevel();
 	TranslationContext->ObjectFlags = ImportContext.ImportObjectFlags;
-	TranslationContext->Time = ImportContext.ImportOptions->bImportAtSpecificTimeCode
-		? ImportContext.ImportOptions->ImportTimeCode
-		: static_cast< float >( UsdUtils::GetDefaultTimeCode() );
-	TranslationContext->PurposesToLoad = ( EUsdPurpose ) ImportContext.ImportOptions->PurposesToImport;
+	TranslationContext->Time = ImportTime;
+	TranslationContext->PurposesToLoad = PurposesToImport;
 	TranslationContext->NaniteTriangleThreshold = ImportContext.ImportOptions->NaniteTriangleThreshold;
 	TranslationContext->RenderContext = ImportContext.ImportOptions->RenderContextToImport;
 	TranslationContext->MaterialPurpose = ImportContext.ImportOptions->MaterialPurpose;
@@ -1977,6 +2060,7 @@ bool UUsdStageImporter::ReimportSingleAsset(
 	TranslationContext->bAllowParsingGroomAssets = ImportContext.ImportOptions->bImportGroomAssets;
 	TranslationContext->bTranslateOnlyUsedMaterials = ImportContext.ImportOptions->bImportOnlyUsedMaterials;
 	TranslationContext->InfoCache = InfoCache;
+	TranslationContext->BBoxCache = ImportContext.BBoxCache;
 	TranslationContext->BlendShapesByPath = &BlendShapesByPath;
 	TranslationContext->GroomInterpolationSettings = ImportContext.ImportOptions->GroomInterpolationSettings;
 	{
@@ -2047,6 +2131,15 @@ bool UUsdStageImporter::ReimportSingleAsset(
 	UsdStageImporterImpl::NotifyAssetRegistry( { ReimportedObject } );
 	UsdStageImporterImpl::RefreshComponents( ImportContext.SceneActor, ImportContext.ImportOptions->bImportAtSpecificTimeCode );
 
+	if (IncludedPurposesToRevertTo.IsSet())
+	{
+		ImportContext.BBoxCache->SetIncludedPurposes(IncludedPurposesToRevertTo.GetValue());
+	}
+	if (TimeToRevertBBoxCacheTo.IsSet())
+	{
+		ImportContext.BBoxCache->SetTime(TimeToRevertBBoxCacheTo.GetValue());
+	}
+
 	FUsdDelegates::OnPostUsdImport.Broadcast(ImportContext.FilePath);
 
 	// Analytics
diff --git a/Engine/Plugins/Importers/USDImporter/Source/USDStageImporter/Public/USDStageImportContext.h b/Engine/Plugins/Importers/USDImporter/Source/USDStageImporter/Public/USDStageImportContext.h
--- a/Engine/Plugins/Importers/USDImporter/Source/USDStageImporter/Public/USDStageImportContext.h
+++ b/Engine/Plugins/Importers/USDImporter/Source/USDStageImporter/Public/USDStageImportContext.h
@@ -64,6 +64,9 @@ struct USDSTAGEIMPORTER_API FUsdStageImportContext
 	UE_DEPRECATED(5.3, "The import process now always builds its own InfoCache, so this member is no longer used")
 	TSharedPtr<FUsdInfoCache> InfoCache;
 
+	/** Bounding box cache used for the USD stage in case we have to spawn bounds components */
+	TSharedPtr<UE::FUsdGeomBBoxCache> BBoxCache;
+
 	/**
 	 * When parsing materials, we keep track of which primvar we mapped to which UV channel.
 	 * When parsing meshes later, we use this data to place the correct primvar values in each UV channel.
diff --git a/Engine/Plugins/Importers/USDImporter/Source/USDUtilities/Private/USDConversionUtils.cpp b/Engine/Plugins/Importers/USDImporter/Source/USDUtilities/Private/USDConversionUtils.cpp
--- a/Engine/Plugins/Importers/USDImporter/Source/USDUtilities/Private/USDConversionUtils.cpp
+++ b/Engine/Plugins/Importers/USDImporter/Source/USDUtilities/Private/USDConversionUtils.cpp
@@ -3,6 +3,7 @@
 #include "USDConversionUtils.h"
 
 #include "USDAssetImportData.h"
+#include "USDDrawModeComponent.h"
 #include "USDClassesModule.h"
 #include "USDDuplicateType.h"
 #include "USDErrorUtils.h"
@@ -10,6 +11,7 @@
 #include "USDLayerUtils.h"
 #include "USDLog.h"
 #include "USDProjectSettings.h"
+#include "USDSkeletalDataConversion.h"
 #include "USDTypesConversion.h"
 #include "USDUnrealAssetInfo.h"
 
@@ -68,6 +70,7 @@
 	#include "pxr/usd/usdGeom/imageable.h"
 	#include "pxr/usd/usdGeom/mesh.h"
 	#include "pxr/usd/usdGeom/metrics.h"
+	#include "pxr/usd/usdGeom/modelAPI.h"
 	#include "pxr/usd/usdGeom/primvar.h"
 	#include "pxr/usd/usdGeom/primvarsAPI.h"
 	#include "pxr/usd/usdLux/diskLight.h"
@@ -76,6 +79,7 @@
 	#include "pxr/usd/usdLux/rectLight.h"
 	#include "pxr/usd/usdLux/shapingAPI.h"
 	#include "pxr/usd/usdLux/sphereLight.h"
+	#include "pxr/usd/usdSkel/animation.h"
 	#include "pxr/usd/usdSkel/binding.h"
 	#include "pxr/usd/usdSkel/cache.h"
 	#include "pxr/usd/usdSkel/root.h"
@@ -977,6 +981,259 @@ bool UsdUtils::HasAnimatedVisibility( const pxr::UsdPrim& Prim )
 	return false;
 }
 
+namespace UE::ConversionUtilsImpl::Private
+{
+	// Convenience function so we don't have to spell this out every time
+	inline void CollectTimeSamplesIfNeeded(bool bCollectTimeSamples, const pxr::UsdAttribute& Attr, TArray<double>* OutTimeSamples)
+	{
+		std::vector<double> TempTimeSamples;
+		if (bCollectTimeSamples && Attr.GetTimeSamples(&TempTimeSamples))
+		{
+			OutTimeSamples->Append(TempTimeSamples.data(), TempTimeSamples.size());
+		}
+	}
+
+	bool GetOrCollectAnimatedBounds(
+		const pxr::UsdPrim& Prim,
+		TArray<double>* OutTimeSamples,
+		bool bCollectTimeSamples,
+		bool bIsParentPrim,
+		EUsdPurpose IncludedPurposes,
+		bool bUseExtentsHint,
+		bool bIgnoreVisibility
+	)
+	{
+		if (!Prim)
+		{
+			return false;
+		}
+
+		// If we want to collect timeSamples we must have some place to put them in
+		if (!ensure(!bCollectTimeSamples || OutTimeSamples))
+		{
+			return false;
+		}
+
+		FScopedUsdAllocs UsdAllocs;
+
+		// If the prim is fully invisible due to visibility or purpose then we shouldn't even check it
+		bool bHasAnimatedVisibility = false;
+		if (pxr::UsdGeomImageable Imageable{Prim}; !bIgnoreVisibility && Imageable)
+		{
+			if (pxr::UsdAttribute Visibility = Imageable.GetVisibilityAttr())
+			{
+				// Keep track of this for later
+				bHasAnimatedVisibility = Visibility.ValueMightBeTimeVarying();
+
+				if (bHasAnimatedVisibility)
+				{
+					CollectTimeSamplesIfNeeded(bCollectTimeSamples, Visibility, OutTimeSamples);
+				}
+				else
+				{
+					pxr::TfToken VisibilityToken;
+					if (!bIsParentPrim && Visibility.Get(&VisibilityToken) && VisibilityToken == pxr::UsdGeomTokens->invisible)
+					{
+						// We don't "propagate the (in)visibility token", we just flat out stop recursing and abandon the subtree
+						return false;
+					}
+				}
+			}
+		}
+		if (!bIsParentPrim && !EnumHasAllFlags(IncludedPurposes, IUsdPrim::GetPurpose(Prim)))
+		{
+			return false;
+		}
+
+		// If the prim has authored animated extents we know we're fully done, because our computed bounds
+		// will also need to be animated and will read *exclusively* from these anyway.
+		// We don't even need to collect any further timeSamples from child prims after this, as we will be ignoring individual
+		// animations on random prims in the subtree and instead just using the authored extent animation.
+		// Also, extentsHint is preferred over extent, so check for that first.
+		if (pxr::UsdGeomModelAPI GeomModelAPI{Prim}; bUseExtentsHint && GeomModelAPI)
+		{
+			if (pxr::UsdAttribute ExtentsHint = GeomModelAPI.GetExtentsHintAttr())
+			{
+				if (ExtentsHint.HasAuthoredValue())
+				{
+					CollectTimeSamplesIfNeeded(bCollectTimeSamples, ExtentsHint, OutTimeSamples);
+					return ExtentsHint.ValueMightBeTimeVarying();
+				}
+			}
+		}
+		if (pxr::UsdGeomBoundable Boundable{Prim})
+		{
+			if (pxr::UsdAttribute Extent = Boundable.GetExtentAttr())
+			{
+				// If we have authored extent or extentsHint (even if not animated, i.e. just default opinions), the
+				// BBoxCache will refuse to compute bounds at any timeCode and just fallback to using the authored stuff
+				if (Extent.HasAuthoredValue())
+				{
+					CollectTimeSamplesIfNeeded(bCollectTimeSamples, Extent, OutTimeSamples);
+					return Extent.ValueMightBeTimeVarying();
+				}
+			}
+		}
+
+		// It's visible at the default timeCode, but has animated visibility. This means
+		// it could affect the bounds as it becomes visible or invisible, so just return now.
+		if (!bCollectTimeSamples && bHasAnimatedVisibility)
+		{
+			return true;
+		}
+
+		bool bHasAnimatedBounds = bHasAnimatedVisibility;
+
+		// Otherwise the prim may have some animated attributes that would make our parent extents animated.
+		// For this function we mostly care about whether the *bounds themselves* are animated.
+		// The parent prim having animated transform means we'll just put this transform on the component itself,
+		// but the bounds could remain un-animated
+		if (pxr::UsdGeomXformable Xformable{Prim}; !bIsParentPrim && Xformable)
+		{
+			if (Xformable.TransformMightBeTimeVarying())
+			{
+				bHasAnimatedBounds = true;
+
+				std::vector<double> TempTimeSamples;
+				if (bCollectTimeSamples && Xformable.GetTimeSamples(&TempTimeSamples))
+				{
+					OutTimeSamples->Append(TempTimeSamples.data(), TempTimeSamples.size());
+				}
+			}
+		}
+		else if (pxr::UsdGeomPointBased PointBased{Prim})
+		{
+			if (pxr::UsdAttribute Points = PointBased.GetPointsAttr())
+			{
+				if (Points.ValueMightBeTimeVarying())
+				{
+					CollectTimeSamplesIfNeeded(bCollectTimeSamples, Points, OutTimeSamples);
+					bHasAnimatedBounds = true;
+				}
+			}
+		}
+		else if (pxr::UsdGeomPointInstancer PointInstancer{Prim})
+		{
+			if (pxr::UsdAttribute Positions = PointInstancer.GetPositionsAttr())
+			{
+				if (Positions.ValueMightBeTimeVarying())
+				{
+					CollectTimeSamplesIfNeeded(bCollectTimeSamples, Positions, OutTimeSamples);
+					bHasAnimatedBounds = true;
+				}
+			}
+		}
+		// Check for a SkelRoot with SkelAnimation
+		else if (UE::FUsdPrim SkelAnimationPrim = UsdUtils::FindFirstAnimationSource(UE::FUsdPrim{Prim}))
+		{
+			pxr::UsdSkelAnimation SkelAnim{SkelAnimationPrim};
+			if (ensure(SkelAnim))
+			{
+				const bool bIncludeInherited = false;
+				for (const pxr::TfToken& SkelAnimAttrName : SkelAnim.GetSchemaAttributeNames(bIncludeInherited))
+				{
+					if (pxr::UsdAttribute Attr = SkelAnim.GetPrim().GetAttribute(SkelAnimAttrName))
+					{
+						if (Attr.ValueMightBeTimeVarying())
+						{
+							bHasAnimatedBounds = true;
+
+							if (!bCollectTimeSamples)
+							{
+								break;
+							}
+							CollectTimeSamplesIfNeeded(bCollectTimeSamples, Attr, OutTimeSamples);
+						}
+					}
+				}
+			}
+		}
+
+		// If we're not collecting timeSamples and we run into a prim with animated bounds then we know that
+		// we're done, and can return then. If we're collecting timeSamples however then we want instead to remember
+		// that we found those animated bounds, but still try to step into children in case they also had animated
+		// bounds and additional timeSamples
+		if (!bCollectTimeSamples && bHasAnimatedBounds)
+		{
+			return true;
+		}
+
+		for (pxr::UsdPrim Child : Prim.GetFilteredChildren(pxr::UsdTraverseInstanceProxies(pxr::UsdPrimAllPrimsPredicate)))
+		{
+			const bool bChildIsParentPrim = false;
+			bHasAnimatedBounds |= GetOrCollectAnimatedBounds(
+				Child,
+				OutTimeSamples,
+				bCollectTimeSamples,
+				bChildIsParentPrim,
+				IncludedPurposes,
+				bUseExtentsHint,
+				bIgnoreVisibility
+			);
+
+			// Don't need to visit any other children, we're done here
+			if (!bCollectTimeSamples && bHasAnimatedBounds)
+			{
+				return true;
+			}
+		}
+
+		return bHasAnimatedBounds;
+	}
+}
+
+bool UsdUtils::HasAnimatedBounds(
+	const pxr::UsdPrim& Prim,
+	EUsdPurpose IncludedPurposes,
+	bool bUseExtentsHint,
+	bool bIgnoreVisibility
+)
+{
+	// "ParentPrim" here because there are slight differences in behavior between handling the actual provided
+	// prim and another random prim in its subtree (for which bIsParentPrim will be 'false')
+	const bool bIsParentPrim = true;
+	const bool bCollectTimeSamples = false;
+	TArray<double>* OutTimeSamples = nullptr;
+	return UE::ConversionUtilsImpl::Private::GetOrCollectAnimatedBounds(
+		Prim,
+		OutTimeSamples,
+		bCollectTimeSamples,
+		bIsParentPrim,
+		IncludedPurposes,
+		bUseExtentsHint,
+		bIgnoreVisibility
+	);
+}
+
+bool UsdUtils::GetAnimatedBoundsTimeSamples(
+	const pxr::UsdPrim& InPrim,
+	TArray<double>& OutTimeSamples,
+	EUsdPurpose InIncludedPurposes,
+	bool bInUseExtentsHint,
+	bool bInIgnoreVisibility
+)
+{
+	OutTimeSamples.Reset();
+
+	// "ParentPrim" here because there are slight differences in behavior between handling the actual provided
+	// prim and another random prim in its subtree (for which bIsParentPrim will be 'false')
+	const bool bIsParentPrim = true;
+	const bool bCollectTimeSamples = true;
+	const bool bHasAnimatedBounds = UE::ConversionUtilsImpl::Private::GetOrCollectAnimatedBounds(
+		InPrim,
+		&OutTimeSamples,
+		bCollectTimeSamples,
+		bIsParentPrim,
+		InIncludedPurposes,
+		bInUseExtentsHint,
+		bInIgnoreVisibility
+	);
+
+	OutTimeSamples.Sort();
+
+	return bHasAnimatedBounds;
+}
+
 EUsdDefaultKind UsdUtils::GetDefaultKind( const pxr::UsdPrim& Prim )
 {
 	FScopedUsdAllocs Allocs;
@@ -1075,6 +1332,83 @@ bool UsdUtils::SetDefaultKind( pxr::UsdPrim& Prim, EUsdDefaultKind NewKind )
 	return IUsdPrim::SetKind( Prim, NewKindToken );
 }
 
+EUsdDrawMode UsdUtils::GetAppliedDrawMode(const pxr::UsdPrim& Prim)
+{
+	// Reference: https://openusd.org/release/api/class_usd_geom_model_a_p_i.html#UsdGeomModelAPI_drawMode
+
+	if (!Prim)
+	{
+		return EUsdDrawMode::Default;
+	}
+
+	FScopedUsdAllocs Allocs;
+
+	// Only "models" should have these (i.e. uninterrupted chain of authored "kind"s back to the root prim)
+	if (!Prim.IsModel())
+	{
+		return EUsdDrawMode::Default;
+	}
+
+	pxr::UsdGeomModelAPI GeomModelAPI{Prim};
+	if (!GeomModelAPI)
+	{
+		return EUsdDrawMode::Default;
+	}
+
+	bool bHasAuthoredApply = false;
+	bool bShouldApplyFromAttr = false;
+	pxr::UsdAttribute Attr = GeomModelAPI.GetModelApplyDrawModeAttr();
+	if (Attr && Attr.HasAuthoredValue() && Attr.Get(&bShouldApplyFromAttr))
+	{
+		if (!bShouldApplyFromAttr)
+		{
+			return EUsdDrawMode::Default;
+		}
+
+		bHasAuthoredApply = true;
+	}
+
+	// "Models of kind component are treated as if model:applyDrawMode were true"
+	// According to UsdImagingDelegate::_IsDrawModeApplied this only works as a "fallback" though:
+	// if the prim has authored whether to apply or not we always use that directly
+	pxr::UsdModelAPI Model{Prim};
+	const bool bIsComponentKind = Model && Model.IsKind(pxr::KindTokens->component, pxr::UsdModelAPI::KindValidationNone);
+	if (!bHasAuthoredApply && !bIsComponentKind)
+	{
+		return EUsdDrawMode::Default;
+	}
+
+	// Note: We can provide the parent draw mode to optimize the ComputeModelDrawMode call if it becomes an issue
+	pxr::TfToken DesiredDrawMode = GeomModelAPI.ComputeModelDrawMode();
+	if (DesiredDrawMode == pxr::UsdGeomTokens->default_)
+	{
+		return EUsdDrawMode::Default;
+	}
+	else if (DesiredDrawMode == pxr::UsdGeomTokens->origin)
+	{
+		return EUsdDrawMode::Origin;
+	}
+	else if (DesiredDrawMode == pxr::UsdGeomTokens->bounds)
+	{
+		return EUsdDrawMode::Bounds;
+	}
+	else if (DesiredDrawMode == pxr::UsdGeomTokens->cards)
+	{
+		return EUsdDrawMode::Cards;
+	}
+	else if (DesiredDrawMode == pxr::UsdGeomTokens->inherited)
+	{
+		// If we're using ComputeModelDrawMode we shouldn't get inherited or anything else here
+		ensure(false);
+		return EUsdDrawMode::Inherited;
+	}
+	else
+	{
+		ensure(false);
+		return EUsdDrawMode::Default;
+	}
+}
+
 TArray< TUsdStore< pxr::UsdPrim > > UsdUtils::GetAllPrimsOfType( const pxr::UsdPrim& StartPrim, const pxr::TfType& SchemaType, const TArray< TUsdStore< pxr::TfType > >& ExcludeSchemaTypes )
 {
     return GetAllPrimsOfType( StartPrim, SchemaType, []( const pxr::UsdPrim& ) { return false; }, ExcludeSchemaTypes );
diff --git a/Engine/Plugins/Importers/USDImporter/Source/USDUtilities/Private/USDGeomMeshConversion.cpp b/Engine/Plugins/Importers/USDImporter/Source/USDUtilities/Private/USDGeomMeshConversion.cpp
--- a/Engine/Plugins/Importers/USDImporter/Source/USDUtilities/Private/USDGeomMeshConversion.cpp
+++ b/Engine/Plugins/Importers/USDImporter/Source/USDUtilities/Private/USDGeomMeshConversion.cpp
@@ -6,7 +6,9 @@
 
 #include "UnrealUSDWrapper.h"
 #include "USDAttributeUtils.h"
+#include "USDClassesModule.h"
 #include "USDConversionUtils.h"
+#include "USDDrawModeComponent.h"
 #include "USDLog.h"
 #include "USDMemory.h"
 #include "USDPrimConversion.h"
@@ -2822,7 +2824,7 @@ bool UsdToUnreal::ConvertGeomPrimitiveTransform(
 {
 	TRACE_CPUPROFILER_EVENT_SCOPE(UsdToUnreal::ConvertGeomPrimitive);
 
-	if (!InPrim || !InPrim.IsA<pxr::UsdGeomGprim>())
+	if (!InPrim || !InPrim.IsA<pxr::UsdGeomGprim>() || UsdUtils::GetAppliedDrawMode(InPrim) != EUsdDrawMode::Default)
 	{
 		return false;
 	}
@@ -3012,6 +3014,7 @@ bool UsdToUnreal::ConvertGeomPrimitiveTransform(
 	return false;
 }
 
+PRAGMA_DISABLE_DEPRECATION_WARNINGS
 UMaterialInstanceDynamic* UsdUtils::CreateDisplayColorMaterialInstanceDynamic( const UsdUtils::FDisplayColorMaterial& DisplayColorDescription )
 {
 	const UUsdProjectSettings* Settings = GetDefault<UUsdProjectSettings>();
@@ -3068,7 +3071,7 @@ UMaterialInstanceDynamic* UsdUtils::CreateDisplayColorMaterialInstanceDynamic( c
 	return nullptr;
 }
 
-UMaterialInstanceConstant* UsdUtils::CreateDisplayColorMaterialInstanceConstant( const FDisplayColorMaterial& DisplayColorDescription )
+UMaterialInstanceConstant* UsdUtils::CreateDisplayColorMaterialInstanceConstant( const UsdUtils::FDisplayColorMaterial& DisplayColorDescription )
 {
 #if WITH_EDITOR
 	const UUsdProjectSettings* Settings = GetDefault<UUsdProjectSettings>();
@@ -3121,6 +3124,7 @@ UMaterialInstanceConstant* UsdUtils::CreateDisplayColorMaterialInstanceConstant(
 #endif // WITH_EDITOR
 	return nullptr;
 }
+PRAGMA_ENABLE_DEPRECATION_WARNINGS
 
 UsdUtils::FUsdPrimMaterialAssignmentInfo UsdUtils::GetPrimMaterialAssignments(
 	const pxr::UsdPrim& UsdPrim,
@@ -3350,7 +3354,7 @@ UsdUtils::FUsdPrimMaterialAssignmentInfo UsdUtils::GetPrimMaterialAssignments(
 		}
 	}
 
-	TOptional<FDisplayColorMaterial> DisplayColor;
+	TOptional<UsdUtils::FDisplayColorMaterial> DisplayColor;
 
 	bool bHasMainAssignment = false;
 	if (bNeedsMainAssignment)
@@ -3787,27 +3791,6 @@ bool UnrealToUsd::ConvertMeshDescriptions( const TArray<FMeshDescription>& LODIn
 	return true;
 }
 
-FString UsdUtils::FDisplayColorMaterial::ToString()
-{
-	return FString::Printf( TEXT("%s_%d_%d"), *UsdGeomMeshImpl::DisplayColorID, bHasOpacity, bIsDoubleSided );
-}
-
-TOptional<UsdUtils::FDisplayColorMaterial> UsdUtils::FDisplayColorMaterial::FromString( const FString& DisplayColorString )
-{
-	TArray<FString> Tokens;
-	DisplayColorString.ParseIntoArray( Tokens, TEXT( "_" ) );
-
-	if ( Tokens.Num() != 3 || Tokens[ 0 ] != UsdGeomMeshImpl::DisplayColorID )
-	{
-		return {};
-	}
-
-	UsdUtils::FDisplayColorMaterial Result;
-	Result.bHasOpacity = static_cast< bool >( FCString::Atoi( *Tokens[ 1 ] ) );
-	Result.bIsDoubleSided = static_cast< bool >( FCString::Atoi( *Tokens[ 2 ] ) );
-	return Result;
-}
-
 TOptional<UsdUtils::FDisplayColorMaterial> UsdUtils::ExtractDisplayColorMaterial(const pxr::UsdGeomGprim& Gprim, const pxr::UsdTimeCode TimeCode)
 {
 	if (!Gprim)
diff --git a/Engine/Plugins/Importers/USDImporter/Source/USDUtilities/Private/USDPrimConversion.cpp b/Engine/Plugins/Importers/USDImporter/Source/USDUtilities/Private/USDPrimConversion.cpp
--- a/Engine/Plugins/Importers/USDImporter/Source/USDUtilities/Private/USDPrimConversion.cpp
+++ b/Engine/Plugins/Importers/USDImporter/Source/USDUtilities/Private/USDPrimConversion.cpp
@@ -3,8 +3,10 @@
 #include "USDPrimConversion.h"
 
 #include "UnrealUSDWrapper.h"
+#include "USDAssetCache2.h"
 #include "USDAssetUserData.h"
 #include "USDAttributeUtils.h"
+#include "USDDrawModeComponent.h"
 #include "USDConversionUtils.h"
 #include "USDLayerUtils.h"
 #include "USDLightConversion.h"
@@ -14,6 +16,7 @@
 #include "USDTypesConversion.h"
 
 #include "UsdWrappers/UsdAttribute.h"
+#include "UsdWrappers/UsdGeomBBoxCache.h"
 #include "UsdWrappers/UsdGeomXformable.h"
 #include "UsdWrappers/UsdPrim.h"
 #include "UsdWrappers/UsdStage.h"
@@ -35,9 +38,11 @@
 #include "Components/SkeletalMeshComponent.h"
 #include "Components/SkinnedMeshComponent.h"
 #include "Components/SpotLightComponent.h"
+#include "EditorFramework/AssetImportData.h"
 #include "Engine/SkeletalMesh.h"
 #include "Engine/SkinnedAssetCommon.h"
 #include "Engine/StaticMesh.h"
+#include "Engine/Texture2D.h"
 #include "GeometryCache.h"
 #include "GeometryCacheComponent.h"
 #include "InstancedFoliageActor.h"
@@ -49,16 +54,19 @@
 #include "Sections/MovieSceneColorSection.h"
 #include "Sections/MovieSceneFloatSection.h"
 #include "Sections/MovieSceneVisibilitySection.h"
+#include "Sections/MovieSceneVectorSection.h"
 #include "Tracks/MovieScene3DTransformTrack.h"
 #include "Tracks/MovieSceneBoolTrack.h"
 #include "Tracks/MovieSceneColorTrack.h"
 #include "Tracks/MovieSceneFloatTrack.h"
 #include "Tracks/MovieScenePropertyTrack.h"
 #include "Tracks/MovieSceneVisibilityTrack.h"
+#include "Tracks/MovieSceneVectorTrack.h"
 
 #if USE_USD_SDK
 
 #include "USDIncludesStart.h"
+	#include "pxr/usd/kind/registry.h"
 	#include "pxr/usd/sdf/changeBlock.h"
 	#include "pxr/usd/usd/attribute.h"
 	#include "pxr/usd/usd/prim.h"
@@ -71,6 +79,7 @@
 	#include "pxr/usd/usdGeom/cylinder.h"
 	#include "pxr/usd/usdGeom/imageable.h"
 	#include "pxr/usd/usdGeom/mesh.h"
+	#include "pxr/usd/usdGeom/modelAPI.h"
 	#include "pxr/usd/usdGeom/plane.h"
 	#include "pxr/usd/usdGeom/pointInstancer.h"
 	#include "pxr/usd/usdGeom/scope.h"
@@ -95,6 +104,13 @@
 	#include "pxr/usd/usdSkel/root.h"
 #include "USDIncludesEnd.h"
 
+static bool GConsiderAllPrimsHaveAnimatedBounds = false;
+static FAutoConsoleVariableRef CVarConsiderAllPrimsHaveAnimatedBounds(
+	TEXT("USD.Bounds.ConsiderAllPrimsHaveAnimatedBounds"),
+	GConsiderAllPrimsHaveAnimatedBounds,
+	TEXT("When active prevents USD from caching computed bounds between timeSamples for any prim, which allows us to force it to recompute accurate bounds for cases it does not naturally consider animated (e.g. for animated Mesh points, skeletal animation, etc.). Warning: This can be extremely expensive!")
+);
+
 namespace UE
 {
 	namespace USDPrimConversionImpl
@@ -920,6 +936,151 @@ bool UsdToUnreal::ConvertTransformTimeSamples( const UE::FUsdStage& Stage, const
 	return true;
 }
 
+bool UsdToUnreal::ConvertBoundsTimeSamples(
+	const UE::FUsdPrim& InPrim,
+	const TArray<double>& InUsdTimeSamples,
+	const FMovieSceneSequenceTransform& InSequenceTransform,
+	UMovieSceneDoubleVectorTrack& InOutMinTrack,
+	UMovieSceneDoubleVectorTrack& InOutMaxTrack,
+	UE::FUsdGeomBBoxCache* InOutBBoxCache
+)
+{
+	if (!InPrim)
+	{
+		return false;
+	}
+
+	TOptional<FWriteScopeLock> BBoxLock;
+	pxr::UsdGeomBBoxCache* BBoxCache = nullptr;
+	if (InOutBBoxCache)
+	{
+		BBoxLock.Emplace(InOutBBoxCache->Lock);
+		BBoxCache = &static_cast<pxr::UsdGeomBBoxCache&>(*InOutBBoxCache);
+	}
+
+	// Create an BBoxCache on-demand (don't need to lock this one as it's purely ours)
+	TOptional<pxr::UsdGeomBBoxCache> TempBBoxCacheStorage;
+	if (BBoxCache == nullptr)
+	{
+		const bool bUseExtentsHint = true;
+		const bool bIgnoreVisibility = false;
+		static std::vector<pxr::TfToken> DefaultTokenVector{pxr::UsdGeomTokens->proxy, pxr::UsdGeomTokens->render};
+		TempBBoxCacheStorage.Emplace(pxr::UsdTimeCode::EarliestTime(), DefaultTokenVector, bUseExtentsHint, bIgnoreVisibility);
+		BBoxCache = &TempBBoxCacheStorage.GetValue();
+	}
+
+	const UMovieScene* MovieScene = InOutMinTrack.GetTypedOuter<UMovieScene>();
+	if (!MovieScene)
+	{
+		return false;
+	}
+
+	const FFrameRate Resolution = MovieScene->GetTickResolution();
+
+	FScopedUsdAllocs Allocs;
+
+	pxr::UsdStageRefPtr UsdStage = InPrim.GetStage();
+	FUsdStageInfo StageInfo{UsdStage};
+
+	TArray<FFrameNumber> FrameNumbers;
+	FrameNumbers.Reserve(InUsdTimeSamples.Num());
+
+	TArray<FMovieSceneDoubleValue> MinXValues;
+	TArray<FMovieSceneDoubleValue> MinYValues;
+	TArray<FMovieSceneDoubleValue> MinZValues;
+	TArray<FMovieSceneDoubleValue> MaxXValues;
+	TArray<FMovieSceneDoubleValue> MaxYValues;
+	TArray<FMovieSceneDoubleValue> MaxZValues;
+	MinXValues.Reserve(InUsdTimeSamples.Num());
+	MinYValues.Reserve(InUsdTimeSamples.Num());
+	MinZValues.Reserve(InUsdTimeSamples.Num());
+	MaxXValues.Reserve(InUsdTimeSamples.Num());
+	MaxYValues.Reserve(InUsdTimeSamples.Num());
+	MaxZValues.Reserve(InUsdTimeSamples.Num());
+
+	const double StageTimeCodesPerSecond = UsdStage->GetTimeCodesPerSecond();
+	const FFrameRate StageFrameRate(StageTimeCodesPerSecond, 1);
+
+	const ERichCurveInterpMode InterpMode = (UsdStage->GetInterpolationType() == pxr::UsdInterpolationTypeLinear)
+												? ERichCurveInterpMode::RCIM_Linear
+												: ERichCurveInterpMode::RCIM_Constant;
+
+	double LastTimeSample = TNumericLimits<double>::Lowest();
+	for (const double UsdTimeSample : InUsdTimeSamples)
+	{
+		// We never want to evaluate the same time twice
+		if (FMath::IsNearlyEqual(UsdTimeSample, LastTimeSample))
+		{
+			continue;
+		}
+		LastTimeSample = UsdTimeSample;
+
+		int32 FrameNumber = FMath::FloorToInt(UsdTimeSample);
+		float SubFrameNumber = UsdTimeSample - FrameNumber;
+
+		FFrameTime FrameTime(FrameNumber, SubFrameNumber);
+
+		FFrameTime KeyFrameTime = FFrameRate::TransformTime(FrameTime, StageFrameRate, Resolution);
+		KeyFrameTime *= InSequenceTransform;
+		FrameNumbers.Add(KeyFrameTime.GetFrame());
+
+		if (GConsiderAllPrimsHaveAnimatedBounds)
+		{
+			BBoxCache->Clear();
+		}
+
+		// It may seem like we're repeatedly invalidating the BBoxCache by doing this but actually it should retain some stuff,
+		// like bounds from non-animated prims
+		BBoxCache->SetTime(UsdTimeSample);
+
+		// Note: This can be extremely expensive, as it may fallback to computing new bounds, traversing points and everything for the entire subtree
+		pxr::GfBBox3d BoxAndMatrix = BBoxCache->ComputeUntransformedBound(InPrim);
+		pxr::GfRange3d Box = BoxAndMatrix.ComputeAlignedRange();
+
+		FBox UEBox;
+		if (!Box.IsEmpty())
+		{
+			FVector UESpaceUSDMin = UsdToUnreal::ConvertVector(StageInfo, Box.GetMin());
+			FVector UESpaceUSDMax = UsdToUnreal::ConvertVector(StageInfo, Box.GetMax());
+			UEBox = FBox{TArray<FVector>{UESpaceUSDMin, UESpaceUSDMax}};
+		}
+
+		MinXValues.Emplace_GetRef(UEBox.Min.X).InterpMode = InterpMode;
+		MinYValues.Emplace_GetRef(UEBox.Min.Y).InterpMode = InterpMode;
+		MinZValues.Emplace_GetRef(UEBox.Min.Z).InterpMode = InterpMode;
+
+		MaxXValues.Emplace_GetRef(UEBox.Max.X).InterpMode = InterpMode;
+		MaxYValues.Emplace_GetRef(UEBox.Max.Y).InterpMode = InterpMode;
+		MaxZValues.Emplace_GetRef(UEBox.Max.Z).InterpMode = InterpMode;
+	}
+
+	bool bSectionAdded = false;
+	UMovieSceneDoubleVectorSection* MinSection = Cast<UMovieSceneDoubleVectorSection>(InOutMinTrack.FindOrAddSection(0, bSectionAdded));
+	UMovieSceneDoubleVectorSection* MaxSection = Cast<UMovieSceneDoubleVectorSection>(InOutMaxTrack.FindOrAddSection(0, bSectionAdded));
+	MinSection->EvalOptions.CompletionMode = EMovieSceneCompletionMode::KeepState;
+	MaxSection->EvalOptions.CompletionMode = EMovieSceneCompletionMode::KeepState;
+
+	TArrayView<FMovieSceneDoubleChannel*> MinChannels = MinSection->GetChannelProxy().GetChannels<FMovieSceneDoubleChannel>();
+	TArrayView<FMovieSceneDoubleChannel*> MaxChannels = MaxSection->GetChannelProxy().GetChannels<FMovieSceneDoubleChannel>();
+	if (!ensure(MinChannels.Num() == 3 && MaxChannels.Num() == 3))
+	{
+		return false;
+	}
+
+	MinChannels[0]->Set(FrameNumbers, MinXValues);
+	MinChannels[1]->Set(FrameNumbers, MinYValues);
+	MinChannels[2]->Set(FrameNumbers, MinZValues);
+
+	MaxChannels[0]->Set(FrameNumbers, MaxXValues);
+	MaxChannels[1]->Set(FrameNumbers, MaxYValues);
+	MaxChannels[2]->Set(FrameNumbers, MaxZValues);
+
+	MinSection->SetRange(MinSection->GetAutoSizeRange().Get(TRange<FFrameNumber>::Empty()));
+	MaxSection->SetRange(MaxSection->GetAutoSizeRange().Get(TRange<FFrameNumber>::Empty()));
+
+	return true;
+}
+
 UsdToUnreal::FPropertyTrackReader UsdToUnreal::CreatePropertyTrackReader( const UE::FUsdPrim& Prim, const FName& PropertyPath, bool bIgnorePrimLocalTransform )
 {
 	UsdToUnreal::FPropertyTrackReader Reader;
@@ -1511,6 +1672,110 @@ UsdToUnreal::FPropertyTrackReader UsdToUnreal::CreatePropertyTrackReader( const
 	return Reader;
 }
 
+bool UsdToUnreal::ConvertBounds(
+	const pxr::UsdPrim& Prim,
+	UUsdDrawModeComponent* BoundsComponent,
+	double EvalTime,
+	pxr::UsdGeomBBoxCache* BBoxCache
+)
+{
+	// We're not going to check if Prim actually has the "bounds" draw mode or if it has "applyDrawMode" set
+	// to true, as that can be expensive and this can get called from UpdateComponents, which can get called
+	// every frame of Time animation. If we have a UUsdDrawModeComponent at all we'll assume we're OK here
+	if (!BoundsComponent)
+	{
+		return false;
+	}
+
+	FScopedUsdAllocs UsdAllocs;
+
+	TOptional<pxr::UsdGeomBBoxCache> TempBBoxCacheStorage;
+	if (BBoxCache == nullptr)
+	{
+		const bool bUseExtentsHint = true;
+		const bool bIgnoreVisibility = false;
+		static std::vector<pxr::TfToken> DefaultTokenVector{pxr::UsdGeomTokens->proxy, pxr::UsdGeomTokens->render};
+		TempBBoxCacheStorage.Emplace(EvalTime, DefaultTokenVector, bUseExtentsHint, bIgnoreVisibility);
+		BBoxCache = &TempBBoxCacheStorage.GetValue();
+	}
+
+	// BBoxCache only considers prims with animated transforms or visibility as needing animated computed bounds.
+	// This doesn't include e.g. animated points, meaning it would try reusing animated mesh bounds across frames for those animations,
+	// which can be very incorrect depending on the animation.
+	// With this very expensive trick we can flush the entire BBoxCache and compute new bounds every time, for each particular
+	// timeSample, which can get us accurate bounds for every frame. Obviously disabled by default, but could be useful if all you want
+	// is to import once, or something like this.
+	if (GConsiderAllPrimsHaveAnimatedBounds)
+	{
+		BBoxCache->Clear();
+	}
+
+	// We should do this here or else the BBoxCache may be set to a different time and we'd be reading wrong bounds
+	// (This can happening when opening stages as we may switch the BBoxCache time around when setting up the Sequencer tracks).
+	// Note that this does nothing in case BBoxCache is already at this time
+	BBoxCache->SetTime(EvalTime);
+
+	// Note: This can be extremely expensive, as it may fallback to computing new bounds, traversing points and everything for the entire subtree.
+	// We don't have a choice if we want decent bounds though, and in practice the user's assets will have (or can be set with) authored bounds,
+	// that should make this pretty fast
+	pxr::GfBBox3d BoxAndMatrix = BBoxCache->ComputeUntransformedBound(Prim);
+	pxr::GfRange3d Box = BoxAndMatrix.ComputeAlignedRange();
+
+	// USD will return a FLT_MAX box in case the prim doesn't contain anything, so we need to check for that as putting FLT_MAX directly
+	// on the component bounds is bad news
+	if (!Box.IsEmpty())
+	{
+		// Note that after we convert the USD min/max to UE coordinate space, due to stage up axis the points may flip sign
+		// (e.g. the USD max ends up at UE's negative Y axis), so we need to compute min/max in UE space again
+		FUsdStageInfo StageInfo{Prim.GetStage()};
+		FVector UESpaceUSDMin = UsdToUnreal::ConvertVector(StageInfo, Box.GetMin());
+		FVector UESpaceUSDMax = UsdToUnreal::ConvertVector(StageInfo, Box.GetMax());
+		FBox UEBox{
+			TArray<FVector>{UESpaceUSDMin, UESpaceUSDMax}
+		 };
+
+		BoundsComponent->SetBoundsMin(UEBox.Min);
+		BoundsComponent->SetBoundsMax(UEBox.Max);
+	}
+
+	if (pxr::UsdGeomModelAPI GeomModelAPI{Prim})
+	{
+		pxr::GfVec3f Color;
+		pxr::UsdAttribute ColorAttr = GeomModelAPI.GetModelDrawModeColorAttr();
+		if (ColorAttr && ColorAttr.Get(&Color))
+		{
+			BoundsComponent->SetBoundsColor(UsdToUnreal::ConvertColor(Color));
+		}
+
+		// It's not super efficient to call this whole function every time but not only this lets us
+		// resolve inherited values for the draw mode, but also lets us have our logic in one place.
+		// That is useful because GetAppliedDrawMode is used to decide which component to spawn, and
+		// it will return Default even when we have a particular draw mode specified (in case e.g.
+		// the prim is not a model, or doesn't have applyDrawMode enabled, etc.) so we can't just check
+		// the drawMode attribute directly here
+		EUsdDrawMode DrawMode = UsdUtils::GetAppliedDrawMode(Prim);
+		BoundsComponent->SetDrawMode(DrawMode);
+
+		pxr::TfToken CardGeometryToken;
+		pxr::UsdAttribute GeometryAttr = GeomModelAPI.GetModelCardGeometryAttr();
+		if (GeometryAttr && GeometryAttr.Get(&CardGeometryToken))
+		{
+			EUsdModelCardGeometry CardGeometry = EUsdModelCardGeometry::Cross;
+			if (CardGeometryToken == pxr::UsdGeomTokens->box)
+			{
+				CardGeometry = EUsdModelCardGeometry::Box;
+			}
+			else if (CardGeometryToken == pxr::UsdGeomTokens->fromTexture)
+			{
+				CardGeometry = EUsdModelCardGeometry::FromTexture;
+			}
+			BoundsComponent->SetCardGeometry(CardGeometry);
+		}
+	}
+
+	return true;
+}
+
 bool UnrealToUsd::ConvertCameraComponent( const UCineCameraComponent& CameraComponent, pxr::UsdPrim& Prim, double UsdTimeCode )
 {
 	FScopedUsdAllocs UsdAllocs;
@@ -1837,6 +2102,166 @@ bool UnrealToUsd::ConvertColorTrack( const UMovieSceneColorTrack& MovieSceneTrac
 	return true;
 }
 
+bool UnrealToUsd::ConvertBoundsVectorTracks(
+	const UMovieSceneDoubleVectorTrack* MinTrack,
+	const UMovieSceneDoubleVectorTrack* MaxTrack,
+	const FMovieSceneSequenceTransform& SequenceTransform,
+	const TFunction<void(const FVector&, const FVector&, double)>& WriterFunc,
+	UE::FUsdPrim& Prim
+)
+{
+	if (!WriterFunc || !Prim)
+	{
+		return false;
+	}
+
+	if (!MinTrack && !MaxTrack)
+	{
+		return false;
+	}
+
+	const UMovieSceneDoubleVectorTrack* MainTrack = MinTrack ? MinTrack : MaxTrack;
+
+	UMovieScene* MovieScene = MainTrack->GetTypedOuter<UMovieScene>();
+	if (!MovieScene)
+	{
+		return false;
+	}
+
+	UE::FUsdStage Stage = Prim.GetStage();
+
+	ERichCurveInterpMode StageInterpMode;
+	{
+		FScopedUsdAllocs Allocs;
+		pxr::UsdStageRefPtr UsdStage{Stage};
+		StageInterpMode = (UsdStage->GetInterpolationType() == pxr::UsdInterpolationTypeLinear)
+							  ? ERichCurveInterpMode::RCIM_Linear
+							  : ERichCurveInterpMode::RCIM_Constant;
+	}
+
+	const FFrameRate Resolution = MovieScene->GetTickResolution();
+	const FFrameRate DisplayRate = MovieScene->GetDisplayRate();
+
+	const double StageTimeCodesPerSecond = Stage.GetTimeCodesPerSecond();
+	const FFrameRate StageFrameRate(StageTimeCodesPerSecond, 1);
+
+	TFunction<void(const FMovieSceneDoubleChannel&, TSet<FFrameNumber>&)> AppendChannelBakeTimes =
+		[&Resolution, &DisplayRate, StageInterpMode](const FMovieSceneDoubleChannel& Channel, TSet<FFrameNumber>& OutBakeTimes)
+	{
+		const FFrameTime BakeInterval = FFrameRate::TransformTime(1, DisplayRate, Resolution);
+
+		TMovieSceneChannelData<const FMovieSceneDoubleValue> ChannelData = Channel.GetData();
+		TArrayView<const FFrameNumber> KeyTimes = ChannelData.GetTimes();
+		TArrayView<const FMovieSceneDoubleValue> KeyValues = ChannelData.GetValues();
+
+		for (int32 KeyIndex = 0; KeyIndex < KeyTimes.Num(); ++KeyIndex)
+		{
+			const FFrameNumber KeyTime = KeyTimes[KeyIndex];
+			const FMovieSceneDoubleValue& KeyValue = KeyValues[KeyIndex];
+
+			// If the channel has the same interpolation type as the stage (or we're the last key),
+			// we don't need to bake anything: Just write out the keyframe as is
+			if (KeyValue.InterpMode == StageInterpMode || KeyIndex == (KeyTimes.Num() - 1))
+			{
+				FFrameNumber SnappedKeyTime = FFrameRate::Snap(KeyTime, Resolution, DisplayRate).FloorToFrame();
+				OutBakeTimes.Emplace(SnappedKeyTime);
+			}
+			// We need to bake: Start from this key up until the next key (non-inclusive). We always want to put a keyframe at
+			// KeyTime, but then snap the other ones to the stage framerate
+			else
+			{
+				// Don't use the snapped key time for the end of the bake range, because if the snapping moves it
+				// later we may end up stepping back again when it's time to bake from that key onwards
+				const FFrameNumber NextKey = KeyTimes[KeyIndex + 1];
+				const FFrameTime NextKeyTime{NextKey};
+
+				for (FFrameTime EvalTime = KeyTime; EvalTime < NextKeyTime; EvalTime += BakeInterval)
+				{
+					FFrameNumber BakedKeyTime = FFrameRate::Snap(EvalTime, Resolution, DisplayRate).FloorToFrame();
+					OutBakeTimes.Emplace(BakedKeyTime);
+				}
+			}
+		}
+	};
+
+	TArray<UMovieSceneSection*> MinSections = MinTrack ? MinTrack->GetAllSections() : TArray<UMovieSceneSection*>{};
+	TArray<UMovieSceneSection*> MaxSections = MaxTrack ? MaxTrack->GetAllSections() : TArray<UMovieSceneSection*>{};
+
+	// Collect all time samples to bake with
+	TSet<FFrameNumber> AllBakeTimes;
+	for (UMovieSceneSection* Section : MinSections)
+	{
+		if (UMovieSceneDoubleVectorSection* VectorSection = Cast<UMovieSceneDoubleVectorSection>(Section))
+		{
+			for (int32 ChannelIndex = 0; ChannelIndex < VectorSection->GetChannelsUsed(); ++ChannelIndex)
+			{
+				const FMovieSceneDoubleChannel& Channel = VectorSection->GetChannel(ChannelIndex);
+				AppendChannelBakeTimes(Channel, AllBakeTimes);
+			}
+		}
+	}
+	for (UMovieSceneSection* Section : MaxSections)
+	{
+		if (UMovieSceneDoubleVectorSection* VectorSection = Cast<UMovieSceneDoubleVectorSection>(Section))
+		{
+			for (int32 ChannelIndex = 0; ChannelIndex < VectorSection->GetChannelsUsed(); ++ChannelIndex)
+			{
+				const FMovieSceneDoubleChannel& Channel = VectorSection->GetChannel(ChannelIndex);
+				AppendChannelBakeTimes(Channel, AllBakeTimes);
+			}
+		}
+	}
+
+	TArray<FFrameNumber> BakeTimeUnion = AllBakeTimes.Array();
+	BakeTimeUnion.Sort();
+
+	// Sample all channels at the union of bake times, construct the value and write it out
+	// This could be done more efficently, but in the general case we're only going to have one
+	// section per track anyway so it shouldn't matter much
+	for (const FFrameNumber UntransformedBakeTime : BakeTimeUnion)
+	{
+		FVector MinValue{0};
+		for (const UMovieSceneSection* Section : MinSections)
+		{
+			const UMovieSceneDoubleVectorSection* VectorSection = Cast<UMovieSceneDoubleVectorSection>(Section);
+			if (!VectorSection->IsTimeWithinSection(UntransformedBakeTime))
+			{
+				continue;
+			}
+
+			for (int32 ChannelIndex = 0; ChannelIndex < VectorSection->GetChannelsUsed(); ++ChannelIndex)
+			{
+				const FMovieSceneDoubleChannel& Channel = VectorSection->GetChannel(ChannelIndex);
+				Channel.Evaluate(UntransformedBakeTime, MinValue[ChannelIndex]);
+			}
+		}
+
+		FVector MaxValue{0};
+		for (const UMovieSceneSection* Section : MaxSections)
+		{
+			const UMovieSceneDoubleVectorSection* VectorSection = Cast<UMovieSceneDoubleVectorSection>(Section);
+			if (!VectorSection->IsTimeWithinSection(UntransformedBakeTime))
+			{
+				continue;
+			}
+
+			for (int32 ChannelIndex = 0; ChannelIndex < VectorSection->GetChannelsUsed(); ++ChannelIndex)
+			{
+				const FMovieSceneDoubleChannel& Channel = VectorSection->GetChannel(ChannelIndex);
+				Channel.Evaluate(UntransformedBakeTime, MaxValue[ChannelIndex]);
+			}
+		}
+
+		FFrameTime TransformedBakedKeyTime{UntransformedBakeTime};
+		TransformedBakedKeyTime *= SequenceTransform.InverseNoLooping();
+		FFrameTime UsdFrameTime = FFrameRate::TransformTime(TransformedBakedKeyTime, Resolution, StageFrameRate);
+
+		WriterFunc(MinValue, MaxValue, UsdFrameTime.AsDecimal());
+	}
+
+	return true;
+}
+
 bool UnrealToUsd::Convert3DTransformTrack( const UMovieScene3DTransformTrack& MovieSceneTrack, const FMovieSceneSequenceTransform& SequenceTransform, const TFunction<void( const FTransform&, double )>& WriterFunc, UE::FUsdPrim& Prim )
 {
 	if ( !WriterFunc || !Prim )
@@ -3021,6 +3446,23 @@ bool UnrealToUsd::CreateComponentPropertyBaker( UE::FUsdPrim& Prim, const UScene
 			}
 		}
 	}
+	else if (const UUsdDrawModeComponent* BoundsComponent = Cast<UUsdDrawModeComponent>(&Component))
+	{
+		static TSet<FString> RelevantProperties = {
+			GET_MEMBER_NAME_CHECKED(UUsdDrawModeComponent, BoundsMin).ToString(),
+			GET_MEMBER_NAME_CHECKED(UUsdDrawModeComponent, BoundsMax).ToString()
+		};
+
+		if (RelevantProperties.Contains(PropertyPath))
+		{
+			BakerType = EBakingType::Bounds;
+			BakerFunction = [UsdStage, BoundsComponent, UsdPrim](double UsdTimeCode) mutable
+			{
+				const bool bWriteExtents = true;
+				UnrealToUsd::ConvertBoundsComponent(*BoundsComponent, UsdPrim, bWriteExtents, UsdTimeCode);
+			};
+		}
+	}
 
 	if ( BakerType != EBakingType::None && BakerFunction)
 	{
@@ -3642,6 +4084,53 @@ UnrealToUsd::FPropertyTrackWriter UnrealToUsd::CreatePropertyTrackWriter( const
 				}
 			}
 		}
+		// Bounds component properties
+		else
+		{
+			Attr = {};
+			if (pxr::UsdGeomModelAPI GeomModelAPI{UsdPrim})
+			{
+				// For whatever reason there is no CreateExtentsHintAttr, so here we copy how the
+				// attr is created from within SetExtentsHint
+				Attr = UsdPrim.CreateAttribute(pxr::UsdGeomTokens->extentsHint, pxr::SdfValueTypeNames->Float3Array, /* custom = */ false);
+			}
+			else if (pxr::UsdGeomBoundable Boundable{UsdPrim})
+			{
+				Attr = Boundable.GetExtentAttr();
+			}
+
+			// If we still don't have an attr try applying the schema.
+			// Not entirely sure how this can possibly happen at this point, but this is more for "parity" with ConvertBoundsComponent
+			if (!Attr)
+			{
+				if (pxr::UsdGeomModelAPI GeomModelAPI = pxr::UsdGeomModelAPI::Apply(UsdPrim))
+				{
+					Attr = UsdPrim.CreateAttribute(pxr::UsdGeomTokens->extentsHint, pxr::SdfValueTypeNames->Float3Array, /* custom = */ false);
+				}
+			}
+
+			if (Attr)
+			{
+				Result.TwoVectorWriter = [Attr, StageInfo](const FVector& UEMinValue, const FVector& UEMaxValue, double UsdTimeCode)
+				{
+					FScopedUsdAllocs Allocs;
+
+					pxr::GfVec3f UEBoundsMinUsdSpace = UnrealToUsd::ConvertVector(StageInfo, UEMinValue);
+					pxr::GfVec3f UEBoundsMaxUsdSpace = UnrealToUsd::ConvertVector(StageInfo, UEMaxValue);
+					pxr::GfVec3f UsdMin{
+						FMath::Min(UEBoundsMinUsdSpace[0], UEBoundsMaxUsdSpace[0]),
+						FMath::Min(UEBoundsMinUsdSpace[1], UEBoundsMaxUsdSpace[1]),
+						FMath::Min(UEBoundsMinUsdSpace[2], UEBoundsMaxUsdSpace[2])};
+					pxr::GfVec3f UsdMax{
+						FMath::Max(UEBoundsMinUsdSpace[0], UEBoundsMaxUsdSpace[0]),
+						FMath::Max(UEBoundsMinUsdSpace[1], UEBoundsMaxUsdSpace[1]),
+						FMath::Max(UEBoundsMinUsdSpace[2], UEBoundsMaxUsdSpace[2])};
+					pxr::VtArray<pxr::GfVec3f> Extents{UsdMin, UsdMax};
+
+					Attr.Set(Extents, UsdTimeCode);
+				};
+			}
+		}
 	}
 
 	if ( Attr )
@@ -4045,7 +4534,210 @@ TArray<UE::FUsdAttribute> UnrealToUsd::GetAttributesForProperty( const UE::FUsdP
 		return { UE::FUsdAttribute{ UsdPrim.GetAttribute( pxr::UsdLuxTokens->inputsAngle ) } };
 	}
 
+	// Bounds component properties
+	// For now we only support bounds animations, but there's nothing preventing us from supporting color/texture/draw mode animations here,
+	// especially since our component works with a "dynamic mesh" approach that already rebuilds its proxy on-demand anyway
+	else if (
+		PropertyPath == GET_MEMBER_NAME_CHECKED(UUsdDrawModeComponent, BoundsMin) ||
+		PropertyPath == GET_MEMBER_NAME_CHECKED(UUsdDrawModeComponent, BoundsMax)
+	)
+	{
+		// If we have a model API, let's direct the caller to `extentsHint` instead, as that has priority over
+		// `extent` in case the model API is applied to a boundable and both are authored
+		if (pxr::UsdGeomModelAPI GeomModelAPI{UsdPrim})
+		{
+			return {UE::FUsdAttribute{GeomModelAPI.GetExtentsHintAttr()}};
+		}
+		else if (pxr::UsdGeomBoundable Boundable{UsdPrim})
+		{
+			return {UE::FUsdAttribute(Boundable.GetExtentAttr())};
+		}
+	}
+
 	return {};
 }
 
+bool UnrealToUsd::ConvertBoundsComponent(const UUsdDrawModeComponent& BoundsComponent, pxr::UsdPrim& UsdPrim, bool bWriteExtents, double UsdTimeCode)
+{
+	if(!UsdPrim)
+	{
+		return false;
+	}
+
+	FScopedUsdAllocs Allocs;
+
+	// If we have a bounds component on an opened stage then we know our prim must satisfy the requirements
+	// to have an alternate draw mode (being a "model", having the API schema, etc.).
+	// However, this function may be used to export imported components onto newly defined prims on a new stage though
+	// (e.g. during level export) so let's enforce those requirements anyway
+	pxr::UsdGeomModelAPI GeomModelAPI = pxr::UsdGeomModelAPI::Apply(UsdPrim);
+
+	// To have bounds, the prim must be a "model". To be a model the prim must have a model kind (i.e. basically any of the
+	// standard "kind"s set) and all of its ancestors up to the pseudoroot must have some "group" kind (like "assembly")
+	if (!UsdPrim.IsModel())
+	{
+		pxr::TfToken Kind = IUsdPrim::GetKind(UsdPrim);
+		if (Kind.IsEmpty())
+		{
+			IUsdPrim::SetKind(UsdPrim, pxr::KindTokens->component);
+		}
+
+		pxr::UsdPrim Parent = UsdPrim.GetParent();
+		while (Parent && !Parent.IsPseudoRoot() && !Parent.IsGroup())
+		{
+			pxr::TfToken ParentKind = IUsdPrim::GetKind(Parent);
+			if (ParentKind.IsEmpty())
+			{
+				IUsdPrim::SetKind(Parent, pxr::KindTokens->group);
+			}
+
+			Parent = Parent.GetParent();
+		}
+	}
+
+	// Author the actual extents.
+	// If we're a boundable, we can just author this as an `extent` opinion, but otherwise we'll need to author
+	// an `extentsHint` via the pxr::UsdGeomModelAPI. However, if a prim has *both* `extentsHint` and `extent`, USD
+	// will favor `extentsHint`, so in order to "affect the stage" we should author to `extentsHint` in that case
+	if (bWriteExtents)
+	{
+		FUsdStageInfo StageInfo{UsdPrim.GetStage()};
+
+		const FVector& UEBoundsMin = BoundsComponent.BoundsMin;
+		const FVector& UEBoundsMax = BoundsComponent.BoundsMax;
+		pxr::GfVec3f UEBoundsMinUsdSpace = UnrealToUsd::ConvertVector(StageInfo, UEBoundsMin);
+		pxr::GfVec3f UEBoundsMaxUsdSpace = UnrealToUsd::ConvertVector(StageInfo, UEBoundsMax);
+		pxr::GfVec3f UsdMin{
+			FMath::Min(UEBoundsMinUsdSpace[0], UEBoundsMaxUsdSpace[0]),
+			FMath::Min(UEBoundsMinUsdSpace[1], UEBoundsMaxUsdSpace[1]),
+			FMath::Min(UEBoundsMinUsdSpace[2], UEBoundsMaxUsdSpace[2])};
+		pxr::GfVec3f UsdMax{
+			FMath::Max(UEBoundsMinUsdSpace[0], UEBoundsMaxUsdSpace[0]),
+			FMath::Max(UEBoundsMinUsdSpace[1], UEBoundsMaxUsdSpace[1]),
+			FMath::Max(UEBoundsMinUsdSpace[2], UEBoundsMaxUsdSpace[2])};
+
+		pxr::VtArray<pxr::GfVec3f> Extents{UsdMin, UsdMax};
+
+		pxr::UsdAttribute ExtentsHintAttr = GeomModelAPI.GetExtentsHintAttr();
+		if (ExtentsHintAttr && ExtentsHintAttr.HasAuthoredValue())
+		{
+			ensure(GeomModelAPI.SetExtentsHint(Extents, UsdTimeCode));
+		}
+		else if (pxr::UsdGeomBoundable Boundable{UsdPrim})
+		{
+			if (pxr::UsdAttribute ExtentAttr = Boundable.CreateExtentAttr())
+			{
+				ExtentAttr.Set(Extents, UsdTimeCode);
+			}
+		}
+		else
+		{
+			ensure(GeomModelAPI.SetExtentsHint(Extents, UsdTimeCode));
+		}
+	}
+
+	if (pxr::UsdAttribute Attr = GeomModelAPI.CreateModelDrawModeAttr())
+	{
+		Attr.Set(
+			BoundsComponent.DrawMode == EUsdDrawMode::Origin	   ? pxr::UsdGeomTokens->origin
+			: BoundsComponent.DrawMode == EUsdDrawMode::Bounds	   ? pxr::UsdGeomTokens->bounds
+			: BoundsComponent.DrawMode == EUsdDrawMode::Cards	   ? pxr::UsdGeomTokens->cards
+			: BoundsComponent.DrawMode == EUsdDrawMode::Inherited  ? pxr::UsdGeomTokens->inherited
+																   : pxr::UsdGeomTokens->default_,
+			UsdTimeCode
+		);
+	}
+
+	if (pxr::UsdAttribute Attr = GeomModelAPI.CreateModelCardGeometryAttr())
+	{
+		Attr.Set(
+			BoundsComponent.CardGeometry == EUsdModelCardGeometry::Cross ? pxr::UsdGeomTokens->cross
+			: BoundsComponent.CardGeometry == EUsdModelCardGeometry::Box ? pxr::UsdGeomTokens->box
+																		 : pxr::UsdGeomTokens->fromTexture,
+			UsdTimeCode
+		);
+	}
+
+	// Technically we don't need this when we're making the prim into a component, but it's probably best to do it anyway
+	// for consistency, and it may be weird for the user if they happen to tweak the prim kind after this and have the bounds disappear
+	if (pxr::UsdAttribute Attr = GeomModelAPI.CreateModelApplyDrawModeAttr())
+	{
+		Attr.Set(true, UsdTimeCode);
+	}
+
+	if (pxr::UsdAttribute Attr = GeomModelAPI.CreateModelDrawModeColorAttr())
+	{
+		// This color is just a vec3f and our color is already linear anyway, so just convert it directly instead of going through
+		// UnrealToUsd::ConvertColor
+		pxr::GfVec3f UsdColor{BoundsComponent.BoundsColor.R, BoundsComponent.BoundsColor.G, BoundsComponent.BoundsColor.B};
+		Attr.Set(UsdColor, UsdTimeCode);
+	}
+
+	// Author the actual card face texture references.
+	// The logic surrounding when a texture is "authored" is complex (see https://openusd.org/release/api/class_usd_geom_model_a_p_i.html#details),
+	// and when reading we must differentiate between not having a texture because something wasn't authored, and not having a texture because
+	// something was authored but we failed to find it. This "AuthoredFaces" member is not exposed to blueprint/details panel though, and in general
+	// when setting/clearing textures via blueprint/details panels we will also tweak AuthoredFaces. The combined effect is that when a user sets a
+	// new texture in a property we will assume that means it became "authored", and when they clear a texture property we will assume that ceases to
+	// be "authored"
+	EUsdModelCardFace AuthoredFaces = BoundsComponent.GetAuthoredFaces();
+
+	using CreateAttrFuncType = decltype(&pxr::UsdGeomModelAPI::CreateModelCardTextureXPosAttr);
+	using GetAttrFuncType = decltype(&pxr::UsdGeomModelAPI::GetModelCardTextureXPosAttr);
+
+	TFunction<void(EUsdModelCardFace, GetAttrFuncType, CreateAttrFuncType)> ExportCardFace =
+		[&BoundsComponent, AuthoredFaces, UsdTimeCode, &GeomModelAPI](EUsdModelCardFace Face, GetAttrFuncType GetAttr, CreateAttrFuncType CreateAttr)
+	{
+		UTexture2D* FaceTexture = BoundsComponent.GetTextureForFace(Face);
+
+		if (EnumHasAllFlags(AuthoredFaces, Face))
+		{
+			const pxr::VtValue DefaultValue;
+			const bool bWriteSparsely = false;
+			if (pxr::UsdAttribute Attr = (GeomModelAPI.*CreateAttr)(DefaultValue, bWriteSparsely))
+			{
+				// We have an existing texture and it's marked as "authored". Try exporting a path to it
+				if (FaceTexture)
+				{
+#if WITH_EDITOR
+					if (FaceTexture->AssetImportData)
+					{
+						// If this texture doesn't have an asset import data, or the filename stored there is not valid then this reference
+						// won't resolve when it read it back from USD...
+						// TOD: Should we export the texture now? To where?
+						FString TextureSourcePath = FaceTexture->AssetImportData->GetFirstFilename();
+						pxr::SdfAssetPath AssetPath = pxr::SdfAssetPath{UnrealToUsd::ConvertString(*TextureSourcePath).Get()};
+						Attr.Set(AssetPath, UsdTimeCode);
+					}
+#endif	  // WITH_EDITOR
+				}
+				else
+				{
+					// This face is marked as "authored" and yet we have no texture for it. The only reason it could end up this way is if the source
+					// prim for this component originally had an authored texture there that didn't resolve when we parsed it, so let's just leave it
+					// alone so that it can still do that if we reload
+				}
+			}
+		}
+		// Face is not marked as authored, let's actively clear our opinion for this attribute on the current edit target and time code
+		else
+		{
+			pxr::UsdAttribute Attr = (GeomModelAPI.*GetAttr)();
+			if (Attr && Attr.HasAuthoredValue())
+			{
+				// This is capable of clearing the default opinion too
+				Attr.ClearAtTime(UsdTimeCode);
+			}
+		}
+	};
+	ExportCardFace(EUsdModelCardFace::XPos, &pxr::UsdGeomModelAPI::GetModelCardTextureXPosAttr, &pxr::UsdGeomModelAPI::CreateModelCardTextureXPosAttr);
+	ExportCardFace(EUsdModelCardFace::YPos, &pxr::UsdGeomModelAPI::GetModelCardTextureYPosAttr, &pxr::UsdGeomModelAPI::CreateModelCardTextureYPosAttr);
+	ExportCardFace(EUsdModelCardFace::ZPos, &pxr::UsdGeomModelAPI::GetModelCardTextureZPosAttr, &pxr::UsdGeomModelAPI::CreateModelCardTextureZPosAttr);
+	ExportCardFace(EUsdModelCardFace::XNeg, &pxr::UsdGeomModelAPI::GetModelCardTextureXNegAttr, &pxr::UsdGeomModelAPI::CreateModelCardTextureXNegAttr);
+	ExportCardFace(EUsdModelCardFace::YNeg, &pxr::UsdGeomModelAPI::GetModelCardTextureYNegAttr, &pxr::UsdGeomModelAPI::CreateModelCardTextureYNegAttr);
+	ExportCardFace(EUsdModelCardFace::ZNeg, &pxr::UsdGeomModelAPI::GetModelCardTextureZNegAttr, &pxr::UsdGeomModelAPI::CreateModelCardTextureZNegAttr);
+
+	return true;
+}
+
 #endif // #if USE_USD_SDK
diff --git a/Engine/Plugins/Importers/USDImporter/Source/USDUtilities/Private/USDTypesConversion.cpp b/Engine/Plugins/Importers/USDImporter/Source/USDUtilities/Private/USDTypesConversion.cpp
--- a/Engine/Plugins/Importers/USDImporter/Source/USDUtilities/Private/USDTypesConversion.cpp
+++ b/Engine/Plugins/Importers/USDImporter/Source/USDUtilities/Private/USDTypesConversion.cpp
@@ -142,6 +142,11 @@ namespace UsdToUnreal
 		return FVector( InValue[0], InValue[1], InValue[2] );
 	}
 
+	FVector ConvertVector(const pxr::GfVec3d& InValue)
+	{
+		return FVector(InValue[0], InValue[1], InValue[2]);
+	}
+
 	FVector ConvertVector( const FUsdStageInfo& StageInfo, const pxr::GfVec3f& InValue )
 	{
 		FVector Value = ConvertVector( InValue );
@@ -166,6 +171,30 @@ namespace UsdToUnreal
 		return Value;
 	}
 
+	FVector ConvertVector(const FUsdStageInfo& StageInfo, const pxr::GfVec3d& InValue)
+	{
+		FVector Value = ConvertVector(InValue);
+
+		const float UEMetersPerUnit = 0.01f;
+		if (!FMath::IsNearlyEqual(StageInfo.MetersPerUnit, UEMetersPerUnit))
+		{
+			Value *= StageInfo.MetersPerUnit / UEMetersPerUnit;
+		}
+
+		const bool bIsZUp = (StageInfo.UpAxis == EUsdUpAxis::ZAxis);
+
+		if (bIsZUp)
+		{
+			Value.Y = -Value.Y;
+		}
+		else
+		{
+			Swap(Value.Y, Value.Z);
+		}
+
+		return Value;
+	}
+
 	FMatrix ConvertMatrix( const pxr::GfMatrix4d& Matrix )
 	{
 		FMatrix UnrealMatrix(
diff --git a/Engine/Plugins/Importers/USDImporter/Source/USDUtilities/Public/USDConversionUtils.h b/Engine/Plugins/Importers/USDImporter/Source/USDUtilities/Public/USDConversionUtils.h
--- a/Engine/Plugins/Importers/USDImporter/Source/USDUtilities/Public/USDConversionUtils.h
+++ b/Engine/Plugins/Importers/USDImporter/Source/USDUtilities/Public/USDConversionUtils.h
@@ -47,6 +47,7 @@ PXR_NAMESPACE_CLOSE_SCOPE
 class USceneComponent;
 class UAssetImportData;
 class UUsdAssetImportData;
+enum class EUsdDrawMode : int32;
 enum class EUsdDuplicateType : uint8;
 enum class EUsdUpAxis : uint8;
 struct FUsdUnrealAssetInfo;
@@ -239,6 +240,32 @@ namespace UsdUtils
 	USDUTILITIES_API bool IsAnimated( const pxr::UsdPrim& Prim );
 	USDUTILITIES_API bool HasAnimatedVisibility( const pxr::UsdPrim& Prim );
 
+	/**
+	 * Returns true if Prim should have animated bounds (either from having authored animated `extent` or `extentsHint`
+	 * attributes, or potentially from having computed bounds and child prims that should make those computed bounds animated,
+	 * like having animated child Xforms)
+	 */
+	USDUTILITIES_API bool HasAnimatedBounds(
+		const pxr::UsdPrim& Prim,
+		EUsdPurpose IncludedPurposes = EUsdPurpose::Proxy | EUsdPurpose::Render,
+		bool bUseExtentsHint = true,
+		bool bIgnoreVisibility = false
+	);
+
+	/**
+	 * Similar to HasAnimatedBounds, except that this will also collect all the timeSamples that are relevant for sampling
+	 * animated bounds. For example this could just be the timeSamples for an authored `extentsHint` or `extent` animation, but
+	 * this could instead be the timeSamples of child Xforms transform animations instead.
+	 * Note that OutTimeSamples may end up with duplicate timeSamples, but those will always be at least sorted.
+	 */
+	USDUTILITIES_API bool GetAnimatedBoundsTimeSamples(
+		const pxr::UsdPrim& InPrim,
+		TArray<double>& OutTimeSamples,
+		EUsdPurpose InIncludedPurposes = EUsdPurpose::Proxy | EUsdPurpose::Render,
+		bool bInUseExtentsHint = true,
+		bool bInIgnoreVisibility = false
+	);
+
 	/** Returns whether Prim belongs to any of the default kinds, or a kind derived from them. The result can be a union of different kinds. */
 	USDUTILITIES_API EUsdDefaultKind GetDefaultKind( const pxr::UsdPrim& Prim );
 
@@ -251,6 +278,12 @@ namespace UsdUtils
 	 */
 	USDUTILITIES_API bool SetDefaultKind( pxr::UsdPrim& Prim, EUsdDefaultKind NewKind );
 
+	/**
+	 * Returns whether the prim has the UsdGeomModelAPI schema and should be drawn with one of the alternative draw modes, such as cards or bounds.
+	 * Will return EUsdDrawMode::Default in case the prim should be drawn as usual instead, or in case of error.
+	 */
+	USDUTILITIES_API EUsdDrawMode GetAppliedDrawMode(const pxr::UsdPrim& Prim);
+
 	/**
 	 * Returns all prims of type SchemaType (or a descendant type) in the subtree of prims rooted at StartPrim.
 	 * Stops going down the subtrees when it hits a schema type to exclude.
diff --git a/Engine/Plugins/Importers/USDImporter/Source/USDUtilities/Public/USDGeomMeshConversion.h b/Engine/Plugins/Importers/USDImporter/Source/USDUtilities/Public/USDGeomMeshConversion.h
--- a/Engine/Plugins/Importers/USDImporter/Source/USDUtilities/Public/USDGeomMeshConversion.h
+++ b/Engine/Plugins/Importers/USDImporter/Source/USDUtilities/Public/USDGeomMeshConversion.h
@@ -4,11 +4,12 @@
 
 #include "CoreMinimal.h"
 
-#if USE_USD_SDK
-
 #include "UnrealUSDWrapper.h"
+#include "USDClassesModule.h"  // Include this for now so that we can redirect FDisplayColorMaterial to the USDClasses definition
 #include "UsdWrappers/ForwardDeclarations.h"
 
+#if USE_USD_SDK
+
 #include "USDIncludesStart.h"
 	#include "pxr/usd/usd/timeCode.h"
 	#include "pxr/usd/usdShade/tokens.h"
@@ -33,7 +34,6 @@ struct FUsdCollapsingCache;
 struct FUsdStageInfo;
 namespace UsdUtils
 {
-	struct FDisplayColorMaterial;
 	struct FUsdPrimMaterialAssignmentInfo;
 }
 namespace UE
@@ -221,15 +221,7 @@ namespace UnrealToUsd
 
 namespace UsdUtils
 {
-	/** Describes the type of vertex color/DisplayColor material that we would need in order to render a prim's displayColor data as intended */
-	struct USDUTILITIES_API FDisplayColorMaterial
-	{
-		bool bHasOpacity = false;
-		bool bIsDoubleSided = false;
-
-		FString ToString();
-		static TOptional<FDisplayColorMaterial> FromString(const FString& DisplayColorString);
-	};
+	using FDisplayColorMaterial = IUsdClassesModule::FDisplayColorMaterial;
 
 	/** Describes what type of material assignment a FUsdPrimMaterialSlot represents */
 	enum class EPrimAssignmentType : uint8
@@ -306,8 +298,12 @@ namespace UsdUtils
 	);
 
 	/** Creates a dynamic material instance using the right reference material depending on the given description */
+	PRAGMA_DISABLE_DEPRECATION_WARNINGS
+	UE_DEPRECATED(5.3, "This has moved to the IUsdClassesModule namespace, from USDClassesModule.h")
 	USDUTILITIES_API UMaterialInstanceDynamic* CreateDisplayColorMaterialInstanceDynamic( const UsdUtils::FDisplayColorMaterial& DisplayColorDescription );
+	UE_DEPRECATED(5.3, "This has moved to the IUsdClassesModule namespace, from USDClassesModule.h")
 	USDUTILITIES_API UMaterialInstanceConstant* CreateDisplayColorMaterialInstanceConstant( const UsdUtils::FDisplayColorMaterial& DisplayColorDescription );
+	PRAGMA_ENABLE_DEPRECATION_WARNINGS
 
 	/**
 	 * Extracts all material assignment data from UsdPrim, including material binding, multiple assignment with GeomSubsets, and the unrealMaterial custom USD attribute.
diff --git a/Engine/Plugins/Importers/USDImporter/Source/USDUtilities/Public/USDPrimConversion.h b/Engine/Plugins/Importers/USDImporter/Source/USDUtilities/Public/USDPrimConversion.h
--- a/Engine/Plugins/Importers/USDImporter/Source/USDUtilities/Public/USDPrimConversion.h
+++ b/Engine/Plugins/Importers/USDImporter/Source/USDUtilities/Public/USDPrimConversion.h
@@ -14,6 +14,7 @@
 
 PXR_NAMESPACE_OPEN_SCOPE
 	class UsdAttribute;
+	class UsdGeomBBoxCache;
 	class UsdGeomCamera;
 	class UsdGeomXformable;
 	class UsdPrim;
@@ -35,6 +36,7 @@ class UMovieScene;
 class UMovieScene3DTransformTrack;
 class UMovieSceneBoolTrack;
 class UMovieSceneColorTrack;
+class UMovieSceneDoubleVectorTrack;
 class UMovieSceneFloatTrack;
 class UMovieScenePropertyTrack;
 class UMovieSceneSkeletalAnimationTrack;
@@ -42,12 +44,15 @@ class UMovieSceneTrack;
 class UMovieSceneVisibilityTrack;
 class USceneComponent;
 class USkeletalMeshComponent;
+class UUsdAssetCache2;
+class UUsdDrawModeComponent;
 struct FFrameRate;
 struct FMovieSceneSequenceTransform;
 struct FUsdStageInfo;
 namespace UE
 {
 	class FUsdAttribute;
+	class FUsdGeomBBoxCache;
 }
 
 namespace UsdToUnreal
@@ -89,12 +94,21 @@ namespace UsdToUnreal
 	/**
 	 * Calls the ReaderFunc on each time sample of UsdTimeSamples in order to bake values into MovieSceneTrack.
 	 * This is mostly used when reading attributes from USD into tracks for the automatically generated ULevelSequence provided with AUsdStageActors.
+	 * The bounds version must use a BBoxCache: If one is not provided it will be created on-demand for this call alone.
 	 */
 	USDUTILITIES_API bool ConvertBoolTimeSamples( const UE::FUsdStage& UsdStage, const TArray<double>& UsdTimeSamples, const TFunction<bool( double )>& ReaderFunc, UMovieSceneBoolTrack& MovieSceneTrack, const FMovieSceneSequenceTransform& SequenceTransform );
 	USDUTILITIES_API bool ConvertBoolTimeSamples( const UE::FUsdStage& UsdStage, const TArray<double>& UsdTimeSamples, const TFunction<bool(double)>& ReaderFunc, UMovieSceneVisibilityTrack& MovieSceneTrack, const FMovieSceneSequenceTransform& SequenceTransform );
 	USDUTILITIES_API bool ConvertFloatTimeSamples( const UE::FUsdStage& UsdStage, const TArray<double>& UsdTimeSamples, const TFunction<float( double )>& ReaderFunc, UMovieSceneFloatTrack& MovieSceneTrack, const FMovieSceneSequenceTransform& SequenceTransform );
 	USDUTILITIES_API bool ConvertColorTimeSamples( const UE::FUsdStage& UsdStage, const TArray<double>& UsdTimeSamples, const TFunction<FLinearColor( double )>& ReaderFunc, UMovieSceneColorTrack& MovieSceneTrack, const FMovieSceneSequenceTransform& SequenceTransform );
 	USDUTILITIES_API bool ConvertTransformTimeSamples( const UE::FUsdStage& UsdStage, const TArray<double>& UsdTimeSamples, const TFunction<FTransform( double )>& ReaderFunc, UMovieScene3DTransformTrack& MovieSceneTrack, const FMovieSceneSequenceTransform& SequenceTransform );
+	USDUTILITIES_API bool ConvertBoundsTimeSamples(
+		const UE::FUsdPrim& InPrim,
+		const TArray<double>& InUsdTimeSamples,
+		const FMovieSceneSequenceTransform& InSequenceTransform,
+		UMovieSceneDoubleVectorTrack& InOutMinTrack,
+		UMovieSceneDoubleVectorTrack& InOutMaxTrack,
+		UE::FUsdGeomBBoxCache* InOutBBoxCache = nullptr
+	);
 
 	/**
 	 * Struct with lambda functions that can be used to sample an FUsdPrim's attributes at a provided UsdTimeCode and return a converted result.
@@ -117,6 +131,18 @@ namespace UsdToUnreal
 	 * When we generate a TransformReader, if bIgnorePrimLocalTransform is true it will cause it to ignore that prim's local transform for the reader.
 	 */
 	USDUTILITIES_API FPropertyTrackReader CreatePropertyTrackReader( const UE::FUsdPrim& Prim, const FName& PropertyPath, bool bIgnorePrimLocalTransform = false );
+
+	/**
+	 * Convert a prim with UsdGeomModelAPI and an alternative drawMode (bounds, cards or origin) into property values of an
+	 * UUsdDrawModeComponent at EvalTime.
+	 * Must use a BBoxCache: If one is not provided it will be created on-demand for this call alone.
+	 */
+	USDUTILITIES_API bool ConvertBounds(
+		const pxr::UsdPrim& Prim,
+		UUsdDrawModeComponent* BoundsComponent,
+		double EvalTime = UsdUtils::GetDefaultTimeCode(),
+		pxr::UsdGeomBBoxCache* BBoxCache = nullptr
+	);
 }
 
 namespace UnrealToUsd
@@ -130,6 +156,13 @@ namespace UnrealToUsd
 	USDUTILITIES_API bool ConvertFloatTrack( const UMovieSceneFloatTrack& MovieSceneTrack, const FMovieSceneSequenceTransform& SequenceTransform, const TFunction<void( float, double )>& WriterFunc, UE::FUsdPrim& Prim );
 	USDUTILITIES_API bool ConvertBoolTrack( const UMovieSceneBoolTrack& MovieSceneTrack, const FMovieSceneSequenceTransform& SequenceTransform, const TFunction<void( bool, double )>& WriterFunc, UE::FUsdPrim& Prim );
 	USDUTILITIES_API bool ConvertColorTrack( const UMovieSceneColorTrack& MovieSceneTrack, const FMovieSceneSequenceTransform& SequenceTransform, const TFunction<void( const FLinearColor&, double )>& WriterFunc, UE::FUsdPrim& Prim );
+	USDUTILITIES_API bool ConvertBoundsVectorTracks(
+		const UMovieSceneDoubleVectorTrack* MinTrack,
+		const UMovieSceneDoubleVectorTrack* MaxTrack,
+		const FMovieSceneSequenceTransform& SequenceTransform,
+		const TFunction<void(const FVector&, const FVector&, double)>& WriterFunc,
+		UE::FUsdPrim& Prim
+	);
 	USDUTILITIES_API bool Convert3DTransformTrack( const UMovieScene3DTransformTrack& MovieSceneTrack, const FMovieSceneSequenceTransform& SequenceTransform, const TFunction<void( const FTransform&, double )>& WriterFunc, UE::FUsdPrim& Prim );
 
 	USDUTILITIES_API bool ConvertSceneComponent( const pxr::UsdStageRefPtr& Stage, const USceneComponent* SceneComponent, pxr::UsdPrim& UsdPrim );
@@ -193,6 +226,7 @@ namespace UnrealToUsd
 		Camera = 4,
 		Light = 8,
 		Skeletal = 16,
+		Bounds = 32,
 	};
 	ENUM_CLASS_FLAGS( EBakingType );
 
@@ -232,6 +266,7 @@ namespace UnrealToUsd
 		TFunction<void( float, double )> FloatWriter;
 		TFunction<void( bool, double )> BoolWriter;
 		TFunction<void( const FLinearColor&, double )> ColorWriter;
+		TFunction<void( const FVector&, const FVector&, double )> TwoVectorWriter;  // Just used for the bounds tracks, where we use a vector for min and another for max
 		TFunction<void( const FTransform&, double )> TransformWriter;
 	};
 
@@ -249,6 +284,18 @@ namespace UnrealToUsd
 	 * we'll receive intensity, exposure, width and height, in that order.
 	 */
 	USDUTILITIES_API TArray<UE::FUsdAttribute> GetAttributesForProperty( const UE::FUsdPrim& Prim, const FName& PropertyPath );
+
+	/**
+	 * Converts the properties values from UUsdDrawModeComponent into attribute values of a prim with UsdGeomModelAPI.
+	 * Note that this will even apply the schema itself to UsdPrim if it doesn't already have it, and potentially author
+	 * the 'kind' metadata on it and all of its ancestors (as UsdGeomModelAPI requires the prim to be a "model" to have function)
+	 */
+	USDUTILITIES_API bool ConvertBoundsComponent(
+		const UUsdDrawModeComponent& BoundsComponent,
+		pxr::UsdPrim& UsdPrim,
+		bool bWriteExtents = false,
+		double UsdTimeCode = UsdUtils::GetDefaultTimeCode()
+	);
 }
 
 #endif // #if USE_USD_SDK
diff --git a/Engine/Plugins/Importers/USDImporter/Source/USDUtilities/Public/USDTypesConversion.h b/Engine/Plugins/Importers/USDImporter/Source/USDUtilities/Public/USDTypesConversion.h
--- a/Engine/Plugins/Importers/USDImporter/Source/USDUtilities/Public/USDTypesConversion.h
+++ b/Engine/Plugins/Importers/USDImporter/Source/USDUtilities/Public/USDTypesConversion.h
@@ -19,6 +19,7 @@ PXR_NAMESPACE_OPEN_SCOPE
 	class GfVec2f;
 	class GfVec3f;
 	class GfVec4f;
+	class GfVec3d;
 	class SdfPath;
 	class TfToken;
 
@@ -56,7 +57,9 @@ namespace UsdToUnreal
 
 	USDUTILITIES_API FVector2D ConvertVector( const pxr::GfVec2f& InValue );
 	USDUTILITIES_API FVector ConvertVector( const pxr::GfVec3f& InValue );
+	USDUTILITIES_API FVector ConvertVector(const pxr::GfVec3d& InValue);
 	USDUTILITIES_API FVector ConvertVector( const FUsdStageInfo& StageInfo, const pxr::GfVec3f& InValue );
+	USDUTILITIES_API FVector ConvertVector(const FUsdStageInfo& StageInfo, const pxr::GfVec3d& InValue);
 
 	USDUTILITIES_API FMatrix ConvertMatrix( const pxr::GfMatrix4d& Matrix );
 	USDUTILITIES_API FTransform ConvertMatrix( const FUsdStageInfo& StageInfo, const pxr::GfMatrix4d& InMatrix );
diff --git a/Engine/Plugins/Importers/USDImporter/Source/USDUtilities/USDUtilities.Build.cs b/Engine/Plugins/Importers/USDImporter/Source/USDUtilities/USDUtilities.Build.cs
--- a/Engine/Plugins/Importers/USDImporter/Source/USDUtilities/USDUtilities.Build.cs
+++ b/Engine/Plugins/Importers/USDImporter/Source/USDUtilities/USDUtilities.Build.cs
@@ -26,6 +26,7 @@ namespace UnrealBuildTool.Rules
 					"Core",
 					"CoreUObject",
 					"UnrealUSDWrapper",
+					"USDClasses", // So that consumers can also include IUsdClassesModule for the new definition of FDisplayColorMaterial
 				}
 			);
 
@@ -50,7 +51,6 @@ namespace UnrealBuildTool.Rules
 					"Slate",
 					"SlateCore",
 					"StaticMeshDescription",
-					"USDClasses",
 				}
 			);
 
diff --git a/Engine/Plugins/Importers/USDImporter/Source/UnrealUSDWrapper/Private/UnrealUSDWrapper.cpp b/Engine/Plugins/Importers/USDImporter/Source/UnrealUSDWrapper/Private/UnrealUSDWrapper.cpp
--- a/Engine/Plugins/Importers/USDImporter/Source/UnrealUSDWrapper/Private/UnrealUSDWrapper.cpp
+++ b/Engine/Plugins/Importers/USDImporter/Source/UnrealUSDWrapper/Private/UnrealUSDWrapper.cpp
@@ -445,6 +445,8 @@ FString UnrealIdentifiers::MaterialFullPurpose = ANSI_TO_TCHAR( pxr::UsdShadeTok
 FString UnrealIdentifiers::PrimvarsDisplayColor = ANSI_TO_TCHAR(pxr::UsdGeomTokens->primvarsDisplayColor.GetString().c_str());
 FString UnrealIdentifiers::PrimvarsDisplayOpacity = ANSI_TO_TCHAR(pxr::UsdGeomTokens->primvarsDisplayOpacity.GetString().c_str());
 FString UnrealIdentifiers::DoubleSided = ANSI_TO_TCHAR(pxr::UsdGeomTokens->doubleSided.GetString().c_str());
+FString UnrealIdentifiers::ModelDrawMode = ANSI_TO_TCHAR(pxr::UsdGeomTokens->modelDrawMode.GetString().c_str());
+FString UnrealIdentifiers::ModelApplyDrawMode = ANSI_TO_TCHAR(pxr::UsdGeomTokens->modelApplyDrawMode.GetString().c_str());
 #else
 FString UnrealIdentifiers::MaterialAllPurpose = TEXT( "" );
 FString UnrealIdentifiers::MaterialPreviewPurpose = TEXT( "preview" );
@@ -452,6 +454,8 @@ FString UnrealIdentifiers::MaterialFullPurpose = TEXT( "full" );
 FString UnrealIdentifiers::PrimvarsDisplayColor = TEXT("primvars:displayColor");
 FString UnrealIdentifiers::PrimvarsDisplayOpacity = TEXT("primvars:displayOpacity");
 FString UnrealIdentifiers::DoubleSided = TEXT("doubleSided");
+FString UnrealIdentifiers::ModelDrawMode = TEXT("model:drawMode");
+FString UnrealIdentifiers::ModelApplyDrawMode = TEXT("model:applyDrawMode");
 #endif // USE_USD_SDK
 
 FUsdDelegates::FUsdImportDelegate FUsdDelegates::OnPreUsdImport;
diff --git a/Engine/Plugins/Importers/USDImporter/Source/UnrealUSDWrapper/Private/UsdWrappers/UsdAttribute.cpp b/Engine/Plugins/Importers/USDImporter/Source/UnrealUSDWrapper/Private/UsdWrappers/UsdAttribute.cpp
--- a/Engine/Plugins/Importers/USDImporter/Source/UnrealUSDWrapper/Private/UsdWrappers/UsdAttribute.cpp
+++ b/Engine/Plugins/Importers/USDImporter/Source/UnrealUSDWrapper/Private/UsdWrappers/UsdAttribute.cpp
@@ -256,6 +256,15 @@ namespace UE
 #endif // #if USE_USD_SDK
 	}
 
+	bool FUsdAttribute::HasAuthoredValue() const
+	{
+#if USE_USD_SDK
+		return Impl->PxrUsdAttribute.Get().HasAuthoredValue();
+#else
+		return false;
+#endif // #if USE_USD_SDK
+	}
+
 	bool FUsdAttribute::HasFallbackValue() const
 	{
 #if USE_USD_SDK
diff --git a/Engine/Plugins/Importers/USDImporter/Source/UnrealUSDWrapper/Private/UsdWrappers/UsdGeomBBoxCache.cpp b/Engine/Plugins/Importers/USDImporter/Source/UnrealUSDWrapper/Private/UsdWrappers/UsdGeomBBoxCache.cpp
new file mode 100644
--- /dev/null
+++ b/Engine/Plugins/Importers/USDImporter/Source/UnrealUSDWrapper/Private/UsdWrappers/UsdGeomBBoxCache.cpp
@@ -0,0 +1,272 @@
+// Copyright Epic Games, Inc. All Rights Reserved.
+
+#include "UsdWrappers/UsdGeomBBoxCache.h"
+
+#include "USDMemory.h"
+
+#include "UsdWrappers/UsdPrim.h"
+
+#if USE_USD_SDK
+
+#include "USDIncludesStart.h"
+#include "pxr/usd/usdGeom/bboxCache.h"
+#include "USDIncludesEnd.h"
+
+#endif	  //  #if USE_USD_SDK
+
+namespace UE
+{
+	namespace Internal
+	{
+		class FUsdGeomBBoxCacheImpl
+		{
+		public:
+			FUsdGeomBBoxCacheImpl(double Time, EUsdPurpose IncludedPurposeFlags, bool bUseExtentsHint, bool bIgnoreVisibility)
+#if USE_USD_SDK
+				: PxrUsdGeomBBoxCache(Time, {}, bUseExtentsHint, bIgnoreVisibility)
+#endif	  // #if USE_USD_SDK
+			{
+#if USE_USD_SDK
+				std::vector<pxr::TfToken> TokenVector;
+				TokenVector.reserve(4);
+				TokenVector.push_back(pxr::UsdGeomTokens->default_);
+
+				if (EnumHasAllFlags(IncludedPurposeFlags, EUsdPurpose::Proxy))
+				{
+					TokenVector.push_back(pxr::UsdGeomTokens->proxy);
+				}
+
+				if (EnumHasAllFlags(IncludedPurposeFlags, EUsdPurpose::Render))
+				{
+					TokenVector.push_back(pxr::UsdGeomTokens->render);
+				}
+
+				if (EnumHasAllFlags(IncludedPurposeFlags, EUsdPurpose::Guide))
+				{
+					TokenVector.push_back(pxr::UsdGeomTokens->guide);
+				}
+
+				PxrUsdGeomBBoxCache.SetIncludedPurposes(TokenVector);
+#endif	  // #if USE_USD_SDK
+			}
+
+#if USE_USD_SDK
+			explicit FUsdGeomBBoxCacheImpl(const pxr::UsdGeomBBoxCache& InUsdGeomBBoxCache)
+				: PxrUsdGeomBBoxCache(InUsdGeomBBoxCache)
+			{
+			}
+
+			explicit FUsdGeomBBoxCacheImpl(pxr::UsdGeomBBoxCache&& InUsdGeomBBoxCache)
+				: PxrUsdGeomBBoxCache(MoveTemp(InUsdGeomBBoxCache))
+			{
+			}
+
+			pxr::UsdGeomBBoxCache PxrUsdGeomBBoxCache;
+#endif	  // #if USE_USD_SDK
+		};
+	}	  // namespace Internal
+
+	FUsdGeomBBoxCache::FUsdGeomBBoxCache(double Time, EUsdPurpose IncludedPurposeFlags, bool bUseExtentsHint, bool bIgnoreVisibility)
+	{
+		FScopedUsdAllocs UsdAllocs;
+		FWriteScopeLock ScopeLock{Lock};
+		Impl = MakeUnique<Internal::FUsdGeomBBoxCacheImpl>(Time, IncludedPurposeFlags, bUseExtentsHint, bIgnoreVisibility);
+	}
+
+	FUsdGeomBBoxCache::FUsdGeomBBoxCache(const FUsdGeomBBoxCache& Other)
+	{
+#if USE_USD_SDK
+		FScopedUsdAllocs UsdAllocs;
+		FWriteScopeLock ScopeLock{Lock};
+		Impl = MakeUnique<Internal::FUsdGeomBBoxCacheImpl>(Other.Impl->PxrUsdGeomBBoxCache);
+#endif	  // #if USE_USD_SDK
+	}
+
+	FUsdGeomBBoxCache::FUsdGeomBBoxCache(FUsdGeomBBoxCache&& Other)
+		: Impl(MakeUnique<Internal::FUsdGeomBBoxCacheImpl>(
+			/*Time*/ 0.0,
+			/*IncludedPurposeFlags*/ EUsdPurpose::Default,
+			/*bUseExtentsHint*/ false,
+			/*bIgnoreVisibility*/ false
+		))
+	{
+		if (this != &Other)
+		{
+			FWriteScopeLock ScopeLock{Lock};
+			FWriteScopeLock OtherScopeLock{Other.Lock};
+			Impl = MoveTemp(Other.Impl);
+		}
+	}
+
+	FUsdGeomBBoxCache& FUsdGeomBBoxCache::operator=(const FUsdGeomBBoxCache& Other)
+	{
+#if USE_USD_SDK
+		FScopedUsdAllocs UsdAllocs;
+		FWriteScopeLock ScopeLock{Lock};
+		Impl = MakeUnique<Internal::FUsdGeomBBoxCacheImpl>(Other.Impl->PxrUsdGeomBBoxCache);
+#endif	  // #if USE_USD_SDK
+
+		return *this;
+	}
+
+	FUsdGeomBBoxCache& FUsdGeomBBoxCache::operator=(FUsdGeomBBoxCache&& Other)
+	{
+		if (this != &Other)
+		{
+			FScopedUsdAllocs UsdAllocs;
+			FWriteScopeLock ScopeLock{Lock};
+			FWriteScopeLock OtherScopeLock{Other.Lock};
+			Impl = MoveTemp(Other.Impl);
+		}
+
+		return *this;
+	}
+
+	FUsdGeomBBoxCache::~FUsdGeomBBoxCache()
+	{
+		FScopedUsdAllocs UsdAllocs;
+		FWriteScopeLock ScopeLock{Lock};
+		Impl.Reset();
+	}
+
+#if USE_USD_SDK
+	FUsdGeomBBoxCache::FUsdGeomBBoxCache(const pxr::UsdGeomBBoxCache& InUsdGeomBBoxCache)
+	{
+		FScopedUsdAllocs UsdAllocs;
+		FWriteScopeLock ScopeLock{Lock};
+		Impl = MakeUnique<Internal::FUsdGeomBBoxCacheImpl>(InUsdGeomBBoxCache);
+	}
+
+	FUsdGeomBBoxCache::FUsdGeomBBoxCache(pxr::UsdGeomBBoxCache&& InUsdGeomBBoxCache)
+	{
+		FScopedUsdAllocs UsdAllocs;
+		FWriteScopeLock ScopeLock{Lock};
+		Impl = MakeUnique<Internal::FUsdGeomBBoxCacheImpl>(MoveTemp(InUsdGeomBBoxCache));
+	}
+
+	FUsdGeomBBoxCache::operator pxr::UsdGeomBBoxCache&()
+	{
+		// Our caller should have manually locked the mutex before calling this.
+		// Note we need to check for a read lock here: It should be considered a mistaked to only have a read lock acquired
+		// and call this function, as it provides mutable access to the underlying type.
+		// If we checked for a write lock here, our lock would have failed to acquire (given the outer read lock from the user)
+		// and we would have incorrectly assumed the user locked correctly. Checking for a read lock should succeed even in that
+		// case, which will let us detect the mistake and cause our ensure to correctly trigger
+		ensure(!Lock.TryReadLock());
+		return Impl->PxrUsdGeomBBoxCache;
+	}
+
+	FUsdGeomBBoxCache::operator const pxr::UsdGeomBBoxCache&() const
+	{
+		// Our caller should have manually locked the mutex before calling this
+		ensure(!Lock.TryReadLock());
+		return Impl->PxrUsdGeomBBoxCache;
+	}
+#endif	  // #if USE_USD_SDK
+
+	void FUsdGeomBBoxCache::Clear()
+	{
+#if USE_USD_SDK
+		FScopedUsdAllocs Allocs;
+		FWriteScopeLock ScopeLock{Lock};
+		Impl->PxrUsdGeomBBoxCache.Clear();
+#endif // #if USE_USD_SDK
+	}
+
+	void FUsdGeomBBoxCache::SetIncludedPurposes(EUsdPurpose IncludedPurposeFlags)
+	{
+#if USE_USD_SDK
+		FScopedUsdAllocs Allocs;
+		FWriteScopeLock ScopeLock{Lock};
+
+		std::vector<pxr::TfToken> TokenVector;
+		TokenVector.reserve(4);
+		TokenVector.push_back(pxr::UsdGeomTokens->default_);
+
+		if (EnumHasAllFlags(IncludedPurposeFlags, EUsdPurpose::Proxy))
+		{
+			TokenVector.push_back(pxr::UsdGeomTokens->proxy);
+		}
+
+		if (EnumHasAllFlags(IncludedPurposeFlags, EUsdPurpose::Render))
+		{
+			TokenVector.push_back(pxr::UsdGeomTokens->render);
+		}
+
+		if (EnumHasAllFlags(IncludedPurposeFlags, EUsdPurpose::Guide))
+		{
+			TokenVector.push_back(pxr::UsdGeomTokens->guide);
+		}
+
+		Impl->PxrUsdGeomBBoxCache.SetIncludedPurposes(TokenVector);
+#endif // #if USE_USD_SDK
+	}
+
+	EUsdPurpose FUsdGeomBBoxCache::GetIncludedPurposes() const
+	{
+#if USE_USD_SDK
+		FScopedUsdAllocs Allocs;
+		FReadScopeLock ScopeLock{Lock};
+
+		EUsdPurpose Flags = EUsdPurpose::Default;
+		for (const pxr::TfToken& Token : Impl->PxrUsdGeomBBoxCache.GetIncludedPurposes())
+		{
+			if (Token == pxr::UsdGeomTokens->proxy)
+			{
+				Flags |= EUsdPurpose::Proxy;
+			}
+			else if (Token == pxr::UsdGeomTokens->render)
+			{
+				Flags |= EUsdPurpose::Render;
+			}
+			else if (Token == pxr::UsdGeomTokens->guide)
+			{
+				Flags |= EUsdPurpose::Guide;
+			}
+		}
+		return Flags;
+#else
+		return EUsdPurpose::Default;
+#endif	  // #if USE_USD_SDK
+	}
+
+	bool FUsdGeomBBoxCache::GetUseExtentsHint() const
+	{
+#if USE_USD_SDK
+		FReadScopeLock ScopeLock{Lock};
+		return Impl->PxrUsdGeomBBoxCache.GetUseExtentsHint();
+#else
+		return false;
+#endif	  // #if USE_USD_SDK
+	}
+
+	bool FUsdGeomBBoxCache::GetIgnoreVisibility() const
+	{
+#if USE_USD_SDK
+		FReadScopeLock ScopeLock{Lock};
+		return Impl->PxrUsdGeomBBoxCache.GetIgnoreVisibility();
+#else
+		return false;
+#endif	  // #if USE_USD_SDK
+	}
+
+	void FUsdGeomBBoxCache::SetTime(double UsdTimeCode)
+	{
+#if USE_USD_SDK
+		FScopedUsdAllocs Allocs;
+		FWriteScopeLock ScopeLock{Lock};
+		Impl->PxrUsdGeomBBoxCache.SetTime(UsdTimeCode);
+#endif	  // #if USE_USD_SDK
+	}
+
+	double FUsdGeomBBoxCache::GetTime() const
+	{
+#if USE_USD_SDK
+		FReadScopeLock ScopeLock{Lock};
+		return Impl->PxrUsdGeomBBoxCache.GetTime().GetValue();
+#else
+		return 0;
+#endif	  // #if USE_USD_SDK
+	}
+
+}	 // namespace UE
diff --git a/Engine/Plugins/Importers/USDImporter/Source/UnrealUSDWrapper/Public/UnrealUSDWrapper.h b/Engine/Plugins/Importers/USDImporter/Source/UnrealUSDWrapper/Public/UnrealUSDWrapper.h
--- a/Engine/Plugins/Importers/USDImporter/Source/UnrealUSDWrapper/Public/UnrealUSDWrapper.h
+++ b/Engine/Plugins/Importers/USDImporter/Source/UnrealUSDWrapper/Public/UnrealUSDWrapper.h
@@ -453,6 +453,10 @@ namespace UnrealIdentifiers
 	extern UNREALUSDWRAPPER_API FString PrimvarsDisplayColor;
 	extern UNREALUSDWRAPPER_API FString PrimvarsDisplayOpacity;
 	extern UNREALUSDWRAPPER_API FString DoubleSided;
+
+	// Tokens from UsdGeomModelAPI that we need to reference from the UsdStageActor
+	extern UNREALUSDWRAPPER_API FString ModelDrawMode;
+	extern UNREALUSDWRAPPER_API FString ModelApplyDrawMode;
 }
 
 struct UNREALUSDWRAPPER_API FUsdDelegates
diff --git a/Engine/Plugins/Importers/USDImporter/Source/UnrealUSDWrapper/Public/UsdWrappers/UsdAttribute.h b/Engine/Plugins/Importers/USDImporter/Source/UnrealUSDWrapper/Public/UsdWrappers/UsdAttribute.h
--- a/Engine/Plugins/Importers/USDImporter/Source/UnrealUSDWrapper/Public/UsdWrappers/UsdAttribute.h
+++ b/Engine/Plugins/Importers/USDImporter/Source/UnrealUSDWrapper/Public/UsdWrappers/UsdAttribute.h
@@ -84,6 +84,7 @@ namespace UE
 		size_t GetNumTimeSamples() const;
 
 		bool HasValue() const;
+		bool HasAuthoredValue() const;
 		bool HasFallbackValue() const;
 
 		bool ValueMightBeTimeVarying() const;
diff --git a/Engine/Plugins/Importers/USDImporter/Source/UnrealUSDWrapper/Public/UsdWrappers/UsdGeomBBoxCache.h b/Engine/Plugins/Importers/USDImporter/Source/UnrealUSDWrapper/Public/UsdWrappers/UsdGeomBBoxCache.h
new file mode 100644
--- /dev/null
+++ b/Engine/Plugins/Importers/USDImporter/Source/UnrealUSDWrapper/Public/UsdWrappers/UsdGeomBBoxCache.h
@@ -0,0 +1,85 @@
+// Copyright Epic Games, Inc. All Rights Reserved.
+
+#pragma once
+
+#include "CoreTypes.h"
+#include "HAL/CriticalSection.h"
+#include "Templates/UniquePtr.h"
+
+#if USE_USD_SDK
+
+#include "USDIncludesStart.h"
+
+#include "pxr/pxr.h"
+
+#include "USDIncludesEnd.h"
+
+PXR_NAMESPACE_OPEN_SCOPE
+class UsdPrim;
+class UsdGeomBBoxCache;
+PXR_NAMESPACE_CLOSE_SCOPE
+
+#endif	  // #if USE_USD_SDK
+
+enum class EUsdPurpose : int32;
+
+namespace UE
+{
+	namespace Internal
+	{
+		class FUsdGeomBBoxCacheImpl;
+	}
+
+	/**
+	 * Minimal pxr::UsdGeomBBoxCache wrapper for Unreal that can be used from no-rtti modules.
+	 * Unlike the other wrapper types, FUsdGeomBBoxCache contains its own mutex in order to try
+	 * and make its usage somewhat thread-safe. The wrapped functions will use this mutex already,
+	 * and the conversion operators that return the underlying type will check that the mutex is
+	 * locked before returning (the caller is expected to lock the mutex before using the operators).
+	 */
+	class UNREALUSDWRAPPER_API FUsdGeomBBoxCache
+	{
+	public:
+		FUsdGeomBBoxCache(double Time, EUsdPurpose IncludedPurposeFlags, bool bUseExtentsHint = false, bool bIgnoreVisibility = false);
+		FUsdGeomBBoxCache(const FUsdGeomBBoxCache& Other);
+		FUsdGeomBBoxCache(FUsdGeomBBoxCache&& Other);
+
+		FUsdGeomBBoxCache& operator=(const FUsdGeomBBoxCache& Other);
+		FUsdGeomBBoxCache& operator=(FUsdGeomBBoxCache&& Other);
+
+		~FUsdGeomBBoxCache();
+
+		// Auto conversion from/to pxr::UsdGeomBBoxCache
+	public:
+#if USE_USD_SDK
+		explicit FUsdGeomBBoxCache(const pxr::UsdGeomBBoxCache& InUsdGeomBBoxCache);
+		explicit FUsdGeomBBoxCache(pxr::UsdGeomBBoxCache&& InUsdGeomBBoxCache);
+
+		operator pxr::UsdGeomBBoxCache&();
+		operator const pxr::UsdGeomBBoxCache&() const;
+#endif	  // #if USE_USD_SDK
+
+		  // Wrapped pxr::UsdGeomBBoxCache functions, refer to the USD SDK documentation
+	public:
+		void Clear();
+
+		void SetIncludedPurposes(EUsdPurpose IncludedPurposeFlags);
+		EUsdPurpose GetIncludedPurposes() const;
+
+		bool GetUseExtentsHint() const;
+		bool GetIgnoreVisibility() const;
+
+		void SetTime(double UsdTimeCode);
+		double GetTime() const;
+
+	public:
+		// The underlying pxr::UsdGeomBBoxCache object is not thread-safe, so we have this.
+		// This mutex will be internally locked for all of the wrapped function calls.
+		// It is exposed so that it can be locked by the user in case it is desirable to use
+		// the underlying pxr::UsdGeomBBoxCache type directly
+		mutable FRWLock Lock;
+
+	private:
+		TUniquePtr<Internal::FUsdGeomBBoxCacheImpl> Impl;
+	};
+}	 // namespace UE
