USD: Implement StageState, that lets you pick whether to open/load/unload stages manually and can be bound in a Sequencer track;

Rename GetOrLoadUsdStage to GetOrOpenUsdStage as that is what it does;
Fix OnPreStageChanged not being called before changing the stage in some contexts;
Fix some inconsistencies in the UsdStageActor setters and SUsdStage widgets.

diff --git a/Engine/Plugins/Importers/USDImporter/Source/USDStage/Private/USDStageActor.cpp b/Engine/Plugins/Importers/USDImporter/Source/USDStage/Private/USDStageActor.cpp
--- a/Engine/Plugins/Importers/USDImporter/Source/USDStage/Private/USDStageActor.cpp
+++ b/Engine/Plugins/Importers/USDImporter/Source/USDStage/Private/USDStageActor.cpp
@@ -130,7 +130,7 @@ struct FUsdStageActorImpl
 	static TSharedRef< FUsdSchemaTranslationContext > CreateUsdSchemaTranslationContext(AUsdStageActor* StageActor, const FString& PrimPath)
 	{
 		TSharedRef< FUsdSchemaTranslationContext > TranslationContext = MakeShared< FUsdSchemaTranslationContext >(
-			StageActor->GetOrLoadUsdStage(),
+			StageActor->GetOrOpenUsdStage(),
 			*StageActor->UsdAssetCache
 		);
 
@@ -963,7 +963,8 @@ TMap<UBlueprint*, FDelegateHandle> FRecompilationTracker::RecompilingBlueprints;
 #endif // WITH_EDITOR
 
 AUsdStageActor::AUsdStageActor()
-	: InitialLoadSet(EUsdInitialLoadSet::LoadAll)
+	: StageState(EUsdStageState::OpenedAndLoaded)
+	, InitialLoadSet(EUsdInitialLoadSet::LoadAll)
 	, InterpolationType(EUsdInterpolationType::Linear)
 	, KindsToCollapse((int32)(EUsdDefaultKind::Component | EUsdDefaultKind::Subcomponent))
 	, bMergeIdenticalMaterialSlots(true)
@@ -1229,6 +1230,8 @@ void AUsdStageActor::IsolateLayer( const UE::FSdfLayer& Layer )
 		return;
 	}
 
+	OnPreStageChanged.Broadcast();
+
 	// Stop isolating
 	if (!Layer || Layer == UsdStage.GetRootLayer())
 	{
@@ -1298,7 +1301,7 @@ void AUsdStageActor::OnUsdObjectsChanged(const UsdUtils::FObjectChangesByPath& I
 		Transactor->Update(InfoChanges, ResyncChanges);
 	}
 
-	const UE::FUsdStage& Stage = GetOrLoadUsdStage();
+	const UE::FUsdStage& Stage = GetOrOpenUsdStage();
 	if (!Stage)
 	{
 		return;
@@ -1842,6 +1845,7 @@ USDSTAGE_API void AUsdStageActor::Reset()
 	Super::Reset();
 
 	UnloadUsdStage();
+	CloseUsdStage();
 
 	Time = 0.f;
 	RootLayer.FilePath.Empty();
@@ -1886,7 +1890,7 @@ UUsdPrimTwin* AUsdStageActor::GetOrCreatePrimTwin(const UE::FSdfPath& UsdPrimPat
 	UUsdPrimTwin* UsdPrimTwin = RootTwin->Find(PrimPath);
 	UUsdPrimTwin* ParentUsdPrimTwin = RootTwin->Find(ParentPrimPath);
 
-	const UE::FUsdPrim Prim = GetOrLoadUsdStage().GetPrimAtPath(UsdPrimPath);
+	const UE::FUsdPrim Prim = GetOrOpenUsdStage().GetPrimAtPath(UsdPrimPath);
 
 	if (!Prim)
 	{
@@ -2106,7 +2110,7 @@ void AUsdStageActor::UpdatePrim(const UE::FSdfPath& InUsdPrimPath, bool bResync,
 
 	if (UsdPrimPath.IsAbsoluteRootOrPrimPath())
 	{
-		UE::FUsdPrim PrimToExpand = GetOrLoadUsdStage().GetPrimAtPath(UsdPrimPath);
+		UE::FUsdPrim PrimToExpand = GetOrOpenUsdStage().GetPrimAtPath(UsdPrimPath);
 		ExpandPrim(PrimToExpand, bResync, TranslationContext);
 	}
 }
@@ -2136,8 +2140,14 @@ void AUsdStageActor::SetUsdStage(const UE::FUsdStage& NewStage)
 	const bool bMarkDirty = false;
 	Modify(bMarkDirty);
 
-	OnPreStageChanged.Broadcast();
+	// Fire this in case CloseUsdStage is not going to
+	if (!UsdStage)
+	{
+		OnPreStageChanged.Broadcast();
+	}
+
 	UnloadUsdStage();
+	CloseUsdStage();
 
 	RootLayer.FilePath = NewStage.GetRootLayer().GetIdentifier();
 
@@ -2164,7 +2174,14 @@ void AUsdStageActor::SetUsdStage(const UE::FUsdStage& NewStage)
 	OnStageChanged.Broadcast();
 }
 
+PRAGMA_DISABLE_DEPRECATION_WARNINGS
 UE::FUsdStage& AUsdStageActor::GetOrLoadUsdStage()
+{
+	return GetOrOpenUsdStage();
+}
+PRAGMA_ENABLE_DEPRECATION_WARNINGS
+
+UE::FUsdStage& AUsdStageActor::GetOrOpenUsdStage()
 {
 	OpenUsdStage();
 
@@ -2173,9 +2190,6 @@ UE::FUsdStage& AUsdStageActor::GetOrLoadUsdStage()
 
 void AUsdStageActor::SetRootLayer(const FString& RootFilePath)
 {
-	const bool bMarkDirty = false;
-	Modify(bMarkDirty);
-
 	FString RelativeFilePath = RootFilePath;
 #if USE_USD_SDK
 	if (!RelativeFilePath.IsEmpty() && !FPaths::IsRelative(RelativeFilePath) && !RelativeFilePath.StartsWith(UnrealIdentifiers::IdentifierPrefix))
@@ -2197,8 +2211,13 @@ void AUsdStageActor::SetRootLayer(const FString& RootFilePath)
 		}
 	}
 
+	const bool bMarkDirty = false;
+	Modify(bMarkDirty);
+
 	UnloadUsdStage();
+	CloseUsdStage();
 	RootLayer.FilePath = RelativeFilePath;
+	OpenUsdStage();
 	LoadUsdStage();
 
 	// Do this here instead of on OpenUsdStage/LoadUsdStage as those also get called when changing any of
@@ -2206,26 +2225,54 @@ void AUsdStageActor::SetRootLayer(const FString& RootFilePath)
 	UsdUtils::CollectSchemaAnalytics(UsdStage, TEXT("Open"));
 }
 
+void AUsdStageActor::SetStageState(EUsdStageState NewState)
+{
+	if (NewState == StageState)
+	{
+		return;
+	}
+
+	const bool bMarkDirty = false;
+	Modify(bMarkDirty);
+
+	StageState = NewState;
+	if (StageState == EUsdStageState::Closed)
+	{
+		UnloadUsdStage();
+		CloseUsdStage();
+	}
+	else if (StageState == EUsdStageState::Opened)
+	{
+		UnloadUsdStage();
+		OpenUsdStage();
+	}
+	else if (StageState == EUsdStageState::OpenedAndLoaded)
+	{
+		OpenUsdStage();
+		LoadUsdStage();
+	}
+}
+
 void AUsdStageActor::SetAssetCache(UUsdAssetCache2* NewCache)
 {
+	if (NewCache == UsdAssetCache)
+	{
+		return;
+	}
+
 	const bool bMarkDirty = false;
 	Modify(bMarkDirty);
 
 	// Remove ourselves from our previous cache. We're going to have to create new assets anyway, so it doesn't matter
 	// if we discard our current assets.
-	// Note that these functions may be called when the cache is already set (if via
-	// AUsdStageActor::HandlePropertyChangedEvent) or to set the actual new cache (when called directly)
-	if (UsdAssetCache != NewCache)
+	if (UsdAssetCache)
 	{
-		if (UsdAssetCache)
-		{
-			UsdAssetCache->RemoveAllAssetReferences(this);
-			UsdAssetCache->RefreshStorage();
-		}
-
-		UsdAssetCache = NewCache;
+		UsdAssetCache->RemoveAllAssetReferences(this);
+		UsdAssetCache->RefreshStorage();
 	}
 
+	UsdAssetCache = NewCache;
+
 	// We can't have no cache while we have a stage loaded, so at least revert the property to a transient cache
 	// instead, as the intent may have been to just have the actor not point at the previous cache anymore.
 	if (!UsdAssetCache && UsdStage)
@@ -2254,6 +2301,11 @@ void AUsdStageActor::SetAssetCache(UUsdAssetCache2* NewCache)
 
 void AUsdStageActor::SetInitialLoadSet(EUsdInitialLoadSet NewLoadSet)
 {
+	if (NewLoadSet == InitialLoadSet)
+	{
+		return;
+	}
+
 	const bool bMarkDirty = false;
 	Modify(bMarkDirty);
 
@@ -2263,6 +2315,11 @@ void AUsdStageActor::SetInitialLoadSet(EUsdInitialLoadSet NewLoadSet)
 
 void AUsdStageActor::SetInterpolationType(EUsdInterpolationType NewType)
 {
+	if (NewType == InterpolationType)
+	{
+		return;
+	}
+
 	const bool bMarkDirty = false;
 	Modify(bMarkDirty);
 
@@ -2272,9 +2329,6 @@ void AUsdStageActor::SetInterpolationType(EUsdInterpolationType NewType)
 
 void AUsdStageActor::SetKindsToCollapse(int32 NewKindsToCollapse)
 {
-	const bool bMarkDirty = false;
-	Modify(bMarkDirty);
-
 	const EUsdDefaultKind NewEnum = (EUsdDefaultKind)NewKindsToCollapse;
 	EUsdDefaultKind Result = NewEnum;
 
@@ -2290,13 +2344,27 @@ void AUsdStageActor::SetKindsToCollapse(int32 NewKindsToCollapse)
 		Result |= (EUsdDefaultKind::Assembly);
 	}
 
+	if ((int32)Result == KindsToCollapse)
+	{
+		return;
+	}
+
+	const bool bMarkDirty = false;
+	Modify(bMarkDirty);
+
 	KindsToCollapse = (int32)Result;
 	LoadUsdStage();
 }
 
 void AUsdStageActor::SetMergeIdenticalMaterialSlots(bool bMerge)
 {
-	Modify();
+	if (bMerge == bMergeIdenticalMaterialSlots)
+	{
+		return;
+	}
+
+	const bool bMarkDirty = false;
+	Modify(bMarkDirty);
 
 	bMergeIdenticalMaterialSlots = bMerge;
 	LoadUsdStage();
@@ -2308,6 +2376,11 @@ void AUsdStageActor::SetCollapseTopLevelPointInstancers(bool bCollapse)
 
 void AUsdStageActor::SetPurposesToLoad(int32 NewPurposesToLoad)
 {
+	if (NewPurposesToLoad == PurposesToLoad)
+	{
+		return;
+	}
+
 	const bool bMarkDirty = false;
 	Modify(bMarkDirty);
 
@@ -2321,6 +2394,11 @@ void AUsdStageActor::SetPurposesToLoad(int32 NewPurposesToLoad)
 
 void AUsdStageActor::SetNaniteTriangleThreshold(int32 NewNaniteTriangleThreshold)
 {
+	if (NewNaniteTriangleThreshold == NaniteTriangleThreshold)
+	{
+		return;
+	}
+
 	const bool bMarkDirty = false;
 	Modify(bMarkDirty);
 
@@ -2330,6 +2408,11 @@ void AUsdStageActor::SetNaniteTriangleThreshold(int32 NewNaniteTriangleThreshold
 
 void AUsdStageActor::SetRenderContext(const FName& NewRenderContext)
 {
+	if (NewRenderContext == RenderContext)
+	{
+		return;
+	}
+
 	const bool bMarkDirty = false;
 	Modify(bMarkDirty);
 
@@ -2339,6 +2422,11 @@ void AUsdStageActor::SetRenderContext(const FName& NewRenderContext)
 
 void AUsdStageActor::SetMaterialPurpose(const FName& NewMaterialPurpose)
 {
+	if (NewMaterialPurpose == MaterialPurpose)
+	{
+		return;
+	}
+
 	const bool bMarkDirty = false;
 	Modify(bMarkDirty);
 
@@ -2346,8 +2434,13 @@ void AUsdStageActor::SetMaterialPurpose(const FName& NewMaterialPurpose)
 	LoadUsdStage();
 }
 
-void AUsdStageActor::SetRootMotionHandling( EUsdRootMotionHandling NewHandlingStrategy )
+void AUsdStageActor::SetRootMotionHandling(EUsdRootMotionHandling NewHandlingStrategy)
 {
+	if (NewHandlingStrategy == RootMotionHandling)
+	{
+		return;
+	}
+
 	const bool bMarkDirty = false;
 	Modify(bMarkDirty);
 
@@ -2371,6 +2464,11 @@ float AUsdStageActor::GetTime() const
 
 void AUsdStageActor::SetTime(float InTime)
 {
+	if (InTime == Time)
+	{
+		return;
+	}
+
 	const bool bMarkDirty = false;
 	Modify(bMarkDirty);
 
@@ -2479,8 +2577,8 @@ FString AUsdStageActor::GetSourcePrimPath(UObject* Object)
 
 void AUsdStageActor::OpenUsdStage()
 {
-	// Early exit if stage is already opened
-	if (UsdStage || RootLayer.FilePath.IsEmpty())
+	// Early exit if stage is already opened, or if we shouldn't be opening anything anyway
+	if (UsdStage || RootLayer.FilePath.IsEmpty() || StageState == EUsdStageState::Closed)
 	{
 		return;
 	}
@@ -2524,25 +2622,35 @@ void AUsdStageActor::OpenUsdStage()
 		const bool bCreateIfNeeded = false;
 		UsdUtils::GetUEPersistentStateSublayer(UsdStage, bCreateIfNeeded);
 #endif // #if USE_USD_SDK
-
-		OnStageChanged.Broadcast();
 	}
 
+	OnStageChanged.Broadcast();
+
 	UsdUtils::ShowErrorsAndStopMonitoring(FText::Format(LOCTEXT("USDOpenError", "Encountered some errors opening USD file at path '{0}!\nCheck the Output Log for details."), FText::FromString(RootLayer.FilePath)));
 }
 
-void AUsdStageActor::CloseUsdStage()
+void AUsdStageActor::CloseUsdStage(bool bUnloadIfNeeded)
 {
-	if (UsdStage)
+	const bool bStageWasOpened = static_cast<bool>(UsdStage);
+	if (bStageWasOpened)
 	{
 		OnPreStageChanged.Broadcast();
+
+		if (bUnloadIfNeeded)
+		{
+			UnloadUsdStage();
+		}
 	}
 
 	FUsdStageActorImpl::DiscardStage(UsdStage, this);
 	UsdStage = UE::FUsdStage();
-	IsolatedStage = UE::FUsdStage();  // We don't keep our isolated stages on the stage cache
-	LevelSequenceHelper.Init(UE::FUsdStage()); // Drop the helper's reference to the stage
-	OnStageChanged.Broadcast();
+	IsolatedStage = UE::FUsdStage();			  // We don't keep our isolated stages on the stage cache
+	LevelSequenceHelper.Init(UE::FUsdStage());	  // Drop the helper's reference to the stage
+
+	if (bStageWasOpened)
+	{
+		OnStageChanged.Broadcast();
+	}
 }
 
 #if WITH_EDITOR
@@ -2668,8 +2776,10 @@ void AUsdStageActor::OnObjectsReplaced(const TMap<UObject*, UObject*>& ObjectRep
 			NewActor->IsBlockedFromUsdNotices.Set(IsBlockedFromUsdNotices.GetValue());
 			NewActor->OldRootLayer = OldRootLayer;
 
-			// Close our stage or else it will remain open forever. NewActor has a a reference to it now so it won't actually close
-			CloseUsdStage();
+			// Close our stage or else it will remain open forever. NewActor has a a reference to it now so it won't actually close.
+			// Don't discard our spawned actors and components though, as they will be used by the replacement
+			const bool bUnloadIfNeeded = false;
+			CloseUsdStage(bUnloadIfNeeded);
 		}
 	}
 }
@@ -2688,10 +2798,24 @@ void AUsdStageActor::OnLevelActorDeleted(AActor* DeletedActor)
 
 #endif // WITH_EDITOR
 
-void AUsdStageActor::LoadUsdStage()
+void AUsdStageActor::LoadUsdStage(bool bOpenIfNeeded)
 {
 	TRACE_CPUPROFILER_EVENT_SCOPE(AUsdStageActor::LoadUsdStage);
 
+	if (StageState != EUsdStageState::OpenedAndLoaded)
+	{
+		return;
+	}
+
+	if (!UsdStage && bOpenIfNeeded)
+	{
+		OpenUsdStage();
+		if (!UsdStage)
+		{
+			return;
+		}
+	}
+
 	// Ensure an asset cache if we're going to load something
 	if (!RootLayer.FilePath.IsEmpty())
 	{
@@ -2740,29 +2864,7 @@ void AUsdStageActor::LoadUsdStage()
 		ScopedReferencer.Emplace(UsdAssetCache, this);
 	}
 
-	// If we have an isolated stage we may have been asked to refresh to display it: We should keep our UsdStage opened
-	UE::FUsdStage StageToLoad;
-	if (IsolatedStage)
-	{
-		StageToLoad = IsolatedStage;
-	}
-	else
-	{
-		// If we're in here we don't expect our current stage to be the same as the new stage we're trying to load, so
-		// get rid of it so that OpenUsdStage can open it
-		UsdStage = UE::FUsdStage();
-
-		OpenUsdStage();
-		if (!UsdStage)
-		{
-			OnStageChanged.Broadcast();
-			return;
-		}
-
-		StageToLoad = UsdStage;
-	}
-
-	StageToLoad.SetInterpolationType(InterpolationType);
+	UE::FUsdStage StageToLoad = GetUsdStage();
 
 	// Create Info and BBoxCache before calling ReloadAnimations as that is when the LevelSequenceHelper will also take
 	// a reference to them
@@ -2876,7 +2978,8 @@ void AUsdStageActor::UnloadUsdStage()
 	if (LevelSequence)
 	{
 #if WITH_EDITOR
-		if (GEditor)
+		// CloseAllEditorsForAsset crashes if called when the engine is closing
+		if (GEditor && !IsEngineExitRequested())
 		{
 			GEditor->GetEditorSubsystem<UAssetEditorSubsystem>()->CloseAllEditorsForAsset(LevelSequence);
 		}
@@ -2922,8 +3025,6 @@ void AUsdStageActor::UnloadUsdStage()
 	}
 
 	OnStageUnloaded.Broadcast();
-
-	CloseUsdStage();
 }
 
 void AUsdStageActor::SetupAssetCacheIfNeeded()
@@ -3024,7 +3125,7 @@ void AUsdStageActor::ReloadAnimations()
 	}
 	TGuardValue<bool> ReentrantGuard(bIsReentrant, true);
 
-	const UE::FUsdStage& CurrentStage = GetOrLoadUsdStage();
+	const UE::FUsdStage& CurrentStage = GetOrOpenUsdStage();
 	if (!CurrentStage)
 	{
 		return;
@@ -3167,10 +3268,33 @@ void AUsdStageActor::PostTransacted(const FTransactionObjectEvent& TransactionEv
 		if (ChangedProperties.Contains(GET_MEMBER_NAME_CHECKED(AUsdStageActor, RootLayer)))
 		{
 			// Changed the path, so we need to reopen the correct stage
-			CloseUsdStage();
+			// Note: We don't unload/load here, as that would wipe the spawned actors and components that were
+			// potentially recreated with the transaction
+			const bool bUnloadIfNeeded = false;
+			CloseUsdStage(bUnloadIfNeeded);
 			OpenUsdStage();
 			ReloadAnimations();
 		}
+		else if (ChangedProperties.Contains(GET_MEMBER_NAME_CHECKED(AUsdStageActor, StageState)))
+		{
+			// Partially copied from SetStageState, except that in here we don't want to call the
+			// Load/UnloadUsdStage functions. Firstly because we'll already have the assets/actors/components in place
+			// since they came along with us for the transaction, and secondly because PostTransacted is itself
+			// outside of a transaction: Any change done in here (creating/destroying/modifying UObjects) is outside of
+			// the transaction system and would cause chaos if we were to hit Undo/Redo afterwards
+			if (StageState == EUsdStageState::Closed)
+			{
+				CloseUsdStage();
+			}
+			else if (StageState == EUsdStageState::Opened)
+			{
+				OpenUsdStage();
+			}
+			else if (StageState == EUsdStageState::OpenedAndLoaded)
+			{
+				OpenUsdStage();
+			}
+		}
 		else if (ChangedProperties.Contains(GET_MEMBER_NAME_CHECKED(AUsdStageActor, Time)))
 		{
 			Refresh();
@@ -3601,6 +3725,7 @@ void AUsdStageActor::UnregisterAllComponents(bool bForReregister)
 	}
 
 	UnloadUsdStage();
+	CloseUsdStage();
 }
 
 void AUsdStageActor::PostUnregisterAllComponents()
@@ -4158,7 +4283,37 @@ void AUsdStageActor::HandlePropertyChangedEvent(FPropertyChangedEvent& PropertyC
 	// Handle property changed events with this function (called from our OnObjectPropertyChanged delegate) instead of overriding PostEditChangeProperty because replicated
 	// multi-user transactions directly broadcast OnObjectPropertyChanged on the properties that were changed, instead of making PostEditChangeProperty events.
 	// Note that UObject::PostEditChangeProperty ends up broadcasting OnObjectPropertyChanged anyway, so this works just the same as before.
-	// see ConcertClientTransactionBridge.cpp, function ConcertClientTransactionBridgeUtil::ProcessTransactionEvent
+	// see ConcertClientTransactionBridge.cpp, function ConcertClientTransactionBridgeUtil::ProcessTransactionEvent.
+
+	// Note that in here we'll delegate to these setter functions (like SetRootLayer) to actually set the new property values.
+	// We want our setter functions to be able to automatically refresh the stage (both for simplicity, since we have a single
+	// code path for changing them that gets reused everywhere) and also due to the fact that the Sequencer uses these setters
+	// when we create Sequencer tracks for these properties: If we make a track for e.g. "PurposesToLoad", we want the stage to
+	// refresh as soon as we hit a keyframe to change the chosen purposes. We don't want to need some separate track to "refresh
+	// the stage" or something like that.
+	//
+	// An issue, however, is the fact that the Sequencer can repeatedly call these setters with the same value over and over in
+	// case it is just e.g. stopped at some frame. We don't want that to keep reloading the stage, so we need the setters to have
+	// an "early out" and not do anything in case they're receiving the same value that was previously set.
+	//
+	// With an "early out" mechanism though, we end up with a problem: This function (called from OnObjectPropertyChanged) is only
+	// called *after* these properties have already been set with their new values. So if we naively delegated to the setters now
+	// they would all just "early out" and do nothing. We do still need to respond from the OnObjectPropertyChanged code path though,
+	// due to the fact that ConcertClientTransactionBridgeUtil::ProcessTransactionEvent calls OnObjectPropertyChanged directly in
+	// order to replicate the multiuser property value changes. We want the stage to automatically refresh when that happens,
+	// meaning we need to do exactly what the setters do anyway and may as well call them. TL;DR: We need this function and for it
+	// to call the setters.
+	//
+	// We can't rely on any other additional event (like OnPreObjectPropertyChanged) because that doesn't tell us the new value
+	// that will be changed anyway, so we'd need some complicated mechanism to store the property values at e.g.
+	// OnPreObjectPropertyChanged time and compare our current values to them to know if something changed...
+	//
+	// This explains the CorrectValues (e.g. CorrectTime, CorretRootLayer, etc.) you'll see below: We will temporarly put a
+	// different value on the properties before calling them to prevent the setters from earlying out. We don't want to record
+	// these spoofed values into the transaction though (otherwise if we hit Undo we would end up with those set), so we Modify()
+	// before we do that.
+	const bool bAlwaysMarkAsDirty = false;
+	Modify(bAlwaysMarkAsDirty);
 
 	FProperty* PropertyThatChanged = PropertyChangedEvent.MemberProperty;
 	const FName PropertyName = PropertyThatChanged ? PropertyThatChanged->GetFName() : NAME_None;
@@ -4170,7 +4325,11 @@ void AUsdStageActor::HandlePropertyChangedEvent(FPropertyChangedEvent& PropertyC
 		const bool bDiscardUndo = CVar && CVar->GetBool();
 #endif // WITH_EDITOR
 
-		SetRootLayer(RootLayer.FilePath);
+		// Technically we don't need this guard value for the root layer itself, since SetRootLayer can compare
+		// RootLayer with the path of the current stage's root layer, but let's just do this for consistency.
+		const FString CorrectRootLayer = RootLayer.FilePath;
+		RootLayer.FilePath = RootLayer.FilePath + TEXT("dummy");
+		SetRootLayer(CorrectRootLayer);
 
 #if WITH_EDITOR
 		if (bDiscardUndo && GEditor)
@@ -4183,45 +4342,71 @@ void AUsdStageActor::HandlePropertyChangedEvent(FPropertyChangedEvent& PropertyC
 		}
 #endif // WITH_EDITOR
 	}
+	else if (PropertyName == GET_MEMBER_NAME_CHECKED(AUsdStageActor, StageState))
+	{
+		const EUsdStageState CorrectState = StageState;
+		StageState = (EUsdStageState)!((uint8)StageState);
+		SetStageState(CorrectState);
+	}
 	else if (PropertyName == GET_MEMBER_NAME_CHECKED(AUsdStageActor, Time))
 	{
-		SetTime(Time);
+		const float CorrectTime = Time;
+		Time = Time + 1.0f;
+		SetTime(CorrectTime);
 	}
 	else if (PropertyName == GET_MEMBER_NAME_CHECKED(AUsdStageActor, InitialLoadSet))
 	{
-		SetInitialLoadSet(InitialLoadSet);
+		const EUsdInitialLoadSet CorrectLoadSet = InitialLoadSet;
+		InitialLoadSet = (EUsdInitialLoadSet) !((uint8)InitialLoadSet);
+		SetInitialLoadSet(CorrectLoadSet);
 	}
 	else if (PropertyName == GET_MEMBER_NAME_CHECKED(AUsdStageActor, InterpolationType))
 	{
-		SetInterpolationType(InterpolationType);
+		const EUsdInterpolationType CorrectInterpolationType = InterpolationType;
+		InterpolationType = (EUsdInterpolationType) !((uint8)InterpolationType);
+		SetInterpolationType(CorrectInterpolationType);
 	}
 	else if (PropertyName == GET_MEMBER_NAME_CHECKED(AUsdStageActor, KindsToCollapse))
 	{
-		SetKindsToCollapse(KindsToCollapse);
+		const int32 CorrectKindsToCollapse = KindsToCollapse;
+		KindsToCollapse += 1;
+		SetKindsToCollapse(CorrectKindsToCollapse);
 	}
 	else if (PropertyName == GET_MEMBER_NAME_CHECKED(AUsdStageActor, bMergeIdenticalMaterialSlots))
 	{
-		SetMergeIdenticalMaterialSlots(bMergeIdenticalMaterialSlots);
+		const bool bCorrectMergeMaterialSlots = bMergeIdenticalMaterialSlots;
+		bMergeIdenticalMaterialSlots = !bMergeIdenticalMaterialSlots;
+		SetMergeIdenticalMaterialSlots(bCorrectMergeMaterialSlots);
 	}
 	else if (PropertyName == GET_MEMBER_NAME_CHECKED(AUsdStageActor, PurposesToLoad))
 	{
-		SetPurposesToLoad(PurposesToLoad);
+		const int32 CorrectPurposesToLoad = PurposesToLoad;
+		PurposesToLoad += 1;
+		SetPurposesToLoad(CorrectPurposesToLoad);
 	}
 	else if (PropertyName == GET_MEMBER_NAME_CHECKED(AUsdStageActor, NaniteTriangleThreshold))
 	{
-		SetNaniteTriangleThreshold(NaniteTriangleThreshold);
+		const int32 CorrectNaniteThreshold = NaniteTriangleThreshold;
+		NaniteTriangleThreshold += 1;
+		SetNaniteTriangleThreshold(CorrectNaniteThreshold);
 	}
 	else if (PropertyName == GET_MEMBER_NAME_CHECKED(AUsdStageActor, RenderContext))
 	{
-		SetRenderContext(RenderContext);
+		const FName CorrectRenderContext = RenderContext;
+		RenderContext = *(RenderContext.ToString() + TEXT("dummy"));
+		SetRenderContext(CorrectRenderContext);
 	}
 	else if (PropertyName == GET_MEMBER_NAME_CHECKED(AUsdStageActor, MaterialPurpose))
 	{
-		SetMaterialPurpose(MaterialPurpose);
+		const FName CorrectMaterialPurpose = MaterialPurpose;
+		MaterialPurpose = *(MaterialPurpose.ToString() + TEXT("dummy"));
+		SetMaterialPurpose(CorrectMaterialPurpose);
 	}
 	else if (PropertyName == GET_MEMBER_NAME_CHECKED(AUsdStageActor, RootMotionHandling))
 	{
-		SetRootMotionHandling(RootMotionHandling);
+		const EUsdRootMotionHandling CorrectHandling = RootMotionHandling;
+		RootMotionHandling = (EUsdRootMotionHandling) !((uint8)RootMotionHandling);
+		SetRootMotionHandling(CorrectHandling);
 	}
 	else if (PropertyName == GET_MEMBER_NAME_CHECKED(AUsdStageActor, SubdivisionLevel))
 	{
@@ -4229,7 +4414,9 @@ void AUsdStageActor::HandlePropertyChangedEvent(FPropertyChangedEvent& PropertyC
 	}
 	else if (PropertyName == GET_MEMBER_NAME_CHECKED(AUsdStageActor, UsdAssetCache))
 	{
-		SetAssetCache(UsdAssetCache);
+		UUsdAssetCache2* const CorrectCache = UsdAssetCache.Get();
+		UsdAssetCache = UsdAssetCache ? nullptr : NewObject<UUsdAssetCache2>();
+		SetAssetCache(CorrectCache);
 	}
 
 	bIsModifyingAProperty = false;
diff --git a/Engine/Plugins/Importers/USDImporter/Source/USDStage/Private/USDTransactor.cpp b/Engine/Plugins/Importers/USDImporter/Source/USDStage/Private/USDTransactor.cpp
--- a/Engine/Plugins/Importers/USDImporter/Source/USDStage/Private/USDTransactor.cpp
+++ b/Engine/Plugins/Importers/USDImporter/Source/USDStage/Private/USDTransactor.cpp
@@ -840,7 +840,7 @@ namespace UsdUtils
 			// will already be accounted for by those actors/components/assets undoing/redoing by themselves via the UE transaction buffer.
 			FScopedBlockNotices BlockNotices( StageActor->GetUsdListener() );
 
-			UE::FUsdStage& Stage = StageActor->GetOrLoadUsdStage();
+			UE::FUsdStage& Stage = StageActor->GetOrOpenUsdStage();
 
 			TArray<FString> PrimsChanged = UsdUtils::ApplyFieldMapToStage(
 				Values,
@@ -891,7 +891,7 @@ namespace UsdUtils
 				BlockNotices.Emplace( StageActor->GetUsdListener() );
 			}
 
-			UE::FUsdStage& Stage = StageActor->GetOrLoadUsdStage();
+			UE::FUsdStage& Stage = StageActor->GetOrOpenUsdStage();
 
 			TArray<FString> PrimsChanged;
 			if ( bIsApplyingConcertSync && ReceivedValuesBeforeUndo.IsSet() )
@@ -1043,7 +1043,7 @@ void UUsdTransactor::Update( const UsdUtils::FObjectChangesByPath& NewInfoChange
 
 	Modify();
 
-	const UE::FUsdStage& Stage = StageActorPtr->GetOrLoadUsdStage();
+	const UE::FUsdStage& Stage = StageActorPtr->GetOrOpenUsdStage();
 	if ( !Stage )
 	{
 		return;
diff --git a/Engine/Plugins/Importers/USDImporter/Source/USDStage/Public/USDStageActor.h b/Engine/Plugins/Importers/USDImporter/Source/USDStage/Public/USDStageActor.h
--- a/Engine/Plugins/Importers/USDImporter/Source/USDStage/Public/USDStageActor.h
+++ b/Engine/Plugins/Importers/USDImporter/Source/USDStage/Public/USDStageActor.h
@@ -28,6 +28,14 @@ namespace UE
 	class FUsdPrim;
 }
 
+UENUM()
+enum class EUsdStageState : uint8
+{
+	Closed,
+	Opened,
+	OpenedAndLoaded
+};
+
 UCLASS( MinimalAPI, config = USDImporter )
 class AUsdStageActor : public AActor
 {
@@ -37,6 +45,12 @@ public:
 	UPROPERTY(EditAnywhere, BlueprintReadOnly, Category = "USD", meta = (RelativeToGameDir))
 	FFilePath RootLayer;
 
+	/**
+	 * What to do with the USD Stage pointed to by RootLayer, if any
+	 */
+	UPROPERTY(EditAnywhere, BlueprintReadOnly, Category = "USD", config)
+	EUsdStageState StageState;
+
 	UPROPERTY(EditAnywhere, BlueprintReadOnly, Category = "USD")
 	TObjectPtr<UUsdAssetCache2> UsdAssetCache;
 
@@ -91,7 +105,7 @@ public:
 
 	// Describes what to add to the root bone animation within generated AnimSequences, if anything
 	UPROPERTY( EditAnywhere, BlueprintReadOnly, Category = "USD", config )
-	EUsdRootMotionHandling RootMotionHandling = EUsdRootMotionHandling::NoAdditionalRootMotion;
+	EUsdRootMotionHandling RootMotionHandling;
 
 	/**
 	 * Subdivision level to use for all subdivision meshes on the opened stage. 0 means "don't subdivide".
@@ -124,6 +138,9 @@ public:
 	UFUNCTION(BlueprintCallable, Category = "USD", meta = (CallInEditor = "true"))
 	USDSTAGE_API void SetRootLayer(const FString& RootFilePath );
 
+	UFUNCTION(BlueprintCallable, Category = "USD", meta = (CallInEditor = "true"))
+	USDSTAGE_API void SetStageState(EUsdStageState NewStageState);
+
 	UFUNCTION(BlueprintCallable, Category = "USD", meta = (CallInEditor = "true"))
 	USDSTAGE_API void SetAssetCache(UUsdAssetCache2* NewCache);
 
@@ -215,10 +232,13 @@ public:
 	USDSTAGE_API FString GetIsolatedRootLayer() const;
 
 public:
-	// Loads the stage with RootLayer if its not loaded already, and returns the either the isolated stage (if any) or
-	// the base stage
+	UE_DEPRECATED(5.4, "This function has been renamed into 'GetOrOpenUsdStage', which better describes what it does")
 	USDSTAGE_API UE::FUsdStage& GetOrLoadUsdStage();
 
+	// Opens the stage with RootLayer if its not loaded already, and returns the either the isolated stage (if any) or
+	// the base stage
+	USDSTAGE_API UE::FUsdStage& GetOrOpenUsdStage();
+
 	// Returns either the isolated stage (if any) or the base stage
 	USDSTAGE_API const UE::FUsdStage& GetUsdStage() const;
 
@@ -296,9 +316,9 @@ protected:
 	void UpdatePrim(const UE::FSdfPath& UsdPrimPath, bool bResync, FUsdSchemaTranslationContext& TranslationContext);
 
 	void OpenUsdStage();
-	void CloseUsdStage();
+	void CloseUsdStage(bool bUnloadIfNeeded = true);
 
-	void LoadUsdStage();
+	void LoadUsdStage(bool bOpenIfNeeded = true);
 	void UnloadUsdStage();
 
 	void SetupAssetCacheIfNeeded();
diff --git a/Engine/Plugins/Importers/USDImporter/Source/USDStageEditor/Private/SUSDStage.cpp b/Engine/Plugins/Importers/USDImporter/Source/USDStageEditor/Private/SUSDStage.cpp
--- a/Engine/Plugins/Importers/USDImporter/Source/USDStageEditor/Private/SUSDStage.cpp
+++ b/Engine/Plugins/Importers/USDImporter/Source/USDStageEditor/Private/SUSDStage.cpp
@@ -843,7 +843,7 @@ void SUsdStage::FillFileMenu( FMenuBuilder& MenuBuilder )
 
 		MenuBuilder.AddMenuEntry(
 			LOCTEXT("Open", "Open..."),
-			LOCTEXT("Open_ToolTip", "Opens a USD file"),
+			LOCTEXT("Open_ToolTip", "Pick an existing USD file"),
 			FSlateIcon(),
 			FUIAction(
 				FExecuteAction::CreateLambda([this]()
@@ -947,22 +947,24 @@ void SUsdStage::FillFileMenu( FMenuBuilder& MenuBuilder )
 	{
 		MenuBuilder.AddMenuEntry(
 			LOCTEXT("Close", "Close"),
-			LOCTEXT("Close_ToolTip", "Closes the opened stage"),
+			LOCTEXT("Close_ToolTip", "Clears the stage actor's root layer, unloading and closing the stage"),
 			FSlateIcon(),
 			FUIAction(
-				FExecuteAction::CreateSP( this, &SUsdStage::FileClose ),
-				FCanExecuteAction::CreateLambda( [this]()
-				{
-					if ( const AUsdStageActor* StageActor = ViewModel.UsdStageActor.Get() )
+				FExecuteAction::CreateSP(this, &SUsdStage::FileClose),
+				FCanExecuteAction::CreateLambda(
+					[this]()
 					{
-						if ( UE::FUsdStage Stage = StageActor->GetUsdStage() )
+						if (const AUsdStageActor* StageActor = ViewModel.UsdStageActor.Get())
 						{
-							return true;
+							if (!StageActor->RootLayer.FilePath.IsEmpty())
+							{
+								return true;
+							}
 						}
-					}
 
-					return false;
-				})
+						return false;
+					}
+				)
 			),
 			NAME_None,
 			EUserInterfaceActionType::Button
@@ -1005,6 +1007,11 @@ void SUsdStage::FillOptionsMenu(FMenuBuilder& MenuBuilder)
 {
 	MenuBuilder.BeginSection( TEXT( "StageOptions" ), LOCTEXT( "StageOptions", "Stage options" ) );
 	{
+		MenuBuilder.AddSubMenu(
+			LOCTEXT("State", "Stage state"),
+			LOCTEXT("State_ToolTip", "Whether we should open the USD Stage or spawn assets, actors and components"),
+			FNewMenuDelegate::CreateSP(this, &SUsdStage::FillStageStateSubMenu));
+
 		MenuBuilder.AddSubMenu(
 			LOCTEXT("Payloads", "Payloads"),
 			LOCTEXT("Payloads_ToolTip", "What to do with payloads when initially opening the stage"),
@@ -1122,6 +1129,72 @@ void SUsdStage::FillExportSubMenu( FMenuBuilder& MenuBuilder )
 	);
 }
 
+void SUsdStage::FillStageStateSubMenu(FMenuBuilder& MenuBuilder)
+{
+	TFunction<void(EUsdStageState, const FText&, const FText&)> AddStageStateEntry =
+		[this, &MenuBuilder](EUsdStageState NewState, const FText& EntryText, const FText& EntryToolTipText)
+	{
+		MenuBuilder.AddMenuEntry(
+			EntryText,
+			EntryToolTipText,
+			FSlateIcon(),
+			FUIAction(
+				FExecuteAction::CreateLambda(
+					[this, NewState, EntryText]()
+					{
+						if (AUsdStageActor* StageActor = GetStageActorOrCDO())
+						{
+							FScopedTransaction Transaction(FText::Format(
+								LOCTEXT("ChangedStageStateTransaction", "Set USD stage actor '{0}' actor to stage state '{1}'"),
+								FText::FromString(StageActor->GetActorLabel()),
+								EntryText
+							));
+
+							// c.f. comment in SUsdStage::FillCollapsingSubMenu
+							TGuardValue<bool> MaintainSelectionGuard(bUpdatingViewportSelection, true);
+
+							StageActor->SetStageState(NewState);
+							if (StageActor->IsTemplate())
+							{
+								StageActor->SaveConfig();
+							}
+						}
+					}
+				),
+				FCanExecuteAction{},
+				FIsActionChecked::CreateLambda(
+					[this, NewState]()
+					{
+						if (AUsdStageActor* StageActor = GetStageActorOrCDO())
+						{
+							return StageActor->StageState == NewState;
+						}
+						return false;
+					}
+				)
+			),
+			NAME_None,
+			EUserInterfaceActionType::RadioButton
+		);
+	};
+
+	AddStageStateEntry(
+		EUsdStageState::Closed,
+		LOCTEXT("Closed", "Closed"),
+		LOCTEXT("Closed_ToolTip", "Don't open or load the USD Stage pointed to by RootLayer")
+	);
+	AddStageStateEntry(
+		EUsdStageState::Opened,
+		LOCTEXT("Opened", "Opened"),
+		LOCTEXT("Opened_ToolTip", "Open the USD Stage pointed to by RootLayer, but don't spawn assets, actors or components")
+	);
+	AddStageStateEntry(
+		EUsdStageState::OpenedAndLoaded,
+		LOCTEXT("OpenedAndLoaded", "Opened and loaded"),
+		LOCTEXT("OpenedAndLoaded_ToolTip", "Open the USD Stage pointed to by RootLayer and spawn assets, actors or components")
+	);
+}
+
 void SUsdStage::FillPayloadsSubMenu(FMenuBuilder& MenuBuilder)
 {
 	MenuBuilder.AddMenuEntry(
@@ -1138,9 +1211,11 @@ void SUsdStage::FillPayloadsSubMenu(FMenuBuilder& MenuBuilder)
 						FText::FromString(StageActor->GetActorLabel())
 					));
 
-					StageActor->Modify();
-					StageActor->InitialLoadSet = EUsdInitialLoadSet::LoadAll;
-					if ( StageActor->IsTemplate() )
+					// c.f. comment in SUsdStage::FillCollapsingSubMenu
+					TGuardValue<bool> MaintainSelectionGuard(bUpdatingViewportSelection, true);
+
+					StageActor->SetInitialLoadSet(EUsdInitialLoadSet::LoadAll);
+					if (StageActor->IsTemplate())
 					{
 						StageActor->SaveConfig();
 					}
@@ -1174,9 +1249,11 @@ void SUsdStage::FillPayloadsSubMenu(FMenuBuilder& MenuBuilder)
 						FText::FromString(StageActor->GetActorLabel())
 					));
 
-					StageActor->Modify();
-					StageActor->InitialLoadSet = EUsdInitialLoadSet::LoadNone;
-					if ( StageActor->IsTemplate() )
+					// c.f. comment in SUsdStage::FillCollapsingSubMenu
+					TGuardValue<bool> MaintainSelectionGuard(bUpdatingViewportSelection, true);
+
+					StageActor->SetInitialLoadSet(EUsdInitialLoadSet::LoadNone);
+					if (StageActor->IsTemplate())
 					{
 						StageActor->SaveConfig();
 					}
@@ -1216,17 +1293,10 @@ void SUsdStage::FillPurposesToLoadSubMenu(FMenuBuilder& MenuBuilder)
 						));
 
 						// c.f. comment in SUsdStage::FillCollapsingSubMenu
-						TGuardValue<bool> MaintainSelectionGuard( bUpdatingViewportSelection, true );
-
-						StageActor->Modify();
-						StageActor->PurposesToLoad = (int32)((EUsdPurpose)StageActor->PurposesToLoad ^ Purpose);
+						TGuardValue<bool> MaintainSelectionGuard(bUpdatingViewportSelection, true);
 
-						FPropertyChangedEvent PropertyChangedEvent(
-							FindFieldChecked< FProperty >( StageActor->GetClass(), GET_MEMBER_NAME_CHECKED( AUsdStageActor, PurposesToLoad ) )
-						);
-						StageActor->PostEditChangeProperty(PropertyChangedEvent);
-
-						if ( StageActor->IsTemplate() )
+						StageActor->SetPurposesToLoad((int32)((EUsdPurpose)StageActor->PurposesToLoad ^ Purpose));
+						if (StageActor->IsTemplate())
 						{
 							StageActor->SaveConfig();
 						}
@@ -1277,17 +1347,10 @@ void SUsdStage::FillRenderContextSubMenu( FMenuBuilder& MenuBuilder )
 						));
 
 						// c.f. comment in SUsdStage::FillCollapsingSubMenu
-						TGuardValue<bool> MaintainSelectionGuard( bUpdatingViewportSelection, true );
-
-						StageActor->Modify();
-						StageActor->RenderContext = RenderContext;
-
-						FPropertyChangedEvent PropertyChangedEvent(
-							FindFieldChecked< FProperty >( StageActor->GetClass(), GET_MEMBER_NAME_CHECKED( AUsdStageActor, RenderContext ) )
-						);
-						StageActor->PostEditChangeProperty(PropertyChangedEvent);
+						TGuardValue<bool> MaintainSelectionGuard(bUpdatingViewportSelection, true);
 
-						if ( StageActor->IsTemplate() )
+						StageActor->SetRenderContext(RenderContext);
+						if (StageActor->IsTemplate())
 						{
 							StageActor->SaveConfig();
 						}
@@ -1383,8 +1446,11 @@ void SUsdStage::FillMaterialPurposeSubMenu( FMenuBuilder& MenuBuilder )
 						FText::FromString(ChosenOption.IsValid() ? *ChosenOption : TEXT(""))
 					));
 
-					StageActor->SetMaterialPurpose( **ChosenOption );
-					if ( StageActor->IsTemplate() )
+					// c.f. comment in SUsdStage::FillCollapsingSubMenu
+					TGuardValue<bool> MaintainSelectionGuard(bUpdatingViewportSelection, true);
+
+					StageActor->SetMaterialPurpose(**ChosenOption);
+					if (StageActor->IsTemplate())
 					{
 						StageActor->SaveConfig();
 					}
@@ -1443,8 +1509,8 @@ void SUsdStage::FillMaterialPurposeSubMenu( FMenuBuilder& MenuBuilder )
 						NewText
 					));
 
-					StageActor->SetMaterialPurpose( NewPurpose );
-					if ( StageActor->IsTemplate() )
+					StageActor->SetMaterialPurpose(NewPurpose);
+					if (StageActor->IsTemplate())
 					{
 						StageActor->SaveConfig();
 					}
@@ -1479,11 +1545,10 @@ void SUsdStage::FillRootMotionSubMenu( FMenuBuilder& MenuBuilder )
 						));
 
 						// c.f. comment in SUsdStage::FillCollapsingSubMenu
-						TGuardValue<bool> MaintainSelectionGuard( bUpdatingViewportSelection, true );
+						TGuardValue<bool> MaintainSelectionGuard(bUpdatingViewportSelection, true);
 
-						StageActor->Modify();
-						StageActor->SetRootMotionHandling( HandlingStrategy );
-						if ( StageActor->IsTemplate() )
+						StageActor->SetRootMotionHandling(HandlingStrategy);
+						if (StageActor->IsTemplate())
 						{
 							StageActor->SaveConfig();
 						}
@@ -1551,9 +1616,11 @@ void SUsdStage::FillCollapsingSubMenu( FMenuBuilder& MenuBuilder )
 						FText::FromString( StageActor->GetActorLabel() )
 					) );
 
-					TGuardValue<bool> MaintainSelectionGuard( bUpdatingViewportSelection, true );
-					StageActor->SetMergeIdenticalMaterialSlots( !StageActor->bMergeIdenticalMaterialSlots );
-					if ( StageActor->IsTemplate() )
+					// c.f. comment within AddKindToCollapseEntry just below
+					TGuardValue<bool> MaintainSelectionGuard(bUpdatingViewportSelection, true);
+
+					StageActor->SetMergeIdenticalMaterialSlots(!StageActor->bMergeIdenticalMaterialSlots);
+					if (StageActor->IsTemplate())
 					{
 						StageActor->SaveConfig();
 					}
@@ -1596,11 +1663,10 @@ void SUsdStage::FillCollapsingSubMenu( FMenuBuilder& MenuBuilder )
 						// Not only is this "selection spam" is very visible on the USD Stage Editor, if our originally selected component
 						// was just collapsed away, we'd end up updating our prim selection to point to the parent prim instead
 						// (the collapsing root), which is not ideal.
-						TGuardValue<bool> MaintainSelectionGuard( bUpdatingViewportSelection, true );
+						TGuardValue<bool> MaintainSelectionGuard(bUpdatingViewportSelection, true);
 
-						int32 NewKindsToCollapse = ( int32 ) ( ( EUsdDefaultKind ) StageActor->KindsToCollapse ^ Kind );
-						StageActor->SetKindsToCollapse( NewKindsToCollapse );
-						if ( StageActor->IsTemplate() )
+						StageActor->SetKindsToCollapse((int32)((EUsdDefaultKind)StageActor->KindsToCollapse ^ Kind));
+						if (StageActor->IsTemplate())
 						{
 							StageActor->SaveConfig();
 						}
@@ -1707,11 +1773,10 @@ void SUsdStage::FillInterpolationTypeSubMenu(FMenuBuilder& MenuBuilder)
 					));
 
 					// c.f. comment in SUsdStage::FillCollapsingSubMenu
-					TGuardValue<bool> MaintainSelectionGuard( bUpdatingViewportSelection, true );
-
-					StageActor->SetInterpolationType( EUsdInterpolationType::Linear );
+					TGuardValue<bool> MaintainSelectionGuard(bUpdatingViewportSelection, true);
 
-					if ( StageActor->IsTemplate() )
+					StageActor->SetInterpolationType(EUsdInterpolationType::Linear);
+					if (StageActor->IsTemplate())
 					{
 						StageActor->SaveConfig();
 					}
@@ -1746,11 +1811,10 @@ void SUsdStage::FillInterpolationTypeSubMenu(FMenuBuilder& MenuBuilder)
 					));
 
 					// c.f. comment in SUsdStage::FillCollapsingSubMenu
-					TGuardValue<bool> MaintainSelectionGuard( bUpdatingViewportSelection, true );
+					TGuardValue<bool> MaintainSelectionGuard(bUpdatingViewportSelection, true);
 
-					StageActor->SetInterpolationType( EUsdInterpolationType::Held );
-
-					if ( StageActor->IsTemplate() )
+					StageActor->SetInterpolationType(EUsdInterpolationType::Held);
+					if (StageActor->IsTemplate())
 					{
 						StageActor->SaveConfig();
 					}
@@ -2578,12 +2642,10 @@ void SUsdStage::OnNaniteTriangleThresholdValueCommitted( int32 InValue, ETextCom
 	) );
 
 	// c.f. comment in SUsdStage::FillCollapsingSubMenu
-	TGuardValue<bool> MaintainSelectionGuard( bUpdatingViewportSelection, true );
+	TGuardValue<bool> MaintainSelectionGuard(bUpdatingViewportSelection, true);
 
-	StageActor->SetNaniteTriangleThreshold( InValue );
+	StageActor->SetNaniteTriangleThreshold(InValue);
 	CurrentNaniteThreshold = InValue;
-
-	if ( StageActor->IsTemplate() )
 	{
 		StageActor->SaveConfig();
 	}
diff --git a/Engine/Plugins/Importers/USDImporter/Source/USDStageEditor/Private/SUSDStage.h b/Engine/Plugins/Importers/USDImporter/Source/USDStageEditor/Private/SUSDStage.h
--- a/Engine/Plugins/Importers/USDImporter/Source/USDStageEditor/Private/SUSDStage.h
+++ b/Engine/Plugins/Importers/USDImporter/Source/USDStageEditor/Private/SUSDStage.h
@@ -83,6 +83,7 @@ protected:
 	void FillActionsMenu( FMenuBuilder& MenuBuilder );
 	void FillOptionsMenu( FMenuBuilder& MenuBuilder );
 	void FillExportSubMenu( FMenuBuilder& MenuBuilder );
+	void FillStageStateSubMenu(FMenuBuilder& MenuBuilder);
 	void FillPayloadsSubMenu( FMenuBuilder& MenuBuilder );
 	void FillPurposesToLoadSubMenu( FMenuBuilder& MenuBuilder );
 	void FillRenderContextSubMenu( FMenuBuilder& MenuBuilder );
diff --git a/Engine/Plugins/Importers/USDImporter/Source/USDStageEditorViewModels/Private/USDStageViewModel.cpp b/Engine/Plugins/Importers/USDImporter/Source/USDStageEditorViewModels/Private/USDStageViewModel.cpp
--- a/Engine/Plugins/Importers/USDImporter/Source/USDStageEditorViewModels/Private/USDStageViewModel.cpp
+++ b/Engine/Plugins/Importers/USDImporter/Source/USDStageEditorViewModels/Private/USDStageViewModel.cpp
@@ -109,7 +109,7 @@ void FUsdStageViewModel::ReloadStage()
 	}
 
 #if USE_USD_SDK
-	UE::FUsdStage Stage = UsdStageActor->GetOrLoadUsdStage();
+	UE::FUsdStage Stage = UsdStageActor->GetOrOpenUsdStage();
 	pxr::UsdStageRefPtr UsdStage = pxr::UsdStageRefPtr( Stage );
 
 	// Can't reload from disk something that doesn't exist on disk yet
@@ -183,7 +183,7 @@ void FUsdStageViewModel::ResetStage()
 		return;
 	}
 
-	UE::FUsdStage Stage = UsdStageActor->GetOrLoadUsdStage();
+	UE::FUsdStage Stage = UsdStageActor->GetOrOpenUsdStage();
 	pxr::UsdStageRefPtr UsdStage = pxr::UsdStageRefPtr( Stage );
 
 	if ( UsdStage )
@@ -212,7 +212,7 @@ void FUsdStageViewModel::SaveStage()
 #if USE_USD_SDK
 	if ( UsdStageActor.IsValid() )
 	{
-		if ( UE::FUsdStage UsdStage = UsdStageActor->GetOrLoadUsdStage() )
+		if ( UE::FUsdStage UsdStage = UsdStageActor->GetOrOpenUsdStage() )
 		{
 			FScopedUsdAllocs UsdAllocs;
 
@@ -248,7 +248,7 @@ void FUsdStageViewModel::SaveStageAs( const TCHAR* FilePath )
 	AUsdStageActor* StageActorPtr = UsdStageActor.Get();
 	if ( StageActorPtr )
 	{
-		if ( UE::FUsdStage UsdStage = UsdStageActor->GetOrLoadUsdStage() )
+		if ( UE::FUsdStage UsdStage = UsdStageActor->GetOrOpenUsdStage() )
 		{
 			UsdUtils::StartMonitoringErrors();
 
@@ -297,7 +297,7 @@ void FUsdStageViewModel::ImportStage( const TCHAR* TargetContentFolder, UUsdStag
 		return;
 	}
 
-	const UE::FUsdStage UsdStage = StageActor->GetOrLoadUsdStage();
+	const UE::FUsdStage UsdStage = StageActor->GetOrOpenUsdStage();
 	if ( !UsdStage )
 	{
 		return;
diff --git a/Engine/Plugins/Importers/USDImporter/Source/USDUtilities/Private/USDConversionUtils.cpp b/Engine/Plugins/Importers/USDImporter/Source/USDUtilities/Private/USDConversionUtils.cpp
--- a/Engine/Plugins/Importers/USDImporter/Source/USDUtilities/Private/USDConversionUtils.cpp
+++ b/Engine/Plugins/Importers/USDImporter/Source/USDUtilities/Private/USDConversionUtils.cpp
@@ -360,10 +360,10 @@ void UsdUtils::SetUsdStageMetersPerUnit( const pxr::UsdStageRefPtr& Stage, float
 	pxr::UsdGeomSetStageMetersPerUnit( Stage, MetersPerUnit );
 }
 
-int32 UsdUtils::GetUsdStageNumFrames( const pxr::UsdStageRefPtr& Stage )
+int32 UsdUtils::GetUsdStageNumFrames(const pxr::UsdStageRefPtr& Stage)
 {
 	// USD time code range is inclusive on both ends
-	return FMath::Abs( FMath::CeilToInt32( Stage->GetEndTimeCode() ) - FMath::FloorToInt32( Stage->GetStartTimeCode() ) + 1 );
+	return Stage ? FMath::Abs(FMath::CeilToInt32(Stage->GetEndTimeCode()) - FMath::FloorToInt32(Stage->GetStartTimeCode()) + 1) : 0;
 }
 
 bool UsdUtils::HasCompositionArcs( const pxr::UsdPrim& Prim )
diff --git a/Engine/Plugins/Importers/USDImporter/Source/UnrealUSDWrapper/Public/UnrealUSDWrapper.h b/Engine/Plugins/Importers/USDImporter/Source/UnrealUSDWrapper/Public/UnrealUSDWrapper.h
--- a/Engine/Plugins/Importers/USDImporter/Source/UnrealUSDWrapper/Public/UnrealUSDWrapper.h
+++ b/Engine/Plugins/Importers/USDImporter/Source/UnrealUSDWrapper/Public/UnrealUSDWrapper.h
@@ -124,7 +124,7 @@ enum class EUsdInterpolationType : uint8
 };
 
 UENUM()
-enum class EUsdRootMotionHandling
+enum class EUsdRootMotionHandling : uint8
 {
 	// Use for the root bone just its regular joint animation as described on the SkelAnimation prim.
 	NoAdditionalRootMotion,
