USD: Add support for handling attach sockets when exporting transforms to USD.

diff --git a/Engine/Plugins/Importers/USDImporter/Source/USDUtilities/Private/USDPrimConversion.cpp b/Engine/Plugins/Importers/USDImporter/Source/USDUtilities/Private/USDPrimConversion.cpp
--- a/Engine/Plugins/Importers/USDImporter/Source/USDUtilities/Private/USDPrimConversion.cpp
+++ b/Engine/Plugins/Importers/USDImporter/Source/USDUtilities/Private/USDPrimConversion.cpp
@@ -2484,6 +2484,18 @@ bool UnrealToUsd::ConvertSceneComponent( const pxr::UsdStageRefPtr& Stage, const
 
 	FTransform RelativeTransform = SceneComponent->GetRelativeTransform();
 
+	// If we're attached to a socket our RelativeTransform will be relative to the socket, instead of the parent
+	// component space. Here we get the socket transform itself and concatenate it, forcing RelativeTransform to
+	// be relative to the parent component space again.
+	if (USceneComponent* Parent = SceneComponent->GetAttachParent())
+	{
+		if (FName SocketName = SceneComponent->GetAttachSocketName(); SocketName != NAME_None)
+		{
+			FTransform SocketTransform = Parent->GetSocketTransform(SocketName, RTS_Component);
+			RelativeTransform = RelativeTransform * SocketTransform;
+		}
+	}
+
 	// Compensate different orientation for light or camera components:
 	// In USD cameras shoot towards local - Z, with + Y up.Lights also emit towards local - Z, with + Y up
 	// In UE cameras shoot towards local + X, with + Z up.Lights also emit towards local + X, with + Z up
@@ -3284,7 +3296,23 @@ bool UnrealToUsd::CreateComponentPropertyBaker( UE::FUsdPrim& Prim, const UScene
 			{
 				FScopedUsdAllocs Allocs;
 
-				FTransform FinalUETransform = AdditionalRotation * Component.GetRelativeTransform();
+				FTransform RelativeTransform = Component.GetRelativeTransform();
+
+				// If we're attached to a socket our RelativeTransform will be relative to the socket, instead of the parent
+				// component space. Here we get the socket transform itself and concatenate it, forcing RelativeTransform to
+				// be relative to the parent component space again.
+				// It may seem wasteful to do this inside the baker function, but you can place "Attach tracks" on the
+				// Sequencer that may make the attach socket change every frame, so we do need this
+				if (USceneComponent* Parent = Component.GetAttachParent())
+				{
+					if (FName SocketName = Component.GetAttachSocketName(); SocketName != NAME_None)
+					{
+						FTransform SocketTransform = Parent->GetSocketTransform(SocketName, RTS_Component);
+						RelativeTransform = RelativeTransform * SocketTransform;
+					}
+				}
+
+				FTransform FinalUETransform = AdditionalRotation * RelativeTransform;
 				pxr::GfMatrix4d UsdTransform = UnrealToUsd::ConvertTransform( StageInfo, FinalUETransform );
 				Attr.Set< pxr::GfMatrix4d>( UsdTransform, UsdTimeCode );
 			};
