USD: Fix static analysis warnings about cl 28836021.

Static analysis job on Horde: https://horde.devtools.epicgames.com/job/652ea19cc83fdc53e699d5ac?step=fd40

#rnx

diff --git a/Engine/Plugins/Importers/USDImporter/Source/USDClasses/Private/USDDrawModeComponent.cpp b/Engine/Plugins/Importers/USDImporter/Source/USDClasses/Private/USDDrawModeComponent.cpp
--- a/Engine/Plugins/Importers/USDImporter/Source/USDClasses/Private/USDDrawModeComponent.cpp
+++ b/Engine/Plugins/Importers/USDImporter/Source/USDClasses/Private/USDDrawModeComponent.cpp
@@ -219,7 +219,7 @@ namespace UE::UsdDrawModeComponentImpl::Private
 
 							FMeshBatch& Mesh = Collector.AllocateMesh();
 							Mesh.VertexFactory = &HitProxyBuffer.VertexFactory;
-							Mesh.MaterialRenderProxy = HitProxyBuffer.Material->GetRenderProxy();
+							Mesh.MaterialRenderProxy = HitProxyBuffer.Material ? HitProxyBuffer.Material->GetRenderProxy() : nullptr;
 							Mesh.ReverseCulling = IsLocalToWorldDeterminantNegative();
 							Mesh.Type = PT_TriangleList;
 							Mesh.DepthPriorityGroup = DepthPriority;
@@ -507,7 +507,13 @@ namespace UE::UsdDrawModeComponentImpl::Private
 				Buffer.TriangleCount = NumTriangles;
 
 				Buffer.Material = MaterialInstances[FaceIndex];
-				ensure(Buffer.Material);	// The component should always have a material for us, even if just the default material
+
+				// The component should always have a material for us, even if just the default material
+				if (!ensure(Buffer.Material))
+				{
+					return;
+				}
+
 				MaterialRelevances.Emplace(Buffer.Material->GetRelevance_Concurrent(GetScene().GetFeatureLevel()));
 
 				ENQUEUE_RENDER_COMMAND(FUsdCardsSceneProxyUploadBuffer)
@@ -584,7 +590,9 @@ namespace UE::UsdDrawModeComponentImpl::Private
 						{
 							FMeshBatch& Mesh = Collector.AllocateMesh();
 							Mesh.VertexFactory = &Buffer.VertexFactory;
-							Mesh.MaterialRenderProxy = OverrideMaterialProxy ? OverrideMaterialProxy : Buffer.Material->GetRenderProxy();
+							Mesh.MaterialRenderProxy = OverrideMaterialProxy ? OverrideMaterialProxy
+													   : Buffer.Material	 ? Buffer.Material->GetRenderProxy()
+																			 : nullptr;
 							Mesh.ReverseCulling = IsLocalToWorldDeterminantNegative();
 							Mesh.Type = PT_TriangleList;
 							Mesh.DepthPriorityGroup = (uint8)GetDepthPriorityGroup(Views[ViewIndex]);
