USD Import - Fixed crash when the context level is null.

#rnx

diff --git a/Engine/Plugins/Importers/USDImporter/Source/USDSchemas/Private/USDShadeMaterialTranslator.cpp b/Engine/Plugins/Importers/USDImporter/Source/USDSchemas/Private/USDShadeMaterialTranslator.cpp
--- a/Engine/Plugins/Importers/USDImporter/Source/USDSchemas/Private/USDShadeMaterialTranslator.cpp
+++ b/Engine/Plugins/Importers/USDImporter/Source/USDSchemas/Private/USDShadeMaterialTranslator.cpp
@@ -352,7 +352,7 @@ void FUsdShadeMaterialTranslator::CreateAssets()
 					// render states for some components, and UWorld::AddToWorld calls FScene::AddPrimitive which expects the component to not have
 					// primitives yet
 					FMaterialUpdateContext::EOptions::Type Options = FMaterialUpdateContext::EOptions::Default;
-					if ( Context->Level->bIsAssociatingLevel )
+					if ( Context->Level && Context->Level->bIsAssociatingLevel )
 					{
 						Options = ( FMaterialUpdateContext::EOptions::Type ) ( Options & ~FMaterialUpdateContext::EOptions::RecreateRenderStates );
 					}
