USD: Added support for geometry cache export.

diff --git a/Engine/Plugins/Importers/USDImporter/Source/USDExporter/Private/GeometryCacheExporterUSD.cpp b/Engine/Plugins/Importers/USDImporter/Source/USDExporter/Private/GeometryCacheExporterUSD.cpp
new file mode 100644
--- /dev/null
+++ b/Engine/Plugins/Importers/USDImporter/Source/USDExporter/Private/GeometryCacheExporterUSD.cpp
@@ -0,0 +1,348 @@
+// Copyright Epic Games, Inc. All Rights Reserved.
+
+#include "GeometryCacheExporterUSD.h"
+
+#include "AssetExportTask.h"
+#include "EngineAnalytics.h"
+#include "GeometryCache.h"
+#include "GeometryCacheExporterUSDOptions.h"
+#include "MaterialExporterUSD.h"
+#include "Misc/EngineVersion.h"
+#include "UnrealUSDWrapper.h"
+#include "USDConversionUtils.h"
+#include "USDExporterModule.h"
+#include "USDGeomMeshConversion.h"
+#include "USDLog.h"
+#include "USDOptionsWindow.h"
+#include "USDUnrealAssetInfo.h"
+
+#include "UsdWrappers/SdfLayer.h"
+#include "UsdWrappers/SdfPath.h"
+#include "UsdWrappers/UsdPrim.h"
+#include "UsdWrappers/UsdStage.h"
+
+namespace UE::GeometryCacheExporterUSD::Private
+{
+	void SendAnalytics(
+		UObject* Asset,
+		UGeometryCacheExporterUSDOptions* Options,
+		bool bAutomated,
+		double ElapsedSeconds,
+		double NumberOfFrames,
+		const FString& Extension
+	)
+	{
+		if (!Asset || !FEngineAnalytics::IsAvailable())
+		{
+			return;
+		}
+
+		FString ClassName = Asset->GetClass()->GetName();
+
+		TArray<FAnalyticsEventAttribute> EventAttributes;
+		EventAttributes.Emplace(TEXT("AssetType"), ClassName);
+
+		if (Options)
+		{
+			UsdUtils::AddAnalyticsAttributes(*Options, EventAttributes);
+		}
+
+		IUsdClassesModule::SendAnalytics(
+			MoveTemp(EventAttributes),
+			FString::Printf(TEXT("Export.%s"), *ClassName),
+			bAutomated,
+			ElapsedSeconds,
+			NumberOfFrames,
+			Extension
+		);
+	}
+}
+
+UGeometryCacheExporterUSD::UGeometryCacheExporterUSD()
+{
+#if USE_USD_SDK
+	for (const FString& Extension : UnrealUSDWrapper::GetNativeFileFormats())
+	{
+		// USDZ is not supported for writing for now
+		if (Extension.Equals(TEXT("usdz")))
+		{
+			continue;
+		}
+
+		FormatExtension.Add(Extension);
+		FormatDescription.Add(TEXT("Universal Scene Description file"));
+	}
+	SupportedClass = UGeometryCache::StaticClass();
+	bText = false;
+#endif // #if USE_USD_SDK
+}
+
+bool UGeometryCacheExporterUSD::ExportBinary(UObject* Object, const TCHAR* Type, FArchive& Ar, FFeedbackContext* Warn, int32 FileIndex, uint32 PortFlags)
+{
+#if USE_USD_SDK
+	UGeometryCache* GeometryCache = CastChecked<UGeometryCache>(Object);
+	if (!GeometryCache)
+	{
+		return false;
+	}
+
+	UGeometryCacheExporterUSDOptions* Options = nullptr;
+	if (ExportTask)
+	{
+		Options = Cast<UGeometryCacheExporterUSDOptions>(ExportTask->Options);
+	}
+	if (!Options)
+	{
+		Options = GetMutableDefault<UGeometryCacheExporterUSDOptions>();
+
+		// Prompt with an options dialog if we can
+		if (Options && (!ExportTask || !ExportTask->bAutomated))
+		{
+			Options->MeshAssetOptions.MaterialBakingOptions.TexturesDir.Path = FPaths::Combine(FPaths::GetPath(UExporter::CurrentFilename), TEXT("Textures"));
+
+			const bool bContinue = SUsdOptionsWindow::ShowExportOptions(*Options);
+			if (!bContinue)
+			{
+				return false;
+			}
+		}
+	}
+	if (!Options)
+	{
+		return false;
+	}
+
+	// The intended export workflow has us exporting into the Ar argument, and would only later actually write to the
+	// file, potentially prompting the user beforehand if they want to overwrite an existing file or not (check
+	// UExporter::RunAssetExportTask). On the USD export workflows we always write to the file directly ourselves as
+	// its easier to just let the USD SDK do it, which means that the "Do you want to overwrite?" file prompt happens
+	// way after we already saved everything, and has no effect. Lets force bPrompt to false then, so this dialog
+	// doesn't appear under any circumstances
+	ExportTask->bPrompt = false;
+
+	// If bUsePayload is true, we'll intercept the filename so that we write the mesh data to
+	// "C:/MyFolder/file_payload.usda" and create an "asset" file "C:/MyFolder/file.usda" that uses it
+	// as a payload, pointing at the default prim
+	FString PayloadFilename = UExporter::CurrentFilename;
+	if (Options->MeshAssetOptions.bUsePayload)
+	{
+		FString PathPart;
+		FString FilenamePart;
+		FString ExtensionPart;
+		FPaths::Split(PayloadFilename, PathPart, FilenamePart, ExtensionPart);
+
+		if (FormatExtension.Contains(Options->MeshAssetOptions.PayloadFormat))
+		{
+			ExtensionPart = Options->MeshAssetOptions.PayloadFormat;
+		}
+
+		PayloadFilename = FPaths::Combine(PathPart, FilenamePart + TEXT("_payload.") + ExtensionPart);
+	}
+
+	if (!IUsdExporterModule::CanExportToLayer(UExporter::CurrentFilename) ||
+		(Options->MeshAssetOptions.bUsePayload && !IUsdExporterModule::CanExportToLayer(PayloadFilename)))
+	{
+		return false;
+	}
+
+	// Get a simple GUID hash/identifier of our mesh
+	FString DDCKeyHash;
+	const FString& DDCKey = GeometryCache->GetHash();
+	FSHA1 SHA1;
+	SHA1.UpdateWithString(*DDCKey, DDCKey.Len());
+	UsdUtils::HashForGeometryCacheExport(*Options, SHA1);
+
+	SHA1.Final();
+	FSHAHash Hash;
+	SHA1.GetHash(&Hash.Hash[0]);
+	DDCKeyHash = Hash.ToString();
+
+	// Check if we already have exported what we plan on exporting anyway
+	if (FPaths::FileExists(UExporter::CurrentFilename) && FPaths::FileExists(PayloadFilename))
+	{
+		if (!ExportTask->bReplaceIdentical)
+		{
+			UE_LOG(LogUsd, Log,
+				TEXT("Skipping export of asset '%s' as the target file '%s' already exists."),
+				*Object->GetPathName(),
+				*UExporter::CurrentFilename
+			);
+			return false;
+		}
+		// If we don't want to re-export this asset we need to check if its the same version
+		else if (!Options->bReExportIdenticalAssets)
+		{
+			bool bSkipMeshExport = false;
+
+			// Don't use the stage cache here as we want this stage to close within this scope in case
+			// we have to overwrite its files due to e.g. missing payload or anything like that
+			const bool bUseStageCache = false;
+			const EUsdInitialLoadSet InitialLoadSet = EUsdInitialLoadSet::LoadNone;
+			if (UE::FUsdStage TempStage = UnrealUSDWrapper::OpenStage(*UExporter::CurrentFilename, InitialLoadSet, bUseStageCache))
+			{
+				if (UE::FUsdPrim DefaultPrim = TempStage.GetDefaultPrim())
+				{
+					FUsdUnrealAssetInfo Info = UsdUtils::GetPrimAssetInfo(DefaultPrim);
+
+					const bool bVersionMatches = !Info.Version.IsEmpty() && Info.Version == DDCKeyHash;
+
+					const bool bAssetTypeMatches = !Info.UnrealAssetType.IsEmpty()
+						&& Info.UnrealAssetType == GeometryCache->GetClass()->GetName();
+
+					if (bVersionMatches && bAssetTypeMatches)
+					{
+						UE_LOG(LogUsd, Log,
+							TEXT("Skipping export of asset '%s' as the target file '%s' already contains up-to-date exported data."),
+							*GeometryCache->GetPathName(),
+							*UExporter::CurrentFilename
+						);
+
+						bSkipMeshExport = true;
+					}
+				}
+			}
+
+			if (bSkipMeshExport)
+			{
+				// Even if we're not going to export the mesh, we may still need to re-bake materials
+				if (Options->MeshAssetOptions.bBakeMaterials)
+				{
+					TSet<UMaterialInterface*> MaterialsToBake;
+					for (UMaterialInterface* MaterialInterface : GeometryCache->Materials)
+					{
+						MaterialsToBake.Add(MaterialInterface);
+					}
+
+					const bool bIsAssetLayer = true;
+					UMaterialExporterUsd::ExportMaterialsForStage(
+						MaterialsToBake.Array(),
+						Options->MeshAssetOptions.MaterialBakingOptions,
+						UExporter::CurrentFilename,
+						bIsAssetLayer,
+						Options->MeshAssetOptions.bUsePayload,
+						ExportTask->bReplaceIdentical,
+						Options->bReExportIdenticalAssets,
+						ExportTask->bAutomated
+					);
+				}
+
+				return true;
+			}
+		}
+	}
+
+	double StartTime = FPlatformTime::Cycles64();
+
+	// UsdStage is the payload stage when exporting with payloads, or just the single stage otherwise
+	UE::FUsdStage UsdStage = UnrealUSDWrapper::NewStage(*PayloadFilename);
+	if (!UsdStage)
+	{
+		return false;
+	}
+
+	UsdUtils::SetUsdStageMetersPerUnit(UsdStage, Options->StageOptions.MetersPerUnit);
+	UsdUtils::SetUsdStageUpAxis(UsdStage, Options->StageOptions.UpAxis);
+
+	FString RootPrimPath = (TEXT("/") + UsdUtils::SanitizeUsdIdentifier(*GeometryCache->GetName()));
+
+	UE::FUsdPrim RootPrim = UsdStage.DefinePrim(UE::FSdfPath(*RootPrimPath), TEXT("Mesh"));
+	if (!RootPrim)
+	{
+		return false;
+	}
+
+	UsdStage.SetDefaultPrim(RootPrim);
+
+	// Asset stage always the stage where we write the material assignments
+	UE::FUsdStage AssetStage;
+
+	// Using payload: Convert mesh data through the asset stage (that references the payload) so that we can
+	// author mesh data on the payload layer and material data on the asset layer
+	if (Options->MeshAssetOptions.bUsePayload)
+	{
+		AssetStage = UnrealUSDWrapper::NewStage(*UExporter::CurrentFilename);
+		if (AssetStage)
+		{
+			UsdUtils::SetUsdStageMetersPerUnit(AssetStage, Options->StageOptions.MetersPerUnit);
+			UsdUtils::SetUsdStageUpAxis(AssetStage, Options->StageOptions.UpAxis);
+
+			if (UE::FUsdPrim AssetRootPrim = AssetStage.DefinePrim(UE::FSdfPath(*RootPrimPath)))
+			{
+				AssetStage.SetDefaultPrim(AssetRootPrim);
+
+				UsdUtils::AddPayload(AssetRootPrim, *PayloadFilename);
+			}
+		}
+	}
+	// Not using payload: Just author everything on the current edit target of the payload (== asset) layer
+	else
+	{
+		AssetStage = UsdStage;
+	}
+
+	UnrealToUsd::ConvertGeometryCache(GeometryCache, RootPrim, &AssetStage);
+
+	// Write asset info now that we finished exporting
+	if( UE::FUsdPrim AssetDefaultPrim = AssetStage.GetDefaultPrim() )
+	{
+		FUsdUnrealAssetInfo Info;
+		Info.Name = GeometryCache->GetName();
+		Info.Identifier = UExporter::CurrentFilename;
+		Info.Version = DDCKeyHash;
+		Info.UnrealContentPath = GeometryCache->GetPathName();
+		Info.UnrealAssetType = GeometryCache->GetClass()->GetName();
+		Info.UnrealExportTime = FDateTime::Now().ToString();
+		Info.UnrealEngineVersion = FEngineVersion::Current().ToString();
+
+		UsdUtils::SetPrimAssetInfo(AssetDefaultPrim, Info);
+	}
+
+	// Bake materials and replace unrealMaterials with references to the baked files.
+	if (Options->MeshAssetOptions.bBakeMaterials)
+	{
+		TSet<UMaterialInterface*> MaterialsToBake;
+		for (UMaterialInterface* MaterialInterface : GeometryCache->Materials)
+		{
+			MaterialsToBake.Add(MaterialInterface);
+		}
+
+		const bool bIsAssetLayer = true;
+		UMaterialExporterUsd::ExportMaterialsForStage(
+			MaterialsToBake.Array(),
+			Options->MeshAssetOptions.MaterialBakingOptions,
+			AssetStage.GetRootLayer().GetRealPath(),
+			bIsAssetLayer,
+			Options->MeshAssetOptions.bUsePayload,
+			ExportTask->bReplaceIdentical,
+			Options->bReExportIdenticalAssets,
+			ExportTask->bAutomated
+		);
+	}
+
+	UsdStage.GetRootLayer().Save();
+	if (AssetStage && UsdStage != AssetStage)
+	{
+		AssetStage.GetRootLayer().Save();
+	}
+
+	// Analytics
+	{
+		bool bAutomated = ExportTask ? ExportTask->bAutomated : false;
+		double ElapsedSeconds = FPlatformTime::ToSeconds64(FPlatformTime::Cycles64() - StartTime);
+		FString Extension = FPaths::GetExtension(UExporter::CurrentFilename);
+
+		UE::GeometryCacheExporterUSD::Private::SendAnalytics(
+			Object,
+			Options,
+			bAutomated,
+			ElapsedSeconds,
+			UsdUtils::GetUsdStageNumFrames(AssetStage),
+			Extension
+		);
+	}
+
+	return true;
+#else
+	return false;
+#endif // #if USE_USD_SDK
+}
diff --git a/Engine/Plugins/Importers/USDImporter/Source/USDExporter/Private/GeometryCacheExporterUSDOptions.cpp b/Engine/Plugins/Importers/USDImporter/Source/USDExporter/Private/GeometryCacheExporterUSDOptions.cpp
new file mode 100644
--- /dev/null
+++ b/Engine/Plugins/Importers/USDImporter/Source/USDExporter/Private/GeometryCacheExporterUSDOptions.cpp
@@ -0,0 +1,21 @@
+// Copyright Epic Games, Inc. All Rights Reserved.
+
+#include "GeometryCacheExporterUSDOptions.h"
+
+#include "AnalyticsEventAttribute.h"
+
+void UsdUtils::AddAnalyticsAttributes(
+	const UGeometryCacheExporterUSDOptions& Options,
+	TArray< FAnalyticsEventAttribute >& InOutAttributes
+)
+{
+	UsdUtils::AddAnalyticsAttributes(Options.StageOptions, InOutAttributes);
+	UsdUtils::AddAnalyticsAttributes(Options.MeshAssetOptions, InOutAttributes);
+	InOutAttributes.Emplace(TEXT("ReExportIdenticalAssets"), Options.bReExportIdenticalAssets);
+}
+
+void UsdUtils::HashForGeometryCacheExport(const UGeometryCacheExporterUSDOptions& Options, FSHA1& HashToUpdate)
+{
+	UsdUtils::HashForExport(Options.StageOptions, HashToUpdate);
+	UsdUtils::HashForMeshExport(Options.MeshAssetOptions, HashToUpdate);
+}
diff --git a/Engine/Plugins/Importers/USDImporter/Source/USDExporter/Public/GeometryCacheExporterUSD.h b/Engine/Plugins/Importers/USDImporter/Source/USDExporter/Public/GeometryCacheExporterUSD.h
new file mode 100644
--- /dev/null
+++ b/Engine/Plugins/Importers/USDImporter/Source/USDExporter/Public/GeometryCacheExporterUSD.h
@@ -0,0 +1,20 @@
+// Copyright Epic Games, Inc. All Rights Reserved.
+
+#pragma once
+
+#include "Exporters/Exporter.h"
+
+#include "GeometryCacheExporterUSD.generated.h"
+
+UCLASS()
+class UGeometryCacheExporterUSD : public UExporter
+{
+	GENERATED_BODY()
+
+public:
+	UGeometryCacheExporterUSD();
+
+	//~ Begin UExporter Interface
+	virtual bool ExportBinary(UObject* Object, const TCHAR* Type, FArchive& Ar, FFeedbackContext* Warn, int32 FileIndex = 0, uint32 PortFlags = 0) override;
+	//~ End UExporter Interface
+};
\ No newline at end of file
diff --git a/Engine/Plugins/Importers/USDImporter/Source/USDExporter/Public/GeometryCacheExporterUSDOptions.h b/Engine/Plugins/Importers/USDImporter/Source/USDExporter/Public/GeometryCacheExporterUSDOptions.h
new file mode 100644
--- /dev/null
+++ b/Engine/Plugins/Importers/USDImporter/Source/USDExporter/Public/GeometryCacheExporterUSDOptions.h
@@ -0,0 +1,48 @@
+// Copyright Epic Games, Inc. All Rights Reserved.
+
+#pragma once
+
+#include "USDAssetOptions.h"
+#include "USDStageOptions.h"
+
+#include "Engine/EngineTypes.h"
+
+#include "GeometryCacheExporterUSDOptions.generated.h"
+
+struct FAnalyticsEventAttribute;
+
+/**
+ * Options for exporting geometry caches to USD format.
+ */
+UCLASS(Config = Editor, Blueprintable)
+class USDEXPORTER_API UGeometryCacheExporterUSDOptions : public UObject
+{
+	GENERATED_BODY()
+
+public:
+	UPROPERTY(EditAnywhere, config, BlueprintReadWrite, Category = "Stage options", meta = (ShowOnlyInnerProperties))
+	FUsdStageOptions StageOptions;
+
+	UPROPERTY(EditAnywhere, config, BlueprintReadWrite, Category = "Mesh options", meta = (ShowOnlyInnerProperties))
+	FUsdMeshAssetOptions MeshAssetOptions;
+
+	/**
+	 * Whether to export any asset (StaticMesh, Material, etc.) even if the existing file already describes the same version of a compatible asset.
+	 * This is only checked when bReplaceIdentical is set on the asset export task. Otherwise we'll never overwrite files.
+	 */
+	UPROPERTY(EditAnywhere, config, BlueprintReadWrite, Category = "Collision", meta = (DisplayName = "Re-export Identical Assets"))
+	bool bReExportIdenticalAssets = false;
+};
+
+namespace UsdUtils
+{
+	USDEXPORTER_API void AddAnalyticsAttributes(
+		const UGeometryCacheExporterUSDOptions& Options,
+		TArray< FAnalyticsEventAttribute >& InOutAttributes
+	);
+
+	USDEXPORTER_API void HashForGeometryCacheExport(
+		const UGeometryCacheExporterUSDOptions& Options,
+		FSHA1& HashToUpdate
+	);
+}
diff --git a/Engine/Plugins/Importers/USDImporter/Source/USDUtilities/Private/USDGeomMeshConversion.cpp b/Engine/Plugins/Importers/USDImporter/Source/USDUtilities/Private/USDGeomMeshConversion.cpp
--- a/Engine/Plugins/Importers/USDImporter/Source/USDUtilities/Private/USDGeomMeshConversion.cpp
+++ b/Engine/Plugins/Importers/USDImporter/Source/USDUtilities/Private/USDGeomMeshConversion.cpp
@@ -9,6 +9,7 @@
 #include "USDClassesModule.h"
 #include "USDConversionUtils.h"
 #include "USDDrawModeComponent.h"
+#include "USDLayerUtils.h"
 #include "USDLog.h"
 #include "USDMemory.h"
 #include "USDPrimConversion.h"
@@ -21,6 +22,9 @@
 #include "UsdWrappers/UsdStage.h"
 
 #include "Engine/StaticMesh.h"
+#include "GeometryCache.h"
+#include "GeometryCacheMeshData.h"
+#include "GeometryCacheTrack.h"
 #include "Materials/MaterialInstanceConstant.h"
 #include "Materials/MaterialInstanceDynamic.h"
 #include "MeshDescription.h"
@@ -3327,7 +3331,11 @@ UsdUtils::FUsdPrimMaterialAssignmentInfo UsdUtils::GetPrimMaterialAssignments(
 				int32 LastAssignmentIndex = Result.Slots.Num() - 1;
 				for (int PolygonIndex : PolygonIndicesInSubset)
 				{
-					Result.MaterialIndices[PolygonIndex] = LastAssignmentIndex;
+					// #ueent_todo: There can be issues with PolygonIndex being bigger that the number of faces with varying GeomSubsets
+					if (Result.MaterialIndices.IsValidIndex(PolygonIndex))
+					{
+						Result.MaterialIndices[PolygonIndex] = LastAssignmentIndex;
+					}
 				}
 			}
 		}
@@ -3763,7 +3771,7 @@ bool UnrealToUsd::ConvertMeshDescriptions( const TArray<FMeshDescription>& LODIn
 			VariantSet.SetVariantSelection( VariantName );
 			EditContext.Emplace( VariantSet.GetVariantEditContext() );
 		}
-
+		
 		pxr::UsdGeomMesh TargetMesh;
 		if ( bExportMultipleLODs )
 		{
@@ -3791,6 +3799,450 @@ bool UnrealToUsd::ConvertMeshDescriptions( const TArray<FMeshDescription>& LODIn
 	return true;
 }
 
+namespace UE::UsdGeometryCacheConversion::Private
+{
+	void AppendGeometryCacheMeshData(const FGeometryCacheMeshData& InMeshData, FGeometryCacheMeshData& InOutFlattenedMeshData)
+	{
+		// MeshData are flattened together by appending their data...
+		const int32 VertexIndexOffset = InOutFlattenedMeshData.Positions.Num();
+		const int32 IndicesIndexOffset = InOutFlattenedMeshData.Indices.Num();
+
+		InOutFlattenedMeshData.Positions.Append(InMeshData.Positions);
+		InOutFlattenedMeshData.TextureCoordinates.Append(InMeshData.TextureCoordinates);
+		InOutFlattenedMeshData.TangentsX.Append(InMeshData.TangentsX);
+		InOutFlattenedMeshData.TangentsZ.Append(InMeshData.TangentsZ);
+		InOutFlattenedMeshData.Colors.Append(InMeshData.Colors);
+
+		// ... and adjusting the indices with the proper offset
+		InOutFlattenedMeshData.Indices.Reserve(InOutFlattenedMeshData.Indices.Num() + InMeshData.Indices.Num());
+		Algo::Transform(InMeshData.Indices, InOutFlattenedMeshData.Indices, [VertexIndexOffset](uint32 Index) { return Index + VertexIndexOffset; });
+
+		// Same with the BatchInfo's StartIndex, which describes where each mesh section starts
+		for (const FGeometryCacheMeshBatchInfo& BatchInfo : InMeshData.BatchesInfo)
+		{
+			FGeometryCacheMeshBatchInfo AdjustedBatchInfo(BatchInfo);
+			AdjustedBatchInfo.StartIndex += IndicesIndexOffset;
+			InOutFlattenedMeshData.BatchesInfo.Add(AdjustedBatchInfo);
+		}
+
+		// Also merge the VertexInfo attributes that are checked when converting the MeshData
+		InOutFlattenedMeshData.VertexInfo.bHasTangentZ |= InMeshData.VertexInfo.bHasTangentZ;
+		InOutFlattenedMeshData.VertexInfo.bHasUV0 |= InMeshData.VertexInfo.bHasUV0;
+		InOutFlattenedMeshData.VertexInfo.bHasColor0 |= InMeshData.VertexInfo.bHasColor0;
+		InOutFlattenedMeshData.VertexInfo.bHasMotionVectors |= InMeshData.VertexInfo.bHasMotionVectors;
+	}
+
+	FGeometryCacheMeshData GetFlattenedGeometryCacheMeshData(const UGeometryCache* GeometryCache, int32 FrameIndex)
+	{
+		FGeometryCacheMeshData FlattenedMeshData;
+		if (GeometryCache->Tracks.Num() == 1)
+		{
+			GeometryCache->Tracks[0]->GetMeshDataAtSampleIndex(FrameIndex, FlattenedMeshData);
+		}
+		else
+		{
+			// MeshData for each track are aggregated together into a single flattened MeshData
+			for (int32 TrackIndex = 0; TrackIndex < GeometryCache->Tracks.Num(); ++TrackIndex)
+			{
+				FGeometryCacheMeshData TrackMeshData;
+				GeometryCache->Tracks[TrackIndex]->GetMeshDataAtSampleIndex(FrameIndex, TrackMeshData);
+
+				AppendGeometryCacheMeshData(TrackMeshData, FlattenedMeshData);
+			}
+		}
+		return MoveTemp(FlattenedMeshData);
+	}
+
+	struct FGeometryCacheExportContext
+	{
+		FGeometryCacheExportContext(const UGeometryCache& GeometryCache)
+			: SlotNames(GeometryCache.MaterialSlotNames)
+			, FrameRate((GeometryCache.GetEndFrame() - GeometryCache.GetStartFrame()) / GeometryCache.CalculateDuration())
+		{
+		}
+
+		const TArray<FName>& SlotNames;
+		const float FrameRate;
+
+		// Cached values of the last written attribute values
+		// Since int cannot be interpolated, the missing timesampled attribute values will be the "held"
+		// values of the previous written timesample
+		pxr::VtArray<int> FaceVertexCounts;
+		pxr::VtArray<int> FaceVertexIndices;
+	};
+
+	void ConvertGeometryCacheMeshData(
+		const FGeometryCacheMeshData& MeshData,
+		pxr::UsdGeomMesh& UsdMesh,
+		const pxr::VtArray<std::string>& MaterialAssignments,
+		const pxr::UsdTimeCode TimeCode,
+		pxr::UsdPrim MaterialPrim,
+		FGeometryCacheExportContext& ExportContext
+	)
+	{
+		pxr::UsdPrim MeshPrim = UsdMesh.GetPrim();
+		pxr::UsdStageRefPtr Stage = MeshPrim.GetStage();
+		if (!Stage)
+		{
+			return;
+		}
+		const FUsdStageInfo StageInfo{Stage};
+
+		// Vertices
+		{
+			const int32 VertexCount = MeshData.Positions.Num();
+
+			// Points
+			{
+				pxr::UsdAttribute Points = UsdMesh.CreatePointsAttr();
+				if (Points)
+				{
+					pxr::VtArray<pxr::GfVec3f> PointsArray;
+					PointsArray.reserve(VertexCount);
+
+					for (int32 VertexIndex = 0; VertexIndex < VertexCount; ++VertexIndex)
+					{
+						PointsArray.push_back(UnrealToUsd::ConvertVector(StageInfo, (FVector)MeshData.Positions[VertexIndex]));
+					}
+
+					Points.Set(PointsArray, TimeCode);
+				}
+			}
+
+			// Normals
+			if (MeshData.VertexInfo.bHasTangentZ)
+			{
+				// We need to emit this if we're writing normals (which we always are) because any DCC that can
+				// actually subdivide (like usdview) will just discard authored normals and fully recompute them
+				// on-demand in case they have a valid subdivision scheme (which is the default state).
+				// Reference: https://graphics.pixar.com/usd/release/api/class_usd_geom_mesh.html#UsdGeom_Mesh_Normals
+				if (pxr::UsdAttribute SubdivisionAttr = UsdMesh.CreateSubdivisionSchemeAttr())
+				{
+					ensure(SubdivisionAttr.Set(pxr::UsdGeomTokens->none));
+				}
+
+				pxr::UsdAttribute NormalsAttribute = UsdMesh.CreateNormalsAttr();
+				if (NormalsAttribute)
+				{
+					pxr::VtArray<pxr::GfVec3f> Normals;
+					Normals.reserve(VertexCount);
+
+					for (int32 VertexIndex = 0; VertexIndex < VertexCount; ++VertexIndex)
+					{
+						FVector VertexNormal = MeshData.TangentsZ[VertexIndex].ToFVector();
+						Normals.push_back(UnrealToUsd::ConvertVector(StageInfo, VertexNormal));
+					}
+
+					NormalsAttribute.Set(Normals, TimeCode);
+				}
+			}
+
+			// UVs
+			if (MeshData.VertexInfo.bHasUV0)
+			{
+				// Only one UV set is supported
+				const int32 TexCoordSourceIndex = 0;
+				pxr::TfToken UsdUVSetName = UsdUtils::GetUVSetName(TexCoordSourceIndex).Get();
+
+				pxr::UsdGeomPrimvar PrimvarST = pxr::UsdGeomPrimvarsAPI(MeshPrim).CreatePrimvar(UsdUVSetName, pxr::SdfValueTypeNames->TexCoord2fArray, pxr::UsdGeomTokens->vertex);
+
+				if (PrimvarST)
+				{
+					pxr::VtVec2fArray UVs;
+
+					for (int32 VertexIndex = 0; VertexIndex < VertexCount; ++VertexIndex)
+					{
+						FVector2D TexCoord = FVector2D(MeshData.TextureCoordinates[VertexIndex]);
+						TexCoord[1] = 1.f - TexCoord[1];
+
+						UVs.push_back(UnrealToUsd::ConvertVector(TexCoord));
+					}
+
+					PrimvarST.Set(UVs, TimeCode);
+				}
+			}
+
+			// Vertex colors
+			if (MeshData.VertexInfo.bHasColor0)
+			{
+				pxr::UsdGeomPrimvar DisplayColorPrimvar = UsdMesh.CreateDisplayColorPrimvar(pxr::UsdGeomTokens->vertex);
+				pxr::UsdGeomPrimvar DisplayOpacityPrimvar = UsdMesh.CreateDisplayOpacityPrimvar(pxr::UsdGeomTokens->vertex);
+
+				if (DisplayColorPrimvar && DisplayOpacityPrimvar)
+				{
+					pxr::VtArray<pxr::GfVec3f> DisplayColors;
+					DisplayColors.reserve(VertexCount);
+
+					pxr::VtArray<float> DisplayOpacities;
+					DisplayOpacities.reserve(VertexCount);
+
+					for (int32 VertexIndex = 0; VertexIndex < VertexCount; ++VertexIndex)
+					{
+						const FColor& VertexColor = MeshData.Colors[VertexIndex];
+
+						// The color in the MeshData is already stored as linear
+						pxr::GfVec4f Color = UnrealToUsd::ConvertColor(VertexColor.ReinterpretAsLinear());
+						DisplayColors.push_back(pxr::GfVec3f(Color[0], Color[1], Color[2]));
+						DisplayOpacities.push_back(Color[3]);
+					}
+
+					DisplayColorPrimvar.Set(DisplayColors, TimeCode);
+					DisplayOpacityPrimvar.Set(DisplayOpacities, TimeCode);
+				}
+			}
+
+			// Velocities
+			if (MeshData.VertexInfo.bHasMotionVectors)
+			{
+				pxr::UsdAttribute VelocitiesAttribute = UsdMesh.CreateVelocitiesAttr();
+				if (VelocitiesAttribute)
+				{
+					pxr::VtArray<pxr::GfVec3f> Velocities;
+					Velocities.reserve(VertexCount);
+
+					for (int32 VertexIndex = 0; VertexIndex < VertexCount; ++VertexIndex)
+					{
+						// The motion vectors in the MeshData are stored as unit per frame so convert it back to unit per second
+						Velocities.push_back(UnrealToUsd::ConvertVector(StageInfo, (FVector) -MeshData.MotionVectors[VertexIndex] * ExportContext.FrameRate));
+					}
+
+					VelocitiesAttribute.Set(Velocities, TimeCode);
+				}
+			}
+
+		}
+
+		// Faces
+		{
+			const int32 NumIndices = MeshData.Indices.Num();
+			const int32 FaceCount = NumIndices / 3;
+			// Face Vertex Counts
+			{
+				pxr::UsdAttribute FaceCountsAttribute = UsdMesh.CreateFaceVertexCountsAttr();
+
+				if (FaceCountsAttribute)
+				{
+					pxr::VtArray<int> FaceVertexCounts;
+					FaceVertexCounts.reserve(FaceCount);
+
+					for (int32 FaceIndex = 0; FaceIndex < FaceCount; ++FaceIndex)
+					{
+						FaceVertexCounts.push_back(3);
+					}
+
+					if (ExportContext.FaceVertexCounts != FaceVertexCounts)
+					{
+						FaceCountsAttribute.Set(FaceVertexCounts, TimeCode);
+						ExportContext.FaceVertexCounts = FaceVertexCounts;
+					}
+				}
+			}
+
+			// Face Vertex Indices
+			{
+				pxr::UsdAttribute FaceVertexIndicesAttribute = UsdMesh.GetFaceVertexIndicesAttr();
+
+				if (FaceVertexIndicesAttribute)
+				{
+					pxr::VtArray<int> FaceVertexIndices;
+					FaceVertexIndices.reserve(NumIndices);
+
+					for (int32 Index = 0; Index < NumIndices; ++Index)
+					{
+						FaceVertexIndices.push_back(MeshData.Indices[Index]);
+					}
+
+					if (ExportContext.FaceVertexIndices != FaceVertexIndices)
+					{
+						FaceVertexIndicesAttribute.Set(FaceVertexIndices, TimeCode);
+						ExportContext.FaceVertexIndices = FaceVertexIndices;
+					}
+				}
+			}
+		}
+
+		// Material assignments
+		{
+			// This mesh has a single material assignment, just add an unrealMaterials attribute to the mesh prim
+			if (MaterialAssignments.size() == 1)
+			{
+				if (pxr::UsdAttribute UEMaterialsAttribute = MaterialPrim.CreateAttribute(UnrealIdentifiers::MaterialAssignment, pxr::SdfValueTypeNames->String))
+				{
+					UEMaterialsAttribute.Set(MaterialAssignments[0]);
+				}
+			}
+			// Multiple material assignments to the same mesh. Need to create a GeomSubset for each UE mesh section
+			else if (MaterialAssignments.size() > 1)
+			{
+				TSet<FString> UsedSectionNames;
+				// Need to fetch all triangles of a section, and add their indices
+				for (int32 SectionIndex = 0; SectionIndex < MeshData.BatchesInfo.Num(); ++SectionIndex)
+				{
+					const FGeometryCacheMeshBatchInfo& Section = MeshData.BatchesInfo[SectionIndex];
+
+					// Note that we will continue on even if we have no material assignment, so as to satisfy the "partition" family condition (below)
+					std::string SectionMaterial;
+					if (Section.MaterialIndex >= 0 && Section.MaterialIndex < MaterialAssignments.size())
+					{
+						SectionMaterial = MaterialAssignments[Section.MaterialIndex];
+					}
+
+					FString SectionName;
+					if (ExportContext.SlotNames.IsValidIndex(SectionIndex))
+					{
+						SectionName = ExportContext.SlotNames[SectionIndex].ToString();
+						SectionName = UsdUtils::GetUniqueName(SectionName, UsedSectionNames);
+						UsedSectionNames.Add(SectionName);
+					}
+					else
+					{
+						SectionName = FString::Printf(TEXT("Section%d"), SectionIndex);
+					}
+
+					FSdfPath PrimPath(*SectionName);
+					pxr::UsdPrim GeomSubsetPrim = Stage->DefinePrim(
+						MeshPrim.GetPath().AppendPath(PrimPath),
+						UnrealToUsd::ConvertToken(TEXT("GeomSubset")).Get()
+					);
+
+					// MaterialPrim may be in another stage, so we may need another GeomSubset there
+					pxr::UsdPrim MaterialGeomSubsetPrim = GeomSubsetPrim;
+					if (MaterialPrim.GetStage() != MeshPrim.GetStage())
+					{
+						MaterialGeomSubsetPrim = MaterialPrim.GetStage()->OverridePrim(
+							MaterialPrim.GetPath().AppendPath(PrimPath)
+						);
+					}
+
+					pxr::UsdGeomSubset GeomSubsetSchema{GeomSubsetPrim};
+
+					// Element type attribute
+					// Write the geomsubset attributes only once since they are at Default time anyway
+					pxr::UsdAttribute ElementTypeAttr = GeomSubsetSchema.CreateElementTypeAttr();
+					if (!ElementTypeAttr.HasAuthoredValue())
+					{
+						ElementTypeAttr.Set(pxr::UsdGeomTokens->face);
+
+						// Indices attribute
+						const uint32 TriangleCount = Section.NumTriangles;
+						const uint32 FirstTriangleIndex = Section.StartIndex / 3; // StartIndex is the first *vertex* instance index
+						pxr::VtArray<int> IndicesAttrValue;
+						for (uint32 TriangleIndex = FirstTriangleIndex; TriangleIndex - FirstTriangleIndex < TriangleCount; ++TriangleIndex)
+						{
+							// Note that we add VertexInstances in sequence to the usda file for the faceVertexInstances attribute, which
+							// also constitutes our triangle order
+							IndicesAttrValue.push_back(static_cast<int>(TriangleIndex));
+						}
+
+						// Since family name and type attributes must be set at time Default, set the Indices at time Default too
+						// #ueent_todo: Add support for varying geomsubsets. This can happen with animation where sections 
+						// visibility are toggled on/off
+						pxr::UsdAttribute IndicesAttr = GeomSubsetSchema.CreateIndicesAttr();
+						IndicesAttr.Set(IndicesAttrValue);
+
+						// Family name attribute
+						pxr::UsdAttribute FamilyNameAttr = GeomSubsetSchema.CreateFamilyNameAttr();
+						FamilyNameAttr.Set(pxr::UsdShadeTokens->materialBind);
+
+						// Family type
+						pxr::UsdGeomSubset::SetFamilyType(UsdMesh, pxr::UsdShadeTokens->materialBind, pxr::UsdGeomTokens->partition);
+
+						// unrealMaterial attribute
+						if (pxr::UsdAttribute UEMaterialsAttribute = MaterialGeomSubsetPrim.CreateAttribute(UnrealIdentifiers::MaterialAssignment, pxr::SdfValueTypeNames->String))
+						{
+							UEMaterialsAttribute.Set(SectionMaterial);
+						}
+					}
+				}
+			}
+		}
+	}
+}
+
+bool UnrealToUsd::ConvertGeometryCache(const UGeometryCache* GeometryCache, pxr::UsdPrim& UsdPrim, UE::FUsdStage* StageForMaterialAssignments)
+{
+	namespace UsdGeometryCacheImpl = UE::UsdGeometryCacheConversion::Private;
+
+	FScopedUsdAllocs UsdAllocs;
+
+	pxr::UsdStageRefPtr Stage = UsdPrim.GetStage();
+	if (!GeometryCache || !Stage)
+	{
+		return false;
+	}
+
+	const FUsdStageInfo StageInfo(Stage);
+
+	pxr::SdfPath ParentPrimPath = UsdPrim.GetPath();
+
+	// Collect all material assignments, referenced by the sections' material indices
+	bool bHasMaterialAssignments = false;
+	pxr::VtArray< std::string > MaterialAssignments;
+	for (const UMaterialInterface* Material : GeometryCache->Materials)
+	{
+		FString AssignedMaterialPathName;
+		if (Material)
+		{
+			if (Material->GetOutermost() != GetTransientPackage())
+			{
+				AssignedMaterialPathName = Material->GetPathName();
+				bHasMaterialAssignments = true;
+			}
+		}
+
+		MaterialAssignments.push_back(UnrealToUsd::ConvertString(*AssignedMaterialPathName).Get());
+	}
+	if (!bHasMaterialAssignments)
+	{
+		// Prevent creation of the unrealMaterials attribute in case we don't have any assignments at all
+		MaterialAssignments.clear();
+	}
+
+	// Author material bindings on the dedicated stage if we have one
+	pxr::UsdStageRefPtr MaterialStage;
+	if (StageForMaterialAssignments)
+	{
+		MaterialStage = static_cast<pxr::UsdStageRefPtr>(*StageForMaterialAssignments);
+	}
+	else
+	{
+		MaterialStage = Stage;
+	}
+
+	pxr::UsdGeomMesh TargetMesh{UsdPrim};
+	pxr::UsdPrim MaterialPrim = MaterialStage->OverridePrim(UsdPrim.GetPath());
+
+	const int32 StartFrame = GeometryCache->GetStartFrame();
+	const int32 EndFrame = GeometryCache->GetEndFrame();
+	int32 ActualStartFrame = -1;
+
+	UsdGeometryCacheImpl::FGeometryCacheExportContext ExportContext(*GeometryCache);
+	for (int32 FrameIndex = StartFrame; FrameIndex <= EndFrame; ++FrameIndex)
+	{
+		FGeometryCacheMeshData MeshData = UsdGeometryCacheImpl::GetFlattenedGeometryCacheMeshData(GeometryCache, FrameIndex - StartFrame);
+		// First frame of the animation cannot be empty otherwise the geometry cache translator would not be able to detect the animation
+		// It is allowed to have empty frames during or at the end of the animation, eg. for fluid sim or FX that disappear
+		const bool bIsValidFrame = MeshData.Positions.Num() > 0 || ActualStartFrame > 0;
+		if (bIsValidFrame)
+		{
+			if (ActualStartFrame == -1)
+			{
+				// The actual start frame is the first frame with some data
+				ActualStartFrame = FrameIndex;
+			}
+			const float TimeCode = FrameIndex;
+			UsdGeometryCacheImpl::ConvertGeometryCacheMeshData(MeshData, TargetMesh, MaterialAssignments, TimeCode, MaterialPrim, ExportContext);
+		}
+	}
+
+	// Configure time metadata for the stage
+	UE::FUsdStage UsdStage(MaterialStage);
+	UsdUtils::AddTimeCodeRangeToLayer(UsdStage.GetRootLayer(), ActualStartFrame, EndFrame);
+	UsdStage.SetTimeCodesPerSecond(ExportContext.FrameRate);
+
+	return true;
+}
+
 TOptional<UsdUtils::FDisplayColorMaterial> UsdUtils::ExtractDisplayColorMaterial(const pxr::UsdGeomGprim& Gprim, const pxr::UsdTimeCode TimeCode)
 {
 	if (!Gprim)
diff --git a/Engine/Plugins/Importers/USDImporter/Source/USDUtilities/Public/USDGeomMeshConversion.h b/Engine/Plugins/Importers/USDImporter/Source/USDUtilities/Public/USDGeomMeshConversion.h
--- a/Engine/Plugins/Importers/USDImporter/Source/USDUtilities/Public/USDGeomMeshConversion.h
+++ b/Engine/Plugins/Importers/USDImporter/Source/USDUtilities/Public/USDGeomMeshConversion.h
@@ -25,6 +25,7 @@ PXR_NAMESPACE_OPEN_SCOPE
 	class UsdTyped;
 PXR_NAMESPACE_CLOSE_SCOPE
 
+class UGeometryCache;
 class UMaterialInstanceConstant;
 class UMaterialInstanceDynamic;
 class UStaticMesh;
@@ -216,6 +217,15 @@ namespace UnrealToUsd
 	 * created for each provided LOD index. Within each variant, a single Mesh prim also named LOD0, LOD1, etc. will contain the mesh data.
 	 */
 	USDUTILITIES_API bool ConvertMeshDescriptions( const TArray<FMeshDescription>& LODIndexToMeshDescription, pxr::UsdPrim& UsdPrim, const FMatrix& AdditionalTransform, const pxr::UsdTimeCode TimeCode = pxr::UsdTimeCode::Default() );
+
+	/**
+	* Extracts animated mesh data from GeometryCache and places the results in UsdPrim.
+	* @param GeometryCache - GeometryCache to convert
+	* @param UsdPrim - Prim to receive the mesh data or LOD variant set
+	* @param StageForMaterialAssignments - Stage to use when authoring material assignments (we use this when we want to export the mesh to a payload layer, but the material assignments to an asset layer)
+	* @return Whether the conversion was successful or not.
+	*/
+	USDUTILITIES_API bool ConvertGeometryCache(const UGeometryCache* GeometryCache, pxr::UsdPrim& UsdPrim, UE::FUsdStage* StageForMaterialAssignments = nullptr);
 }
 #endif // USE_USD_SDK
 
