USD: Fixed UsdUtils::AddTimeCodeRangeToLayer was not setting the provided EndTimeCode.

diff --git a/Engine/Plugins/Importers/USDImporter/Source/USDUtilities/Private/USDLayerUtils.cpp b/Engine/Plugins/Importers/USDImporter/Source/USDUtilities/Private/USDLayerUtils.cpp
--- a/Engine/Plugins/Importers/USDImporter/Source/USDUtilities/Private/USDLayerUtils.cpp
+++ b/Engine/Plugins/Importers/USDImporter/Source/USDUtilities/Private/USDLayerUtils.cpp
@@ -581,7 +581,7 @@ void UsdUtils::AddTimeCodeRangeToLayer( const pxr::SdfLayerRefPtr& Layer, double
 
 	if ( EndTimeCode > Layer->GetEndTimeCode() || !Layer->HasEndTimeCode() )
 	{
-		Layer->SetEndTimeCode( StartTimeCode );
+		Layer->SetEndTimeCode( EndTimeCode );
 	}
 }
 
