USD: Fixed possible crash in CanMeshSubtreeBeCollapsed because it can be called in worker threads by RecursiveQueryCollapsesChildren.

diff --git a/Engine/Plugins/Importers/USDImporter/Source/USDSchemas/Private/USDInfoCache.cpp b/Engine/Plugins/Importers/USDImporter/Source/USDSchemas/Private/USDInfoCache.cpp
--- a/Engine/Plugins/Importers/USDImporter/Source/USDSchemas/Private/USDInfoCache.cpp
+++ b/Engine/Plugins/Importers/USDImporter/Source/USDSchemas/Private/USDInfoCache.cpp
@@ -144,6 +144,8 @@ struct FUsdInfoCache::FUsdInfoCacheImpl
 public:
 	void RegisterAuxiliaryPrims(const UE::FSdfPath& MainPrimPath, const TSet<UE::FSdfPath>& AuxPrimPaths)
 	{
+		TRACE_CPUPROFILER_EVENT_SCOPE(RegisterAuxiliaryPrims);
+
 		if (AuxPrimPaths.Num() == 0)
 		{
 			return;
@@ -947,6 +949,7 @@ namespace UE::USDInfoCacheImpl::Private
 
 		pxr::SdfPath UsdPrimPath = UsdPrim.GetPrimPath();
 
+		FWriteScopeLock ScopeLock(Impl.InfoMapLock);
 		UE::UsdInfoCache::Private::FUsdPrimInfo& Info = Impl.InfoMap.FindOrAdd(UE::FSdfPath{UsdPrimPath});
 
 		if (Info.ExpectedVertexCountForSubtree > GMaxNumVerticesCollapsedMesh)
@@ -1068,6 +1071,8 @@ namespace UE::USDInfoCacheImpl::Private
 	// In other words, "instanceable composition arcs from local prims"
 	TSet<UE::FSdfPath> GetLocalNonRootCompositionArcSourcePaths(const pxr::UsdPrim& UsdPrim)
 	{
+		TRACE_CPUPROFILER_EVENT_SCOPE(GetLocalNonRootCompositionArcSourcePaths);
+
 		TSet<UE::FSdfPath> Result;
 
 		if (!UsdPrim)
@@ -1112,6 +1117,8 @@ namespace UE::USDInfoCacheImpl::Private
 		pxr::UsdStageRefPtr Stage = UsdPrim.GetStage();
 		for(const pxr::UsdPrim& Prototype : Stage->GetPrototypes())
 		{
+			TRACE_CPUPROFILER_EVENT_SCOPE(RegisterInstanceableAuxPrims::Prototype);
+
 			if(!Prototype)
 			{
 				continue;
@@ -1120,6 +1127,8 @@ namespace UE::USDInfoCacheImpl::Private
 			// Step into every instance of this prototype on the stage
 			for (const pxr::UsdPrim& Instance : Prototype.GetInstances())
 			{
+				TRACE_CPUPROFILER_EVENT_SCOPE(RegisterInstanceableAuxPrims::PrototypeInstance);
+
 				UE::FSdfPath InstancePath{Instance.GetPrimPath()};
 
 				// Adding a dependency on the prototype directly is interesting, even though we currently don't display those.
@@ -1135,6 +1144,8 @@ namespace UE::USDInfoCacheImpl::Private
 				pxr::UsdPrimRange PrimRange(Instance, pxr::UsdTraverseInstanceProxies());
 				for (pxr::UsdPrimRange::iterator InstanceChildIt = ++PrimRange.begin(); InstanceChildIt != PrimRange.end(); ++InstanceChildIt)
 				{
+					TRACE_CPUPROFILER_EVENT_SCOPE(RegisterInstanceableAuxPrims::InstanceChild);
+
 					pxr::SdfPath SdfChildPrimPath = InstanceChildIt->GetPrimPath();
 					UE::FSdfPath ChildPrimPath{SdfChildPrimPath};
 
