USD: Fix weirdness and broken LevelSequence bindings caused by trying to remap all invalid soft references to some asset.

diff --git a/Engine/Plugins/Importers/USDImporter/Source/USDStageImporter/Private/USDStageImporter.cpp b/Engine/Plugins/Importers/USDImporter/Source/USDStageImporter/Private/USDStageImporter.cpp
--- a/Engine/Plugins/Importers/USDImporter/Source/USDStageImporter/Private/USDStageImporter.cpp
+++ b/Engine/Plugins/Importers/USDImporter/Source/USDStageImporter/Private/USDStageImporter.cpp
@@ -780,13 +780,16 @@ namespace UsdStageImporterImpl
 				MovedAsset = Asset;
 			}
 
-			OutAssetsToFinalize.Add(ExistingAsset);
 			OutAssetsToFinalize.Add(MovedAsset);
+			if (ExistingAsset)
+			{
+				OutAssetsToFinalize.Add(ExistingAsset);
 
-			// If we're replacing ExistingAsset, we must update all references we can find from ExistingAsset to the new
-			// asset, otherwise they'll be left pointing at transient or GC'd/stomped assets
-			SoftObjectsToRemap.Add(ExistingAsset, MovedAsset);
-			ObjectsToRemap.Add(ExistingAsset, MovedAsset);
+				// If we're replacing ExistingAsset, we must update all references we can find from ExistingAsset to the new
+				// asset, otherwise they'll be left pointing at transient or GC'd/stomped assets
+				SoftObjectsToRemap.Add(ExistingAsset, MovedAsset);
+				ObjectsToRemap.Add(ExistingAsset, MovedAsset);
+			}
 		}
 		else
 		{
@@ -1322,14 +1325,25 @@ namespace UsdStageImporterImpl
 			return;
 		}
 
-		// Remap references held by assets that were moved directly to the destination package, and won't be in ObjectsToRemap
+		// We never want to remap from all instances of nullptr to something else
+		TMap<UObject*, UObject*> ProcessedObjectsToRemap;
+		ProcessedObjectsToRemap.Reserve(ObjectsToRemap.Num());
+		for (const TPair<UObject*, UObject*>& SourceTarget : ObjectsToRemap)
+		{
+			if (SourceTarget.Key)
+			{
+				ProcessedObjectsToRemap.Add(SourceTarget);
+			}
+		}
+
+		// Remap references held by assets that were moved directly to the destination package, and won't be in ProcessedObjectsToRemap
 		TSet<UObject*> Referencers = PublishedObjects;
 		if ( AActor* SceneActor = ImportContext.SceneActor )
 		{
 			// Remap references to spawned actors
 			Referencers.Add( ImportContext.SceneActor->GetWorld()->GetCurrentLevel() );
 		}
-		for ( const TPair<UObject*, UObject*>& Pair : ObjectsToRemap )
+		for ( const TPair<UObject*, UObject*>& Pair : ProcessedObjectsToRemap )
 		{
 			// Remap internal references between the remapped objects
 			Referencers.Add( Pair.Value );
@@ -1345,7 +1359,7 @@ namespace UsdStageImporterImpl
 			}
 
 			constexpr EArchiveReplaceObjectFlags ReplaceFlags = (EArchiveReplaceObjectFlags::IgnoreOuterRef | EArchiveReplaceObjectFlags::IgnoreArchetypeRef);
-			FArchiveReplaceObjectRef< UObject > ArchiveReplaceObjectRefInner(Referencer, ObjectsToRemap, ReplaceFlags);
+			FArchiveReplaceObjectRef< UObject > ArchiveReplaceObjectRefInner(Referencer, ProcessedObjectsToRemap, ReplaceFlags);
 		}
 	}
 
@@ -1503,8 +1517,19 @@ namespace UsdStageImporterImpl
 		// if we try deleting those assets afterwards
 		Packages.Remove( GetTransientPackage() );
 
+		// We never want to remap all invalid reference to something. That particularly seems to break LevelSequence bindings somehow
+		TMap<FSoftObjectPath, FSoftObjectPath> ProcessedPaths;
+		ProcessedPaths.Reserve(SoftObjectsToRemap.Num());
+		for (const TPair<FSoftObjectPath, FSoftObjectPath>& SourceTarget : SoftObjectsToRemap)
+		{
+			if (SourceTarget.Key.IsValid())
+			{
+				ProcessedPaths.Add(SourceTarget);
+			}
+		}
+
 		IAssetTools& AssetTools = FModuleManager::LoadModuleChecked<FAssetToolsModule>( "AssetTools" ).Get();
-		AssetTools.RenameReferencingSoftObjectPaths( Packages.Array(), SoftObjectsToRemap );
+		AssetTools.RenameReferencingSoftObjectPaths( Packages.Array(), ProcessedPaths );
 	}
 
 	void GetPublishedAssetsAndDependencies(
