USD: Fix animations being one frame off when exporting attached components in some cases;

Better way of handling socket transforms that may also handle additional cases;
Bake transforms for attached and bound components in case they are animated ith skeletal animations.

diff --git a/Engine/Plugins/Importers/USDImporter/Source/USDExporter/Private/LevelSequenceExporterUSD.cpp b/Engine/Plugins/Importers/USDImporter/Source/USDExporter/Private/LevelSequenceExporterUSD.cpp
--- a/Engine/Plugins/Importers/USDImporter/Source/USDExporter/Private/LevelSequenceExporterUSD.cpp
+++ b/Engine/Plugins/Importers/USDImporter/Source/USDExporter/Private/LevelSequenceExporterUSD.cpp
@@ -406,7 +406,7 @@ namespace UE
 			struct FCombinedComponentBakers
 			{
 				UnrealToUsd::EBakingType CombinedBakingType = UnrealToUsd::EBakingType::None;
-				TArray<TFunction<void( double )>> Bakers;
+				TArray<UnrealToUsd::FComponentBaker> Bakers;
 			};
 
 			class FLevelSequenceExportContext
@@ -965,6 +965,7 @@ namespace UE
 						continue;
 					}
 
+					bool bHasTransformBaker = false;
 					if ( const FMovieSceneBinding* Binding = MovieScene->FindBinding( Guid ) )
 					{
 						for ( const UMovieSceneTrack* Track : Binding->GetTracks() )
@@ -1011,9 +1012,34 @@ namespace UE
 							FCombinedComponentBakers& ExistingBakers = InOutComponentBakers.FindOrAdd( BoundComponent );
 							if ( Baker.BakerType != UnrealToUsd::EBakingType::None && !EnumHasAnyFlags( ExistingBakers.CombinedBakingType, Baker.BakerType ) )
 							{
-								ExistingBakers.Bakers.Add( Baker.BakerFunction );
+								ExistingBakers.Bakers.Add(Baker);
 								ExistingBakers.CombinedBakingType |= Baker.BakerType;
 							}
+
+							if (Baker.BakerType == UnrealToUsd::EBakingType::Transform)
+							{
+								bHasTransformBaker = true;
+							}
+						}
+					}
+
+					// If our component is attached to a specific socket of its parent, make sure its transform is baked.
+					// This because if this parent has any AnimSequence animating it, BoundComponent's world transform may change
+					// without it ever having an animated transform, and we don't have any form of rigging/socket attachment on USD that
+					// would cause the parent prim's skeletal animation to also affect its child prims.
+					// Ideally we'd actually search through the tracks to know for sure whether our parent has a SkeletalAnimation section,
+					// but it's probably safer to just do this in case it is hidden behind N subsequences or some obscure feature
+					if (!bHasTransformBaker && BoundComponent->GetAttachSocketName() != NAME_None)
+					{
+						UnrealToUsd::FComponentBaker Baker;
+						const FString PropertyPath = TEXT("Transform");
+						UnrealToUsd::CreateComponentPropertyBaker(Prim, *BoundComponent, PropertyPath, Baker);
+
+						FCombinedComponentBakers& ExistingBakers = InOutComponentBakers.FindOrAdd(BoundComponent);
+						if (Baker.BakerType != UnrealToUsd::EBakingType::None && !EnumHasAnyFlags(ExistingBakers.CombinedBakingType, Baker.BakerType))
+						{
+							ExistingBakers.Bakers.Add(Baker);
+							ExistingBakers.CombinedBakingType |= Baker.BakerType;
 						}
 					}
 				}
@@ -1099,6 +1125,45 @@ namespace UE
 				static constexpr bool bSorted = true;
 				const TArray<TWeakObjectPtr<UTickableConstraint>> AllConstraints = Controller.GetAllConstraints(bSorted);
 
+				// Collect and sort the bakers: We need all skeletal animations evaluated first, as we need to manually force
+				// the component to update and attached components will only get the correct values if they are evaluated
+				// after that
+				TArray<UnrealToUsd::FComponentBaker> SortedBakers;
+				SortedBakers.Reserve(ComponentBakers.Num() * UnrealToUsd::NumBakingTypes);
+				for (const TPair<USceneComponent*, FCombinedComponentBakers>& Pair : ComponentBakers)
+				{
+					for (const UnrealToUsd::FComponentBaker& Baker : Pair.Value.Bakers)
+					{
+						if (Baker.BakerFunction)
+						{
+							SortedBakers.Add(Baker);
+						}
+					}
+				}
+				SortedBakers.Sort(
+					[](const UnrealToUsd::FComponentBaker& LHS, const UnrealToUsd::FComponentBaker& RHS) -> bool
+					{
+						if (LHS.BakerType != RHS.BakerType)
+						{
+							if (LHS.BakerType == UnrealToUsd::EBakingType::Skeletal)
+							{
+								// We want all skeletal bakers first so that the joints are updated before we refresh transforms attached
+								// to joints and sockets
+								return true;
+							}
+							else if (RHS.BakerType == UnrealToUsd::EBakingType::Skeletal)
+							{
+								return false;
+							}
+						}
+
+						// Parents should go first (so that if we have two nested skeletal mesh components attached to each other, we refresh the
+						// parent first). To be honest this is likely not necessary as child joint transforms shouldn't depend on parent joint
+						// transforms, but we should enforce some consistent ordering anyway so might as well use this
+						return LHS.ComponentPath < RHS.ComponentPath;
+					}
+				);
+
 				for ( FFrameTime EvalTime = StartFrame; EvalTime <= EndFrame; EvalTime += Interval )
 				{
 					Context.Sequencer->SetLocalTimeDirectly( EvalTime );
@@ -1114,17 +1179,11 @@ namespace UE
 					}
 
 					FFrameTime KeyTime = FFrameRate::Snap( EvalTime, Resolution, DisplayRate ).FloorToFrame();
-					double UsdTimeCode = FFrameRate::TransformTime( KeyTime, Resolution, StageFrameRate ).AsDecimal();
+					double UsdTimeCode = FFrameRate::TransformTime(KeyTime, Resolution, StageFrameRate).AsDecimal();
 
-					for ( const TPair<USceneComponent*, FCombinedComponentBakers>& Pair : ComponentBakers )
+					for (const UnrealToUsd::FComponentBaker& Baker : SortedBakers)
 					{
-						for ( const TFunction<void( double )>& Lambda : Pair.Value.Bakers )
-						{
-							if ( Lambda )
-							{
-								Lambda( UsdTimeCode );
-							}
-						}
+						Baker.BakerFunction(UsdTimeCode);
 					}
 				}
 
diff --git a/Engine/Plugins/Importers/USDImporter/Source/USDUtilities/Private/USDPrimConversion.cpp b/Engine/Plugins/Importers/USDImporter/Source/USDUtilities/Private/USDPrimConversion.cpp
--- a/Engine/Plugins/Importers/USDImporter/Source/USDUtilities/Private/USDPrimConversion.cpp
+++ b/Engine/Plugins/Importers/USDImporter/Source/USDUtilities/Private/USDPrimConversion.cpp
@@ -2482,18 +2482,21 @@ bool UnrealToUsd::ConvertSceneComponent( const pxr::UsdStageRefPtr& Stage, const
 		return false;
 	}
 
-	FTransform RelativeTransform = SceneComponent->GetRelativeTransform();
-
 	// If we're attached to a socket our RelativeTransform will be relative to the socket, instead of the parent
-	// component space. Here we get the socket transform itself and concatenate it, forcing RelativeTransform to
-	// be relative to the parent component space again.
+	// component space. If we were to use GetRelativeTransform directly, we're in charge of managing the socket
+	// transform too (and any other N obscure features we don't know about/don't exist yet). If we fetch directly
+	// the component-to-world transform however, the component will do that on its own (as that is the transform
+	// that is actually used to show it on the level), so we don't have to worry about it!
+	FTransform RelativeTransform;
 	if (USceneComponent* Parent = SceneComponent->GetAttachParent())
 	{
-		if (FName SocketName = SceneComponent->GetAttachSocketName(); SocketName != NAME_None)
-		{
-			FTransform SocketTransform = Parent->GetSocketTransform(SocketName, RTS_Component);
-			RelativeTransform = RelativeTransform * SocketTransform;
-		}
+		Parent->ConditionalUpdateComponentToWorld();
+		Parent->UpdateChildTransforms();
+		RelativeTransform = SceneComponent->GetComponentTransform().GetRelativeTransform(Parent->GetComponentTransform());
+	}
+	else
+	{
+		RelativeTransform = SceneComponent->GetRelativeTransform();
 	}
 
 	// Compensate different orientation for light or camera components:
@@ -3296,20 +3299,23 @@ bool UnrealToUsd::CreateComponentPropertyBaker( UE::FUsdPrim& Prim, const UScene
 			{
 				FScopedUsdAllocs Allocs;
 
-				FTransform RelativeTransform = Component.GetRelativeTransform();
-
 				// If we're attached to a socket our RelativeTransform will be relative to the socket, instead of the parent
-				// component space. Here we get the socket transform itself and concatenate it, forcing RelativeTransform to
-				// be relative to the parent component space again.
+				// component space. If we were to use GetRelativeTransform directly, we're in charge of managing the socket
+				// transform too (and any other N obscure features we don't know about/don't exist yet). If we fetch directly
+				// the component-to-world transform however, the component will do that on its own (as that is the transform
+				// that is actually used to show it on the level), so we don't have to worry about it!
 				// It may seem wasteful to do this inside the baker function, but you can place "Attach tracks" on the
 				// Sequencer that may make the attach socket change every frame, so we do need this
+				FTransform RelativeTransform;
 				if (USceneComponent* Parent = Component.GetAttachParent())
 				{
-					if (FName SocketName = Component.GetAttachSocketName(); SocketName != NAME_None)
-					{
-						FTransform SocketTransform = Parent->GetSocketTransform(SocketName, RTS_Component);
-						RelativeTransform = RelativeTransform * SocketTransform;
-					}
+					Parent->ConditionalUpdateComponentToWorld();
+					Parent->UpdateChildTransforms();
+					RelativeTransform = Component.GetComponentTransform().GetRelativeTransform(Parent->GetComponentTransform());
+				}
+				else
+				{
+					RelativeTransform = Component.GetRelativeTransform();
 				}
 
 				FTransform FinalUETransform = AdditionalRotation * RelativeTransform;
@@ -3496,6 +3502,7 @@ bool UnrealToUsd::CreateComponentPropertyBaker( UE::FUsdPrim& Prim, const UScene
 	{
 		OutBaker.BakerType = BakerType;
 		OutBaker.BakerFunction = BakerFunction;
+		OutBaker.ComponentPath = Component.GetPathName();
 		return true;
 	}
 
@@ -3587,6 +3594,7 @@ bool UnrealToUsd::CreateSkeletalAnimationBaker( UE::FUsdPrim& SkelRoot, UE::FUsd
 	pxr::VtVec3hArray Scales;
 	pxr::VtArray< float > BlendShapeWeights;
 
+	OutBaker.ComponentPath = Component.GetPathName();
 	OutBaker.BakerType = EBakingType::Skeletal;
 	OutBaker.BakerFunction =
 		[&Component, StageInfo, Translations, Rotations, Scales, BlendShapeWeights, TranslationsAttr, RotationsAttr, ScalesAttr, BlendShapeWeightsAttr, NumBones, NumMorphTargets]
@@ -3609,6 +3617,13 @@ bool UnrealToUsd::CreateSkeletalAnimationBaker( UE::FUsdPrim& SkelRoot, UE::FUsd
 			Component.MarkRenderTransformDirty();
 			Component.MarkRenderDynamicDataDirty();
 
+			// I'm not entirely sure why this is needed but FFbxExporter::ExportAnimTrack and FFbxExporter::ExportLevelSequenceBaked3DTransformTrack
+			// do this so for safety maybe we should as well?
+			if (AActor* Owner = Component.GetOwner())
+			{
+				Owner->Tick(0.0f);
+			}
+
 			const TArray< FTransform >& LocalBoneTransforms = Component.GetBoneSpaceTransforms();
 			for ( int32 BoneIndex = 0; BoneIndex < NumBones; ++BoneIndex )
 			{
diff --git a/Engine/Plugins/Importers/USDImporter/Source/USDUtilities/Public/USDPrimConversion.h b/Engine/Plugins/Importers/USDImporter/Source/USDUtilities/Public/USDPrimConversion.h
--- a/Engine/Plugins/Importers/USDImporter/Source/USDUtilities/Public/USDPrimConversion.h
+++ b/Engine/Plugins/Importers/USDImporter/Source/USDUtilities/Public/USDPrimConversion.h
@@ -229,11 +229,13 @@ namespace UnrealToUsd
 		Bounds = 32,
 	};
 	ENUM_CLASS_FLAGS( EBakingType );
+	const static inline int32 NumBakingTypes = 6;
 
 	// Contains a lambda function responsible for baking a USceneComponent on the level into a prim on an FUsdStage
 	struct USDUTILITIES_API FComponentBaker
 	{
 		EBakingType BakerType;
+		FString ComponentPath;  // Used for sorting, to ensure we evaluate parents and skeletal stuff before children and attached stuff
 		TFunction<void( double UsdTimeCode )> BakerFunction;
 	};
 
