USD: Fix missing nullptr check.

diff --git a/Engine/Plugins/Importers/USDImporter/Source/USDStageEditorViewModels/Private/USDStageViewModel.cpp b/Engine/Plugins/Importers/USDImporter/Source/USDStageEditorViewModels/Private/USDStageViewModel.cpp
--- a/Engine/Plugins/Importers/USDImporter/Source/USDStageEditorViewModels/Private/USDStageViewModel.cpp
+++ b/Engine/Plugins/Importers/USDImporter/Source/USDStageEditorViewModels/Private/USDStageViewModel.cpp
@@ -357,7 +357,10 @@ void FUsdStageViewModel::ImportStage( const TCHAR* TargetContentFolder, UUsdStag
 
 			// Preemptively remove the stage actor as a user of the assets on the cache so that the stage importer
 			// can just take the assets from the cache directly. Otherwise it will be forced to duplicate them
-			ImportContext.AssetCache->RemoveAllAssetReferences(StageActor);
+			if (ImportContext.AssetCache)
+			{
+				ImportContext.AssetCache->RemoveAllAssetReferences(StageActor);
+			}
 
 			UUsdStageImporter* USDImporter = IUsdStageImporterModule::Get().GetImporter();
 			USDImporter->ImportFromFile(ImportContext);
