USD: Add support for exporting the output of PCG graphs to USD.

diff --git a/Engine/Plugins/Importers/USDImporter/Source/USDUtilities/Private/USDPrimConversion.cpp b/Engine/Plugins/Importers/USDImporter/Source/USDUtilities/Private/USDPrimConversion.cpp
--- a/Engine/Plugins/Importers/USDImporter/Source/USDUtilities/Private/USDPrimConversion.cpp
+++ b/Engine/Plugins/Importers/USDImporter/Source/USDUtilities/Private/USDPrimConversion.cpp
@@ -27,6 +27,7 @@
 #include "Channels/MovieSceneFloatChannel.h"
 #include "CineCameraActor.h"
 #include "CineCameraComponent.h"
+#include "Components/BrushComponent.h"
 #include "Components/DirectionalLightComponent.h"
 #include "Components/HierarchicalInstancedStaticMeshComponent.h"
 #include "Components/LightComponent.h"
@@ -2546,8 +2547,12 @@ bool UnrealToUsd::ConvertSceneComponent( const pxr::UsdStageRefPtr& Stage, const
 		}
 		else if ( !SceneComponent->ComponentTags.Contains( UnrealIdentifiers::Inherited ) )
 		{
-			// We don't have visible nor inherited tags: We're probably exporting a pure UE component, so write out component visibility instead
-			Value = SceneComponent->bHiddenInGame ? pxr::UsdGeomTokens->invisible : pxr::UsdGeomTokens->inherited;
+			// We don't have visible nor inherited tags: We're probably exporting a pure UE component, so write out component visibility instead.
+			// Ignore invisibility from brush components though because they are always forced to bHiddenInGame=true, with the property even
+			// being hidden on the details panel
+			Value = SceneComponent->bHiddenInGame && !SceneComponent->IsA<UBrushComponent>()
+						? pxr::UsdGeomTokens->invisible
+						: pxr::UsdGeomTokens->inherited;
 		}
 
 		VisibilityAttr.Set<pxr::TfToken>( Value );
