USD: Improve warning in case shader prim doesn't specify an "st" to use;

Downgrade message from ensure to just a warning, as it can happen with invalid data.

diff --git a/Engine/Plugins/Importers/USDImporter/Source/USDUtilities/Private/USDShadeConversion.cpp b/Engine/Plugins/Importers/USDImporter/Source/USDUtilities/Private/USDShadeConversion.cpp
--- a/Engine/Plugins/Importers/USDImporter/Source/USDUtilities/Private/USDShadeConversion.cpp
+++ b/Engine/Plugins/Importers/USDImporter/Source/USDUtilities/Private/USDShadeConversion.cpp
@@ -1200,11 +1200,16 @@ namespace UE
 					}
 					else
 					{
-						ensureMsgf(false, TEXT("Failed to find primvar '%s' when setting material parameter '%s' on material '%s'. Available primvars and UV indices: %s"),
+						UE_LOG(
+							LogUsd,
+							Warning,
+							TEXT("Failed to find primvar '%s' when setting material parameter '%s' on material '%s'. Available primvars and UV "
+								 "indices: %s.%s"),
 							*TextureValue.Primvar,
 							ParameterName,
 							*Material.GetPathName(),
-							*UsdUtils::StringifyMap(PrimvarToUVIndex)
+							*UsdUtils::StringifyMap(PrimvarToUVIndex),
+							TextureValue.Primvar.IsEmpty() ? TEXT(" Is your UsdUVTexture Shader missing the 'inputs:st' attribute? (It specifies which UV set to sample the texture with)") : TEXT("")
 						);
 					}
 
@@ -2355,9 +2360,13 @@ bool UsdToUnreal::ConvertMaterial(
 					}
 					else
 					{
-						ensureMsgf(false, TEXT("Failed to find primvar '%s' when setting material parameter. Available primvars and UV indices: %s"),
+						UE_LOG(
+							LogTemp,
+							Warning,
+							TEXT("Failed to find primvar '%s' when setting material parameter. Available primvars and UV indices: %s.%s"),
 							*TextureParameterValue->Primvar,
-							*UsdUtils::StringifyMap(PrimvarToUVIndex)
+							*UsdUtils::StringifyMap(PrimvarToUVIndex),
+							TextureParameterValue->Primvar.IsEmpty() ? TEXT(" Is your UsdUVTexture Shader missing the 'inputs:st' attribute? (It specifies which UV set to sample the texture with)") : TEXT("")
 						);
 					}
 				}
