USD: Fix allocator crashes when parsing grooms and geometry caches.

diff --git a/Engine/Plugins/Importers/USDImporter/Source/USDSchemas/Private/USDGeometryCacheTranslator.cpp b/Engine/Plugins/Importers/USDImporter/Source/USDSchemas/Private/USDGeometryCacheTranslator.cpp
--- a/Engine/Plugins/Importers/USDImporter/Source/USDSchemas/Private/USDGeometryCacheTranslator.cpp
+++ b/Engine/Plugins/Importers/USDImporter/Source/USDSchemas/Private/USDGeometryCacheTranslator.cpp
@@ -114,6 +114,10 @@ namespace UsdGeometryCacheTranslatorImpl
 			return;
 		}
 
+		// MeshDescriptions are always allocated on the UE allocator as the allocation happens within
+		// another dll, so we need to deallocate them using it too
+		FScopedUnrealAllocs Allocs;
+
 		FMeshDescription TempMeshDescription;
 		UsdUtils::FUsdPrimMaterialAssignmentInfo TempMaterialInfo;
 
@@ -240,6 +244,10 @@ namespace UsdGeometryCacheTranslatorImpl
 
 	bool ReadMeshData(const FReadMeshDataArgs& Args, const UE::FUsdPrim& MeshPrim, int32 MaterialOffset, float Time, FGeometryCacheMeshData& OutMeshData)
 	{
+		// MeshDescriptions are always allocated on the UE allocator as the allocation happens within
+		// another dll, so we need to deallocate them using it too
+		FScopedUnrealAllocs Allocs;
+
 		FTransform PropatagedTransform = FTransform::Identity;
 		if (Args.bPropagateTransform)
 		{
diff --git a/Engine/Plugins/Importers/USDImporter/Source/USDSchemas/Private/USDGroomTranslatorUtils.cpp b/Engine/Plugins/Importers/USDImporter/Source/USDSchemas/Private/USDGroomTranslatorUtils.cpp
--- a/Engine/Plugins/Importers/USDImporter/Source/USDSchemas/Private/USDGroomTranslatorUtils.cpp
+++ b/Engine/Plugins/Importers/USDImporter/Source/USDSchemas/Private/USDGroomTranslatorUtils.cpp
@@ -199,7 +199,6 @@ namespace UsdGroomTranslatorUtils
 	{
 		// At this point, the prim (SkelRoot or GeomMesh) has already been checked to have the GroomBindingAPI,
 		// so we need to set up the GroomComponent and the groom binding asset to be able to bind it to the mesh
-		FScopedUsdAllocs Allocs;
 
 		// The GroomBinding schema must specify a groom prim to bind to the mesh
 		const FString GroomPrimPath = GetGroomPrimPath(Prim);
