USD - Added a reparent function to move a list of specified prims to a new parent prim.
#tests Reparented an entire tree of mesh prims from an untype parent prim to an xform.

diff --git a/Engine/Plugins/Importers/USDImporter/Source/USDUtilities/Private/USDConversionUtils.cpp b/Engine/Plugins/Importers/USDImporter/Source/USDUtilities/Private/USDConversionUtils.cpp
--- a/Engine/Plugins/Importers/USDImporter/Source/USDUtilities/Private/USDConversionUtils.cpp
+++ b/Engine/Plugins/Importers/USDImporter/Source/USDUtilities/Private/USDConversionUtils.cpp
@@ -73,6 +73,7 @@
 	#include "pxr/usd/usdGeom/modelAPI.h"
 	#include "pxr/usd/usdGeom/primvar.h"
 	#include "pxr/usd/usdGeom/primvarsAPI.h"
+	#include "pxr/usd/usdGeom/xform.h"
 	#include "pxr/usd/usdLux/diskLight.h"
 	#include "pxr/usd/usdLux/distantLight.h"
 	#include "pxr/usd/usdLux/domeLight.h"
@@ -2047,6 +2048,65 @@ bool UsdUtils::RenamePrim( UE::FUsdPrim& Prim, const TCHAR* NewPrimName )
 #endif // #if USE_USD_SDK
 }
 
+TArray<UE::FSdfPath> UsdUtils::ReparentPrims(const TArray<UE::FUsdPrim>& Prims, const UE::FUsdPrim& ParentPrim)
+{
+	TArray<UE::FSdfPath> Result;
+
+#if USE_USD_SDK
+	Result.Reserve(Prims.Num());
+
+	FScopedUsdAllocs Allocs;
+
+	if ( const pxr::UsdPrim UsdParentPrim{ ParentPrim } )
+	{
+		if ( const pxr::UsdStageRefPtr ParentUsdStage = UsdParentPrim.GetStage() )
+		{
+			if ( const pxr::SdfLayerHandle ParentEditTargetLayer = ParentUsdStage->GetEditTarget().GetLayer() )
+			{
+				for ( const UE::FUsdPrim& Prim : Prims )
+				{
+					if ( const pxr::UsdPrim UsdPrim{ Prim } )
+					{
+						const pxr::UsdStageRefPtr UsdStage = UsdPrim.GetStage();
+
+						if ( const pxr::SdfLayerHandle EditTargetLayer = UsdStage->GetEditTarget().GetLayer() )
+						{
+							const pxr::SdfPath SourcePath = UsdPrim.GetPrimPath();
+							const pxr::SdfPath TargetPath = UsdParentPrim.GetPrimPath().AppendChild(UsdPrim.GetName());
+
+							// Copy the old prim to the new path
+							if ( pxr::SdfCopySpec(EditTargetLayer, SourcePath, ParentEditTargetLayer, TargetPath) )
+							{
+								// Remove the old prim
+								if ( UsdStage->RemovePrim(SourcePath) )
+								{
+									UE_LOG(LogUsd, Log, TEXT("Moved prim '%s' from layer '%s' under new parent prim '%s' from layer '%s'."),
+										*UsdToUnreal::ConvertPath(SourcePath),
+										*UsdToUnreal::ConvertString(EditTargetLayer->GetIdentifier()),
+										*UsdToUnreal::ConvertPath(UsdParentPrim.GetPrimPath()),
+										*UsdToUnreal::ConvertString(ParentEditTargetLayer->GetIdentifier()));
+								}
+								else
+								{
+									UE_LOG(LogUsd, Warning, TEXT("Copied prim '%s' from layer '%s' under new parent prim '%s' from layer '%s', but failed to remove the source prim from under its old parent."),
+										*UsdToUnreal::ConvertPath(SourcePath),
+										*UsdToUnreal::ConvertString(EditTargetLayer->GetIdentifier()),
+										*UsdToUnreal::ConvertPath(UsdParentPrim.GetPrimPath()),
+										*UsdToUnreal::ConvertString(ParentEditTargetLayer->GetIdentifier()));
+								}
+								Result.Emplace(TargetPath);
+							}
+						}
+					}
+				}
+			}
+		}
+	}
+#endif // USE_USD_SDK
+
+	return Result;
+}
+
 bool UsdUtils::RemoveNumberedSuffix( FString& Prefix )
 {
 	if ( Prefix.IsNumeric() )
diff --git a/Engine/Plugins/Importers/USDImporter/Source/USDUtilities/Public/USDConversionUtils.h b/Engine/Plugins/Importers/USDImporter/Source/USDUtilities/Public/USDConversionUtils.h
--- a/Engine/Plugins/Importers/USDImporter/Source/USDUtilities/Public/USDConversionUtils.h
+++ b/Engine/Plugins/Importers/USDImporter/Source/USDUtilities/Public/USDConversionUtils.h
@@ -333,6 +333,16 @@ namespace UsdUtils
 	 */
 	USDUTILITIES_API bool RenamePrim( UE::FUsdPrim& Prim, const TCHAR* NewPrimName );
 
+	/** 
+	 * Reparent the specified prims under a new parent prim.
+	 * The prims will be removed from under their old parent and moved under the new parent prim unless
+	 * the removing operation fails, in which case they are only copied under the new parent prim.
+	 * @param Prims - Prims to reparent
+	 * @param ParentPrim - The new parent prim
+	 * @return Paths to the succesfully copied prim specs, after they were added as children of ParentPrim
+	 */
+	USDUTILITIES_API TArray<UE::FSdfPath> ReparentPrims( const TArray<UE::FUsdPrim>& Prims, const UE::FUsdPrim& ParentPrim);
+
 	/**
 	 * Removes any numbered suffix, followed by any number of underscores (e.g. Asset_2, Asset__232_31 or Asset94 all become 'Asset'), making
 	 * sure the string is kept at least one character long. Returns true if it removed anything.
