USDImporter: expose more type name-related functions of UsdPrim on UE wrapper

This adds `SetTypeName()`, `ClearTypeName()`, and `HasAuthoredTypeName()`.

diff --git a/Engine/Plugins/Importers/USDImporter/Source/UnrealUSDWrapper/Private/UsdWrappers/UsdPrim.cpp b/Engine/Plugins/Importers/USDImporter/Source/UnrealUSDWrapper/Private/UsdWrappers/UsdPrim.cpp
--- a/Engine/Plugins/Importers/USDImporter/Source/UnrealUSDWrapper/Private/UsdWrappers/UsdPrim.cpp
+++ b/Engine/Plugins/Importers/USDImporter/Source/UnrealUSDWrapper/Private/UsdWrappers/UsdPrim.cpp
@@ -433,6 +433,34 @@ namespace UE
 #endif // #if USE_USD_SDK
 	}
 
+	bool FUsdPrim::SetTypeName(FName TypeName) const
+	{
+#if USE_USD_SDK
+		FScopedUsdAllocs UsdAllocs;
+		return Impl->PxrUsdPrim.Get().SetTypeName(pxr::TfToken(TCHAR_TO_ANSI(*TypeName.ToString())));
+#else
+		return false;
+#endif // #if USE_USD_SDK
+	}
+
+	bool FUsdPrim::ClearTypeName() const
+	{
+#if USE_USD_SDK
+		return Impl->PxrUsdPrim.Get().ClearTypeName();
+#else
+		return false;
+#endif // #if USE_USD_SDK
+	}
+
+	bool FUsdPrim::HasAuthoredTypeName() const
+	{
+#if USE_USD_SDK
+		return Impl->PxrUsdPrim.Get().HasAuthoredTypeName();
+#else
+		return false;
+#endif // #if USE_USD_SDK
+	}
+
 	FUsdPrim FUsdPrim::GetParent() const
 	{
 #if USE_USD_SDK
diff --git a/Engine/Plugins/Importers/USDImporter/Source/UnrealUSDWrapper/Public/UsdWrappers/UsdPrim.h b/Engine/Plugins/Importers/USDImporter/Source/UnrealUSDWrapper/Public/UsdWrappers/UsdPrim.h
--- a/Engine/Plugins/Importers/USDImporter/Source/UnrealUSDWrapper/Public/UsdWrappers/UsdPrim.h
+++ b/Engine/Plugins/Importers/USDImporter/Source/UnrealUSDWrapper/Public/UsdWrappers/UsdPrim.h
@@ -111,7 +111,11 @@ namespace UE
 		FUsdStage GetStage() const;
 
 		FName GetName() const;
+
 		FName GetTypeName() const;
+		bool SetTypeName(FName TypeName) const;
+		bool ClearTypeName() const;
+		bool HasAuthoredTypeName() const;
 
 		FUsdPrim GetParent() const;
 
