USD: Added support for USD collision APIs
* UsdPhysicsCollisionAPI can be applied to mesh and primitive prims to enable them as colliders.
* UsdPhysicsMeshCollisionAPI can be applied to mesh prim to specify its collision approximation shapes.
Custom collision shapes are specified by applying both UsdPhysicsCollisionAPI and UsdPhysicsMeshCollisionAPI on a mesh prim, setting its approximartion to none and its purpose to guide. The custom collision mesh prim must be a sibling to the rendering mesh.

diff --git a/Engine/Plugins/Importers/USDImporter/Source/USDSchemas/Private/USDGeomMeshTranslator.cpp b/Engine/Plugins/Importers/USDImporter/Source/USDSchemas/Private/USDGeomMeshTranslator.cpp
--- a/Engine/Plugins/Importers/USDImporter/Source/USDSchemas/Private/USDGeomMeshTranslator.cpp
+++ b/Engine/Plugins/Importers/USDImporter/Source/USDSchemas/Private/USDGeomMeshTranslator.cpp
@@ -21,6 +21,7 @@
 #include "UsdWrappers/UsdPrim.h"
 #include "UsdWrappers/UsdStage.h"
 
+#include "CompGeom/FitKDOP3.h"
 #include "Components/StaticMeshComponent.h"
 #include "Engine/CollisionProfile.h"
 #include "Engine/Level.h"
@@ -45,6 +46,7 @@
 #include "UObject/SoftObjectPath.h"
 
 #if WITH_EDITOR
+#include "ConvexDecompTool.h"
 #include "IMeshBuilderModule.h"
 #include "MeshBudgetProjectSettings.h"
 #endif // WITH_EDITOR
@@ -55,6 +57,9 @@
 #include "pxr/usd/usd/typed.h"
 #include "pxr/usd/usdGeom/mesh.h"
 #include "pxr/usd/usdGeom/subset.h"
+#include "pxr/usd/usdPhysics/collisionAPI.h"
+#include "pxr/usd/usdPhysics/meshCollisionAPI.h"
+#include "pxr/usd/usdPhysics/tokens.h"
 #include "USDIncludesEnd.h"
 
 static float GMeshNormalRepairThreshold = 0.05f;
@@ -702,21 +707,6 @@ namespace UsdGeomMeshTranslatorImpl
 
 #endif // WITH_EDITOR
 
-		static IConsoleVariable* CVar = IConsoleManager::Get().FindConsoleVariable(TEXT("USD.EnableCollision"));
-		bool bEnableCollision = CVar && CVar->GetBool();
-
-		if (StaticMesh.GetBodySetup())
-		{
-			if (bEnableCollision)
-			{
-				StaticMesh.GetBodySetup()->CreatePhysicsMeshes();
-			}
-			else
-			{
-				StaticMesh.GetBodySetup()->DefaultInstance.SetCollisionEnabled(ECollisionEnabled::NoCollision);
-				StaticMesh.GetBodySetup()->DefaultInstance.SetCollisionProfileName(UCollisionProfile::NoCollision_ProfileName);
-			}
-		}
 		return true;
 	}
 
@@ -748,6 +738,630 @@ namespace UsdGeomMeshTranslatorImpl
 	}
 }
 
+namespace UE::UsdCollision::Private
+{
+	// Based on GeomFitUtils.cpp
+
+	// k-DOP (k-Discrete Oriented Polytopes) Direction Vectors
+	constexpr float RCP_SQRT2 = 0.70710678118654752440084436210485f;
+
+	TArray<FVector> KDopDir18 = 
+	{
+		FVector( 1.f, 0.f, 0.f),
+		FVector(-1.f, 0.f, 0.f),
+		FVector( 0.f, 1.f, 0.f),
+		FVector( 0.f,-1.f, 0.f),
+		FVector( 0.f, 0.f, 1.f),
+		FVector( 0.f, 0.f,-1.f),
+		FVector( 0.f, RCP_SQRT2,  RCP_SQRT2),
+		FVector( 0.f,-RCP_SQRT2, -RCP_SQRT2),
+		FVector( 0.f, RCP_SQRT2, -RCP_SQRT2),
+		FVector( 0.f,-RCP_SQRT2,  RCP_SQRT2),
+		FVector( RCP_SQRT2, 0.f,  RCP_SQRT2),
+		FVector(-RCP_SQRT2, 0.f, -RCP_SQRT2),
+		FVector( RCP_SQRT2, 0.f, -RCP_SQRT2),
+		FVector(-RCP_SQRT2, 0.f,  RCP_SQRT2),
+		FVector( RCP_SQRT2,  RCP_SQRT2, 0.f),
+		FVector(-RCP_SQRT2, -RCP_SQRT2, 0.f),
+		FVector( RCP_SQRT2, -RCP_SQRT2, 0.f),
+		FVector(-RCP_SQRT2,  RCP_SQRT2, 0.f)
+	};
+
+	void GenerateKDopAsSimpleCollision(const FMeshDescription& MeshDescription, const TArray<FVector> &Dirs, FKAggregateGeom& CollisionShapes)
+	{
+		TArray<FVector> HullVertices;
+		TVertexAttributesConstRef<FVector3f> VertexPositions = MeshDescription.GetVertexPositions();
+		UE::Geometry::FitKDOPVertices3<double>(
+			Dirs, 
+			VertexPositions.GetNumElements(),
+			[&](int32 VertexId) 
+			{
+				return static_cast<FVector>(VertexPositions[VertexId]);
+			},
+			HullVertices);
+
+
+		FKConvexElem ConvexElem;
+		ConvexElem.VertexData = HullVertices;
+		// Note: UpdateElemBox also computes the convex hull indices
+		ConvexElem.UpdateElemBox();
+
+		CollisionShapes.ConvexElems.Add(ConvexElem);
+	}
+
+	void CalcBoundingSphere(const FMeshDescription& MeshDescription, FSphere& Sphere)
+	{
+		FBox Box;
+		FVector MinIx[3];
+		FVector MaxIx[3];
+
+		TVertexAttributesConstRef<FVector3f> VertexPositions = MeshDescription.GetVertexPositions();
+
+		bool bFirstVertex = true;
+		for (const FVertexID VertexID : MeshDescription.Vertices().GetElementIDs())
+		{
+			FVector Point = static_cast<FVector>(VertexPositions[VertexID]);
+			if (bFirstVertex)
+			{
+				// First, find AABB, remembering furthest points in each dir.
+				Box.Min = Point;
+				Box.Max = Box.Min;
+
+				MinIx[0] = static_cast<FVector>(VertexPositions[VertexID]);
+				MinIx[1] = static_cast<FVector>(VertexPositions[VertexID]);
+				MinIx[2] = static_cast<FVector>(VertexPositions[VertexID]);
+
+				MaxIx[0] = static_cast<FVector>(VertexPositions[VertexID]);
+				MaxIx[1] = static_cast<FVector>(VertexPositions[VertexID]);
+				MaxIx[2] = static_cast<FVector>(VertexPositions[VertexID]);
+				bFirstVertex = false;
+				continue;
+			}
+
+			// X //
+			if (Point.X < Box.Min.X)
+			{
+				Box.Min.X = Point.X;
+				MinIx[0] = static_cast<FVector>(VertexPositions[VertexID]);
+			}
+			else if (Point.X > Box.Max.X)
+			{
+				Box.Max.X = Point.X;
+				MaxIx[0] = static_cast<FVector>(VertexPositions[VertexID]);
+			}
+
+			// Y //
+			if (Point.Y < Box.Min.Y)
+			{
+				Box.Min.Y = Point.Y;
+				MinIx[1] = static_cast<FVector>(VertexPositions[VertexID]);
+			}
+			else if (Point.Y > Box.Max.Y)
+			{
+				Box.Max.Y = Point.Y;
+				MaxIx[1] = static_cast<FVector>(VertexPositions[VertexID]);
+			}
+
+			// Z //
+			if (Point.Z < Box.Min.Z)
+			{
+				Box.Min.Z = Point.Z;
+				MinIx[2] = static_cast<FVector>(VertexPositions[VertexID]);
+			}
+			else if (Point.Z > Box.Max.Z)
+			{
+				Box.Max.Z = Point.Z;
+				MaxIx[2] = static_cast<FVector>(VertexPositions[VertexID]);
+			}
+		}
+
+		const FVector Extremes[3] = { (MaxIx[0] - MinIx[0]), (MaxIx[1] - MinIx[1]), (MaxIx[2] - MinIx[2]) };
+
+		// Now find extreme points furthest apart, and initial center and radius of sphere.
+		float MaxDist2 = 0.f;
+		for (int32 i = 0; i < 3; ++i)
+		{
+			const float TmpDist2 = Extremes[i].SizeSquared();
+			if (TmpDist2 > MaxDist2)
+			{
+				MaxDist2 = TmpDist2;
+				Sphere.Center = (MinIx[i] + (0.5f * Extremes[i]));
+				Sphere.W = 0.f;
+			}
+		}
+
+		const FVector Extents = FVector(Extremes[0].X, Extremes[1].Y, Extremes[2].Z);
+
+		// radius and radius squared
+		float Radius = 0.5f * Extents.GetMax();
+		float Radius2 = FMath::Square(Radius);
+
+		// Now check each point lies within this sphere. If not - expand it a bit.
+		for (const FVertexID VertexID : MeshDescription.Vertices().GetElementIDs())
+		{
+			const FVector CenterToPoint = ((FVector)VertexPositions[VertexID]) - Sphere.Center;
+			const float CenterToPoint2 = CenterToPoint.SizeSquared();
+
+			// If this point is outside our current bounding sphere's radius
+			if (CenterToPoint2 > Radius2)
+			{
+				// ..expand radius just enough to include this point.
+				const float PointRadius = FMath::Sqrt(CenterToPoint2);
+				Radius = 0.5f * (Radius + PointRadius);
+				Radius2 = FMath::Square(Radius);
+
+				Sphere.Center += ((PointRadius - Radius) / PointRadius * CenterToPoint);
+			}
+		}
+
+		Sphere.W = Radius;
+	}
+
+	// This is the one thats already used by unreal.
+	// Seems to do better with more symmetric input...
+	void CalcBoundingSphere2(const FMeshDescription& MeshDescription, FSphere& Sphere)
+	{
+		FVector Center = MeshDescription.ComputeBoundingBox().GetCenter();
+
+		Sphere.Center = Center;
+		Sphere.W = 0.0f;
+
+		TVertexAttributesConstRef<FVector3f> VertexPositions = MeshDescription.GetVertexPositions();
+		for (const FVertexID VertexID : MeshDescription.Vertices().GetElementIDs())
+		{
+			float Dist2 = FVector::DistSquared((FVector)VertexPositions[VertexID], Sphere.Center);
+			if (Dist2 > Sphere.W)
+				Sphere.W = Dist2;
+		}
+		Sphere.W = FMath::Sqrt(Sphere.W);
+	}
+
+	void GenerateSphereAsSimpleCollision(const FMeshDescription& MeshDescription, FKAggregateGeom& CollisionShapes)
+	{
+		FSphere Sphere, Sphere2, BestSphere;
+
+		// Calculate bounding sphere.
+		CalcBoundingSphere(MeshDescription, Sphere);
+		CalcBoundingSphere2(MeshDescription, Sphere2);
+
+		if(Sphere.W < Sphere2.W)
+			BestSphere = Sphere;
+		else
+			BestSphere = Sphere2;
+
+		// Don't use if radius is zero.
+		if(BestSphere.W <= 0.f)
+		{
+			return;
+		}
+
+		FKSphereElem SphereElem;
+		SphereElem.Center = BestSphere.Center;
+		SphereElem.Radius = BestSphere.W;
+		CollisionShapes.SphereElems.Add(SphereElem);
+	}
+
+	void GenerateBoxAsSimpleCollision(const FMeshDescription& MeshDescription, FKAggregateGeom& CollisionShapes)
+	{
+		// Calculate bounding Box.
+		FVector Center, Extents;
+		MeshDescription.ComputeBoundingBox().GetCenterAndExtents(Center, Extents);
+
+		FKBoxElem BoxElem;
+		BoxElem.Center = Center;
+		BoxElem.X = Extents.X * 2.0f;
+		BoxElem.Y = Extents.Y * 2.0f;
+		BoxElem.Z = Extents.Z * 2.0f;
+		CollisionShapes.BoxElems.Add(BoxElem);
+	}
+
+	void CalcBoundingSphyl(const FMeshDescription& MeshDescription, FSphere& Sphere, float& Length, FRotator& Rotation)
+	{
+		FVector Center, Extents;
+		MeshDescription.ComputeBoundingBox().GetCenterAndExtents(Center, Extents);
+
+		Sphere.Center = Center;
+
+		// Work out best axis aligned orientation (longest side)
+		double Extent = Extents.GetMax();
+		if (Extent == Extents.X)
+		{
+			Rotation = FRotator(90.f, 0.f, 0.f);
+			Extents.X = 0.0f;
+		}
+		else if (Extent == Extents.Y)
+		{
+			Rotation = FRotator(0.f, 0.f, 90.f);
+			Extents.Y = 0.0f;
+		}
+		else
+		{
+			Rotation = FRotator(0.f, 0.f, 0.f);
+			Extents.Z = 0.0f;
+		}
+
+		// Cleared the largest axis above, remaining determines the radius
+		float Radius = Extents.GetMax();
+		float Radius2 = FMath::Square(Radius);
+
+		TVertexAttributesConstRef<FVector3f> VertexPositions = MeshDescription.GetVertexPositions();
+
+		// Now check each point lies within this the radius. If not - expand it a bit.
+		for (const FVertexID VertexID : MeshDescription.Vertices().GetElementIDs())
+		{
+			FVector CenterToPoint = (static_cast<FVector>(VertexPositions[VertexID])) - Sphere.Center;
+			CenterToPoint = Rotation.UnrotateVector(CenterToPoint);
+
+			const float PointRadius2 = CenterToPoint.SizeSquared2D();	// Ignore Z here...
+
+			// If this point is outside our current bounding sphere's radius
+			if (PointRadius2 > Radius2)
+			{
+				// ..expand radius just enough to include this point.
+				const float PointRadius = FMath::Sqrt(PointRadius2);
+				Radius = 0.5f * (Radius + PointRadius);
+				Radius2 = FMath::Square(Radius);
+			}
+		}
+
+		// The length is the longest side minus the radius.
+		float HalfLength = FMath::Max(0.0f, Extent - Radius);
+
+		// Now check each point lies within the length. If not - expand it a bit.
+		for (const FVertexID VertexID : MeshDescription.Vertices().GetElementIDs())
+		{
+			FVector CenterToPoint = (static_cast<FVector>(VertexPositions[VertexID])) - Sphere.Center;
+			CenterToPoint = Rotation.UnrotateVector(CenterToPoint);
+
+			// If this point is outside our current bounding sphyl's length
+			if (FMath::Abs(CenterToPoint.Z) > HalfLength)
+			{
+				const bool bFlip = (CenterToPoint.Z < 0.f ? true : false);
+				const FVector Origin(0.f, 0.f, (bFlip ? -HalfLength : HalfLength));
+
+				const float PointRadius2 = (Origin - CenterToPoint).SizeSquared();
+
+				// If this point is outside our current bounding sphyl's radius
+				if (PointRadius2 > Radius2)
+				{
+					FVector ClosestPoint;
+					FMath::SphereDistToLine(Origin, Radius, CenterToPoint, (bFlip ? FVector(0.f, 0.f, 1.f) : FVector(0.f, 0.f, -1.f)), ClosestPoint);
+
+					// Don't accept zero as a valid diff when we know it's outside the sphere (saves needless retest on further iterations of like points)
+					HalfLength += FMath::Max<float>(FMath::Abs(CenterToPoint.Z - ClosestPoint.Z), 1.e-6f);
+				}
+			}
+		}
+
+		Sphere.W = Radius;
+		Length = HalfLength * 2.0f;
+	}
+
+	void GenerateSphylAsSimpleCollision(const FMeshDescription& MeshDescription, FKAggregateGeom& CollisionShapes)
+	{
+		FSphere Sphere;
+		float Length;
+		FRotator Rotation;
+
+		// Calculate bounding box.
+		CalcBoundingSphyl(MeshDescription, Sphere, Length, Rotation);
+
+		// Don't use if radius is zero.
+		if (Sphere.W <= 0.f)
+		{
+			return;
+		}
+
+		// If height is zero, then a sphere would be better (should we just create one instead?)
+		if (Length <= 0.f)
+		{
+			Length = SMALL_NUMBER;
+		}
+
+		FKSphylElem SphylElem;
+		SphylElem.Center = Sphere.Center;
+		SphylElem.Rotation = Rotation;
+		SphylElem.Radius = Sphere.W;
+		SphylElem.Length = Length;
+		CollisionShapes.SphylElems.Add(SphylElem);
+	}
+
+	void ConvertMeshToSimpleCollision(const FMeshDescription& MeshDescription, FKAggregateGeom& CollisionShapes)
+	{
+		FKConvexElem ConvexElem;
+		ConvexElem.VertexData.Reserve(MeshDescription.Vertices().Num());
+
+		TVertexAttributesConstRef<FVector3f> VertexPositions = MeshDescription.GetVertexPositions();
+		for (const FVertexID VertexID : MeshDescription.Vertices().GetElementIDs())
+		{
+			ConvexElem.VertexData.Add(static_cast<FVector>(VertexPositions[VertexID]));
+		}
+
+		// Note: UpdateElemBox also computes the convex hull indices
+		ConvexElem.UpdateElemBox();
+		CollisionShapes.ConvexElems.Add(ConvexElem);
+	}
+
+	// Approximation types from PhysicsMeshCollisionAPI and UE-specific approximations
+	enum class EUsdCollisionType : uint8
+	{
+		None,
+		ConvexDecomposition,
+		ConvexHull,
+		Sphere,
+		Cube,
+		MeshSimplification,
+		Capsule,
+		CustomMesh
+	};
+
+	void CollectCustomCollisionShapes(const pxr::UsdPrim& UsdPrim, FKAggregateGeom& CollisionShapes)
+	{
+		UsdToUnreal::FUsdMeshConversionOptions Options;
+		Options.PurposesToLoad = EUsdPurpose::Guide; // custom collision mesh must have guide purpose
+		Options.bMergeIdenticalMaterialSlots = false;
+
+		FTransform ParentTransform;
+		UsdToUnreal::ConvertXformable(UsdPrim.GetStage(), pxr::UsdGeomMesh{UsdPrim}, ParentTransform, Options.TimeCode.GetValue());
+
+		for (const pxr::UsdPrim& ChildPrim : UsdPrim.GetParent().GetChildren())
+		{
+			if (UsdUtils::IsCollisionMesh(ChildPrim))
+			{
+				FMeshDescription MeshDescription;
+				UsdUtils::FUsdPrimMaterialAssignmentInfo MaterialAssignments;
+
+				FStaticMeshAttributes MeshAttributes(MeshDescription);
+				MeshAttributes.Register();
+
+				pxr::UsdGeomMesh UsdMesh{ChildPrim};
+
+				// Since the custom collision shapes are "attached" to the static mesh, its transform must be
+				// removed when their own transforms are baked into the collision shapes
+				UsdToUnreal::ConvertXformable(ChildPrim.GetStage(), UsdMesh, Options.AdditionalTransform, Options.TimeCode.GetValue());
+				Options.AdditionalTransform = Options.AdditionalTransform * ParentTransform.Inverse();
+
+				if (UsdToUnreal::ConvertGeomMesh(UsdMesh, MeshDescription, MaterialAssignments, Options))
+				{
+					// By default, use the custom mesh as convex hull, which takes care of the UCX_ prefix
+					EUsdCollisionType Approximation(EUsdCollisionType::ConvexHull);
+
+					// Use the FBX custom collision name convention as a hint
+					FString PrimName(UsdToUnreal::ConvertToken(ChildPrim.GetName()));
+					if (PrimName.StartsWith(TEXT("UBX_")))
+					{
+						Approximation = EUsdCollisionType::Cube;
+					}
+					else if (PrimName.StartsWith(TEXT("UCP_")))
+					{
+						Approximation = EUsdCollisionType::Capsule;
+					}
+					else if (PrimName.StartsWith(TEXT("USP_")))
+					{
+						Approximation = EUsdCollisionType::Sphere;
+					}
+
+					switch (Approximation)
+					{
+						case EUsdCollisionType::ConvexHull:
+						{
+							GenerateKDopAsSimpleCollision(MeshDescription, KDopDir18, CollisionShapes);
+							break;
+						}
+						case EUsdCollisionType::Cube:
+						{
+							GenerateBoxAsSimpleCollision(MeshDescription, CollisionShapes);
+							break;
+						}
+						case EUsdCollisionType::Capsule:
+						{
+							GenerateSphylAsSimpleCollision(MeshDescription, CollisionShapes);
+							break;
+						}
+						case EUsdCollisionType::Sphere:
+						{
+							GenerateSphereAsSimpleCollision(MeshDescription, CollisionShapes);
+							break;
+						}
+					}
+				}
+			}
+		}
+	}
+
+	void SetupSimpleCollision(const pxr::UsdPrim& UsdPrim, UStaticMesh& StaticMesh)
+	{
+		UBodySetup* BodySetup = StaticMesh.GetBodySetup();
+		if (!BodySetup)
+		{
+			return;
+		}
+
+		bool bIsCollisionEnabled = false;
+		if (pxr::UsdPhysicsCollisionAPI CollisionAPI{UsdPrim})
+		{
+			if (pxr::UsdAttribute CollisionAttr = CollisionAPI.GetCollisionEnabledAttr())
+			{
+				CollisionAttr.Get(&bIsCollisionEnabled);
+			}
+		}
+
+		static IConsoleVariable* CVar = IConsoleManager::Get().FindConsoleVariable(TEXT("USD.EnableCollision"));
+		const bool bEnableCollision = CVar && CVar->GetBool();
+
+		if (!bIsCollisionEnabled || !bEnableCollision)
+		{
+			BodySetup->Modify();
+			BodySetup->RemoveSimpleCollision();
+			BodySetup->bNeverNeedsCookedCollisionData = true; // this will prevent the complex collision mesh from being generated
+			BodySetup->DefaultInstance.SetCollisionEnabled(ECollisionEnabled::NoCollision);
+			BodySetup->DefaultInstance.SetCollisionProfileName(UCollisionProfile::NoCollision_ProfileName);
+
+			return;
+		}
+
+		FMeshDescription* MeshDescription = StaticMesh.GetMeshDescription(0);
+		if (!MeshDescription || MeshDescription->Vertices().Num() == 0)
+		{
+			return;
+		}
+
+		EUsdCollisionType Approximation(EUsdCollisionType::None);
+		if (pxr::UsdGeomMesh UsdMesh{UsdPrim})
+		{
+			// Get the collision approximation type (only meshes should have UsdPhysicsMeshCollisionAPI)
+			if (pxr::UsdPhysicsMeshCollisionAPI MeshCollisionAPI{UsdMesh})
+			{
+				pxr::TfToken ApproximationAttr(pxr::UsdPhysicsTokens->none);
+				if (MeshCollisionAPI.GetApproximationAttr())
+				{
+					MeshCollisionAPI.GetApproximationAttr().Get(&ApproximationAttr);
+				}
+				
+				if (ApproximationAttr == pxr::UsdPhysicsTokens->convexDecomposition)
+				{
+#if WITH_EDITOR
+					Approximation = EUsdCollisionType::ConvexDecomposition;
+#else
+					Approximation = EUsdCollisionType::ConvexHull;
+#endif
+				}
+				else if (ApproximationAttr == pxr::UsdPhysicsTokens->convexHull)
+				{
+					Approximation = EUsdCollisionType::ConvexHull;
+				}
+				else if (ApproximationAttr == pxr::UsdPhysicsTokens->boundingSphere)
+				{
+					Approximation = EUsdCollisionType::Sphere;
+				}
+				else if (ApproximationAttr == pxr::UsdPhysicsTokens->boundingCube)
+				{
+					Approximation = EUsdCollisionType::Cube;
+				}
+				else if (ApproximationAttr == pxr::UsdPhysicsTokens->meshSimplification)
+				{
+					Approximation = EUsdCollisionType::MeshSimplification;
+				}
+			}
+		}
+		else
+		{
+			// Collision for primitives are converted to their closest approximation
+			if (UsdPrim.IsA(pxr::UsdGeomTokens->Capsule))
+			{
+				Approximation = EUsdCollisionType::Capsule;
+			}
+			else if (UsdPrim.IsA(pxr::UsdGeomTokens->Cone))
+			{
+				Approximation = EUsdCollisionType::CustomMesh;
+			}
+			else if (UsdPrim.IsA(pxr::UsdGeomTokens->Cube))
+			{
+				Approximation = EUsdCollisionType::Cube;
+			}
+			else if (UsdPrim.IsA(pxr::UsdGeomTokens->Cylinder))
+			{
+				Approximation = EUsdCollisionType::CustomMesh;
+			}
+			else if (UsdPrim.IsA(pxr::UsdGeomTokens->Sphere))
+			{
+				Approximation = EUsdCollisionType::Sphere;
+			}
+			else if (UsdPrim.IsA(pxr::UsdGeomTokens->Plane))
+			{
+				Approximation = EUsdCollisionType::CustomMesh;
+			}
+		}
+
+		BodySetup->Modify();
+		BodySetup->RemoveSimpleCollision();
+		BodySetup->bNeverNeedsCookedCollisionData = false;
+		BodySetup->DefaultInstance.SetCollisionEnabled(ECollisionEnabled::QueryAndPhysics);
+		BodySetup->DefaultInstance.SetCollisionProfileName(UCollisionProfile::BlockAllDynamic_ProfileName);
+
+		switch (Approximation)
+		{
+			case EUsdCollisionType::None:
+			{
+				// No approximation but allow for custom collision shapes
+				FKAggregateGeom CollisionShapes;
+				CollectCustomCollisionShapes(UsdPrim, CollisionShapes);
+				if (CollisionShapes.GetElementCount())
+				{
+					BodySetup->AddCollisionFrom(CollisionShapes);
+				}
+				break;
+			}
+#if WITH_EDITOR
+			case EUsdCollisionType::ConvexDecomposition:
+			case EUsdCollisionType::MeshSimplification:
+			{
+				// Could use IMeshReduction for MeshSimplification, but it will get converted to a single convex hull
+				// so the result would be not as good as convex decomposition anyway
+
+				TArray<FVector3f> Vertices;
+				Vertices.Reserve(MeshDescription->Vertices().Num());
+				TVertexAttributesConstRef<FVector3f> VertexPositions = MeshDescription->GetVertexPositions();
+				for (const FVertexID VertexID : MeshDescription->Vertices().GetElementIDs())
+				{
+					Vertices.Add(VertexPositions[VertexID]);
+				}
+
+				TArray<uint32> Indices;
+				Indices.Reserve(MeshDescription->VertexInstances().Num());
+				for (const FVertexInstanceID InstanceID : MeshDescription->VertexInstances().GetElementIDs())
+				{
+					Indices.Add(MeshDescription->GetVertexInstanceVertex(InstanceID).GetValue());
+				}
+
+				// Do not perform any action if we have invalid input
+				if (Vertices.Num() >= 3 && Indices.Num() >= 3)
+				{
+					const uint32 HullCount = 32;
+					const int32 MaxHullVertices = 32;
+					DecomposeMeshToHulls(BodySetup, Vertices, Indices, HullCount, MaxHullVertices);
+				}
+				break;
+			}
+#endif // WITH_EDITOR
+			case EUsdCollisionType::ConvexHull:
+			{
+				GenerateKDopAsSimpleCollision(*MeshDescription, KDopDir18, BodySetup->AggGeom);
+				break;
+			}
+			case EUsdCollisionType::Sphere:
+			{
+				GenerateSphereAsSimpleCollision(*MeshDescription, BodySetup->AggGeom);
+				break;
+			}
+			case EUsdCollisionType::Cube:
+			{
+				GenerateBoxAsSimpleCollision(*MeshDescription, BodySetup->AggGeom);
+				break;
+			}
+			case EUsdCollisionType::Capsule:
+			{
+				GenerateSphylAsSimpleCollision(*MeshDescription, BodySetup->AggGeom);
+				break;
+			}
+			case EUsdCollisionType::CustomMesh:
+			{
+				ConvertMeshToSimpleCollision(*MeshDescription, BodySetup->AggGeom);
+				break;
+			}
+			default:
+				break;
+		}
+
+		BodySetup->CreatePhysicsMeshes();
+
+		StaticMesh.bCustomizedCollision = true;
+
+		const bool bIsUpdate = true;
+		StaticMesh.CreateNavCollision(bIsUpdate);
+	}
+}
+
 FBuildStaticMeshTaskChain::FBuildStaticMeshTaskChain(
 	const TSharedRef< FUsdSchemaTranslationContext >& InContext,
 	const UE::FSdfPath& InPrimPath,
@@ -847,6 +1461,11 @@ void FBuildStaticMeshTaskChain::SetupTasks()
 					}
 #endif // WITH_EDITOR
 				}
+				else
+				{
+					// Setup collision on existing mesh in case the collision settings have changed
+					UE::UsdCollision::Private::SetupSimpleCollision(GetPrim(), *StaticMesh);
+				}
 			}
 
 			// Only need to continue building the mesh if we just created it
@@ -906,6 +1525,8 @@ void FBuildStaticMeshTaskChain::SetupTasks()
 				return false;
 			}
 
+			UE::UsdCollision::Private::SetupSimpleCollision(GetPrim(), *StaticMesh);
+
 			return true;
 		});
 
@@ -1009,6 +1630,11 @@ void FUsdGeomMeshTranslator::CreateAssets()
 		return;
 	}
 
+	if (UsdUtils::IsCollisionMesh(GetPrim()))
+	{
+		return;
+	}
+
 	TSharedRef< FGeomMeshCreateAssetsTaskChain > AssetsTaskChain = MakeShared< FGeomMeshCreateAssetsTaskChain >(Context, PrimPath);
 
 	Context->TranslatorTasks.Add(MoveTemp(AssetsTaskChain));
@@ -1016,6 +1642,10 @@ void FUsdGeomMeshTranslator::CreateAssets()
 
 USceneComponent* FUsdGeomMeshTranslator::CreateComponents()
 {
+	if (UsdUtils::IsCollisionMesh(GetPrim()))
+	{
+		return nullptr;
+	}
 	return Super::CreateComponents();
 }
 
@@ -1083,6 +1713,12 @@ bool FUsdGeomMeshTranslator::CanBeCollapsed(ECollapsingType CollapsingType) cons
 		return false;
 	}
 
+	// Prevent collapse of custom collision mesh
+	if (UsdUtils::IsCollisionMesh(Prim))
+	{
+		return false;
+	}
+
 	return Super::CanBeCollapsed(CollapsingType);
 }
 
@@ -1099,6 +1735,11 @@ TSet<UE::FSdfPath> FUsdGeomMeshTranslator::CollectAuxiliaryPrims() const
 	}
 
 	TSet<UE::FSdfPath> Result;
+	if (UsdUtils::IsCollisionMesh(GetPrim()))
+	{
+		return Result;
+	}
+
 	{
 		FScopedUsdAllocs UsdAllocs;
 
@@ -1116,7 +1757,17 @@ TSet<UE::FSdfPath> FUsdGeomMeshTranslator::CollectAuxiliaryPrims() const
 		{
 			Result.Add(UE::FSdfPath{ChildPrim.Get().GetPrimPath()});
 		}
+
+		// For mesh prim, collect the sibling collision meshes as auxiliary prims
+		for (const pxr::UsdPrim& ChildPrim : GetPrim().GetParent().GetChildren())
+		{
+			if (UsdUtils::IsCollisionMesh(ChildPrim))
+			{
+				Result.Add(UE::FSdfPath{ChildPrim.GetPrimPath()});
+			}
+		}
 	}
+
 	return Result;
 }
 
diff --git a/Engine/Plugins/Importers/USDImporter/Source/USDSchemas/Private/USDGeomXformableTranslator.cpp b/Engine/Plugins/Importers/USDImporter/Source/USDSchemas/Private/USDGeomXformableTranslator.cpp
--- a/Engine/Plugins/Importers/USDImporter/Source/USDSchemas/Private/USDGeomXformableTranslator.cpp
+++ b/Engine/Plugins/Importers/USDImporter/Source/USDSchemas/Private/USDGeomXformableTranslator.cpp
@@ -33,6 +33,7 @@
 #include "LiveLinkRole.h"
 #include "Misc/App.h"
 #include "Modules/ModuleManager.h"
+#include "PhysicsEngine/BodySetup.h"
 #include "Roles/LiveLinkTransformRole.h"
 #include "StaticMeshAttributes.h"
 
@@ -670,6 +671,13 @@ void FUsdGeomXformableTranslator::UpdateComponents( USceneComponent* SceneCompon
 				// We can't register yet, as UsdToUnreal::ConvertXformable below us may want to move the component.
 				// We'll always re-register when needed below, though.
 			}
+
+			// Update the collision settings of the component in case they changed on the static mesh
+			if (PrimStaticMesh && PrimStaticMesh->GetBodySetup())
+			{
+				StaticMeshComponent->BodyInstance.SetCollisionEnabled(PrimStaticMesh->GetBodySetup()->DefaultInstance.GetCollisionEnabled());
+				StaticMeshComponent->BodyInstance.SetCollisionProfileName(PrimStaticMesh->GetBodySetup()->DefaultInstance.GetCollisionProfileName());
+			}
 		}
 		else if (UUsdDrawModeComponent* BoundsComponent = Cast<UUsdDrawModeComponent>(SceneComponent))
 		{
diff --git a/Engine/Plugins/Importers/USDImporter/Source/USDSchemas/USDSchemas.Build.cs b/Engine/Plugins/Importers/USDImporter/Source/USDSchemas/USDSchemas.Build.cs
--- a/Engine/Plugins/Importers/USDImporter/Source/USDSchemas/USDSchemas.Build.cs
+++ b/Engine/Plugins/Importers/USDImporter/Source/USDSchemas/USDSchemas.Build.cs
@@ -28,6 +28,7 @@ namespace UnrealBuildTool.Rules
 					"CoreUObject",
 					"Engine",
 					"GeometryCache",
+					"GeometryCore", // for FitKDOP
 					"InterchangeCore",
 					"InterchangeEngine",
 					"InterchangeFactoryNodes",
@@ -68,7 +69,7 @@ namespace UnrealBuildTool.Rules
 						"MaterialEditor",
 						"MDLImporter",
 						"MeshUtilities",
-						"PhysicsUtilities", // For generating UPhysicsAssets for SkeletalMeshes
+						"PhysicsUtilities", // For generating UPhysicsAssets for SkeletalMeshes and ConvexDecompTool
 						"PropertyEditor",
 						"UnrealEd",
 					}
diff --git a/Engine/Plugins/Importers/USDImporter/Source/USDStageEditor/Private/SUSDStageTreeView.cpp b/Engine/Plugins/Importers/USDImporter/Source/USDStageEditor/Private/SUSDStageTreeView.cpp
--- a/Engine/Plugins/Importers/USDImporter/Source/USDStageEditor/Private/SUSDStageTreeView.cpp
+++ b/Engine/Plugins/Importers/USDImporter/Source/USDStageEditor/Private/SUSDStageTreeView.cpp
@@ -42,6 +42,10 @@
 
 #if USE_USD_SDK
 
+#include "USDIncludesStart.h"
+	#include "pxr/usd/usdPhysics/tokens.h"
+#include "USDIncludesEnd.h"
+
 namespace UE::USDStageTreeView::Private
 {
 	static const FText NoSpecOnLocalLayerStack =
@@ -949,6 +953,86 @@ TSharedPtr< SWidget > SUsdStageTreeView::ConstructPrimContextMenu()
 			EUserInterfaceActionType::Button
 		);
 
+		PrimOptions.AddMenuEntry(
+			LOCTEXT("ApplyPhysicsCollisionchema", "Apply PhysicsCollision schema"),
+			LOCTEXT("ApplyPhysicsCollisionSchema_ToolTip", "Applies the PhysicsCollision schema to enable the prim as a collider"),
+			FSlateIcon(),
+			FUIAction(
+				FExecuteAction::CreateSP(
+					this,
+					&SUsdStageTreeView::OnApplySchema,
+					FName{*UsdToUnreal::ConvertToken(pxr::UsdPhysicsTokens->PhysicsCollisionAPI)}
+				),
+				FCanExecuteAction::CreateSP(
+					this,
+					&SUsdStageTreeView::CanApplySchema,
+					FName{*UsdToUnreal::ConvertToken(pxr::UsdPhysicsTokens->PhysicsCollisionAPI)}
+				)
+			),
+			NAME_None,
+			EUserInterfaceActionType::Button
+		);
+
+		PrimOptions.AddMenuEntry(
+			LOCTEXT("RemovePhysicsCollisionSchema", "Remove PhysicsCollision schema"),
+			LOCTEXT("RemovePhysicsCollisionSchema_ToolTip", "Removes the PhysicsCollision schema from the prim to disable it as a collider"),
+			FSlateIcon(),
+			FUIAction(
+				FExecuteAction::CreateSP(
+					this,
+					&SUsdStageTreeView::OnRemoveSchema,
+					FName{*UsdToUnreal::ConvertToken(pxr::UsdPhysicsTokens->PhysicsCollisionAPI)}
+				),
+				FCanExecuteAction::CreateSP(
+					this,
+					&SUsdStageTreeView::CanRemoveSchema,
+					FName{*UsdToUnreal::ConvertToken(pxr::UsdPhysicsTokens->PhysicsCollisionAPI)}
+				)
+			),
+			NAME_None,
+			EUserInterfaceActionType::Button
+		);
+
+		PrimOptions.AddMenuEntry(
+			LOCTEXT("ApplyPhysicsMeshCollisionchema", "Apply PhysicsMeshCollision schema"),
+			LOCTEXT("ApplyPhysicsMeshCollisionSchema_ToolTip", "Applies the PhysicsMeshCollision schema to control how the mesh is made into a collider"),
+			FSlateIcon(),
+			FUIAction(
+				FExecuteAction::CreateSP(
+					this,
+					&SUsdStageTreeView::OnApplySchema,
+					FName{*UsdToUnreal::ConvertToken(pxr::UsdPhysicsTokens->PhysicsMeshCollisionAPI)}
+				),
+				FCanExecuteAction::CreateSP(
+					this,
+					&SUsdStageTreeView::CanApplySchema,
+					FName{*UsdToUnreal::ConvertToken(pxr::UsdPhysicsTokens->PhysicsMeshCollisionAPI)}
+				)
+			),
+			NAME_None,
+			EUserInterfaceActionType::Button
+		);
+
+		PrimOptions.AddMenuEntry(
+			LOCTEXT("RemovePhysicsMeshCollisionSchema", "Remove PhysicsMeshCollision schema"),
+			LOCTEXT("RemovePhysicsMeshCollisionSchema_ToolTip", "Removes the PhysicsMeshCollision schema from the prim"),
+			FSlateIcon(),
+			FUIAction(
+				FExecuteAction::CreateSP(
+					this,
+					&SUsdStageTreeView::OnRemoveSchema,
+					FName{*UsdToUnreal::ConvertToken(pxr::UsdPhysicsTokens->PhysicsMeshCollisionAPI)}
+				),
+				FCanExecuteAction::CreateSP(
+					this,
+					&SUsdStageTreeView::CanRemoveSchema,
+					FName{*UsdToUnreal::ConvertToken(pxr::UsdPhysicsTokens->PhysicsMeshCollisionAPI)}
+				)
+			),
+			NAME_None,
+			EUserInterfaceActionType::Button
+		);
+
 	}
 	PrimOptions.EndSection();
 
diff --git a/Engine/Plugins/Importers/USDImporter/Source/USDStageEditorViewModels/Private/USDObjectFieldViewModel.cpp b/Engine/Plugins/Importers/USDImporter/Source/USDStageEditorViewModels/Private/USDObjectFieldViewModel.cpp
--- a/Engine/Plugins/Importers/USDImporter/Source/USDStageEditorViewModels/Private/USDObjectFieldViewModel.cpp
+++ b/Engine/Plugins/Importers/USDImporter/Source/USDStageEditorViewModels/Private/USDObjectFieldViewModel.cpp
@@ -33,6 +33,7 @@
 	#include "pxr/usd/usd/prim.h"
 	#include "pxr/usd/usd/relationship.h"
 	#include "pxr/usd/usdGeom/tokens.h"
+	#include "pxr/usd/usdPhysics/tokens.h"
 #include "USDIncludesEnd.h"
 
 #endif // #if USE_USD_SDK
@@ -97,6 +98,19 @@ TArray< TSharedPtr< FString > > FUsdObjectFieldViewModel::GetDropdownOptions() c
 			MakeShared< FString >( UsdToUnreal::ConvertToken( pxr::UsdGeomTokens->invisible ) ),
 		};
 	}
+	else if (Label == UsdToUnreal::ConvertToken(pxr::UsdPhysicsTokens->physicsApproximation))
+	{
+		return TArray<TSharedPtr<FString>>
+		{
+			MakeShared<FString>(UsdToUnreal::ConvertToken(pxr::UsdPhysicsTokens->none)),
+			MakeShared<FString>(UsdToUnreal::ConvertToken(pxr::UsdPhysicsTokens->convexDecomposition)),
+			MakeShared<FString>(UsdToUnreal::ConvertToken(pxr::UsdPhysicsTokens->convexHull)),
+			MakeShared<FString>(UsdToUnreal::ConvertToken(pxr::UsdPhysicsTokens->boundingSphere)),
+			MakeShared<FString>(UsdToUnreal::ConvertToken(pxr::UsdPhysicsTokens->boundingCube)),
+			// meshSimplification will get mapped functionally to convexDecomposition
+		};
+	}
+
 #endif // #if USE_USD_SDK
 
 	return {};
diff --git a/Engine/Plugins/Importers/USDImporter/Source/USDUtilities/Private/USDGeomMeshConversion.cpp b/Engine/Plugins/Importers/USDImporter/Source/USDUtilities/Private/USDGeomMeshConversion.cpp
--- a/Engine/Plugins/Importers/USDImporter/Source/USDUtilities/Private/USDGeomMeshConversion.cpp
+++ b/Engine/Plugins/Importers/USDImporter/Source/USDUtilities/Private/USDGeomMeshConversion.cpp
@@ -60,6 +60,9 @@
 	#include "pxr/usd/usdGeom/sphere.h"
 	#include "pxr/usd/usdGeom/subset.h"
 	#include "pxr/usd/usdGeom/tokens.h"
+	#include "pxr/usd/usdPhysics/collisionAPI.h"
+	#include "pxr/usd/usdPhysics/meshCollisionAPI.h"
+	#include "pxr/usd/usdPhysics/tokens.h"
 	#include "pxr/usd/usdShade/material.h"
 	#include "pxr/usd/usdShade/materialBindingAPI.h"
 	#include "pxr/usd/usdShade/tokens.h"
@@ -4343,6 +4346,69 @@ bool UsdUtils::IsGeomMeshALOD( const pxr::UsdPrim& UsdMeshPrim )
 	return UE::UsdGeomMeshConversion::Private::DoesPrimContainMeshLODsInternal( UsdMeshPrim.GetParent() );
 }
 
+bool UsdUtils::IsCollisionMesh(const pxr::UsdPrim& UsdPrim)
+{
+	// https://openusd.org/release/api/usd_physics_page_front.html#usdPhysics_collision_shapes
+	// "Collision meshes may be specified explicitly by adding the custom collider mesh as a sibling to the
+	// original graphics mesh, UsdGeomImageable purpose to "guide" so it does not render, and apply
+	// UsdPhysicsCollisionAPI and UsdPhysicsMeshCollisionAPI to it specifying no approximation."
+	FScopedUsdAllocs Allocs;
+
+	pxr::UsdGeomMesh UsdMesh{UsdPrim};
+	if (!UsdMesh)
+	{
+		return false;
+	}
+
+	bool bIsCollisionEnabled = false;
+	if (pxr::UsdPhysicsCollisionAPI CollisionAPI{UsdPrim})
+	{
+		if (pxr::UsdAttribute CollisionAttr = CollisionAPI.GetCollisionEnabledAttr())
+		{
+			CollisionAttr.Get(&bIsCollisionEnabled);
+		}
+
+		if (!bIsCollisionEnabled)
+		{
+			return false;
+		}
+	}
+	else
+	{
+		return false;
+	}
+
+	if (pxr::UsdPhysicsMeshCollisionAPI MeshCollisionAPI{UsdPrim})
+	{
+		pxr::TfToken Approximation(pxr::UsdPhysicsTokens->none);
+		if (pxr::UsdAttribute ApproximationAttr = MeshCollisionAPI.GetApproximationAttr())
+		{
+			ApproximationAttr.Get(&Approximation);
+		}
+
+		if (Approximation != pxr::UsdPhysicsTokens->none)
+		{
+			return false;
+		}
+	}
+	else
+	{
+		return false;
+	}
+
+	if (pxr::UsdAttribute PurposeAttr = UsdMesh.GetPurposeAttr())
+	{
+		pxr::TfToken Purpose;
+		PurposeAttr.Get(&Purpose);
+		if (Purpose != pxr::UsdGeomTokens->guide)
+		{
+			return false;
+		}
+	}
+
+	return true;
+}
+
 int32 UsdUtils::GetNumberOfLODVariants( const pxr::UsdPrim& Prim )
 {
 	FScopedUsdAllocs Allocs;
diff --git a/Engine/Plugins/Importers/USDImporter/Source/USDUtilities/Public/USDGeomMeshConversion.h b/Engine/Plugins/Importers/USDImporter/Source/USDUtilities/Public/USDGeomMeshConversion.h
--- a/Engine/Plugins/Importers/USDImporter/Source/USDUtilities/Public/USDGeomMeshConversion.h
+++ b/Engine/Plugins/Importers/USDImporter/Source/USDUtilities/Public/USDGeomMeshConversion.h
@@ -346,6 +346,9 @@ namespace UsdUtils
 	/** Returns whether this UsdMesh can be interpreted as a LOD of a mesh with multiple LODs (i.e. if the Mesh prim is inside a LOD variant) */
 	USDUTILITIES_API bool IsGeomMeshALOD( const pxr::UsdPrim& UsdMeshPrim );
 
+	/** Returns whether this UsdPrim is a custom collision mesh */
+	USDUTILITIES_API bool IsCollisionMesh(const pxr::UsdPrim& UsdPrim);
+
 	/** Returns how many LOD variants the Prim has. Note that this will return 0 if called on one of the LOD meshes themselves, it's meant to be called on its parent */
 	USDUTILITIES_API int32 GetNumberOfLODVariants( const pxr::UsdPrim& Prim );
 
