USD: Better handle nested skeletal data on import/export;

Fix root motion handling option misbehaving on import;
Fix AnimBlueprint not being created (incorrect check on the start of CreateAnimBlueprint).

diff --git a/Engine/Plugins/Importers/USDImporter/Source/USDExporter/Private/AnimSequenceExporterUSD.cpp b/Engine/Plugins/Importers/USDImporter/Source/USDExporter/Private/AnimSequenceExporterUSD.cpp
--- a/Engine/Plugins/Importers/USDImporter/Source/USDExporter/Private/AnimSequenceExporterUSD.cpp
+++ b/Engine/Plugins/Importers/USDImporter/Source/USDExporter/Private/AnimSequenceExporterUSD.cpp
@@ -308,11 +308,22 @@ bool UAnimSequenceExporterUSD::ExportBinary( UObject* Object, const TCHAR* Type,
 			return false;
 		}
 
-		AnimationStage.SetDefaultPrim( SkelRootPrim );
-		UsdUtils::BindAnimationSource( SkelRootPrim, SkelAnimPrim );
+		UE::FSdfPath SkeletonPath = SkelRootPath.AppendChild(UnrealIdentifiers::ExportedSkeletonPrimName);
+		UE::FUsdPrim SkeletonPrim = AnimationStage.DefinePrim(SkeletonPath, TEXT("Skeleton"));
+		if (!SkeletonPrim)
+		{
+			return false;
+		}
+
+		AnimationStage.SetDefaultPrim(SkelRootPrim);
 
 		// Add a reference to the SkelRoot of the static mesh, which will compose in the Mesh and Skeleton prims
-		UsdUtils::AddReference( SkelRootPrim, *MeshAssetFile );
+		UsdUtils::AddReference(SkelRootPrim, *MeshAssetFile);
+
+		// We bind the animation directly to the Skeleton (and not the skel root) because binding it to the SkelRoot
+		// may lead to trouble when we're exporting nested SkeletalMeshComponents, as it will be inherited by
+		// all child Skeletons (even the ones that wouldn't otherwise receive any animation)
+		UsdUtils::BindAnimationSource(SkeletonPrim, SkelAnimPrim);
 	}
 
 	// Configure stage metadata
diff --git a/Engine/Plugins/Importers/USDImporter/Source/USDExporter/Private/LevelSequenceExporterUSD.cpp b/Engine/Plugins/Importers/USDImporter/Source/USDExporter/Private/LevelSequenceExporterUSD.cpp
--- a/Engine/Plugins/Importers/USDImporter/Source/USDExporter/Private/LevelSequenceExporterUSD.cpp
+++ b/Engine/Plugins/Importers/USDImporter/Source/USDExporter/Private/LevelSequenceExporterUSD.cpp
@@ -995,15 +995,29 @@ namespace UE
 							// original skeletal animation track sections as disabled, so they'd fail the "IsTrackAnimated" check above
 							else if ( Track->IsA< UMovieSceneSkeletalAnimationTrack >() || Track->IsA< UMovieSceneControlRigParameterTrack >() )
 							{
-								if ( USkeletalMeshComponent* SkeletalBoundComponent = Cast<USkeletalMeshComponent>( BoundComponent ) )
+								if (USkeletalMeshComponent* SkeletalBoundComponent = Cast<USkeletalMeshComponent>(BoundComponent))
 								{
-									UE::FUsdPrim SkelAnimPrim = UsdStage.DefinePrim( UE::FSdfPath{ *PrimPath }.AppendChild( TEXT( "Anim" ) ), TEXT( "SkelAnimation" ) );
-									if ( !SkelAnimPrim )
+									UE::FUsdPrim SkelAnimPrim = UsdStage.DefinePrim(
+										UE::FSdfPath{*PrimPath}.AppendChild(TEXT("Anim")),
+										TEXT("SkelAnimation")
+									);
+
+									UE::FUsdPrim SkeletonPrim = UsdStage.DefinePrim(
+										UE::FSdfPath{*PrimPath}.AppendChild(UnrealIdentifiers::ExportedSkeletonPrimName),
+										TEXT("Skeleton")
+									);
+
+									if (!SkelAnimPrim || !SkeletonPrim)
 									{
-										UE_LOG( LogUsd, Warning, TEXT( "Failed to generate SkelAnimation prim when baking out SkelRoot '%s'" ), *PrimPath );
+										UE_LOG(
+											LogUsd,
+											Warning,
+											TEXT("Failed to generate Skeleton or SkelAnimation prim when baking out SkelRoot '%s'"),
+											*PrimPath
+										);
 										continue;
 									}
-									UnrealToUsd::CreateSkeletalAnimationBaker( Prim, SkelAnimPrim, *SkeletalBoundComponent, Baker );
+									UnrealToUsd::CreateSkeletalAnimationBaker(SkeletonPrim, SkelAnimPrim, *SkeletalBoundComponent, Baker);
 								}
 							}
 
diff --git a/Engine/Plugins/Importers/USDImporter/Source/USDSchemas/Private/MeshTranslationImpl.h b/Engine/Plugins/Importers/USDImporter/Source/USDSchemas/Private/MeshTranslationImpl.h
--- a/Engine/Plugins/Importers/USDImporter/Source/USDSchemas/Private/MeshTranslationImpl.h
+++ b/Engine/Plugins/Importers/USDImporter/Source/USDSchemas/Private/MeshTranslationImpl.h
@@ -26,7 +26,7 @@ namespace UsdUtils
 	struct FUsdPrimMaterialAssignmentInfo;
 }
 
-/** Implementation that can be shared between the SkelRoot translator and GeomMesh translators */
+/** Implementation that can be shared between the Skeleton translator and GeomMesh translators */
 namespace MeshTranslationImpl
 {
 	/** Resolves the material assignments in AssignmentInfo, returning an UMaterialInterface for each material slot */
diff --git a/Engine/Plugins/Importers/USDImporter/Source/USDSchemas/Private/USDGeomMeshTranslator.cpp b/Engine/Plugins/Importers/USDImporter/Source/USDSchemas/Private/USDGeomMeshTranslator.cpp
--- a/Engine/Plugins/Importers/USDImporter/Source/USDSchemas/Private/USDGeomMeshTranslator.cpp
+++ b/Engine/Plugins/Importers/USDImporter/Source/USDSchemas/Private/USDGeomMeshTranslator.cpp
@@ -12,7 +12,6 @@
 #include "USDConversionUtils.h"
 #include "USDGeomMeshConversion.h"
 #include "USDInfoCache.h"
-#include "USDIntegrationUtils.h"
 #include "USDLog.h"
 #include "USDPrimConversion.h"
 #include "USDTypesConversion.h"
@@ -32,8 +31,6 @@
 #include "Interfaces/ITargetPlatformManagerModule.h"
 #include "MaterialDomain.h"
 #include "Materials/Material.h"
-#include "Materials/MaterialInstanceConstant.h"
-#include "Materials/MaterialInstanceDynamic.h"
 #include "Materials/MaterialInterface.h"
 #include "Misc/App.h"
 #include "Misc/SecureHash.h"
@@ -43,7 +40,6 @@
 #include "StaticMeshOperations.h"
 #include "StaticMeshResources.h"
 #include "UObject/Package.h"
-#include "UObject/SoftObjectPath.h"
 
 #if WITH_EDITOR
 #include "ConvexDecompTool.h"
@@ -52,14 +48,16 @@
 #endif // WITH_EDITOR
 
 #include "USDIncludesStart.h"
-#include "pxr/usd/usd/prim.h"
-#include "pxr/usd/usd/stage.h"
-#include "pxr/usd/usd/typed.h"
-#include "pxr/usd/usdGeom/mesh.h"
-#include "pxr/usd/usdGeom/subset.h"
-#include "pxr/usd/usdPhysics/collisionAPI.h"
-#include "pxr/usd/usdPhysics/meshCollisionAPI.h"
-#include "pxr/usd/usdPhysics/tokens.h"
+	#include "pxr/usd/usd/prim.h"
+	#include "pxr/usd/usd/stage.h"
+	#include "pxr/usd/usd/typed.h"
+	#include "pxr/usd/usdGeom/mesh.h"
+	#include "pxr/usd/usdGeom/subset.h"
+	#include "pxr/usd/usdSkel/bindingAPI.h"
+	#include "pxr/usd/usdSkel/root.h"
+	#include "pxr/usd/usdPhysics/collisionAPI.h"
+	#include "pxr/usd/usdPhysics/meshCollisionAPI.h"
+	#include "pxr/usd/usdPhysics/tokens.h"
 #include "USDIncludesEnd.h"
 
 static float GMeshNormalRepairThreshold = 0.05f;
@@ -1568,7 +1566,7 @@ void FGeomMeshCreateAssetsTaskChain::SetupTasks()
 
 	// Create mesh descriptions (Async or ExclusiveSync)
 	Do(LaunchPolicy,
-		[this]() -> bool
+		[this, LaunchPolicy]() -> bool
 		{
 			pxr::TfToken RenderContextToken = pxr::UsdShadeTokens->universalRenderContext;
 			if (!Context->RenderContext.IsNone())
@@ -1596,7 +1594,7 @@ void FGeomMeshCreateAssetsTaskChain::SetupTasks()
 				LODIndexToMeshDescription,
 				LODIndexToMaterialInfo,
 				Options,
-				Context->bAllowInterpretingLODs
+				Context->bAllowInterpretingLODs && LaunchPolicy == ESchemaTranslationLaunchPolicy::ExclusiveSync
 			);
 
 			// If we have at least one valid LOD, we should keep going
@@ -1635,6 +1633,11 @@ void FUsdGeomMeshTranslator::CreateAssets()
 		return;
 	}
 
+	if (ShouldSkipSkinnablePrim())
+	{
+		return;
+	}
+
 	TSharedRef< FGeomMeshCreateAssetsTaskChain > AssetsTaskChain = MakeShared< FGeomMeshCreateAssetsTaskChain >(Context, PrimPath);
 
 	Context->TranslatorTasks.Add(MoveTemp(AssetsTaskChain));
@@ -1646,11 +1649,39 @@ USceneComponent* FUsdGeomMeshTranslator::CreateComponents()
 	{
 		return nullptr;
 	}
+
+	// It's not great that we have to check for an alternate draw mode here and Super::CreateComponents will
+	// also do it, however we must check for ShouldSkipSkinnablePrim before we call Super::CreateComponents,
+	// and we need the alt draw mode check to take priority over ShouldSkipSkinnablePrim...
+	EUsdDrawMode DrawMode = UsdUtils::GetAppliedDrawMode(GetPrim());
+	if (DrawMode == EUsdDrawMode::Default)
+	{
+		const bool bCheckForComponent = true;
+		if (ShouldSkipSkinnablePrim(bCheckForComponent))
+		{
+			return nullptr;
+		}
+	}
+
 	return Super::CreateComponents();
 }
 
 void FUsdGeomMeshTranslator::UpdateComponents(USceneComponent* SceneComponent)
 {
+	// If we're a bounds component we don't need to handle any mesh stuff, as drawMode takes priority
+	// over even checking for whether we're skinned or not
+	if (Cast<UUsdDrawModeComponent>(SceneComponent))
+	{
+		Super::UpdateComponents(SceneComponent);
+		return;
+	}
+
+	const bool bCheckForComponent = true;
+	if (ShouldSkipSkinnablePrim(bCheckForComponent))
+	{
+		return;
+	}
+
 	if (IsMeshPrim())
 	{
 		if (SceneComponent)
@@ -1673,11 +1704,6 @@ void FUsdGeomMeshTranslator::UpdateComponents(USceneComponent* SceneComponent)
 
 bool FUsdGeomMeshTranslator::CollapsesChildren(ECollapsingType CollapsingType) const
 {
-	if (!IsMeshPrim())
-	{
-		return Super::CollapsesChildren(CollapsingType);
-	}
-
 	// If we have a custom draw mode, it means we should draw bounds/cards/etc. instead
 	// of our entire subtree, which is basically the same thing as collapsing
 	EUsdDrawMode DrawMode = UsdUtils::GetAppliedDrawMode(GetPrim());
@@ -1686,14 +1712,24 @@ bool FUsdGeomMeshTranslator::CollapsesChildren(ECollapsingType CollapsingType) c
 		return true;
 	}
 
-	// We can't claim we collapse anything here since we'll just parse the mesh for this prim and that's it,
-	// otherwise the translation context wouldn't spawn translators for our child prims.
-	// Another approach would be to actually recursively collapse our child mesh prims, but that leads to a few
-	// issues. For example this translator could end up globbing a child Mesh prim, while the translation context
-	// could simultaneously spawn other translators that could also end up accounting for that same mesh.
-	// Generally Gprims shouldn't be nested into each other anyway (see https://graphics.pixar.com/usd/release/glossary.html#usdglossary-gprim)
-	// so it's likely best to just not collapse anything here.
-	return false;
+	if (ShouldSkipSkinnablePrim())
+	{
+		return false;
+	}
+
+	if (IsMeshPrim())
+	{
+		// We can't claim we collapse anything here since we'll just parse the mesh for this prim and that's it,
+		// otherwise the translation context wouldn't spawn translators for our child prims.
+		// Another approach would be to actually recursively collapse our child mesh prims, but that leads to a few
+		// issues. For example this translator could end up globbing a child Mesh prim, while the translation context
+		// could simultaneously spawn other translators that could also end up accounting for that same mesh.
+		// Generally Gprims shouldn't be nested into each other anyway (see https://graphics.pixar.com/usd/release/glossary.html#usdglossary-gprim)
+		// so it's likely best to just not collapse anything here.
+		return false;
+	}
+
+	return Super::CollapsesChildren(CollapsingType);
 }
 
 bool FUsdGeomMeshTranslator::CanBeCollapsed(ECollapsingType CollapsingType) const
@@ -1703,6 +1739,12 @@ bool FUsdGeomMeshTranslator::CanBeCollapsed(ECollapsingType CollapsingType) cons
 		return Super::CanBeCollapsed(CollapsingType);
 	}
 
+	// If we're a skinned mesh prim we're inside of a SkelRoot, that won't be collapsed anyway
+	if (ShouldSkipSkinnablePrim())
+	{
+		return false;
+	}
+
 	UE::FUsdPrim Prim = GetPrim();
 
 	// Don't collapse if our final UStaticMesh would have multiple LODs
@@ -1729,6 +1771,11 @@ TSet<UE::FSdfPath> FUsdGeomMeshTranslator::CollectAuxiliaryPrims() const
 		return Super::CollectAuxiliaryPrims();
 	}
 
+	if (ShouldSkipSkinnablePrim())
+	{
+		return {};
+	}
+
 	if (!Context->bIsBuildingInfoCache)
 	{
 		return Context->InfoCache->GetAuxiliaryPrims(PrimPath);
@@ -1782,4 +1829,85 @@ bool FUsdGeomMeshTranslator::IsMeshPrim() const
 	return false;
 }
 
+bool FUsdGeomMeshTranslator::ShouldSkipSkinnablePrim(bool bCheckForComponent) const
+{
+	// At runtime we don't handle skinned meshes as SkeletalMeshes anyway, so early out and pretend it is not skinned
+	if (!GIsEditor)
+	{
+		return false;
+	}
+
+	FScopedUsdAllocs Allocs;
+
+	pxr::UsdPrim Prim = GetPrim();
+	if (!Prim)
+	{
+		return false;
+	}
+
+	// Search for the API schema on the mesh prims
+	bool bHasAPISchema = false;
+	if (pxr::UsdGeomMesh GeomMesh{Prim})
+	{
+		if (Prim.HasAPI<pxr::UsdSkelBindingAPI>())
+		{
+			bHasAPISchema = true;
+		}
+	}
+	// The groom translator means we could potentially have been called on an Xform LOD prim with skinned meshes in the LODs.
+	// We could even be a LOD prim setup where the SkelRoot is the "container" and the different LOD meshes are inside of it. In that
+	// particular case though then we still want to spawn/keep our SceneComponent for the SkelRoot (we always keep the SkelRoots)
+	else if (Context->bAllowInterpretingLODs && UsdUtils::DoesPrimContainMeshLODs(Prim))
+	{
+		pxr::UsdPrimSiblingRange PrimRange = Prim.GetChildren();
+		for (pxr::UsdPrimSiblingRange::iterator PrimRangeIt = PrimRange.begin(); PrimRangeIt != PrimRange.end(); ++PrimRangeIt)
+		{
+			const pxr::UsdPrim& Child = *PrimRangeIt;
+			if (pxr::UsdGeomMesh ChildMesh{Child})
+			{
+				if (Prim.HasAPI<pxr::UsdSkelBindingAPI>())
+				{
+					bHasAPISchema = true;
+				}
+			}
+		}
+	}
+	if (!bHasAPISchema)
+	{
+		return false;
+	}
+
+	// If we're here we know we have a valid API schema. Check for an ancestor SkelRoot then,
+	// which is required to enable skinning
+	if (pxr::UsdPrim ParentSkelRoot = UsdUtils::GetClosestParentSkelRoot(Prim))
+	{
+		// For components there is still one edge case to worry about: If the user has placed a non-skinnable
+		// prim as a child of our skinnable prim (our 'Prim' here). We'll translate this skinnable prim data when
+		// generating the SkeletalMesh for it from somewhere else, and we'll still never translate an asset for 'Prim'
+		// directly again here (e.g. as a StaticMesh), but when it comes to the *component* then we may still need to spawn it.
+		// This because the component for 'Prim' may have some transform, visibility, animation (or something) that affects
+		// its children
+		if (bCheckForComponent)
+		{
+			bool bHasRelevantChildPrim = false;
+			for (pxr::UsdPrim Child : Prim.GetChildren())
+			{
+				if (!Child.IsA<pxr::UsdGeomSubset>())
+				{
+					bHasRelevantChildPrim = true;
+					break;
+				}
+			}
+
+			return !bHasRelevantChildPrim;
+		}
+		else
+		{
+			return true;
+		}
+	}
+
+	return false;
+}
+
 #endif // #if USE_USD_SDK
diff --git a/Engine/Plugins/Importers/USDImporter/Source/USDSchemas/Private/USDGeomXformableTranslator.cpp b/Engine/Plugins/Importers/USDImporter/Source/USDSchemas/Private/USDGeomXformableTranslator.cpp
--- a/Engine/Plugins/Importers/USDImporter/Source/USDSchemas/Private/USDGeomXformableTranslator.cpp
+++ b/Engine/Plugins/Importers/USDImporter/Source/USDSchemas/Private/USDGeomXformableTranslator.cpp
@@ -57,6 +57,7 @@
 	#include "pxr/usd/usdGeom/subset.h"
 	#include "pxr/usd/usdGeom/xformable.h"
 	#include "pxr/usd/usdShade/materialBindingAPI.h"
+	#include "pxr/usd/usdSkel/root.h"
 #include "USDIncludesEnd.h"
 
 static bool GCollapsePrimsWithoutKind = true;
@@ -399,10 +400,13 @@ USceneComponent* FUsdGeomXformableTranslator::CreateComponentsEx( TOptional< TSu
 					UsdPrim.IsModel() ||
 					UsdPrim.IsGroup() ||
 					UsdUtils::HasCompositionArcs( UsdPrim ) ||
-					UsdPrim.HasAttribute( TEXT( "unrealCameraPrimName" ) );  // If we have this, then we correspond to the root component
-																			 // of an exported ACineCameraActor. Let's create an actual
-																			 // CineCameraActor here so that our child camera prim can just
-																			 // take it's UCineCameraComponent instead
+					UsdPrim.HasAttribute( TEXT( "unrealCameraPrimName" ) ) ||  // If we have this, then we correspond to the root component
+																			   // of an exported ACineCameraActor. Let's create an actual
+																			   // CineCameraActor here so that our child camera prim can just
+																			   // take it's UCineCameraComponent instead
+					UsdPrim.IsA(TEXT("SkelRoot"));  // Now that we use the UsdSkelSkeletonTranslator, UsdSkelRoots will be handled like regular
+													// Xforms. We likely always want then to show up on the outliner though, as they are important
+													// prims
 		};
 
 		bNeedsActor =
@@ -691,9 +695,9 @@ void FUsdGeomXformableTranslator::UpdateComponents( USceneComponent* SceneCompon
 			UsdToUnreal::ConvertBounds(Prim, BoundsComponent, Context->Time, PxrBBoxCache);
 		}
 
-		// Handle LiveLink, but only i we're not a skeletal mesh component: The SkelRootTranslator will deal with the
+		// Handle LiveLink, but only if we're not a skeletal case: The SkelSkeletonTranslator will deal with the
 		// skeletal version of the LiveLink configuration, we only handle setting up LiveLink for simple transforms
-		if ( !SceneComponent->IsA<USkeletalMeshComponent>() )
+		if (!Prim.IsA(TEXT("SkelRoot")) && !Prim.IsA(TEXT("Skeleton")))
 		{
 			if ( UsdUtils::PrimHasSchema( Prim, UnrealIdentifiers::LiveLinkAPI ) )
 			{
@@ -711,9 +715,29 @@ void FUsdGeomXformableTranslator::UpdateComponents( USceneComponent* SceneCompon
 		{
 			// Don't update the component's transform if this is already factored in as root motion within the AnimSequence
 			bool bConvertTransform = true;
-			if ( Prim.IsA( TEXT( "SkelRoot" ) ) && Context->RootMotionHandling == EUsdRootMotionHandling::UseMotionFromSkelRoot )
+			switch (Context->RootMotionHandling)
 			{
-				bConvertTransform = false;
+				default:
+				case EUsdRootMotionHandling::NoAdditionalRootMotion:
+				{
+					break;
+				}
+				case EUsdRootMotionHandling::UseMotionFromSkelRoot:
+				{
+					if (Prim.IsA(TEXT("SkelRoot")))
+					{
+						bConvertTransform = false;
+					}
+					break;
+				}
+				case EUsdRootMotionHandling::UseMotionFromSkeleton:
+				{
+					if (Prim.IsA(TEXT("Skeleton")))
+					{
+						bConvertTransform = false;
+					}
+					break;
+				}
 			}
 
 			UsdToUnreal::ConvertXformable( Context->Stage, pxr::UsdGeomXformable( Prim ), *SceneComponent, Context->Time, bConvertTransform );
@@ -978,6 +1002,13 @@ bool FUsdGeomXformableTranslator::CollapsesChildren(ECollapsingType CollapsingTy
 	pxr::UsdPrim Prim = GetPrim();
 	pxr::UsdModelAPI Model{ pxr::UsdTyped( Prim ) };
 
+	// Now that we use UsdSkelSkeletonTranslator the SkelRoots will be handled by the FUsdGeomXformableTranslator (here).
+	// SkelRoots are likely going to end up with SkeletalMeshes though, so we can assume we won't be collapsing them just from that
+	if (Prim.IsA<pxr::UsdSkelRoot>())
+	{
+		return false;
+	}
+
 	if ( Model )
 	{
 		EUsdDefaultKind PrimKind = UsdUtils::GetDefaultKind( Prim );
@@ -1006,6 +1037,11 @@ bool FUsdGeomXformableTranslator::CollapsesChildren(ECollapsingType CollapsingTy
 			TArray< TUsdStore< pxr::UsdPrim > > ChildXformPrims = UsdUtils::GetAllPrimsOfType( Prim, pxr::TfType::Find< pxr::UsdGeomXformable >() );
 			for ( const TUsdStore< pxr::UsdPrim >& ChildXformPrim : ChildXformPrims )
 			{
+				if (ChildXformPrim.Get().IsA<pxr::UsdSkelRoot>())
+				{
+					return false;
+				}
+
 				if ( TSharedPtr< FUsdSchemaTranslator > SchemaTranslator = UsdSchemasModule.GetTranslatorRegistry().CreateTranslatorForSchema( Context, UE::FUsdTyped( ChildXformPrim.Get() ) ) )
 				{
 					if ( !SchemaTranslator->CanBeCollapsed( CollapsingType ) )
@@ -1032,6 +1068,7 @@ bool FUsdGeomXformableTranslator::CanBeCollapsed( ECollapsingType CollapsingType
 
 	if (UsdUtils::IsAnimated(UsdPrim) ||
 		UsdUtils::PrimHasSchema(UsdPrim, UnrealIdentifiers::LiveLinkAPI) ||
+		UsdPrim.IsA<pxr::UsdSkelRoot>() ||
 		(Context->bAllowInterpretingLODs && UsdUtils::DoesPrimContainMeshLODs(UsdPrim))
 	)
 	{
diff --git a/Engine/Plugins/Importers/USDImporter/Source/USDSchemas/Private/USDGeometryCacheTranslator.cpp b/Engine/Plugins/Importers/USDImporter/Source/USDSchemas/Private/USDGeometryCacheTranslator.cpp
--- a/Engine/Plugins/Importers/USDImporter/Source/USDSchemas/Private/USDGeometryCacheTranslator.cpp
+++ b/Engine/Plugins/Importers/USDImporter/Source/USDSchemas/Private/USDGeometryCacheTranslator.cpp
@@ -858,6 +858,11 @@ void FUsdGeometryCacheTranslator::CreateAssets()
 		return;
 	}
 
+	if (ShouldSkipSkinnablePrim())
+	{
+		return;
+	}
+
 	// Create the GeometryCache TaskChain
 	TSharedRef<FGeometryCacheCreateAssetsTaskChain> AssetsTaskChain = MakeShared<FGeometryCacheCreateAssetsTaskChain>(Context, PrimPath);
 
@@ -876,6 +881,12 @@ USceneComponent* FUsdGeometryCacheTranslator::CreateComponents()
 	EUsdDrawMode DrawMode = UsdUtils::GetAppliedDrawMode(GetPrim());
 	if (DrawMode == EUsdDrawMode::Default)
 	{
+		const bool bCheckForComponent = true;
+		if (ShouldSkipSkinnablePrim(bCheckForComponent))
+		{
+			return nullptr;
+		}
+
 		SceneComponent = CreateComponentsEx(
 			{Context->bIsImporting ? UGeometryCacheComponent::StaticClass() : UGeometryCacheUsdComponent::StaticClass()},
 			{}
@@ -941,7 +952,13 @@ void FUsdGeometryCacheTranslator::UpdateComponents(USceneComponent* SceneCompone
 {
 	UGeometryCacheComponent* GeometryCacheComponent = Cast<UGeometryCacheComponent>(SceneComponent);
 
-	if (!IsPotentialGeometryCacheRoot() || !GeometryCacheComponent)
+	const bool bCheckForComponent = true;
+	if (!Cast<UUsdDrawModeComponent>(SceneComponent) && ShouldSkipSkinnablePrim(bCheckForComponent))
+	{
+		return;
+	}
+
+	if (!GeometryCacheComponent || !IsPotentialGeometryCacheRoot())
 	{
 		Super::UpdateComponents(SceneComponent);
 		return;
@@ -1053,18 +1070,27 @@ bool FUsdGeometryCacheTranslator::CollapsesChildren(ECollapsingType CollapsingTy
 {
 	// If we have a custom draw mode, it means we should draw bounds/cards/etc. instead
 	// of our entire subtree, which is basically the same thing as collapsing
-	EUsdDrawMode DrawMode = UsdUtils::GetAppliedDrawMode(GetPrim());
-	if (DrawMode != EUsdDrawMode::Default)
+	if (IsPotentialGeometryCacheRoot() || UsdUtils::GetAppliedDrawMode(GetPrim()) != EUsdDrawMode::Default)
 	{
 		return true;
 	}
 
-	return IsPotentialGeometryCacheRoot() ? true : Super::CollapsesChildren(CollapsingType);
+	if (ShouldSkipSkinnablePrim())
+	{
+		return false;
+	}
+
+	return Super::CollapsesChildren(CollapsingType);
 }
 
 bool FUsdGeometryCacheTranslator::CanBeCollapsed(ECollapsingType CollapsingType) const
 {
-	return IsPotentialGeometryCacheRoot() ? false : Super::CanBeCollapsed(CollapsingType);
+	if (IsPotentialGeometryCacheRoot() || ShouldSkipSkinnablePrim())
+	{
+		return false;
+	}
+
+	return Super::CanBeCollapsed(CollapsingType);
 }
 
 TSet<UE::FSdfPath> FUsdGeometryCacheTranslator::CollectAuxiliaryPrims() const
@@ -1079,6 +1105,11 @@ TSet<UE::FSdfPath> FUsdGeometryCacheTranslator::CollectAuxiliaryPrims() const
 		return Context->InfoCache->GetAuxiliaryPrims(PrimPath);
 	}
 
+	if (ShouldSkipSkinnablePrim())
+	{
+		return {};
+	}
+
 	TSet<UE::FSdfPath> AuxPrims;
 
 	// Here, we collect all meshes even non-animated ones since they'll be collapse into the cache
diff --git a/Engine/Plugins/Importers/USDImporter/Source/USDSchemas/Private/USDGroomTranslatorUtils.cpp b/Engine/Plugins/Importers/USDImporter/Source/USDSchemas/Private/USDGroomTranslatorUtils.cpp
--- a/Engine/Plugins/Importers/USDImporter/Source/USDSchemas/Private/USDGroomTranslatorUtils.cpp
+++ b/Engine/Plugins/Importers/USDImporter/Source/USDSchemas/Private/USDGroomTranslatorUtils.cpp
@@ -24,10 +24,11 @@
 #include "UsdWrappers/SdfPath.h"
 
 #include "USDIncludesStart.h"
-#include "pxr/usd/usd/prim.h"
-#include "pxr/usd/usd/relationship.h"
-#include "pxr/usd/usdGeom/mesh.h"
-#include "pxr/usd/usdSkel/root.h"
+	#include "pxr/usd/usd/prim.h"
+	#include "pxr/usd/usd/relationship.h"
+	#include "pxr/usd/usdGeom/mesh.h"
+	#include "pxr/usd/usdSkel/root.h"
+	#include "pxr/usd/usdSkel/skeleton.h"
 #include "USDIncludesEnd.h"
 
 namespace UE::UsdGroomTranslatorUtils::Private
@@ -172,7 +173,7 @@ namespace UE::UsdGroomTranslatorUtils::Private
 
 				// Validate that the target prim and associated asset are of the expected type for the binding
 				pxr::UsdPrim TargetPrim = Prim.GetPrimAtPath(TargetPrimPath);
-				if (BindingType == EGroomBindingMeshType::SkeletalMesh && pxr::UsdSkelRoot(TargetPrim))
+				if (BindingType == EGroomBindingMeshType::SkeletalMesh && (pxr::UsdSkelRoot(TargetPrim) || pxr::UsdSkelSkeleton(TargetPrim)))
 				{
 					return InfoCache.GetSingleAssetForPrim<USkeletalMesh>(UE::FSdfPath{TargetPrimPath});
 				}
@@ -197,7 +198,7 @@ namespace UsdGroomTranslatorUtils
 		EObjectFlags ObjectFlags
 	)
 	{
-		// At this point, the prim (SkelRoot or GeomMesh) has already been checked to have the GroomBindingAPI,
+		// At this point, the prim (SkelRoot/Skeleton or GeomMesh) has already been checked to have the GroomBindingAPI,
 		// so we need to set up the GroomComponent and the groom binding asset to be able to bind it to the mesh
 
 		// The GroomBinding schema must specify a groom prim to bind to the mesh
diff --git a/Engine/Plugins/Importers/USDImporter/Source/USDSchemas/Private/USDInfoCache.cpp b/Engine/Plugins/Importers/USDImporter/Source/USDSchemas/Private/USDInfoCache.cpp
--- a/Engine/Plugins/Importers/USDImporter/Source/USDSchemas/Private/USDInfoCache.cpp
+++ b/Engine/Plugins/Importers/USDImporter/Source/USDSchemas/Private/USDInfoCache.cpp
@@ -31,6 +31,7 @@
 	#include "pxr/usd/usdGeom/subset.h"
 	#include "pxr/usd/usdGeom/xform.h"
 	#include "pxr/usd/usdShade/materialBindingAPI.h"
+	#include "pxr/usd/usdSkel/bindingAPI.h"
 	#include "pxr/usd/usdSkel/root.h"
 	#include "pxr/usd/usd/primRange.h"
 #include "USDIncludesEnd.h"
@@ -940,11 +941,10 @@ namespace UE::USDInfoCacheImpl::Private
 			return false;
 		}
 
-		// Only care about collapsing into a StaticMesh: We should always collapse into a SkeletalMesh as we have
-		// no real alternative for handling them
+		// We should never be able to collapse SkelRoots because the UsdSkelSkeletonTranslator doesn't collapse
 		if (UsdPrim.IsA<pxr::UsdSkelRoot>())
 		{
-			return true;
+			return false;
 		}
 
 		pxr::SdfPath UsdPrimPath = UsdPrim.GetPrimPath();
@@ -1220,7 +1220,14 @@ namespace UE::USDInfoCacheImpl::Private
 		OutState = EGeometryCachePrimState::Uncollapsible;
 	}
 
-	void RecursiveCheckForGeometryCache(const pxr::UsdPrim& UsdPrim, FUsdSchemaTranslationContext& Context, FUsdInfoCache::FUsdInfoCacheImpl& Impl, int32& OutDepth, UE::UsdInfoCache::Private::EGeometryCachePrimState& OutState)
+	void RecursiveCheckForGeometryCache(
+		const pxr::UsdPrim& UsdPrim,
+		FUsdSchemaTranslationContext& Context,
+		FUsdInfoCache::FUsdInfoCacheImpl& Impl,
+		bool bIsInsideSkelRoot,
+		int32& OutDepth,
+		UE::UsdInfoCache::Private::EGeometryCachePrimState& OutState
+	)
 	{
 		TRACE_CPUPROFILER_EVENT_SCOPE(RecursiveCheckForGeometryCache);
 
@@ -1249,10 +1256,20 @@ namespace UE::USDInfoCacheImpl::Private
 		States.SetNum(Prims.Num());
 
 		const int32 MinBatchSize = 1;
-		ParallelFor(TEXT("RecursiveCheckForGeometryCache"), Prims.Num(), MinBatchSize,
-			[&Prims, &Context, &Impl, &Depths, &States](int32 Index)
+		ParallelFor(
+			TEXT("RecursiveCheckForGeometryCache"),
+			Prims.Num(),
+			MinBatchSize,
+			[&Prims, &Context, &Impl, bIsInsideSkelRoot, &Depths, &States](int32 Index)
 			{
-				RecursiveCheckForGeometryCache(Prims[Index], Context, Impl, Depths[Index], States[Index]);
+				RecursiveCheckForGeometryCache(
+					Prims[Index],
+					Context,
+					Impl,
+					bIsInsideSkelRoot || Prims[Index].IsA<pxr::UsdSkelRoot>(),
+					Depths[Index],
+					States[Index]
+				);
 			}
 		);
 
@@ -1308,10 +1325,20 @@ namespace UE::USDInfoCacheImpl::Private
 		const bool bIsXform = !!pxr::UsdGeomXform(UsdPrim);
 		if (bIsMesh)
 		{
-			// Animated or static mesh. Static meshes could potentially be animated by transforms in their hierarchy.
-			// A mesh prim should be a leaf, but it can have GeomSubset prims as children, but those don't
-			// affect the collapsibility status.
-			PrimState = EGeometryCachePrimState::Mesh;
+			// A skinned mesh can never be considered part of a geometry cache.
+			// Now that we use the UsdSkelSkeletonTranslator instead of the old UsdSkelRootTranslator we may run into these
+			// skinned meshes that were already handled by a SkeletonTranslator elsewhere, and need to manually skip them
+			if (GIsEditor && bIsInsideSkelRoot && UsdPrim.HasAPI<pxr::UsdSkelBindingAPI>())
+			{
+				PrimState = EGeometryCachePrimState::Uncollapsible;
+			}
+			else
+			{
+				// Animated or static mesh. Static meshes could potentially be animated by transforms in their hierarchy.
+				// A mesh prim should be a leaf, but it can have GeomSubset prims as children, but those don't
+				// affect the collapsibility status.
+				PrimState = EGeometryCachePrimState::Mesh;
+			}
 		}
 		else if (bIsXform)
 		{
@@ -1432,9 +1459,11 @@ namespace UE::USDInfoCacheImpl::Private
 			return;
 		}
 
+		const bool bIsInsideSkelRoot = static_cast<bool>(UsdUtils::GetClosestParentSkelRoot(UsdPrim));
+
 		int32 Depth = -1;
 		EGeometryCachePrimState State = EGeometryCachePrimState::None;
-		RecursiveCheckForGeometryCache(UsdPrim, Context, Impl, Depth, State);
+		RecursiveCheckForGeometryCache(UsdPrim, Context, Impl, bIsInsideSkelRoot, Depth, State);
 
 		// If we end up with a positive depth, it means the check found an animated mesh somewhere
 		// but no potential root before reaching the pseudoroot, so find one
diff --git a/Engine/Plugins/Importers/USDImporter/Source/USDSchemas/Private/USDSchemasModule.cpp b/Engine/Plugins/Importers/USDImporter/Source/USDSchemas/Private/USDSchemasModule.cpp
--- a/Engine/Plugins/Importers/USDImporter/Source/USDSchemas/Private/USDSchemasModule.cpp
+++ b/Engine/Plugins/Importers/USDImporter/Source/USDSchemas/Private/USDSchemasModule.cpp
@@ -22,7 +22,7 @@
 #include "USDLuxLightTranslator.h"
 #include "USDMemory.h"
 #include "USDShadeMaterialTranslator.h"
-#include "USDSkelRootTranslator.h"
+#include "USDSkelSkeletonTranslator.h"
 #include "Custom/MaterialXUSDShadeMaterialTranslator.h"
 #include "Custom/MDLUSDShadeMaterialTranslator.h"
 #endif // #if USE_USD_SDK
@@ -64,7 +64,7 @@ public:
 		if (GIsEditor)
 		{
 			TranslatorHandles.Append({
-				Registry.Register<FUsdSkelRootTranslator>(TEXT("UsdSkelRoot")),
+				Registry.Register<FUsdSkelSkeletonTranslator>(TEXT("UsdSkelSkeleton")),
 				Registry.Register<FUsdGroomTranslator>(TEXT("UsdGeomXformable")),
 				// The GeometryCacheTranslator also works on UsdGeomXformable through the GroomTranslator
 				Registry.Register<FUsdGeometryCacheTranslator>(TEXT("UsdGeomMesh")),
diff --git a/Engine/Plugins/Importers/USDImporter/Source/USDSchemas/Private/USDSkelRootTranslator.cpp b/Engine/Plugins/Importers/USDImporter/Source/USDSchemas/Private/USDSkelRootTranslator.cpp
--- a/Engine/Plugins/Importers/USDImporter/Source/USDSchemas/Private/USDSkelRootTranslator.cpp
+++ b/Engine/Plugins/Importers/USDImporter/Source/USDSchemas/Private/USDSkelRootTranslator.cpp
@@ -71,12 +71,6 @@
 
 #define LOCTEXT_NAMESPACE "UsdSkelRoot"
 
-static bool bGeneratePhysicsAssets = true;
-static FAutoConsoleVariableRef CVarGeneratePhysicsAssets(
-	TEXT( "USD.GeneratePhysicsAssets" ),
-	bGeneratePhysicsAssets,
-	TEXT( "Whether to automatically generate and assign PhysicsAssets to generated SkeletalMeshes." ) );
-
 namespace UsdSkelRootTranslatorImpl
 {
 #if WITH_EDITOR
@@ -1429,6 +1423,9 @@ namespace UsdSkelRootTranslatorImpl
 						SkeletalMesh->GetSkeleton()->SetPreviewMesh(SkeletalMesh);
 					}
 
+					static IConsoleVariable* PhysicsAssetsCvar = IConsoleManager::Get().FindConsoleVariable(TEXT("USD.GeneratePhysicsAssets"));
+					const bool bGeneratePhysicsAssets = PhysicsAssetsCvar && PhysicsAssetsCvar->GetBool();
+
 					if ( bGeneratePhysicsAssets )
 					{
 						if ( !SkeletalMesh->GetPhysicsAsset() )
@@ -1665,6 +1662,7 @@ namespace UsdSkelRootTranslatorImpl
 #endif // WITH_EDITOR
 }
 
+PRAGMA_DISABLE_DEPRECATION_WARNINGS
 void FUsdSkelRootTranslator::CreateAssets()
 {
 	TRACE_CPUPROFILER_EVENT_SCOPE(FUsdSkelRootTranslator::CreateAssets);
@@ -1997,6 +1995,7 @@ TSet<UE::FSdfPath> FUsdSkelRootTranslator::CollectAuxiliaryPrims() const
 	}
 	return Result;
 }
+PRAGMA_ENABLE_DEPRECATION_WARNINGS
 
 #undef LOCTEXT_NAMESPACE
 
diff --git a/Engine/Plugins/Importers/USDImporter/Source/USDSchemas/Private/USDSkelSkeletonTranslator.cpp b/Engine/Plugins/Importers/USDImporter/Source/USDSchemas/Private/USDSkelSkeletonTranslator.cpp
new file mode 100644
--- /dev/null
+++ b/Engine/Plugins/Importers/USDImporter/Source/USDSchemas/Private/USDSkelSkeletonTranslator.cpp
@@ -0,0 +1,2180 @@
+// Copyright Epic Games, Inc. All Rights Reserved.
+
+#include "USDSkelSkeletonTranslator.h"
+
+#if USE_USD_SDK
+
+#include "MeshTranslationImpl.h"
+#include "USDAssetUserData.h"
+#include "USDClassesModule.h"
+#include "USDConversionUtils.h"
+#include "USDDrawModeComponent.h"
+#include "USDGeomMeshConversion.h"
+#include "USDGroomTranslatorUtils.h"
+#include "USDInfoCache.h"
+#include "USDIntegrationUtils.h"
+#include "USDLayerUtils.h"
+#include "USDLog.h"
+#include "USDMemory.h"
+#include "USDSkeletalDataConversion.h"
+#include "USDTypesConversion.h"
+
+#include "UsdWrappers/SdfLayer.h"
+#include "UsdWrappers/SdfPath.h"
+#include "UsdWrappers/UsdPrim.h"
+
+#include "Animation/AnimBlueprint.h"
+#include "Animation/AnimInstance.h"
+#include "Animation/AnimSequence.h"
+#include "Animation/Skeleton.h"
+#include "Components/SkeletalMeshComponent.h"
+#include "Components/SkinnedMeshComponent.h"
+#include "Engine/SkeletalMesh.h"
+#include "Engine/SkinnedAssetCommon.h"
+#include "GroomComponent.h"
+#include "MaterialDomain.h"
+#include "Materials/Material.h"
+#include "Rendering/SkeletalMeshLODImporterData.h"
+
+#if WITH_EDITOR
+#include "AnimGraphNode_LiveLinkPose.h"
+#include "AnimNode_LiveLinkPose.h"
+#include "BlueprintCompilationManager.h"
+#include "EdGraph/EdGraph.h"
+#include "EdGraph/EdGraphNode.h"
+#include "Kismet2/BlueprintEditorUtils.h"
+#include "Kismet2/CompilerResultsLog.h"
+#include "ObjectTools.h"
+#include "PhysicsAssetUtils.h"
+#endif	  // WITH_EDITOR
+
+#include "USDIncludesStart.h"
+	#include "pxr/usd/usd/primRange.h"
+	#include "pxr/usd/usdGeom/mesh.h"
+	#include "pxr/usd/usdSkel/binding.h"
+	#include "pxr/usd/usdSkel/bindingAPI.h"
+	#include "pxr/usd/usdSkel/blendShapeQuery.h"
+	#include "pxr/usd/usdSkel/cache.h"
+	#include "pxr/usd/usdSkel/root.h"
+	#include "pxr/usd/usdSkel/skeletonQuery.h"
+#include "USDIncludesEnd.h"
+
+#define LOCTEXT_NAMESPACE "UsdSkelRoot"
+
+static bool bGeneratePhysicsAssets = true;
+static FAutoConsoleVariableRef CVarGeneratePhysicsAssets(
+	TEXT("USD.GeneratePhysicsAssets"),
+	bGeneratePhysicsAssets,
+	TEXT("Whether to automatically generate and assign PhysicsAssets to generated SkeletalMeshes.")
+);
+
+namespace UsdSkelSkeletonTranslatorImpl
+{
+#if WITH_EDITOR
+	bool ProcessMaterials(
+		const pxr::UsdPrim& UsdPrim,
+		TArray<UsdUtils::FUsdPrimMaterialAssignmentInfo>& LODIndexToMaterialInfo,
+		USkeletalMesh* SkeletalMesh,
+		UUsdAssetCache2& AssetCache,
+		FUsdInfoCache& InfoCache,
+		float Time,
+		EObjectFlags Flags,
+		bool bSkeletalMeshHasMorphTargets
+	)
+	{
+		TRACE_CPUPROFILER_EVENT_SCOPE(UsdSkelSkeletonTranslatorImpl::ProcessMaterials);
+
+		if (!SkeletalMesh)
+		{
+			return false;
+		}
+
+		TArray<UMaterialInterface*> ExistingAssignments;
+		for (const FSkeletalMaterial& SkeletalMaterial : SkeletalMesh->GetMaterials())
+		{
+			ExistingAssignments.Add(SkeletalMaterial.MaterialInterface);
+		}
+
+		UUsdMeshAssetUserData* UserData = SkeletalMesh->GetAssetUserData<UUsdMeshAssetUserData>();
+		ensureMsgf(
+			UserData,
+			TEXT("Skeletal Mesh '%s' generated for prim '%s' should have an UUsdMeshAssetUserData at this point!"),
+			*SkeletalMesh->GetPathName(),
+			*UsdToUnreal::ConvertPath(UsdPrim.GetPrimPath())
+		);
+
+		TMap<const UsdUtils::FUsdPrimMaterialSlot*, UMaterialInterface*>
+			ResolvedMaterials = MeshTranslationImpl::ResolveMaterialAssignmentInfo(UsdPrim, LODIndexToMaterialInfo, AssetCache, InfoCache, Flags);
+
+		bool bMaterialsHaveChanged = false;
+
+		uint32 SkeletalMeshSlotIndex = 0;
+		for (int32 LODIndex = 0; LODIndex < LODIndexToMaterialInfo.Num(); ++LODIndex)
+		{
+			const TArray<UsdUtils::FUsdPrimMaterialSlot>& LODSlots = LODIndexToMaterialInfo[LODIndex].Slots;
+
+			// We need to fill this in with the mapping from LOD material slots (i.e. sections) to the skeletal mesh's material slots
+			FSkeletalMeshLODInfo* LODInfo = SkeletalMesh->GetLODInfo(LODIndex);
+			if (!LODInfo)
+			{
+				UE_LOG(
+					LogUsd,
+					Error,
+					TEXT("When processing materials for SkeletalMesh '%s', encountered no LOD info for LOD index %d!"),
+					*SkeletalMesh->GetName(),
+					LODIndex
+				);
+				continue;
+			}
+			TArray<int32>& LODMaterialMap = LODInfo->LODMaterialMap;
+			LODMaterialMap.Reserve(LODSlots.Num());
+
+			for (int32 LODSlotIndex = 0; LODSlotIndex < LODSlots.Num(); ++LODSlotIndex, ++SkeletalMeshSlotIndex)
+			{
+				const UsdUtils::FUsdPrimMaterialSlot& Slot = LODSlots[LODSlotIndex];
+
+				UMaterialInterface* Material = UMaterial::GetDefaultMaterial(MD_Surface);
+
+				if (UMaterialInterface** FoundMaterial = ResolvedMaterials.Find(&Slot))
+				{
+					Material = *FoundMaterial;
+				}
+				else
+				{
+					UE_LOG(
+						LogUsd,
+						Error,
+						TEXT("Failed to resolve material '%s' for slot '%d' of LOD '%d' for mesh '%s'"),
+						*Slot.MaterialSource,
+						LODSlotIndex,
+						LODIndex,
+						*UsdToUnreal::ConvertPath(UsdPrim.GetPath())
+					);
+					continue;
+				}
+
+				if (Material)
+				{
+					bool bNeedsRecompile = false;
+					Material->GetMaterial()->SetMaterialUsage(bNeedsRecompile, MATUSAGE_SkeletalMesh);
+					if (bSkeletalMeshHasMorphTargets)
+					{
+						Material->GetMaterial()->SetMaterialUsage(bNeedsRecompile, MATUSAGE_MorphTargets);
+					}
+				}
+
+				FName MaterialSlotName = *LexToString(SkeletalMeshSlotIndex);
+
+				// Already have a material at that skeletal mesh slot, need to reassign
+				if (SkeletalMesh->GetMaterials().IsValidIndex(SkeletalMeshSlotIndex))
+				{
+					FSkeletalMaterial& ExistingMaterial = SkeletalMesh->GetMaterials()[SkeletalMeshSlotIndex];
+
+					if (ExistingMaterial.MaterialInterface != Material || ExistingMaterial.MaterialSlotName != MaterialSlotName
+						|| ExistingMaterial.ImportedMaterialSlotName != MaterialSlotName)
+					{
+						ExistingMaterial.MaterialInterface = Material;
+						ExistingMaterial.MaterialSlotName = MaterialSlotName;
+						ExistingMaterial.ImportedMaterialSlotName = MaterialSlotName;
+						bMaterialsHaveChanged = true;
+					}
+				}
+				// Add new material
+				else
+				{
+					const bool bEnableShadowCasting = true;
+					const bool bRecomputeTangents = false;
+					SkeletalMesh->GetMaterials().Add(
+						FSkeletalMaterial(Material, bEnableShadowCasting, bRecomputeTangents, MaterialSlotName, MaterialSlotName)
+					);
+					bMaterialsHaveChanged = true;
+				}
+
+				// Already have a material at that LOD remap slot, need to reassign
+				if (LODMaterialMap.IsValidIndex(LODSlotIndex))
+				{
+					LODMaterialMap[LODSlotIndex] = SkeletalMeshSlotIndex;
+				}
+				// Add new material slot remap
+				else
+				{
+					LODMaterialMap.Add(SkeletalMeshSlotIndex);
+				}
+			}
+		}
+
+		return bMaterialsHaveChanged;
+	}
+
+	FSHAHash ComputeSHAHash(
+		const TArray<FSkeletalMeshImportData>& LODIndexToSkeletalMeshImportData,
+		TArray<SkeletalMeshImportData::FBone>& ImportedBones,
+		UsdUtils::FBlendShapeMap* BlendShapes
+	)
+	{
+		FSHA1 HashState;
+
+		for (const FSkeletalMeshImportData& ImportData : LODIndexToSkeletalMeshImportData)
+		{
+			HashState.Update((uint8*)ImportData.Points.GetData(), ImportData.Points.Num() * ImportData.Points.GetTypeSize());
+			HashState.Update((uint8*)ImportData.Wedges.GetData(), ImportData.Wedges.Num() * ImportData.Wedges.GetTypeSize());
+			HashState.Update((uint8*)ImportData.Faces.GetData(), ImportData.Faces.Num() * ImportData.Faces.GetTypeSize());
+			HashState.Update((uint8*)ImportData.Influences.GetData(), ImportData.Influences.Num() * ImportData.Influences.GetTypeSize());
+		}
+
+		// Hash the bones as well because it is possible for the mesh to be identical while only the bone configuration changed, and in that case we'd
+		// need new skeleton and ref skeleton Maybe in the future (as a separate feature) we could split off the skeleton import so that it could vary
+		// independently of the skeletal mesh
+		for (const SkeletalMeshImportData::FBone& Bone : ImportedBones)
+		{
+			HashState.UpdateWithString(*Bone.Name, Bone.Name.Len());
+			HashState.Update((uint8*)&Bone.Flags, sizeof(Bone.Flags));
+			HashState.Update((uint8*)&Bone.NumChildren, sizeof(Bone.NumChildren));
+			HashState.Update((uint8*)&Bone.ParentIndex, sizeof(Bone.ParentIndex));
+			HashState.Update((uint8*)&Bone.BonePos, sizeof(Bone.BonePos));
+		}
+
+		if (BlendShapes)
+		{
+			for (const TPair<FString, UsdUtils::FUsdBlendShape>& Pair : *BlendShapes)
+			{
+				const UsdUtils::FUsdBlendShape& BlendShape = Pair.Value;
+				HashState.UpdateWithString(*BlendShape.Name, BlendShape.Name.Len());
+
+				HashState.Update((uint8*)BlendShape.Vertices.GetData(), BlendShape.Vertices.Num() * BlendShape.Vertices.GetTypeSize());
+
+				for (const UsdUtils::FUsdBlendShapeInbetween& InBetween : BlendShape.Inbetweens)
+				{
+					HashState.UpdateWithString(*InBetween.Name, InBetween.Name.Len());
+					HashState.Update((uint8*)&InBetween.InbetweenWeight, sizeof(InBetween.InbetweenWeight));
+				}
+
+				HashState.Update((uint8*)&BlendShape.bHasAuthoredTangents, sizeof(BlendShape.bHasAuthoredTangents));
+			}
+		}
+
+		FSHAHash OutHash;
+
+		HashState.Final();
+		HashState.GetHash(&OutHash.Hash[0]);
+
+		return OutHash;
+	}
+
+	FSHAHash ComputeSHAHash(
+		const pxr::UsdSkelSkeletonQuery& InUsdSkeletonQuery,
+		const pxr::UsdPrim& RootMotionPrim,
+		const FString& SkeletalMeshHashString
+	)
+	{
+		TRACE_CPUPROFILER_EVENT_SCOPE(UsdSkelSkeletonTranslatorImpl::ComputeSHAHash_SkelQuery);
+
+		FSHAHash OutHash;
+		FSHA1 HashState;
+
+		FScopedUsdAllocs Allocs;
+
+		pxr::UsdSkelAnimQuery AnimQuery = InUsdSkeletonQuery.GetAnimQuery();
+		if (!AnimQuery)
+		{
+			return OutHash;
+		}
+
+		pxr::UsdPrim UsdPrim = InUsdSkeletonQuery.GetPrim();
+		if (!UsdPrim)
+		{
+			return OutHash;
+		}
+
+		pxr::UsdStageRefPtr Stage = UsdPrim.GetStage();
+		if (!Stage)
+		{
+			return OutHash;
+		}
+
+		int32 InterpolationType = static_cast<int32>(Stage->GetInterpolationType());
+		HashState.Update((uint8*)&InterpolationType, sizeof(int32));
+
+		// Hash blend shape and joint order tokens
+		TFunction<void(const pxr::VtArray<pxr::TfToken>&)> HashTokens = [&HashState](const pxr::VtArray<pxr::TfToken>& Tokens)
+		{
+			for (const pxr::TfToken& Token : Tokens)
+			{
+				const std::string& TokenString = Token.GetString();
+				const char* TokenCStr = TokenString.c_str();
+				HashState.Update((const uint8*)TokenCStr, sizeof(char) * TokenString.size());
+			}
+		};
+		HashTokens(AnimQuery.GetJointOrder());
+		HashTokens(AnimQuery.GetBlendShapeOrder());
+
+		// Time samples for joint transforms
+		std::vector<double> TimeData;
+		AnimQuery.GetJointTransformTimeSamples(&TimeData);
+		HashState.Update((uint8*)TimeData.data(), TimeData.size() * sizeof(double));
+
+		// Joint transform values
+		pxr::VtArray<pxr::GfMatrix4d> JointTransforms;
+		for (double JointTimeSample : TimeData)
+		{
+			InUsdSkeletonQuery.ComputeJointLocalTransforms(&JointTransforms, JointTimeSample);
+			HashState.Update((uint8*)JointTransforms.data(), JointTransforms.size() * sizeof(pxr::GfMatrix4d));
+		}
+
+		// restTransforms
+		pxr::VtArray<pxr::GfMatrix4d> Transforms;
+		const bool bAtRest = true;
+		InUsdSkeletonQuery.ComputeJointLocalTransforms(&Transforms, pxr::UsdTimeCode::EarliestTime(), bAtRest);
+		HashState.Update((uint8*)Transforms.data(), Transforms.size() * sizeof(pxr::GfMatrix4d));
+
+		// bindTransforms
+		InUsdSkeletonQuery.GetJointWorldBindTransforms(&Transforms);
+		HashState.Update((uint8*)Transforms.data(), Transforms.size() * sizeof(pxr::GfMatrix4d));
+
+		// Time samples for blend shape curves
+		AnimQuery.GetBlendShapeWeightTimeSamples(&TimeData);
+		HashState.Update((uint8*)TimeData.data(), TimeData.size() * sizeof(double));
+
+		// Blend shape curve values
+		pxr::VtArray<float> WeightsForSample;
+		for (double CurveTimeSample : TimeData)
+		{
+			AnimQuery.ComputeBlendShapeWeights(&WeightsForSample, pxr::UsdTimeCode(CurveTimeSample));
+			HashState.Update((uint8*)WeightsForSample.data(), WeightsForSample.size() * sizeof(float));
+		}
+
+		// If we're pulling root motion from anywhere, hash that too because if it changes we'll need to rebake
+		// the AnimSequence asset
+		if (pxr::UsdGeomXformable Xformable{RootMotionPrim})
+		{
+			// Hash non-animated transform too, because we'll put these directly on the components if
+			// the SkelRoot/Skeleton is not animated, and we need our AnimSequence to combine nicely with
+			// those
+			pxr::GfMatrix4d Transform{};
+			bool bResetsXformSack = false;
+			Xformable.GetLocalTransformation(&Transform, &bResetsXformSack, pxr::UsdTimeCode::Default());
+			HashState.Update((uint8*)Transform.data(), 16 * sizeof(double));
+
+			std::vector<double> TimeSamples;
+			Xformable.GetTimeSamples(&TimeSamples);
+
+			for (double TimeSample : TimeSamples)
+			{
+				Xformable.GetLocalTransformation(&Transform, &bResetsXformSack, TimeSample);
+				HashState.Update((uint8*)Transform.data(), 16 * sizeof(double));
+			}
+
+			HashState.Update((uint8*)&bResetsXformSack, sizeof(bool));
+		}
+
+		// An anim sequence matches a particular skeleton. If the skeleton is different, we'll likely
+		// need a new AnimSequence, even if the SkelAnimation prim itself hashes the same. The same
+		// applies to SkeletalMesh morph targets too: In USD a SkelAnimation has decoupled "blend
+		// shape channel" curves that can animate blend shapes specific to each mesh, but in UE the
+		// AnimSequence curves are specific to each morph target name, so if we have a different mesh,
+		// we're probably better off ensuring we have a different AnimSequence too
+		HashState.UpdateWithString(*SkeletalMeshHashString, SkeletalMeshHashString.Len());
+
+		HashState.Final();
+		HashState.GetHash(&OutHash.Hash[0]);
+
+		return OutHash;
+	}
+
+	void SetMorphTargetWeight(USkeletalMeshComponent& SkeletalMeshComponent, const FString& MorphTargetName, float Weight)
+	{
+		USkeletalMesh* SkeletalMesh = SkeletalMeshComponent.GetSkeletalMeshAsset();
+
+		// We try keeping a perfect correspondence between SkeletalMesh->GetMorphTargets() and SkeletalMeshComponent.ActiveMorphTargets
+		int32 IndexInSkeletalMesh = INDEX_NONE;
+		SkeletalMeshComponent.GetSkeletalMeshAsset()->FindMorphTargetAndIndex(*MorphTargetName, IndexInSkeletalMesh);
+		if (IndexInSkeletalMesh == INDEX_NONE)
+		{
+			return;
+		}
+
+		const UMorphTarget* MorphTarget = SkeletalMeshComponent.GetSkeletalMeshAsset()->GetMorphTargets()[IndexInSkeletalMesh];
+		if (!MorphTarget)
+		{
+			return;
+		}
+
+		int32 WeightIndex = INDEX_NONE;
+		if (SkeletalMeshComponent.ActiveMorphTargets.Contains(MorphTarget))
+		{
+			WeightIndex = SkeletalMeshComponent.ActiveMorphTargets[MorphTarget];
+		}
+
+		// Morph target is not at expected location (i.e. after CreateComponents, duplicate for PIE or undo/redo) --> Rebuild ActiveMorphTargets
+		// This may lead to one frame of glitchiness, as we'll reset all weights to zero...
+		if (WeightIndex == INDEX_NONE)
+		{
+			SkeletalMeshComponent.ActiveMorphTargets.Reset();
+			SkeletalMeshComponent.MorphTargetWeights.Reset();
+			TArray<TObjectPtr<UMorphTarget>>& MorphTargets = SkeletalMesh->GetMorphTargets();
+			for (int32 MorphTargetIndex = 0; MorphTargetIndex < MorphTargets.Num(); ++MorphTargetIndex)
+			{
+				SkeletalMeshComponent.ActiveMorphTargets.Add(MorphTargets[MorphTargetIndex], MorphTargetIndex);
+				SkeletalMeshComponent.MorphTargetWeights.Add(0.0f);	   // We'll update these right afterwards when we call UpdateComponents
+			}
+
+			WeightIndex = IndexInSkeletalMesh;
+		}
+
+		SkeletalMeshComponent.MorphTargetWeights[WeightIndex] = Weight;
+	}
+
+	bool LoadAllSkeletalData(
+		const pxr::UsdSkelBinding& InSkeletonBinding,
+		const pxr::UsdSkelCache& InSkelCache,
+		TArray<FSkeletalMeshImportData>& OutLODIndexToSkeletalMeshImportData,
+		TArray<UsdUtils::FUsdPrimMaterialAssignmentInfo>& OutLODIndexToMaterialInfo,
+		TArray<SkeletalMeshImportData::FBone>& OutSkeletonBones,
+		FName& OutSkeletonName,
+		UsdUtils::FBlendShapeMap* OutBlendShapes,
+		TSet<FString>& InOutUsedMorphTargetNames,
+		bool bInInterpretLODs,
+		const UsdToUnreal::FUsdMeshConversionOptions& Options
+	)
+	{
+		FScopedUsdAllocs UsdAllocs;
+
+		const pxr::UsdSkelSkeleton& Skeleton = InSkeletonBinding.GetSkeleton();
+		if (!Skeleton)
+		{
+			return false;
+		}
+
+		const pxr::UsdSkelSkeletonQuery& SkeletonQuery = InSkelCache.GetSkelQuery(Skeleton);
+		if (!SkeletonQuery)
+		{
+			return false;
+		}
+
+		pxr::UsdPrim SkeletonPrim = Skeleton.GetPrim();
+		pxr::SdfPath SkeletonPrimPath = SkeletonPrim.GetPrimPath();
+
+		pxr::UsdStageRefPtr Stage = SkeletonPrim.GetStage();
+		const FUsdStageInfo StageInfo(Stage);
+
+		// Import skeleton data
+		{
+			FSkeletalMeshImportData DummyImportData;
+			const bool bSkeletonValid = UsdToUnreal::ConvertSkeleton(SkeletonQuery, DummyImportData);
+			if (!bSkeletonValid)
+			{
+				return false;
+			}
+			OutSkeletonBones = MoveTemp(DummyImportData.RefBonesBinary);
+			OutSkeletonName = *UsdToUnreal::ConvertString(SkeletonPrim.GetName());
+		}
+
+		TMap<int32, FSkeletalMeshImportData> LODIndexToSkeletalMeshImportDataMap;
+		TMap<int32, UsdUtils::FUsdPrimMaterialAssignmentInfo> LODIndexToMaterialInfoMap;
+		TSet<FString> ProcessedLODParentPaths;
+
+		// Since we may need to switch variants to parse LODs, we could invalidate references to SkinningQuery objects, so we need
+		// to keep track of these by path and construct one whenever we need them
+		TArray<pxr::SdfPath> PathsToSkinnedPrims;
+		for (const pxr::UsdSkelSkinningQuery& SkinningQuery : InSkeletonBinding.GetSkinningTargets())
+		{
+			// In USD, the skinning target need not be a mesh, but for Unreal we are only interested in skinning meshes
+			if (pxr::UsdGeomMesh SkinningMesh = pxr::UsdGeomMesh(SkinningQuery.GetPrim()))
+			{
+				// Let's only care about prims with the SkelBindingAPI for now as we'll *need* joint influences and weights
+				if (SkinningQuery.GetPrim().HasAPI<pxr::UsdSkelBindingAPI>())
+				{
+					PathsToSkinnedPrims.Add(SkinningMesh.GetPrim().GetPath());
+				}
+				else
+				{
+					UE_LOG(
+						LogUsd,
+						Log,
+						TEXT("Ignoring skinned prim '%s' when generating Skeletal Mesh for Skeleton '%s' as the prim doesn't have the SkelBindingAPI"
+						),
+						*UsdToUnreal::ConvertPath(SkinningQuery.GetPrim().GetPrimPath()),
+						*UsdToUnreal::ConvertPath(SkeletonPrim.GetPrimPath())
+					);
+				}
+			}
+		}
+
+		TFunction<bool(const pxr::UsdGeomMesh&, int32)> ConvertLOD =
+			[&LODIndexToSkeletalMeshImportDataMap,
+			 &LODIndexToMaterialInfoMap,
+			 &InSkelCache,
+			 &Stage,
+			 &SkeletonPrimPath,
+			 &InOutUsedMorphTargetNames,
+			 OutBlendShapes,
+			 &StageInfo,
+			 Options](const pxr::UsdGeomMesh& LODMesh, int32 LODIndex)
+		{
+			// Construct this and SkinningQuery every time so as to survive the prim reference invalidation caused by flipping LODs
+			pxr::UsdSkelSkeletonQuery SkeletonQuery = InSkelCache.GetSkelQuery(pxr::UsdSkelSkeleton{Stage->GetPrimAtPath(SkeletonPrimPath)});
+			if (!SkeletonQuery)
+			{
+				return true;	// Continue trying other LODs
+			}
+
+			pxr::UsdSkelSkinningQuery SkinningQuery = UsdUtils::CreateSkinningQuery(LODMesh, SkeletonQuery);
+			if (!SkinningQuery)
+			{
+				return true;
+			}
+
+			// Ignore prims from disabled purposes
+			if (!EnumHasAllFlags(Options.PurposesToLoad, IUsdPrim::GetPurpose(LODMesh.GetPrim())))
+			{
+				return true;
+			}
+
+			// If a skinned prim has an alt draw mode, let's prioritize showing the alt draw mode (that is, geommesh/geometrycache translators
+			// will handle it) and skip skinning it
+			EUsdDrawMode DrawMode = UsdUtils::GetAppliedDrawMode(LODMesh.GetPrim());
+			if (DrawMode != EUsdDrawMode::Default)
+			{
+				return true;
+			}
+
+			if (LODMesh && LODMesh.ComputeVisibility() == pxr::UsdGeomTokens->invisible)
+			{
+				return true;
+			}
+
+			FSkeletalMeshImportData& LODImportData = LODIndexToSkeletalMeshImportDataMap.FindOrAdd(LODIndex);
+			UsdUtils::FUsdPrimMaterialAssignmentInfo& LODMaterialInfo = LODIndexToMaterialInfoMap.FindOrAdd(LODIndex);
+
+			// BlendShape data is respective to point indices for each mesh in isolation, but we combine all points
+			// into one FSkeletalMeshImportData per LOD, so we need to remap the indices using this
+			uint32 NumPointsBeforeThisMesh = static_cast<uint32>(LODImportData.Points.Num());
+
+			bool bSuccess = UsdToUnreal::ConvertSkinnedMesh(SkinningQuery, SkeletonQuery, LODImportData, LODMaterialInfo, Options);
+			if (!bSuccess)
+			{
+				return true;
+			}
+
+			if (OutBlendShapes)
+			{
+				pxr::SdfPath MeshPrimPath = LODMesh.GetPrim().GetPrimPath();
+
+				pxr::UsdSkelBindingAPI SkelBindingAPI{Stage->GetPrimAtPath(MeshPrimPath)};
+				if (pxr::UsdSkelBlendShapeQuery BlendShapeQuery{SkelBindingAPI})
+				{
+					for (uint32 BlendShapeIndex = 0; BlendShapeIndex < BlendShapeQuery.GetNumBlendShapes(); ++BlendShapeIndex)
+					{
+						UsdToUnreal::ConvertBlendShape(
+							BlendShapeQuery.GetBlendShape(BlendShapeIndex),
+							StageInfo,
+							LODIndex,
+							NumPointsBeforeThisMesh,
+							InOutUsedMorphTargetNames,
+							*OutBlendShapes
+						);
+					}
+				}
+			}
+
+			return true;
+		};
+
+		// Actually parse all mesh data
+		for (const pxr::SdfPath& SkinnedPrimPath : PathsToSkinnedPrims)
+		{
+			pxr::UsdGeomMesh SkinnedMesh{Stage->GetPrimAtPath(SkinnedPrimPath)};
+			if (!SkinnedMesh)
+			{
+				continue;
+			}
+
+			pxr::UsdPrim ParentPrim = SkinnedMesh.GetPrim().GetParent();
+			FString ParentPrimPath = UsdToUnreal::ConvertPath(ParentPrim.GetPath());
+
+			bool bInterpretedLODs = false;
+			if (bInInterpretLODs && ParentPrim && !ProcessedLODParentPaths.Contains(ParentPrimPath))
+			{
+				// At the moment we only consider a single mesh per variant, so if multiple meshes tell us to process the same parent prim, we skip.
+				// This check would also prevent us from getting in here in case we just have many meshes children of a same prim, outside
+				// of a variant. In this case they don't fit the "one mesh per variant" pattern anyway, and we want to fallback to ignoring LODs
+				ProcessedLODParentPaths.Add(ParentPrimPath);
+
+				// WARNING: After this is called, references to objects that were inside any of the LOD Meshes will be invalidated!
+				bInterpretedLODs = UsdUtils::IterateLODMeshes(ParentPrim, ConvertLOD);
+			}
+
+			if (!bInterpretedLODs)
+			{
+				// Refresh reference to this prim as it could have been inside a variant that was temporarily switched by IterateLODMeshes
+				ConvertLOD(pxr::UsdGeomMesh{Stage->GetPrimAtPath(SkinnedPrimPath)}, 0);
+			}
+		}
+
+		// Place the LODs in order as we can't have e.g. LOD0 and LOD2 without LOD1, and there's no reason downstream code needs to care about
+		// what LOD number these data originally wanted to be
+		TMap<int32, int32> OldLODIndexToNewLODIndex;
+		LODIndexToSkeletalMeshImportDataMap.KeySort(TLess<int32>());
+		OutLODIndexToSkeletalMeshImportData.Reset(LODIndexToSkeletalMeshImportDataMap.Num());
+		OutLODIndexToMaterialInfo.Reset(LODIndexToMaterialInfoMap.Num());
+		for (TPair<int32, FSkeletalMeshImportData>& Entry : LODIndexToSkeletalMeshImportDataMap)
+		{
+			FSkeletalMeshImportData& ImportData = Entry.Value;
+			if (Entry.Value.Points.Num() == 0)
+			{
+				continue;
+			}
+
+			const int32 OldLODIndex = Entry.Key;
+			const int32 NewLODIndex = OutLODIndexToSkeletalMeshImportData.Add(MoveTemp(ImportData));
+			OutLODIndexToMaterialInfo.Add(LODIndexToMaterialInfoMap[OldLODIndex]);
+
+			// Keep track of these to remap blendshapes
+			OldLODIndexToNewLODIndex.Add(OldLODIndex, NewLODIndex);
+		}
+		if (OutBlendShapes)
+		{
+			for (auto& Pair : *OutBlendShapes)
+			{
+				UsdUtils::FUsdBlendShape& BlendShape = Pair.Value;
+
+				TSet<int32> NewLODIndexUsers;
+				NewLODIndexUsers.Reserve(BlendShape.LODIndicesThatUseThis.Num());
+
+				for (int32 OldLODIndexUser : BlendShape.LODIndicesThatUseThis)
+				{
+					if (int32* FoundNewLODIndex = OldLODIndexToNewLODIndex.Find(OldLODIndexUser))
+					{
+						NewLODIndexUsers.Add(*FoundNewLODIndex);
+					}
+					else
+					{
+						UE_LOG(LogUsd, Error, TEXT("Failed to remap blend shape '%s's LOD index '%d'"), *BlendShape.Name, OldLODIndexUser);
+					}
+				}
+
+				BlendShape.LODIndicesThatUseThis = MoveTemp(NewLODIndexUsers);
+			}
+		}
+
+		return true;
+	}
+
+	/** Warning: This function will temporarily switch the active LOD variant if one exists, so it's *not* thread safe! */
+	void SetMaterialOverrides(
+		const pxr::UsdSkelBinding& SkeletonBinding,
+		const TArray<UMaterialInterface*>& ExistingAssignments,
+		UMeshComponent& MeshComponent,
+		UUsdAssetCache2& AssetCache,
+		FUsdInfoCache& InfoCache,
+		float Time,
+		EObjectFlags Flags,
+		bool bInterpretLODs,
+		const FName& RenderContext,
+		const FName& MaterialPurpose,
+		const EUsdPurpose PurposesToLoad
+	)
+	{
+		FScopedUsdAllocs Allocs;
+
+		const pxr::UsdSkelSkeleton& Skeleton = SkeletonBinding.GetSkeleton();
+		if (!Skeleton)
+		{
+			return;
+		}
+		pxr::UsdPrim SkeletonPrim = Skeleton.GetPrim();
+		pxr::SdfPath SkeletonPrimPath = SkeletonPrim.GetPath();
+
+		pxr::UsdStageRefPtr Stage = SkeletonPrim.GetStage();
+
+		pxr::TfToken RenderContextToken = pxr::UsdShadeTokens->universalRenderContext;
+		if (!RenderContext.IsNone())
+		{
+			RenderContextToken = UnrealToUsd::ConvertToken(*RenderContext.ToString()).Get();
+		}
+
+		pxr::TfToken MaterialPurposeToken = pxr::UsdShadeTokens->allPurpose;
+		if (!MaterialPurpose.IsNone())
+		{
+			MaterialPurposeToken = UnrealToUsd::ConvertToken(*MaterialPurpose.ToString()).Get();
+		}
+
+		TMap<int32, UsdUtils::FUsdPrimMaterialAssignmentInfo> LODIndexToMaterialInfoMap;
+		TMap<int32, TSet<UsdUtils::FUsdPrimMaterialSlot>> CombinedSlotsForLODIndex;
+		TFunction<bool(const pxr::UsdGeomMesh&, int32)> IterateLODsLambda =
+			[&LODIndexToMaterialInfoMap, &CombinedSlotsForLODIndex, Time, RenderContextToken, MaterialPurposeToken, PurposesToLoad]
+			(const pxr::UsdGeomMesh& LODMesh, int32 LODIndex)
+			{
+				if (LODMesh && LODMesh.ComputeVisibility() == pxr::UsdGeomTokens->invisible)
+				{
+					return true;
+				}
+
+				// Ignore prims with disabled purposes: We need to match the material slot ordering that was used
+				// to generate the mesh in the first place, so this is important
+				if (!EnumHasAllFlags(PurposesToLoad, IUsdPrim::GetPurpose(LODMesh.GetPrim())))
+				{
+					return true;
+				}
+
+				TArray<UsdUtils::FUsdPrimMaterialSlot>& CombinedLODSlots = LODIndexToMaterialInfoMap.FindOrAdd(LODIndex).Slots;
+				TSet<UsdUtils::FUsdPrimMaterialSlot>& CombinedLODSlotsSet = CombinedSlotsForLODIndex.FindOrAdd(LODIndex);
+
+				const bool bProvideMaterialIndices = false;	   // We have no use for material indices and it can be slow to retrieve, as it will iterate
+															// all faces
+				UsdUtils::FUsdPrimMaterialAssignmentInfo LocalInfo = UsdUtils::GetPrimMaterialAssignments(
+					LODMesh.GetPrim(),
+					pxr::UsdTimeCode(Time),
+					bProvideMaterialIndices,
+					RenderContextToken,
+					MaterialPurposeToken
+				);
+
+				// Combine material slots in the same order that UsdToUnreal::ConvertSkinnedMesh does
+				for (UsdUtils::FUsdPrimMaterialSlot& LocalSlot : LocalInfo.Slots)
+				{
+					if (!CombinedLODSlotsSet.Contains(LocalSlot))
+					{
+						CombinedLODSlots.Add(LocalSlot);
+						CombinedLODSlotsSet.Add(LocalSlot);
+					}
+				}
+
+				return true;
+			};
+
+		TSet<FString> ProcessedLODParentPaths;
+
+		// Because we combine all skinning target meshes into a single skeletal mesh, we'll have to reconstruct the combined
+		// material assignment info that this SkelRoot wants in order to compare with the existing assignments.
+		for (const pxr::UsdSkelSkinningQuery& SkinningQuery : SkeletonBinding.GetSkinningTargets())
+		{
+			pxr::UsdPrim MeshPrim = SkinningQuery.GetPrim();
+			pxr::UsdGeomMesh Mesh{MeshPrim};
+			if (!Mesh)
+			{
+				continue;
+			}
+
+			// GetSkinningTargets can also return prims without the skel binding API, but we don't want to collect
+			// material bindings from those as they are not going to be globbed into the SkeletalMesh
+			if (!MeshPrim.HasAPI<pxr::UsdSkelBindingAPI>())
+			{
+				continue;
+			}
+
+			pxr::SdfPath MeshPrimPath = MeshPrim.GetPath();
+
+			pxr::UsdPrim ParentPrim = MeshPrim.GetParent();
+			FString ParentPrimPath = UsdToUnreal::ConvertPath(ParentPrim.GetPath());
+
+			bool bInterpretedLODs = false;
+			if (bInterpretLODs && UsdUtils::IsGeomMeshALOD(MeshPrim) && !ProcessedLODParentPaths.Contains(ParentPrimPath))
+			{
+				ProcessedLODParentPaths.Add(ParentPrimPath);
+
+				bInterpretedLODs = UsdUtils::IterateLODMeshes(ParentPrim, IterateLODsLambda);
+			}
+
+			if (!bInterpretedLODs)
+			{
+				// Refresh reference to this prim as it could have been inside a variant that was temporarily switched by IterateLODMeshes
+				IterateLODsLambda(pxr::UsdGeomMesh{Stage->GetPrimAtPath(MeshPrimPath)}, 0);
+			}
+		}
+
+		// Refresh reference to Skeleton prim because variant switching potentially invalidated it
+		pxr::UsdPrim ValidSkeletonPrim = Stage->GetPrimAtPath(SkeletonPrimPath);
+
+		// Place the LODs in order as we can't have e.g. LOD0 and LOD2 without LOD1, and there's no reason downstream code needs to care about
+		// what LOD number these data originally wanted to be
+		TArray<UsdUtils::FUsdPrimMaterialAssignmentInfo> LODIndexToAssignments;
+		LODIndexToMaterialInfoMap.KeySort(TLess<int32>());
+		for (TPair<int32, UsdUtils::FUsdPrimMaterialAssignmentInfo>& Entry : LODIndexToMaterialInfoMap)
+		{
+			LODIndexToAssignments.Add(MoveTemp(Entry.Value));
+		}
+
+		// Stash our mesh PrimvarToUVIndex into the assignment info, as that's where ResolveMaterialAssignmentInfo will look for it
+		UUsdMeshAssetUserData* UserData = nullptr;
+		if (USkeletalMeshComponent* SkeletalMeshComponent = Cast<USkeletalMeshComponent>(&MeshComponent))
+		{
+			if (USkeletalMesh* Mesh = SkeletalMeshComponent->GetSkeletalMeshAsset())
+			{
+				UserData = Mesh->GetAssetUserData<UUsdMeshAssetUserData>();
+			}
+		}
+		if (UserData)
+		{
+			LODIndexToAssignments[0].PrimvarToUVIndex = UserData->PrimvarToUVIndex;
+		}
+		else
+		{
+			ensureMsgf(
+				UserData,
+				TEXT("Expected component '%s''s SkeletalMesh to have an instance of UUsdMeshAssetUserData at this point!"),
+				*MeshComponent.GetPathName()
+			);
+		}
+
+		TMap<const UsdUtils::FUsdPrimMaterialSlot*, UMaterialInterface*> ResolvedMaterials = MeshTranslationImpl::ResolveMaterialAssignmentInfo(
+			ValidSkeletonPrim,
+			LODIndexToAssignments,
+			AssetCache,
+			InfoCache,
+			Flags
+		);
+
+		// Compare resolved materials with existing assignments, and create overrides if we need to
+		uint32 SkeletalMeshSlotIndex = 0;
+		for (int32 LODIndex = 0; LODIndex < LODIndexToAssignments.Num(); ++LODIndex)
+		{
+			const TArray<UsdUtils::FUsdPrimMaterialSlot>& LODSlots = LODIndexToAssignments[LODIndex].Slots;
+			for (int32 LODSlotIndex = 0; LODSlotIndex < LODSlots.Num(); ++LODSlotIndex, ++SkeletalMeshSlotIndex)
+			{
+				const UsdUtils::FUsdPrimMaterialSlot& Slot = LODSlots[LODSlotIndex];
+
+				UMaterialInterface* Material = nullptr;
+				if (UMaterialInterface** FoundMaterial = ResolvedMaterials.Find(&Slot))
+				{
+					Material = *FoundMaterial;
+				}
+				else
+				{
+					UE_LOG(
+						LogUsd,
+						Error,
+						TEXT("Lost track of resolved material for slot '%d' of LOD '%d' for skeletal mesh '%s'"),
+						LODSlotIndex,
+						LODIndex,
+						*UsdToUnreal::ConvertPath(ValidSkeletonPrim.GetPath())
+					);
+					continue;
+				}
+
+				UMaterialInterface* ExistingMaterial = ExistingAssignments.IsValidIndex(SkeletalMeshSlotIndex)
+														   ? ExistingAssignments[SkeletalMeshSlotIndex]
+														   : nullptr;
+				if (!ExistingMaterial || ExistingMaterial == Material)
+				{
+					continue;
+				}
+				else
+				{
+					MeshComponent.SetMaterial(SkeletalMeshSlotIndex, Material);
+				}
+			}
+		}
+	}
+
+	bool HasLODSkinningTargets(const pxr::UsdSkelBinding& SkelBinding)
+	{
+		FScopedUsdAllocs Allocs;
+
+		for (const pxr::UsdSkelSkinningQuery& SkinningQuery : SkelBinding.GetSkinningTargets())
+		{
+			if (UsdUtils::IsGeomMeshALOD(SkinningQuery.GetPrim()))
+			{
+				return true;
+			}
+		}
+
+		return false;
+	}
+
+	UPhysicsAsset* GenerateAndAssignPhysicsAsset(USkeletalMesh* SkeletalMesh, EObjectFlags Flags)
+	{
+#if WITH_EDITOR
+		FName AssetName = MakeUniqueObjectName(
+			GetTransientPackage(),
+			UPhysicsAsset::StaticClass(),
+			*IUsdClassesModule::SanitizeObjectName(FPaths::GetBaseFilename(TEXT("PHYS_") + SkeletalMesh->GetName()))
+		);
+
+		UPhysicsAsset* Result = NewObject<UPhysicsAsset>(
+			GetTransientPackage(),
+			AssetName,
+			Flags | EObjectFlags::RF_Public | EObjectFlags::RF_Transient
+		);
+
+		FPhysAssetCreateParams NewBodyData;
+		FText CreationErrorMessage;
+		bool bSuccess = FPhysicsAssetUtils::CreateFromSkeletalMesh(Result, SkeletalMesh, NewBodyData, CreationErrorMessage);
+		if (!bSuccess)
+		{
+			UE_LOG(
+				LogUsd,
+				Error,
+				TEXT("Couldn't create PhysicsAsset for SkeletalMesh '%s': %s"),
+				*SkeletalMesh->GetName(),
+				*CreationErrorMessage.ToString()
+			);
+
+			ensure(ObjectTools::DeleteSingleObject(Result));
+
+			Result = nullptr;
+		}
+
+		return Result;
+#else
+		return nullptr;
+#endif	  // WITH_EDITOR
+	}
+
+	UAnimBlueprint* CreateAnimBlueprint(
+		const FUsdSchemaTranslationContext& Context,
+		const pxr::UsdPrim& SkeletonPrim,
+		bool bDelayRecompilation = false,	  // Whether we should recompile within this function or not
+		bool* bOutNeedsRecompile = nullptr	  // Whether this function wants the returned AnimBP to be recompiled
+	)
+	{
+		if (!SkeletonPrim || !Context.InfoCache || !Context.AssetCache)
+		{
+			return nullptr;
+		}
+
+		FString PrimName = UsdToUnreal::ConvertString(SkeletonPrim.GetName());
+
+		USkeletalMesh* SkeletalMesh = Context.InfoCache->GetSingleAssetForPrim<USkeletalMesh>(UE::FSdfPath{SkeletonPrim.GetPath()});
+		if (!SkeletalMesh)
+		{
+			return nullptr;
+		}
+
+		USkeleton* Skeleton = SkeletalMesh->GetSkeleton();
+		if (!Skeleton)
+		{
+			return nullptr;
+		}
+
+		// Fetch relevant attributes from prim, since we know it has the schema
+		FString AnimBPPath;
+		{
+			FScopedUsdAllocs Allocs;
+
+			if (pxr::UsdAttribute Attr = SkeletonPrim.GetAttribute(UnrealIdentifiers::UnrealAnimBlueprintPath))
+			{
+				std::string PathString;
+				if (Attr.Get(&PathString))
+				{
+					AnimBPPath = UsdToUnreal::ConvertString(PathString);
+				}
+			}
+			// Temporarily check the SkelRoot for the same attribute for backwards compatibility
+			else if (pxr::UsdPrim ClosestParentSkelRoot = UsdUtils::GetClosestParentSkelRoot(SkeletonPrim))
+			{
+				if (pxr::UsdAttribute ParentAttr = ClosestParentSkelRoot.GetAttribute(UnrealIdentifiers::UnrealAnimBlueprintPath))
+				{
+					std::string PathString;
+					if (ParentAttr.Get(&PathString))
+					{
+						AnimBPPath = UsdToUnreal::ConvertString(PathString);
+					}
+				}
+			}
+		}
+		if (AnimBPPath.IsEmpty())
+		{
+			return nullptr;
+		}
+
+		bool bNeedRecompile = false;
+
+		UAnimBlueprint* AnimBP = Cast<UAnimBlueprint>(FSoftObjectPath(AnimBPPath).TryLoad());
+		if (!AnimBP)
+		{
+			return nullptr;
+		}
+
+		// Create transient AnimBP based on our template, so that we can assign it a proper skeleton
+		const static FString DefaultAnimBPPath = TEXT("/USDImporter/Blueprint/DefaultLiveLinkAnimBP.DefaultLiveLinkAnimBP");
+		if (DefaultAnimBPPath == AnimBPPath)
+		{
+			const FString PrimPath = UsdToUnreal::ConvertPath(SkeletonPrim.GetPrimPath());
+
+			// Let's try to never reuse AnimBP between prims (as we want to be able to switch subject names
+			// independently and we likely won't have more than a handful of these anyway).
+			// We should have at least something deterministic though so that we don't repeatedly recreate assets for the same prim.
+			FSHAHash Hash;
+			FSHA1 SHA1;
+			// Each stage actor has a separate info cache that is assigned to the context
+			SHA1.Update(reinterpret_cast<const uint8*>(Context.InfoCache.Get()), sizeof(Context.InfoCache.Get()));
+			SHA1.UpdateWithString(*PrimPath, PrimPath.Len());
+			SHA1.Final();
+			SHA1.GetHash(&Hash.Hash[0]);
+			const FString CacheKey = Hash.ToString();
+
+			// Check if we can find an AnimBP for this prim in the asset cache (useful when doing Action->Import)
+			bool bReusedAnimBP = false;
+			if (UAnimBlueprint* CachedAnimBP = Cast<UAnimBlueprint>(Context.AssetCache->GetCachedAsset(CacheKey)))
+			{
+				if (CachedAnimBP->TargetSkeleton == Skeleton)
+				{
+					AnimBP = CachedAnimBP;
+					bReusedAnimBP = true;
+				}
+			}
+
+			// We have to generate a new transient AnimBP
+			if (!bReusedAnimBP)
+			{
+				FName UniqueName = MakeUniqueObjectName(
+					GetTransientPackage(),
+					UAnimBlueprint::StaticClass(),
+					*IUsdClassesModule::SanitizeObjectName(PrimName + TEXT("_DefaultAnimBlueprint"))
+				);
+
+				// Duplicate and never reuse these so that they can be assigned independent subject names if desired.
+				// Its not as if scenes will have thousands of these anyway.
+				AnimBP = DuplicateObject(AnimBP, GetTransientPackage(), UniqueName);
+
+				// Patch up the flags here or else the rest of the engine code (and our plugin code) will get confused
+				// as to what a non-transient asset in the transient package even means. It can lead to some crashes
+				// when saving after import if we don't do this
+				AnimBP->ClearFlags(AnimBP->GetFlags());
+				AnimBP->SetFlags(Skeleton->GetFlags());
+
+				AnimBP->TargetSkeleton = Skeleton;
+				AnimBP->bIsTemplate = false;
+
+				bNeedRecompile = true;
+
+				Context.AssetCache->CacheAsset(CacheKey, AnimBP);
+			}
+		}
+		// Path is pointing to an existing, persistent AnimBP
+		else
+		{
+			// Force skeletons to be compatible if they aren't (we need both ways!)
+			if (!Skeleton->IsCompatibleForEditor(AnimBP->TargetSkeleton))
+			{
+				AnimBP->TargetSkeleton->AddCompatibleSkeleton(Skeleton);
+				Skeleton->AddCompatibleSkeleton(AnimBP->TargetSkeleton);
+
+				if (AnimBP->TargetSkeleton->GetReferenceSkeleton().GetRefBoneInfo() != Skeleton->GetReferenceSkeleton().GetRefBoneInfo())
+				{
+					UE_LOG(
+						LogUsd,
+						Warning,
+						TEXT("Forcing AnimBlueprint '%s's Skeleton '%s' to be compatible with the Skeleton generated for prim '%s', but they may be "
+							 "different!"),
+						*AnimBP->GetPathName(),
+						*AnimBP->TargetSkeleton->GetPathName(),
+						*PrimName
+					);
+				}
+			}
+		}
+
+		if (bNeedRecompile && !bDelayRecompilation)
+		{
+			FCompilerResultsLog Results;
+			FBPCompileRequest Request(AnimBP, EBlueprintCompileOptions::None, &Results);
+			FBlueprintCompilationManager::CompileSynchronously(Request);
+			bNeedRecompile = false;
+		}
+
+		if (bOutNeedsRecompile)
+		{
+			*bOutNeedsRecompile = bNeedRecompile;
+		}
+
+		return AnimBP;
+	}
+
+	void UpdateLiveLinkProperties(const FUsdSchemaTranslationContext& Context, USkeletalMeshComponent* Component, const pxr::UsdPrim& SkeletonPrim)
+	{
+		if (!Component || !Component->GetSkeletalMeshAsset() || !SkeletonPrim)
+		{
+			return;
+		}
+
+		FString PrimName = UsdToUnreal::ConvertString(SkeletonPrim.GetName());
+
+		USkeleton* Skeleton = Component->GetSkeletalMeshAsset()->GetSkeleton();
+		if (!Skeleton)
+		{
+			return;
+		}
+
+		UAnimBlueprint* ExistingAnimBP = nullptr;
+		if (Component->AnimClass)
+		{
+			ExistingAnimBP = Cast<UAnimBlueprint>(Component->AnimClass->ClassGeneratedBy);
+		}
+
+		// Fetch relevant attributes from prim, since we know it has the schema
+		TUsdStore<pxr::UsdSkelRoot> ClosestSkelRootParent;
+		TFunction<FString(const pxr::TfToken&)> GetAttrValue = [&SkeletonPrim, &ClosestSkelRootParent](const pxr::TfToken& AttrName) -> FString
+		{
+			FScopedUsdAllocs Allocs;
+
+			if (pxr::UsdAttribute Attr = SkeletonPrim.GetAttribute(AttrName))
+			{
+				std::string SubjectNameString;
+				if (Attr.Get(&SubjectNameString))
+				{
+					return UsdToUnreal::ConvertString(SubjectNameString);
+				}
+			}
+			// Temporarily check the SkelRoot for the same attribute for backwards compatibility
+			else
+			{
+				if (!ClosestSkelRootParent.Get())
+				{
+					ClosestSkelRootParent = pxr::UsdSkelRoot{UsdUtils::GetClosestParentSkelRoot(SkeletonPrim)};
+				}
+
+				if (ClosestSkelRootParent.Get())
+				{
+					if (pxr::UsdAttribute ParentAttr = ClosestSkelRootParent.Get().GetPrim().GetAttribute(AttrName))
+					{
+						std::string SubjectNameString;
+						if (ParentAttr.Get(&SubjectNameString))
+						{
+							return UsdToUnreal::ConvertString(SubjectNameString);
+						}
+					}
+				}
+			}
+
+			return {};
+		};
+		FString SubjectName = GetAttrValue(UnrealIdentifiers::UnrealLiveLinkSubjectName);
+		FString AnimBPPath = GetAttrValue(UnrealIdentifiers::UnrealAnimBlueprintPath);
+
+		UAnimBlueprint* AnimBP = Cast<UAnimBlueprint>(FSoftObjectPath(AnimBPPath).TryLoad());
+
+		// Check if we need to change the AnimBP
+		bool bNeedRecompile = false;
+		if (ExistingAnimBP != AnimBP)
+		{
+			const bool bDelayRecompilation = true;
+			AnimBP = UsdSkelSkeletonTranslatorImpl::CreateAnimBlueprint(Context, SkeletonPrim, bDelayRecompilation, &bNeedRecompile);
+		}
+
+		// Apply subject name to live link pose AnimBlueprint node
+		// Reference: UAnimationBlueprintLibrary::AddNodeAssetOverride
+		if (AnimBP)
+		{
+			TArray<UBlueprint*> BlueprintHierarchy;
+			AnimBP->GetBlueprintHierarchyFromClass(AnimBP->GetAnimBlueprintGeneratedClass(), BlueprintHierarchy);
+
+			TArray<UAnimGraphNode_LiveLinkPose*> LiveLinkNodes;
+
+			for (int32 BlueprintIndex = 0; BlueprintIndex < BlueprintHierarchy.Num(); ++BlueprintIndex)
+			{
+				UBlueprint* CurrentBlueprint = BlueprintHierarchy[BlueprintIndex];
+
+				TArray<UEdGraph*> Graphs;
+				CurrentBlueprint->GetAllGraphs(Graphs);
+
+				for (UEdGraph* Graph : Graphs)
+				{
+					for (UEdGraphNode* Node : Graph->Nodes)
+					{
+						if (UAnimGraphNode_LiveLinkPose* AnimNode = Cast<UAnimGraphNode_LiveLinkPose>(Node))
+						{
+							LiveLinkNodes.Add(AnimNode);
+						}
+					}
+				}
+			}
+
+			if (LiveLinkNodes.Num() > 1 && !ExistingAnimBP)
+			{
+				UE_LOG(
+					LogUsd,
+					Warning,
+					TEXT("Found more than one LiveLinkPose blueprint node on AnimBlueprint '%s's graphs."
+						 "Note that all of those nodes will have their LiveLink SubjectName's updated to '%s', as described for prim '%s'!"),
+					*AnimBP->GetPathName(),
+					*SubjectName,
+					*UsdToUnreal::ConvertPath(SkeletonPrim.GetPrimPath())
+				);
+			}
+
+			for (UAnimGraphNode_LiveLinkPose* Node : LiveLinkNodes)
+			{
+				const UEdGraphSchema* Schema = Node->GetSchema();
+				if (!Schema)
+				{
+					continue;
+				}
+
+				UEdGraphPin* SubjectNamePin = nullptr;
+				for (UEdGraphPin* Pin : Node->Pins)
+				{
+					if (Pin->GetName() == GET_MEMBER_NAME_STRING_CHECKED(FAnimNode_LiveLinkPose, LiveLinkSubjectName))
+					{
+						SubjectNamePin = Pin;
+						break;
+					}
+				}
+
+				// The subject name pin is already connected to something...
+				if (SubjectNamePin->LinkedTo.Num() != 0)
+				{
+					if (!ExistingAnimBP)
+					{
+						UE_LOG(
+							LogUsd,
+							Warning,
+							TEXT("Failed to update a LiveLinkPose node's 'Subject Name' to '%s' on AnimBlueprint '%s', "
+								 "because the pin is already connected to some other node. Disconnect it if you want it to be updated automatically."
+							),
+							*SubjectName,
+							*AnimBP->GetPathName()
+						);
+					}
+
+					continue;
+				}
+
+				if (SubjectNamePin)
+				{
+					// The pin type is FLiveLinkSubjectName, so we must create an instance of it and serialize it using
+					// UScriptStruct::ExportText to generate a proper default value string
+					FLiveLinkSubjectName Dummy;
+					Dummy.Name = *SubjectName;
+
+					FString ValueString;
+					const void* Defaults = nullptr;
+					UObject* OwnerObject = nullptr;
+					const int32 PortFlags = EPropertyPortFlags::PPF_None;
+					UObject* ExportRootScope = nullptr;
+					FLiveLinkSubjectName::StaticStruct()->ExportText(ValueString, &Dummy, Defaults, OwnerObject, PortFlags, ExportRootScope);
+
+					if (!Schema->DoesDefaultValueMatch(*SubjectNamePin, ValueString))
+					{
+						SubjectNamePin->Modify();
+						Schema->TrySetDefaultValue(*SubjectNamePin, ValueString);
+
+						FBlueprintEditorUtils::MarkBlueprintAsModified(AnimBP);
+						bNeedRecompile = true;
+					}
+				}
+			}
+		}
+
+		if (bNeedRecompile)
+		{
+			FCompilerResultsLog Results;
+			FBPCompileRequest Request(AnimBP, EBlueprintCompileOptions::None, &Results);
+			FBlueprintCompilationManager::CompileSynchronously(Request);
+
+			// We need to force the component to update its anim after we regenerate the blueprint class
+			Component->ClearAnimScriptInstance();
+			Component->InitAnim(true);
+		}
+
+		if (AnimBP != ExistingAnimBP)
+		{
+			// This can internally change AnimationMode, but lets revert it to what it was so that we can control it from
+			// that single place in ::UpdateComponents
+			EAnimationMode::Type OldMode = Component->GetAnimationMode();
+			Component->SetAnimInstanceClass(AnimBP->GeneratedClass);
+			Component->SetAnimationMode(OldMode);
+		}
+	}
+
+	class FSkelSkeletonCreateAssetsTaskChain : public FUsdSchemaTranslatorTaskChain
+	{
+	public:
+		explicit FSkelSkeletonCreateAssetsTaskChain(
+			const TSharedRef<FUsdSchemaTranslationContext>& InContext,
+			const UE::FSdfPath& InSkeletonPrimPath
+		);
+
+	protected:
+		// Inputs
+		UE::FSdfPath SkeletonPrimPath;
+		TSharedRef<FUsdSchemaTranslationContext> Context;
+
+		// Outputs
+		TArray<FSkeletalMeshImportData> LODIndexToSkeletalMeshImportData;
+		TArray<UsdUtils::FUsdPrimMaterialAssignmentInfo> LODIndexToMaterialInfo;
+		TArray<SkeletalMeshImportData::FBone> SkeletonBones;
+		FName SkeletonName;
+		UsdUtils::FBlendShapeMap NewBlendShapes;
+
+		// Note that we want this to be case insensitive so that our UMorphTarget FNames are unique not only due to case differences
+		TSet<FString> UsedMorphTargetNames;
+		TUsdStore<pxr::UsdSkelBinding> SkeletonBinding;
+		TUsdStore<pxr::UsdSkelCache> SkelCache;
+		TUsdStore<pxr::UsdSkelRoot> ClosestParentSkelRoot;
+		TUsdStore<pxr::UsdSkelSkeletonQuery> SkeletonQuery;
+		FString SkeletalMeshHashString;
+
+		// Don't keep a live reference to the prim because other translators may mutate the stage in an ExclusiveSync translation step, invalidating
+		// the reference
+		UE::FUsdPrim GetSkeletonPrim() const
+		{
+			return Context->Stage.GetPrimAtPath(SkeletonPrimPath);
+		}
+
+		void SetupTasks();
+	};
+
+	FSkelSkeletonCreateAssetsTaskChain::FSkelSkeletonCreateAssetsTaskChain(
+		const TSharedRef<FUsdSchemaTranslationContext>& InContext,
+		const UE::FSdfPath& InPrimPath
+	)
+		: SkeletonPrimPath(InPrimPath)
+		, Context(InContext)
+	{
+		// Collect our skel binding info
+		{
+			FScopedUsdAllocs Allocs;
+
+			pxr::UsdPrim SkeletonPrim = GetSkeletonPrim();
+
+			bool bSuccess = false;
+
+			ClosestParentSkelRoot = pxr::UsdSkelRoot{UsdUtils::GetClosestParentSkelRoot(SkeletonPrim)};
+			if (ClosestParentSkelRoot.Get() && SkelCache.Get().Populate(ClosestParentSkelRoot.Get(), pxr::UsdTraverseInstanceProxies()))
+			{
+				bSuccess = UsdUtils::GetSkelQueries(
+					ClosestParentSkelRoot.Get(),
+					pxr::UsdSkelSkeleton{SkeletonPrim},
+					SkeletonBinding.Get(),
+					SkeletonQuery.Get(),
+					&SkelCache.Get()
+				);
+			}
+
+			if (!bSuccess)
+			{
+				UE_LOG(
+					LogUsd,
+					Warning,
+					TEXT("Failed to find skeleton binding info for skeleton at path '%s'. Is it within a SkelRoot prim?"),
+					*InPrimPath.GetString()
+				);
+				return;
+			}
+		}
+
+		SetupTasks();
+	}
+
+	void FSkelSkeletonCreateAssetsTaskChain::SetupTasks()
+	{
+		// To parse all LODs we need to actively switch variant sets to other variants (triggering prim loading/unloading and notices),
+		// which could cause race conditions if other async translation tasks are trying to access those prims
+		const bool bTryLODParsing = Context->bAllowInterpretingLODs && HasLODSkinningTargets(SkeletonBinding.Get());
+		ESchemaTranslationLaunchPolicy LaunchPolicy = bTryLODParsing ? ESchemaTranslationLaunchPolicy::ExclusiveSync
+																	 : ESchemaTranslationLaunchPolicy::Async;
+
+		// Create SkeletalMeshImportData (Async or ExclusiveSync)
+		Do(LaunchPolicy,
+		   [this, bTryLODParsing]()
+		   {
+			   // No point in importing blend shapes if the import context doesn't want them
+			   UsdUtils::FBlendShapeMap* OutBlendShapes = Context->BlendShapesByPath ? &NewBlendShapes : nullptr;
+
+			   pxr::TfToken RenderContextToken = pxr::UsdShadeTokens->universalRenderContext;
+			   if (!Context->RenderContext.IsNone())
+			   {
+				   RenderContextToken = UnrealToUsd::ConvertToken(*Context->RenderContext.ToString()).Get();
+			   }
+
+			   pxr::TfToken MaterialPurposeToken = pxr::UsdShadeTokens->allPurpose;
+			   if (!Context->MaterialPurpose.IsNone())
+			   {
+				   MaterialPurposeToken = UnrealToUsd::ConvertToken(*Context->MaterialPurpose.ToString()).Get();
+			   }
+
+			   UsdToUnreal::FUsdMeshConversionOptions Options;
+			   Options.TimeCode = Context->Time;
+			   Options.RenderContext = RenderContextToken;
+			   Options.MaterialPurpose = MaterialPurposeToken;
+			   Options.SubdivisionLevel = Context->SubdivisionLevel;
+			   Options.PurposesToLoad = Context->PurposesToLoad;
+			   Options.bMergeIdenticalMaterialSlots = Context->bMergeIdenticalMaterialSlots;
+
+			   bool bContinueTaskChain = UsdSkelSkeletonTranslatorImpl::LoadAllSkeletalData(
+				   SkeletonBinding.Get(),
+				   SkelCache.Get(),
+				   LODIndexToSkeletalMeshImportData,
+				   LODIndexToMaterialInfo,
+				   SkeletonBones,
+				   SkeletonName,
+				   OutBlendShapes,
+				   UsedMorphTargetNames,
+				   bTryLODParsing,
+				   Options
+			   );
+
+			   int32 NumRootBones = 0;
+			   for (const SkeletalMeshImportData::FBone& Bone : SkeletonBones)
+			   {
+				   if (Bone.ParentIndex == INDEX_NONE)
+				   {
+					   NumRootBones++;
+				   }
+
+				   if (NumRootBones > 1)
+				   {
+					   break;
+				   }
+			   }
+
+			   if (NumRootBones != 1)
+			   {
+				   UE_LOG(
+					   LogUsd,
+					   Warning,
+					   TEXT("Ignoring SkelRoot '%s' as the bound Skeleton prim must have exactly one root bone!"),
+					   *SkeletonPrimPath.GetString()
+				   );
+				   bContinueTaskChain = false;
+			   }
+
+			   // If we parsed LODs we could potentially have invalidated our references to our queries,
+			   // so let's refresh them before the downstream tasks try using them.
+			   if (bContinueTaskChain && bTryLODParsing)
+			   {
+				   // We are forced to repopulate here unfortunately, or else the queries/SkelCache will actually happily hold on to and return their
+				   // internal, invalid prim references...
+				   // We can't just keep prim paths of things instead (paths to skinned prims, paths to SkelAnimation prim, etc.) because parsing
+				   // these prims involves using those query objects to compute joint transforms and etc., and we can't create them ourselves.
+				   // This should only be temporary, as I believe we'll eventually have a cleverer way of handling LODs that doesn't invalidate
+				   // everything (e.g. open a separate stage with a population mask to only parse the LOD mesh, etc.)
+				   SkelCache.Get().Populate(ClosestParentSkelRoot.Get(), pxr::UsdTraverseInstanceProxies());
+				   ensure(UsdUtils::GetSkelQueries(
+					   ClosestParentSkelRoot.Get(),
+					   pxr::UsdSkelSkeleton{GetSkeletonPrim()},
+					   SkeletonBinding.Get(),
+					   SkeletonQuery.Get(),
+					   &SkelCache.Get()
+				   ));
+			   }
+
+			   return bContinueTaskChain;
+		   });
+
+		// Create USkeletalMesh (Main thread)
+		Then(
+			ESchemaTranslationLaunchPolicy::Sync,
+			[this]()
+			{
+				if (!Context->InfoCache || !Context->AssetCache)
+				{
+					return false;
+				}
+
+				UsdUtils::FBlendShapeMap* BlendShapes = Context->BlendShapesByPath ? &NewBlendShapes : nullptr;
+
+				FSHAHash SkeletalMeshHash = UsdSkelSkeletonTranslatorImpl::ComputeSHAHash(
+					LODIndexToSkeletalMeshImportData,
+					SkeletonBones,
+					BlendShapes
+				);
+				SkeletalMeshHashString = SkeletalMeshHash.ToString();
+
+				USkeletalMesh* SkeletalMesh = Cast<USkeletalMesh>(Context->AssetCache->GetCachedAsset(SkeletalMeshHashString));
+
+				bool bIsNew = false;
+				if (!SkeletalMesh)
+				{
+					// Even though we're translating SkeletalMeshes from Skeleton prims now, keep using the SkelRoot
+					// prim name as the SkeletalMesh asset name. This because in the general case the entire SkelRoot will
+					// represent a character and be named after it, and the Skeleton will be named just "Skeleton" or "Skel"
+					FString MeshName;
+					{
+						FScopedUsdAllocs Allocs;
+						MeshName = UsdToUnreal::ConvertString(ClosestParentSkelRoot.Get().GetPrim().GetName());
+					}
+
+					bIsNew = true;
+					SkeletalMesh = UsdToUnreal::GetSkeletalMeshFromImportData(
+						LODIndexToSkeletalMeshImportData,
+						SkeletonBones,
+						NewBlendShapes,
+						Context->ObjectFlags,
+						*MeshName,
+						SkeletonName
+					);
+				}
+
+				if (SkeletalMesh)
+				{
+					UUsdMeshAssetUserData* UserData = SkeletalMesh->GetAssetUserData<UUsdMeshAssetUserData>();
+					if (!UserData)
+					{
+						UserData = NewObject<UUsdMeshAssetUserData>(SkeletalMesh, TEXT("USDAssetUserData"));
+						UserData->PrimvarToUVIndex = LODIndexToMaterialInfo[0].PrimvarToUVIndex;	// We use the same primvar mapping for all LODs
+						SkeletalMesh->AddAssetUserData(UserData);
+					}
+					UserData->PrimPaths.AddUnique(SkeletonPrimPath.GetString());
+
+					MeshTranslationImpl::RecordSourcePrimsForMaterialSlots(LODIndexToMaterialInfo, UserData);
+
+					if (bIsNew)
+					{
+						const bool bMaterialsHaveChanged = UsdSkelSkeletonTranslatorImpl::ProcessMaterials(
+							GetSkeletonPrim(),
+							LODIndexToMaterialInfo,
+							SkeletalMesh,
+							*Context->AssetCache.Get(),
+							*Context->InfoCache.Get(),
+							Context->Time,
+							Context->ObjectFlags,
+							NewBlendShapes.Num() > 0
+						);
+
+						if (bMaterialsHaveChanged)
+						{
+							const bool bRebuildAll = true;
+							SkeletalMesh->UpdateUVChannelData(bRebuildAll);
+						}
+
+						Context->AssetCache->CacheAsset(SkeletalMeshHashString, SkeletalMesh);
+						Context->AssetCache->CacheAsset(SkeletalMeshHashString + TEXT("_Skeleton"), SkeletalMesh->GetSkeleton());
+
+						// The PreviewSkeletalMesh property on the Skeleton is a soft object path, that is set within
+						// UsdToUnreal::GetSkeletalMeshFromImportData before the SkeletalMesh is part of the cache. When we cache the
+						// SkeletalMesh, it is renamed into the AssetCache package, and so PreviewSkeletalMesh is no longer pointing at it.
+						// Ideally we'd create all our asset cache assets directly outered to the AssetCache to avoid this issue, but
+						// for now I think it's only a problem for SkeletalMesh/Skeleton pairs (since they are created at the same time
+						// already knowing each other), so it's simpler to just do this
+						SkeletalMesh->GetSkeleton()->SetPreviewMesh(SkeletalMesh);
+					}
+
+					if (bGeneratePhysicsAssets)
+					{
+						if (!SkeletalMesh->GetPhysicsAsset())
+						{
+							UPhysicsAsset* PhysicsAsset = UsdSkelSkeletonTranslatorImpl::GenerateAndAssignPhysicsAsset(
+								SkeletalMesh,
+								Context->ObjectFlags
+							);
+
+							if (PhysicsAsset)
+							{
+								Context->AssetCache->CacheAsset(SkeletalMeshHashString + TEXT("_PhysicsAsset"), PhysicsAsset);
+								Context->InfoCache->LinkAssetToPrim(SkeletonPrimPath, PhysicsAsset);
+							}
+						}
+					}
+					else
+					{
+						// Actively clear this so that if we toggle the cvar and reload we'll clear our physics assets
+						SkeletalMesh->SetPhysicsAsset(nullptr);
+					}
+
+					Context->InfoCache->LinkAssetToPrim(SkeletonPrimPath, SkeletalMesh);
+
+					// Track our Skeleton by the source skeleton prim path
+					if (USkeleton* Skeleton = SkeletalMesh->GetSkeleton())
+					{
+						Context->InfoCache->LinkAssetToPrim(SkeletonPrimPath, Skeleton);
+					}
+
+					// We may be reusing a skeletal mesh we got in the cache, but we always need the BlendShapesByPath stored on the
+					// actor to be up-to-date with the Skeletal Mesh that is actually being displayed
+					if (Context->BlendShapesByPath)
+					{
+						Context->BlendShapesByPath->Append(NewBlendShapes);
+					}
+				}
+
+				// Continuing even if the mesh is not new as we currently don't add the SkelAnimation info to the mesh hash, so the animations
+				// may have changed
+				return true;
+			}
+		);
+
+		// Create AnimBP asset if we need to
+		Then(
+			ESchemaTranslationLaunchPolicy::Sync,
+			[this]()
+			{
+				UE::FUsdPrim Prim = GetSkeletonPrim();
+
+				// We need to *also* create the AnimBP within CreateAssets so that we will still generate it even if
+				// we never call Create/UpdateComponents (e.g. importing without importing actors).
+				bool bPrimHasLiveLinkSchema = UsdUtils::PrimHasSchema(Prim, UnrealIdentifiers::LiveLinkAPI);
+				if (!bPrimHasLiveLinkSchema)
+				{
+					// Temporarily check the closest skelroot prim for the schema for backwards compatibility
+					bPrimHasLiveLinkSchema = UsdUtils::PrimHasSchema(ClosestParentSkelRoot.Get().GetPrim(), UnrealIdentifiers::LiveLinkAPI);
+				}
+
+				if (bPrimHasLiveLinkSchema)
+				{
+					TSharedPtr<FUsdSchemaTranslationContext> ContextPtr = Context;
+
+					// When importing, we can't delay creating the AnimBP to the next frame as that will be after the
+					// import. We don't need to worry about deadlocks though, because we will never *import* as a
+					// response to a USD event: It's always an intentional function call, where the Python GIL is
+					// properly handled
+					if (ContextPtr && ContextPtr->bIsImporting)
+					{
+						UsdSkelSkeletonTranslatorImpl::CreateAnimBlueprint(*ContextPtr.Get(), Prim);
+					}
+					else
+					{
+						// HACK. c.f. the large comment on the analogous timer inside
+						// FUsdSkelSkeletonTranslator::UpdateComponents
+						FTSTicker::GetCoreTicker().AddTicker(FTickerDelegate::CreateLambda(
+							[ContextPtr, Prim](float Time)
+							{
+								if (ContextPtr)
+								{
+									UsdSkelSkeletonTranslatorImpl::CreateAnimBlueprint(*ContextPtr.Get(), Prim);
+								}
+
+								// Returning false means this is a one-off, and won't repeat
+								return false;
+							}
+						));
+					}
+				}
+
+				return true;
+			}
+		);
+
+		// Create UAnimSequences (requires a completed USkeleton. Main thread as some steps of the animation compression require it)
+		Then(
+			ESchemaTranslationLaunchPolicy::Sync,
+			[this]()
+			{
+				if (!Context->bAllowParsingSkeletalAnimations || !Context->InfoCache)
+				{
+					return false;
+				}
+
+				USkeletalMesh* SkeletalMesh = Context->InfoCache->GetSingleAssetForPrim<USkeletalMesh>(SkeletonPrimPath);
+				if (!SkeletalMesh)
+				{
+					return false;
+				}
+
+				FScopedUsdAllocs Allocs;
+
+				const pxr::UsdSkelSkeleton& Skeleton = SkeletonBinding.Get().GetSkeleton();
+				const pxr::UsdSkelAnimQuery& UsdAnimQuery = SkeletonQuery.Get().GetAnimQuery();
+				if (!UsdAnimQuery)
+				{
+					return false;
+				}
+
+				pxr::UsdPrim SkelAnimationPrim = UsdAnimQuery.GetPrim();
+				if (!SkelAnimationPrim)
+				{
+					return false;
+				}
+				FString SkelAnimationPrimPath = UsdToUnreal::ConvertPath(SkelAnimationPrim.GetPath());
+
+				std::vector<double> JointTimeSamples;
+				std::vector<double> BlendShapeTimeSamples;
+				if ((!UsdAnimQuery.GetJointTransformTimeSamples(&JointTimeSamples) || JointTimeSamples.size() == 0)
+					&& (NewBlendShapes.Num() == 0
+						|| (!UsdAnimQuery.GetBlendShapeWeightTimeSamples(&BlendShapeTimeSamples) || BlendShapeTimeSamples.size() == 0)))
+				{
+					return false;
+				}
+
+				pxr::UsdPrim RootMotionPrim;
+				switch(Context->RootMotionHandling)
+				{
+					case EUsdRootMotionHandling::UseMotionFromSkelRoot:
+					{
+						if (pxr::UsdPrim Root = UsdUtils::GetClosestParentSkelRoot(Skeleton.GetPrim()))
+						{
+							RootMotionPrim = Root.GetPrim();
+						}
+						break;
+					}
+					case EUsdRootMotionHandling::UseMotionFromSkeleton:
+					{
+						RootMotionPrim = Skeleton.GetPrim();
+						break;
+					}
+					default:
+					case EUsdRootMotionHandling::NoAdditionalRootMotion:
+					{
+						break;
+					}
+				}
+
+				FSHAHash Hash = UsdSkelSkeletonTranslatorImpl::ComputeSHAHash(SkeletonQuery.Get(), RootMotionPrim, SkeletalMeshHashString);
+				FString HashString = Hash.ToString();
+				UAnimSequence* AnimSequence = Cast<UAnimSequence>(Context->AssetCache->GetCachedAsset(HashString));
+
+				if (!AnimSequence || AnimSequence->GetSkeleton() != SkeletalMesh->GetSkeleton())
+				{
+					FScopedUnrealAllocs UEAllocs;
+
+					FName AnimSequenceName = MakeUniqueObjectName(
+						GetTransientPackage(),
+						UAnimSequence::StaticClass(),
+						*IUsdClassesModule::SanitizeObjectName(UsdToUnreal::ConvertToken(SkelAnimationPrim.GetName()))
+					);
+
+					// The UAnimSequence can't be created with the RF_Transactional flag, or else it will be serialized without
+					// Bone/CurveCompressionSettings. Undoing that transaction would call UAnimSequence::Serialize with nullptr values for
+					// both, which crashes. Besides, this particular asset type is only ever created when we import to content folder assets
+					// (so never for realtime), and in that case we don't need it to be transactional anyway
+					AnimSequence = NewObject<UAnimSequence>(
+						GetTransientPackage(),
+						AnimSequenceName,
+						(Context->ObjectFlags & ~EObjectFlags::RF_Transactional) | EObjectFlags::RF_Transient
+					);
+					AnimSequence->SetSkeleton(SkeletalMesh->GetSkeleton());
+
+					// This is read back in the USDImporter, so that if we ever import this AnimSequence we will always also import the
+					// SkeletalMesh for it
+					AnimSequence->SetPreviewMesh(SkeletalMesh);
+
+					TUsdStore<pxr::VtArray<pxr::UsdSkelSkinningQuery>> SkinningTargets = SkeletonBinding.Get().GetSkinningTargets();
+					float LayerStartOffsetSeconds = 0.0f;
+					const bool bSuccess = UsdToUnreal::ConvertSkelAnim(
+						SkeletonQuery.Get(),
+						&SkinningTargets.Get(),
+						&NewBlendShapes,
+						Context->bAllowInterpretingLODs,
+						RootMotionPrim,
+						AnimSequence,
+						&LayerStartOffsetSeconds
+					);
+
+					if (bSuccess
+						&& (AnimSequence->GetDataModel()->GetNumBoneTracks() != 0 || AnimSequence->GetDataModel()->GetNumberOfFloatCurves() != 0))
+					{
+						UUsdAnimSequenceAssetUserData* UserData = NewObject<UUsdAnimSequenceAssetUserData>(
+							AnimSequence,
+							TEXT("USDAssetUserData")
+						);
+						UserData->PrimPaths = {SkelAnimationPrimPath};
+						UserData->LayerStartOffsetSeconds = LayerStartOffsetSeconds;
+
+						AnimSequence->AddAssetUserData(UserData);
+
+						Context->AssetCache->CacheAsset(HashString, AnimSequence);
+					}
+					else
+					{
+						AnimSequence->MarkAsGarbage();
+					}
+				}
+
+				if (AnimSequence && Context->InfoCache)
+				{
+					Context->InfoCache->LinkAssetToPrim(SkeletonPrimPath, AnimSequence);
+				}
+
+				return true;
+			}
+		);
+	}
+
+#endif	  // WITH_EDITOR
+}	 // namespace UsdSkelSkeletonTranslatorImpl
+
+void FUsdSkelSkeletonTranslator::CreateAssets()
+{
+	TRACE_CPUPROFILER_EVENT_SCOPE(FUsdSkelSkeletonTranslator::CreateAssets);
+
+	// Don't bother generating assets if we're going to just draw some bounds for this prim instead
+	EUsdDrawMode DrawMode = UsdUtils::GetAppliedDrawMode(GetPrim());
+	if (DrawMode != EUsdDrawMode::Default)
+	{
+		CreateAlternativeDrawModeAssets(DrawMode);
+		return;
+	}
+
+#if WITH_EDITOR
+	Context->TranslatorTasks.Add(MakeShared<UsdSkelSkeletonTranslatorImpl::FSkelSkeletonCreateAssetsTaskChain>(Context, PrimPath));
+#endif	  // WITH_EDITOR
+}
+
+USceneComponent* FUsdSkelSkeletonTranslator::CreateComponents()
+{
+	TRACE_CPUPROFILER_EVENT_SCOPE(FUsdSkelSkeletonTranslator::CreateComponents);
+
+	USceneComponent* SceneComponent = FUsdGeomXformableTranslator::CreateComponents();
+
+	// TODO: DrawMode stuff here too?
+
+#if WITH_EDITOR
+	UE::FUsdPrim SkeletonPrim = GetPrim();
+
+	pxr::UsdPrim ClosestParentSkelRoot = UsdUtils::GetClosestParentSkelRoot(SkeletonPrim);
+	if (!ClosestParentSkelRoot)
+	{
+		UE_LOG(
+			LogUsd,
+			Warning,
+			TEXT("Ignoring skeleton '%s' when creating components as it is not contained within a SkelRoot prim scope"),
+			*PrimPath.GetString()
+		);
+		return SceneComponent;
+	}
+
+	// Check if the prim has the GroomBinding schema and setup the component and assets necessary to bind the groom to the SkeletalMesh
+	if (Context->bAllowParsingGroomAssets &&
+		Context->AssetCache &&
+		Context->InfoCache)
+	{
+		UE::FUsdPrim PrimWithSchema;
+		if (UsdUtils::PrimHasSchema(SkeletonPrim, UnrealIdentifiers::GroomBindingAPI))
+		{
+			PrimWithSchema = SkeletonPrim;
+		}
+		// Temporarily also allow this to be on the closest skelroot for backwards compatibility
+		else if (UsdUtils::PrimHasSchema(ClosestParentSkelRoot, UnrealIdentifiers::GroomBindingAPI))
+		{
+			PrimWithSchema = ClosestParentSkelRoot;
+		}
+
+		if (PrimWithSchema)
+		{
+			UsdGroomTranslatorUtils::CreateGroomBindingAsset(PrimWithSchema, *Context->AssetCache, *Context->InfoCache, Context->ObjectFlags);
+
+			// For the groom binding to work, the GroomComponent must be a child of the SceneComponent
+			// so the Context ParentComponent is set to the SceneComponent temporarily
+			TGuardValue<USceneComponent*> ParentComponentGuard{Context->ParentComponent, SceneComponent};
+			const bool bNeedsActor = false;
+			UGroomComponent* GroomComponent = Cast<UGroomComponent>(
+				CreateComponentsEx(TSubclassOf<USceneComponent>(UGroomComponent::StaticClass()), bNeedsActor)
+			);
+			if (GroomComponent)
+			{
+				UpdateComponents(SceneComponent);
+			}
+		}
+	}
+#endif	  // WITH_EDITOR
+
+	return SceneComponent;
+}
+
+void FUsdSkelSkeletonTranslator::UpdateComponents(USceneComponent* SceneComponent)
+{
+	USkeletalMeshComponent* SkeletalMeshComponent = Cast<USkeletalMeshComponent>(SceneComponent);
+	if (!SkeletalMeshComponent || !Context->InfoCache)
+	{
+		return;
+	}
+
+	UE::FUsdPrim SkeletonPrim = GetPrim();
+	UE::FUsdPrim ClosestSkelRootPrim = UsdUtils::GetClosestParentSkelRoot(SkeletonPrim);
+	if (!ClosestSkelRootPrim)
+	{
+		UE_LOG(
+			LogUsd,
+			Warning,
+			TEXT("Ignoring skeleton '%s' when updating components as it is not contained within a SkelRoot prim scope"),
+			*PrimPath.GetString()
+		);
+		return;
+	}
+
+	UE::FUsdPrim PrimWithLiveLinkSchema;
+	if (UsdUtils::PrimHasSchema(SkeletonPrim, UnrealIdentifiers::LiveLinkAPI))
+	{
+		PrimWithLiveLinkSchema = SkeletonPrim;
+	}
+	else if (UsdUtils::PrimHasSchema(ClosestSkelRootPrim, UnrealIdentifiers::LiveLinkAPI))
+	{
+		PrimWithLiveLinkSchema = ClosestSkelRootPrim;
+
+		// Commenting the usual deprecation macro so that we can find this with search and replace later
+		// UE_DEPRECATED(5.4, "schemas")
+		UE_LOG(
+			LogUsd,
+			Warning,
+			TEXT("Placing integration schemas (Live Link, Control Rig, Groom Binding) on SkelRoot prims (like '%s') has been deprecated on "
+				 "version 5.4 and will be unsupported in a future release. Please place your integration schemas directly on the Skeleton "
+				 "prims instead!"),
+			*ClosestSkelRootPrim.GetPrimPath().GetString()
+		);
+	}
+	const bool bPrimHasLiveLinkSchema = static_cast<bool>(PrimWithLiveLinkSchema);
+
+	bool bPrimHasLiveLinkEnabled = bPrimHasLiveLinkSchema;
+	if (bPrimHasLiveLinkSchema)
+	{
+		FScopedUsdAllocs Allocs;
+
+		if (pxr::UsdAttribute Attr = pxr::UsdPrim{PrimWithLiveLinkSchema}.GetAttribute(UnrealIdentifiers::UnrealLiveLinkEnabled))
+		{
+			bPrimHasLiveLinkEnabled = Attr.Get(&bPrimHasLiveLinkEnabled) && bPrimHasLiveLinkEnabled;
+		}
+	}
+
+	SkeletalMeshComponent->Modify();
+
+	SkeletalMeshComponent->SetAnimationMode(
+		bPrimHasLiveLinkEnabled			 ? EAnimationMode::AnimationBlueprint
+		: Context->bSequencerIsAnimating ? EAnimationMode::AnimationCustomMode
+										 : EAnimationMode::AnimationSingleNode
+	);
+
+	UE::FUsdPrim SkelAnimPrim;
+	TUsdStore<pxr::UsdSkelBinding> SkeletonBindingForPrim;
+	{
+		FScopedUsdAllocs Allocs;
+
+		pxr::UsdSkelSkeletonQuery SkelQuery;
+		const bool bSuccess = UsdUtils::GetSkelQueries(
+			pxr::UsdSkelRoot{ClosestSkelRootPrim},
+			pxr::UsdSkelSkeleton{SkeletonPrim},
+			SkeletonBindingForPrim.Get(),
+			SkelQuery
+		);
+
+		pxr::UsdSkelAnimQuery AnimQuery = SkelQuery.GetAnimQuery();
+		if (bSuccess && AnimQuery)
+		{
+			SkelAnimPrim = AnimQuery.GetPrim();
+		}
+	}
+
+	if (SkelAnimPrim)
+	{
+		UAnimSequence* TargetAnimSequence = Context->InfoCache->GetSingleAssetForPrim<UAnimSequence>(PrimPath);
+		if (TargetAnimSequence != SkeletalMeshComponent->AnimationData.AnimToPlay)
+		{
+			SkeletalMeshComponent->AnimationData.AnimToPlay = TargetAnimSequence;
+			SkeletalMeshComponent->AnimationData.bSavedLooping = false;
+			SkeletalMeshComponent->AnimationData.bSavedPlaying = false;
+			SkeletalMeshComponent->SetAnimation(TargetAnimSequence);
+		}
+	}
+
+	Super::UpdateComponents(SceneComponent);
+
+	// We always want this, but need to be registered for this to work (Super::UpdateComponents should register us)
+	const bool bNewUpdateState = true;
+	SkeletalMeshComponent->SetUpdateAnimationInEditor(bNewUpdateState);
+
+#if WITH_EDITOR
+	// Re-set the skeletal mesh if we created a new one (maybe the hash changed, a skinned UsdGeomMesh was hidden, etc.)
+	USkeletalMesh* TargetSkeletalMesh = Context->InfoCache->GetSingleAssetForPrim<USkeletalMesh>(PrimPath);
+	if (SkeletalMeshComponent->GetSkeletalMeshAsset() != TargetSkeletalMesh)
+	{
+		SkeletalMeshComponent->SetSkeletalMesh(TargetSkeletalMesh);
+
+		// Handle material overrides
+		if (TargetSkeletalMesh)
+		{
+			TArray<UMaterialInterface*> ExistingAssignments;
+			for (FSkeletalMaterial& SkeletalMaterial : TargetSkeletalMesh->GetMaterials())
+			{
+				ExistingAssignments.Add(SkeletalMaterial.MaterialInterface);
+			}
+
+			UsdSkelSkeletonTranslatorImpl::SetMaterialOverrides(
+				SkeletonBindingForPrim.Get(),
+				ExistingAssignments,
+				*SkeletalMeshComponent,
+				*Context->AssetCache.Get(),
+				*Context->InfoCache.Get(),
+				Context->Time,
+				Context->ObjectFlags,
+				Context->bAllowInterpretingLODs,
+				Context->RenderContext,
+				Context->MaterialPurpose,
+				Context->PurposesToLoad
+			);
+		}
+	}
+
+	if (bPrimHasLiveLinkSchema)
+	{
+		TSharedPtr<FUsdSchemaTranslationContext> ContextPtr = Context;
+		TStrongObjectPtr<USkeletalMeshComponent> PinnedSkelMeshComponent{SkeletalMeshComponent};
+
+		if (ContextPtr && ContextPtr->bIsImporting)
+		{
+			UsdSkelSkeletonTranslatorImpl::UpdateLiveLinkProperties(*ContextPtr.Get(), PinnedSkelMeshComponent.Get(), SkeletonPrim);
+		}
+		else
+		{
+			// HACK: This is a temporary work-around for a GIL deadlock. See UE-156489 and UE-156370 for more details, but
+			// the long-story short of it is that at this point we may have a callstack that originates from Python,
+			// triggers a USD stage notice and causes C++ code to run as our stage actor is listening to them. If we
+			// cause GC to run right now (which the DuplicateObject inside UpdateLiveLinkProperties will) we may cause
+			// a deadlock, as the game thread still holds the GIL and a background reference collector thread may want
+			// to acquire it too. What this does is run this part of the UpdateComponents on tick, that has a callstack
+			// that doesn't originate from Python, and so doesn't have the GIL locked
+			FTSTicker::GetCoreTicker().AddTicker(FTickerDelegate::CreateLambda(
+				[ContextPtr, PinnedSkelMeshComponent, SkeletonPrim](float Time)
+				{
+					if (ContextPtr)
+					{
+						UsdSkelSkeletonTranslatorImpl::UpdateLiveLinkProperties(*ContextPtr.Get(), PinnedSkelMeshComponent.Get(), SkeletonPrim);
+					}
+
+					// Returning false means this is a one-off, and won't repeat
+					return false;
+				}
+			));
+		}
+	}
+
+	// Update the animation state
+	// Don't try animating ourselves if the sequencer is animating as it will just overwrite the animation state on next
+	// tick anyway, and all this would do is lead to flickering and other issues
+	if (!Context->bSequencerIsAnimating && SkeletalMeshComponent->GetSkeletalMeshAsset() && !bPrimHasLiveLinkEnabled)
+	{
+		if (UAnimSequence* AnimSequence = Cast<UAnimSequence>(SkeletalMeshComponent->AnimationData.AnimToPlay.Get()))
+		{
+			UE::FSdfLayerOffset CombinedOffset;
+			if (SkelAnimPrim)
+			{
+				CombinedOffset = UsdUtils::GetPrimToStageOffset(SkelAnimPrim);
+			}
+
+			double LayerStartOffsetSeconds = 0.0f;
+			if (UUsdAnimSequenceAssetUserData* UserData = AnimSequence->GetAssetUserData<UUsdAnimSequenceAssetUserData>())
+			{
+				LayerStartOffsetSeconds = UserData->LayerStartOffsetSeconds;
+			}
+
+			// Always change the mode here because the sequencer will change it back to AnimationCustomMode when animating
+			SkeletalMeshComponent->SetAnimationMode(EAnimationMode::AnimationSingleNode);
+
+			// Part of the CombinedOffset will be due to a framerate difference. We don't care about that part here though, so remove it
+			const double TimeCodesPerSecondDifference = Context->Stage.GetTimeCodesPerSecond() / AnimSequence->ImportFileFramerate;
+			CombinedOffset.Scale /= TimeCodesPerSecondDifference;
+
+			// Always use the sequence's framerate here because we need to sample the UAnimSequence with in seconds, and that
+			// asset may have been created when the stage had a different framesPerSecond (and was reused by the assets cache)
+			// Use the import framerate here because we will need to change the sampling framerate of the sequence in order to get it
+			// to match the target duration in seconds and the number of source frames.
+			const double LayerTimeCode = ((Context->Time - CombinedOffset.Offset) / CombinedOffset.Scale);
+			const double AnimSequenceTime = LayerTimeCode / AnimSequence->ImportFileFramerate;
+			SkeletalMeshComponent->SetPosition(AnimSequenceTime - LayerStartOffsetSeconds);
+
+			SkeletalMeshComponent->TickAnimation(0.f, false);
+			SkeletalMeshComponent->RefreshBoneTransforms();
+			SkeletalMeshComponent->RefreshFollowerComponents();
+			SkeletalMeshComponent->UpdateComponentToWorld();
+			SkeletalMeshComponent->FinalizeBoneTransform();
+			SkeletalMeshComponent->MarkRenderTransformDirty();
+			SkeletalMeshComponent->MarkRenderDynamicDataDirty();
+		}
+	}
+
+	// If the prim has a GroomBinding schema, apply the target groom to its associated GroomComponent
+	if (Context->InfoCache)
+	{
+		if (UsdUtils::PrimHasSchema(SkeletonPrim, UnrealIdentifiers::GroomBindingAPI))
+		{
+			UsdGroomTranslatorUtils::SetGroomFromPrim(SkeletonPrim, *Context->InfoCache, SceneComponent);
+		}
+		else if (UsdUtils::PrimHasSchema(ClosestSkelRootPrim, UnrealIdentifiers::GroomBindingAPI))
+		{
+			// Commenting the usual deprecation macro so that we can find this with search and replace later
+			// UE_DEPRECATED(5.4, "schemas")
+			UE_LOG(
+				LogUsd,
+				Warning,
+				TEXT("Placing integration schemas (Live Link, Control Rig, Groom Binding) on SkelRoot prims (like '%s') has been deprecated on "
+					 "version 5.4 and will be unsupported in a future release. Please place your integration schemas directly on the Skeleton "
+					 "prims instead!"),
+				*ClosestSkelRootPrim.GetPrimPath().GetString()
+			);
+			UsdGroomTranslatorUtils::SetGroomFromPrim(ClosestSkelRootPrim, *Context->InfoCache, SceneComponent);
+		}
+	}
+#endif	  // WITH_EDITOR
+}
+
+bool FUsdSkelSkeletonTranslator::CollapsesChildren(ECollapsingType CollapsingType) const
+{
+	// There's no real reason to collapse: We're going to find all of our skinned meshes directly wherever they are, and other translators
+	// are going to skip skinned stuff anyway. If the user happens to put an e.g. Light or Camera prim inside of the Skeleton for
+	// some reason there's no reason we can't just handle it as normal
+	return false;
+}
+
+bool FUsdSkelSkeletonTranslator::CanBeCollapsed(ECollapsingType CollapsingType) const
+{
+	// Definetely cannot be collapsed into another StaticMesh. We could maybe be collapsed into another SkeletalMesh if we merged the
+	// skeletons together, but that seems like a lot of complicated work and likely undesirable anyway: Users likely want to see their
+	// skeleton prims generate dedicated USkeletons, every time
+	return false;
+}
+
+TSet<UE::FSdfPath> FUsdSkelSkeletonTranslator::CollectAuxiliaryPrims() const
+{
+	if (!Context->bIsBuildingInfoCache)
+	{
+		return Context->InfoCache->GetAuxiliaryPrims(PrimPath);
+	}
+
+	TSet<UE::FSdfPath> Result;
+	{
+		FScopedUsdAllocs UsdAllocs;
+
+		pxr::UsdPrim SkeletonPrim = GetPrim();
+		pxr::UsdPrim ClosestParentSkelRoot = UsdUtils::GetClosestParentSkelRoot(SkeletonPrim);
+
+		pxr::TfType ImageableSchema = pxr::TfType::Find<pxr::UsdGeomImageable>();
+		pxr::TfType MeshSchema = pxr::TfType::Find<pxr::UsdGeomMesh>();
+
+		pxr::UsdPrimRange PrimRange{ClosestParentSkelRoot, pxr::UsdTraverseInstanceProxies()};
+		for (pxr::UsdPrimRange::iterator PrimRangeIt = PrimRange.begin(); PrimRangeIt != PrimRange.end(); ++PrimRangeIt)
+		{
+			if (PrimRangeIt->IsA(MeshSchema))
+			{
+				Result.Add(UE::FSdfPath{PrimRangeIt->GetPrimPath()});
+
+				if (pxr::UsdSkelBindingAPI SkelBindingAPI{*PrimRangeIt})
+				{
+					// Collect blend shapes, which don't have to be within the Mesh or SkelRoot prim at all
+					if (pxr::UsdSkelBlendShapeQuery BlendShapeQuery{SkelBindingAPI})
+					{
+						for (uint32 BlendShapeIndex = 0; BlendShapeIndex < BlendShapeQuery.GetNumBlendShapes(); ++BlendShapeIndex)
+						{
+							if (pxr::UsdSkelBlendShape BlendShape = BlendShapeQuery.GetBlendShape(BlendShapeIndex))
+							{
+								Result.Add(UE::FSdfPath{BlendShape.GetPrim().GetPrimPath()});
+							}
+						}
+					}
+				}
+			}
+			// All meshes, xforms, skeleton are imageables. Registering any imageable is also a good idea because their
+			// visibility could affect child Mesh prims, and so the combined skeletal mesh.
+			else if (PrimRangeIt->IsA(ImageableSchema))
+			{
+				Result.Add(UE::FSdfPath{PrimRangeIt->GetPrimPath()});
+			}
+		}
+
+		// Collect the bound SkelAnimation, that doesn't have to be within the actual SkelRoot
+		if (ClosestParentSkelRoot)
+		{
+			pxr::UsdSkelBinding SkelBinding;
+			pxr::UsdSkelSkeletonQuery SkelQuery;
+			const bool bSuccess = UsdUtils::GetSkelQueries(
+				pxr::UsdSkelRoot{ClosestParentSkelRoot},
+				pxr::UsdSkelSkeleton{SkeletonPrim},
+				SkelBinding,
+				SkelQuery
+			);
+
+			pxr::UsdSkelAnimQuery AnimQuery = SkelQuery.GetAnimQuery();
+			if (bSuccess && AnimQuery)
+			{
+				if (pxr::UsdPrim SkelAnimationPrim = AnimQuery.GetPrim())
+				{
+					Result.Add(UE::FSdfPath{SkelAnimationPrim.GetPrimPath()});
+				}
+			}
+		}
+	}
+	return Result;
+}
+
+#undef LOCTEXT_NAMESPACE
+
+#endif	  // #if USE_USD_SDK
diff --git a/Engine/Plugins/Importers/USDImporter/Source/USDSchemas/Public/USDGeomMeshTranslator.h b/Engine/Plugins/Importers/USDImporter/Source/USDSchemas/Public/USDGeomMeshTranslator.h
--- a/Engine/Plugins/Importers/USDImporter/Source/USDSchemas/Public/USDGeomMeshTranslator.h
+++ b/Engine/Plugins/Importers/USDImporter/Source/USDSchemas/Public/USDGeomMeshTranslator.h
@@ -98,6 +98,17 @@ protected:
 	// Because of how the GroomTranslator derives the GeometryCacheTranslator that derives this, we may end up in situations where this
 	// translator is not invoked on Mesh prims directly, and so we just yield back to the Xformable translator
 	bool IsMeshPrim() const;
+
+	// Returns true if the Mesh prim we're meant to translate is inside a SkelRoot and has the SkelBindingAPI.
+	// We handle skeletal data now via an UsdSkelSkeletonTranslator for each Skeleton prim we encounter (the SkelRootTranslator has been
+	// deprecated). The SkeletonTranslator can find the skinned meshes for its skeleton on its own, regardless of
+	// where they are, but the fact that we don't have a SkelRootTranslator collapsing anymore means we'll get FUsdGeomMeshTranslator
+	// invokations on skinned mesh prims that were already handled by a SkeletonTranslator somewhere else. This check lets us avoid
+	// those meshes.
+	// The "bCheckForComponent" parameter exists because we may want to always skip spawning assets for a skinnable prim (as a
+	// SkeletonTranslator may have already handled the mesh data), but still want to spawn a component in some cases
+	// (e.g. in the edge case that the skinned mesh prim has a transform and non-skinned children that inherit it)
+	bool ShouldSkipSkinnablePrim(bool bCheckForComponent = false) const;
 };
 
 #endif // #if USE_USD_SDK
diff --git a/Engine/Plugins/Importers/USDImporter/Source/USDSchemas/Public/USDSchemaTranslator.h b/Engine/Plugins/Importers/USDImporter/Source/USDSchemas/Public/USDSchemaTranslator.h
--- a/Engine/Plugins/Importers/USDImporter/Source/USDSchemas/Public/USDSchemaTranslator.h
+++ b/Engine/Plugins/Importers/USDImporter/Source/USDSchemas/Public/USDSchemaTranslator.h
@@ -237,7 +237,7 @@ struct USDSCHEMAS_API FUsdSchemaTranslationContext : public TSharedFromThis< FUs
 	 */
 	bool bAllowInterpretingLODs = true;
 
-	/** If true, we will also try creating UAnimSequence skeletal animation assets when parsing SkelRoot prims */
+	/** If true, we will also try creating UAnimSequence skeletal animation assets when parsing Skeleton prims */
 	bool bAllowParsingSkeletalAnimations = true;
 
 	/** If true, means we will try generating GroomAssets, GroomCaches and GroomBindings */
diff --git a/Engine/Plugins/Importers/USDImporter/Source/USDSchemas/Public/USDSkelRootTranslator.h b/Engine/Plugins/Importers/USDImporter/Source/USDSchemas/Public/USDSkelRootTranslator.h
--- a/Engine/Plugins/Importers/USDImporter/Source/USDSchemas/Public/USDSkelRootTranslator.h
+++ b/Engine/Plugins/Importers/USDImporter/Source/USDSchemas/Public/USDSkelRootTranslator.h
@@ -6,17 +6,7 @@
 
 #if USE_USD_SDK
 
-#include "USDIncludesStart.h"
-	#include "pxr/pxr.h"
-#include "USDIncludesEnd.h"
-
-PXR_NAMESPACE_OPEN_SCOPE
-	class UsdSkelRoot;
-PXR_NAMESPACE_CLOSE_SCOPE
-
-struct FUsdSchemaTranslationContext;
-
-class USDSCHEMAS_API FUsdSkelRootTranslator : public FUsdGeomXformableTranslator
+class USDSCHEMAS_API UE_DEPRECATED(5.4, "Use the UsdSkelSkeletonTranslator for skeletal data") FUsdSkelRootTranslator : public FUsdGeomXformableTranslator
 {
 	using Super = FUsdGeomXformableTranslator;
 
diff --git a/Engine/Plugins/Importers/USDImporter/Source/USDSchemas/Public/USDSkelSkeletonTranslator.h b/Engine/Plugins/Importers/USDImporter/Source/USDSchemas/Public/USDSkelSkeletonTranslator.h
new file mode 100644
--- /dev/null
+++ b/Engine/Plugins/Importers/USDImporter/Source/USDSchemas/Public/USDSkelSkeletonTranslator.h
@@ -0,0 +1,26 @@
+// Copyright Epic Games, Inc. All Rights Reserved.
+
+#pragma once
+
+#include "USDGeomXformableTranslator.h"
+
+#if USE_USD_SDK
+
+class USDSCHEMAS_API FUsdSkelSkeletonTranslator : public FUsdGeomXformableTranslator
+{
+	using Super = FUsdGeomXformableTranslator;
+
+public:
+	using FUsdGeomXformableTranslator::FUsdGeomXformableTranslator;
+
+	virtual void CreateAssets() override;
+	virtual USceneComponent* CreateComponents() override;
+	virtual void UpdateComponents( USceneComponent* SceneComponent ) override;
+
+	virtual bool CollapsesChildren(ECollapsingType CollapsingType) const override;
+	virtual bool CanBeCollapsed(ECollapsingType CollapsingType) const override;
+
+	virtual TSet<UE::FSdfPath> CollectAuxiliaryPrims() const override;
+};
+
+#endif //#if USE_USD_SDK
diff --git a/Engine/Plugins/Importers/USDImporter/Source/USDStage/Private/USDLevelSequenceHelper.cpp b/Engine/Plugins/Importers/USDImporter/Source/USDStage/Private/USDLevelSequenceHelper.cpp
--- a/Engine/Plugins/Importers/USDImporter/Source/USDStage/Private/USDLevelSequenceHelper.cpp
+++ b/Engine/Plugins/Importers/USDImporter/Source/USDStage/Private/USDLevelSequenceHelper.cpp
@@ -16,6 +16,7 @@
 #include "USDPrimConversion.h"
 #include "USDPrimTwin.h"
 #include "USDProjectSettings.h"
+#include "USDSkeletalDataConversion.h"
 #include "USDStageActor.h"
 #include "USDTypesConversion.h"
 #include "USDValueConversion.h"
@@ -542,6 +543,8 @@ public:
 	void CreateLocalLayersSequences();
 	void BindToUsdStageActor(AUsdStageActor* InStageActor);
 	void UnbindFromUsdStageActor();
+	EUsdRootMotionHandling GetRootMotionHandling() const;
+	void SetRootMotionHandling(EUsdRootMotionHandling NewValue);
 	void OnStageActorRenamed();
 
 	ULevelSequence* GetMainLevelSequence() const { return MainLevelSequence; }
@@ -695,6 +698,7 @@ private:
 	TWeakObjectPtr<AUsdStageActor> StageActor = nullptr;
 	TSharedPtr<FUsdInfoCache> InfoCache = nullptr;  // We keep a pointer to this directly because we may be called via the USDStageImporter directly, when we don't have an available actor
 	TSharedPtr<UE::FUsdGeomBBoxCache> BBoxCache = nullptr; // Same as for the info cache
+	EUsdRootMotionHandling RootMotionHandling = EUsdRootMotionHandling::NoAdditionalRootMotion;
 	FGuid StageActorBinding;
 
 	// Only when this is zero we write LevelSequence object (tracks, moviescene, sections, etc.) transactions back to the USD stage
@@ -900,6 +904,7 @@ void FUsdLevelSequenceHelperImpl::BindToUsdStageActor(AUsdStageActor* InStageAct
 	StageActor = InStageActor;
 	SetInfoCache(InStageActor ? InStageActor->GetInfoCache() : nullptr);
 	SetBBoxCache(InStageActor ? InStageActor->GetBBoxCache() : nullptr);
+	SetRootMotionHandling(InStageActor ? InStageActor->RootMotionHandling : EUsdRootMotionHandling::NoAdditionalRootMotion);
 
 	if (!StageActor.IsValid() || !MainLevelSequence || !MainLevelSequence->GetMovieScene())
 	{
@@ -947,10 +952,21 @@ void FUsdLevelSequenceHelperImpl::UnbindFromUsdStageActor()
 	}
 
 	SetInfoCache(nullptr);
+	SetRootMotionHandling(EUsdRootMotionHandling::NoAdditionalRootMotion);
 
 	OnStageEditTargetChangedHandle.Reset();
 }
 
+EUsdRootMotionHandling FUsdLevelSequenceHelperImpl::GetRootMotionHandling() const
+{
+	return RootMotionHandling;
+}
+
+void FUsdLevelSequenceHelperImpl::SetRootMotionHandling(EUsdRootMotionHandling NewValue)
+{
+	RootMotionHandling = NewValue;
+}
+
 void FUsdLevelSequenceHelperImpl::OnStageActorRenamed()
 {
 	AUsdStageActor* StageActorPtr = StageActor.Get();
@@ -1485,17 +1501,30 @@ void FUsdLevelSequenceHelperImpl::AddCommonTracks(const UUsdPrimTwin& PrimTwin,
 		}
 	}
 
-	// Check whether we should ignore the prim's local transform or not. We only do this for SkelRoots, in case we
-	// already appended their transform animations to the root bone as additional root motion
+	// Check whether we should ignore the prim's local transform or not
 	bool bIgnorePrimLocalTransform = false;
-	if (AUsdStageActor* StageActorPtr = StageActor.Get())
+	switch (GetRootMotionHandling())
 	{
-		if (StageActorPtr->RootMotionHandling == EUsdRootMotionHandling::UseMotionFromSkelRoot)
+		default:
+		case EUsdRootMotionHandling::NoAdditionalRootMotion:
+		{
+			break;
+		}
+		case EUsdRootMotionHandling::UseMotionFromSkelRoot:
 		{
 			if (Prim.IsA(TEXT("SkelRoot")))
 			{
 				bIgnorePrimLocalTransform = true;
 			}
+			break;
+		}
+		case EUsdRootMotionHandling::UseMotionFromSkeleton:
+		{
+			if (Prim.IsA(TEXT("Skeleton")))
+			{
+				bIgnorePrimLocalTransform = true;
+			}
+			break;
 		}
 	}
 
@@ -2049,9 +2078,11 @@ void FUsdLevelSequenceHelperImpl::AddSkeletalTracks(const UUsdPrimTwin& PrimTwin
 		return;
 	}
 
+	UE::FUsdPrim SkelRootPrim = UsdUtils::GetClosestParentSkelRoot(Prim);
+
 	// We'll place the skeletal animation track wherever the SkelAnimation prim is defined (not necessarily the
 	// same layer as the skel root)
-	UE::FUsdPrim SkelAnimationPrim = UsdUtils::FindFirstAnimationSource(Prim);
+	UE::FUsdPrim SkelAnimationPrim = UsdUtils::FindAnimationSource(SkelRootPrim, Prim);
 	if (!SkelAnimationPrim)
 	{
 		return;
@@ -2297,7 +2328,7 @@ void FUsdLevelSequenceHelperImpl::AddPrim(UUsdPrimTwin& PrimTwin, bool bForceVis
 	{
 		AddLightTracks(PrimTwin, UsdPrim);
 	}
-	else if (UsdPrim.IsA(TEXT("SkelRoot")))
+	else if (UsdPrim.IsA(TEXT("Skeleton")))
 	{
 		AddSkeletalTracks(PrimTwin, UsdPrim);
 	}
@@ -2413,15 +2444,13 @@ void FUsdLevelSequenceHelperImpl::UpdateControlRigTracks(UUsdPrimTwin& PrimTwin)
 		return;
 	}
 
-	UE::FSdfPath PrimPath(*PrimTwin.PrimPath);
-	UE::FUsdPrim UsdPrim(UsdStage.GetPrimAtPath(PrimPath));
-	if (!UsdPrim)
+	UE::FSdfPath PrimPath{*PrimTwin.PrimPath};
+	UE::FUsdPrim SkeletonPrim{UsdStage.GetPrimAtPath(PrimPath)};
+	if (!SkeletonPrim)
 	{
 		return;
 	}
 
-	ensure(UsdPrim.IsA(TEXT("SkelRoot")));
-
 	USkeletalMeshComponent* ComponentToBind = Cast<USkeletalMeshComponent>(PrimTwin.GetSceneComponent());
 	if (!ComponentToBind)
 	{
@@ -2433,11 +2462,23 @@ void FUsdLevelSequenceHelperImpl::UpdateControlRigTracks(UUsdPrimTwin& PrimTwin)
 	// repeatedly create Animation prims
 	FScopedBlockNoticeListening BlockNotices(StageActor.Get());
 
-	UE::FSdfLayer SkelAnimationLayer;
+	UE::FUsdPrim SkelRootPrim = UsdUtils::GetClosestParentSkelRoot(SkeletonPrim);
+	UE::FUsdPrim SkelAnimationPrim = UsdUtils::FindAnimationSource(SkelRootPrim, SkeletonPrim);
+
+	// Temporarily consider how our API schema can be applied to the Skeleton prim or a parent SkelRoot
+	UE::FUsdPrim PrimWithSchema;
+	if (UsdUtils::PrimHasSchema(SkeletonPrim, UnrealIdentifiers::ControlRigAPI))
+	{
+		PrimWithSchema = SkeletonPrim;
+	}
+	else if (SkelRootPrim && UsdUtils::PrimHasSchema(SkelRootPrim, UnrealIdentifiers::ControlRigAPI))
+	{
+		PrimWithSchema = SkelRootPrim;
+	}
 
 	// We'll place the skeletal animation track wherever the SkelAnimation prim is defined (not necessarily the
 	// same layer as the skel root)
-	UE::FUsdPrim SkelAnimationPrim = UsdUtils::FindFirstAnimationSource(UsdPrim);
+	UE::FSdfLayer SkelAnimationLayer;
 	if (SkelAnimationPrim)
 	{
 		SkelAnimationLayer = UsdUtils::FindLayerForPrim(SkelAnimationPrim);
@@ -2447,9 +2488,9 @@ void FUsdLevelSequenceHelperImpl::UpdateControlRigTracks(UUsdPrimTwin& PrimTwin)
 		// If this SkelRoot doesn't have any animation, lets create a new one on the current edit target
 		SkelAnimationLayer = UsdStage.GetEditTarget();
 
-		FString UniqueChildName = UsdUtils::GetValidChildName(TEXT("Animation"), UsdPrim);
+		FString UniqueChildName = UsdUtils::GetValidChildName(TEXT("Animation"), SkelRootPrim);
 		SkelAnimationPrim = UsdStage.DefinePrim(
-			UsdPrim.GetPrimPath().AppendChild(*UniqueChildName),
+			SkelRootPrim.GetPrimPath().AppendChild(*UniqueChildName),
 			TEXT("SkelAnimation")
 		);
 		if (!SkelAnimationPrim)
@@ -2457,7 +2498,9 @@ void FUsdLevelSequenceHelperImpl::UpdateControlRigTracks(UUsdPrimTwin& PrimTwin)
 			return;
 		}
 
-		UsdUtils::BindAnimationSource(UsdPrim, SkelAnimationPrim);
+		// Let's always choose to author animSource within skeletons, as it works best in setups where
+		// we have authored nested SkelRoots: The outer animSource would be inherited by the inner animSource otherwise!
+		UsdUtils::BindAnimationSource(SkeletonPrim, SkelAnimationPrim);
 	}
 	if (!SkelAnimationLayer)
 	{
@@ -2495,17 +2538,13 @@ void FUsdLevelSequenceHelperImpl::UpdateControlRigTracks(UUsdPrimTwin& PrimTwin)
 
 	const FGuid ComponentBinding = GetOrCreateComponentBinding(PrimTwin, *ComponentToBind, *SkelAnimationSequence);
 
-	// NOTE: We are fetching the first skel track we find, since we can't actually use SkelAnimationPrim.GetName() here at all!
-	// The property tracks do derive GetTrackName(), but the skeletal track doesn't, so FindTrack will never find them.
-	// This likely has no effect since we only ever spawn a single skeletal track per prim anyway, but its worth to keep in mind!
-	UMovieSceneSkeletalAnimationTrack* SkelTrack = MovieScene->FindTrack< UMovieSceneSkeletalAnimationTrack >(ComponentBinding);
 	UMovieSceneControlRigParameterTrack* ControlRigTrack = MovieScene->FindTrack< UMovieSceneControlRigParameterTrack >(ComponentBinding);
 
 	// We should be in control rig track mode but don't have any tracks yet --> Setup for Control Rig
 	if (!ControlRigTrack)
 	{
 		bool bControlRigReduceKeys = false;
-		if (UE::FUsdAttribute Attr = UsdPrim.GetAttribute(*UsdToUnreal::ConvertToken(UnrealIdentifiers::UnrealControlRigReduceKeys)))
+		if (UE::FUsdAttribute Attr = PrimWithSchema.GetAttribute(*UsdToUnreal::ConvertToken(UnrealIdentifiers::UnrealControlRigReduceKeys)))
 		{
 			UE::FVtValue Value;
 			if (Attr.Get(Value) && !Value.IsEmpty())
@@ -2518,7 +2557,7 @@ void FUsdLevelSequenceHelperImpl::UpdateControlRigTracks(UUsdPrimTwin& PrimTwin)
 		}
 
 		float ControlRigReduceTolerance = 0.001f;
-		if (UE::FUsdAttribute Attr = UsdPrim.GetAttribute(*UsdToUnreal::ConvertToken(UnrealIdentifiers::UnrealControlRigReductionTolerance)))
+		if (UE::FUsdAttribute Attr = PrimWithSchema.GetAttribute(*UsdToUnreal::ConvertToken(UnrealIdentifiers::UnrealControlRigReductionTolerance)))
 		{
 			UE::FVtValue Value;
 			if (Attr.Get(Value) && !Value.IsEmpty())
@@ -2531,7 +2570,7 @@ void FUsdLevelSequenceHelperImpl::UpdateControlRigTracks(UUsdPrimTwin& PrimTwin)
 		}
 
 		bool bIsFKControlRig = false;
-		if (UE::FUsdAttribute Attr = UsdPrim.GetAttribute(*UsdToUnreal::ConvertToken(UnrealIdentifiers::UnrealUseFKControlRig)))
+		if (UE::FUsdAttribute Attr = PrimWithSchema.GetAttribute(*UsdToUnreal::ConvertToken(UnrealIdentifiers::UnrealUseFKControlRig)))
 		{
 			UE::FVtValue Value;
 			if (Attr.Get(Value))
@@ -2551,7 +2590,7 @@ void FUsdLevelSequenceHelperImpl::UpdateControlRigTracks(UUsdPrimTwin& PrimTwin)
 		else
 		{
 			FString ControlRigBPPath;
-			if (UE::FUsdAttribute Attr = UsdPrim.GetAttribute(*UsdToUnreal::ConvertToken(UnrealIdentifiers::UnrealControlRigPath)))
+			if (UE::FUsdAttribute Attr = PrimWithSchema.GetAttribute(*UsdToUnreal::ConvertToken(UnrealIdentifiers::UnrealControlRigPath)))
 			{
 				UE::FVtValue Value;
 				if (Attr.Get(Value) && !Value.IsEmpty())
@@ -3240,7 +3279,9 @@ void FUsdLevelSequenceHelperImpl::HandleMovieSceneChange(UMovieScene& MovieScene
 					{
 						if (!MovieScene.FindTrack(UMovieSceneSkeletalAnimationTrack::StaticClass(), Guid))
 						{
-							if (UE::FUsdPrim SkelAnimationPrim = UsdUtils::FindFirstAnimationSource(UsdPrim))
+							UE::FUsdPrim SkelRootPrim = UsdUtils::GetClosestParentSkelRoot(UsdPrim);
+							UE::FUsdPrim SkelAnimationPrim = UsdUtils::FindAnimationSource(SkelRootPrim, UsdPrim);
+							if (SkelAnimationPrim)
 							{
 								if (UE::FSdfLayer SkelAnimationLayer = UsdUtils::FindLayerForPrim(SkelAnimationPrim))
 								{
@@ -3336,15 +3377,21 @@ void FUsdLevelSequenceHelperImpl::HandleControlRigSectionChange(UMovieSceneContr
 		return;
 	}
 
-	UE::FUsdPrim UsdPrim = UsdStage.GetPrimAtPath(UE::FSdfPath(*PrimTwin->PrimPath));
-	if (!UsdPrim)
+	UE::FUsdPrim SkeletonPrim = UsdStage.GetPrimAtPath(UE::FSdfPath(*PrimTwin->PrimPath));
+	if (!SkeletonPrim)
+	{
+		return;
+	}
+
+	UE::FUsdPrim SkelRootPrim = UsdUtils::GetClosestParentSkelRoot(SkeletonPrim);
+	if(!SkelRootPrim)
 	{
 		return;
 	}
 
 	// We'll place the skeletal animation track wherever the SkelAnimation prim is defined (not necessarily the
 	// same layer as the skel root)
-	UE::FUsdPrim SkelAnimationPrim = UsdUtils::FindFirstAnimationSource(UsdPrim);
+	UE::FUsdPrim SkelAnimationPrim = UsdUtils::FindAnimationSource(SkelRootPrim, SkeletonPrim);
 	if (!SkelAnimationPrim)
 	{
 		return;
@@ -3425,7 +3472,7 @@ void FUsdLevelSequenceHelperImpl::HandleControlRigSectionChange(UMovieSceneContr
 		MovieScene,
 		Player,
 		Skeleton->GetReferenceSkeleton(),
-		UsdPrim,
+		SkelRootPrim,
 		SkelAnimationPrim,
 		&BlendShapeMap
 	);
@@ -3614,7 +3661,8 @@ void FUsdLevelSequenceHelperImpl::HandleTrackChange(const UMovieSceneTrack& Trac
 					bAllSectionsMuted &= !Section->IsActive();
 				}
 
-				if (UE::FUsdPrim SkelAnimationPrim = UsdUtils::FindFirstAnimationSource(UsdPrim))
+				UE::FUsdPrim SkelRootPrim = UsdUtils::GetClosestParentSkelRoot(UsdPrim);
+				if (UE::FUsdPrim SkelAnimationPrim = UsdUtils::FindAnimationSource(SkelRootPrim, UsdPrim))
 				{
 					UE::FUsdAttribute TranslationsAttr = SkelAnimationPrim.GetAttribute(TEXT("translations"));
 					UE::FUsdAttribute RotationsAttr = SkelAnimationPrim.GetAttribute(TEXT("rotations"));
@@ -3830,6 +3878,8 @@ public:
 
 	void BindToUsdStageActor(AUsdStageActor* InStageActor) {}
 	void UnbindFromUsdStageActor() {}
+	EUsdRootMotionHandling GetRootMotionHandling() const {return EUsdRootMotionHandling::NoAdditionalRootMotion;}
+	void SetRootMotionHandling(EUsdRootMotionHandling NewValue){};
 	void OnStageActorRenamed() {};
 
 	void AddPrim(UUsdPrimTwin& PrimTwin, bool bForceVisibilityTracks, TOptional<bool> HasAnimatedBounds) {}
@@ -3938,6 +3988,24 @@ void FUsdLevelSequenceHelper::UnbindFromUsdStageActor()
 	}
 }
 
+EUsdRootMotionHandling FUsdLevelSequenceHelper::GetRootMotionHandling() const
+{
+	if (UsdSequencerImpl.IsValid())
+	{
+		return UsdSequencerImpl->GetRootMotionHandling();
+	}
+
+	return EUsdRootMotionHandling::NoAdditionalRootMotion;
+}
+
+void FUsdLevelSequenceHelper::SetRootMotionHandling(EUsdRootMotionHandling NewValue)
+{
+	if (UsdSequencerImpl.IsValid())
+	{
+		return UsdSequencerImpl->SetRootMotionHandling(NewValue);
+	}
+}
+
 void FUsdLevelSequenceHelper::AddPrim(UUsdPrimTwin & PrimTwin, bool bForceVisibilityTracks, TOptional<bool> HasAnimatedBounds)
 {
 	if (UsdSequencerImpl.IsValid())
diff --git a/Engine/Plugins/Importers/USDImporter/Source/USDStage/Private/USDStageActor.cpp b/Engine/Plugins/Importers/USDImporter/Source/USDStage/Private/USDStageActor.cpp
--- a/Engine/Plugins/Importers/USDImporter/Source/USDStage/Private/USDStageActor.cpp
+++ b/Engine/Plugins/Importers/USDImporter/Source/USDStage/Private/USDStageActor.cpp
@@ -26,7 +26,7 @@
 #include "USDProjectSettings.h"
 #include "USDSchemasModule.h"
 #include "USDSchemaTranslator.h"
-#include "USDSkelRootTranslator.h"
+#include "USDSkelSkeletonTranslator.h"
 #include "USDStageModule.h"
 #include "USDTransactor.h"
 #include "USDTypesConversion.h"
@@ -2001,9 +2001,23 @@ UUsdPrimTwin* AUsdStageActor::ExpandPrim(const UE::FUsdPrim& Prim, bool bResync,
 		}
 	}
 
-	// Update the prim animated status
+	// Check if the prim should have Sequencer tracks or not
+	bool bIsAnimated = false;
+	if (USkeletalMeshComponent* SkeletalMeshComponent = Cast<USkeletalMeshComponent>(TwinSceneComponent))
+	{
+		if (SkeletalMeshComponent->AnimationData.AnimToPlay)
+		{
+			// We know we're animated if we have an animation of course
+			bIsAnimated = true;
+		}
+	}
+	if (!bIsAnimated)
+	{
+		bIsAnimated = UsdUtils::IsAnimated(Prim);
+	}
+
+	// Create Sequencer tracks for the prim
 	bool bHasAnimatedBounds = false;
-	bool bIsAnimated = UsdUtils::IsAnimated(Prim);
 	if (bIsAnimated)
 	{
 		if (!PrimsToAnimate.Contains(UsdPrimTwin->PrimPath))
@@ -2067,9 +2081,32 @@ UUsdPrimTwin* AUsdStageActor::ExpandPrim(const UE::FUsdPrim& Prim, bool bResync,
 
 	// Setup Control Rig tracks if we need to. This must be done after adding regular skeletal animation tracks
 	// if we have any as if will properly deactivate them like the usual "Bake to Control Rig" workflow.
-	if (Prim.IsA(TEXT("SkelRoot")))
+	if (Prim.IsA(TEXT("Skeleton")))
 	{
+		UE::FUsdPrim PrimWithSchema;
 		if (UsdUtils::PrimHasSchema(Prim, UnrealIdentifiers::ControlRigAPI))
+		{
+			PrimWithSchema = Prim;
+		}
+		else if (UE::FUsdPrim ParentSkelRoot = UsdUtils::GetClosestParentSkelRoot(Prim))
+		{
+			if (UsdUtils::PrimHasSchema(ParentSkelRoot, UnrealIdentifiers::ControlRigAPI))
+			{
+				// Commenting the usual deprecation macro so that we can find this with search and replace later
+				// UE_DEPRECATED(5.4, "schemas")
+				UE_LOG(
+					LogUsd,
+					Warning,
+					TEXT("Placing integration schemas (Live Link, Control Rig, Groom Binding) on SkelRoot prims (like '%s') has been deprecated on "
+						 "version 5.4 and will be unsupported in a future release. Please place your integration schemas directly on the Skeleton "
+						 "prims instead!"),
+					*Prim.GetPrimPath().GetString()
+				);
+				PrimWithSchema = ParentSkelRoot;
+			}
+		}
+
+		if (PrimWithSchema)
 		{
 			LevelSequenceHelper.UpdateControlRigTracks(*UsdPrimTwin);
 
@@ -2536,8 +2573,6 @@ TArray<UObject*> AUsdStageActor::GetGeneratedAssets(const FString& PrimPath)
 	}
 
 	// Prefer checking the prim directly, but also check its collapsed root if it is collapsed.
-	// This because we have some exception cases like USkeleton/UAnimSequences that can be found by querying the actual
-	// Skeleton/SKelAnimation prims even though they are collapsed into the SkelRoot prim.
 	TArray<TWeakObjectPtr<UObject>> AssetsPtrs = InfoCache->GetAllAssetsForPrim(UsdPath);
 	if (AssetsPtrs.Num() == 0 && InfoCache->IsPathCollapsed(UsdPath, ECollapsingType::Assets))
 	{
@@ -2704,7 +2739,7 @@ void AUsdStageActor::OnPostPIEStarted(bool bIsSimulating)
 	}
 
 	// Setup for the very first frame when we duplicate into PIE, or else we will display skeletal mesh components on their
-	// StartTimeCode state. We have to do this here (after duplicating) as we need the calls to FUsdSkelRootTranslator::UpdateComponents
+	// StartTimeCode state. We have to do this here (after duplicating) as we need the calls to FUsdSkelSkeletonTranslator::UpdateComponents
 	// to actually animate the components, and they will only be able to do anything after they have been registered (which
 	// needs to be done by the engine when going into PIE)
 	AnimatePrims();
@@ -4435,7 +4470,7 @@ bool AUsdStageActor::HasAuthorityOverStage() const
 	return !IsTemplate();
 }
 
-void AUsdStageActor::OnSkelAnimationBaked(const FString& SkelRootPrimPath)
+void AUsdStageActor::OnSkelAnimationBaked(const FString& SkeletonPrimPath)
 {
 #if USE_USD_SDK
 	const UE::FUsdStage& CurrentStage = static_cast<const AUsdStageActor*>(this)->GetUsdStage();
@@ -4444,8 +4479,8 @@ void AUsdStageActor::OnSkelAnimationBaked(const FString& SkelRootPrimPath)
 		return;
 	}
 
-	UE::FUsdPrim SkelRootPrim = CurrentStage.GetPrimAtPath(UE::FSdfPath{*SkelRootPrimPath});
-	if (!SkelRootPrim || !SkelRootPrim.IsA(TEXT("SkelRoot")))
+	UE::FUsdPrim SkeletonPrim = CurrentStage.GetPrimAtPath(UE::FSdfPath{*SkeletonPrimPath});
+	if (!SkeletonPrim || !SkeletonPrim.IsA(TEXT("Skeleton")))
 	{
 		return;
 	}
@@ -4456,7 +4491,7 @@ void AUsdStageActor::OnSkelAnimationBaked(const FString& SkelRootPrimPath)
 		return;
 	}
 
-	UUsdPrimTwin* Twin = RootTwin->Find(SkelRootPrimPath);
+	UUsdPrimTwin* Twin = RootTwin->Find(SkeletonPrimPath);
 	if (!Twin)
 	{
 		return;
@@ -4468,14 +4503,14 @@ void AUsdStageActor::OnSkelAnimationBaked(const FString& SkelRootPrimPath)
 		return;
 	}
 
-	TSharedRef< FUsdSchemaTranslationContext > TranslationContext = FUsdStageActorImpl::CreateUsdSchemaTranslationContext(this, SkelRootPrimPath);
+	TSharedRef< FUsdSchemaTranslationContext > TranslationContext = FUsdStageActorImpl::CreateUsdSchemaTranslationContext(this, SkeletonPrimPath);
 	// The only way we could have baked a skel animation is via the sequencer, so we know its playing
 	TranslationContext->bSequencerIsAnimating = true;
 
 	IUsdSchemasModule& UsdSchemasModule = FModuleManager::Get().LoadModuleChecked< IUsdSchemasModule >(TEXT("USDSchemas"));
-	if (TSharedPtr< FUsdSchemaTranslator > SchemaTranslator = UsdSchemasModule.GetTranslatorRegistry().CreateTranslatorForSchema(TranslationContext, UE::FUsdTyped(SkelRootPrim)))
+	if (TSharedPtr< FUsdSchemaTranslator > SchemaTranslator = UsdSchemasModule.GetTranslatorRegistry().CreateTranslatorForSchema(TranslationContext, UE::FUsdTyped(SkeletonPrim)))
 	{
-		if (TSharedPtr< FUsdSkelRootTranslator > SkelRootTranslator = StaticCastSharedPtr<FUsdSkelRootTranslator>(SchemaTranslator))
+		if (TSharedPtr<FUsdSkelSkeletonTranslator> SkelRootTranslator = StaticCastSharedPtr<FUsdSkelSkeletonTranslator>(SchemaTranslator))
 		{
 			// For now we're regenerating all asset types (including skeletal meshes) but we could
 			// eventually just split off the anim sequence generation and call exclusively that from
diff --git a/Engine/Plugins/Importers/USDImporter/Source/USDStage/Public/USDLevelSequenceHelper.h b/Engine/Plugins/Importers/USDImporter/Source/USDStage/Public/USDLevelSequenceHelper.h
--- a/Engine/Plugins/Importers/USDImporter/Source/USDStage/Public/USDLevelSequenceHelper.h
+++ b/Engine/Plugins/Importers/USDImporter/Source/USDStage/Public/USDLevelSequenceHelper.h
@@ -13,6 +13,7 @@ class FUsdInfoCache;
 class FUsdLevelSequenceHelperImpl;
 class ULevelSequence;
 class UUsdPrimTwin;
+enum class EUsdRootMotionHandling : uint8;
 
 namespace UE
 {
@@ -54,10 +55,22 @@ public:
 	/** Resets the helper, abandoning all managed LevelSequences */
 	void Clear();
 
-	/** Creates the time track for the StageActor */
+	/** Creates the time track for the StageActor. This will also set the root handling mode to the provided actor's, if any */
 	void BindToUsdStageActor(AUsdStageActor* StageActor);
 	void UnbindFromUsdStageActor();
 
+	/**
+	 * Gets the current root motion handling mode.
+	 * We use this to prevent animation that has already been baked into AnimSequences as root joint motion from also
+	 * being parsed as LevelSequence tracks.
+	 */
+	EUsdRootMotionHandling GetRootMotionHandling() const;
+
+	/**
+	 * Sets the current root motion handling mode.
+	 */
+	void SetRootMotionHandling(EUsdRootMotionHandling NewValue);
+
 	/**
 	 * Adds the necessary tracks for a given prim to the level sequence.
 	 * If bForceVisibilityTracks is true, will add visibility tracks even if this prim
@@ -82,7 +95,7 @@ public:
 	ULevelSequence* GetMainLevelSequence() const;
 	TArray< ULevelSequence* > GetSubSequences() const;
 
-	DECLARE_EVENT_OneParam( FUsdLevelSequenceHelper, FOnSkelAnimationBaked, const FString& /*SkelRootPrimPath*/ );
+	DECLARE_EVENT_OneParam( FUsdLevelSequenceHelper, FOnSkelAnimationBaked, const FString& /*SkeletonPrimPath*/ );
 	FOnSkelAnimationBaked& GetOnSkelAnimationBaked();
 
 private:
diff --git a/Engine/Plugins/Importers/USDImporter/Source/USDStage/Public/USDStageActor.h b/Engine/Plugins/Importers/USDImporter/Source/USDStage/Public/USDStageActor.h
--- a/Engine/Plugins/Importers/USDImporter/Source/USDStage/Public/USDStageActor.h
+++ b/Engine/Plugins/Importers/USDImporter/Source/USDStage/Public/USDStageActor.h
@@ -344,7 +344,7 @@ protected:
 	void OnUsdPrimTwinDestroyed( const UUsdPrimTwin& UsdPrimTwin );
 	void OnObjectPropertyChanged( UObject* ObjectBeingModified, FPropertyChangedEvent& PropertyChangedEvent );
 	void HandlePropertyChangedEvent( FPropertyChangedEvent& PropertyChangedEvent );
-	void OnSkelAnimationBaked( const FString& SkelRootPrimPath );
+	void OnSkelAnimationBaked(const FString& SkeletonPrimPath);
 
 protected:
 	friend struct FUsdStageActorImpl;
diff --git a/Engine/Plugins/Importers/USDImporter/Source/USDStageEditor/Private/Widgets/SUSDIntegrationsPanel.cpp b/Engine/Plugins/Importers/USDImporter/Source/USDStageEditor/Private/Widgets/SUSDIntegrationsPanel.cpp
--- a/Engine/Plugins/Importers/USDImporter/Source/USDStageEditor/Private/Widgets/SUSDIntegrationsPanel.cpp
+++ b/Engine/Plugins/Importers/USDImporter/Source/USDStageEditor/Private/Widgets/SUSDIntegrationsPanel.cpp
@@ -360,9 +360,10 @@ TSharedRef< SWidget > SUsdIntegrationsPanelRow::GenerateWidgetForColumn( const F
 				SLiveLinkSubjectRepresentationPicker::FLiveLinkSourceSubjectRole Result;
 				if ( AttributeCopy && *AttributeCopy )
 				{
-					Result.Role = AttributeCopy->GetPrim().IsA(TEXT("SkelRoot"))
-						? ULiveLinkAnimationRole::StaticClass()
-						: ULiveLinkTransformRole::StaticClass();
+					const UE::FUsdPrim& Prim = AttributeCopy->GetPrim();
+					Result.Role = (Prim.IsA(TEXT("SkelRoot")) || Prim.IsA(TEXT("Skeleton")))
+									  ? ULiveLinkAnimationRole::StaticClass()
+									  : ULiveLinkTransformRole::StaticClass();
 
 					UE::FVtValue Value;
 					if ( AttributeCopy->Get( Value ) )
diff --git a/Engine/Plugins/Importers/USDImporter/Source/USDStageEditorViewModels/Private/USDIntegrationsViewModel.cpp b/Engine/Plugins/Importers/USDImporter/Source/USDStageEditorViewModels/Private/USDIntegrationsViewModel.cpp
--- a/Engine/Plugins/Importers/USDImporter/Source/USDStageEditorViewModels/Private/USDIntegrationsViewModel.cpp
+++ b/Engine/Plugins/Importers/USDImporter/Source/USDStageEditorViewModels/Private/USDIntegrationsViewModel.cpp
@@ -5,12 +5,11 @@
 #include "USDIntegrationUtils.h"
 #include "USDTypesConversion.h"
 
-#include "UsdWrappers/SdfAttributeSpec.h"
 #include "UsdWrappers/SdfPath.h"
 #include "UsdWrappers/UsdPrim.h"
 #include "UsdWrappers/UsdStage.h"
 
-void FUsdIntegrationsViewModel::UpdateAttributes( const UE::FUsdStageWeak& InUsdStage, const TCHAR* InPrimPath )
+void FUsdIntegrationsViewModel::UpdateAttributes(const UE::FUsdStageWeak& InUsdStage, const TCHAR* InPrimPath)
 {
 #if USE_USD_SDK
 	UsdStage = InUsdStage;
@@ -18,58 +17,57 @@ void FUsdIntegrationsViewModel::UpdateAttributes( const UE::FUsdStageWeak& InUsd
 
 	Attributes.Reset();
 
-	if ( !UsdStage || PrimPath.IsEmpty() || UE::FSdfPath{ InPrimPath }.IsAbsoluteRootPath() )
+	if (!UsdStage || PrimPath.IsEmpty() || UE::FSdfPath{InPrimPath}.IsAbsoluteRootPath())
 	{
 		return;
 	}
 
 	TArray<FString> AttributeNames;
 
-	TFunction<void(bool, bool)> AddLiveLinkAttributes = [&AttributeNames]( bool bHasLiveLink, bool bIsSkelRoot )
+	TFunction<void(bool, bool)> AddLiveLinkAttributes = [&AttributeNames](bool bHasLiveLink, bool bIsSkeletal)
 	{
-		if ( bHasLiveLink )
+		if (bHasLiveLink)
 		{
-			if ( bIsSkelRoot )
+			if (bIsSkeletal)
 			{
-				AttributeNames.Add( *UsdToUnreal::ConvertToken( UnrealIdentifiers::UnrealAnimBlueprintPath ) );
+				AttributeNames.Add(*UsdToUnreal::ConvertToken(UnrealIdentifiers::UnrealAnimBlueprintPath));
 			}
 
-			AttributeNames.Add( *UsdToUnreal::ConvertToken( UnrealIdentifiers::UnrealLiveLinkSubjectName ) );
-			AttributeNames.Add( *UsdToUnreal::ConvertToken( UnrealIdentifiers::UnrealLiveLinkEnabled ) );
+			AttributeNames.Add(*UsdToUnreal::ConvertToken(UnrealIdentifiers::UnrealLiveLinkSubjectName));
+			AttributeNames.Add(*UsdToUnreal::ConvertToken(UnrealIdentifiers::UnrealLiveLinkEnabled));
 		}
 	};
 
-	TFunction<void( bool, bool )> AddControlRigAttributes = [&AttributeNames]( bool bHasControlRig, bool bIsSkelRoot )
+	TFunction<void(bool, bool)> AddControlRigAttributes = [&AttributeNames](bool bHasControlRig, bool bIsSkeletal)
 	{
-		if ( bHasControlRig && bIsSkelRoot )
+		if (bHasControlRig && bIsSkeletal)
 		{
-			AttributeNames.Add( *UsdToUnreal::ConvertToken( UnrealIdentifiers::UnrealControlRigPath ) );
-			AttributeNames.Add( *UsdToUnreal::ConvertToken( UnrealIdentifiers::UnrealUseFKControlRig ) );
-			AttributeNames.Add( *UsdToUnreal::ConvertToken( UnrealIdentifiers::UnrealControlRigReduceKeys ) );
-			AttributeNames.Add( *UsdToUnreal::ConvertToken( UnrealIdentifiers::UnrealControlRigReductionTolerance ) );
+			AttributeNames.Add(*UsdToUnreal::ConvertToken(UnrealIdentifiers::UnrealControlRigPath));
+			AttributeNames.Add(*UsdToUnreal::ConvertToken(UnrealIdentifiers::UnrealUseFKControlRig));
+			AttributeNames.Add(*UsdToUnreal::ConvertToken(UnrealIdentifiers::UnrealControlRigReduceKeys));
+			AttributeNames.Add(*UsdToUnreal::ConvertToken(UnrealIdentifiers::UnrealControlRigReductionTolerance));
 		}
 	};
 
-    if ( UE::FUsdPrim Prim = InUsdStage.GetPrimAtPath( UE::FSdfPath{ InPrimPath } ) )
-    {
-        if ( !Prim.IsPseudoRoot() )
-        {
-            const bool bHasLiveLink = UsdUtils::PrimHasSchema( Prim, UnrealIdentifiers::LiveLinkAPI );
-            const bool bHasControlRig = UsdUtils::PrimHasSchema( Prim, UnrealIdentifiers::ControlRigAPI );
-            const bool bIsSkelRoot = Prim.IsA( TEXT( "SkelRoot" ) );
+	if (UE::FUsdPrim Prim = InUsdStage.GetPrimAtPath(UE::FSdfPath{InPrimPath}))
+	{
+		if (!Prim.IsPseudoRoot())
+		{
+			const bool bHasLiveLink = UsdUtils::PrimHasSchema(Prim, UnrealIdentifiers::LiveLinkAPI);
+			const bool bHasControlRig = UsdUtils::PrimHasSchema(Prim, UnrealIdentifiers::ControlRigAPI);
+			const bool bIsSkeletal = Prim.IsA(TEXT("SkelRoot")) || Prim.IsA(TEXT("Skeleton"));
 
-            AddLiveLinkAttributes( bHasLiveLink, bIsSkelRoot );
-            AddControlRigAttributes( bHasControlRig, bIsSkelRoot );
+			AddLiveLinkAttributes(bHasLiveLink, bIsSkeletal);
+			AddControlRigAttributes(bHasControlRig, bIsSkeletal);
 
-            for ( const FString& AttributeName : AttributeNames )
-            {
-                if ( UE::FUsdAttribute Attr = Prim.GetAttribute( *AttributeName ) )
-                {
-					Attributes.Add( MakeShared<UE::FUsdAttribute>( Attr ) );
-                }
-            }
-        }
-    }
-#endif // USE_USD_SDK
+			for (const FString& AttributeName : AttributeNames)
+			{
+				if (UE::FUsdAttribute Attr = Prim.GetAttribute(*AttributeName))
+				{
+					Attributes.Add(MakeShared<UE::FUsdAttribute>(Attr));
+				}
+			}
+		}
+	}
+#endif	  // USE_USD_SDK
 }
-
diff --git a/Engine/Plugins/Importers/USDImporter/Source/USDStageEditorViewModels/Private/USDPrimViewModel.cpp b/Engine/Plugins/Importers/USDImporter/Source/USDStageEditorViewModels/Private/USDPrimViewModel.cpp
--- a/Engine/Plugins/Importers/USDImporter/Source/USDStageEditorViewModels/Private/USDPrimViewModel.cpp
+++ b/Engine/Plugins/Importers/USDImporter/Source/USDStageEditorViewModels/Private/USDPrimViewModel.cpp
@@ -4,11 +4,9 @@
 
 #include "USDConversionUtils.h"
 #include "USDIntegrationUtils.h"
-#include "USDLog.h"
 #include "USDMemory.h"
 #include "USDTypesConversion.h"
 
-#include "UsdWrappers/SdfLayer.h"
 #include "UsdWrappers/SdfPath.h"
 #include "UsdWrappers/UsdPrim.h"
 #include "UsdWrappers/UsdStage.h"
@@ -16,14 +14,12 @@
 #if USE_USD_SDK
 
 #include "USDIncludesStart.h"
-	#include "pxr/pxr.h"
 	#include "pxr/usd/sdf/path.h"
-	#include "pxr/usd/usd/modelAPI.h"
 	#include "pxr/usd/usd/payloads.h"
 	#include "pxr/usd/usd/prim.h"
-	#include "pxr/usd/usd/tokens.h"
 	#include "pxr/usd/usdGeom/xform.h"
 	#include "pxr/usd/usdSkel/root.h"
+	#include "pxr/usd/usdSkel/skeleton.h"
 #include "USDIncludesEnd.h"
 
 #endif // #if USE_USD_SDK
@@ -254,7 +250,7 @@ bool FUsdPrimViewModel::CanApplySchema( FName SchemaName ) const
 	pxr::UsdPrim PxrUsdPrim{ UsdPrim };
 	pxr::TfToken SchemaToken = UnrealToUsd::ConvertToken( *SchemaName.ToString() ).Get();
 
-	if ( SchemaToken == UnrealIdentifiers::ControlRigAPI && !PxrUsdPrim.IsA<pxr::UsdSkelRoot>() )
+	if (SchemaToken == UnrealIdentifiers::ControlRigAPI && !(PxrUsdPrim.IsA<pxr::UsdSkelRoot>() || PxrUsdPrim.IsA<pxr::UsdSkelSkeleton>()))
 	{
 		return false;
 	}
diff --git a/Engine/Plugins/Importers/USDImporter/Source/USDStageImporter/Private/USDStageImporter.cpp b/Engine/Plugins/Importers/USDImporter/Source/USDStageImporter/Private/USDStageImporter.cpp
--- a/Engine/Plugins/Importers/USDImporter/Source/USDStageImporter/Private/USDStageImporter.cpp
+++ b/Engine/Plugins/Importers/USDImporter/Source/USDStageImporter/Private/USDStageImporter.cpp
@@ -496,6 +496,11 @@ namespace UsdStageImporterImpl
 		else if (USkeletalMesh* SkMesh = Cast<USkeletalMesh>(Asset))
 		{
 			AssetPrefix = TEXT("SK_");
+
+			// Our SkeletalMesh asset is now assigned to the Skeleton prim path, but we don't want our actual USkeletalMesh
+			// to be named "Skeleton" or something like that. Let's keep naming it after the SkelRoot instead (which the
+			// UAsset will already be named after anyway)
+			AssetName = Asset->GetFName().GetPlainNameString();
 		}
 		else if (USkeleton* Skeleton = Cast<USkeleton>(Asset))
 		{
@@ -1881,6 +1886,7 @@ void UUsdStageImporter::ImportFromFile(FUsdStageImportContext& ImportContext)
 	ImportContext.AssetCache->MarkAssetsAsStale();
 	ImportContext.LevelSequenceHelper.SetInfoCache(InfoCache);
 	ImportContext.LevelSequenceHelper.Init( ImportContext.Stage );  // Must happen after the context gets an InfoCache!
+	ImportContext.LevelSequenceHelper.SetRootMotionHandling(ImportContext.ImportOptions->RootMotionHandling);
 
 	EUsdPurpose PurposesToImport = static_cast<EUsdPurpose>(ImportContext.ImportOptions->PurposesToImport);
 
@@ -2037,6 +2043,7 @@ bool UUsdStageImporter::ReimportSingleAsset(
 	ImportContext.AssetCache->MarkAssetsAsStale();
 	ImportContext.LevelSequenceHelper.SetInfoCache(InfoCache);
 	ImportContext.LevelSequenceHelper.Init(ImportContext.Stage);  // Must happen after the context gets an InfoCache!
+	ImportContext.LevelSequenceHelper.SetRootMotionHandling(ImportContext.ImportOptions->RootMotionHandling);
 
 	EUsdPurpose PurposesToImport = static_cast<EUsdPurpose>(ImportContext.ImportOptions->PurposesToImport);
 
diff --git a/Engine/Plugins/Importers/USDImporter/Source/USDUtilities/Private/USDConversionUtils.cpp b/Engine/Plugins/Importers/USDImporter/Source/USDUtilities/Private/USDConversionUtils.cpp
--- a/Engine/Plugins/Importers/USDImporter/Source/USDUtilities/Private/USDConversionUtils.cpp
+++ b/Engine/Plugins/Importers/USDImporter/Source/USDUtilities/Private/USDConversionUtils.cpp
@@ -441,7 +441,7 @@ UClass* UsdUtils::GetActorTypeForPrim( const pxr::UsdPrim& Prim )
 
 UClass* UsdUtils::GetComponentTypeForPrim( const pxr::UsdPrim& Prim )
 {
-	if ( Prim.IsA< pxr::UsdSkelRoot >() )
+	if ( Prim.IsA< pxr::UsdSkelSkeleton >() )
 	{
 		return USkeletalMeshComponent::StaticClass();
 	}
@@ -921,35 +921,24 @@ bool UsdUtils::IsAnimated(const pxr::UsdPrim& Prim)
 		}
 	}
 
-	if (pxr::UsdSkelRoot SkeletonRoot{Prim})
+	if (pxr::UsdSkelSkeleton Skeleton{Prim})
 	{
-		pxr::UsdSkelCache SkeletonCache;
-		SkeletonCache.Populate(SkeletonRoot, pxr::UsdTraverseInstanceProxies());
-
-		std::vector< pxr::UsdSkelBinding > SkeletonBindings;
-		SkeletonCache.ComputeSkelBindings(SkeletonRoot, &SkeletonBindings, pxr::UsdTraverseInstanceProxies());
-
-		for (const pxr::UsdSkelBinding& Binding : SkeletonBindings)
+		if (pxr::UsdSkelRoot ClosestParentSkelRoot = pxr::UsdSkelRoot{UsdUtils::GetClosestParentSkelRoot(Prim)})
 		{
-			const pxr::UsdSkelSkeleton& Skeleton = Binding.GetSkeleton();
-			pxr::UsdSkelSkeletonQuery SkelQuery = SkeletonCache.GetSkelQuery(Skeleton);
-			pxr::UsdSkelAnimQuery AnimQuery = SkelQuery.GetAnimQuery();
-			if (!AnimQuery)
-			{
-				continue;
-			}
+			pxr::UsdSkelCache SkeletonCache;
+			SkeletonCache.Populate(ClosestParentSkelRoot, pxr::UsdTraverseInstanceProxies());
 
-			std::vector<double> JointTimeSamples;
-			std::vector<double> BlendShapeTimeSamples;
-			if ( ( AnimQuery.GetJointTransformTimeSamples( &JointTimeSamples ) && JointTimeSamples.size() > 0 ) ||
-				 ( AnimQuery.GetBlendShapeWeightTimeSamples( &BlendShapeTimeSamples ) && BlendShapeTimeSamples.size() > 0 ) )
+			pxr::UsdSkelSkeletonQuery SkelQuery = SkeletonCache.GetSkelQuery(Skeleton);
+			if (pxr::UsdSkelAnimQuery AnimQuery = SkelQuery.GetAnimQuery())
 			{
-				return true;
+				std::vector<double> JointTimeSamples;
+				std::vector<double> BlendShapeTimeSamples;
+				if ((AnimQuery.GetJointTransformTimeSamples(&JointTimeSamples) && JointTimeSamples.size() > 0)
+					|| (AnimQuery.GetBlendShapeWeightTimeSamples(&BlendShapeTimeSamples) && BlendShapeTimeSamples.size() > 0))
+				{
+					return true;
+				}
 			}
-
-			// We only parse Skeletons and SkelAnimations from the first skeletal binding of a SkelRoot, so
-			// if that one is not animated then this entire SkelRoot is not animated to us either (for now)
-			break;
 		}
 	}
 
diff --git a/Engine/Plugins/Importers/USDImporter/Source/USDUtilities/Private/USDGeomMeshConversion.cpp b/Engine/Plugins/Importers/USDImporter/Source/USDUtilities/Private/USDGeomMeshConversion.cpp
--- a/Engine/Plugins/Importers/USDImporter/Source/USDUtilities/Private/USDGeomMeshConversion.cpp
+++ b/Engine/Plugins/Importers/USDImporter/Source/USDUtilities/Private/USDGeomMeshConversion.cpp
@@ -15,6 +15,7 @@
 #include "USDPrimConversion.h"
 #include "USDProjectSettings.h"
 #include "USDShadeConversion.h"
+#include "USDSkeletalDataConversion.h"
 #include "USDTypesConversion.h"
 
 #include "UsdWrappers/SdfPath.h"
@@ -66,6 +67,8 @@
 	#include "pxr/usd/usdShade/material.h"
 	#include "pxr/usd/usdShade/materialBindingAPI.h"
 	#include "pxr/usd/usdShade/tokens.h"
+	#include "pxr/usd/usdSkel/bindingAPI.h"
+	#include "pxr/usd/usdSkel/root.h"
 	#include "pxr/usdImaging/usdImaging/implicitSurfaceMeshUtils.h"
 #include "USDIncludesEnd.h"
 
@@ -765,7 +768,8 @@ namespace UE::UsdGeomMeshConversion::Private
 		FMeshDescription& OutMeshDescription,
 		UsdUtils::FUsdPrimMaterialAssignmentInfo& OutMaterialAssignments,
 		UsdToUnreal::FUsdMeshConversionOptions& Options,
-		bool bIsFirstPrim
+		bool bIsFirstPrim,
+		bool bIsInsideSkelRoot
 	)
 	{
 		// Ignore meshes from disabled purposes
@@ -829,7 +833,12 @@ namespace UE::UsdGeomMeshConversion::Private
 
 		if ( pxr::UsdGeomMesh Mesh = pxr::UsdGeomMesh( Prim ) )
 		{
-			bSuccess = UsdToUnreal::ConvertGeomMesh( Mesh, OutMeshDescription, OutMaterialAssignments, Options );
+			// We never want to glob up *skinned* meshes inside SkelRoots, as those presumably will be handled by the
+			// UsdSkelSkeletonTranslator and the skeletal data code path.
+			if (!GIsEditor || !bIsInsideSkelRoot || !Prim.HasAPI<pxr::UsdSkelBindingAPI>())
+			{
+				bSuccess = UsdToUnreal::ConvertGeomMesh( Mesh, OutMeshDescription, OutMaterialAssignments, Options );
+			}
 		}
 		// Check for a cube/capsule/etc. ConvertGeomPrimitive will internally check for all specific types
 		else if (pxr::UsdGeomGprim Gprim = pxr::UsdGeomGprim{Prim})
@@ -865,7 +874,8 @@ namespace UE::UsdGeomMeshConversion::Private
 					OutMeshDescription,
 					OutMaterialAssignments,
 					Options,
-					bChildIsFirstPrim
+					bChildIsFirstPrim,
+					bIsInsideSkelRoot || ChildPrim.IsA<pxr::UsdSkelRoot>()
 				);
 			}
 		}
@@ -2708,12 +2718,14 @@ bool UsdToUnreal::ConvertGeomMeshHierarchy(
 		);
 	}
 
+	const bool bIsInSkelRoot = static_cast<bool>(UsdUtils::GetClosestParentSkelRoot(Prim));
 	return UsdGeomMeshImpl::RecursivelyCollapseChildMeshes(
 		Prim,
 		OutMeshDescription,
 		OutMaterialAssignments,
 		OptionsCopy,
-		bSkipRootPrimTransformAndVisibility
+		bSkipRootPrimTransformAndVisibility,
+		bIsInSkelRoot
 	);
 }
 
@@ -3774,7 +3786,7 @@ bool UnrealToUsd::ConvertMeshDescriptions( const TArray<FMeshDescription>& LODIn
 			VariantSet.SetVariantSelection( VariantName );
 			EditContext.Emplace( VariantSet.GetVariantEditContext() );
 		}
-		
+
 		pxr::UsdGeomMesh TargetMesh;
 		if ( bExportMultipleLODs )
 		{
@@ -4138,7 +4150,7 @@ namespace UE::UsdGeometryCacheConversion::Private
 						}
 
 						// Since family name and type attributes must be set at time Default, set the Indices at time Default too
-						// #ueent_todo: Add support for varying geomsubsets. This can happen with animation where sections 
+						// #ueent_todo: Add support for varying geomsubsets. This can happen with animation where sections
 						// visibility are toggled on/off
 						pxr::UsdAttribute IndicesAttr = GeomSubsetSchema.CreateIndicesAttr();
 						IndicesAttr.Set(IndicesAttrValue);
@@ -4343,6 +4355,14 @@ bool UsdUtils::IsGeomMeshALOD( const pxr::UsdPrim& UsdMeshPrim )
 		return false;
 	}
 
+	// Note that we can't robustly check whether UsdMeshPrim "is inside of the LOD variant set" or not,
+	// because that can vary *per layer*... For example, a stage with layers root.usda and sub.usda can have
+	// the MeshA prim inside the LOD variant on root.usa, and MeshB prim inside the LOD variant on sub.usda.
+	// The LOD variant setup is a set of rules we specify ourselves and the users must adhere to, and one of
+	// them is to have a single Mesh prim as a child of the variant set prim. This means that as soon
+	// as the user puts more than one Mesh prim inside of the variant set prim, we're already in a "garbage in"
+	// scenario, and will likely generate some garbage in turn. We'll emit a bunch of warning for that though.
+
 	return UE::UsdGeomMeshConversion::Private::DoesPrimContainMeshLODsInternal( UsdMeshPrim.GetParent() );
 }
 
diff --git a/Engine/Plugins/Importers/USDImporter/Source/USDUtilities/Private/USDPrimConversion.cpp b/Engine/Plugins/Importers/USDImporter/Source/USDUtilities/Private/USDPrimConversion.cpp
--- a/Engine/Plugins/Importers/USDImporter/Source/USDUtilities/Private/USDPrimConversion.cpp
+++ b/Engine/Plugins/Importers/USDImporter/Source/USDUtilities/Private/USDPrimConversion.cpp
@@ -102,7 +102,10 @@
 	#include "pxr/usd/usdShade/shader.h"
 	#include "pxr/usd/usdShade/tokens.h"
 	#include "pxr/usd/usdSkel/animation.h"
+	#include "pxr/usd/usdSkel/binding.h"
 	#include "pxr/usd/usdSkel/root.h"
+	#include "pxr/usd/usdSkel/skeleton.h"
+	#include "pxr/usd/usdSkel/skeletonQuery.h"
 #include "USDIncludesEnd.h"
 
 static bool GConsiderAllPrimsHaveAnimatedBounds = false;
@@ -2582,13 +2585,6 @@ bool UnrealToUsd::ConvertMeshComponent(const pxr::UsdStageRefPtr& Stage, const U
 	}
 	else if (const USkinnedMeshComponent* SkinnedMeshComponent = Cast<const USkinnedMeshComponent>(MeshComponent))
 	{
-		FScopedUsdAllocs Allocs;
-		pxr::UsdSkelRoot SkelRoot{UsdPrim};
-		if (!SkelRoot)
-		{
-			return false;
-		}
-
 		MeshAsset = SkinnedMeshComponent->GetSkinnedAsset();
 	}
 
@@ -2890,16 +2886,41 @@ bool UnrealToUsd::ConvertMaterialOverrides(
 		}
 		const bool bHasLODs = NumLODs > 1;
 
-		FString MeshName;
-		if (!bHasLODs)
+		if (!UsdPrim.IsA<pxr::UsdSkelSkeleton>())
 		{
-			for (const pxr::UsdPrim& Child : UsdPrim.GetChildren())
+			UE_LOG(
+				LogUsd,
+				Warning,
+				TEXT("For the skeletal case, ConvertMaterialOverrides must now receive a Skeleton prim! ('%s' was provided)"),
+				*UsdToUnreal::ConvertPath(UsdPrim.GetPrimPath())
+			);
+			return false;
+		}
+
+		// If performance becomes an issue we can start storing our skel caches in the info cache and optionally provide it
+		// to this function
+		pxr::UsdSkelRoot SkelRoot = pxr::UsdSkelRoot{UsdUtils::GetClosestParentSkelRoot(UsdPrim)};
+		if (!SkelRoot)
+		{
+			return false;
+		}
+		pxr::UsdSkelBinding SkelBinding;
+		pxr::UsdSkelSkeletonQuery SkeletonQuery;
+		if (!UsdUtils::GetSkelQueries(SkelRoot, pxr::UsdSkelSkeleton{UsdPrim}, SkelBinding, SkeletonQuery))
+		{
+			return false;
+		}
+
+		// Collect all skinned prim paths
+		const pxr::VtArray<pxr::UsdSkelSkinningQuery>& SkinningTargets = SkelBinding.GetSkinningTargets();
+		TSet<FString> SkinnedMeshPaths;
+		SkinnedMeshPaths.Reserve(SkinningTargets.size());
+		for (const pxr::UsdSkelSkinningQuery& SkinningTarget : SkinningTargets)
+		{
+			pxr::UsdPrim SkinnedPrim = SkinningTarget.GetPrim();
+			if (pxr::UsdGeomMesh SkinnedMesh = pxr::UsdGeomMesh{SkinnedPrim})
 			{
-				if (pxr::UsdGeomMesh Mesh{Child})
-				{
-					MeshName = UsdToUnreal::ConvertToken(Child.GetName());
-					break;
-				}
+				SkinnedMeshPaths.Add(UsdToUnreal::ConvertPath(SkinnedPrim.GetPrimPath()));
 			}
 		}
 
@@ -2943,32 +2964,54 @@ bool UnrealToUsd::ConvertMaterialOverrides(
 					{
 						if (const FUsdPrimPathList* SourcePrimPaths = UserData->MaterialSlotToPrimPaths.Find(SectionMatIndex))
 						{
-							for (const FString& PrimPath : SourcePrimPaths->PrimPaths)
+							// The N^2 here is not great but note that "ConvertMaterialOverrides" is not exactly spammed all that much, and that
+							// these arrays will in the general case have like 3 items each
+							for (const FString& SourcePrimPath : SourcePrimPaths->PrimPaths)
 							{
-								// See comment on the analogue part for the geometry cache component
-								if (PrimPath.StartsWith(UsdPrimPath))
+								for (const FString& SkinnedMeshPath : SkinnedMeshPaths)
 								{
-									pxr::SdfPath OverridePrimPath = UnrealToUsd::ConvertPath(*PrimPath).Get();
-									pxr::UsdPrim MeshPrim = Stage->OverridePrim(OverridePrimPath);
-									UE::USDPrimConversionImpl::Private::AuthorMaterialOverride(MeshPrim, Override->GetPathName());
+									// See comment on the analogue part for the geometry cache component
+									// Note we use 'StartsWith' here because our SourcePrimPath may be a UsdGeomSubset or something like that,
+									// but only the actual Mesh prim will count as a "skinned mesh"
+									if (SourcePrimPath.StartsWith(SkinnedMeshPath))
+									{
+										pxr::SdfPath OverridePrimPath = UnrealToUsd::ConvertPath(*SourcePrimPath).Get();
+										pxr::UsdPrim MeshPrim = Stage->OverridePrim(OverridePrimPath);
+										UE::USDPrimConversionImpl::Private::AuthorMaterialOverride(MeshPrim, Override->GetPathName());
+										break;
+									}
 								}
 							}
 						}
 					}
+					// TODO: We really need a separate function for ConvertingMaterialOverrides (to an opened stage) and ExportingMaterialOverrides
+					// that we can use when the SkeletalMesh is not something we generated ourselves (with annotated MaterialSlotToPrimPaths).
+					// We could collect an analogue for MaterialSlotToPrimPaths during the mesh export process to accurately author these overrides too
 					else
 					{
-						pxr::SdfPath OverridePrimPath = UsdPrim.GetPath();
+						pxr::SdfPath OverridePrimPath;
 
 						// If we have only 1 LOD, the asset's DefaultPrim will be a SkelRoot, and the Mesh will be a subprim
 						// with the same name. If we have multiple LODS, the default prim is also the SkelRoot but will contain separate
 						// Mesh prims for each LOD named "LOD0", "LOD1", etc., switched via a "LOD" variant set
 						if (bHasLODs)
 						{
-							OverridePrimPath = OverridePrimPath.AppendPath(UnrealToUsd::ConvertPath(*FString::Printf(TEXT("LOD%d"), LODIndex)).Get());
+							OverridePrimPath = SkelRoot.GetPath().AppendPath(UnrealToUsd::ConvertPath(*FString::Printf(TEXT("LOD%d"), LODIndex)).Get());
 						}
 						else
 						{
-							OverridePrimPath = OverridePrimPath.AppendElementString(UnrealToUsd::ConvertString(*MeshName).Get());
+							// Here we're guessing that we're converting material overrides for our exported level, which will use our
+							// own prims from exported SkeletalMeshes that all just have a single skinned mesh anyway
+							FString MeshName;
+							if (SkinningTargets.size() > 0)
+							{
+								if (pxr::UsdPrim FirstSkinnedPrim = SkinningTargets[0].GetPrim())
+								{
+									MeshName = UsdToUnreal::ConvertString(FirstSkinnedPrim.GetName());
+								}
+							}
+
+							OverridePrimPath = SkelRoot.GetPath().AppendElementString(UnrealToUsd::ConvertString(*MeshName).Get());
 						}
 
 						// If our LOD has only one section, its material assignment will be authored directly on the Mesh prim.
@@ -3514,7 +3557,7 @@ bool UnrealToUsd::CreateComponentPropertyBaker( UE::FUsdPrim& Prim, const UScene
 	return false;
 }
 
-bool UnrealToUsd::CreateSkeletalAnimationBaker( UE::FUsdPrim& SkelRoot, UE::FUsdPrim& SkelAnimation, USkeletalMeshComponent& Component, FComponentBaker& OutBaker )
+bool UnrealToUsd::CreateSkeletalAnimationBaker( UE::FUsdPrim& SkeletonPrim, UE::FUsdPrim& SkelAnimation, USkeletalMeshComponent& Component, FComponentBaker& OutBaker )
 {
 #if WITH_EDITOR
 	USkeletalMesh* SkeletalMesh = Component.GetSkeletalMeshAsset();
@@ -3525,19 +3568,29 @@ bool UnrealToUsd::CreateSkeletalAnimationBaker( UE::FUsdPrim& SkelRoot, UE::FUsd
 
 	FScopedUsdAllocs Allocs;
 
-	pxr::UsdSkelRoot UsdSkelRoot{ SkelRoot };
+	pxr::UsdSkelSkeleton UsdSkeleton{SkeletonPrim};
+	if (!UsdSkeleton)
+	{
+		UE_LOG(
+			LogUsd,
+			Warning,
+			TEXT("Failed to create skeletal animation baker: Prim '%s' must be a UsdSkeleton!"),
+			*SkeletonPrim.GetPrimPath().GetString()
+		);
+		return false;
+	}
+
 	pxr::UsdSkelAnimation UsdSkelAnimation{ SkelAnimation };
-	if ( !SkelRoot || !SkelAnimation )
+	if ( !UsdSkeleton || !SkelAnimation )
 	{
 		return false;
 	}
 
-	// Make sure that the skel root is using our animation
-	pxr::UsdPrim SkelRootPrim = UsdSkelRoot.GetPrim();
+	// Make sure that the skeleton is using our animation
 	pxr::UsdPrim SkelAnimPrim = UsdSkelAnimation.GetPrim();
-	UsdUtils::BindAnimationSource( SkelRootPrim, SkelAnimPrim );
+	UsdUtils::BindAnimationSource(SkeletonPrim, SkelAnimPrim);
 
-	FUsdStageInfo StageInfo{ SkelRoot.GetStage() };
+	FUsdStageInfo StageInfo{ SkeletonPrim.GetStage() };
 
 	pxr::UsdAttribute JointsAttr			= UsdSkelAnimation.CreateJointsAttr();
 	pxr::UsdAttribute TranslationsAttr		= UsdSkelAnimation.CreateTranslationsAttr();
@@ -3630,6 +3683,27 @@ bool UnrealToUsd::CreateSkeletalAnimationBaker( UE::FUsdPrim& SkelRoot, UE::FUsd
 			}
 
 			const TArray< FTransform >& LocalBoneTransforms = Component.GetBoneSpaceTransforms();
+
+			// For whatever reason it seems that sometimes this is not ready for us, so let's force it to be recalculated
+			if (LocalBoneTransforms.Num() == 0)
+			{
+				const int32 LODIndex = 0;
+				Component.RecalcRequiredBones(LODIndex);
+			}
+			if (LocalBoneTransforms.Num() != NumBones)
+			{
+				UE_LOG(
+					LogUsd,
+					Warning,
+					TEXT("Failed to retrieve bone transforms when baking skeletal animation for component '%s' at timeCode '%f'. Expected %d transforms, received %d"),
+					*Component.GetPathName(),
+					UsdTimeCode,
+					NumBones,
+					LocalBoneTransforms.Num()
+				);
+				return;
+			}
+
 			for ( int32 BoneIndex = 0; BoneIndex < NumBones; ++BoneIndex )
 			{
 				FTransform BoneTransform = LocalBoneTransforms[ BoneIndex ];
diff --git a/Engine/Plugins/Importers/USDImporter/Source/USDUtilities/Private/USDSkeletalDataConversion.cpp b/Engine/Plugins/Importers/USDImporter/Source/USDUtilities/Private/USDSkeletalDataConversion.cpp
--- a/Engine/Plugins/Importers/USDImporter/Source/USDUtilities/Private/USDSkeletalDataConversion.cpp
+++ b/Engine/Plugins/Importers/USDImporter/Source/USDUtilities/Private/USDSkeletalDataConversion.cpp
@@ -2392,7 +2392,7 @@ bool UsdToUnreal::ConvertSkelAnim(
 		// We don't want to store just StartSeconds here, because part of that may be because the layer itself
 		// has an offset/scale within the stage. In OutStartOffsetSeconds we need to store the start of the animation
 		// in seconds *with respect to its own layer*. The layer's offset/scale can be retrieved later at any time by
-		// just looking at the FSdfLayerOffset for the SkelAnimation prim (like what is done in FUsdSkelRootTranslator::UpdateComponents)
+		// just looking at the FSdfLayerOffset for the SkelAnimation prim (like what is done in FUsdSkelSkeletonTranslator::UpdateComponents)
 		*OutStartOffsetSeconds = LayerStartSeconds;
 	}
 
@@ -2408,6 +2408,11 @@ bool UsdToUnreal::ConvertBlendShape( const pxr::UsdSkelBlendShape& UsdBlendShape
 {
 	TRACE_CPUPROFILER_EVENT_SCOPE( UsdToUnreal::ConvertBlendShape );
 
+	if (!UsdBlendShape)
+	{
+		return false;
+	}
+
 	FScopedUsdAllocs Allocs;
 
 	pxr::UsdAttribute OffsetsAttr = UsdBlendShape.GetOffsetsAttr();
@@ -2769,15 +2774,24 @@ pxr::UsdSkelSkinningQuery UsdUtils::CreateSkinningQuery( const pxr::UsdGeomMesh&
 	);
 }
 
-void UsdUtils::BindAnimationSource( pxr::UsdPrim& Prim, const pxr::UsdPrim& AnimationSource )
+void UsdUtils::BindAnimationSource(pxr::UsdPrim& Prim, const pxr::UsdPrim& AnimationSource)
 {
 	FScopedUsdAllocs UsdAllocs;
 
-	pxr::UsdSkelBindingAPI SkelBindingAPI = pxr::UsdSkelBindingAPI::Apply( Prim );
-	SkelBindingAPI.CreateAnimationSourceRel().SetTargets( pxr::SdfPathVector( { AnimationSource.GetPath() } ) );
+	pxr::UsdSkelBindingAPI SkelBindingAPI = pxr::UsdSkelBindingAPI::Apply(Prim);
+	pxr::UsdRelationship AnimSourceRel = SkelBindingAPI.CreateAnimationSourceRel();
+	if (AnimationSource)
+	{
+		AnimSourceRel.SetTargets(pxr::SdfPathVector({AnimationSource.GetPath()}));
+	}
+	else
+	{
+		const bool bRemoveSpec = false;
+		AnimSourceRel.ClearTargets(bRemoveSpec);
+	}
 }
 
-UE::FUsdPrim UsdUtils::FindFirstAnimationSource( const UE::FUsdPrim& SkelRootPrim )
+UE::FUsdPrim UsdUtils::FindFirstAnimationSource(const UE::FUsdPrim& SkelRootPrim)
 {
 	if ( !SkelRootPrim )
 	{
@@ -2786,36 +2800,108 @@ UE::FUsdPrim UsdUtils::FindFirstAnimationSource( const UE::FUsdPrim& SkelRootPri
 
 	FScopedUsdAllocs UsdAllocs;
 
-	// For now we really only parse the first skeletal binding of a SkelRoot (check USDSkelRootTranslator.cpp,
+	// For now we really only parse the first skeletal binding of a SkelRoot (check USDSkelSkeletonTranslator.cpp,
 	// LoadAllSkeletalData) and its SkelAnimation, if any.
 	// Note that we don't check the SkelRoot prim directly for the SkelAnimation binding: If it has a valid one
 	// it will propagate down to child namespaces and affect our first skeletal binding anyway
 
-	if ( pxr::UsdSkelRoot SkeletonRoot{ pxr::UsdPrim{ SkelRootPrim } } )
+	if (pxr::UsdSkelRoot SkeletonRoot{pxr::UsdPrim{SkelRootPrim}})
 	{
-		std::vector< pxr::UsdSkelBinding > SkeletonBindings;
+		std::vector<pxr::UsdSkelBinding> SkeletonBindings;
 
 		pxr::UsdSkelCache SkeletonCache;
-		SkeletonCache.Populate( SkeletonRoot, pxr::UsdTraverseInstanceProxies() );
-		SkeletonCache.ComputeSkelBindings( SkeletonRoot, &SkeletonBindings, pxr::UsdTraverseInstanceProxies() );
+		SkeletonCache.Populate(SkeletonRoot, pxr::UsdTraverseInstanceProxies());
+		SkeletonCache.ComputeSkelBindings(SkeletonRoot, &SkeletonBindings, pxr::UsdTraverseInstanceProxies());
 
-		for ( const pxr::UsdSkelBinding& Binding : SkeletonBindings )
+		for (const pxr::UsdSkelBinding& Binding : SkeletonBindings)
 		{
 			const pxr::UsdSkelSkeleton& Skeleton = Binding.GetSkeleton();
-			pxr::UsdSkelSkeletonQuery SkelQuery = SkeletonCache.GetSkelQuery( Skeleton );
+			pxr::UsdSkelSkeletonQuery SkelQuery = SkeletonCache.GetSkelQuery(Skeleton);
 			pxr::UsdSkelAnimQuery AnimQuery = SkelQuery.GetAnimQuery();
-			if ( !AnimQuery )
+			if (!AnimQuery)
 			{
 				continue;
 			}
 
-			return UE::FUsdPrim{ AnimQuery.GetPrim() };
+			return UE::FUsdPrim{AnimQuery.GetPrim()};
+		}
+	}
+
+	return {};
+}
+
+UE::FUsdPrim UsdUtils::FindAnimationSource(const pxr::UsdPrim& SkelRootPrim, const pxr::UsdPrim& SkeletonPrim)
+{
+	if (!SkeletonPrim)
+	{
+		return {};
+	}
+
+	FScopedUsdAllocs UsdAllocs;
+
+	pxr::UsdSkelSkeleton Skeleton{SkeletonPrim};
+	pxr::UsdSkelRoot ClosestParentSkelRoot{SkelRootPrim};
+	if (Skeleton && ClosestParentSkelRoot)
+	{
+		pxr::UsdSkelCache SkeletonCache;
+		SkeletonCache.Populate(ClosestParentSkelRoot, pxr::UsdTraverseInstanceProxies());
+
+		pxr::UsdSkelSkeletonQuery SkelQuery = SkeletonCache.GetSkelQuery(Skeleton);
+		if (pxr::UsdSkelAnimQuery AnimQuery = SkelQuery.GetAnimQuery())
+		{
+			return UE::FUsdPrim{AnimQuery.GetPrim()};
+		}
+	}
+
+	return {};
+}
+
+UE::FUsdPrim UsdUtils::GetClosestParentSkelRoot(const pxr::UsdPrim& SomePrim)
+{
+	FScopedUsdAllocs Allocs;
+
+	pxr::UsdPrim Parent = SomePrim;
+	while (Parent && !Parent.IsPseudoRoot())
+	{
+		if (Parent.IsA<pxr::UsdSkelRoot>())
+		{
+			return UE::FUsdPrim{Parent};
 		}
+
+		Parent = Parent.GetParent();
 	}
 
 	return {};
 }
 
+bool UsdUtils::GetSkelQueries(
+	const pxr::UsdSkelRoot& InSkelRootPrim,
+	const pxr::UsdSkelSkeleton& InSkeletonPrim,
+	pxr::UsdSkelBinding& OutSkelBinding,
+	pxr::UsdSkelSkeletonQuery& OutSkeletonQuery,
+	pxr::UsdSkelCache* InOutSkelCache
+)
+{
+	if (!InSkelRootPrim || !InSkeletonPrim)
+	{
+		return false;
+	}
+
+	FScopedUsdAllocs Allocs;
+
+	TOptional<pxr::UsdSkelCache> TempCache;
+	if (!InOutSkelCache)
+	{
+		TempCache.Emplace();
+		InOutSkelCache = &TempCache.GetValue();
+		InOutSkelCache->Populate(InSkelRootPrim, pxr::UsdTraverseInstanceProxies());
+	}
+
+	OutSkeletonQuery = InOutSkelCache->GetSkelQuery(InSkeletonPrim);
+
+	return InOutSkelCache->ComputeSkelBinding(InSkelRootPrim, InSkeletonPrim, &OutSkelBinding, pxr::UsdTraverseInstanceProxies());
+}
+
 #endif // USE_USD_SDK
 
 #if USE_USD_SDK && WITH_EDITOR
@@ -2989,7 +3075,7 @@ bool UnrealToUsd::ConvertSkeletalMesh( const USkeletalMesh* SkeletalMesh, pxr::U
 	pxr::UsdSkelBindingAPI SkelBindingAPI = pxr::UsdSkelBindingAPI::Apply( SkelRootPrim );
 	{
 		pxr::UsdPrim SkeletonPrim = Stage->DefinePrim(
-			SkelRootPrim.GetPath().AppendChild( UnrealToUsd::ConvertToken(TEXT("Skel")).Get() ),
+			SkelRootPrim.GetPath().AppendChild(UnrealToUsd::ConvertToken(UnrealIdentifiers::ExportedSkeletonPrimName).Get()),
 			UnrealToUsd::ConvertToken(TEXT("Skeleton")).Get()
 		);
 		pxr::UsdSkelSkeleton SkelSkeleton{ SkeletonPrim };
@@ -3501,6 +3587,8 @@ bool UnrealToUsd::ConvertControlRigSection(
 			// We'll change the blend shape channel names, so we need to update all meshes that were using them too.
 			// For now we'll assume that they're all inside the same skel root. We could upgrade this for the stage
 			// later too, if needed
+			// TODO: This could probably be updated to just find the actual skinned meshes, and have some better parameters
+			// like skinning/skeleton queries
 			for ( UE::FUsdPrim& MeshPrim : UsdUtils::GetAllPrimsOfType( UE::FUsdPrim{ InSkelRoot }, TEXT( "UsdGeomMesh" ) ) )
 			{
 				pxr::UsdSkelBindingAPI SkelBindingAPI{ MeshPrim };
diff --git a/Engine/Plugins/Importers/USDImporter/Source/USDUtilities/Public/USDPrimConversion.h b/Engine/Plugins/Importers/USDImporter/Source/USDUtilities/Public/USDPrimConversion.h
--- a/Engine/Plugins/Importers/USDImporter/Source/USDUtilities/Public/USDPrimConversion.h
+++ b/Engine/Plugins/Importers/USDImporter/Source/USDUtilities/Public/USDPrimConversion.h
@@ -255,7 +255,7 @@ namespace UnrealToUsd
 	 * Creates a property baker responsible for baking the joint/blend shape state of Component as a SkelAnimation,
 	 * returning true if a EBakingType::Skeletal baker was created
 	 */
-	USDUTILITIES_API bool CreateSkeletalAnimationBaker( UE::FUsdPrim& SkelRoot, UE::FUsdPrim& SkelAnimation, USkeletalMeshComponent& Component, FComponentBaker& OutBaker );
+	USDUTILITIES_API bool CreateSkeletalAnimationBaker( UE::FUsdPrim& SkeletonPrim, UE::FUsdPrim& SkelAnimation, USkeletalMeshComponent& Component, FComponentBaker& OutBaker );
 
 	/**
 	 * Struct with lambda functions that can be used to convert an UE property's value into the corresponding USD attribute values at the received UsdTimeCode.
diff --git a/Engine/Plugins/Importers/USDImporter/Source/USDUtilities/Public/USDSkeletalDataConversion.h b/Engine/Plugins/Importers/USDImporter/Source/USDUtilities/Public/USDSkeletalDataConversion.h
--- a/Engine/Plugins/Importers/USDImporter/Source/USDUtilities/Public/USDSkeletalDataConversion.h
+++ b/Engine/Plugins/Importers/USDImporter/Source/USDUtilities/Public/USDSkeletalDataConversion.h
@@ -21,7 +21,9 @@
 PXR_NAMESPACE_OPEN_SCOPE
 	class UsdGeomMesh;
 	class UsdSkelAnimQuery;
+	class UsdSkelBinding;
 	class UsdSkelBlendShape;
+	class UsdSkelCache;
 	class UsdSkelRoot;
 	class UsdSkelSkeleton;
 	class UsdSkelSkeletonQuery;
@@ -118,10 +120,33 @@ namespace UsdUtils
 
 	/**
 	 * Returns the SkelAnimation prim that is resolved for the first skeletal binding of SkelRootPrim, if it is a SkelRoot.
-	 * We use this as we currently only parse a single Skeleton per SkelRoot (and so a single SkelAnimation), but in the future we may
-	 * decide to do something else.
 	 */
-	USDUTILITIES_API UE::FUsdPrim FindFirstAnimationSource( const UE::FUsdPrim& SkelRootPrim );
+	USDUTILITIES_API UE::FUsdPrim FindFirstAnimationSource(const UE::FUsdPrim& InSkelRootPrim);
+
+	/**
+	 * Returns the SkelAnimation prim that is bound as animation source for the provided skeleton
+	 */
+	USDUTILITIES_API UE::FUsdPrim FindAnimationSource(const pxr::UsdPrim& InSkelRootPrim, const pxr::UsdPrim& InSkeletonPrim);
+
+	/**
+	 * Convenience function that traverses through ancestors up from SomePrim and returns the first SkelRoot prim that it finds,
+	 * returning an invalid prim in case it doesn't find anything
+	 */
+	USDUTILITIES_API UE::FUsdPrim GetClosestParentSkelRoot(const pxr::UsdPrim& SomePrim);
+
+	/**
+	 * Convenience function to use/populate a SkelCache on a SkelRoot and retrieve the SkelBinding/SkeletonQuery for a
+	 * particular InSkeletonPrim. You can retrieve the AnimQuery from the SkeletonQuery as well.
+	 * A temp SkelCache will be constructed if none is provided.
+	 * Returns true whether we managed to retrieve everything.
+	 */
+	USDUTILITIES_API bool GetSkelQueries(
+		const pxr::UsdSkelRoot& InSkelRootPrim,
+		const pxr::UsdSkelSkeleton& InSkeletonPrim,
+		pxr::UsdSkelBinding& OutSkelBinding,
+		pxr::UsdSkelSkeletonQuery& OutSkeletonQuery,
+		pxr::UsdSkelCache* InOutSkelCache = nullptr
+	);
 #endif // USE_USD_SDK
 }
 
diff --git a/Engine/Plugins/Importers/USDImporter/Source/UnrealUSDWrapper/Private/UnrealUSDWrapper.cpp b/Engine/Plugins/Importers/USDImporter/Source/UnrealUSDWrapper/Private/UnrealUSDWrapper.cpp
--- a/Engine/Plugins/Importers/USDImporter/Source/UnrealUSDWrapper/Private/UnrealUSDWrapper.cpp
+++ b/Engine/Plugins/Importers/USDImporter/Source/UnrealUSDWrapper/Private/UnrealUSDWrapper.cpp
@@ -415,6 +415,7 @@ const TCHAR* UnrealIdentifiers::TwoSidedMaterialSuffix = TEXT("_TwoSided");
 const TCHAR* UnrealIdentifiers::Invisible = TEXT("invisible");
 const TCHAR* UnrealIdentifiers::Inherited = TEXT("inherited");
 const TCHAR* UnrealIdentifiers::IdentifierPrefix = TEXT("@identifier:");
+const TCHAR* UnrealIdentifiers::ExportedSkeletonPrimName = TEXT("Skel");
 
 FName UnrealIdentifiers::TransformPropertyName = TEXT( "Transform" ); // Fake FName because the transform is stored decomposed on the component
 FName UnrealIdentifiers::HiddenInGamePropertyName = GET_MEMBER_NAME_CHECKED( USceneComponent, bHiddenInGame );
diff --git a/Engine/Plugins/Importers/USDImporter/Source/UnrealUSDWrapper/Public/UnrealUSDWrapper.h b/Engine/Plugins/Importers/USDImporter/Source/UnrealUSDWrapper/Public/UnrealUSDWrapper.h
--- a/Engine/Plugins/Importers/USDImporter/Source/UnrealUSDWrapper/Public/UnrealUSDWrapper.h
+++ b/Engine/Plugins/Importers/USDImporter/Source/UnrealUSDWrapper/Public/UnrealUSDWrapper.h
@@ -410,6 +410,8 @@ namespace UnrealIdentifiers
 	extern UNREALUSDWRAPPER_API const TCHAR* Inherited;
 	extern UNREALUSDWRAPPER_API const TCHAR* IdentifierPrefix;
 
+	extern UNREALUSDWRAPPER_API const TCHAR* ExportedSkeletonPrimName;
+
 	// USceneComponent properties
 	extern UNREALUSDWRAPPER_API FName TransformPropertyName;
 	extern UNREALUSDWRAPPER_API FName HiddenInGamePropertyName;
