Build fix

diff --git a/Engine/Plugins/Importers/USDImporter/Source/USDSchemas/Private/USDGeomMeshTranslator.cpp b/Engine/Plugins/Importers/USDImporter/Source/USDSchemas/Private/USDGeomMeshTranslator.cpp
--- a/Engine/Plugins/Importers/USDImporter/Source/USDSchemas/Private/USDGeomMeshTranslator.cpp
+++ b/Engine/Plugins/Importers/USDImporter/Source/USDSchemas/Private/USDGeomMeshTranslator.cpp
@@ -1198,6 +1198,7 @@ namespace UE::UsdCollision::Private
 			return;
 		}
 
+#if WITH_EDITORONLY_DATA
 		FMeshDescription* MeshDescription = StaticMesh.GetMeshDescription(0);
 		if (!MeshDescription || MeshDescription->Vertices().Num() == 0)
 		{
@@ -1357,6 +1358,7 @@ namespace UE::UsdCollision::Private
 
 		const bool bIsUpdate = true;
 		StaticMesh.CreateNavCollision(bIsUpdate);
+#endif // WITH_EDITORONLY_DATA
 	}
 }
 
