USD: Fix draw mode component not being directly spawnable by the user.

diff --git a/Engine/Plugins/Importers/USDImporter/Source/USDClasses/Public/USDDrawModeComponent.h b/Engine/Plugins/Importers/USDImporter/Source/USDClasses/Public/USDDrawModeComponent.h
--- a/Engine/Plugins/Importers/USDImporter/Source/USDClasses/Public/USDDrawModeComponent.h
+++ b/Engine/Plugins/Importers/USDImporter/Source/USDClasses/Public/USDDrawModeComponent.h
@@ -56,7 +56,7 @@ namespace UsdUtils
  * Component type that is used to draw bounds, cards and origin axes when translating a prim that has the UsdGeomModelAPI schema
  * applied.
  */
-UCLASS()
+UCLASS(ClassGroup = (USD), meta = (BlueprintSpawnableComponent))
 class USDCLASSES_API UUsdDrawModeComponent : public UPrimitiveComponent
 {
 	GENERATED_BODY()
