USD: Add support for multiple root bones when opening or importing stages.

diff --git a/Engine/Plugins/Importers/USDImporter/Source/USDSchemas/Private/USDSkelRootTranslator.cpp b/Engine/Plugins/Importers/USDImporter/Source/USDSchemas/Private/USDSkelRootTranslator.cpp
--- a/Engine/Plugins/Importers/USDImporter/Source/USDSchemas/Private/USDSkelRootTranslator.cpp
+++ b/Engine/Plugins/Importers/USDImporter/Source/USDSchemas/Private/USDSkelRootTranslator.cpp
@@ -1320,31 +1320,6 @@ namespace UsdSkelRootTranslatorImpl
 					Options
 				);
 
-				int32 NumRootBones = 0;
-				for (const SkeletalMeshImportData::FBone& Bone : SkeletonBones)
-				{
-					if (Bone.ParentIndex == INDEX_NONE)
-					{
-						NumRootBones++;
-					}
-
-					if (NumRootBones > 1)
-					{
-						break;
-					}
-				}
-
-				if (NumRootBones != 1)
-				{
-					UE_LOG(
-						LogUsd,
-						Warning,
-						TEXT("Ignoring SkelRoot '%s' as the bound Skeleton prim must have exactly one root bone!"),
-						*PrimPath.GetString()
-					);
-					bContinueTaskChain = false;
-				}
-
 				return bContinueTaskChain;
 			} );
 
diff --git a/Engine/Plugins/Importers/USDImporter/Source/USDSchemas/Private/USDSkelSkeletonTranslator.cpp b/Engine/Plugins/Importers/USDImporter/Source/USDSchemas/Private/USDSkelSkeletonTranslator.cpp
--- a/Engine/Plugins/Importers/USDImporter/Source/USDSchemas/Private/USDSkelSkeletonTranslator.cpp
+++ b/Engine/Plugins/Importers/USDImporter/Source/USDSchemas/Private/USDSkelSkeletonTranslator.cpp
@@ -1406,31 +1406,6 @@ namespace UsdSkelSkeletonTranslatorImpl
 				   Options
 			   );
 
-			   int32 NumRootBones = 0;
-			   for (const SkeletalMeshImportData::FBone& Bone : SkeletonBones)
-			   {
-				   if (Bone.ParentIndex == INDEX_NONE)
-				   {
-					   NumRootBones++;
-				   }
-
-				   if (NumRootBones > 1)
-				   {
-					   break;
-				   }
-			   }
-
-			   if (NumRootBones != 1)
-			   {
-				   UE_LOG(
-					   LogUsd,
-					   Warning,
-					   TEXT("Ignoring SkelRoot '%s' as the bound Skeleton prim must have exactly one root bone!"),
-					   *SkeletonPrimPath.GetString()
-				   );
-				   bContinueTaskChain = false;
-			   }
-
 			   // If we parsed LODs we could potentially have invalidated our references to our queries,
 			   // so let's refresh them before the downstream tasks try using them.
 			   if (bContinueTaskChain && bTryLODParsing)
diff --git a/Engine/Plugins/Importers/USDImporter/Source/USDUtilities/Private/USDSkeletalDataConversion.cpp b/Engine/Plugins/Importers/USDImporter/Source/USDUtilities/Private/USDSkeletalDataConversion.cpp
--- a/Engine/Plugins/Importers/USDImporter/Source/USDUtilities/Private/USDSkeletalDataConversion.cpp
+++ b/Engine/Plugins/Importers/USDImporter/Source/USDUtilities/Private/USDSkeletalDataConversion.cpp
@@ -1178,6 +1178,8 @@ bool UsdToUnreal::ConvertSkeleton(const pxr::UsdSkelSkeletonQuery& SkeletonQuery
 	TArray<FString> JointNames;
 	TArray<int32> ParentJointIndices;
 
+	uint32 RootBoneCount = 0;
+
 	// Retrieve the joint names and parent indices from the skeleton topology
 	// GetJointOrder already orders them from parent-to-child
 	VtArray<TfToken> JointOrder = SkeletonQuery.GetJointOrder();
@@ -1191,6 +1193,11 @@ bool UsdToUnreal::ConvertSkeleton(const pxr::UsdSkelSkeletonQuery& SkeletonQuery
 
 		int ParentIndex = SkelTopology.GetParent(Index);
 		ParentJointIndices.Add(ParentIndex);
+
+		if (ParentIndex == -1)
+		{
+			RootBoneCount++;
+		}
 	}
 
 	// Skeleton has no joints: Generate a dummy single "Root" bone skeleton
@@ -1252,6 +1259,35 @@ bool UsdToUnreal::ConvertSkeleton(const pxr::UsdSkelSkeletonQuery& SkeletonQuery
 		return false;
 	}
 
+	// If we have more than one root bone, let's create a new "true root bone" and add the
+	// previously root bones as children of it
+	if (RootBoneCount > 1)
+	{
+		for (int32& Index : ParentJointIndices)
+		{
+			// Have previously root bones point at the new bone we'll add soon
+			if (Index == INDEX_NONE)
+			{
+				Index = 0;
+			}
+			// All other index references have to move one over since we'll push
+			// a new root bone into the start of the array
+			else
+			{
+				Index += 1;
+			}
+		}
+
+		const uint64 FirstIndex = 0;
+		const int32 RootParentIndex = INDEX_NONE;
+
+		FString UniqueNewRootName = UsdUtils::GetUniqueName(TEXT("Root"), TSet<FString>{JointNames});
+
+		JointNames.Insert(UniqueNewRootName, FirstIndex);
+		BoneTransforms.Insert(FTransform::Identity, FirstIndex);
+		ParentJointIndices.Insert(RootParentIndex, FirstIndex);
+	}
+
 	// Store the retrieved data as bones into the SkeletalMeshImportData
 	SkelMeshImportData.RefBonesBinary.AddZeroed( JointNames.Num() );
 	for (int32 Index = 0; Index < JointNames.Num(); ++Index)
@@ -1281,6 +1317,29 @@ bool UsdToUnreal::ConvertSkeleton(const pxr::UsdSkelSkeletonQuery& SkeletonQuery
 	return true;
 }
 
+namespace UE::USDSkeletalDataConversion::Private
+{
+	bool HasMultipleRootBones(const pxr::UsdSkelSkeletonQuery& SkeletonQuery)
+	{
+		const pxr::UsdSkelTopology& SkelTopology = SkeletonQuery.GetTopology();
+		const pxr::VtArray<int>& JointParentIndices = SkelTopology.GetParentIndices();
+		bool bFoundRoot = false;
+		for (int ParentIndex : JointParentIndices)
+		{
+			if (ParentIndex == INDEX_NONE)
+			{
+				if (bFoundRoot)
+				{
+					return true;
+				}
+				bFoundRoot = true;
+			}
+		}
+
+		return false;
+	}
+};
+
 bool UsdToUnreal::ConvertSkinnedMesh(
 	const pxr::UsdSkelSkinningQuery& SkinningQuery,
 	const pxr::UsdSkelSkeletonQuery& SkeletonQuery,
@@ -1858,6 +1917,13 @@ bool UsdToUnreal::ConvertSkinnedMesh(
 	VtArray<float> JointWeights;
 	SkinningQuery.ComputeVaryingJointInfluences(NumPoints, &JointIndices, &JointWeights);
 
+	// Keep track of whether we added an additional "true" root bone in the cases the bound skeleton has
+	// multiple root bones
+	// We'll only ever set NumAdditionalBones to 1 or 0 (as we'll only either need a "true root bone" or
+	// not), but naming it this way allows us to use it like an offset, which should make it easier to
+	// understand whenever it is used
+	uint32 NumAdditionalBones = UE::USDSkeletalDataConversion::Private::HasMultipleRootBones(SkeletonQuery) ? 1 : 0;
+
 	// Recompute the joint influences if we need to
 	uint32 NumInfluencesPerComponent = SkinningQuery.GetNumInfluencesPerComponent();
 	const uint32 MaxAllowedInfluences = EXTRA_BONE_INFLUENCES;
@@ -1886,7 +1952,7 @@ bool UsdToUnreal::ConvertSkinnedMesh(
 				if (BoneWeight != 0.f)
 				{
 					SkelMeshImportData.Influences.AddUninitialized();
-					SkelMeshImportData.Influences.Last().BoneIndex = JointIndices[JointIndex];
+					SkelMeshImportData.Influences.Last().BoneIndex = NumAdditionalBones + JointIndices[JointIndex];
 					SkelMeshImportData.Influences.Last().Weight = BoneWeight;
 					SkelMeshImportData.Influences.Last().VertexIndex = NumExistingPoints + PointIndex;
 				}
@@ -1907,8 +1973,8 @@ bool UsdToUnreal::ConvertSkinnedMesh(
 				if ( SkeletonJointsAttr.Get( &SkeletonJoints ) )
 				{
 					// If the skeleton has N bones, this will just contain { 0, 1, 2, ..., N-1 }
-					int NumSkeletonBones = static_cast< int >( SkeletonJoints.size() );
-					for ( int SkeletonBoneIndex = 0; SkeletonBoneIndex < NumSkeletonBones; ++SkeletonBoneIndex )
+					int NumUsdSkeletonBones = static_cast< int >( SkeletonJoints.size() );
+					for ( int SkeletonBoneIndex = 0; SkeletonBoneIndex < NumUsdSkeletonBones; ++SkeletonBoneIndex )
 					{
 						SkeletonBoneIndices.push_back( SkeletonBoneIndex );
 					}
@@ -1924,7 +1990,11 @@ bool UsdToUnreal::ConvertSkinnedMesh(
 						for ( int32 AddedInfluenceIndex = NumInfluencesBefore; AddedInfluenceIndex < NumInfluencesAfter; ++AddedInfluenceIndex )
 						{
 							SkeletalMeshImportData::FRawBoneInfluence& Influence = SkelMeshImportData.Influences[ AddedInfluenceIndex ];
-							Influence.BoneIndex = BoneIndexRemapping[ Influence.BoneIndex ];
+
+							// We have to remove our "NumAdditionalBones" offset from the influence's bone index because that's a UE concept that
+							// the BoneIndexRemapping array doesn't really know about. After that, we have a bone index that matches the USD Skeleton
+							// joint order, then we can remap with BoneIndexRemapping and add our NumAdditionalBones back in so that it matches our USkeleton
+							Influence.BoneIndex = NumAdditionalBones + BoneIndexRemapping[Influence.BoneIndex - NumAdditionalBones];
 						}
 					}
 				}
@@ -2017,11 +2087,20 @@ bool UsdToUnreal::ConvertSkelAnim(
 
 	const FReferenceSkeleton& RefSkeleton = Skeleton->GetReferenceSkeleton();
 	const TArray<FMeshBoneInfo>& BoneInfo = RefSkeleton.GetRawRefBoneInfo();
-	int32 NumBones = BoneInfo.Num();
-	int32 NumSkelQueryJoints = static_cast<int32>(InUsdSkeletonQuery.GetJointOrder().size());
+	int32 NumBonesInUE = BoneInfo.Num();  // This will already contain any new "true root bone" we may have created
+	int32 NumBonesInUsd = static_cast<int32>(InUsdSkeletonQuery.GetJointOrder().size());
+
+	// Keep track of whether we added an additional "true" root bone in the cases the bound skeleton has
+	// multiple root bones
+	// We'll only ever set NumAdditionalBones to 1 or 0 (as we'll only either need a "true root bone" or
+	// not), but naming it this way allows us to use it like an offset, which should make it easier to
+	// understand whenever it is used
+	uint32 NumAdditionalBones = UE::USDSkeletalDataConversion::Private::HasMultipleRootBones(InUsdSkeletonQuery) ? 1 : 0;
+
 	// If we have zero bones on our skeleton we'll generate a dummy "Root" bone just so that Unreal can have a USkeleton asset,
 	// so we have to check for that case
-	if ((NumSkelQueryJoints != NumBones) && !(NumSkelQueryJoints == 0 && NumBones == 1 && BoneInfo[0].Name == TEXT("Root")))
+	if (((NumBonesInUsd + NumAdditionalBones) != (uint32)NumBonesInUE)
+		&& !(NumBonesInUsd == 0 && NumBonesInUE == 1 && BoneInfo[0].Name == TEXT("Root")))
 	{
 		return false;
 	}
@@ -2135,9 +2214,9 @@ bool UsdToUnreal::ConvertSkelAnim(
 		FScopedUsdAllocs Allocs;
 
 		TArray<FRawAnimSequenceTrack> JointTracks;
-		JointTracks.SetNum(NumBones);
+		JointTracks.SetNum(NumBonesInUE);
 
-		for ( int32 BoneIndex = 0; BoneIndex < NumBones; ++BoneIndex )
+		for ( int32 BoneIndex = 0; BoneIndex < NumBonesInUE; ++BoneIndex )
 		{
 			FRawAnimSequenceTrack& JointTrack = JointTracks[ BoneIndex ];
 			JointTrack.PosKeys.Reserve(NumBakedFrames);
@@ -2153,10 +2232,17 @@ bool UsdToUnreal::ConvertSkelAnim(
 			const double StageFrameTimeCodes = StageStartTimeCode + FrameIndex * StageBakeIntervalTimeCodes;
 
 			InUsdSkeletonQuery.ComputeJointLocalTransforms( &UsdJointTransforms, StageFrameTimeCodes );
-			for ( int32 BoneIndex = 0; BoneIndex < NumBones; ++BoneIndex )
+			for ( int32 BoneIndex = 0; BoneIndex < NumBonesInUE; ++BoneIndex )
 			{
-				pxr::GfMatrix4d& UsdJointTransform = UsdJointTransforms[ BoneIndex ];
-				FTransform UEJointTransform = UsdToUnreal::ConvertMatrix( StageInfo, UsdJointTransform );
+				// UsdJointTransforms will never have a transform value for our AdditionalBones that we manually added
+				// (inserted "true root" bone), so we have to have this annoying check here to redirect the bone indices
+				// properly when querying USD with them
+				FTransform UEJointTransform = FTransform::Identity;
+				if (NumAdditionalBones == 0 || BoneIndex != 0)
+				{
+					pxr::GfMatrix4d& UsdJointTransform = UsdJointTransforms[BoneIndex - NumAdditionalBones];
+					UEJointTransform = UsdToUnreal::ConvertMatrix(StageInfo, UsdJointTransform);
+				}
 
 				// Concatenate the root bone transform with the transform track actually present on the skel root as a
 				// whole
@@ -2186,7 +2272,7 @@ bool UsdToUnreal::ConvertSkelAnim(
 			}
 		}
 
-		for ( int32 BoneIndex = 0; BoneIndex < NumBones; ++BoneIndex )
+		for ( int32 BoneIndex = 0; BoneIndex < NumBonesInUE; ++BoneIndex )
 		{
 			Controller.AddBoneCurve( BoneInfo[ BoneIndex ].Name, bShouldTransact );
 			Controller.SetBoneTrackKeys( BoneInfo[ BoneIndex ].Name, JointTracks[ BoneIndex ].PosKeys, JointTracks[ BoneIndex ].RotKeys, JointTracks[ BoneIndex ].ScaleKeys, bShouldTransact );
