USD: Add support for importing and exporting prim metadata.

diff --git a/Engine/Plugins/Importers/USDImporter/Source/USDClasses/Public/USDAssetUserData.h b/Engine/Plugins/Importers/USDImporter/Source/USDClasses/Public/USDAssetUserData.h
--- a/Engine/Plugins/Importers/USDImporter/Source/USDClasses/Public/USDAssetUserData.h
+++ b/Engine/Plugins/Importers/USDImporter/Source/USDClasses/Public/USDAssetUserData.h
@@ -4,6 +4,8 @@
 
 #include "Engine/AssetUserData.h"
 
+#include "USDMetadata.h"
+
 #include "USDAssetUserData.generated.h"
 
 UCLASS()
@@ -12,8 +14,16 @@ class USDCLASSES_API UUsdAssetUserData : public UAssetUserData
 	GENERATED_BODY()
 
 public:
-	UPROPERTY()
+	// Paths to prims that generated the asset that owns this AssetUserData
+	UPROPERTY(EditAnywhere, BlueprintReadWrite, Category = "USD")
 	TArray<FString> PrimPaths;
+
+	// Holds metadata collected for this asset, from all relevant Source prims.
+	// The asset that owns this user data may be shared via the asset cache, and reused for
+	// even entirely different stages. This map lets us keep track of which stage owns which
+	// bits of metadata
+	UPROPERTY(EditAnywhere, BlueprintReadWrite, Category = "USD")
+	TMap<FString, FUsdCombinedPrimMetadata> StageIdentifierToMetadata;
 };
 
 UCLASS()
@@ -91,6 +101,18 @@ public:
 	TMap<FString, int32> PrimvarToUVIndex;
 };
 
+/** We assign these to UGeometryCaches generated from USD */
+UCLASS()
+class USDCLASSES_API UUsdGeometryCacheAssetUserData : public UUsdMeshAssetUserData
+{
+	GENERATED_BODY()
+
+public:
+	// Check analogous comment on UUsdAnimSequenceAssetUserData
+	UPROPERTY()
+	float LayerStartOffsetSeconds = 0.0f;
+};
+
 /**
  * We assign these to persistent LevelSequences that bind to one of the actors/components that the stage actor spawns.
  * We need this as part of a mechanism to automatically repair those bindings when they break if we close/reload the stage.
diff --git a/Engine/Plugins/Importers/USDImporter/Source/USDExporter/Private/AnimSequenceExporterUSD.cpp b/Engine/Plugins/Importers/USDImporter/Source/USDExporter/Private/AnimSequenceExporterUSD.cpp
--- a/Engine/Plugins/Importers/USDImporter/Source/USDExporter/Private/AnimSequenceExporterUSD.cpp
+++ b/Engine/Plugins/Importers/USDImporter/Source/USDExporter/Private/AnimSequenceExporterUSD.cpp
@@ -2,19 +2,14 @@
 
 #include "AnimSequenceExporterUSD.h"
 
-#include "AnimSequenceExporterUSDOptions.h"
-#include "Engine/SkeletalMesh.h"
-#include "EngineAnalytics.h"
-#include "MaterialExporterUSD.h"
-#include "Misc/EngineVersion.h"
 #include "SkeletalMeshExporterUSDOptions.h"
 #include "UnrealUSDWrapper.h"
-#include "USDClassesModule.h"
 #include "USDErrorUtils.h"
 #include "USDExporterModule.h"
 #include "USDLayerUtils.h"
 #include "USDLog.h"
 #include "USDOptionsWindow.h"
+#include "USDPrimConversion.h"
 #include "USDSkeletalDataConversion.h"
 #include "USDUnrealAssetInfo.h"
 
@@ -24,7 +19,12 @@
 #include "UsdWrappers/UsdStage.h"
 
 #include "Animation/AnimSequence.h"
+#include "AnimSequenceExporterUSDOptions.h"
 #include "AssetExportTask.h"
+#include "Engine/SkeletalMesh.h"
+#include "EngineAnalytics.h"
+#include "MaterialExporterUSD.h"
+#include "Misc/EngineVersion.h"
 #include "UObject/GCObjectScopeGuard.h"
 
 namespace UE::AnimSequenceExporterUSD::Private
@@ -162,6 +162,7 @@ bool UAnimSequenceExporterUSD::ExportBinary( UObject* Object, const TCHAR* Type,
 			USkeletalMeshExporterUSDOptions* SkeletalMeshOptions = GetMutableDefault<USkeletalMeshExporterUSDOptions>();
 			SkeletalMeshOptions->StageOptions = Options->StageOptions;
 			SkeletalMeshOptions->MeshAssetOptions = Options->PreviewMeshOptions;
+			SkeletalMeshOptions->MetadataOptions = Options->MetadataOptions;
 			SkeletalMeshOptions->bReExportIdenticalAssets = Options->bReExportIdenticalAssets;
 
 			UAssetExportTask* LevelExportTask = NewObject<UAssetExportTask>();
@@ -343,7 +344,7 @@ bool UAnimSequenceExporterUSD::ExportBinary( UObject* Object, const TCHAR* Type,
 
 	UnrealToUsd::ConvertAnimSequence( AnimSequence, SkelAnimPrim );
 
-	// Write asset info now that we finished exporting
+	if (Options->MetadataOptions.bExportAssetInfo)
 	{
 		FUsdUnrealAssetInfo Info;
 		Info.Name = AnimSequence->GetName();
@@ -357,6 +358,19 @@ bool UAnimSequenceExporterUSD::ExportBinary( UObject* Object, const TCHAR* Type,
 		UsdUtils::SetPrimAssetInfo( SkelAnimPrim, Info );
 	}
 
+	if (Options->MetadataOptions.bExportAssetMetadata)
+	{
+		if (UUsdAssetUserData* UserData = UsdUtils::GetAssetUserData(AnimSequence))
+		{
+			UnrealToUsd::ConvertMetadata(
+				UserData,
+				SkelAnimPrim,
+				Options->MetadataOptions.BlockedPrefixFilters,
+				Options->MetadataOptions.bInvertFilters
+			);
+		}
+	}
+
 	AnimationStage.GetRootLayer().Save();
 
 	// Analytics
diff --git a/Engine/Plugins/Importers/USDImporter/Source/USDExporter/Private/AnimSequenceExporterUSDOptions.cpp b/Engine/Plugins/Importers/USDImporter/Source/USDExporter/Private/AnimSequenceExporterUSDOptions.cpp
--- a/Engine/Plugins/Importers/USDImporter/Source/USDExporter/Private/AnimSequenceExporterUSDOptions.cpp
+++ b/Engine/Plugins/Importers/USDImporter/Source/USDExporter/Private/AnimSequenceExporterUSDOptions.cpp
@@ -4,22 +4,21 @@
 
 #include "AnalyticsEventAttribute.h"
 
-void UsdUtils::AddAnalyticsAttributes(
-	const UAnimSequenceExporterUSDOptions& Options,
-	TArray< FAnalyticsEventAttribute >& InOutAttributes
-)
+void UsdUtils::AddAnalyticsAttributes(const UAnimSequenceExporterUSDOptions& Options, TArray<FAnalyticsEventAttribute>& InOutAttributes)
 {
-	UsdUtils::AddAnalyticsAttributes( Options.StageOptions, InOutAttributes );
-	InOutAttributes.Emplace( TEXT( "ExportPreviewMesh" ), LexToString( Options.bExportPreviewMesh ) );
-	if ( Options.bExportPreviewMesh )
+	UsdUtils::AddAnalyticsAttributes(Options.StageOptions, InOutAttributes);
+	InOutAttributes.Emplace(TEXT("ExportPreviewMesh"), LexToString(Options.bExportPreviewMesh));
+	if (Options.bExportPreviewMesh)
 	{
-		UsdUtils::AddAnalyticsAttributes( Options.PreviewMeshOptions, InOutAttributes );
+		UsdUtils::AddAnalyticsAttributes(Options.PreviewMeshOptions, InOutAttributes);
 	}
-	InOutAttributes.Emplace( TEXT( "ReExportIdenticalAssets" ), Options.bReExportIdenticalAssets );
+	UsdUtils::AddAnalyticsAttributes(Options.MetadataOptions, InOutAttributes);
+	InOutAttributes.Emplace(TEXT("ReExportIdenticalAssets"), Options.bReExportIdenticalAssets);
 }
 
-void UsdUtils::HashForAnimSequenceExport( const UAnimSequenceExporterUSDOptions& Options, FSHA1& HashToUpdate )
+void UsdUtils::HashForAnimSequenceExport(const UAnimSequenceExporterUSDOptions& Options, FSHA1& HashToUpdate)
 {
-	UsdUtils::HashForExport( Options.StageOptions, HashToUpdate );
-	HashToUpdate.Update( reinterpret_cast< const uint8* >( &Options.bExportPreviewMesh ), sizeof( Options.bExportPreviewMesh ) );
+	UsdUtils::HashForExport(Options.StageOptions, HashToUpdate);
+	UsdUtils::HashForExport(Options.MetadataOptions, HashToUpdate);
+	HashToUpdate.Update(reinterpret_cast<const uint8*>(&Options.bExportPreviewMesh), sizeof(Options.bExportPreviewMesh));
 }
diff --git a/Engine/Plugins/Importers/USDImporter/Source/USDExporter/Private/GeometryCacheExporterUSD.cpp b/Engine/Plugins/Importers/USDImporter/Source/USDExporter/Private/GeometryCacheExporterUSD.cpp
--- a/Engine/Plugins/Importers/USDImporter/Source/USDExporter/Private/GeometryCacheExporterUSD.cpp
+++ b/Engine/Plugins/Importers/USDImporter/Source/USDExporter/Private/GeometryCacheExporterUSD.cpp
@@ -2,18 +2,15 @@
 
 #include "GeometryCacheExporterUSD.h"
 
-#include "AssetExportTask.h"
-#include "EngineAnalytics.h"
-#include "GeometryCache.h"
 #include "GeometryCacheExporterUSDOptions.h"
 #include "MaterialExporterUSD.h"
-#include "Misc/EngineVersion.h"
 #include "UnrealUSDWrapper.h"
 #include "USDConversionUtils.h"
 #include "USDExporterModule.h"
 #include "USDGeomMeshConversion.h"
 #include "USDLog.h"
 #include "USDOptionsWindow.h"
+#include "USDPrimConversion.h"
 #include "USDUnrealAssetInfo.h"
 
 #include "UsdWrappers/SdfLayer.h"
@@ -21,6 +18,11 @@
 #include "UsdWrappers/UsdPrim.h"
 #include "UsdWrappers/UsdStage.h"
 
+#include "AssetExportTask.h"
+#include "EngineAnalytics.h"
+#include "GeometryCache.h"
+#include "Misc/EngineVersion.h"
+
 namespace UE::GeometryCacheExporterUSD::Private
 {
 	void SendAnalytics(
@@ -217,6 +219,7 @@ bool UGeometryCacheExporterUSD::ExportBinary(UObject* Object, const TCHAR* Type,
 					UMaterialExporterUsd::ExportMaterialsForStage(
 						MaterialsToBake.Array(),
 						Options->MeshAssetOptions.MaterialBakingOptions,
+						Options->MetadataOptions,
 						UExporter::CurrentFilename,
 						bIsAssetLayer,
 						Options->MeshAssetOptions.bUsePayload,
@@ -282,19 +285,34 @@ bool UGeometryCacheExporterUSD::ExportBinary(UObject* Object, const TCHAR* Type,
 
 	UnrealToUsd::ConvertGeometryCache(GeometryCache, RootPrim, &AssetStage);
 
-	// Write asset info now that we finished exporting
-	if( UE::FUsdPrim AssetDefaultPrim = AssetStage.GetDefaultPrim() )
+	if (UE::FUsdPrim AssetDefaultPrim = AssetStage.GetDefaultPrim())
 	{
-		FUsdUnrealAssetInfo Info;
-		Info.Name = GeometryCache->GetName();
-		Info.Identifier = UExporter::CurrentFilename;
-		Info.Version = DDCKeyHash;
-		Info.UnrealContentPath = GeometryCache->GetPathName();
-		Info.UnrealAssetType = GeometryCache->GetClass()->GetName();
-		Info.UnrealExportTime = FDateTime::Now().ToString();
-		Info.UnrealEngineVersion = FEngineVersion::Current().ToString();
-
-		UsdUtils::SetPrimAssetInfo(AssetDefaultPrim, Info);
+		if (Options->MetadataOptions.bExportAssetInfo)
+		{
+			FUsdUnrealAssetInfo Info;
+			Info.Name = GeometryCache->GetName();
+			Info.Identifier = UExporter::CurrentFilename;
+			Info.Version = DDCKeyHash;
+			Info.UnrealContentPath = GeometryCache->GetPathName();
+			Info.UnrealAssetType = GeometryCache->GetClass()->GetName();
+			Info.UnrealExportTime = FDateTime::Now().ToString();
+			Info.UnrealEngineVersion = FEngineVersion::Current().ToString();
+
+			UsdUtils::SetPrimAssetInfo(AssetDefaultPrim, Info);
+		}
+
+		if (Options->MetadataOptions.bExportAssetMetadata)
+		{
+			if (UUsdAssetUserData* UserData = UsdUtils::GetAssetUserData(GeometryCache))
+			{
+				UnrealToUsd::ConvertMetadata(
+					UserData,
+					AssetDefaultPrim,
+					Options->MetadataOptions.BlockedPrefixFilters,
+					Options->MetadataOptions.bInvertFilters
+				);
+			}
+		}
 	}
 
 	// Bake materials and replace unrealMaterials with references to the baked files.
@@ -310,6 +328,7 @@ bool UGeometryCacheExporterUSD::ExportBinary(UObject* Object, const TCHAR* Type,
 		UMaterialExporterUsd::ExportMaterialsForStage(
 			MaterialsToBake.Array(),
 			Options->MeshAssetOptions.MaterialBakingOptions,
+			Options->MetadataOptions,
 			AssetStage.GetRootLayer().GetRealPath(),
 			bIsAssetLayer,
 			Options->MeshAssetOptions.bUsePayload,
diff --git a/Engine/Plugins/Importers/USDImporter/Source/USDExporter/Private/LevelExporterUSDOptions.cpp b/Engine/Plugins/Importers/USDImporter/Source/USDExporter/Private/LevelExporterUSDOptions.cpp
--- a/Engine/Plugins/Importers/USDImporter/Source/USDExporter/Private/LevelExporterUSDOptions.cpp
+++ b/Engine/Plugins/Importers/USDImporter/Source/USDExporter/Private/LevelExporterUSDOptions.cpp
@@ -2,6 +2,7 @@
 
 #include "LevelExporterUSDOptions.h"
 
+#include "UnrealUSDWrapper.h"
 #include "USDAssetOptions.h"
 #include "USDExporterModule.h"
 
@@ -10,84 +11,77 @@
 TArray<FString> ULevelExporterUSDOptions::GetUsdExtensions()
 {
 	TArray<FString> Extensions = UnrealUSDWrapper::GetNativeFileFormats();
-	Extensions.Remove( TEXT( "usdz" ) );
+	Extensions.Remove(TEXT("usdz"));
 	return Extensions;
 }
 
-void UsdUtils::AddAnalyticsAttributes(
-	const FLevelExporterUSDOptionsInner& Options,
-	TArray< FAnalyticsEventAttribute >& InOutAttributes
-)
+void UsdUtils::AddAnalyticsAttributes(const FLevelExporterUSDOptionsInner& Options, TArray<FAnalyticsEventAttribute>& InOutAttributes)
 {
-	InOutAttributes.Emplace( TEXT( "SelectionOnly" ), Options.bSelectionOnly );
-	InOutAttributes.Emplace( TEXT( "ExportActorFolders" ), Options.bExportActorFolders );
-	InOutAttributes.Emplace( TEXT( "IgnoreSequencerAnimations" ), Options.bIgnoreSequencerAnimations );
-	InOutAttributes.Emplace( TEXT( "ExportFoliageOnActorsLayer" ), Options.bExportFoliageOnActorsLayer );
-	UsdUtils::AddAnalyticsAttributes( Options.AssetOptions, InOutAttributes );
-	InOutAttributes.Emplace( TEXT( "LowestLandscapeLOD" ), LexToString( Options.LowestLandscapeLOD ) );
-	InOutAttributes.Emplace( TEXT( "HighestLandscapeLOD" ), LexToString( Options.HighestLandscapeLOD ) );
-	InOutAttributes.Emplace( TEXT( "LandscapeBakeResolution" ), Options.LandscapeBakeResolution.ToString() );
-	InOutAttributes.Emplace( TEXT( "ExportSublayers" ), LexToString( Options.bExportSublayers ) );
-	InOutAttributes.Emplace( TEXT( "NumLevelsToIgnore" ), LexToString( Options.LevelsToIgnore.Num() ) );
+	InOutAttributes.Emplace(TEXT("SelectionOnly"), Options.bSelectionOnly);
+	InOutAttributes.Emplace(TEXT("ExportActorFolders"), Options.bExportActorFolders);
+	InOutAttributes.Emplace(TEXT("IgnoreSequencerAnimations"), Options.bIgnoreSequencerAnimations);
+	InOutAttributes.Emplace(TEXT("ExportFoliageOnActorsLayer"), Options.bExportFoliageOnActorsLayer);
+	UsdUtils::AddAnalyticsAttributes(Options.AssetOptions, InOutAttributes);
+	UsdUtils::AddAnalyticsAttributes(Options.MetadataOptions, InOutAttributes);
+	InOutAttributes.Emplace(TEXT("LowestLandscapeLOD"), LexToString(Options.LowestLandscapeLOD));
+	InOutAttributes.Emplace(TEXT("HighestLandscapeLOD"), LexToString(Options.HighestLandscapeLOD));
+	InOutAttributes.Emplace(TEXT("LandscapeBakeResolution"), Options.LandscapeBakeResolution.ToString());
+	InOutAttributes.Emplace(TEXT("ExportSublayers"), LexToString(Options.bExportSublayers));
+	InOutAttributes.Emplace(TEXT("NumLevelsToIgnore"), LexToString(Options.LevelsToIgnore.Num()));
 }
 
-void UsdUtils::AddAnalyticsAttributes(
-	const ULevelExporterUSDOptions& Options,
-	TArray< FAnalyticsEventAttribute >& InOutAttributes
-)
+void UsdUtils::AddAnalyticsAttributes(const ULevelExporterUSDOptions& Options, TArray<FAnalyticsEventAttribute>& InOutAttributes)
 {
-	UsdUtils::AddAnalyticsAttributes( Options.StageOptions, InOutAttributes );
-	InOutAttributes.Emplace( TEXT( "StartTimeCode" ), LexToString( Options.StartTimeCode ) );
-	InOutAttributes.Emplace( TEXT( "EndTimeCode" ), LexToString( Options.EndTimeCode ) );
-	UsdUtils::AddAnalyticsAttributes( Options.Inner, InOutAttributes );
-	InOutAttributes.Emplace( TEXT( "EndTimeCode" ), LexToString( Options.EndTimeCode ) );
-	InOutAttributes.Emplace(
-		TEXT( "ReExportIdenticalLevelsAndSequences" ),
-		Options.bReExportIdenticalLevelsAndSequences
-	);
-	InOutAttributes.Emplace( TEXT( "ReExportIdenticalAssets" ), Options.bReExportIdenticalAssets );
+	UsdUtils::AddAnalyticsAttributes(Options.StageOptions, InOutAttributes);
+	InOutAttributes.Emplace(TEXT("StartTimeCode"), LexToString(Options.StartTimeCode));
+	InOutAttributes.Emplace(TEXT("EndTimeCode"), LexToString(Options.EndTimeCode));
+	UsdUtils::AddAnalyticsAttributes(Options.Inner, InOutAttributes);
+	InOutAttributes.Emplace(TEXT("EndTimeCode"), LexToString(Options.EndTimeCode));
+	InOutAttributes.Emplace(TEXT("ReExportIdenticalLevelsAndSequences"), Options.bReExportIdenticalLevelsAndSequences);
+	InOutAttributes.Emplace(TEXT("ReExportIdenticalAssets"), Options.bReExportIdenticalAssets);
 }
 
-void UsdUtils::HashForLevelExport( const FLevelExporterUSDOptionsInner& Options, FSHA1& HashToUpdate )
+void UsdUtils::HashForLevelExport(const FLevelExporterUSDOptionsInner& Options, FSHA1& HashToUpdate)
 {
-	HashToUpdate.Update( reinterpret_cast< const uint8* >( &Options.bSelectionOnly ), sizeof( Options.bSelectionOnly ) );
-	if ( Options.bSelectionOnly )
+	HashToUpdate.Update(reinterpret_cast<const uint8*>(&Options.bSelectionOnly), sizeof(Options.bSelectionOnly));
+	if (Options.bSelectionOnly)
 	{
-		IUsdExporterModule::HashEditorSelection( HashToUpdate );
+		IUsdExporterModule::HashEditorSelection(HashToUpdate);
 	}
 
-	HashToUpdate.Update( reinterpret_cast< const uint8* >( &Options.bExportActorFolders ), sizeof( Options.bExportActorFolders ) );
-	HashToUpdate.Update( reinterpret_cast< const uint8* >( &Options.bIgnoreSequencerAnimations ), sizeof( Options.bIgnoreSequencerAnimations ) );
-	HashToUpdate.Update( reinterpret_cast< const uint8* >( &Options.bExportFoliageOnActorsLayer ), sizeof( Options.bExportFoliageOnActorsLayer ) );
+	HashToUpdate.Update(reinterpret_cast<const uint8*>(&Options.bExportActorFolders), sizeof(Options.bExportActorFolders));
+	HashToUpdate.Update(reinterpret_cast<const uint8*>(&Options.bIgnoreSequencerAnimations), sizeof(Options.bIgnoreSequencerAnimations));
+	HashToUpdate.Update(reinterpret_cast<const uint8*>(&Options.bExportFoliageOnActorsLayer), sizeof(Options.bExportFoliageOnActorsLayer));
 
 	// If we changed where we want the assets exported we need to re-export them and update the reference paths
-	HashToUpdate.UpdateWithString( *Options.AssetFolder.Path, Options.AssetFolder.Path.Len() );
+	HashToUpdate.UpdateWithString(*Options.AssetFolder.Path, Options.AssetFolder.Path.Len());
 
 	// This affects how we author material overrides
-	HashToUpdate.Update( reinterpret_cast< const uint8* >( &Options.AssetOptions.bBakeMaterials ), sizeof( Options.AssetOptions.bBakeMaterials ) );
+	HashToUpdate.Update(reinterpret_cast<const uint8*>(&Options.AssetOptions.bBakeMaterials), sizeof(Options.AssetOptions.bBakeMaterials));
 
 	// This affects how we author material overrides and landscapes
-	HashToUpdate.Update( reinterpret_cast< const uint8* >( &Options.AssetOptions.bUsePayload ), sizeof( Options.AssetOptions.bUsePayload ) );
-	if ( Options.AssetOptions.bUsePayload )
+	HashToUpdate.Update(reinterpret_cast<const uint8*>(&Options.AssetOptions.bUsePayload), sizeof(Options.AssetOptions.bUsePayload));
+	if (Options.AssetOptions.bUsePayload)
 	{
-		HashToUpdate.UpdateWithString( *Options.AssetOptions.PayloadFormat, Options.AssetOptions.PayloadFormat.Len() );
+		HashToUpdate.UpdateWithString(*Options.AssetOptions.PayloadFormat, Options.AssetOptions.PayloadFormat.Len());
 	}
 
 	// If we changed to/from exporting LODs, how we author material overrides may also need to change
-	HashToUpdate.Update( reinterpret_cast< const uint8* >( &Options.AssetOptions.LowestMeshLOD ), sizeof( Options.AssetOptions.LowestMeshLOD ) );
-	HashToUpdate.Update( reinterpret_cast< const uint8* >( &Options.AssetOptions.HighestMeshLOD ), sizeof( Options.AssetOptions.HighestMeshLOD ) );
+	HashToUpdate.Update(reinterpret_cast<const uint8*>(&Options.AssetOptions.LowestMeshLOD), sizeof(Options.AssetOptions.LowestMeshLOD));
+	HashToUpdate.Update(reinterpret_cast<const uint8*>(&Options.AssetOptions.HighestMeshLOD), sizeof(Options.AssetOptions.HighestMeshLOD));
 
-	HashToUpdate.Update( reinterpret_cast< const uint8* >( &Options.LowestLandscapeLOD ), sizeof( Options.LowestLandscapeLOD ) );
-	HashToUpdate.Update( reinterpret_cast< const uint8* >( &Options.HighestLandscapeLOD ), sizeof( Options.HighestLandscapeLOD ) );
-	HashToUpdate.Update( reinterpret_cast< const uint8* >( &Options.LandscapeBakeResolution ), sizeof( Options.LandscapeBakeResolution ) );
-	HashToUpdate.Update( reinterpret_cast< const uint8* >( &Options.bExportSublayers ), sizeof( Options.bExportSublayers ) );
+	UsdUtils::HashForExport(Options.MetadataOptions, HashToUpdate);
+
+	HashToUpdate.Update(reinterpret_cast<const uint8*>(&Options.LowestLandscapeLOD), sizeof(Options.LowestLandscapeLOD));
+	HashToUpdate.Update(reinterpret_cast<const uint8*>(&Options.HighestLandscapeLOD), sizeof(Options.HighestLandscapeLOD));
+	HashToUpdate.Update(reinterpret_cast<const uint8*>(&Options.LandscapeBakeResolution), sizeof(Options.LandscapeBakeResolution));
+	HashToUpdate.Update(reinterpret_cast<const uint8*>(&Options.bExportSublayers), sizeof(Options.bExportSublayers));
 }
 
-void UsdUtils::HashForLevelExport( const ULevelExporterUSDOptions& Options, FSHA1& HashToUpdate )
+void UsdUtils::HashForLevelExport(const ULevelExporterUSDOptions& Options, FSHA1& HashToUpdate)
 {
-	UsdUtils::HashForExport( Options.StageOptions, HashToUpdate );
-	HashToUpdate.Update( reinterpret_cast< const uint8* >( &Options.StartTimeCode ), sizeof( Options.StartTimeCode ) );
-	HashToUpdate.Update( reinterpret_cast< const uint8* >( &Options.EndTimeCode ), sizeof( Options.EndTimeCode ) );
-	UsdUtils::HashForLevelExport( Options.Inner, HashToUpdate );
+	UsdUtils::HashForExport(Options.StageOptions, HashToUpdate);
+	HashToUpdate.Update(reinterpret_cast<const uint8*>(&Options.StartTimeCode), sizeof(Options.StartTimeCode));
+	HashToUpdate.Update(reinterpret_cast<const uint8*>(&Options.EndTimeCode), sizeof(Options.EndTimeCode));
+	UsdUtils::HashForLevelExport(Options.Inner, HashToUpdate);
 }
-
diff --git a/Engine/Plugins/Importers/USDImporter/Source/USDExporter/Private/LevelExporterUSDOptionsCustomization.cpp b/Engine/Plugins/Importers/USDImporter/Source/USDExporter/Private/LevelExporterUSDOptionsCustomization.cpp
--- a/Engine/Plugins/Importers/USDImporter/Source/USDExporter/Private/LevelExporterUSDOptionsCustomization.cpp
+++ b/Engine/Plugins/Importers/USDImporter/Source/USDExporter/Private/LevelExporterUSDOptionsCustomization.cpp
@@ -239,6 +239,7 @@ void FLevelExporterUSDOptionsCustomization::CustomizeDetails(IDetailLayoutBuilde
 	DetailLayoutBuilder.EditCategory( TEXT( "Stage options" ) );
 	DetailLayoutBuilder.EditCategory( TEXT( "Export settings" ) );
 	IDetailCategoryBuilder& AssetOptionsCategory = DetailLayoutBuilder.EditCategory( TEXT( "Asset options" ) );
+	IDetailCategoryBuilder& MetadataOptionsCategory = DetailLayoutBuilder.EditCategory( TEXT( "Metadata options" ) );
 
 	// Promote all AssetOptions up a level on LevelExportUsdOptions or else we'll end up with a property named AssetOptions inside the AssetOptions category
 	// This is the same effect as ShowOnlyInnerProperties, but in this case we need to do it manually as it doesn't work recursively
@@ -256,6 +257,21 @@ void FLevelExporterUSDOptionsCustomization::CustomizeDetails(IDetailLayoutBuilde
 			}
 		}
 	}
+	// Do the same for MetadataOptions...
+	if (TSharedPtr<IPropertyHandle> MetadataOptions = DetailLayoutBuilder.GetProperty(TEXT("Inner.MetadataOptions")))
+	{
+		DetailLayoutBuilder.HideProperty(MetadataOptions);
+
+		uint32 NumChildren = 0;
+		if (MetadataOptions->GetNumChildren(NumChildren) == FPropertyAccess::Result::Success)
+		{
+			for (uint32 Index = 0; Index < NumChildren; ++Index)
+			{
+				TSharedPtr<IPropertyHandle> ChildProperty = MetadataOptions->GetChildHandle(Index);
+				MetadataOptionsCategory.AddProperty(ChildProperty);
+			}
+		}
+	}
 
 	// Hide the textures dir property because we'll add multiple textures folders (one next to each exported material)
 	TSharedPtr<IPropertyHandle> TexturesDirProperty = DetailLayoutBuilder.GetProperty( TexturesDirPropName );
diff --git a/Engine/Plugins/Importers/USDImporter/Source/USDExporter/Private/LevelSequenceExporterUSD.cpp b/Engine/Plugins/Importers/USDImporter/Source/USDExporter/Private/LevelSequenceExporterUSD.cpp
--- a/Engine/Plugins/Importers/USDImporter/Source/USDExporter/Private/LevelSequenceExporterUSD.cpp
+++ b/Engine/Plugins/Importers/USDImporter/Source/USDExporter/Private/LevelSequenceExporterUSD.cpp
@@ -1415,19 +1415,34 @@ namespace UE
 					UsdUtils::InsertSubLayer( UsdStage.GetRootLayer(), *Context.LevelFilePath );
 				}
 
-				// Write asset info now that we finished exporting
-				if ( UE::FUsdPrim AssetDefaultPrim = UsdStage.GetDefaultPrim() )
+				if (UE::FUsdPrim AssetDefaultPrim = UsdStage.GetDefaultPrim())
 				{
-					FUsdUnrealAssetInfo Info;
-					Info.Name = MovieSceneSequence.GetName();
-					Info.Identifier = UniqueFilePath;
-					Info.Version = LevelSequenceVersion;
-					Info.UnrealContentPath = MovieSceneSequence.GetPathName();
-					Info.UnrealAssetType = MovieSceneSequence.GetClass()->GetName();
-					Info.UnrealExportTime = FDateTime::Now().ToString();
-					Info.UnrealEngineVersion = FEngineVersion::Current().ToString();
-
-					UsdUtils::SetPrimAssetInfo( AssetDefaultPrim, Info );
+					if (Context.ExportOptions->LevelExportOptions.MetadataOptions.bExportAssetInfo)
+					{
+						FUsdUnrealAssetInfo Info;
+						Info.Name = MovieSceneSequence.GetName();
+						Info.Identifier = UniqueFilePath;
+						Info.Version = LevelSequenceVersion;
+						Info.UnrealContentPath = MovieSceneSequence.GetPathName();
+						Info.UnrealAssetType = MovieSceneSequence.GetClass()->GetName();
+						Info.UnrealExportTime = FDateTime::Now().ToString();
+						Info.UnrealEngineVersion = FEngineVersion::Current().ToString();
+
+						UsdUtils::SetPrimAssetInfo(AssetDefaultPrim, Info);
+					}
+
+					if (Context.ExportOptions->LevelExportOptions.MetadataOptions.bExportAssetMetadata)
+					{
+						if (UUsdAssetUserData* UserData = UsdUtils::GetAssetUserData(Cast<ULevelSequence>(&MovieSceneSequence)))
+						{
+							UnrealToUsd::ConvertMetadata(
+								UserData,
+								AssetDefaultPrim,
+								Context.ExportOptions->LevelExportOptions.MetadataOptions.BlockedPrefixFilters,
+								Context.ExportOptions->LevelExportOptions.MetadataOptions.bInvertFilters
+							);
+						}
+					}
 				}
 
 				Context.ExportedMovieScenes.Add( &MovieSceneSequence, UniqueFilePath );
diff --git a/Engine/Plugins/Importers/USDImporter/Source/USDExporter/Private/MaterialExporterUSD.cpp b/Engine/Plugins/Importers/USDImporter/Source/USDExporter/Private/MaterialExporterUSD.cpp
--- a/Engine/Plugins/Importers/USDImporter/Source/USDExporter/Private/MaterialExporterUSD.cpp
+++ b/Engine/Plugins/Importers/USDImporter/Source/USDExporter/Private/MaterialExporterUSD.cpp
@@ -3,17 +3,15 @@
 #include "MaterialExporterUSD.h"
 
 #include "MaterialExporterUSDOptions.h"
-#include "Materials/Material.h"
-#include "MaterialShared.h"
-#include "Misc/EngineVersion.h"
 #include "UnrealUSDWrapper.h"
 #include "USDClassesModule.h"
 #include "USDConversionUtils.h"
 #include "USDExporterModule.h"
 #include "USDGeomMeshConversion.h"
 #include "USDLog.h"
-#include "USDMemory.h"
+#include "USDMetadataExportOptions.h"
 #include "USDOptionsWindow.h"
+#include "USDPrimConversion.h"
 #include "USDShadeConversion.h"
 #include "USDUnrealAssetInfo.h"
 
@@ -26,14 +24,14 @@
 #include "Engine/Font.h"
 #include "Engine/Texture.h"
 #include "EngineAnalytics.h"
-#include "HAL/FileManager.h"
-#include "IMaterialBakingModule.h"
 #include "MaterialOptions.h"
+#include "Materials/Material.h"
 #include "Materials/MaterialInstance.h"
 #include "Materials/MaterialInstanceConstant.h"
 #include "Materials/MaterialInterface.h"
+#include "MaterialShared.h"
+#include "Misc/EngineVersion.h"
 #include "Misc/Paths.h"
-#include "Modules/ModuleManager.h"
 #include "UObject/UObjectGlobals.h"
 #include "VT/RuntimeVirtualTexture.h"
 
@@ -265,6 +263,7 @@ bool UMaterialExporterUsd::ExportBinary( UObject* Object, const TCHAR* Type, FAr
 	return UMaterialExporterUsd::ExportMaterial(
 		*Material,
 		Options->MaterialBakingOptions,
+		Options->MetadataOptions,
 		FFilePath{ UExporter::CurrentFilename },
 		ExportTask->bReplaceIdentical,
 		Options->bReExportIdenticalAssets,
@@ -278,6 +277,7 @@ bool UMaterialExporterUsd::ExportBinary( UObject* Object, const TCHAR* Type, FAr
 bool UMaterialExporterUsd::ExportMaterial(
 	const UMaterialInterface& Material,
 	const FUsdMaterialBakingOptions& Options,
+	const FUsdMetadataExportOptions& MetadataOptions,
 	const FFilePath& FilePath,
 	bool bReplaceIdentical,
 	bool bReExportIdenticalAssets,
@@ -372,7 +372,7 @@ bool UMaterialExporterUsd::ExportMaterial(
 		Options.bConstantColorAsSingleValue
 	);
 
-	// Write asset info now that we finished exporting
+	if (MetadataOptions.bExportAssetInfo)
 	{
 		FUsdUnrealAssetInfo Info;
 		Info.Name = Material.GetName();
@@ -383,7 +383,20 @@ bool UMaterialExporterUsd::ExportMaterial(
 		Info.UnrealExportTime = FDateTime::Now().ToString();
 		Info.UnrealEngineVersion = FEngineVersion::Current().ToString();
 
-		UsdUtils::SetPrimAssetInfo( RootPrim, Info );
+		UsdUtils::SetPrimAssetInfo(RootPrim, Info);
+	}
+
+	if (MetadataOptions.bExportAssetMetadata)
+	{
+		if (UUsdAssetUserData* UserData = UsdUtils::GetAssetUserData(&Material))
+		{
+			UnrealToUsd::ConvertMetadata(
+				UserData,
+				RootPrim,
+				MetadataOptions.BlockedPrefixFilters,
+				MetadataOptions.bInvertFilters
+			);
+		}
 	}
 
 	UsdStage.GetRootLayer().Save();
@@ -410,9 +423,33 @@ bool UMaterialExporterUsd::ExportMaterial(
 #endif // USE_USD_SDK
 }
 
+// Deprecated signature
+bool UMaterialExporterUsd::ExportMaterial(
+	const UMaterialInterface& Material,
+	const FUsdMaterialBakingOptions& Options,
+	const FFilePath& FilePath,
+	bool bReplaceIdentical,
+	bool bReExportIdenticalAssets,
+	bool bIsAutomated
+)
+{
+	static FUsdMetadataExportOptions DefaultOptions;
+	return UMaterialExporterUsd::ExportMaterial(
+		Material,
+		Options,
+		DefaultOptions,
+		FilePath,
+		bReplaceIdentical,
+		bReExportIdenticalAssets,
+		bIsAutomated
+	);
+}
+
+
 bool UMaterialExporterUsd::ExportMaterialsForStage(
 	const TArray<UMaterialInterface*>& Materials,
 	const FUsdMaterialBakingOptions& Options,
+	const FUsdMetadataExportOptions& MetadataOptions,
 	const FString& StageRootLayerPath,
 	bool bIsAssetLayer,
 	bool bUsePayload,
@@ -470,17 +507,18 @@ bool UMaterialExporterUsd::ExportMaterialsForStage(
 		// "C:/MyFolder/Export/Blue_4.usda"
 		FString FinalPath = FString::Printf( TEXT( "%s.%s" ), *FinalPathNoExt, *ExtensionNoDot );
 
-		if ( UMaterialExporterUsd::ExportMaterial(
+		if (UMaterialExporterUsd::ExportMaterial(
 			*Material,
 			Options,
-			FFilePath{ FinalPath },
+			MetadataOptions,
+			FFilePath{FinalPath},
 			bReplaceIdentical,
 			bReExportIdenticalAssets,
 			bIsAutomated
 		))
 		{
-			UsedFilePathsWithoutExt.Add( FinalPathNoExt );
-			MaterialPathNameToFilePath.Add( MaterialPathName, FinalPath );
+			UsedFilePathsWithoutExt.Add(FinalPathNoExt);
+			MaterialPathNameToFilePath.Add(MaterialPathName, FinalPath);
 		}
 	}
 
@@ -516,3 +554,29 @@ bool UMaterialExporterUsd::ExportMaterialsForStage(
 	return false;
 #endif // USE_USD_SDK
 }
+
+// Deprecated signature
+bool UMaterialExporterUsd::ExportMaterialsForStage(
+	const TArray<UMaterialInterface*>& Materials,
+	const FUsdMaterialBakingOptions& Options,
+	const FString& StageRootLayerPath,
+	bool bIsAssetLayer,
+	bool bUsePayload,
+	bool bReplaceIdentical,
+	bool bReExportIdenticalAssets,
+	bool bIsAutomated
+)
+{
+	static FUsdMetadataExportOptions DefaultOptions;
+	return UMaterialExporterUsd::ExportMaterialsForStage(
+		Materials,
+		Options,
+		DefaultOptions,
+		StageRootLayerPath,
+		bIsAssetLayer,
+		bUsePayload,
+		bReplaceIdentical,
+		bReExportIdenticalAssets,
+		bIsAutomated
+	);
+}
diff --git a/Engine/Plugins/Importers/USDImporter/Source/USDExporter/Private/MaterialExporterUSDOptions.cpp b/Engine/Plugins/Importers/USDImporter/Source/USDExporter/Private/MaterialExporterUSDOptions.cpp
--- a/Engine/Plugins/Importers/USDImporter/Source/USDExporter/Private/MaterialExporterUSDOptions.cpp
+++ b/Engine/Plugins/Importers/USDImporter/Source/USDExporter/Private/MaterialExporterUSDOptions.cpp
@@ -4,16 +4,15 @@
 
 #include "AnalyticsEventAttribute.h"
 
-void UsdUtils::AddAnalyticsAttributes(
-	const UMaterialExporterUSDOptions& Options,
-	TArray< FAnalyticsEventAttribute >& InOutAttributes
-)
+void UsdUtils::AddAnalyticsAttributes(const UMaterialExporterUSDOptions& Options, TArray<FAnalyticsEventAttribute>& InOutAttributes)
 {
-	UsdUtils::AddAnalyticsAttributes( Options.MaterialBakingOptions, InOutAttributes );
-	InOutAttributes.Emplace( TEXT( "ReExportIdenticalAssets" ), Options.bReExportIdenticalAssets );
+	UsdUtils::AddAnalyticsAttributes(Options.MaterialBakingOptions, InOutAttributes);
+	UsdUtils::AddAnalyticsAttributes(Options.MetadataOptions, InOutAttributes);
+	InOutAttributes.Emplace(TEXT("ReExportIdenticalAssets"), Options.bReExportIdenticalAssets);
 }
 
-void UsdUtils::HashForMaterialExport( const UMaterialExporterUSDOptions& Options, FSHA1& HashToUpdate )
+void UsdUtils::HashForMaterialExport(const UMaterialExporterUSDOptions& Options, FSHA1& HashToUpdate)
 {
-	UsdUtils::HashForMaterialExport( Options.MaterialBakingOptions, HashToUpdate );
+	UsdUtils::HashForMaterialExport(Options.MaterialBakingOptions, HashToUpdate);
+	UsdUtils::HashForExport(Options.MetadataOptions, HashToUpdate);
 }
diff --git a/Engine/Plugins/Importers/USDImporter/Source/USDExporter/Private/SkeletalMeshExporterUSD.cpp b/Engine/Plugins/Importers/USDImporter/Source/USDExporter/Private/SkeletalMeshExporterUSD.cpp
--- a/Engine/Plugins/Importers/USDImporter/Source/USDExporter/Private/SkeletalMeshExporterUSD.cpp
+++ b/Engine/Plugins/Importers/USDImporter/Source/USDExporter/Private/SkeletalMeshExporterUSD.cpp
@@ -2,19 +2,14 @@
 
 #include "SkeletalMeshExporterUSD.h"
 
-#include "Engine/SkinnedAssetCommon.h"
-#include "EngineAnalytics.h"
-#include "MaterialExporterUSD.h"
-#include "Misc/EngineVersion.h"
-#include "Misc/PackageName.h"
-#include "Rendering/SkeletalMeshRenderData.h"
 #include "SkeletalMeshExporterUSDOptions.h"
+#include "USDAssetUserData.h"
 #include "USDClassesModule.h"
 #include "USDConversionUtils.h"
 #include "USDExporterModule.h"
 #include "USDLog.h"
-#include "USDMemory.h"
 #include "USDOptionsWindow.h"
+#include "USDPrimConversion.h"
 #include "USDSkeletalDataConversion.h"
 #include "USDTypesConversion.h"
 #include "USDUnrealAssetInfo.h"
@@ -26,6 +21,11 @@
 
 #include "AssetExportTask.h"
 #include "Engine/SkeletalMesh.h"
+#include "Engine/SkinnedAssetCommon.h"
+#include "EngineAnalytics.h"
+#include "MaterialExporterUSD.h"
+#include "Misc/EngineVersion.h"
+#include "Rendering/SkeletalMeshRenderData.h"
 
 namespace UE::SkeletalMeshExporterUSD::Private
 {
@@ -220,6 +220,7 @@ bool USkeletalMeshExporterUsd::ExportBinary( UObject* Object, const TCHAR* Type,
 					UMaterialExporterUsd::ExportMaterialsForStage(
 						MaterialsToBake.Array(),
 						Options->MeshAssetOptions.MaterialBakingOptions,
+						Options->MetadataOptions,
 						UExporter::CurrentFilename,
 						bIsAssetLayer,
 						Options->MeshAssetOptions.bUsePayload,
@@ -287,19 +288,53 @@ bool USkeletalMeshExporterUsd::ExportBinary( UObject* Object, const TCHAR* Type,
 
 	UnrealToUsd::ConvertSkeletalMesh( SkeletalMesh, RootPrim, UsdUtils::GetDefaultTimeCode(), &AssetStage, Options->MeshAssetOptions.LowestMeshLOD, Options->MeshAssetOptions.HighestMeshLOD );
 
-	// Write asset info now that we finished exporting
-	if ( UE::FUsdPrim AssetDefaultPrim = AssetStage.GetDefaultPrim() )
+	if (UE::FUsdPrim AssetDefaultPrim = AssetStage.GetDefaultPrim())
 	{
-		FUsdUnrealAssetInfo Info;
-		Info.Name = SkeletalMesh->GetName();
-		Info.Identifier = UExporter::CurrentFilename;
-		Info.Version = DDCKeyHash;
-		Info.UnrealContentPath = SkeletalMesh->GetPathName();
-		Info.UnrealAssetType = SkeletalMesh->GetClass()->GetName();
-		Info.UnrealExportTime = FDateTime::Now().ToString();
-		Info.UnrealEngineVersion = FEngineVersion::Current().ToString();
-
-		UsdUtils::SetPrimAssetInfo( AssetDefaultPrim, Info );
+		if (Options->MetadataOptions.bExportAssetInfo)
+		{
+			FUsdUnrealAssetInfo Info;
+			Info.Name = SkeletalMesh->GetName();
+			Info.Identifier = UExporter::CurrentFilename;
+			Info.Version = DDCKeyHash;
+			Info.UnrealContentPath = SkeletalMesh->GetPathName();
+			Info.UnrealAssetType = SkeletalMesh->GetClass()->GetName();
+			Info.UnrealExportTime = FDateTime::Now().ToString();
+			Info.UnrealEngineVersion = FEngineVersion::Current().ToString();
+
+			UsdUtils::SetPrimAssetInfo(AssetDefaultPrim, Info);
+		}
+
+		if (Options->MetadataOptions.bExportAssetMetadata)
+		{
+			if (UUsdAssetUserData* UserData = UsdUtils::GetAssetUserData(SkeletalMesh))
+			{
+				UnrealToUsd::ConvertMetadata(
+					UserData,
+					AssetDefaultPrim,
+					Options->MetadataOptions.BlockedPrefixFilters,
+					Options->MetadataOptions.bInvertFilters
+				);
+			}
+
+			if (USkeleton* Skeleton = SkeletalMesh->GetSkeleton())
+			{
+				if (UUsdAssetUserData* UserData = UsdUtils::GetAssetUserData(Skeleton))
+				{
+					if (UserData->StageIdentifierToMetadata.Num() > 0)
+					{
+						UE::FUsdPrim SkelPrim = AssetStage.OverridePrim(
+							UE::FSdfPath(*RootPrimPath).AppendChild(UnrealIdentifiers::ExportedSkeletonPrimName)
+						);
+						UnrealToUsd::ConvertMetadata(
+							UserData,
+							SkelPrim,
+							Options->MetadataOptions.BlockedPrefixFilters,
+							Options->MetadataOptions.bInvertFilters
+						);
+					}
+				}
+			}
+		}
 	}
 
 	// Bake materials and replace unrealMaterials with references to the baked files.
@@ -315,6 +350,7 @@ bool USkeletalMeshExporterUsd::ExportBinary( UObject* Object, const TCHAR* Type,
 		UMaterialExporterUsd::ExportMaterialsForStage(
 			MaterialsToBake.Array(),
 			Options->MeshAssetOptions.MaterialBakingOptions,
+			Options->MetadataOptions,
 			AssetStage.GetRootLayer().GetRealPath(),
 			bIsAssetLayer,
 			Options->MeshAssetOptions.bUsePayload,
diff --git a/Engine/Plugins/Importers/USDImporter/Source/USDExporter/Private/SkeletalMeshExporterUSDOptions.cpp b/Engine/Plugins/Importers/USDImporter/Source/USDExporter/Private/SkeletalMeshExporterUSDOptions.cpp
--- a/Engine/Plugins/Importers/USDImporter/Source/USDExporter/Private/SkeletalMeshExporterUSDOptions.cpp
+++ b/Engine/Plugins/Importers/USDImporter/Source/USDExporter/Private/SkeletalMeshExporterUSDOptions.cpp
@@ -4,18 +4,17 @@
 
 #include "AnalyticsEventAttribute.h"
 
-void UsdUtils::AddAnalyticsAttributes(
-	const USkeletalMeshExporterUSDOptions& Options,
-	TArray< FAnalyticsEventAttribute >& InOutAttributes
-)
+void UsdUtils::AddAnalyticsAttributes(const USkeletalMeshExporterUSDOptions& Options, TArray<FAnalyticsEventAttribute>& InOutAttributes)
 {
-	UsdUtils::AddAnalyticsAttributes( Options.StageOptions, InOutAttributes );
-	UsdUtils::AddAnalyticsAttributes( Options.MeshAssetOptions, InOutAttributes );
-	InOutAttributes.Emplace( TEXT( "ReExportIdenticalAssets" ), Options.bReExportIdenticalAssets );
+	UsdUtils::AddAnalyticsAttributes(Options.StageOptions, InOutAttributes);
+	UsdUtils::AddAnalyticsAttributes(Options.MeshAssetOptions, InOutAttributes);
+	UsdUtils::AddAnalyticsAttributes(Options.MetadataOptions, InOutAttributes);
+	InOutAttributes.Emplace(TEXT("ReExportIdenticalAssets"), Options.bReExportIdenticalAssets);
 }
 
-void UsdUtils::HashForSkeletalMeshExport( const USkeletalMeshExporterUSDOptions& Options, FSHA1& HashToUpdate )
+void UsdUtils::HashForSkeletalMeshExport(const USkeletalMeshExporterUSDOptions& Options, FSHA1& HashToUpdate)
 {
-	UsdUtils::HashForExport( Options.StageOptions, HashToUpdate );
-	UsdUtils::HashForMeshExport( Options.MeshAssetOptions, HashToUpdate );
+	UsdUtils::HashForExport(Options.StageOptions, HashToUpdate);
+	UsdUtils::HashForMeshExport(Options.MeshAssetOptions, HashToUpdate);
+	UsdUtils::HashForExport(Options.MetadataOptions, HashToUpdate);
 }
diff --git a/Engine/Plugins/Importers/USDImporter/Source/USDExporter/Private/StaticMeshExporterUSD.cpp b/Engine/Plugins/Importers/USDImporter/Source/USDExporter/Private/StaticMeshExporterUSD.cpp
--- a/Engine/Plugins/Importers/USDImporter/Source/USDExporter/Private/StaticMeshExporterUSD.cpp
+++ b/Engine/Plugins/Importers/USDImporter/Source/USDExporter/Private/StaticMeshExporterUSD.cpp
@@ -12,8 +12,8 @@
 #include "USDExporterModule.h"
 #include "USDGeomMeshConversion.h"
 #include "USDLog.h"
-#include "USDMemory.h"
 #include "USDOptionsWindow.h"
+#include "USDPrimConversion.h"
 #include "USDTypesConversion.h"
 #include "USDUnrealAssetInfo.h"
 
@@ -23,7 +23,6 @@
 #include "UsdWrappers/UsdStage.h"
 
 #include "AssetExportTask.h"
-#include "CoreMinimal.h"
 #include "Engine/StaticMesh.h"
 
 namespace UE::StaticMeshExporterUSD::Private
@@ -233,6 +232,7 @@ bool UStaticMeshExporterUsd::ExportBinary( UObject* Object, const TCHAR* Type, F
 					UMaterialExporterUsd::ExportMaterialsForStage(
 						MaterialsToBake.Array(),
 						Options->MeshAssetOptions.MaterialBakingOptions,
+						Options->MetadataOptions,
 						UExporter::CurrentFilename,
 						bIsAssetLayer,
 						Options->MeshAssetOptions.bUsePayload,
@@ -298,19 +298,34 @@ bool UStaticMeshExporterUsd::ExportBinary( UObject* Object, const TCHAR* Type, F
 
 	UnrealToUsd::ConvertStaticMesh( StaticMesh, RootPrim, UsdUtils::GetDefaultTimeCode(), &AssetStage, Options->MeshAssetOptions.LowestMeshLOD, Options->MeshAssetOptions.HighestMeshLOD );
 
-	// Write asset info now that we finished exporting
-	if( UE::FUsdPrim AssetDefaultPrim = AssetStage.GetDefaultPrim() )
+	if (UE::FUsdPrim AssetDefaultPrim = AssetStage.GetDefaultPrim())
 	{
-		FUsdUnrealAssetInfo Info;
-		Info.Name = StaticMesh->GetName();
-		Info.Identifier = UExporter::CurrentFilename;
-		Info.Version = DDCKeyHash;
-		Info.UnrealContentPath = StaticMesh->GetPathName();
-		Info.UnrealAssetType = StaticMesh->GetClass()->GetName();
-		Info.UnrealExportTime = FDateTime::Now().ToString();
-		Info.UnrealEngineVersion = FEngineVersion::Current().ToString();
-
-		UsdUtils::SetPrimAssetInfo( AssetDefaultPrim, Info );
+		if (Options->MetadataOptions.bExportAssetInfo)
+		{
+			FUsdUnrealAssetInfo Info;
+			Info.Name = StaticMesh->GetName();
+			Info.Identifier = UExporter::CurrentFilename;
+			Info.Version = DDCKeyHash;
+			Info.UnrealContentPath = StaticMesh->GetPathName();
+			Info.UnrealAssetType = StaticMesh->GetClass()->GetName();
+			Info.UnrealExportTime = FDateTime::Now().ToString();
+			Info.UnrealEngineVersion = FEngineVersion::Current().ToString();
+
+			UsdUtils::SetPrimAssetInfo(AssetDefaultPrim, Info);
+		}
+
+		if (Options->MetadataOptions.bExportAssetMetadata)
+		{
+			if (UUsdAssetUserData* UserData = UsdUtils::GetAssetUserData(StaticMesh))
+			{
+				UnrealToUsd::ConvertMetadata(
+					UserData,
+					AssetDefaultPrim,
+					Options->MetadataOptions.BlockedPrefixFilters,
+					Options->MetadataOptions.bInvertFilters
+				);
+			}
+		}
 	}
 
 	// Bake materials and replace unrealMaterials with references to the baked files.
@@ -326,6 +341,7 @@ bool UStaticMeshExporterUsd::ExportBinary( UObject* Object, const TCHAR* Type, F
 		UMaterialExporterUsd::ExportMaterialsForStage(
 			MaterialsToBake.Array(),
 			Options->MeshAssetOptions.MaterialBakingOptions,
+			Options->MetadataOptions,
 			AssetStage.GetRootLayer().GetRealPath(),
 			bIsAssetLayer,
 			Options->MeshAssetOptions.bUsePayload,
diff --git a/Engine/Plugins/Importers/USDImporter/Source/USDExporter/Private/StaticMeshExporterUSDOptions.cpp b/Engine/Plugins/Importers/USDImporter/Source/USDExporter/Private/StaticMeshExporterUSDOptions.cpp
--- a/Engine/Plugins/Importers/USDImporter/Source/USDExporter/Private/StaticMeshExporterUSDOptions.cpp
+++ b/Engine/Plugins/Importers/USDImporter/Source/USDExporter/Private/StaticMeshExporterUSDOptions.cpp
@@ -4,18 +4,17 @@
 
 #include "AnalyticsEventAttribute.h"
 
-void UsdUtils::AddAnalyticsAttributes(
-	const UStaticMeshExporterUSDOptions& Options,
-	TArray< FAnalyticsEventAttribute >& InOutAttributes
-)
+void UsdUtils::AddAnalyticsAttributes(const UStaticMeshExporterUSDOptions& Options, TArray<FAnalyticsEventAttribute>& InOutAttributes)
 {
-	UsdUtils::AddAnalyticsAttributes( Options.StageOptions, InOutAttributes );
-	UsdUtils::AddAnalyticsAttributes( Options.MeshAssetOptions, InOutAttributes );
-	InOutAttributes.Emplace( TEXT( "ReExportIdenticalAssets" ), Options.bReExportIdenticalAssets );
+	UsdUtils::AddAnalyticsAttributes(Options.StageOptions, InOutAttributes);
+	UsdUtils::AddAnalyticsAttributes(Options.MeshAssetOptions, InOutAttributes);
+	UsdUtils::AddAnalyticsAttributes(Options.MetadataOptions, InOutAttributes);
+	InOutAttributes.Emplace(TEXT("ReExportIdenticalAssets"), Options.bReExportIdenticalAssets);
 }
 
-void UsdUtils::HashForStaticMeshExport( const UStaticMeshExporterUSDOptions& Options, FSHA1& HashToUpdate )
+void UsdUtils::HashForStaticMeshExport(const UStaticMeshExporterUSDOptions& Options, FSHA1& HashToUpdate)
 {
-	UsdUtils::HashForExport( Options.StageOptions, HashToUpdate );
-	UsdUtils::HashForMeshExport( Options.MeshAssetOptions, HashToUpdate );
+	UsdUtils::HashForExport(Options.StageOptions, HashToUpdate);
+	UsdUtils::HashForMeshExport(Options.MeshAssetOptions, HashToUpdate);
+	UsdUtils::HashForExport(Options.MetadataOptions, HashToUpdate);
 }
diff --git a/Engine/Plugins/Importers/USDImporter/Source/USDExporter/Private/USDConversionBlueprintContext.cpp b/Engine/Plugins/Importers/USDImporter/Source/USDExporter/Private/USDConversionBlueprintContext.cpp
--- a/Engine/Plugins/Importers/USDImporter/Source/USDExporter/Private/USDConversionBlueprintContext.cpp
+++ b/Engine/Plugins/Importers/USDImporter/Source/USDExporter/Private/USDConversionBlueprintContext.cpp
@@ -7,12 +7,15 @@
 #include "USDGeomMeshConversion.h"
 #include "USDLightConversion.h"
 #include "USDLog.h"
+#include "USDMetadata.h"
 #include "USDPrimConversion.h"
 #include "USDShadeConversion.h"
 #include "USDUnrealAssetInfo.h"
+#include "USDValueConversion.h"
 
 #include "UsdWrappers/SdfPath.h"
 #include "UsdWrappers/UsdPrim.h"
+#include "UsdWrappers/VtValue.h"
 
 #include "InstancedFoliageActor.h"
 #include "LandscapeProxy.h"
@@ -555,3 +558,65 @@ FUsdUnrealAssetInfo UUsdConversionBlueprintContext::GetPrimAssetInfo( const FStr
 	return {};
 }
 
+void UUsdConversionBlueprintContext::SetPrimMetadata(
+	const FString& PrimPath,
+	const FUsdCombinedPrimMetadata& Metadata,
+	const TArray<FString>& BlockedPrefixFilter,
+	bool bInvertFilter
+)
+{
+#if USE_USD_SDK
+	if (Stage)
+	{
+		if (UE::FUsdPrim Prim = UnrealToUsdImpl::GetPrim(Stage, PrimPath))
+		{
+			UnrealToUsd::ConvertMetadata(Metadata, Prim, BlockedPrefixFilter, bInvertFilter);
+		}
+	}
+#endif	  // USE_USD_SDK
+}
+
+void UUsdConversionBlueprintContext::SetPrimMetadataFromUserData(
+	const FString& PrimPath,
+	const UUsdAssetUserData* UserData,
+	const TArray<FString>& BlockedPrefixFilter,
+	bool bInvertFilter
+)
+{
+#if USE_USD_SDK
+	if (Stage && UserData)
+	{
+		if (UE::FUsdPrim Prim = UnrealToUsdImpl::GetPrim(Stage, PrimPath))
+		{
+			UnrealToUsd::ConvertMetadata(UserData, Prim, BlockedPrefixFilter, bInvertFilter);
+		}
+	}
+#endif	  // USE_USD_SDK
+}
+
+FUsdCombinedPrimMetadata UUsdConversionBlueprintContext::GetPrimMetadata(
+	const FString& PrimPath,
+	const TArray<FString>& BlockedPrefixFilter,
+	bool bInvertFilter,
+	bool bCollectFromEntireSubtrees
+)
+{
+	FUsdCombinedPrimMetadata Result;
+
+#if USE_USD_SDK
+	if (Stage)
+	{
+		if (UE::FUsdPrim Prim = UnrealToUsdImpl::GetPrim(Stage, PrimPath))
+		{
+			bool bSuccess = UsdToUnreal::ConvertMetadata(Prim, Result, BlockedPrefixFilter, bInvertFilter, bCollectFromEntireSubtrees);
+			if (!bSuccess)
+			{
+				UE_LOG(LogUsd, Warning, TEXT("Failed to get metadata from prim '%s'"), *PrimPath);
+				Result = {};
+			}
+		}
+	}
+#endif	  // USE_USD_SDK
+
+	return Result;
+}
diff --git a/Engine/Plugins/Importers/USDImporter/Source/USDExporter/Private/USDConversionBlueprintLibrary.cpp b/Engine/Plugins/Importers/USDImporter/Source/USDExporter/Private/USDConversionBlueprintLibrary.cpp
--- a/Engine/Plugins/Importers/USDImporter/Source/USDExporter/Private/USDConversionBlueprintLibrary.cpp
+++ b/Engine/Plugins/Importers/USDImporter/Source/USDExporter/Private/USDConversionBlueprintLibrary.cpp
@@ -5,16 +5,20 @@
 #include "LevelExporterUSD.h"
 #include "LevelExporterUSDOptions.h"
 #include "UnrealUSDWrapper.h"
+#include "USDAssetUserData.h"
 #include "USDClassesModule.h"
 #include "USDConversionUtils.h"
 #include "USDExporterModule.h"
 #include "USDLayerUtils.h"
 #include "USDLog.h"
+#include "USDTypesConversion.h"
+#include "USDValueConversion.h"
 
 #include "UsdWrappers/SdfLayer.h"
 #include "UsdWrappers/SdfPath.h"
 #include "UsdWrappers/UsdPrim.h"
 #include "UsdWrappers/UsdStage.h"
+#include "UsdWrappers/VtValue.h"
 
 #include "AnalyticsBlueprintLibrary.h"
 #include "AnalyticsEventAttribute.h"
@@ -749,3 +753,973 @@ TArray<FString> UUsdConversionBlueprintLibrary::DuplicatePrims( const FString& S
 
 	return Result;
 }
+
+UUsdAssetUserData* UUsdConversionBlueprintLibrary::GetUsdAssetUserData(UObject* Object)
+{
+	return UsdUtils::GetAssetUserData(Object);
+}
+
+bool UUsdConversionBlueprintLibrary::SetUsdAssetUserData(UObject* Object, UUsdAssetUserData* AssetUserData)
+{
+	return UsdUtils::SetAssetUserData(Object, AssetUserData);
+}
+
+namespace UE::UsdConversionBlueprintLibraryImpl::Private
+{
+	FEditPropertyChain& GetMetadataPropertyChain()
+	{
+		static TOptional<FEditPropertyChain> Chain;
+
+		if (!Chain.IsSet())
+		{
+			FProperty* StageIdentifierProp = UUsdAssetUserData::StaticClass()->FindPropertyByName(
+				GET_MEMBER_NAME_CHECKED(UUsdAssetUserData, StageIdentifierToMetadata)
+			);
+			FProperty* StringifiedValueProp = FUsdMetadataValue::StaticStruct()->FindPropertyByName(
+				GET_MEMBER_NAME_CHECKED(FUsdMetadataValue, StringifiedValue)
+			);
+
+			Chain.Emplace();
+			Chain->AddHead(StageIdentifierProp);
+			Chain->AddTail(FUsdCombinedPrimMetadata::StaticStruct()->FindPropertyByName(GET_MEMBER_NAME_CHECKED(FUsdCombinedPrimMetadata, PrimPathToMetadata)));
+			Chain->AddTail(FUsdPrimMetadata::StaticStruct()->FindPropertyByName(GET_MEMBER_NAME_CHECKED(FUsdPrimMetadata, Metadata)));
+			Chain->AddTail(StringifiedValueProp);
+			Chain->SetActivePropertyNode(StringifiedValueProp);
+			Chain->SetActiveMemberPropertyNode(StageIdentifierProp);
+		}
+
+		return Chain.GetValue();
+	}
+
+	FPropertyChangedChainEvent GetPostChangeChainEvent(UObject* TopLevelObject, EPropertyChangeType::Type ChangeType)
+	{
+		FPropertyChangedEvent PostChangeEvent{
+			FUsdMetadataValue::StaticStruct()->FindPropertyByName(GET_MEMBER_NAME_CHECKED(FUsdMetadataValue, StringifiedValue)),
+			ChangeType,
+			{TopLevelObject}
+		};
+		PostChangeEvent.MemberProperty = UUsdAssetUserData::StaticClass()->FindPropertyByName(
+			GET_MEMBER_NAME_CHECKED(UUsdAssetUserData, StageIdentifierToMetadata)
+		);
+
+		FPropertyChangedChainEvent ChainEvent{
+			GetMetadataPropertyChain(),
+			PostChangeEvent
+		};
+
+		return ChainEvent;
+	}
+}
+
+bool UUsdConversionBlueprintLibrary::SetMetadataField(
+	UUsdAssetUserData* AssetUserData,
+	const FString& Key,
+	const FString& Value,
+	const FString& ValueTypeName,
+	const FString& StageIdentifier,
+	const FString& PrimPath,
+	bool bTriggerPropertyChangeEvents
+)
+{
+	using namespace UE::UsdConversionBlueprintLibraryImpl::Private;
+
+	if (!AssetUserData)
+	{
+		return false;
+	}
+
+	if (StageIdentifier.IsEmpty() && AssetUserData->StageIdentifierToMetadata.Num() != 1)
+	{
+		UE_LOG(LogUsd, Warning, TEXT("Failed to set metadata field '%s' on AssetUserData '%s': Please provide a valid value for the StageIdentifier parameter"),
+			*AssetUserData->GetPathName(),
+			*Key
+		);
+		return false;
+	}
+
+	const FString& StageIdentifierToUse = !StageIdentifier.IsEmpty() ? StageIdentifier : AssetUserData->StageIdentifierToMetadata.CreateIterator()->Key;
+	FUsdCombinedPrimMetadata& CombinedPrimMetadata = AssetUserData->StageIdentifierToMetadata.FindOrAdd(StageIdentifierToUse);
+
+	if (PrimPath.IsEmpty() && CombinedPrimMetadata.PrimPathToMetadata.Num() != 1)
+	{
+		UE_LOG(LogUsd, Warning, TEXT("Failed to set metadata field '%s' on AssetUserData '%s': Please provide a valid value for the PrimPath parameter"),
+			*AssetUserData->GetPathName(),
+			*Key
+		);
+		return false;
+	}
+
+	const FString& PrimPathToUse = !PrimPath.IsEmpty() ? PrimPath : CombinedPrimMetadata.PrimPathToMetadata.CreateConstIterator()->Key;
+	FUsdPrimMetadata& PrimMetadata = CombinedPrimMetadata.PrimPathToMetadata.FindOrAdd(PrimPathToUse);
+
+	if (bTriggerPropertyChangeEvents)
+	{
+		AssetUserData->PreEditChange(GetMetadataPropertyChain());
+	}
+
+	FUsdMetadataValue& MetadataValue = PrimMetadata.Metadata.FindOrAdd(Key);
+	MetadataValue.StringifiedValue = Value;
+	MetadataValue.TypeName = ValueTypeName;
+
+	if (bTriggerPropertyChangeEvents)
+	{
+		FPropertyChangedChainEvent Event = GetPostChangeChainEvent(AssetUserData, EPropertyChangeType::ValueSet);
+		AssetUserData->PostEditChangeChainProperty(Event);
+	}
+
+	return true;
+}
+
+bool UUsdConversionBlueprintLibrary::ClearMetadataField(
+	UUsdAssetUserData* AssetUserData,
+	const FString& Key,
+	const FString& StageIdentifier,
+	const FString& PrimPath,
+	bool bTriggerPropertyChangeEvents
+)
+{
+	using namespace UE::UsdConversionBlueprintLibraryImpl::Private;
+
+	if (!AssetUserData)
+	{
+		return false;
+	}
+
+	if (StageIdentifier.IsEmpty() && AssetUserData->StageIdentifierToMetadata.Num() != 1)
+	{
+		UE_LOG(LogUsd, Warning, TEXT("Failed to clear metadata field '%s' on AssetUserData '%s': Please provide a valid value for the StageIdentifier parameter"),
+			*AssetUserData->GetPathName(),
+			*Key
+		);
+		return false;
+	}
+
+	const FString& StageIdentifierToUse = !StageIdentifier.IsEmpty() ? StageIdentifier : AssetUserData->StageIdentifierToMetadata.CreateIterator()->Key;
+	FUsdCombinedPrimMetadata& CombinedPrimMetadata = AssetUserData->StageIdentifierToMetadata.FindOrAdd(StageIdentifierToUse);
+
+	if (PrimPath.IsEmpty() && CombinedPrimMetadata.PrimPathToMetadata.Num() != 1)
+	{
+		UE_LOG(LogUsd, Warning, TEXT("Failed to clear metadata field '%s' on AssetUserData '%s': Please provide a valid value for the PrimPath parameter"),
+			*AssetUserData->GetPathName(),
+			*Key
+		);
+		return false;
+	}
+
+	const FString& PrimPathToUse = !PrimPath.IsEmpty() ? PrimPath : CombinedPrimMetadata.PrimPathToMetadata.CreateConstIterator()->Key;
+	FUsdPrimMetadata& PrimMetadata = CombinedPrimMetadata.PrimPathToMetadata.FindOrAdd(PrimPathToUse);
+
+	if (bTriggerPropertyChangeEvents)
+	{
+		AssetUserData->PreEditChange(GetMetadataPropertyChain());
+	}
+
+	PrimMetadata.Metadata.Remove(Key);
+
+	if (bTriggerPropertyChangeEvents)
+	{
+		FPropertyChangedChainEvent Event = GetPostChangeChainEvent(AssetUserData, EPropertyChangeType::ArrayRemove);
+		AssetUserData->PostEditChangeChainProperty(Event);
+	}
+
+	return true;
+}
+
+bool UUsdConversionBlueprintLibrary::HasMetadataField(
+	UUsdAssetUserData* AssetUserData,
+	const FString& Key,
+	const FString& StageIdentifier,
+	const FString& PrimPath
+)
+{
+	if (!AssetUserData)
+	{
+		return false;
+	}
+
+	if (StageIdentifier.IsEmpty() && AssetUserData->StageIdentifierToMetadata.Num() != 1)
+	{
+		UE_LOG(LogUsd, Warning, TEXT("Failed to check for metadata field '%s' on AssetUserData '%s': Please provide a valid value for the StageIdentifier parameter"),
+			*AssetUserData->GetPathName(),
+			*Key
+		);
+		return false;
+	}
+
+	const FString& StageIdentifierToUse = !StageIdentifier.IsEmpty() ? StageIdentifier : AssetUserData->StageIdentifierToMetadata.CreateIterator()->Key;
+	const FUsdCombinedPrimMetadata* CombinedPrimMetadata = AssetUserData->StageIdentifierToMetadata.Find(StageIdentifierToUse);
+	if (!CombinedPrimMetadata)
+	{
+		return false;
+	}
+
+	if (PrimPath.IsEmpty() && CombinedPrimMetadata->PrimPathToMetadata.Num() != 1)
+	{
+		UE_LOG(LogUsd, Warning, TEXT("Failed to clear metadata field '%s' on AssetUserData '%s': Please provide a valid value for the PrimPath parameter"),
+			*AssetUserData->GetPathName(),
+			*Key
+		);
+		return false;
+	}
+
+	const FString& PrimPathToUse = !PrimPath.IsEmpty() ? PrimPath : CombinedPrimMetadata->PrimPathToMetadata.CreateConstIterator()->Key;
+	const FUsdPrimMetadata* PrimMetadata = CombinedPrimMetadata->PrimPathToMetadata.Find(PrimPathToUse);
+	if (!PrimMetadata)
+	{
+		return false;
+	}
+
+	return PrimMetadata->Metadata.Contains(Key);
+}
+
+FUsdMetadataValue UUsdConversionBlueprintLibrary::GetMetadataField(
+	UUsdAssetUserData* AssetUserData,
+	const FString& Key,
+	const FString& StageIdentifier,
+	const FString& PrimPath
+)
+{
+	if (!AssetUserData)
+	{
+		return {};
+	}
+
+	if (StageIdentifier.IsEmpty() && AssetUserData->StageIdentifierToMetadata.Num() != 1)
+	{
+		UE_LOG(LogUsd, Warning, TEXT("Failed to get metadata field '%s' on AssetUserData '%s': Please provide a valid value for the StageIdentifier parameter"),
+			*AssetUserData->GetPathName(),
+			*Key
+		);
+		return {};
+	}
+
+	const FString& StageIdentifierToUse = !StageIdentifier.IsEmpty() ? StageIdentifier : AssetUserData->StageIdentifierToMetadata.CreateIterator()->Key;
+	const FUsdCombinedPrimMetadata* CombinedPrimMetadata = AssetUserData->StageIdentifierToMetadata.Find(StageIdentifierToUse);
+	if (!CombinedPrimMetadata)
+	{
+		return {};
+	}
+
+	if (PrimPath.IsEmpty() && CombinedPrimMetadata->PrimPathToMetadata.Num() != 1)
+	{
+		UE_LOG(LogUsd, Warning, TEXT("Failed to get metadata field '%s' on AssetUserData '%s': Please provide a valid value for the PrimPath parameter"),
+			*AssetUserData->GetPathName(),
+			*Key
+		);
+		return {};
+	}
+
+	const FString& PrimPathToUse = !PrimPath.IsEmpty() ? PrimPath : CombinedPrimMetadata->PrimPathToMetadata.CreateConstIterator()->Key;
+	const FUsdPrimMetadata* PrimMetadata = CombinedPrimMetadata->PrimPathToMetadata.Find(PrimPathToUse);
+	if (!PrimMetadata)
+	{
+		return {};
+	}
+
+	if (const FUsdMetadataValue* FoundValue = PrimMetadata->Metadata.Find(Key))
+	{
+		return *FoundValue;
+	}
+
+	UE_LOG(LogUsd, Warning, TEXT("Failed to find a metadata entry in '%s' with stage identifier '%s', prim path '%s' and key '%s'"),
+		*AssetUserData->GetPathName(),
+		*StageIdentifier,
+		*PrimPath,
+		*Key
+	);
+	return {};
+}
+
+FString UUsdConversionBlueprintLibrary::StringifyAsBool(bool Value)
+{
+	return UsdUtils::StringifyAsBool(Value);
+}
+
+FString UUsdConversionBlueprintLibrary::StringifyAsUChar(uint8 Value)
+{
+	return UsdUtils::StringifyAsUChar(Value);
+}
+
+FString UUsdConversionBlueprintLibrary::StringifyAsInt(int32 Value)
+{
+	return UsdUtils::StringifyAsInt(Value);
+}
+
+FString UUsdConversionBlueprintLibrary::StringifyAsUInt(int32 Value)
+{
+	// Have to cast as there are no unsigned types higher than uint8 on blueprint...
+	return UsdUtils::StringifyAsUInt(static_cast<uint32>(Value));
+}
+
+FString UUsdConversionBlueprintLibrary::StringifyAsInt64(int64 Value)
+{
+	return UsdUtils::StringifyAsInt64(Value);
+}
+
+FString UUsdConversionBlueprintLibrary::StringifyAsUInt64(int64 Value)
+{
+	// Have to cast as there are no unsigned types higher than uint8 on blueprint...
+	return UsdUtils::StringifyAsUInt64(static_cast<uint64>(Value));
+}
+
+FString UUsdConversionBlueprintLibrary::StringifyAsHalf(float Value)
+{
+	return UsdUtils::StringifyAsHalf(Value);
+}
+
+FString UUsdConversionBlueprintLibrary::StringifyAsFloat(float Value)
+{
+	return UsdUtils::StringifyAsFloat(Value);
+}
+
+FString UUsdConversionBlueprintLibrary::StringifyAsDouble(double Value)
+{
+	return UsdUtils::StringifyAsDouble(Value);
+}
+
+FString UUsdConversionBlueprintLibrary::StringifyAsTimeCode(double Value)
+{
+	return UsdUtils::StringifyAsTimeCode(Value);
+}
+
+FString UUsdConversionBlueprintLibrary::StringifyAsString(const FString& Value)
+{
+	return UsdUtils::StringifyAsString(Value);
+}
+
+FString UUsdConversionBlueprintLibrary::StringifyAsToken(const FString& Value)
+{
+	return UsdUtils::StringifyAsToken(Value);
+}
+
+FString UUsdConversionBlueprintLibrary::StringifyAsAssetPath(const FString& Value)
+{
+	return UsdUtils::StringifyAsAssetPath(Value);
+}
+
+FString UUsdConversionBlueprintLibrary::StringifyAsMatrix2d(const FMatrix2D& Value)
+{
+	return UsdUtils::StringifyAsMatrix2d(Value);
+}
+
+FString UUsdConversionBlueprintLibrary::StringifyAsMatrix3d(const FMatrix3D& Value)
+{
+	return UsdUtils::StringifyAsMatrix3d(Value);
+}
+
+FString UUsdConversionBlueprintLibrary::StringifyAsMatrix4d(const FMatrix& Value)
+{
+	return UsdUtils::StringifyAsMatrix4d(Value);
+}
+
+FString UUsdConversionBlueprintLibrary::StringifyAsQuatd(const FQuat& Value)
+{
+	return UsdUtils::StringifyAsQuatd(Value);
+}
+
+FString UUsdConversionBlueprintLibrary::StringifyAsQuatf(const FQuat& Value)
+{
+	return UsdUtils::StringifyAsQuatf(Value);
+}
+
+FString UUsdConversionBlueprintLibrary::StringifyAsQuath(const FQuat& Value)
+{
+	return UsdUtils::StringifyAsQuath(Value);
+}
+
+FString UUsdConversionBlueprintLibrary::StringifyAsDouble2(const FVector2D& Value)
+{
+	return UsdUtils::StringifyAsDouble2(Value);
+}
+
+FString UUsdConversionBlueprintLibrary::StringifyAsFloat2(const FVector2D& Value)
+{
+	return UsdUtils::StringifyAsFloat2(Value);
+}
+
+FString UUsdConversionBlueprintLibrary::StringifyAsHalf2(const FVector2D& Value)
+{
+	return UsdUtils::StringifyAsHalf2(Value);
+}
+
+FString UUsdConversionBlueprintLibrary::StringifyAsInt2(const FIntPoint& Value)
+{
+	return UsdUtils::StringifyAsInt2(Value);
+}
+
+FString UUsdConversionBlueprintLibrary::StringifyAsDouble3(const FVector& Value)
+{
+	return UsdUtils::StringifyAsDouble3(Value);
+}
+
+FString UUsdConversionBlueprintLibrary::StringifyAsFloat3(const FVector& Value)
+{
+	return UsdUtils::StringifyAsFloat3(Value);
+}
+
+FString UUsdConversionBlueprintLibrary::StringifyAsHalf3(const FVector& Value)
+{
+	return UsdUtils::StringifyAsHalf3(Value);
+}
+
+FString UUsdConversionBlueprintLibrary::StringifyAsInt3(const FIntVector& Value)
+{
+	return UsdUtils::StringifyAsInt3(Value);
+}
+
+FString UUsdConversionBlueprintLibrary::StringifyAsDouble4(const FVector4& Value)
+{
+	return UsdUtils::StringifyAsDouble4(Value);
+}
+
+FString UUsdConversionBlueprintLibrary::StringifyAsFloat4(const FVector4& Value)
+{
+	return UsdUtils::StringifyAsFloat4(Value);
+}
+
+FString UUsdConversionBlueprintLibrary::StringifyAsHalf4(const FVector4& Value)
+{
+	return UsdUtils::StringifyAsHalf4(Value);
+}
+
+FString UUsdConversionBlueprintLibrary::StringifyAsInt4(const FIntVector4& Value)
+{
+	return UsdUtils::StringifyAsInt4(Value);
+}
+
+FString UUsdConversionBlueprintLibrary::StringifyAsBoolArray(const TArray<bool>& Value)
+{
+	return UsdUtils::StringifyAsBoolArray(Value);
+}
+
+FString UUsdConversionBlueprintLibrary::StringifyAsUCharArray(const TArray<uint8>& Value)
+{
+	return UsdUtils::StringifyAsUCharArray(Value);
+}
+
+FString UUsdConversionBlueprintLibrary::StringifyAsIntArray(const TArray<int32>& Value)
+{
+	return UsdUtils::StringifyAsIntArray(Value);
+}
+
+FString UUsdConversionBlueprintLibrary::StringifyAsUIntArray(const TArray<int32>& Value)
+{
+	static_assert(sizeof(int32) == sizeof(uint32));
+
+	// Have to cast as there are no unsigned types higher than uint8 on blueprint...
+	TArray<uint32> CastValue;
+	CastValue.SetNumUninitialized(Value.Num());
+	FMemory::Memcpy(CastValue.GetData(), Value.GetData(), Value.Num() * Value.GetTypeSize());
+
+	return UsdUtils::StringifyAsUIntArray(CastValue);
+}
+
+FString UUsdConversionBlueprintLibrary::StringifyAsInt64Array(const TArray<int64>& Value)
+{
+	return UsdUtils::StringifyAsInt64Array(Value);
+}
+
+FString UUsdConversionBlueprintLibrary::StringifyAsUInt64Array(const TArray<int64>& Value)
+{
+	static_assert(sizeof(int64) == sizeof(uint64));
+
+	// Have to cast as there are no unsigned types higher than uint8 on blueprint...
+	TArray<uint64> CastValue;
+	CastValue.SetNumUninitialized(Value.Num());
+	FMemory::Memcpy(CastValue.GetData(), Value.GetData(), Value.Num() * Value.GetTypeSize());
+
+	return UsdUtils::StringifyAsUInt64Array(CastValue);
+}
+
+FString UUsdConversionBlueprintLibrary::StringifyAsHalfArray(const TArray<float>& Value)
+{
+	return UsdUtils::StringifyAsHalfArray(Value);
+}
+
+FString UUsdConversionBlueprintLibrary::StringifyAsFloatArray(const TArray<float>& Value)
+{
+	return UsdUtils::StringifyAsFloatArray(Value);
+}
+
+FString UUsdConversionBlueprintLibrary::StringifyAsDoubleArray(const TArray<double>& Value)
+{
+	return UsdUtils::StringifyAsDoubleArray(Value);
+}
+
+FString UUsdConversionBlueprintLibrary::StringifyAsTimeCodeArray(const TArray<double>& Value)
+{
+	return UsdUtils::StringifyAsTimeCodeArray(Value);
+}
+
+FString UUsdConversionBlueprintLibrary::StringifyAsStringArray(const TArray<FString>& Value)
+{
+	return UsdUtils::StringifyAsStringArray(Value);
+}
+
+FString UUsdConversionBlueprintLibrary::StringifyAsTokenArray(const TArray<FString>& Value)
+{
+	return UsdUtils::StringifyAsTokenArray(Value);
+}
+
+FString UUsdConversionBlueprintLibrary::StringifyAsAssetPathArray(const TArray<FString>& Value)
+{
+	return UsdUtils::StringifyAsAssetPathArray(Value);
+}
+
+FString UUsdConversionBlueprintLibrary::StringifyAsListOpTokens(const TArray<FString>& Value)
+{
+	return UsdUtils::StringifyAsListOpTokens(Value);
+}
+
+FString UUsdConversionBlueprintLibrary::StringifyAsMatrix2dArray(const TArray<FMatrix2D>& Value)
+{
+	return UsdUtils::StringifyAsMatrix2dArray(Value);
+}
+
+FString UUsdConversionBlueprintLibrary::StringifyAsMatrix3dArray(const TArray<FMatrix3D>& Value)
+{
+	return UsdUtils::StringifyAsMatrix3dArray(Value);
+}
+
+FString UUsdConversionBlueprintLibrary::StringifyAsMatrix4dArray(const TArray<FMatrix>& Value)
+{
+	return UsdUtils::StringifyAsMatrix4dArray(Value);
+}
+
+FString UUsdConversionBlueprintLibrary::StringifyAsQuatdArray(const TArray<FQuat>& Value)
+{
+	return UsdUtils::StringifyAsQuatdArray(Value);
+}
+
+FString UUsdConversionBlueprintLibrary::StringifyAsQuatfArray(const TArray<FQuat>& Value)
+{
+	return UsdUtils::StringifyAsQuatfArray(Value);
+}
+
+FString UUsdConversionBlueprintLibrary::StringifyAsQuathArray(const TArray<FQuat>& Value)
+{
+	return UsdUtils::StringifyAsQuathArray(Value);
+}
+
+FString UUsdConversionBlueprintLibrary::StringifyAsDouble2Array(const TArray<FVector2D>& Value)
+{
+	return UsdUtils::StringifyAsDouble2Array(Value);
+}
+
+FString UUsdConversionBlueprintLibrary::StringifyAsFloat2Array(const TArray<FVector2D>& Value)
+{
+	return UsdUtils::StringifyAsFloat2Array(Value);
+}
+
+FString UUsdConversionBlueprintLibrary::StringifyAsHalf2Array(const TArray<FVector2D>& Value)
+{
+	return UsdUtils::StringifyAsHalf2Array(Value);
+}
+
+FString UUsdConversionBlueprintLibrary::StringifyAsInt2Array(const TArray<FIntPoint>& Value)
+{
+	return UsdUtils::StringifyAsInt2Array(Value);
+}
+
+FString UUsdConversionBlueprintLibrary::StringifyAsDouble3Array(const TArray<FVector>& Value)
+{
+	return UsdUtils::StringifyAsDouble3Array(Value);
+}
+
+FString UUsdConversionBlueprintLibrary::StringifyAsFloat3Array(const TArray<FVector>& Value)
+{
+	return UsdUtils::StringifyAsFloat3Array(Value);
+}
+
+FString UUsdConversionBlueprintLibrary::StringifyAsHalf3Array(const TArray<FVector>& Value)
+{
+	return UsdUtils::StringifyAsHalf3Array(Value);
+}
+
+FString UUsdConversionBlueprintLibrary::StringifyAsInt3Array(const TArray<FIntVector>& Value)
+{
+	return UsdUtils::StringifyAsInt3Array(Value);
+}
+
+FString UUsdConversionBlueprintLibrary::StringifyAsDouble4Array(const TArray<FVector4>& Value)
+{
+	return UsdUtils::StringifyAsDouble4Array(Value);
+}
+
+FString UUsdConversionBlueprintLibrary::StringifyAsFloat4Array(const TArray<FVector4>& Value)
+{
+	return UsdUtils::StringifyAsFloat4Array(Value);
+}
+
+FString UUsdConversionBlueprintLibrary::StringifyAsHalf4Array(const TArray<FVector4>& Value)
+{
+	return UsdUtils::StringifyAsHalf4Array(Value);
+}
+
+FString UUsdConversionBlueprintLibrary::StringifyAsInt4Array(const TArray<FIntVector4>& Value)
+{
+	return UsdUtils::StringifyAsInt4Array(Value);
+}
+
+namespace UE::UsdConversionLibrary::Private
+{
+	// It seems the most idiomatic behavior for stuff exposed to scripting is to always return a default value on failure
+	// but to emit warnings in case anything went wrong, as dealing with error outputs in Blueprint can be tedious, and having
+	// output parameters in Python feels weird. Here we use this adapter to convert the UsdUtils converters into signatures like that
+	template<typename T, bool (*Unstringifier)(const FString&, T&)>
+	T UnstringifyOrDefault(const FString& String, const FString& TypeName)
+	{
+		T Output{};
+
+		bool bSuccess = Unstringifier(String, Output);
+		if (!bSuccess)
+		{
+			UE_LOG(LogUsd, Warning, TEXT("Failed to unstringify '%s' with the type '%s'!"), *String, *TypeName);
+		}
+
+		return Output;
+	}
+}
+
+bool UUsdConversionBlueprintLibrary::UnstringifyAsBool(const FString& String)
+{
+	return UE::UsdConversionLibrary::Private::UnstringifyOrDefault<bool, UsdUtils::UnstringifyAsBool>(String, TEXT("bool"));
+}
+
+uint8 UUsdConversionBlueprintLibrary::UnstringifyAsUChar(const FString& String)
+{
+	return UE::UsdConversionLibrary::Private::UnstringifyOrDefault<uint8, UsdUtils::UnstringifyAsUChar>(String, TEXT("uchar"));
+}
+
+int32 UUsdConversionBlueprintLibrary::UnstringifyAsInt(const FString& String)
+{
+	return UE::UsdConversionLibrary::Private::UnstringifyOrDefault<int32, UsdUtils::UnstringifyAsInt>(String, TEXT("int"));
+}
+
+int32 UUsdConversionBlueprintLibrary::UnstringifyAsUInt(const FString& String)
+{
+	// Have to cast as there are no unsigned types higher than uint8 on blueprint...
+	return static_cast<int32>(
+		UE::UsdConversionLibrary::Private::UnstringifyOrDefault<uint32, UsdUtils::UnstringifyAsUInt>(String, TEXT("uint"))
+	);
+}
+
+int64 UUsdConversionBlueprintLibrary::UnstringifyAsInt64(const FString& String)
+{
+	return UE::UsdConversionLibrary::Private::UnstringifyOrDefault<int64, UsdUtils::UnstringifyAsInt64>(String, TEXT("int64"));
+}
+
+int64 UUsdConversionBlueprintLibrary::UnstringifyAsUInt64(const FString& String)
+{
+	// Have to cast as there are no unsigned types higher than uint8 on blueprint...
+	return static_cast<int64>(
+		UE::UsdConversionLibrary::Private::UnstringifyOrDefault<uint64, UsdUtils::UnstringifyAsUInt64>(String, TEXT("uint64"))
+	);
+}
+
+float UUsdConversionBlueprintLibrary::UnstringifyAsHalf(const FString& String)
+{
+	return UE::UsdConversionLibrary::Private::UnstringifyOrDefault<float, UsdUtils::UnstringifyAsHalf>(String, TEXT("half"));
+}
+
+float UUsdConversionBlueprintLibrary::UnstringifyAsFloat(const FString& String)
+{
+	return UE::UsdConversionLibrary::Private::UnstringifyOrDefault<float, UsdUtils::UnstringifyAsFloat>(String, TEXT("float"));
+}
+
+double UUsdConversionBlueprintLibrary::UnstringifyAsDouble(const FString& String)
+{
+	return UE::UsdConversionLibrary::Private::UnstringifyOrDefault<double, UsdUtils::UnstringifyAsDouble>(String, TEXT("double"));
+}
+
+double UUsdConversionBlueprintLibrary::UnstringifyAsTimeCode(const FString& String)
+{
+	return UE::UsdConversionLibrary::Private::UnstringifyOrDefault<double, UsdUtils::UnstringifyAsTimeCode>(String, TEXT("timecode"));
+}
+
+FString UUsdConversionBlueprintLibrary::UnstringifyAsString(const FString& String)
+{
+	return String;
+}
+
+FString UUsdConversionBlueprintLibrary::UnstringifyAsToken(const FString& String)
+{
+	return UE::UsdConversionLibrary::Private::UnstringifyOrDefault<FString, UsdUtils::UnstringifyAsToken>(String, TEXT("token"));
+}
+
+FString UUsdConversionBlueprintLibrary::UnstringifyAsAssetPath(const FString& String)
+{
+	return UE::UsdConversionLibrary::Private::UnstringifyOrDefault<FString, UsdUtils::UnstringifyAsAssetPath>(String, TEXT("asset"));
+}
+
+FMatrix2D UUsdConversionBlueprintLibrary::UnstringifyAsMatrix2d(const FString& String)
+{
+	return UE::UsdConversionLibrary::Private::UnstringifyOrDefault<FMatrix2D, UsdUtils::UnstringifyAsMatrix2d>(String, TEXT("matrix2d"));
+}
+
+FMatrix3D UUsdConversionBlueprintLibrary::UnstringifyAsMatrix3d(const FString& String)
+{
+	return UE::UsdConversionLibrary::Private::UnstringifyOrDefault<FMatrix3D, UsdUtils::UnstringifyAsMatrix3d>(String, TEXT("matrix3d"));
+}
+
+FMatrix UUsdConversionBlueprintLibrary::UnstringifyAsMatrix4d(const FString& String)
+{
+	return UE::UsdConversionLibrary::Private::UnstringifyOrDefault<FMatrix, UsdUtils::UnstringifyAsMatrix4d>(String, TEXT("matrix4d"));
+}
+
+FQuat UUsdConversionBlueprintLibrary::UnstringifyAsQuatd(const FString& String)
+{
+	return UE::UsdConversionLibrary::Private::UnstringifyOrDefault<FQuat, UsdUtils::UnstringifyAsQuatd>(String, TEXT("quatd"));
+}
+
+FQuat UUsdConversionBlueprintLibrary::UnstringifyAsQuatf(const FString& String)
+{
+	return UE::UsdConversionLibrary::Private::UnstringifyOrDefault<FQuat, UsdUtils::UnstringifyAsQuatf>(String, TEXT("quatf"));
+}
+
+FQuat UUsdConversionBlueprintLibrary::UnstringifyAsQuath(const FString& String)
+{
+	return UE::UsdConversionLibrary::Private::UnstringifyOrDefault<FQuat, UsdUtils::UnstringifyAsQuath>(String, TEXT("quath"));
+}
+
+FVector2D UUsdConversionBlueprintLibrary::UnstringifyAsDouble2(const FString& String)
+{
+	return UE::UsdConversionLibrary::Private::UnstringifyOrDefault<FVector2D, UsdUtils::UnstringifyAsDouble2>(String, TEXT("double2"));
+}
+
+FVector2D UUsdConversionBlueprintLibrary::UnstringifyAsFloat2(const FString& String)
+{
+	return UE::UsdConversionLibrary::Private::UnstringifyOrDefault<FVector2D, UsdUtils::UnstringifyAsFloat2>(String, TEXT("float2"));
+}
+
+FVector2D UUsdConversionBlueprintLibrary::UnstringifyAsHalf2(const FString& String)
+{
+	return UE::UsdConversionLibrary::Private::UnstringifyOrDefault<FVector2D, UsdUtils::UnstringifyAsHalf2>(String, TEXT("half2"));
+}
+
+FIntPoint UUsdConversionBlueprintLibrary::UnstringifyAsInt2(const FString& String)
+{
+	return UE::UsdConversionLibrary::Private::UnstringifyOrDefault<FIntPoint, UsdUtils::UnstringifyAsInt2>(String, TEXT("int2"));
+}
+
+FVector UUsdConversionBlueprintLibrary::UnstringifyAsDouble3(const FString& String)
+{
+	return UE::UsdConversionLibrary::Private::UnstringifyOrDefault<FVector, UsdUtils::UnstringifyAsDouble3>(String, TEXT("double3"));
+}
+
+FVector UUsdConversionBlueprintLibrary::UnstringifyAsFloat3(const FString& String)
+{
+	return UE::UsdConversionLibrary::Private::UnstringifyOrDefault<FVector, UsdUtils::UnstringifyAsFloat3>(String, TEXT("float3"));
+}
+
+FVector UUsdConversionBlueprintLibrary::UnstringifyAsHalf3(const FString& String)
+{
+	return UE::UsdConversionLibrary::Private::UnstringifyOrDefault<FVector, UsdUtils::UnstringifyAsHalf3>(String, TEXT("half3"));
+}
+
+FIntVector UUsdConversionBlueprintLibrary::UnstringifyAsInt3(const FString& String)
+{
+	return UE::UsdConversionLibrary::Private::UnstringifyOrDefault<FIntVector, UsdUtils::UnstringifyAsInt3>(String, TEXT("int3"));
+}
+
+FVector4 UUsdConversionBlueprintLibrary::UnstringifyAsDouble4(const FString& String)
+{
+	return UE::UsdConversionLibrary::Private::UnstringifyOrDefault<FVector4, UsdUtils::UnstringifyAsDouble4>(String, TEXT("double4"));
+}
+
+FVector4 UUsdConversionBlueprintLibrary::UnstringifyAsFloat4(const FString& String)
+{
+	return UE::UsdConversionLibrary::Private::UnstringifyOrDefault<FVector4, UsdUtils::UnstringifyAsFloat4>(String, TEXT("float4"));
+}
+
+FVector4 UUsdConversionBlueprintLibrary::UnstringifyAsHalf4(const FString& String)
+{
+	return UE::UsdConversionLibrary::Private::UnstringifyOrDefault<FVector4, UsdUtils::UnstringifyAsHalf4>(String, TEXT("half4"));
+}
+
+FIntVector4 UUsdConversionBlueprintLibrary::UnstringifyAsInt4(const FString& String)
+{
+	return UE::UsdConversionLibrary::Private::UnstringifyOrDefault<FIntVector4, UsdUtils::UnstringifyAsInt4>(String, TEXT("int4"));
+}
+
+TArray<bool> UUsdConversionBlueprintLibrary::UnstringifyAsBoolArray(const FString& String)
+{
+	return UE::UsdConversionLibrary::Private::UnstringifyOrDefault<TArray<bool>, UsdUtils::UnstringifyAsBoolArray>(String, TEXT("bool[]"));
+}
+
+TArray<uint8> UUsdConversionBlueprintLibrary::UnstringifyAsUCharArray(const FString& String)
+{
+	return UE::UsdConversionLibrary::Private::UnstringifyOrDefault<TArray<uint8>, UsdUtils::UnstringifyAsUCharArray>(String, TEXT("uchar[]"));
+}
+
+TArray<int32> UUsdConversionBlueprintLibrary::UnstringifyAsIntArray(const FString& String)
+{
+	return UE::UsdConversionLibrary::Private::UnstringifyOrDefault<TArray<int32>, UsdUtils::UnstringifyAsIntArray>(String, TEXT("int[]"));
+}
+
+TArray<int32> UUsdConversionBlueprintLibrary::UnstringifyAsUIntArray(const FString& String)
+{
+	static_assert(sizeof(uint32) == sizeof(int32));
+
+	TArray<uint32> Converted = UE::UsdConversionLibrary::Private::UnstringifyOrDefault<TArray<uint32>, UsdUtils::UnstringifyAsUIntArray>(
+		String,
+		TEXT("uint[]")
+	);
+
+	TArray<int32> Result;
+	Result.SetNumUninitialized(Converted.Num());
+	FMemory::Memcpy(Result.GetData(), Converted.GetData(), Converted.Num() * Converted.GetTypeSize());
+
+	return Result;
+}
+
+TArray<int64> UUsdConversionBlueprintLibrary::UnstringifyAsInt64Array(const FString& String)
+{
+	return UE::UsdConversionLibrary::Private::UnstringifyOrDefault<TArray<int64>, UsdUtils::UnstringifyAsInt64Array>(String, TEXT("int64[]"));
+}
+
+TArray<int64> UUsdConversionBlueprintLibrary::UnstringifyAsUInt64Array(const FString& String)
+{
+	static_assert(sizeof(uint64) == sizeof(int64));
+
+	TArray<uint64> Converted = UE::UsdConversionLibrary::Private::UnstringifyOrDefault<TArray<uint64>, UsdUtils::UnstringifyAsUInt64Array>(
+		String,
+		TEXT("uint64[]")
+	);
+
+	TArray<int64> Result;
+	Result.SetNumUninitialized(Converted.Num());
+	FMemory::Memcpy(Result.GetData(), Converted.GetData(), Converted.Num() * Converted.GetTypeSize());
+
+	return Result;
+}
+
+TArray<float> UUsdConversionBlueprintLibrary::UnstringifyAsHalfArray(const FString& String)
+{
+	return UE::UsdConversionLibrary::Private::UnstringifyOrDefault<TArray<float>, UsdUtils::UnstringifyAsHalfArray>(String, TEXT("half[]"));
+}
+
+TArray<float> UUsdConversionBlueprintLibrary::UnstringifyAsFloatArray(const FString& String)
+{
+	return UE::UsdConversionLibrary::Private::UnstringifyOrDefault<TArray<float>, UsdUtils::UnstringifyAsFloatArray>(String, TEXT("float[]"));
+}
+
+TArray<double> UUsdConversionBlueprintLibrary::UnstringifyAsDoubleArray(const FString& String)
+{
+	return UE::UsdConversionLibrary::Private::UnstringifyOrDefault<TArray<double>, UsdUtils::UnstringifyAsDoubleArray>(String, TEXT("double[]"));
+}
+
+TArray<double> UUsdConversionBlueprintLibrary::UnstringifyAsTimeCodeArray(const FString& String)
+{
+	return UE::UsdConversionLibrary::Private::UnstringifyOrDefault<TArray<double>, UsdUtils::UnstringifyAsTimeCodeArray>(String, TEXT("timecode[]"));
+}
+
+TArray<FString> UUsdConversionBlueprintLibrary::UnstringifyAsStringArray(const FString& String)
+{
+	return UE::UsdConversionLibrary::Private::UnstringifyOrDefault<TArray<FString>, UsdUtils::UnstringifyAsStringArray>(String, TEXT("string[]"));
+}
+
+TArray<FString> UUsdConversionBlueprintLibrary::UnstringifyAsTokenArray(const FString& String)
+{
+	return UE::UsdConversionLibrary::Private::UnstringifyOrDefault<TArray<FString>, UsdUtils::UnstringifyAsTokenArray>(String, TEXT("token[]"));
+}
+
+TArray<FString> UUsdConversionBlueprintLibrary::UnstringifyAsAssetPathArray(const FString& String)
+{
+	return UE::UsdConversionLibrary::Private::UnstringifyOrDefault<TArray<FString>, UsdUtils::UnstringifyAsAssetPathArray>(String, TEXT("asset[]"));
+}
+
+TArray<FString> UUsdConversionBlueprintLibrary::UnstringifyAsListOpTokens(const FString& String)
+{
+	return UE::UsdConversionLibrary::Private::UnstringifyOrDefault<TArray<FString>, UsdUtils::UnstringifyAsListOpTokens>(String, TEXT("SdfListOp<Token>"));
+}
+
+TArray<FMatrix2D> UUsdConversionBlueprintLibrary::UnstringifyAsMatrix2dArray(const FString& String)
+{
+	return UE::UsdConversionLibrary::Private::UnstringifyOrDefault<TArray<FMatrix2D>, UsdUtils::UnstringifyAsMatrix2dArray>(String, TEXT("matrix2d[]"));
+}
+
+TArray<FMatrix3D> UUsdConversionBlueprintLibrary::UnstringifyAsMatrix3dArray(const FString& String)
+{
+	return UE::UsdConversionLibrary::Private::UnstringifyOrDefault<TArray<FMatrix3D>, UsdUtils::UnstringifyAsMatrix3dArray>(String, TEXT("matrix3d[]"));
+}
+
+TArray<FMatrix> UUsdConversionBlueprintLibrary::UnstringifyAsMatrix4dArray(const FString& String)
+{
+	return UE::UsdConversionLibrary::Private::UnstringifyOrDefault<TArray<FMatrix>, UsdUtils::UnstringifyAsMatrix4dArray>(String, TEXT("matrix4d[]"));
+}
+
+TArray<FQuat> UUsdConversionBlueprintLibrary::UnstringifyAsQuatdArray(const FString& String)
+{
+	return UE::UsdConversionLibrary::Private::UnstringifyOrDefault<TArray<FQuat>, UsdUtils::UnstringifyAsQuatdArray>(String, TEXT("quatd[]"));
+}
+
+TArray<FQuat> UUsdConversionBlueprintLibrary::UnstringifyAsQuatfArray(const FString& String)
+{
+	return UE::UsdConversionLibrary::Private::UnstringifyOrDefault<TArray<FQuat>, UsdUtils::UnstringifyAsQuatfArray>(String, TEXT("quatf[]"));
+}
+
+TArray<FQuat> UUsdConversionBlueprintLibrary::UnstringifyAsQuathArray(const FString& String)
+{
+	return UE::UsdConversionLibrary::Private::UnstringifyOrDefault<TArray<FQuat>, UsdUtils::UnstringifyAsQuathArray>(String, TEXT("quath[]"));
+}
+
+TArray<FVector2D> UUsdConversionBlueprintLibrary::UnstringifyAsDouble2Array(const FString& String)
+{
+	return UE::UsdConversionLibrary::Private::UnstringifyOrDefault<TArray<FVector2D>, UsdUtils::UnstringifyAsDouble2Array>(String, TEXT("double2"));
+}
+
+TArray<FVector2D> UUsdConversionBlueprintLibrary::UnstringifyAsFloat2Array(const FString& String)
+{
+	return UE::UsdConversionLibrary::Private::UnstringifyOrDefault<TArray<FVector2D>, UsdUtils::UnstringifyAsFloat2Array>(String, TEXT("float2"));
+}
+
+TArray<FVector2D> UUsdConversionBlueprintLibrary::UnstringifyAsHalf2Array(const FString& String)
+{
+	return UE::UsdConversionLibrary::Private::UnstringifyOrDefault<TArray<FVector2D>, UsdUtils::UnstringifyAsHalf2Array>(String, TEXT("half2"));
+}
+
+TArray<FIntPoint> UUsdConversionBlueprintLibrary::UnstringifyAsInt2Array(const FString& String)
+{
+	return UE::UsdConversionLibrary::Private::UnstringifyOrDefault<TArray<FIntPoint>, UsdUtils::UnstringifyAsInt2Array>(String, TEXT("int2"));
+}
+
+TArray<FVector> UUsdConversionBlueprintLibrary::UnstringifyAsDouble3Array(const FString& String)
+{
+	return UE::UsdConversionLibrary::Private::UnstringifyOrDefault<TArray<FVector>, UsdUtils::UnstringifyAsDouble3Array>(String, TEXT("double3"));
+}
+
+TArray<FVector> UUsdConversionBlueprintLibrary::UnstringifyAsFloat3Array(const FString& String)
+{
+	return UE::UsdConversionLibrary::Private::UnstringifyOrDefault<TArray<FVector>, UsdUtils::UnstringifyAsFloat3Array>(String, TEXT("float3"));
+}
+
+TArray<FVector> UUsdConversionBlueprintLibrary::UnstringifyAsHalf3Array(const FString& String)
+{
+	return UE::UsdConversionLibrary::Private::UnstringifyOrDefault<TArray<FVector>, UsdUtils::UnstringifyAsHalf3Array>(String, TEXT("half3"));
+}
+
+TArray<FIntVector> UUsdConversionBlueprintLibrary::UnstringifyAsInt3Array(const FString& String)
+{
+	return UE::UsdConversionLibrary::Private::UnstringifyOrDefault<TArray<FIntVector>, UsdUtils::UnstringifyAsInt3Array>(String, TEXT("int3"));
+}
+
+TArray<FVector4> UUsdConversionBlueprintLibrary::UnstringifyAsDouble4Array(const FString& String)
+{
+	return UE::UsdConversionLibrary::Private::UnstringifyOrDefault<TArray<FVector4>, UsdUtils::UnstringifyAsDouble4Array>(String, TEXT("double4"));
+}
+
+TArray<FVector4> UUsdConversionBlueprintLibrary::UnstringifyAsFloat4Array(const FString& String)
+{
+	return UE::UsdConversionLibrary::Private::UnstringifyOrDefault<TArray<FVector4>, UsdUtils::UnstringifyAsFloat4Array>(String, TEXT("float4"));
+}
+
+TArray<FVector4> UUsdConversionBlueprintLibrary::UnstringifyAsHalf4Array(const FString& String)
+{
+	return UE::UsdConversionLibrary::Private::UnstringifyOrDefault<TArray<FVector4>, UsdUtils::UnstringifyAsHalf4Array>(String, TEXT("half4"));
+}
+
+TArray<FIntVector4> UUsdConversionBlueprintLibrary::UnstringifyAsInt4Array(const FString& String)
+{
+	return UE::UsdConversionLibrary::Private::UnstringifyOrDefault<TArray<FIntVector4>, UsdUtils::UnstringifyAsInt4Array>(String, TEXT("int4"));
+}
diff --git a/Engine/Plugins/Importers/USDImporter/Source/USDExporter/Private/USDExporterModule.cpp b/Engine/Plugins/Importers/USDImporter/Source/USDExporter/Private/USDExporterModule.cpp
--- a/Engine/Plugins/Importers/USDImporter/Source/USDExporter/Private/USDExporterModule.cpp
+++ b/Engine/Plugins/Importers/USDImporter/Source/USDExporter/Private/USDExporterModule.cpp
@@ -6,9 +6,11 @@
 #include "USDAssetOptions.h"
 #include "USDLog.h"
 #include "USDMemory.h"
+#include "USDMetadataExportOptions.h"
 
 #include "UsdWrappers/SdfLayer.h"
 
+#include "DetailLayoutBuilder.h"
 #include "Editor.h"
 #include "Framework/Notifications/NotificationManager.h"
 #include "Modules/ModuleManager.h"
@@ -19,6 +21,25 @@
 
 #define LOCTEXT_NAMESPACE "USDExporterModule"
 
+// Small customization to hide the "bExportComponentMetadata" property of shown FUsdMetadataExportOptions
+// for the asset exporters, since the inverse of having it hidden by default and only shown for level/sequence exporters
+// does not seem possible
+class FHideExportComponentMetadataCustomization : public IDetailCustomization
+{
+public:
+	static TSharedRef<IDetailCustomization> MakeInstance()
+	{
+		return MakeShared<FHideExportComponentMetadataCustomization>();
+	}
+	virtual void CustomizeDetails(IDetailLayoutBuilder& DetailLayoutBuilder) override
+	{
+		if (TSharedPtr<IPropertyHandle> Property = DetailLayoutBuilder.GetProperty(TEXT("MetadataOptions.bExportComponentMetadata")))
+		{
+			DetailLayoutBuilder.HideProperty(Property);
+		}
+	}
+};
+
 class FUsdExporterModule : public IUsdExporterModule
 {
 public:
@@ -28,6 +49,13 @@ public:
 
 		FPropertyEditorModule& PropertyModule = FModuleManager::LoadModuleChecked< FPropertyEditorModule >( TEXT( "PropertyEditor" ) );
 
+		// Hide the "bExportComponentMetadata" property if we're just exporting assets
+		PropertyModule.RegisterCustomClassLayout(TEXT("AnimSequenceExporterUSDOptions"), FOnGetDetailCustomizationInstance::CreateStatic(&FHideExportComponentMetadataCustomization::MakeInstance));
+		PropertyModule.RegisterCustomClassLayout(TEXT("MaterialExporterUSDOptions"), FOnGetDetailCustomizationInstance::CreateStatic(&FHideExportComponentMetadataCustomization::MakeInstance));
+		PropertyModule.RegisterCustomClassLayout(TEXT("SkeletalMeshExporterUSDOptions"), FOnGetDetailCustomizationInstance::CreateStatic(&FHideExportComponentMetadataCustomization::MakeInstance));
+		PropertyModule.RegisterCustomClassLayout(TEXT("StaticMeshExporterUSDOptions"), FOnGetDetailCustomizationInstance::CreateStatic(&FHideExportComponentMetadataCustomization::MakeInstance));
+		PropertyModule.RegisterCustomClassLayout(TEXT("GeometryCacheExporterUSDOptions"), FOnGetDetailCustomizationInstance::CreateStatic(&FHideExportComponentMetadataCustomization::MakeInstance));
+
 		// We intentionally use the same customization for both of these
 		PropertyModule.RegisterCustomClassLayout( TEXT( "LevelExporterUSDOptions" ), FOnGetDetailCustomizationInstance::CreateStatic( &FLevelExporterUSDOptionsCustomization::MakeInstance ) );
 		PropertyModule.RegisterCustomClassLayout( TEXT( "LevelSequenceExporterUSDOptions" ), FOnGetDetailCustomizationInstance::CreateStatic( &FLevelExporterUSDOptionsCustomization::MakeInstance ) );
@@ -52,6 +80,12 @@ public:
 	{
 		if ( FPropertyEditorModule* PropertyModule = FModuleManager::GetModulePtr< FPropertyEditorModule >( TEXT( "PropertyEditor" ) ) )
 		{
+			PropertyModule->UnregisterCustomClassLayout(TEXT("AnimSequenceExporterUSDOptions"));
+			PropertyModule->UnregisterCustomClassLayout(TEXT("MaterialExporterUSDOptions"));
+			PropertyModule->UnregisterCustomClassLayout(TEXT("SkeletalMeshExporterUSDOptions"));
+			PropertyModule->UnregisterCustomClassLayout(TEXT("StaticMeshExporterUSDOptions"));
+			PropertyModule->UnregisterCustomClassLayout(TEXT("GeometryCacheExporterUSDOptions"));
+
 			PropertyModule->UnregisterCustomClassLayout( TEXT( "LevelExporterUSDOptions" ) );
 			PropertyModule->UnregisterCustomClassLayout( TEXT( "LevelSequenceExporterUSDOptions" ) );
 		}
diff --git a/Engine/Plugins/Importers/USDImporter/Source/USDExporter/Public/AnimSequenceExporterUSDOptions.h b/Engine/Plugins/Importers/USDImporter/Source/USDExporter/Public/AnimSequenceExporterUSDOptions.h
--- a/Engine/Plugins/Importers/USDImporter/Source/USDExporter/Public/AnimSequenceExporterUSDOptions.h
+++ b/Engine/Plugins/Importers/USDImporter/Source/USDExporter/Public/AnimSequenceExporterUSDOptions.h
@@ -3,6 +3,7 @@
 #pragma once
 
 #include "USDAssetOptions.h"
+#include "USDMetadataExportOptions.h"
 #include "USDStageOptions.h"
 
 #include "Engine/EngineTypes.h"
@@ -32,6 +33,10 @@ public:
 	UPROPERTY( EditAnywhere, config, BlueprintReadWrite, Category = "Mesh options", meta = ( ShowOnlyInnerProperties, EditCondition = bExportPreviewMesh ) )
 	FUsdMeshAssetOptions PreviewMeshOptions;
 
+	/** Options to use when exporting UsdAssetUserData metadata to USD for all assets (including the preview mesh) */
+	UPROPERTY(EditAnywhere, config, BlueprintReadWrite, Category = "Metadata options", meta = (ShowOnlyInnerProperties))
+	FUsdMetadataExportOptions MetadataOptions;
+
 	/**
 	 * Whether to export any asset (StaticMesh, Material, etc.) even if the existing file already describes the same version of a compatible asset.
 	 * This is only checked when bReplaceIdentical is set on the asset export task. Otherwise we'll never overwrite files.
diff --git a/Engine/Plugins/Importers/USDImporter/Source/USDExporter/Public/GeometryCacheExporterUSDOptions.h b/Engine/Plugins/Importers/USDImporter/Source/USDExporter/Public/GeometryCacheExporterUSDOptions.h
--- a/Engine/Plugins/Importers/USDImporter/Source/USDExporter/Public/GeometryCacheExporterUSDOptions.h
+++ b/Engine/Plugins/Importers/USDImporter/Source/USDExporter/Public/GeometryCacheExporterUSDOptions.h
@@ -3,10 +3,9 @@
 #pragma once
 
 #include "USDAssetOptions.h"
+#include "USDMetadataExportOptions.h"
 #include "USDStageOptions.h"
 
-#include "Engine/EngineTypes.h"
-
 #include "GeometryCacheExporterUSDOptions.generated.h"
 
 struct FAnalyticsEventAttribute;
@@ -26,6 +25,9 @@ public:
 	UPROPERTY(EditAnywhere, config, BlueprintReadWrite, Category = "Mesh options", meta = (ShowOnlyInnerProperties))
 	FUsdMeshAssetOptions MeshAssetOptions;
 
+	UPROPERTY(EditAnywhere, config, BlueprintReadWrite, Category = "Metadata options", meta = (ShowOnlyInnerProperties))
+	FUsdMetadataExportOptions MetadataOptions;
+
 	/**
 	 * Whether to export any asset (StaticMesh, Material, etc.) even if the existing file already describes the same version of a compatible asset.
 	 * This is only checked when bReplaceIdentical is set on the asset export task. Otherwise we'll never overwrite files.
diff --git a/Engine/Plugins/Importers/USDImporter/Source/USDExporter/Public/LevelExporterUSDOptions.h b/Engine/Plugins/Importers/USDImporter/Source/USDExporter/Public/LevelExporterUSDOptions.h
--- a/Engine/Plugins/Importers/USDImporter/Source/USDExporter/Public/LevelExporterUSDOptions.h
+++ b/Engine/Plugins/Importers/USDImporter/Source/USDExporter/Public/LevelExporterUSDOptions.h
@@ -2,13 +2,11 @@
 
 #pragma once
 
-#include "UnrealUSDWrapper.h"
 #include "USDAssetOptions.h"
+#include "USDMetadataExportOptions.h"
 #include "USDStageOptions.h"
 
 #include "AssetExportTask.h"
-#include "Engine/EngineTypes.h"
-#include "RHIDefinitions.h"
 
 #include "LevelExporterUSDOptions.generated.h"
 
@@ -47,6 +45,10 @@ struct USDEXPORTER_API FLevelExporterUSDOptionsInner
 	UPROPERTY( EditAnywhere, config, BlueprintReadWrite, Category = "Asset options", meta = ( ShowOnlyInnerProperties ) )
 	FUsdMeshAssetOptions AssetOptions;
 
+	/** Options to use when exporting UsdAssetUserData metadata to USD for all assets (including the level itself) */
+	UPROPERTY(EditAnywhere, config, BlueprintReadWrite, Category = "Metadata options", meta = (ShowOnlyInnerProperties))
+	FUsdMetadataExportOptions MetadataOptions;
+
 	/** Lowest of the LOD indices to export landscapes with (use 0 for full resolution) */
 	UPROPERTY( EditAnywhere, config, BlueprintReadWrite, Category = "Landscape options", meta = ( ClampMin = "0" ) )
 	int32 LowestLandscapeLOD = 0;
diff --git a/Engine/Plugins/Importers/USDImporter/Source/USDExporter/Public/MaterialExporterUSD.h b/Engine/Plugins/Importers/USDImporter/Source/USDExporter/Public/MaterialExporterUSD.h
--- a/Engine/Plugins/Importers/USDImporter/Source/USDExporter/Public/MaterialExporterUSD.h
+++ b/Engine/Plugins/Importers/USDImporter/Source/USDExporter/Public/MaterialExporterUSD.h
@@ -4,12 +4,11 @@
 
 #include "Exporters/Exporter.h"
 
-#include "UsdWrappers/ForwardDeclarations.h"
-
 #include "MaterialExporterUSD.generated.h"
 
 class UMaterialInterface;
 struct FUsdMaterialBakingOptions;
+struct FUsdMetadataExportOptions;
 
 UCLASS()
 class UMaterialExporterUsd : public UExporter
@@ -20,7 +19,7 @@ public:
 	UMaterialExporterUsd();
 
 	//~ Begin UExporter Interface
-	virtual bool ExportBinary( UObject* Object, const TCHAR* Type, FArchive& Ar, FFeedbackContext* Warn, int32 FileIndex = 0, uint32 PortFlags=0 ) override;
+	USDEXPORTER_API virtual bool ExportBinary( UObject* Object, const TCHAR* Type, FArchive& Ar, FFeedbackContext* Warn, int32 FileIndex = 0, uint32 PortFlags=0 ) override;
 	//~ End UExporter Interface
 
 	/**
@@ -33,6 +32,16 @@ public:
 	 * @param bIsAutomated - Whether the export is being done by a script or not. Just used for analytics
 	 * @return Whether the export was successful or not.
 	 */
+	USDEXPORTER_API static bool ExportMaterial(
+		const UMaterialInterface& Material,
+		const FUsdMaterialBakingOptions& Options,
+		const FUsdMetadataExportOptions& MetadataOptions,
+		const FFilePath& FilePath,
+		bool bReplaceIdentical = true,
+		bool bReExportIdenticalAssets = false,
+		bool bIsAutomated = false
+	);
+	UE_DEPRECATED(5.4, "Please use the function overload that can receive metadata options.")
 	static bool ExportMaterial(
 		const UMaterialInterface& Material,
 		const FUsdMaterialBakingOptions& Options,
@@ -61,6 +70,18 @@ public:
 	 * @param bIsAutomated - Whether the export is being done by a script or not. Just used for analytics
 	 * @return Whether the export was successful or not.
 	 */
+	USDEXPORTER_API static bool ExportMaterialsForStage(
+		const TArray<UMaterialInterface*>& Materials,
+		const FUsdMaterialBakingOptions& Options,
+		const FUsdMetadataExportOptions& MetadataOptions,
+		const FString& StageRootLayerPath,
+		bool bIsAssetLayer,
+		bool bUsePayload,
+		bool bReplaceIdentical = true,
+		bool bReExportIdenticalAssets = false,
+		bool bIsAutomated = false
+	);
+	UE_DEPRECATED(5.4, "Please use the function overload that can receive metadata options.")
 	static bool ExportMaterialsForStage(
 		const TArray<UMaterialInterface*>& Materials,
 		const FUsdMaterialBakingOptions& Options,
diff --git a/Engine/Plugins/Importers/USDImporter/Source/USDExporter/Public/MaterialExporterUSDOptions.h b/Engine/Plugins/Importers/USDImporter/Source/USDExporter/Public/MaterialExporterUSDOptions.h
--- a/Engine/Plugins/Importers/USDImporter/Source/USDExporter/Public/MaterialExporterUSDOptions.h
+++ b/Engine/Plugins/Importers/USDImporter/Source/USDExporter/Public/MaterialExporterUSDOptions.h
@@ -3,8 +3,7 @@
 #pragma once
 
 #include "USDAssetOptions.h"
-
-#include "Engine/EngineTypes.h"
+#include "USDMetadataExportOptions.h"
 
 #include "MaterialExporterUSDOptions.generated.h"
 
@@ -24,6 +23,9 @@ public:
 	UPROPERTY( EditAnywhere, config, BlueprintReadWrite, Category = "Material baking options", meta = ( ShowOnlyInnerProperties ))
 	FUsdMaterialBakingOptions MaterialBakingOptions;
 
+	UPROPERTY(EditAnywhere, config, BlueprintReadWrite, Category = "Metadata options", meta = (ShowOnlyInnerProperties))
+	FUsdMetadataExportOptions MetadataOptions;
+
 	/**
 	 * Whether to export any asset (StaticMesh, Material, etc.) even if the existing file already describes the same version of a compatible asset.
 	 * This is only checked when bReplaceIdentical is set on the asset export task. Otherwise we'll never overwrite files.
diff --git a/Engine/Plugins/Importers/USDImporter/Source/USDExporter/Public/SkeletalMeshExporterUSDOptions.h b/Engine/Plugins/Importers/USDImporter/Source/USDExporter/Public/SkeletalMeshExporterUSDOptions.h
--- a/Engine/Plugins/Importers/USDImporter/Source/USDExporter/Public/SkeletalMeshExporterUSDOptions.h
+++ b/Engine/Plugins/Importers/USDImporter/Source/USDExporter/Public/SkeletalMeshExporterUSDOptions.h
@@ -3,10 +3,9 @@
 #pragma once
 
 #include "USDAssetOptions.h"
+#include "USDMetadataExportOptions.h"
 #include "USDStageOptions.h"
 
-#include "Engine/EngineTypes.h"
-
 #include "SkeletalMeshExporterUSDOptions.generated.h"
 
 struct FAnalyticsEventAttribute;
@@ -26,6 +25,9 @@ public:
 	UPROPERTY( EditAnywhere, config, BlueprintReadWrite, Category = "Mesh options", meta = ( ShowOnlyInnerProperties ) )
 	FUsdMeshAssetOptions MeshAssetOptions;
 
+	UPROPERTY( EditAnywhere, config, BlueprintReadWrite, Category = "Metadata options", meta = ( ShowOnlyInnerProperties ) )
+	FUsdMetadataExportOptions MetadataOptions;
+
 	/**
 	 * Whether to export any asset (StaticMesh, Material, etc.) even if the existing file already describes the same version of a compatible asset.
 	 * This is only checked when bReplaceIdentical is set on the asset export task. Otherwise we'll never overwrite files.
diff --git a/Engine/Plugins/Importers/USDImporter/Source/USDExporter/Public/StaticMeshExporterUSDOptions.h b/Engine/Plugins/Importers/USDImporter/Source/USDExporter/Public/StaticMeshExporterUSDOptions.h
--- a/Engine/Plugins/Importers/USDImporter/Source/USDExporter/Public/StaticMeshExporterUSDOptions.h
+++ b/Engine/Plugins/Importers/USDImporter/Source/USDExporter/Public/StaticMeshExporterUSDOptions.h
@@ -4,8 +4,7 @@
 
 #include "USDAssetOptions.h"
 #include "USDStageOptions.h"
-
-#include "Engine/EngineTypes.h"
+#include "USDMetadataExportOptions.h"
 
 #include "StaticMeshExporterUSDOptions.generated.h"
 
@@ -26,6 +25,9 @@ public:
 	UPROPERTY( EditAnywhere, config, BlueprintReadWrite, Category = "Mesh options", meta = ( ShowOnlyInnerProperties ) )
 	FUsdMeshAssetOptions MeshAssetOptions;
 
+	UPROPERTY(EditAnywhere, config, BlueprintReadWrite, Category = "Metadata options", meta = (ShowOnlyInnerProperties))
+	FUsdMetadataExportOptions MetadataOptions;
+
 	/**
 	 * Whether to export any asset (StaticMesh, Material, etc.) even if the existing file already describes the same version of a compatible asset.
 	 * This is only checked when bReplaceIdentical is set on the asset export task. Otherwise we'll never overwrite files.
