USD: Add support for importing and exporting prim metadata.

diff --git a/Engine/Plugins/Importers/USDImporter/Source/USDExporter/Public/USDConversionBlueprintContext.h b/Engine/Plugins/Importers/USDImporter/Source/USDExporter/Public/USDConversionBlueprintContext.h
--- a/Engine/Plugins/Importers/USDImporter/Source/USDExporter/Public/USDConversionBlueprintContext.h
+++ b/Engine/Plugins/Importers/USDImporter/Source/USDExporter/Public/USDConversionBlueprintContext.h
@@ -3,6 +3,7 @@
 #pragma once
 
 #include "UnrealUSDWrapper.h"
+#include "USDAssetUserData.h"
 #include "USDStageOptions.h"
 #include "USDUnrealAssetInfo.h"
 #include "UsdWrappers/SdfLayer.h"
@@ -27,6 +28,8 @@ class USceneComponent;
 class USkyLightComponent;
 class USpotLightComponent;
 class UUsdDrawModeComponent;
+struct FUsdPrimMetadata;
+struct FUsdCombinedPrimMetadata;
 
 /**
  * Wraps the UnrealToUsd component conversion functions from the USDUtilities module so that they can be used by
@@ -181,10 +184,22 @@ public:
 	int32 GetUsdStageNumFrames();
 
 	/** Adds to Prim the assetInfo metadata the values described in Info */
-	UFUNCTION( BlueprintCallable, Category = "Conversion utils" )
+	UFUNCTION( BlueprintCallable, Category = "Conversion utils|Metadata" )
 	void SetPrimAssetInfo( const FString& PrimPath, const FUsdUnrealAssetInfo& Info );
 
 	/** Retrieves from Prim the assetInfo metadata values that we use as export metadata, when exporting Unreal assets */
-	UFUNCTION( BlueprintCallable, Category = "Conversion utils" )
+	UFUNCTION( BlueprintCallable, Category = "Conversion utils|Metadata" )
 	FUsdUnrealAssetInfo GetPrimAssetInfo( const FString& PrimPath );
+
+	/** Applies Metadata to the prim specified at PrimPath, using the provided filters */
+	UFUNCTION(BlueprintCallable, Category = "Conversion utils|Metadata")
+	void SetPrimMetadata(const FString& PrimPath, const FUsdCombinedPrimMetadata& Metadata, const TArray<FString>& BlockedPrefixFilter, bool bInvertFilter);
+
+	/** Applies all the metadata contained in UserData to the prim specified at PrimPath, using the provided filters */
+	UFUNCTION(BlueprintCallable, Category = "Conversion utils|Metadata")
+	void SetPrimMetadataFromUserData(const FString& PrimPath, const UUsdAssetUserData* UserData, const TArray<FString>& BlockedPrefixFilter, bool bInvertFilter);
+
+	/** Extracts metadata from the prim at PrimPath using the provided filters and returns it */
+	UFUNCTION(BlueprintCallable, Category = "Conversion utils|Metadata")
+	FUsdCombinedPrimMetadata GetPrimMetadata(const FString& PrimPath, const TArray<FString>& BlockedPrefixFilter, bool bInvertFilter, bool bCollectFromEntireSubtrees);
 };
diff --git a/Engine/Plugins/Importers/USDImporter/Source/USDExporter/Public/USDConversionBlueprintLibrary.h b/Engine/Plugins/Importers/USDImporter/Source/USDExporter/Public/USDConversionBlueprintLibrary.h
--- a/Engine/Plugins/Importers/USDImporter/Source/USDExporter/Public/USDConversionBlueprintLibrary.h
+++ b/Engine/Plugins/Importers/USDImporter/Source/USDExporter/Public/USDConversionBlueprintLibrary.h
@@ -2,6 +2,8 @@
 
 #pragma once
 
+#include "USDMetadata.h"
+
 #include "AnalyticsBlueprintLibrary.h"
 #include "Engine/EngineTypes.h"
 #include "Kismet/BlueprintFunctionLibrary.h"
@@ -12,6 +14,9 @@ class AInstancedFoliageActor;
 class UFoliageType;
 class ULevel;
 class ULevelExporterUSDOptions;
+class UUsdAssetUserData;
+struct FMatrix2D;
+struct FMatrix3D;
 
 /** Wrapped static conversion functions from the UsdUtilities module, so that they can be used via scripting */
 UCLASS(meta=(ScriptName="UsdConversionLibrary"))
@@ -21,46 +26,46 @@ class USDEXPORTER_API UUsdConversionBlueprintLibrary : public UBlueprintFunction
 
 public:
 	/** Returns how many total Unreal levels (persistent + all sublevels) will be exported if we consider LevelsToIgnore */
-	UFUNCTION( BlueprintCallable, Category = "World utils" )
+	UFUNCTION( BlueprintCallable, Category = "USD|World utils" )
 	static int32 GetNumLevelsToExport( UWorld* World, const TSet<FString>& LevelsToIgnore );
 
 	/** Fully streams in and displays all levels whose names are not in LevelsToIgnore */
-	UFUNCTION( BlueprintCallable, Category = "World utils" )
+	UFUNCTION( BlueprintCallable, Category = "USD|World utils" )
 	static void StreamInRequiredLevels( UWorld* World, const TSet<FString>& LevelsToIgnore );
 
 	/**
 	 * If we have the Sequencer open with a level sequence animating the level before export, this function can revert
 	 * any actor or component to its unanimated state
 	 */
-	UFUNCTION( BlueprintCallable, Category = "World utils" )
+	UFUNCTION( BlueprintCallable, Category = "USD|World utils" )
 	static void RevertSequencerAnimations();
 
 	/**
 	 * If we used `ReverseSequencerAnimations` to undo the effect of an opened sequencer before export, this function
 	 * can be used to re-apply the sequencer state back to the level after the export is complete
 	 */
-	UFUNCTION( BlueprintCallable, Category = "World utils" )
+	UFUNCTION( BlueprintCallable, Category = "USD|World utils" )
 	static void ReapplySequencerAnimations();
 
 	/**
 	 * Returns the path name (e.g. "/Game/Maps/MyLevel") of levels that are loaded on `World`.
 	 * We use these to revert the `World` to its initial state after we force-stream levels in for exporting
 	 */
-	UFUNCTION( BlueprintCallable, Category = "World utils" )
+	UFUNCTION( BlueprintCallable, Category = "USD|World utils" )
 	static TArray<FString> GetLoadedLevelNames( UWorld* World );
 
 	/**
 	 * Returns the path name (e.g. "/Game/Maps/MyLevel") of levels that checked to be visible in the editor within `World`.
 	 * We use these to revert the `World` to its initial state after we force-stream levels in for exporting
 	 */
-	UFUNCTION( BlueprintCallable, Category = "World utils" )
+	UFUNCTION( BlueprintCallable, Category = "USD|World utils" )
 	static TArray<FString> GetVisibleInEditorLevelNames( UWorld* World );
 
 	/** Streams out/hides sublevels that were streamed in before export */
-	UFUNCTION( BlueprintCallable, Category = "World utils" )
+	UFUNCTION( BlueprintCallable, Category = "USD|World utils" )
 	static void StreamOutLevels( UWorld* OwningWorld, const TArray<FString>& LevelNamesToStreamOut, const TArray<FString>& LevelNamesToHide );
 
-	UFUNCTION( BlueprintCallable, Category = "World utils" )
+	UFUNCTION( BlueprintCallable, Category = "USD|World utils" )
 	static TSet<AActor*> GetActorsToConvert( UWorld* World );
 
 	/**
@@ -72,29 +77,29 @@ public:
 	 * This can be used to track the version of exported assets and levels, to prevent reexporting of actors and
 	 * components.
 	 */
-	UFUNCTION( BlueprintCallable, Category = "World utils" )
+	UFUNCTION( BlueprintCallable, Category = "USD|World utils" )
 	static FString GenerateObjectVersionString( const UObject* ObjectToExport, UObject* ExportOptions );
 
 	/** Checks whether we can create a USD Layer with "TargetFilePath" as identifier and export to it */
-	UFUNCTION( BlueprintCallable, Category = "World utils" )
+	UFUNCTION( BlueprintCallable, Category = "USD|World utils" )
 	static bool CanExportToLayer( const FString& TargetFilePath );
 
-	UFUNCTION( BlueprintCallable, Category = "Layer utils" )
+	UFUNCTION( BlueprintCallable, Category = "USD|Layer utils" )
 	static FString MakePathRelativeToLayer( const FString& AnchorLayerPath, const FString& PathToMakeRelative );
 
-	UFUNCTION( BlueprintCallable, Category = "Layer utils" )
+	UFUNCTION( BlueprintCallable, Category = "USD|Layer utils" )
 	static void InsertSubLayer( const FString& ParentLayerPath, const FString& SubLayerPath, int32 Index = -1 );
 
-	UFUNCTION( BlueprintCallable, Category = "Layer utils" )
+	UFUNCTION( BlueprintCallable, Category = "USD|Layer utils" )
 	static void AddReference( const FString& ReferencingStagePath, const FString& ReferencingPrimPath, const FString& TargetStagePath );
 
-	UFUNCTION( BlueprintCallable, Category = "Layer utils" )
+	UFUNCTION( BlueprintCallable, Category = "USD|Layer utils" )
 	static void AddPayload( const FString& ReferencingStagePath, const FString& ReferencingPrimPath, const FString& TargetStagePath );
 
-	UFUNCTION( BlueprintCallable, Category = "Prim utils" )
+	UFUNCTION( BlueprintCallable, Category = "USD|Prim utils" )
 	static FString GetPrimPathForObject( const UObject* ActorOrComponent, const FString& ParentPrimPath = TEXT(""), bool bUseActorFolders = false );
 
-	UFUNCTION( BlueprintCallable, Category = "Prim utils" )
+	UFUNCTION( BlueprintCallable, Category = "USD|Prim utils" )
 	static FString GetSchemaNameForComponent( const USceneComponent* Component );
 
 	/**
@@ -103,7 +108,7 @@ public:
 	 * This function is useful because it's difficult to retrieve this actor otherwise, as it will be filtered from
 	 * the results of functions like EditorLevelLibrary.get_all_level_actors()
 	 */
-	UFUNCTION( BlueprintCallable, Category = "USD Foliage Exporter" )
+	UFUNCTION( BlueprintCallable, Category = "USD|Foliage Exporter" )
 	static AInstancedFoliageActor* GetInstancedFoliageActorForLevel( bool bCreateIfNone = false, ULevel* Level = nullptr );
 
 	/**
@@ -111,29 +116,29 @@ public:
 	 * This function exists because we want to retrieve all instances of all foliage types on an actor, but we
 	 * can't return nested containers from UFUNCTIONs, so users of this API should call this, and then GetInstanceTransforms.
 	 */
-	UFUNCTION( BlueprintCallable, meta = ( ScriptMethod ), Category = "USD Foliage Exporter" )
+	UFUNCTION( BlueprintCallable, meta = ( ScriptMethod ), Category = "USD|Foliage Exporter" )
 	static TArray<UFoliageType*> GetUsedFoliageTypes( AInstancedFoliageActor* Actor );
 
 	/**
 	 * Returns the source asset for a UFoliageType.
 	 * It can be a UStaticMesh in case we're dealing with a UFoliageType_InstancedStaticMesh, but it can be other types of objects.
 	 */
-	UFUNCTION( BlueprintCallable, meta = ( ScriptMethod ), Category = "USD Foliage Exporter" )
+	UFUNCTION( BlueprintCallable, meta = ( ScriptMethod ), Category = "USD|Foliage Exporter" )
 	static UObject* GetSource( UFoliageType* FoliageType );
 
 	/**
 	 * Returns the transforms of all instances of a particular UFoliageType on a given level. If no level is provided all instances will be returned.
 	 * Use GetUsedFoliageTypes() to retrieve all foliage types managed by a particular actor.
 	 */
-	UFUNCTION( BlueprintCallable, meta = ( ScriptMethod ), Category = "USD Foliage Exporter" )
+	UFUNCTION( BlueprintCallable, meta = ( ScriptMethod ), Category = "USD|Foliage Exporter" )
 	static TArray<FTransform> GetInstanceTransforms( AInstancedFoliageActor* Actor, UFoliageType* FoliageType, ULevel* InstancesLevel = nullptr );
 
 	/** Retrieves the analytics attributes to send for the provided options object */
-	UFUNCTION( BlueprintCallable, Category = "Analytics" )
+	UFUNCTION( BlueprintCallable, Category = "USD|Analytics" )
 	static TArray<FAnalyticsEventAttr> GetAnalyticsAttributes( const ULevelExporterUSDOptions* Options );
 
 	/** Defer to the USDClasses module to actually send analytics information */
-	UFUNCTION( BlueprintCallable, Category = "Analytics" )
+	UFUNCTION( BlueprintCallable, Category = "USD|Analytics" )
 	static void SendAnalytics( const TArray<FAnalyticsEventAttr>& Attrs, const FString& EventName, bool bAutomated, double ElapsedSeconds, double NumberOfFrames, const FString& Extension );
 
     /**
@@ -150,7 +155,7 @@ public:
 	 * @param Layer - Layer to remove prim specs from. This can be left with the invalid layer (default) in order to remove all
 	 *				  specs from the entire stage's local layer stack.
 	 */
-	UFUNCTION( BlueprintCallable, Category = "Prim utils" )
+	UFUNCTION( BlueprintCallable, Category = "USD|Prim utils" )
 	static void RemoveAllPrimSpecs( const FString& StageRootLayer, const FString& PrimPath, const FString& TargetLayer );
 
 	/**
@@ -161,7 +166,7 @@ public:
 	 * @param PrimPaths - Prims to cut
 	 * @return True if we managed to cut
 	 */
-	UFUNCTION( BlueprintCallable, Category = "Prim utils" )
+	UFUNCTION( BlueprintCallable, Category = "USD|Prim utils" )
 	static bool CutPrims( const FString& StageRootLayer, const TArray<FString>& PrimPaths );
 
 	/**
@@ -172,7 +177,7 @@ public:
 	 * @param PrimPaths - Prims to copy
 	 * @return True if we managed to copy
 	 */
-	UFUNCTION( BlueprintCallable, Category = "Prim utils" )
+	UFUNCTION( BlueprintCallable, Category = "USD|Prim utils" )
 	static bool CopyPrims( const FString& StageRootLayer, const TArray<FString>& PrimPaths );
 
 	/**
@@ -188,15 +193,15 @@ public:
 	 * @param ParentPrimPath - Prim that will become parent to the pasted prims
 	 * @return Paths to the pasted prim specs, after they were added as children of ParentPrim
 	 */
-	UFUNCTION( BlueprintCallable, Category = "Prim utils" )
+	UFUNCTION( BlueprintCallable, Category = "USD|Prim utils" )
 	static TArray<FString> PastePrims( const FString& StageRootLayer, const FString& ParentPrimPath );
 
 	/** Returns true if we have prims that we can paste within our clipboard stage */
-	UFUNCTION( BlueprintCallable, Category = "Prim utils" )
+	UFUNCTION( BlueprintCallable, Category = "USD|Prim utils" )
 	static bool CanPastePrims();
 
 	/** Clears all prims from our clipboard stage */
-	UFUNCTION( BlueprintCallable, Category = "Prim utils" )
+	UFUNCTION( BlueprintCallable, Category = "USD|Prim utils" )
 	static void ClearPrimClipboard();
 
 	/**
@@ -215,11 +220,452 @@ public:
 	 * @param TargetLayer - Target layer to use when duplicating, if relevant for that duplication type
 	 * @return Paths to the duplicated prim specs, after they were added as children of ParentPrim.
 	 */
-	UFUNCTION( BlueprintCallable, Category = "Prim utils" )
+	UFUNCTION( BlueprintCallable, Category = "USD|Prim utils" )
 	static TArray<FString> DuplicatePrims(
 		const FString& StageRootLayer,
 		const TArray<FString>& PrimPaths,
 		EUsdDuplicateType DuplicateType,
 		const FString& TargetLayer
 	);
+
+	/* Retrieve the first instance of UUsdAssetUserData contained on the Object, if any */
+	UFUNCTION(BlueprintCallable, Category = "USD|Metadata utils")
+	static UUsdAssetUserData* GetUsdAssetUserData(UObject* Object);
+
+	/*
+	 * Sets AssetUserData as the single UUsdAssetUserData on the Object, overwriting an existing one if encountered.
+	 * Returns true if it managed to add AssetUserData to Object.
+	 */
+	UFUNCTION(BlueprintCallable, Category = "USD|Metadata utils")
+	static bool SetUsdAssetUserData(UObject* Object, UUsdAssetUserData* AssetUserData);
+
+	/*
+	 * Utilities that make it easier to get/set metadata fields without having to manipulate the nested struct instances directly.
+	 * It will create the struct entries automatically if needed, and overwrite existing entries with the same key if needed.
+	 *
+	 * If the AssetUserData contains exactly one entry for StageIdentifier and one entry for PrimPath, you can omit those arguments
+	 * and that single entry will be used. If there are more or less than exactly one entry for StageIdentifier or for PrimPath however,
+	 * you must specify which one to use, and failing to do so will cause the functions to return false and emit a warning.
+	 *
+	 * It is possible to get these functions to automatically trigger pre/post property changed events by providing "true" for
+	 * bTriggerPropertyChangeEvents, which is useful as it is not trivial to trigger those from Python/Blueprint given how the
+	 * metadata is stored inside nested structs and maps. If these AssetUserData belong to generated transient assets when opening
+	 * stages, emitting property change events causes those edits to be immediately written out to the opened USD Stage.
+	 *
+	 * Returns true if it managed to set the new key-value pair.
+	 */
+	UFUNCTION(BlueprintCallable, Category = "USD|Metadata utils")
+	static bool SetMetadataField(
+		UUsdAssetUserData* AssetUserData,
+		const FString& Key,
+		const FString& Value,
+		const FString& ValueTypeName,
+		const FString& StageIdentifier = "",
+		const FString& PrimPath = "",
+		bool bTriggerPropertyChangeEvents = true
+	);
+	UFUNCTION(BlueprintCallable, Category = "USD|Metadata utils")
+	static bool ClearMetadataField(
+		UUsdAssetUserData* AssetUserData,
+		const FString& Key,
+		const FString& StageIdentifier = "",
+		const FString& PrimPath = "",
+		bool bTriggerPropertyChangeEvents = true
+	);
+	UFUNCTION(BlueprintCallable, Category = "USD|Metadata utils")
+	static bool HasMetadataField(
+		UUsdAssetUserData* AssetUserData,
+		const FString& Key,
+		const FString& StageIdentifier = "",
+		const FString& PrimPath = ""
+	);
+	UFUNCTION(BlueprintCallable, Category = "USD|Metadata utils")
+	static FUsdMetadataValue GetMetadataField(
+		UUsdAssetUserData* AssetUserData,
+		const FString& Key,
+		const FString& StageIdentifier = "",
+		const FString& PrimPath = ""
+	);
+
+public: // Stringify functions
+	UFUNCTION(BlueprintCallable, Category = "USD|Stringify utils")
+	static FString StringifyAsBool(bool Value);
+
+	UFUNCTION(BlueprintCallable, Category = "USD|Stringify utils", Meta=(ScriptName="stringify_as_uchar"))
+	static FString StringifyAsUChar(uint8 Value);
+
+	UFUNCTION(BlueprintCallable, Category = "USD|Stringify utils")
+	static FString StringifyAsInt(int32 Value);
+
+	UFUNCTION(BlueprintCallable, Category = "USD|Stringify utils", Meta=(ScriptName="stringify_as_uint"))
+	static FString StringifyAsUInt(int32 Value);
+
+	UFUNCTION(BlueprintCallable, Category = "USD|Stringify utils")
+	static FString StringifyAsInt64(int64 Value);
+
+	UFUNCTION(BlueprintCallable, Category = "USD|Stringify utils", Meta=(ScriptName="stringify_as_uint64"))
+	static FString StringifyAsUInt64(int64 Value);
+
+	UFUNCTION(BlueprintCallable, Category = "USD|Stringify utils")
+	static FString StringifyAsHalf(float Value);
+
+	UFUNCTION(BlueprintCallable, Category = "USD|Stringify utils")
+	static FString StringifyAsFloat(float Value);
+
+	UFUNCTION(BlueprintCallable, Category = "USD|Stringify utils")
+	static FString StringifyAsDouble(double Value);
+
+	UFUNCTION(BlueprintCallable, Category = "USD|Stringify utils", Meta=(ScriptName="stringify_as_timecode"))
+	static FString StringifyAsTimeCode(double Value);
+
+	UFUNCTION(BlueprintCallable, Category = "USD|Stringify utils")
+	static FString StringifyAsString(const FString& Value);
+
+	UFUNCTION(BlueprintCallable, Category = "USD|Stringify utils")
+	static FString StringifyAsToken(const FString& Value);
+
+	UFUNCTION(BlueprintCallable, Category = "USD|Stringify utils")
+	static FString StringifyAsAssetPath(const FString& Value);
+
+	UFUNCTION(BlueprintCallable, Category = "USD|Stringify utils")
+	static FString StringifyAsMatrix2d(const FMatrix2D& Value);
+
+	UFUNCTION(BlueprintCallable, Category = "USD|Stringify utils")
+	static FString StringifyAsMatrix3d(const FMatrix3D& Value);
+
+	UFUNCTION(BlueprintCallable, Category = "USD|Stringify utils")
+	static FString StringifyAsMatrix4d(const FMatrix& Value);
+
+	UFUNCTION(BlueprintCallable, Category = "USD|Stringify utils")
+	static FString StringifyAsQuatd(const FQuat& Value);
+
+	UFUNCTION(BlueprintCallable, Category = "USD|Stringify utils")
+	static FString StringifyAsQuatf(const FQuat& Value);
+
+	UFUNCTION(BlueprintCallable, Category = "USD|Stringify utils")
+	static FString StringifyAsQuath(const FQuat& Value);
+
+	UFUNCTION(BlueprintCallable, Category = "USD|Stringify utils")
+	static FString StringifyAsDouble2(const FVector2D& Value);
+
+	UFUNCTION(BlueprintCallable, Category = "USD|Stringify utils")
+	static FString StringifyAsFloat2(const FVector2D& Value);
+
+	UFUNCTION(BlueprintCallable, Category = "USD|Stringify utils")
+	static FString StringifyAsHalf2(const FVector2D& Value);
+
+	UFUNCTION(BlueprintCallable, Category = "USD|Stringify utils")
+	static FString StringifyAsInt2(const FIntPoint& Value);
+
+	UFUNCTION(BlueprintCallable, Category = "USD|Stringify utils")
+	static FString StringifyAsDouble3(const FVector& Value);
+
+	UFUNCTION(BlueprintCallable, Category = "USD|Stringify utils")
+	static FString StringifyAsFloat3(const FVector& Value);
+
+	UFUNCTION(BlueprintCallable, Category = "USD|Stringify utils")
+	static FString StringifyAsHalf3(const FVector& Value);
+
+	UFUNCTION(BlueprintCallable, Category = "USD|Stringify utils")
+	static FString StringifyAsInt3(const FIntVector& Value);
+
+	UFUNCTION(BlueprintCallable, Category = "USD|Stringify utils")
+	static FString StringifyAsDouble4(const FVector4& Value);
+
+	UFUNCTION(BlueprintCallable, Category = "USD|Stringify utils")
+	static FString StringifyAsFloat4(const FVector4& Value);
+
+	UFUNCTION(BlueprintCallable, Category = "USD|Stringify utils")
+	static FString StringifyAsHalf4(const FVector4& Value);
+
+	UFUNCTION(BlueprintCallable, Category = "USD|Stringify utils")
+	static FString StringifyAsInt4(const FIntVector4& Value);
+
+public: // Stringify array functions
+	UFUNCTION(BlueprintCallable, Category = "USD|Stringify utils")
+	static FString StringifyAsBoolArray(const TArray<bool>& Value);
+
+	UFUNCTION(BlueprintCallable, Category = "USD|Stringify utils", Meta=(ScriptName="stringify_as_uchar_array"))
+	static FString StringifyAsUCharArray(const TArray<uint8>& Value);
+
+	UFUNCTION(BlueprintCallable, Category = "USD|Stringify utils")
+	static FString StringifyAsIntArray(const TArray<int32>& Value);
+
+	UFUNCTION(BlueprintCallable, Category = "USD|Stringify utils", Meta=(ScriptName="stringify_as_uint_array"))
+	static FString StringifyAsUIntArray(const TArray<int32>& Value);
+
+	UFUNCTION(BlueprintCallable, Category = "USD|Stringify utils")
+	static FString StringifyAsInt64Array(const TArray<int64>& Value);
+
+	UFUNCTION(BlueprintCallable, Category = "USD|Stringify utils", Meta=(ScriptName="stringify_as_uint64_array"))
+	static FString StringifyAsUInt64Array(const TArray<int64>& Value);
+
+	UFUNCTION(BlueprintCallable, Category = "USD|Stringify utils")
+	static FString StringifyAsHalfArray(const TArray<float>& Value);
+
+	UFUNCTION(BlueprintCallable, Category = "USD|Stringify utils")
+	static FString StringifyAsFloatArray(const TArray<float>& Value);
+
+	UFUNCTION(BlueprintCallable, Category = "USD|Stringify utils")
+	static FString StringifyAsDoubleArray(const TArray<double>& Value);
+
+	UFUNCTION(BlueprintCallable, Category = "USD|Stringify utils", Meta=(ScriptName="stringify_as_timecode_array"))
+	static FString StringifyAsTimeCodeArray(const TArray<double>& Value);
+
+	UFUNCTION(BlueprintCallable, Category = "USD|Stringify utils")
+	static FString StringifyAsStringArray(const TArray<FString>& Value);
+
+	UFUNCTION(BlueprintCallable, Category = "USD|Stringify utils")
+	static FString StringifyAsTokenArray(const TArray<FString>& Value);
+
+	UFUNCTION(BlueprintCallable, Category = "USD|Stringify utils")
+	static FString StringifyAsAssetPathArray(const TArray<FString>& Value);
+
+	UFUNCTION(BlueprintCallable, Category = "USD|Stringify utils")
+	static FString StringifyAsListOpTokens(const TArray<FString>& Value);
+
+	UFUNCTION(BlueprintCallable, Category = "USD|Stringify utils")
+	static FString StringifyAsMatrix2dArray(const TArray<FMatrix2D>& Value);
+
+	UFUNCTION(BlueprintCallable, Category = "USD|Stringify utils")
+	static FString StringifyAsMatrix3dArray(const TArray<FMatrix3D>& Value);
+
+	UFUNCTION(BlueprintCallable, Category = "USD|Stringify utils")
+	static FString StringifyAsMatrix4dArray(const TArray<FMatrix>& Value);
+
+	UFUNCTION(BlueprintCallable, Category = "USD|Stringify utils")
+	static FString StringifyAsQuatdArray(const TArray<FQuat>& Value);
+
+	UFUNCTION(BlueprintCallable, Category = "USD|Stringify utils")
+	static FString StringifyAsQuatfArray(const TArray<FQuat>& Value);
+
+	UFUNCTION(BlueprintCallable, Category = "USD|Stringify utils")
+	static FString StringifyAsQuathArray(const TArray<FQuat>& Value);
+
+	UFUNCTION(BlueprintCallable, Category = "USD|Stringify utils")
+	static FString StringifyAsDouble2Array(const TArray<FVector2D>& Value);
+
+	UFUNCTION(BlueprintCallable, Category = "USD|Stringify utils")
+	static FString StringifyAsFloat2Array(const TArray<FVector2D>& Value);
+
+	UFUNCTION(BlueprintCallable, Category = "USD|Stringify utils")
+	static FString StringifyAsHalf2Array(const TArray<FVector2D>& Value);
+
+	UFUNCTION(BlueprintCallable, Category = "USD|Stringify utils")
+	static FString StringifyAsInt2Array(const TArray<FIntPoint>& Value);
+
+	UFUNCTION(BlueprintCallable, Category = "USD|Stringify utils")
+	static FString StringifyAsDouble3Array(const TArray<FVector>& Value);
+
+	UFUNCTION(BlueprintCallable, Category = "USD|Stringify utils")
+	static FString StringifyAsFloat3Array(const TArray<FVector>& Value);
+
+	UFUNCTION(BlueprintCallable, Category = "USD|Stringify utils")
+	static FString StringifyAsHalf3Array(const TArray<FVector>& Value);
+
+	UFUNCTION(BlueprintCallable, Category = "USD|Stringify utils")
+	static FString StringifyAsInt3Array(const TArray<FIntVector>& Value);
+
+	UFUNCTION(BlueprintCallable, Category = "USD|Stringify utils")
+	static FString StringifyAsDouble4Array(const TArray<FVector4>& Value);
+
+	UFUNCTION(BlueprintCallable, Category = "USD|Stringify utils")
+	static FString StringifyAsFloat4Array(const TArray<FVector4>& Value);
+
+	UFUNCTION(BlueprintCallable, Category = "USD|Stringify utils")
+	static FString StringifyAsHalf4Array(const TArray<FVector4>& Value);
+
+	UFUNCTION(BlueprintCallable, Category = "USD|Stringify utils")
+	static FString StringifyAsInt4Array(const TArray<FIntVector4>& Value);
+
+public: // Unstringify functions
+	UFUNCTION(BlueprintCallable, Category = "USD|Stringify utils")
+	static bool UnstringifyAsBool(const FString& String);
+
+	UFUNCTION(BlueprintCallable, Category = "USD|Stringify utils", Meta=(ScriptName="unstringify_as_uchar"))
+	static uint8 UnstringifyAsUChar(const FString& String);
+
+	UFUNCTION(BlueprintCallable, Category = "USD|Stringify utils")
+	static int32 UnstringifyAsInt(const FString& String);
+
+	UFUNCTION(BlueprintCallable, Category = "USD|Stringify utils", Meta=(ScriptName="unstringify_as_uint"))
+	static int32 UnstringifyAsUInt(const FString& String);
+
+	UFUNCTION(BlueprintCallable, Category = "USD|Stringify utils")
+	static int64 UnstringifyAsInt64(const FString& String);
+
+	UFUNCTION(BlueprintCallable, Category = "USD|Stringify utils", Meta=(ScriptName="unstringify_as_uint64"))
+	static int64 UnstringifyAsUInt64(const FString& String);
+
+	UFUNCTION(BlueprintCallable, Category = "USD|Stringify utils")
+	static float UnstringifyAsHalf(const FString& String);
+
+	UFUNCTION(BlueprintCallable, Category = "USD|Stringify utils")
+	static float UnstringifyAsFloat(const FString& String);
+
+	UFUNCTION(BlueprintCallable, Category = "USD|Stringify utils")
+	static double UnstringifyAsDouble(const FString& String);
+
+	UFUNCTION(BlueprintCallable, Category = "USD|Stringify utils", Meta=(ScriptName="unstringify_as_timecode"))
+	static double UnstringifyAsTimeCode(const FString& String);
+
+	UFUNCTION(BlueprintCallable, Category = "USD|Stringify utils")
+	static FString UnstringifyAsString(const FString& String);
+
+	UFUNCTION(BlueprintCallable, Category = "USD|Stringify utils")
+	static FString UnstringifyAsToken(const FString& String);
+
+	UFUNCTION(BlueprintCallable, Category = "USD|Stringify utils")
+	static FString UnstringifyAsAssetPath(const FString& String);
+
+	UFUNCTION(BlueprintCallable, Category = "USD|Stringify utils")
+	static FMatrix2D UnstringifyAsMatrix2d(const FString& String);
+
+	UFUNCTION(BlueprintCallable, Category = "USD|Stringify utils")
+	static FMatrix3D UnstringifyAsMatrix3d(const FString& String);
+
+	UFUNCTION(BlueprintCallable, Category = "USD|Stringify utils")
+	static FMatrix UnstringifyAsMatrix4d(const FString& String);
+
+	UFUNCTION(BlueprintCallable, Category = "USD|Stringify utils")
+	static FQuat UnstringifyAsQuatd(const FString& String);
+
+	UFUNCTION(BlueprintCallable, Category = "USD|Stringify utils")
+	static FQuat UnstringifyAsQuatf(const FString& String);
+
+	UFUNCTION(BlueprintCallable, Category = "USD|Stringify utils")
+	static FQuat UnstringifyAsQuath(const FString& String);
+
+	UFUNCTION(BlueprintCallable, Category = "USD|Stringify utils")
+	static FVector2D UnstringifyAsDouble2(const FString& String);
+
+	UFUNCTION(BlueprintCallable, Category = "USD|Stringify utils")
+	static FVector2D UnstringifyAsFloat2(const FString& String);
+
+	UFUNCTION(BlueprintCallable, Category = "USD|Stringify utils")
+	static FVector2D UnstringifyAsHalf2(const FString& String);
+
+	UFUNCTION(BlueprintCallable, Category = "USD|Stringify utils")
+	static FIntPoint UnstringifyAsInt2(const FString& String);
+
+	UFUNCTION(BlueprintCallable, Category = "USD|Stringify utils")
+	static FVector UnstringifyAsDouble3(const FString& String);
+
+	UFUNCTION(BlueprintCallable, Category = "USD|Stringify utils")
+	static FVector UnstringifyAsFloat3(const FString& String);
+
+	UFUNCTION(BlueprintCallable, Category = "USD|Stringify utils")
+	static FVector UnstringifyAsHalf3(const FString& String);
+
+	UFUNCTION(BlueprintCallable, Category = "USD|Stringify utils")
+	static FIntVector UnstringifyAsInt3(const FString& String);
+
+	UFUNCTION(BlueprintCallable, Category = "USD|Stringify utils")
+	static FVector4 UnstringifyAsDouble4(const FString& String);
+
+	UFUNCTION(BlueprintCallable, Category = "USD|Stringify utils")
+	static FVector4 UnstringifyAsFloat4(const FString& String);
+
+	UFUNCTION(BlueprintCallable, Category = "USD|Stringify utils")
+	static FVector4 UnstringifyAsHalf4(const FString& String);
+
+	UFUNCTION(BlueprintCallable, Category = "USD|Stringify utils")
+	static FIntVector4 UnstringifyAsInt4(const FString& String);
+
+public: // Unstringify array functions
+	UFUNCTION(BlueprintCallable, Category = "USD|Stringify utils")
+	static TArray<bool> UnstringifyAsBoolArray(const FString& String);
+
+	UFUNCTION(BlueprintCallable, Category = "USD|Stringify utils", Meta=(ScriptName="unstringify_as_uchar_array"))
+	static TArray<uint8> UnstringifyAsUCharArray(const FString& String);
+
+	UFUNCTION(BlueprintCallable, Category = "USD|Stringify utils")
+	static TArray<int32> UnstringifyAsIntArray(const FString& String);
+
+	UFUNCTION(BlueprintCallable, Category = "USD|Stringify utils", Meta=(ScriptName="unstringify_as_uint_array"))
+	static TArray<int32> UnstringifyAsUIntArray(const FString& String);
+
+	UFUNCTION(BlueprintCallable, Category = "USD|Stringify utils")
+	static TArray<int64> UnstringifyAsInt64Array(const FString& String);
+
+	UFUNCTION(BlueprintCallable, Category = "USD|Stringify utils", Meta=(ScriptName="unstringify_as_uint64_array"))
+	static TArray<int64> UnstringifyAsUInt64Array(const FString& String);
+
+	UFUNCTION(BlueprintCallable, Category = "USD|Stringify utils")
+	static TArray<float> UnstringifyAsHalfArray(const FString& String);
+
+	UFUNCTION(BlueprintCallable, Category = "USD|Stringify utils")
+	static TArray<float> UnstringifyAsFloatArray(const FString& String);
+
+	UFUNCTION(BlueprintCallable, Category = "USD|Stringify utils")
+	static TArray<double> UnstringifyAsDoubleArray(const FString& String);
+
+	UFUNCTION(BlueprintCallable, Category = "USD|Stringify utils", Meta=(ScriptName="unstringify_as_timecode_array"))
+	static TArray<double> UnstringifyAsTimeCodeArray(const FString& String);
+
+	UFUNCTION(BlueprintCallable, Category = "USD|Stringify utils")
+	static TArray<FString> UnstringifyAsStringArray(const FString& String);
+
+	UFUNCTION(BlueprintCallable, Category = "USD|Stringify utils")
+	static TArray<FString> UnstringifyAsTokenArray(const FString& String);
+
+	UFUNCTION(BlueprintCallable, Category = "USD|Stringify utils")
+	static TArray<FString> UnstringifyAsAssetPathArray(const FString& String);
+
+	UFUNCTION(BlueprintCallable, Category = "USD|Stringify utils")
+	static TArray<FString> UnstringifyAsListOpTokens(const FString& String);
+
+	UFUNCTION(BlueprintCallable, Category = "USD|Stringify utils")
+	static TArray<FMatrix2D> UnstringifyAsMatrix2dArray(const FString& String);
+
+	UFUNCTION(BlueprintCallable, Category = "USD|Stringify utils")
+	static TArray<FMatrix3D> UnstringifyAsMatrix3dArray(const FString& String);
+
+	UFUNCTION(BlueprintCallable, Category = "USD|Stringify utils")
+	static TArray<FMatrix> UnstringifyAsMatrix4dArray(const FString& String);
+
+	UFUNCTION(BlueprintCallable, Category = "USD|Stringify utils")
+	static TArray<FQuat> UnstringifyAsQuatdArray(const FString& String);
+
+	UFUNCTION(BlueprintCallable, Category = "USD|Stringify utils")
+	static TArray<FQuat> UnstringifyAsQuatfArray(const FString& String);
+
+	UFUNCTION(BlueprintCallable, Category = "USD|Stringify utils")
+	static TArray<FQuat> UnstringifyAsQuathArray(const FString& String);
+
+	UFUNCTION(BlueprintCallable, Category = "USD|Stringify utils")
+	static TArray<FVector2D> UnstringifyAsDouble2Array(const FString& String);
+
+	UFUNCTION(BlueprintCallable, Category = "USD|Stringify utils")
+	static TArray<FVector2D> UnstringifyAsFloat2Array(const FString& String);
+
+	UFUNCTION(BlueprintCallable, Category = "USD|Stringify utils")
+	static TArray<FVector2D> UnstringifyAsHalf2Array(const FString& String);
+
+	UFUNCTION(BlueprintCallable, Category = "USD|Stringify utils")
+	static TArray<FIntPoint> UnstringifyAsInt2Array(const FString& String);
+
+	UFUNCTION(BlueprintCallable, Category = "USD|Stringify utils")
+	static TArray<FVector> UnstringifyAsDouble3Array(const FString& String);
+
+	UFUNCTION(BlueprintCallable, Category = "USD|Stringify utils")
+	static TArray<FVector> UnstringifyAsFloat3Array(const FString& String);
+
+	UFUNCTION(BlueprintCallable, Category = "USD|Stringify utils")
+	static TArray<FVector> UnstringifyAsHalf3Array(const FString& String);
+
+	UFUNCTION(BlueprintCallable, Category = "USD|Stringify utils")
+	static TArray<FIntVector> UnstringifyAsInt3Array(const FString& String);
+
+	UFUNCTION(BlueprintCallable, Category = "USD|Stringify utils")
+	static TArray<FVector4> UnstringifyAsDouble4Array(const FString& String);
+
+	UFUNCTION(BlueprintCallable, Category = "USD|Stringify utils")
+	static TArray<FVector4> UnstringifyAsFloat4Array(const FString& String);
+
+	UFUNCTION(BlueprintCallable, Category = "USD|Stringify utils")
+	static TArray<FVector4> UnstringifyAsHalf4Array(const FString& String);
+
+	UFUNCTION(BlueprintCallable, Category = "USD|Stringify utils")
+	static TArray<FIntVector4> UnstringifyAsInt4Array(const FString& String);
 };
diff --git a/Engine/Plugins/Importers/USDImporter/Source/USDSchemas/Private/USDGeomMeshTranslator.cpp b/Engine/Plugins/Importers/USDImporter/Source/USDSchemas/Private/USDGeomMeshTranslator.cpp
--- a/Engine/Plugins/Importers/USDImporter/Source/USDSchemas/Private/USDGeomMeshTranslator.cpp
+++ b/Engine/Plugins/Importers/USDImporter/Source/USDSchemas/Private/USDGeomMeshTranslator.cpp
@@ -296,7 +296,9 @@ namespace UsdGeomMeshTranslatorImpl
 		const UE::FUsdPrim& MeshPrim,
 		TArray<FMeshDescription>& OutLODIndexToMeshDescription,
 		TArray<UsdUtils::FUsdPrimMaterialAssignmentInfo>& OutLODIndexToMaterialInfo,
-		const UsdToUnreal::FUsdMeshConversionOptions& Options
+		FUsdCombinedPrimMetadata& OutLODMetadata,
+		const UsdToUnreal::FUsdMeshConversionOptions& Options,
+		const FUsdMetadataImportOptions& MetadataOptions
 	)
 	{
 		if(!MeshPrim)
@@ -345,19 +347,18 @@ namespace UsdGeomMeshTranslatorImpl
 		{
 			return false;
 		}
-		TMap<FString, int32> PrimvarToUVIndex = UsdUtils::CombinePrimvarsIntoUVSets(
-			AllPrimvars,
-			PreferredPrimvars
-		);
+		TMap<FString, int32> PrimvarToUVIndex = UsdUtils::CombinePrimvarsIntoUVSets(AllPrimvars, PreferredPrimvars);
 
 		TFunction<bool(const pxr::UsdGeomMesh&, int32)> ConvertLOD =
-			[&OptionsCopy, &Options, &PrimvarToUVIndex, &LODIndexToMeshDescriptionMap, &LODIndexToMaterialInfoMap]
-			(const pxr::UsdGeomMesh& LODMesh, int32 LODIndex)
+			[&OptionsCopy, &Options, &PrimvarToUVIndex, &LODIndexToMeshDescriptionMap, &LODIndexToMaterialInfoMap, &OutLODMetadata, &MetadataOptions](
+				const pxr::UsdGeomMesh& LODMesh,
+				int32 LODIndex
+			)
 		{
 			FMeshDescription TempMeshDescription;
 
 			UsdUtils::FUsdPrimMaterialAssignmentInfo TempMaterialInfo;
-				TempMaterialInfo.PrimvarToUVIndex = PrimvarToUVIndex;
+			TempMaterialInfo.PrimvarToUVIndex = PrimvarToUVIndex;
 
 			FStaticMeshAttributes StaticMeshAttributes(TempMeshDescription);
 			StaticMeshAttributes.Register();
@@ -401,6 +402,17 @@ namespace UsdGeomMeshTranslatorImpl
 			{
 				LODIndexToMeshDescriptionMap.Add(LODIndex, MoveTemp(TempMeshDescription));
 				LODIndexToMaterialInfoMap.Add(LODIndex, MoveTemp(TempMaterialInfo));
+
+				if (MetadataOptions.bCollectMetadata)
+				{
+					UsdToUnreal::ConvertMetadata(
+						LODMesh.GetPrim(),
+						OutLODMetadata,
+						MetadataOptions.BlockedPrefixFilters,
+						MetadataOptions.bInvertFilters,
+						MetadataOptions.bCollectFromEntireSubtrees
+					);
+				}
 			}
 
 			return true;
@@ -428,8 +440,10 @@ namespace UsdGeomMeshTranslatorImpl
 		UE::FUsdPrim MeshPrim,
 		TArray<FMeshDescription>& OutLODIndexToMeshDescription,
 		TArray<UsdUtils::FUsdPrimMaterialAssignmentInfo>& OutLODIndexToMaterialInfo,
+		FUsdCombinedPrimMetadata& OutLODMetadata,
 		const UsdToUnreal::FUsdMeshConversionOptions& Options,
-		bool bInterpretLODs = false
+		const FUsdMetadataImportOptions& MetadataOptions,
+		bool bInterpretLODs
 	)
 	{
 		if (!MeshPrim)
@@ -443,7 +457,14 @@ namespace UsdGeomMeshTranslatorImpl
 		bool bInterpretedLODs = false;
 		if (bInterpretLODs)
 		{
-			bInterpretedLODs = TryLoadingMultipleLODs(MeshPrim, OutLODIndexToMeshDescription, OutLODIndexToMaterialInfo, Options);
+			bInterpretedLODs = TryLoadingMultipleLODs(
+				MeshPrim,
+				OutLODIndexToMeshDescription,
+				OutLODIndexToMaterialInfo,
+				OutLODMetadata,
+				Options,
+				MetadataOptions
+			);
 
 			// Have to be very careful here as flipping through LODs invalidates prim references, so we need to
 			// re-acquire them
@@ -491,6 +512,17 @@ namespace UsdGeomMeshTranslatorImpl
 			{
 				OutLODIndexToMeshDescription = {MoveTemp(TempMeshDescription)};
 				OutLODIndexToMaterialInfo = {MoveTemp(TempMaterialInfo)};
+
+				if (MetadataOptions.bCollectMetadata)
+				{
+					UsdToUnreal::ConvertMetadata(
+						MeshPrim,
+						OutLODMetadata,
+						MetadataOptions.BlockedPrefixFilters,
+						MetadataOptions.bInvertFilters,
+						MetadataOptions.bCollectFromEntireSubtrees
+					);
+				}
 			}
 		}
 	}
@@ -1392,12 +1424,14 @@ void FBuildStaticMeshTaskChain::SetupTasks()
 
 			const FString PrimPathString = PrimPath.GetString();
 
+			const bool bParsedLODs = Context->bAllowInterpretingLODs && UsdUtils::IsGeomMeshALOD(GetPrim());
+
 			// It's useful to have the LOD Mesh prims be named "LOD0", "LOD1", etc. within the LOD variants so that we
 			// can easily tell which Mesh is actually meant to be the LOD mesh (in case there are more Meshes in each
 			// variant or other Meshes outside of the variant), but it's not ideal to have all the generated assets end
 			// up imported as "SM_LOD0_22", "SM_LOD0_23", etc. So here we fetch the parent prim name in case we're a LOD
 			FString MeshName;
-			if (Context->bAllowInterpretingLODs && UsdUtils::IsGeomMeshALOD(GetPrim()))
+			if (bParsedLODs)
 			{
 				MeshName = PrimPath.GetParentPath().GetString();
 			}
@@ -1424,16 +1458,34 @@ void FBuildStaticMeshTaskChain::SetupTasks()
 				StaticMesh->NaniteSettings.bEnabled = bShouldEnableNanite;
 #endif // WITH_EDITOR
 
-				UUsdMeshAssetUserData* UserData = StaticMesh->GetAssetUserData<UUsdMeshAssetUserData>();
-				if (!UserData)
+				if (UUsdMeshAssetUserData* UserData = UsdUtils::GetOrCreateAssetUserData<UUsdMeshAssetUserData>(StaticMesh))
 				{
-					UserData = NewObject<UUsdMeshAssetUserData>(StaticMesh, TEXT("UUSDAssetUserData"));
 					UserData->PrimvarToUVIndex = LODIndexToMaterialInfo[0].PrimvarToUVIndex;	// We use the same primvar mapping for all LODs
-					StaticMesh->AddAssetUserData(UserData);
-				}
-				UserData->PrimPaths.AddUnique(PrimPath.GetString());
+					UserData->PrimPaths.AddUnique(PrimPathString);
 
-				MeshTranslationImpl::RecordSourcePrimsForMaterialSlots(LODIndexToMaterialInfo, UserData);
+					if (Context->MetadataOptions.bCollectMetadata)
+					{
+						// If we already have collected metadata just stash it into our UserData directly (the GeomMeshTranslator task
+						// chain itself does this as it parses lods, but the regular FBuildStaticMeshTaskChain won't do this on
+						// its own, and is reused for other task chains)
+						if (bCollectedMetadata)
+						{
+							UserData->StageIdentifierToMetadata.Add(GetPrim().GetStage().GetRootLayer().GetIdentifier(), LODMetadata);
+						}
+						else
+						{
+							UsdToUnreal::ConvertMetadata(
+								bParsedLODs ? GetPrim().GetParent() : GetPrim(),
+								UserData,
+								Context->MetadataOptions.BlockedPrefixFilters,
+								Context->MetadataOptions.bInvertFilters,
+								Context->MetadataOptions.bCollectFromEntireSubtrees
+							);
+						}
+					}
+
+					MeshTranslationImpl::RecordSourcePrimsForMaterialSlots(LODIndexToMaterialInfo, UserData);
+				}
 
 				// Only the original creator of the prim at creation time gets to set the material assignments
 				// directly on the mesh, all others prims ensure their materials via material overrides on the
@@ -1560,15 +1612,16 @@ void FGeomMeshCreateAssetsTaskChain::SetupTasks()
 
 	// To parse all LODs we need to actively switch variant sets to other variants (triggering prim loading/unloading and notices),
 	// which could cause race conditions if other async translation tasks are trying to access those prims
+	bool bParseLODs = Context->bAllowInterpretingLODs && UsdUtils::IsGeomMeshALOD(GetPrim());
 	ESchemaTranslationLaunchPolicy LaunchPolicy = ESchemaTranslationLaunchPolicy::Async;
-	if (Context->bAllowInterpretingLODs && UsdUtils::IsGeomMeshALOD(GetPrim()))
+	if (bParseLODs)
 	{
 		LaunchPolicy = ESchemaTranslationLaunchPolicy::ExclusiveSync;
 	}
 
 	// Create mesh descriptions (Async or ExclusiveSync)
 	Do(LaunchPolicy,
-		[this, LaunchPolicy]() -> bool
+		[this, bParseLODs]() -> bool
 		{
 			pxr::TfToken RenderContextToken = pxr::UsdShadeTokens->universalRenderContext;
 			if (!Context->RenderContext.IsNone())
@@ -1595,10 +1648,28 @@ void FGeomMeshCreateAssetsTaskChain::SetupTasks()
 				GetPrim(),
 				LODIndexToMeshDescription,
 				LODIndexToMaterialInfo,
+				LODMetadata,
 				Options,
-				Context->bAllowInterpretingLODs && LaunchPolicy == ESchemaTranslationLaunchPolicy::ExclusiveSync
+				Context->MetadataOptions,
+				Context->bAllowInterpretingLODs && bParseLODs
 			);
 
+			// If we're parsing LODs, LoadMeshDescriptions will have already collected metadata from the Mesh
+			// prims directly, but we still have to collect stuff from the actual LOD root prim, which is kind
+			// of absorbed into the asset
+			if (bParseLODs && Context->MetadataOptions.bCollectMetadata)
+			{
+				const bool bCollectMetadataFromSubtree = false;
+				UsdToUnreal::ConvertMetadata(
+					GetPrim().GetParent(),
+					LODMetadata,
+					Context->MetadataOptions.BlockedPrefixFilters,
+					Context->MetadataOptions.bInvertFilters,
+					bCollectMetadataFromSubtree
+				);
+			}
+			bCollectedMetadata = true;
+
 			// If we have at least one valid LOD, we should keep going
 			for (const FMeshDescription& MeshDescription : LODIndexToMeshDescription)
 			{
diff --git a/Engine/Plugins/Importers/USDImporter/Source/USDSchemas/Private/USDGeomXformableTranslator.cpp b/Engine/Plugins/Importers/USDImporter/Source/USDSchemas/Private/USDGeomXformableTranslator.cpp
--- a/Engine/Plugins/Importers/USDImporter/Source/USDSchemas/Private/USDGeomXformableTranslator.cpp
+++ b/Engine/Plugins/Importers/USDImporter/Source/USDSchemas/Private/USDGeomXformableTranslator.cpp
@@ -568,6 +568,25 @@ USceneComponent* FUsdGeomXformableTranslator::CreateComponentsEx( TOptional< TSu
 		}
 	}
 
+	if (Context->MetadataOptions.bCollectMetadata && Context->MetadataOptions.bCollectOnComponents)
+	{
+		UUsdAssetUserData* UserData = UsdUtils::GetOrCreateAssetUserData(SceneComponent);
+
+		// It makes sense for asset metadata to "include all prims in the subtree", as when we generate an
+		// asset we don't generate additional separate assets for child prims. This is not the same behavior
+		// for components though, so it doesn't feel like "collecting from the entire subtree" should be
+		// allowed for them. In other words, if we allowed this the root scene component for the stage will
+		// contain metadata from the entire stage every time...
+		const bool bCollectMetadataFromSubtree = false;
+		UsdToUnreal::ConvertMetadata(
+			Prim,
+			UserData,
+			Context->MetadataOptions.BlockedPrefixFilters,
+			Context->MetadataOptions.bInvertFilters,
+			bCollectMetadataFromSubtree
+		);
+	}
+
 	if ( SceneComponent )
 	{
 		if ( !GEnableCollision )
diff --git a/Engine/Plugins/Importers/USDImporter/Source/USDSchemas/Private/USDGeometryCacheTranslator.cpp b/Engine/Plugins/Importers/USDImporter/Source/USDSchemas/Private/USDGeometryCacheTranslator.cpp
--- a/Engine/Plugins/Importers/USDImporter/Source/USDSchemas/Private/USDGeometryCacheTranslator.cpp
+++ b/Engine/Plugins/Importers/USDImporter/Source/USDSchemas/Private/USDGeometryCacheTranslator.cpp
@@ -779,20 +779,25 @@ void FGeometryCacheCreateAssetsTaskChain::SetupTasks()
 
 			if (GeometryCache)
 			{
-				UUsdMeshAssetUserData* UserData = GeometryCache->GetAssetUserData<UUsdMeshAssetUserData>();
-				if (!UserData)
+				if (UUsdGeometryCacheAssetUserData* UserData = UsdUtils::GetOrCreateAssetUserData<UUsdGeometryCacheAssetUserData>(GeometryCache.Get()))
 				{
-					UserData = NewObject<UUsdMeshAssetUserData>(GeometryCache.Get(), TEXT("UUSDAssetUserData"));
 					UserData->PrimvarToUVIndex = LODIndexToMaterialInfo[0].PrimvarToUVIndex;	// We use the same primvar mapping for all LODs
-					GeometryCache->AddAssetUserData(UserData);
+					UserData->LayerStartOffsetSeconds = StartTimeOffset;
+					UserData->PrimPaths.AddUnique(PrimPathString);
 
-					UUsdAnimSequenceAssetUserData* AnimUserData = NewObject<UUsdAnimSequenceAssetUserData>(GeometryCache.Get(), TEXT("UsdAnimUserData"));
-					AnimUserData->LayerStartOffsetSeconds = StartTimeOffset;
-					GeometryCache->AddAssetUserData(AnimUserData);
-				}
-				UserData->PrimPaths.AddUnique(PrimPath.GetString());
+					if (Context->MetadataOptions.bCollectMetadata)
+					{
+						UsdToUnreal::ConvertMetadata(
+							GetPrim(),
+							UserData,
+							Context->MetadataOptions.BlockedPrefixFilters,
+							Context->MetadataOptions.bInvertFilters,
+							Context->MetadataOptions.bCollectFromEntireSubtrees
+						);
+					}
 
-				MeshTranslationImpl::RecordSourcePrimsForMaterialSlots(LODIndexToMaterialInfo, UserData);
+					MeshTranslationImpl::RecordSourcePrimsForMaterialSlots(LODIndexToMaterialInfo, UserData);
+				}
 
 				if (bIsNew)
 				{
@@ -1021,7 +1026,7 @@ void FUsdGeometryCacheTranslator::UpdateComponents(USceneComponent* SceneCompone
 			float LayerStartOffsetSeconds = 0.0f;
 			if (GeometryCache)
 			{
-				if (UUsdAnimSequenceAssetUserData* UserData = GeometryCache->GetAssetUserData<UUsdAnimSequenceAssetUserData>())
+				if (UUsdGeometryCacheAssetUserData* UserData = GeometryCache->GetAssetUserData<UUsdGeometryCacheAssetUserData>())
 				{
 					LayerStartOffsetSeconds = UserData->LayerStartOffsetSeconds;
 				}
diff --git a/Engine/Plugins/Importers/USDImporter/Source/USDSchemas/Private/USDGroomTranslator.cpp b/Engine/Plugins/Importers/USDImporter/Source/USDSchemas/Private/USDGroomTranslator.cpp
--- a/Engine/Plugins/Importers/USDImporter/Source/USDSchemas/Private/USDGroomTranslator.cpp
+++ b/Engine/Plugins/Importers/USDImporter/Source/USDSchemas/Private/USDGroomTranslator.cpp
@@ -19,12 +19,13 @@
 
 #include "UnrealUSDWrapper.h"
 #include "USDAssetUserData.h"
-#include "USDDrawModeComponent.h"
 #include "USDClassesModule.h"
+#include "USDDrawModeComponent.h"
 #include "USDGroomConversion.h"
 #include "USDGroomTranslatorUtils.h"
 #include "USDIntegrationUtils.h"
 #include "USDLog.h"
+#include "USDPrimConversion.h"
 #include "USDTypesConversion.h"
 
 #include "USDIncludesStart.h"
@@ -197,10 +198,21 @@ protected:
 					{
 						Context->AssetCache->CacheAsset(SHAHash.ToString(), GroomAsset);
 
-						UUsdAssetUserData* UserData = NewObject<UUsdAssetUserData>(GroomAsset, TEXT("UUSDAssetUserData"));
-						UserData->PrimPaths = {PrimPathString};
+						if (UUsdAssetUserData* UserData = UsdUtils::GetOrCreateAssetUserData(GroomAsset))
+						{
+							UserData->PrimPaths.AddUnique(PrimPathString);
 
-						GroomAsset->AddAssetUserData(UserData);
+							if (Context->MetadataOptions.bCollectMetadata)
+							{
+								UsdToUnreal::ConvertMetadata(
+									GetPrim(),
+									UserData,
+									Context->MetadataOptions.BlockedPrefixFilters,
+									Context->MetadataOptions.bInvertFilters,
+									Context->MetadataOptions.bCollectFromEntireSubtrees
+								);
+							}
+						}
 					}
 				}
 
@@ -340,10 +352,21 @@ protected:
 				UGroomCache* GroomCache = FGroomCacheImporter::ProcessToGroomCache(*GroomCacheProcessor, AnimInfo, HairImportContext, UniqueName.ToString());
 				if (GroomCache)
 				{
-					UUsdAssetUserData* UserData = NewObject<UUsdAssetUserData>(GroomCache, TEXT("UUSDAssetUserData"));
-					UserData->PrimPaths = {PrimPath.GetString()};
+					if (UUsdAssetUserData* UserData = UsdUtils::GetOrCreateAssetUserData(GroomCache))
+					{
+						UserData->PrimPaths.AddUnique(PrimPath.GetString());
 
-					GroomCache->AddAssetUserData(UserData);
+						if (Context->MetadataOptions.bCollectMetadata)
+						{
+							UsdToUnreal::ConvertMetadata(
+								GetPrim(),
+								UserData,
+								Context->MetadataOptions.BlockedPrefixFilters,
+								Context->MetadataOptions.bInvertFilters,
+								Context->MetadataOptions.bCollectFromEntireSubtrees
+							);
+						}
+					}
 
 					Context->AssetCache->CacheAsset(GroomCacheHash.ToString(), GroomCache);
 					Context->InfoCache->LinkAssetToPrim(PrimPath, GroomCache);
diff --git a/Engine/Plugins/Importers/USDImporter/Source/USDSchemas/Private/USDShadeMaterialTranslator.cpp b/Engine/Plugins/Importers/USDImporter/Source/USDSchemas/Private/USDShadeMaterialTranslator.cpp
--- a/Engine/Plugins/Importers/USDImporter/Source/USDSchemas/Private/USDShadeMaterialTranslator.cpp
+++ b/Engine/Plugins/Importers/USDImporter/Source/USDSchemas/Private/USDShadeMaterialTranslator.cpp
@@ -6,8 +6,8 @@
 #include "USDAssetCache.h"
 #include "USDAssetUserData.h"
 #include "USDClassesModule.h"
-#include "USDClassesModule.h"
 #include "USDLog.h"
+#include "USDPrimConversion.h"
 #include "USDShadeConversion.h"
 #include "USDTypesConversion.h"
 
@@ -464,13 +464,21 @@ void FUsdShadeMaterialTranslator::PostImportMaterial(const FString& MaterialHash
 		return;
 	}
 
-	UUsdMaterialAssetUserData* UserData = ImportedMaterial->GetAssetUserData<UUsdMaterialAssetUserData>();
-	if (!UserData)
+	if (UUsdMaterialAssetUserData* UserData = UsdUtils::GetOrCreateAssetUserData<UUsdMaterialAssetUserData>(ImportedMaterial))
 	{
-		UserData = NewObject<UUsdMaterialAssetUserData>(ImportedMaterial, TEXT("USDAssetUserData"));
-		ImportedMaterial->AddAssetUserData(UserData);
+		UserData->PrimPaths.AddUnique(PrimPath.GetString());
+
+		if (Context->MetadataOptions.bCollectMetadata)
+		{
+			UsdToUnreal::ConvertMetadata(
+				GetPrim(),
+				UserData,
+				Context->MetadataOptions.BlockedPrefixFilters,
+				Context->MetadataOptions.bInvertFilters,
+				Context->MetadataOptions.bCollectFromEntireSubtrees
+			);
+		}
 	}
-	UserData->PrimPaths.AddUnique(PrimPath.GetString());
 
 	// Note that this needs to run even if we found this material in the asset cache already, otherwise we won't
 	// re-register the prim asset links when we reload a stage
diff --git a/Engine/Plugins/Importers/USDImporter/Source/USDSchemas/Private/USDSkelSkeletonTranslator.cpp b/Engine/Plugins/Importers/USDImporter/Source/USDSchemas/Private/USDSkelSkeletonTranslator.cpp
--- a/Engine/Plugins/Importers/USDImporter/Source/USDSchemas/Private/USDSkelSkeletonTranslator.cpp
+++ b/Engine/Plugins/Importers/USDImporter/Source/USDSchemas/Private/USDSkelSkeletonTranslator.cpp
@@ -16,6 +16,7 @@
 #include "USDLayerUtils.h"
 #include "USDLog.h"
 #include "USDMemory.h"
+#include "USDPrimConversion.h"
 #include "USDSkeletalDataConversion.h"
 #include "USDTypesConversion.h"
 
@@ -429,12 +430,14 @@ namespace UsdSkelSkeletonTranslatorImpl
 		const pxr::UsdSkelCache& InSkelCache,
 		TArray<FSkeletalMeshImportData>& OutLODIndexToSkeletalMeshImportData,
 		TArray<UsdUtils::FUsdPrimMaterialAssignmentInfo>& OutLODIndexToMaterialInfo,
+		FUsdCombinedPrimMetadata& LODMetadata,
 		TArray<SkeletalMeshImportData::FBone>& OutSkeletonBones,
 		FName& OutSkeletonName,
 		UsdUtils::FBlendShapeMap* OutBlendShapes,
 		TSet<FString>& InOutUsedMorphTargetNames,
 		bool bInInterpretLODs,
-		const UsdToUnreal::FUsdMeshConversionOptions& Options
+		const UsdToUnreal::FUsdMeshConversionOptions& Options,
+		const FUsdMetadataImportOptions& MetadataOptions
 	)
 	{
 		FScopedUsdAllocs UsdAllocs;
@@ -453,6 +456,8 @@ namespace UsdSkelSkeletonTranslatorImpl
 
 		pxr::UsdPrim SkeletonPrim = Skeleton.GetPrim();
 		pxr::SdfPath SkeletonPrimPath = SkeletonPrim.GetPrimPath();
+		pxr::UsdPrim ClosestParentSkelRoot = UsdUtils::GetClosestParentSkelRoot(SkeletonPrim);
+		pxr::SdfPath SkelRootPrimPath = ClosestParentSkelRoot.GetPrimPath();
 
 		pxr::UsdStageRefPtr Stage = SkeletonPrim.GetStage();
 		const FUsdStageInfo StageInfo(Stage);
@@ -469,6 +474,24 @@ namespace UsdSkelSkeletonTranslatorImpl
 			OutSkeletonName = *UsdToUnreal::ConvertString(SkeletonPrim.GetName());
 		}
 
+		// Note that the approach is to store skelroot + skinned mesh metadata onto the USkeletalMesh, and
+		// purely skeleton metadata onto the USkeleton.
+		// Here we collect metadata from the skelroot itself, as the process inside ConvertLOD will only collect
+		// metadata from the skinned meshes
+		if (MetadataOptions.bCollectMetadata)
+		{
+			// Here we're always setting this to false otherwise we'll also end up collecting metadata on the skeleton for
+			// skel animation and other prims that weren't handled
+			const bool bCollectMetadataFromSubtree = false;
+			UsdToUnreal::ConvertMetadata(
+				ClosestParentSkelRoot,
+				LODMetadata,
+				MetadataOptions.BlockedPrefixFilters,
+				MetadataOptions.bInvertFilters,
+				bCollectMetadataFromSubtree
+			);
+		}
+
 		TMap<int32, FSkeletalMeshImportData> LODIndexToSkeletalMeshImportDataMap;
 		TMap<int32, UsdUtils::FUsdPrimMaterialAssignmentInfo> LODIndexToMaterialInfoMap;
 		TSet<FString> ProcessedLODParentPaths;
@@ -501,15 +524,20 @@ namespace UsdSkelSkeletonTranslatorImpl
 		}
 
 		TFunction<bool(const pxr::UsdGeomMesh&, int32)> ConvertLOD =
-			[&LODIndexToSkeletalMeshImportDataMap,
-			 &LODIndexToMaterialInfoMap,
-			 &InSkelCache,
-			 &Stage,
-			 &SkeletonPrimPath,
-			 &InOutUsedMorphTargetNames,
-			 OutBlendShapes,
-			 &StageInfo,
-			 Options](const pxr::UsdGeomMesh& LODMesh, int32 LODIndex)
+		[
+			&LODIndexToSkeletalMeshImportDataMap,
+			&LODIndexToMaterialInfoMap,
+			&LODMetadata,
+			&InSkelCache,
+			&Stage,
+			&SkeletonPrimPath,
+			&InOutUsedMorphTargetNames,
+			OutBlendShapes,
+			&StageInfo,
+			Options,
+			&MetadataOptions
+		]
+		(const pxr::UsdGeomMesh& LODMesh, int32 LODIndex)
 		{
 			// Construct this and SkinningQuery every time so as to survive the prim reference invalidation caused by flipping LODs
 			pxr::UsdSkelSkeletonQuery SkeletonQuery = InSkelCache.GetSkelQuery(pxr::UsdSkelSkeleton{Stage->GetPrimAtPath(SkeletonPrimPath)});
@@ -556,6 +584,18 @@ namespace UsdSkelSkeletonTranslatorImpl
 				return true;
 			}
 
+			if (MetadataOptions.bCollectMetadata && MetadataOptions.bCollectFromEntireSubtrees)
+			{
+				// Collect metadata from this particular LOD mesh prim
+				UsdToUnreal::ConvertMetadata(
+					LODMesh.GetPrim(),
+					LODMetadata,
+					MetadataOptions.BlockedPrefixFilters,
+					MetadataOptions.bInvertFilters,
+					MetadataOptions.bCollectFromEntireSubtrees
+				);
+			}
+
 			if (OutBlendShapes)
 			{
 				pxr::SdfPath MeshPrimPath = LODMesh.GetPrim().GetPrimPath();
@@ -611,6 +651,9 @@ namespace UsdSkelSkeletonTranslatorImpl
 			}
 		}
 
+		// Repopulate the skeleton cache because flipping through LODs can invalidate some stuff like skeleton references
+		InSkelCache.Populate(pxr::UsdSkelRoot{Stage->GetPrimAtPath(SkelRootPrimPath)}, pxr::UsdTraverseInstanceProxies());
+
 		// Place the LODs in order as we can't have e.g. LOD0 and LOD2 without LOD1, and there's no reason downstream code needs to care about
 		// what LOD number these data originally wanted to be
 		TMap<int32, int32> OldLODIndexToNewLODIndex;
@@ -1294,6 +1337,7 @@ namespace UsdSkelSkeletonTranslatorImpl
 		// Outputs
 		TArray<FSkeletalMeshImportData> LODIndexToSkeletalMeshImportData;
 		TArray<UsdUtils::FUsdPrimMaterialAssignmentInfo> LODIndexToMaterialInfo;
+		FUsdCombinedPrimMetadata LODMetadata;
 		TArray<SkeletalMeshImportData::FBone> SkeletonBones;
 		FName SkeletonName;
 		UsdUtils::FBlendShapeMap NewBlendShapes;
@@ -1398,12 +1442,14 @@ namespace UsdSkelSkeletonTranslatorImpl
 				   SkelCache.Get(),
 				   LODIndexToSkeletalMeshImportData,
 				   LODIndexToMaterialInfo,
+				   LODMetadata,
 				   SkeletonBones,
 				   SkeletonName,
 				   OutBlendShapes,
 				   UsedMorphTargetNames,
 				   bTryLODParsing,
-				   Options
+				   Options,
+				   Context->MetadataOptions
 			   );
 
 			   // If we parsed LODs we could potentially have invalidated our references to our queries,
@@ -1475,16 +1521,43 @@ namespace UsdSkelSkeletonTranslatorImpl
 
 				if (SkeletalMesh)
 				{
-					UUsdMeshAssetUserData* UserData = SkeletalMesh->GetAssetUserData<UUsdMeshAssetUserData>();
-					if (!UserData)
+					// Handle the SkeletalMesh AssetUserData
+					if (UUsdMeshAssetUserData* UserData = UsdUtils::GetOrCreateAssetUserData<UUsdMeshAssetUserData>(SkeletalMesh))
 					{
-						UserData = NewObject<UUsdMeshAssetUserData>(SkeletalMesh, TEXT("USDAssetUserData"));
 						UserData->PrimvarToUVIndex = LODIndexToMaterialInfo[0].PrimvarToUVIndex;	// We use the same primvar mapping for all LODs
-						SkeletalMesh->AddAssetUserData(UserData);
+						UserData->PrimPaths.AddUnique(SkeletonPrimPath.GetString());
+
+						// For the skel task chain we always collect skeletal mesh metadata when first parsing the prims directly, as it
+						// allows us to do it while we're flipping through LOD variants, if any
+						UserData->StageIdentifierToMetadata.Add(GetSkeletonPrim().GetStage().GetRootLayer().GetIdentifier(), LODMetadata);
+
+						MeshTranslationImpl::RecordSourcePrimsForMaterialSlots(LODIndexToMaterialInfo, UserData);
 					}
-					UserData->PrimPaths.AddUnique(SkeletonPrimPath.GetString());
 
-					MeshTranslationImpl::RecordSourcePrimsForMaterialSlots(LODIndexToMaterialInfo, UserData);
+					// Handle the Skeleton AssetUserData
+					if (USkeleton* Skeleton = SkeletalMesh->GetSkeleton())
+					{
+						if (UUsdAssetUserData* UserData = UsdUtils::GetOrCreateAssetUserData(Skeleton))
+						{
+							UserData->PrimPaths.AddUnique(SkeletonPrimPath.GetString());
+
+							if (Context->MetadataOptions.bCollectMetadata)
+							{
+								// Since we never collapse, we'll spawn assets components for any child prim that happens to be inside
+								// the skeleton itself, and the skeleton type doesn't have any relevant "child prim" type (like for
+								// Mesh prims and UsdGeomSubsets), so we're probably safe in never collecting metadata from the skeleton
+								// prim subtree
+								const bool bCollectFromEntireSubtrees = false;
+								UsdToUnreal::ConvertMetadata(
+									GetSkeletonPrim(),
+									UserData,
+									Context->MetadataOptions.BlockedPrefixFilters,
+									Context->MetadataOptions.bInvertFilters,
+									bCollectFromEntireSubtrees
+								);
+							}
+						}
+					}
 
 					if (bIsNew)
 					{
@@ -1717,17 +1790,29 @@ namespace UsdSkelSkeletonTranslatorImpl
 						&LayerStartOffsetSeconds
 					);
 
-					if (bSuccess
-						&& (AnimSequence->GetDataModel()->GetNumBoneTracks() != 0 || AnimSequence->GetDataModel()->GetNumberOfFloatCurves() != 0))
+					if (bSuccess && (AnimSequence->GetDataModel()->GetNumBoneTracks() != 0 || AnimSequence->GetDataModel()->GetNumberOfFloatCurves() != 0))
 					{
-						UUsdAnimSequenceAssetUserData* UserData = NewObject<UUsdAnimSequenceAssetUserData>(
-							AnimSequence,
-							TEXT("USDAssetUserData")
-						);
-						UserData->PrimPaths = {SkelAnimationPrimPath};
-						UserData->LayerStartOffsetSeconds = LayerStartOffsetSeconds;
+						if (UUsdAnimSequenceAssetUserData* UserData = UsdUtils::GetOrCreateAssetUserData<UUsdAnimSequenceAssetUserData>(AnimSequence))
+						{
+							UserData->PrimPaths.AddUnique(SkelAnimationPrimPath);
+							UserData->LayerStartOffsetSeconds = LayerStartOffsetSeconds;
 
-						AnimSequence->AddAssetUserData(UserData);
+							if (Context->MetadataOptions.bCollectMetadata)
+							{
+								// Since we never collapse, we'll spawn assets components for any child prim that happens to be inside
+								// the skeleton itself, and the SkelAnimation type doesn't have any relevant "child prim" type (like for
+								// Mesh prims and UsdGeomSubsets), so we're probably safe in never collecting metadata from the SkelAnimation
+								// prim subtree
+								const bool bCollectFromEntireSubtrees = false;
+								UsdToUnreal::ConvertMetadata(
+									SkelAnimationPrim,
+									UserData,
+									Context->MetadataOptions.BlockedPrefixFilters,
+									Context->MetadataOptions.bInvertFilters,
+									bCollectFromEntireSubtrees
+								);
+							}
+						}
 
 						Context->AssetCache->CacheAsset(HashString, AnimSequence);
 					}
diff --git a/Engine/Plugins/Importers/USDImporter/Source/USDStage/Private/USDLevelSequenceHelper.cpp b/Engine/Plugins/Importers/USDImporter/Source/USDStage/Private/USDLevelSequenceHelper.cpp
--- a/Engine/Plugins/Importers/USDImporter/Source/USDStage/Private/USDLevelSequenceHelper.cpp
+++ b/Engine/Plugins/Importers/USDImporter/Source/USDStage/Private/USDLevelSequenceHelper.cpp
@@ -2202,7 +2202,7 @@ void FUsdLevelSequenceHelperImpl::AddGeometryCacheTracks(const UUsdPrimTwin& Pri
 		if (GeometryCacheSequence == MainLevelSequence)
 		{
 			// The LayerStartOffset needs to be applied only for the track in the main sequence. For subsequences, it's the subsequence section that is offset.
-			if (UUsdAnimSequenceAssetUserData* UserData = GeometryCache->GetAssetUserData<UUsdAnimSequenceAssetUserData>())
+			if (UUsdGeometryCacheAssetUserData* UserData = GeometryCache->GetAssetUserData<UUsdGeometryCacheAssetUserData>())
 			{
 				LayerStartOffsetSeconds = UserData->LayerStartOffsetSeconds;
 			}
diff --git a/Engine/Plugins/Importers/USDImporter/Source/USDStage/Private/USDStageActor.cpp b/Engine/Plugins/Importers/USDImporter/Source/USDStage/Private/USDStageActor.cpp
--- a/Engine/Plugins/Importers/USDImporter/Source/USDStage/Private/USDStageActor.cpp
+++ b/Engine/Plugins/Importers/USDImporter/Source/USDStage/Private/USDStageActor.cpp
@@ -143,6 +143,7 @@ struct FUsdStageActorImpl
 		TranslationContext->MaterialPurpose = StageActor->MaterialPurpose;
 		TranslationContext->RootMotionHandling = StageActor->RootMotionHandling;
 		TranslationContext->SubdivisionLevel = StageActor->SubdivisionLevel;
+		TranslationContext->MetadataOptions = StageActor->MetadataOptions;
 		TranslationContext->BlendShapesByPath = &StageActor->BlendShapesByPath;
 		TranslationContext->InfoCache = StageActor->InfoCache;
 		TranslationContext->BBoxCache = StageActor->BBoxCache;
@@ -577,6 +578,8 @@ struct FUsdStageActorImpl
 			EventAttributes.Emplace(TEXT("RootMotionHandling"), LexToString((uint8)StageActor->RootMotionHandling));
 			EventAttributes.Emplace(TEXT("SubdivisionLevel"), LexToString(StageActor->SubdivisionLevel));
 
+			UsdUtils::AddAnalyticsAttributes(StageActor->MetadataOptions, EventAttributes);
+
 			const bool bAutomated = false;
 			IUsdClassesModule::SendAnalytics(MoveTemp(EventAttributes), TEXT("Open"), bAutomated, ElapsedSeconds, NumberOfFrames, Extension);
 		}
@@ -891,6 +894,93 @@ struct FUsdStageActorImpl
 			}
 		}
 	}
+
+	// This function is in charge of writing out to USD the analogous metadata change that
+	// we just received for ChangedUserData via the PropertyChangedEvent
+	static void WriteOutAssetMetadataChange(
+		const AUsdStageActor* StageActor,
+		const UUsdAssetUserData* ChangedUserData,
+		const FPropertyChangedEvent& PropertyChangedEvent
+	)
+	{
+#if USE_USD_SDK
+		if (!StageActor || !ChangedUserData)
+		{
+			return;
+		}
+
+		UE::FUsdStage Stage = StageActor->GetUsdStage();
+		if (!Stage)
+		{
+			return;
+		}
+
+		const bool bChangeWasRemoval = PropertyChangedEvent.ChangeType == EPropertyChangeType::ArrayRemove
+									   || PropertyChangedEvent.ChangeType == EPropertyChangeType::ArrayClear;
+
+		const bool bHasMetadataFilters = StageActor->MetadataOptions.BlockedPrefixFilters.Num() > 0;
+
+		// For this call, we'll only care about metadata on prims that are from the stage opened by this stage actor.
+		// If we need to modify other stages for this PropertyChangedEvent somehow, the other stage actors' call to this
+		// function will handle them
+		FString StageIdentifier = Stage.GetRootLayer().GetIdentifier();
+		const FUsdCombinedPrimMetadata* StageMetadata = ChangedUserData->StageIdentifierToMetadata.Find(StageIdentifier);
+
+		// This asset user data doesn't have any metadata for this particular stage, nothing to do
+		if (!StageMetadata)
+		{
+			return;
+		}
+
+		for (const TPair<FString, FUsdPrimMetadata>& PrimPathToMetadata : StageMetadata->PrimPathToMetadata)
+		{
+			const FString& PrimPath = PrimPathToMetadata.Key;
+			const FUsdPrimMetadata& PrimMetadata = PrimPathToMetadata.Value;
+
+			const UE::FUsdPrim& Prim = Stage.GetPrimAtPath(UE::FSdfPath{*PrimPath});
+
+			// If the change we need to write out is a removal, since we can't tell *what* was removed from the PropertyChangedEvent,
+			// the only thing we can do is wipe all metadata on the prim and replace that with what we have on our AssetUserData
+			if (bChangeWasRemoval)
+			{
+				// If the metadata we have was obtained with metadata filters, we're in trouble: We can't just clear everything
+				// and write what we have, because we just have the stuff that passed the filter. What we'll do here then is
+				// invert the filters and collect metadata again (which gives us the stuff that are *not* already in our
+				// AssetUserData), then clear all metadata on the prim, write out that "inverted" dataset, and (later) also
+				// write out our current AssetUserData
+				if (bHasMetadataFilters)
+				{
+					// We are clearing/writing to a particular prim here, "collecting from subtrees" is an UE-concept
+					const bool bCollectFromEntireSubtrees = false;
+					FUsdCombinedPrimMetadata TempInvertedMetadata;
+
+					const bool bSuccess = UsdToUnreal::ConvertMetadata(
+						Prim,
+						TempInvertedMetadata,
+						StageActor->MetadataOptions.BlockedPrefixFilters,
+						!StageActor->MetadataOptions.bInvertFilters,
+						bCollectFromEntireSubtrees
+					);
+
+					// Don't clear anything if we failed to collect the inverted dataset
+					if (bSuccess)
+					{
+						UsdUtils::ClearNonEssentialPrimMetadata(Prim);
+						UnrealToUsd::ConvertMetadata(TempInvertedMetadata, Prim);
+					}
+				}
+				// If what we have currently was obtained without any filters, we can be sure that what we have is a good
+				// representation of all metadata on this prim, so we can just clear everything and write what we have
+				else
+				{
+					UsdUtils::ClearNonEssentialPrimMetadata(Prim);
+				}
+			}
+
+			UnrealToUsd::ConvertMetadata(PrimMetadata, Prim);
+		}
+#endif	  // USE_USD_SDK
+	}
 };
 
 /**
@@ -973,6 +1063,13 @@ AUsdStageActor::AUsdStageActor()
 	, MaterialPurpose(*UnrealIdentifiers::MaterialPreviewPurpose)
 	, RootMotionHandling(EUsdRootMotionHandling::NoAdditionalRootMotion)
 	, SubdivisionLevel(0)
+	, MetadataOptions(FUsdMetadataImportOptions{
+		false, 	/* bCollectMetadata */
+		false, 	/* bCollectFromEntireSubtrees */
+		false, 	/* bCollectOnComponents */
+		{},		/* BlockedPrefixFilters */
+		false	/* bInvertFilters */
+	})
 	, Time(0.0f)
 	, bIsTransitioningIntoPIE(false)
 	, bIsModifyingAProperty(false)
@@ -2487,6 +2584,11 @@ void AUsdStageActor::SetRootMotionHandling(EUsdRootMotionHandling NewHandlingStr
 
 void AUsdStageActor::SetSubdivisionLevel(int32 NewSubdivisionLevel)
 {
+	if (NewSubdivisionLevel == SubdivisionLevel)
+	{
+		return;
+	}
+
 	const bool bMarkDirty = false;
 	Modify(bMarkDirty);
 
@@ -2494,6 +2596,76 @@ void AUsdStageActor::SetSubdivisionLevel(int32 NewSubdivisionLevel)
 	LoadUsdStage();
 }
 
+void AUsdStageActor::SetCollectMetadata(bool bNewCollectValue)
+{
+	if (bNewCollectValue == MetadataOptions.bCollectMetadata)
+	{
+		return;
+	}
+
+	const bool bMarkDirty = false;
+	Modify(bMarkDirty);
+
+	MetadataOptions.bCollectMetadata = bNewCollectValue;
+	LoadUsdStage();
+}
+
+void AUsdStageActor::SetCollectFromEntireSubtrees(bool bNewCollectValue)
+{
+	if (bNewCollectValue == MetadataOptions.bCollectFromEntireSubtrees)
+	{
+		return;
+	}
+
+	const bool bMarkDirty = false;
+	Modify(bMarkDirty);
+
+	MetadataOptions.bCollectFromEntireSubtrees = bNewCollectValue;
+	LoadUsdStage();
+}
+
+void AUsdStageActor::SetCollectOnComponents(bool bNewCollectValue)
+{
+	if (bNewCollectValue == MetadataOptions.bCollectOnComponents)
+	{
+		return;
+	}
+
+	const bool bMarkDirty = false;
+	Modify(bMarkDirty);
+
+	MetadataOptions.bCollectOnComponents = bNewCollectValue;
+	LoadUsdStage();
+}
+
+void AUsdStageActor::SetBlockedPrefixFilters(const TArray<FString>& NewFilters)
+{
+	if (NewFilters == MetadataOptions.BlockedPrefixFilters)
+	{
+		return;
+	}
+
+	const bool bMarkDirty = false;
+	Modify(bMarkDirty);
+
+	MetadataOptions.BlockedPrefixFilters = NewFilters;
+	LoadUsdStage();
+}
+
+void AUsdStageActor::SetInvertFilters(bool bNewInvertValue)
+{
+	if (bNewInvertValue == MetadataOptions.bInvertFilters)
+	{
+		return;
+	}
+
+	const bool bMarkDirty = false;
+	Modify(bMarkDirty);
+
+	MetadataOptions.bInvertFilters = bNewInvertValue;
+	LoadUsdStage();
+}
+
 float AUsdStageActor::GetTime() const
 {
 	return Time;
@@ -4042,6 +4214,18 @@ void AUsdStageActor::OnObjectPropertyChanged(UObject* ObjectBeingModified, FProp
 		FUsdStageActorImpl::AllowListComponentHierarchy(GetRootComponent(), VisitedObjects);
 	}
 
+	// If the user is just setting metadata on one of our transient UAssets, then try to author the metadata back
+	// out to the relevant prims
+	if (UUsdAssetUserData* UserData = Cast<UUsdAssetUserData>(ObjectBeingModified))
+	{
+		if (UsdAssetCache->IsAssetOwnedByCache(ObjectBeingModified->GetOuter()->GetPathName()))
+		{
+			FScopedBlockNoticeListening BlockNoticeListening(this);
+			FUsdStageActorImpl::WriteOutAssetMetadataChange(this, UserData, PropertyChangedEvent);
+			return;
+		}
+	}
+
 	// We have to accept actor and component events here, because actor transform changes do not trigger root component
 	// transform property events, and component property changes don't trigger actor property change events
 	bool bIsActorEvent = false;
@@ -4350,7 +4534,9 @@ void AUsdStageActor::HandlePropertyChangedEvent(FPropertyChangedEvent& PropertyC
 	const bool bAlwaysMarkAsDirty = false;
 	Modify(bAlwaysMarkAsDirty);
 
-	FProperty* PropertyThatChanged = PropertyChangedEvent.MemberProperty;
+	// If we're changing a property inside a struct, like "bCollectMetadata" inside our MetadataOptions, then
+	// "MemberProperty" will point to "MetadataOptions", and "Property" is the thing that will point to "bCollectMetadata"
+	FProperty* PropertyThatChanged = PropertyChangedEvent.Property;
 	const FName PropertyName = PropertyThatChanged ? PropertyThatChanged->GetFName() : NAME_None;
 
 	if (PropertyName == GET_MEMBER_NAME_CHECKED(AUsdStageActor, RootLayer))
@@ -4445,7 +4631,39 @@ void AUsdStageActor::HandlePropertyChangedEvent(FPropertyChangedEvent& PropertyC
 	}
 	else if (PropertyName == GET_MEMBER_NAME_CHECKED(AUsdStageActor, SubdivisionLevel))
 	{
-		SetSubdivisionLevel(SubdivisionLevel);
+		int32 CorrectSubdivisionLevel = SubdivisionLevel;
+		SubdivisionLevel = !SubdivisionLevel;
+		SetSubdivisionLevel(CorrectSubdivisionLevel);
+	}
+	else if (PropertyName == GET_MEMBER_NAME_CHECKED(FUsdMetadataImportOptions, bCollectMetadata))
+	{
+		bool bCorrectCollectValue = MetadataOptions.bCollectMetadata;
+		MetadataOptions.bCollectMetadata = !bCorrectCollectValue;
+		SetCollectMetadata(bCorrectCollectValue);
+	}
+	else if (PropertyName == GET_MEMBER_NAME_CHECKED(FUsdMetadataImportOptions, bCollectFromEntireSubtrees))
+	{
+		bool bCorrectCollectValue = MetadataOptions.bCollectFromEntireSubtrees;
+		MetadataOptions.bCollectFromEntireSubtrees = !bCorrectCollectValue;
+		SetCollectFromEntireSubtrees(bCorrectCollectValue);
+	}
+	else if (PropertyName == GET_MEMBER_NAME_CHECKED(FUsdMetadataImportOptions, bCollectOnComponents))
+	{
+		bool bCorrectCollectValue = MetadataOptions.bCollectOnComponents;
+		MetadataOptions.bCollectOnComponents = !bCorrectCollectValue;
+		SetCollectOnComponents(bCorrectCollectValue);
+	}
+	else if (PropertyName == GET_MEMBER_NAME_CHECKED(FUsdMetadataImportOptions, BlockedPrefixFilters))
+	{
+		TArray<FString> CorrectFilters = MetadataOptions.BlockedPrefixFilters;
+		MetadataOptions.BlockedPrefixFilters.Add(TEXT("dummy"));
+		SetBlockedPrefixFilters(CorrectFilters);
+	}
+	else if (PropertyName == GET_MEMBER_NAME_CHECKED(FUsdMetadataImportOptions, bInvertFilters))
+	{
+		bool bCorrectInvertValue = MetadataOptions.bInvertFilters;
+		MetadataOptions.bInvertFilters = !bCorrectInvertValue;
+		SetInvertFilters(bCorrectInvertValue);
 	}
 	else if (PropertyName == GET_MEMBER_NAME_CHECKED(AUsdStageActor, UsdAssetCache))
 	{
diff --git a/Engine/Plugins/Importers/USDImporter/Source/USDStage/Public/USDStageActor.h b/Engine/Plugins/Importers/USDImporter/Source/USDStage/Public/USDStageActor.h
--- a/Engine/Plugins/Importers/USDImporter/Source/USDStage/Public/USDStageActor.h
+++ b/Engine/Plugins/Importers/USDImporter/Source/USDStage/Public/USDStageActor.h
@@ -4,6 +4,7 @@
 
 #include "USDLevelSequenceHelper.h"
 #include "USDListener.h"
+#include "USDMetadataImportOptions.h"
 #include "USDSkeletalDataConversion.h"
 #include "UsdWrappers/SdfPath.h"
 #include "UsdWrappers/UsdStage.h"
@@ -114,6 +115,10 @@ public:
 	UPROPERTY(EditAnywhere, BlueprintReadOnly, Category = "USD", config)
 	int32 SubdivisionLevel;
 
+	/* Describes if/how we should collect metadata from USD prims onto the assets and components we generate when opening stages */
+	UPROPERTY(EditAnywhere, BlueprintReadOnly, Category = "USD", config)
+	FUsdMetadataImportOptions MetadataOptions;
+
 public:
 	DECLARE_EVENT_OneParam(AUsdStageActor, FOnActorLoaded, AUsdStageActor*);
 	USDSTAGE_API static FOnActorLoaded OnActorLoaded;
@@ -178,6 +183,21 @@ public:
 	UFUNCTION(BlueprintCallable, Category = "USD", meta = (CallInEditor = "true"))
 	USDSTAGE_API void SetSubdivisionLevel(int32 NewLevel);
 
+	UFUNCTION(BlueprintCallable, Category = "USD|Metadata", meta = (CallInEditor = "true"))
+	USDSTAGE_API void SetCollectMetadata(bool bNewCollectValue);
+
+	UFUNCTION(BlueprintCallable, Category = "USD|Metadata", meta = (CallInEditor = "true"))
+	USDSTAGE_API void SetCollectFromEntireSubtrees(bool bNewCollectValue);
+
+	UFUNCTION(BlueprintCallable, Category = "USD|Metadata", meta = (CallInEditor = "true"))
+	USDSTAGE_API void SetCollectOnComponents(bool bNewCollectValue);
+
+	UFUNCTION(BlueprintCallable, Category = "USD|Metadata", meta = (CallInEditor = "true"))
+	USDSTAGE_API void SetBlockedPrefixFilters(const TArray<FString>& NewFilters);
+
+	UFUNCTION(BlueprintCallable, Category = "USD|Metadata", meta = (CallInEditor = "true"))
+	USDSTAGE_API void SetInvertFilters(bool bNewInvertValue);
+
 	UFUNCTION(BlueprintCallable, Category = "USD", meta = (CallInEditor = "true"))
 	USDSTAGE_API float GetTime() const;
 
diff --git a/Engine/Plugins/Importers/USDImporter/Source/USDStageEditor/Private/SUSDStage.cpp b/Engine/Plugins/Importers/USDImporter/Source/USDStageEditor/Private/SUSDStage.cpp
--- a/Engine/Plugins/Importers/USDImporter/Source/USDStageEditor/Private/SUSDStage.cpp
+++ b/Engine/Plugins/Importers/USDImporter/Source/USDStageEditor/Private/SUSDStage.cpp
@@ -31,7 +31,9 @@
 #include "Framework/MultiBox/MultiBoxBuilder.h"
 #include "ISceneOutliner.h"
 #include "ISceneOutlinerColumn.h"
+#include "ISinglePropertyView.h"
 #include "Modules/ModuleManager.h"
+#include "PropertyEditorModule.h"
 #include "SceneOutlinerModule.h"
 #include "ScopedTransaction.h"
 #include "Selection.h"
@@ -1042,6 +1044,11 @@ void SUsdStage::FillOptionsMenu(FMenuBuilder& MenuBuilder)
 			LOCTEXT("SubdivLevel_ToolTip", "Subdivision level to use for all subdivision meshes on the opened stage. 0 means 'don't subdivide'"),
 			FNewMenuDelegate::CreateSP(this, &SUsdStage::FillSubdivisionLevelSubMenu));
 
+		MenuBuilder.AddSubMenu(
+			LOCTEXT("Metadata", "Metadata"),
+			LOCTEXT("Metadata_ToolTip", "How to collect and handle metadata from USD prims"),
+			FNewMenuDelegate::CreateSP(this, &SUsdStage::FillMetadataSubMenu));
+
 		MenuBuilder.AddSubMenu(
 			LOCTEXT( "Collapsing", "Collapsing" ),
 			LOCTEXT( "Collapsing_ToolTip", "Whether to try to combine individual assets and components of the same type on a Kind-per-Kind basis, like multiple Mesh prims into a single Static Mesh" ),
@@ -1600,6 +1607,82 @@ void SUsdStage::FillSubdivisionLevelSubMenu(FMenuBuilder& MenuBuilder)
 	MenuBuilder.AddWidget(Slider, FText::FromString(TEXT("Subdivision level: ")), bNoIndent);
 }
 
+void SUsdStage::FillMetadataSubMenu(FMenuBuilder& MenuBuilder)
+{
+	AUsdStageActor* StageActor = GetStageActorOrCDO();
+	if (!StageActor)
+	{
+		return;
+	}
+
+	FPropertyEditorModule& PropertyEditorModule = FModuleManager::GetModuleChecked<FPropertyEditorModule>("PropertyEditor");
+
+	FDetailsViewArgs DetailsViewArgs;
+	DetailsViewArgs.bAllowSearch = false;
+	DetailsViewArgs.NameAreaSettings = FDetailsViewArgs::HideNameArea;
+	DetailsViewArgs.ColumnWidth = 0.5;
+	TSharedRef<IDetailsView> DetailView = PropertyEditorModule.CreateDetailView(DetailsViewArgs);
+
+	DetailView->SetIsPropertyVisibleDelegate(FIsPropertyVisible::CreateLambda(
+		[](const FPropertyAndParent& PropertyAndParent) -> bool
+		{
+			static const FName MetadataOptionsName = GET_MEMBER_NAME_CHECKED(AUsdStageActor, MetadataOptions);
+
+			const FName PropertyName = PropertyAndParent.Property.GetFName();
+			if (PropertyName == MetadataOptionsName)
+			{
+				return true;
+			}
+
+			for(const FProperty* ParentProp : PropertyAndParent.ParentProperties)
+			{
+				const FName ParentPropertyName = ParentProp->GetFName();
+				if(ParentPropertyName == MetadataOptionsName)
+				{
+					return true;
+				}
+			}
+
+			return false;
+		}
+	));
+	DetailView->SetIsCustomRowVisibleDelegate(FIsCustomRowVisible::CreateLambda(
+		[](FName RowName, FName ParentName) -> bool
+		{
+			static const FName FilterPropertyName = GET_MEMBER_NAME_CHECKED(FUsdMetadataImportOptions, BlockedPrefixFilters);
+			return RowName == FilterPropertyName;
+		}
+	));
+	DetailView->SetObject(StageActor);
+
+	DetailView->OnFinishedChangingProperties().AddLambda(
+		[this](const FPropertyChangedEvent& PropertyChangedEvent)
+		{
+			if (AUsdStageActor* StageActor = GetStageActorOrCDO())
+			{
+				if (StageActor->IsTemplate())
+				{
+					StageActor->SaveConfig();
+				}
+			}
+		}
+	);
+
+	// Override the widget with some extra width because the prefix filters need some space to be used
+	// comfortably
+	TSharedRef<SBox> Box = SNew(SBox)
+	.Padding(FMargin(0.0f, 0.0f))
+	.VAlign(VAlign_Fill)
+	.HAlign(HAlign_Fill)
+	.WidthOverride(500)
+	[
+		DetailView
+	];
+
+	const bool bNoIndent = true;
+	MenuBuilder.AddWidget(Box, FText::GetEmpty(), bNoIndent);
+}
+
 void SUsdStage::FillCollapsingSubMenu( FMenuBuilder& MenuBuilder )
 {
 	MenuBuilder.AddMenuEntry(
@@ -1612,7 +1695,7 @@ void SUsdStage::FillCollapsingSubMenu( FMenuBuilder& MenuBuilder )
 				if ( AUsdStageActor* StageActor = GetStageActorOrCDO() )
 				{
 					FScopedTransaction Transaction( FText::Format(
-						LOCTEXT( "MergeIdenticalMaterialSlotsTransaction", "Toggle bMergeIdenticalMaterialSlots on USD stage actor '{1}'" ),
+						LOCTEXT( "MergeIdenticalMaterialSlotsTransaction", "Toggle bMergeIdenticalMaterialSlots on USD stage actor '{0}'" ),
 						FText::FromString( StageActor->GetActorLabel() )
 					) );
 
diff --git a/Engine/Plugins/Importers/USDImporter/Source/USDStageEditor/Private/SUSDStage.h b/Engine/Plugins/Importers/USDImporter/Source/USDStageEditor/Private/SUSDStage.h
--- a/Engine/Plugins/Importers/USDImporter/Source/USDStageEditor/Private/SUSDStage.h
+++ b/Engine/Plugins/Importers/USDImporter/Source/USDStageEditor/Private/SUSDStage.h
@@ -90,6 +90,7 @@ protected:
 	void FillMaterialPurposeSubMenu( FMenuBuilder& MenuBuilder );
 	void FillRootMotionSubMenu( FMenuBuilder& MenuBuilder );
 	void FillSubdivisionLevelSubMenu(FMenuBuilder& MenuBuilder);
+	void FillMetadataSubMenu(FMenuBuilder& MenuBuilder);
 	void FillCollapsingSubMenu( FMenuBuilder& MenuBuilder );
 	void FillInterpolationTypeSubMenu( FMenuBuilder& MenuBuilder );
 	void FillSelectionSubMenu( FMenuBuilder& MenuBuilder );
diff --git a/Engine/Plugins/Importers/USDImporter/Source/USDStageEditorViewModels/Private/USDStageViewModel.cpp b/Engine/Plugins/Importers/USDImporter/Source/USDStageEditorViewModels/Private/USDStageViewModel.cpp
--- a/Engine/Plugins/Importers/USDImporter/Source/USDStageEditorViewModels/Private/USDStageViewModel.cpp
+++ b/Engine/Plugins/Importers/USDImporter/Source/USDStageEditorViewModels/Private/USDStageViewModel.cpp
@@ -317,6 +317,7 @@ void FUsdStageViewModel::ImportStage( const TCHAR* TargetContentFolder, UUsdStag
 		ImportContext.ImportOptions->NaniteTriangleThreshold = StageActor->NaniteTriangleThreshold;
 		ImportContext.ImportOptions->RootMotionHandling = StageActor->RootMotionHandling;
 		ImportContext.ImportOptions->SubdivisionLevel = StageActor->SubdivisionLevel;
+		ImportContext.ImportOptions->MetadataOptions = StageActor->MetadataOptions;
 		ImportContext.ImportOptions->KindsToCollapse = StageActor->KindsToCollapse;
 		ImportContext.ImportOptions->bMergeIdenticalMaterialSlots = StageActor->bMergeIdenticalMaterialSlots;
 
diff --git a/Engine/Plugins/Importers/USDImporter/Source/USDStageImporter/Private/USDStageImportOptions.cpp b/Engine/Plugins/Importers/USDImporter/Source/USDStageImporter/Private/USDStageImportOptions.cpp
--- a/Engine/Plugins/Importers/USDImporter/Source/USDStageImporter/Private/USDStageImportOptions.cpp
+++ b/Engine/Plugins/Importers/USDImporter/Source/USDStageImporter/Private/USDStageImportOptions.cpp
@@ -27,6 +27,13 @@ UUsdStageImportOptions::UUsdStageImportOptions(const FObjectInitializer& ObjectI
 	RenderContextToImport = UsdSchemasModule.GetRenderContextRegistry().GetUnrealRenderContext();
 	MaterialPurpose = *UnrealIdentifiers::MaterialPreviewPurpose;
 	SubdivisionLevel = 0;
+	MetadataOptions = FUsdMetadataImportOptions{
+		true, 	/* bCollectMetadata */
+		true, 	/* bCollectFromEntireSubtrees */
+		false, 	/* bCollectOnComponents */
+		{},		/* BlockedPrefixFilters */
+		false	/* bInvertFilters */
+	};
 	bOverrideStageOptions = false;
 	StageOptions.MetersPerUnit = 0.01f;
 	StageOptions.UpAxis = EUsdUpAxis::ZAxis;
@@ -97,6 +104,9 @@ void UsdUtils::AddAnalyticsAttributes(
 	InOutAttributes.Emplace( TEXT( "MaterialPurpose" ), Options.MaterialPurpose.ToString() );
 	InOutAttributes.Emplace( TEXT( "RootMotionHandling" ), LexToString( ( uint8 ) Options.RootMotionHandling ) );
 	InOutAttributes.Emplace(TEXT("SubdivisionLevel"), LexToString(Options.SubdivisionLevel));
+
+	UsdUtils::AddAnalyticsAttributes(Options.MetadataOptions, InOutAttributes);
+
 	InOutAttributes.Emplace( TEXT( "OverrideStageOptions" ), Options.bOverrideStageOptions );
 	if ( Options.bOverrideStageOptions )
 	{
diff --git a/Engine/Plugins/Importers/USDImporter/Source/USDStageImporter/Private/USDStageImporter.cpp b/Engine/Plugins/Importers/USDImporter/Source/USDStageImporter/Private/USDStageImporter.cpp
--- a/Engine/Plugins/Importers/USDImporter/Source/USDStageImporter/Private/USDStageImporter.cpp
+++ b/Engine/Plugins/Importers/USDImporter/Source/USDStageImporter/Private/USDStageImporter.cpp
@@ -854,6 +854,33 @@ namespace UsdStageImporterImpl
 		return MovedAsset;
 	}
 
+	// Our asset could have come from an asset cache, where it could have been reused across different stages. After import, we only
+	// want to keep the metadata fields that came from the actually imported stage though, so here we prune the others
+	void PrunePrimMetadata(UObject* Asset, const UE::FUsdStage& Stage)
+	{
+		IInterface_AssetUserData* UserDataInterface = Cast<IInterface_AssetUserData>(Asset);
+		if (!UserDataInterface)
+		{
+			return;
+		}
+
+		UUsdAssetUserData* UserData = UserDataInterface->GetAssetUserData<UUsdAssetUserData>();
+		if (!UserData)
+		{
+			return;
+		}
+
+		FString ImportedStageIdentifier = Stage.GetRootLayer().GetIdentifier();
+
+		for (TMap<FString, FUsdCombinedPrimMetadata>::TIterator Iter = UserData->StageIdentifierToMetadata.CreateIterator(); Iter; ++Iter)
+		{
+			if (Iter->Key != ImportedStageIdentifier)
+			{
+				Iter.RemoveCurrent();
+			}
+		}
+	}
+
 	// Move imported assets from transient folder to their final package, updating AssetCache to point to the moved assets
 	void PublishAssets(
 		FUsdStageImportContext& ImportContext,
@@ -1008,7 +1035,7 @@ namespace UsdStageImporterImpl
 				{
 					FString TargetAssetName = GetUserFriendlyName(Asset, UniqueAssetNames);
 					FString DestPackagePath = FPaths::Combine(ImportContext.PackagePath, *AssetTypeFolder, TargetAssetName);
-					PublishAsset(
+					UObject* PublishedAsset = PublishAsset(
 						ImportContext,
 						Asset,
 						DestPackagePath,
@@ -1016,6 +1043,11 @@ namespace UsdStageImporterImpl
 						SoftObjectsToRemap,
 						OutAssetsToFinalize
 					);
+
+					if (PublishedAsset)
+					{
+						PrunePrimMetadata(PublishedAsset, ImportContext.Stage);
+					}
 				}
 			}
 		};
@@ -1929,6 +1961,7 @@ void UUsdStageImporter::ImportFromFile(FUsdStageImportContext& ImportContext)
 	TranslationContext->MaterialPurpose = ImportContext.ImportOptions->MaterialPurpose;
 	TranslationContext->RootMotionHandling = ImportContext.ImportOptions->RootMotionHandling;
 	TranslationContext->SubdivisionLevel = ImportContext.ImportOptions->SubdivisionLevel;
+	TranslationContext->MetadataOptions = ImportContext.ImportOptions->MetadataOptions;
 	TranslationContext->ParentComponent = ImportContext.SceneActor ? ImportContext.SceneActor->GetRootComponent() : nullptr;
 	TranslationContext->KindsToCollapse = ( EUsdDefaultKind ) ImportContext.ImportOptions->KindsToCollapse;
 	TranslationContext->bMergeIdenticalMaterialSlots = ImportContext.ImportOptions->bMergeIdenticalMaterialSlots;
@@ -2085,6 +2118,7 @@ bool UUsdStageImporter::ReimportSingleAsset(
 	TranslationContext->MaterialPurpose = ImportContext.ImportOptions->MaterialPurpose;
 	TranslationContext->RootMotionHandling = ImportContext.ImportOptions->RootMotionHandling;
 	TranslationContext->SubdivisionLevel = ImportContext.ImportOptions->SubdivisionLevel;
+	TranslationContext->MetadataOptions = ImportContext.ImportOptions->MetadataOptions;
 	TranslationContext->KindsToCollapse = ( EUsdDefaultKind ) ImportContext.ImportOptions->KindsToCollapse;
 	TranslationContext->bMergeIdenticalMaterialSlots = ImportContext.ImportOptions->bMergeIdenticalMaterialSlots;
 	TranslationContext->bAllowInterpretingLODs = ImportContext.ImportOptions->bInterpretLODs;
diff --git a/Engine/Plugins/Importers/USDImporter/Source/USDStageImporter/Public/USDStageImportOptions.h b/Engine/Plugins/Importers/USDImporter/Source/USDStageImporter/Public/USDStageImportOptions.h
--- a/Engine/Plugins/Importers/USDImporter/Source/USDStageImporter/Public/USDStageImportOptions.h
+++ b/Engine/Plugins/Importers/USDImporter/Source/USDStageImporter/Public/USDStageImportOptions.h
@@ -3,6 +3,7 @@
 #pragma once
 
 #include "UnrealUSDWrapper.h"
+#include "USDMetadataImportOptions.h"
 #include "USDStageOptions.h"
 
 #include "CoreMinimal.h"
@@ -109,6 +110,10 @@ public:
 	UPROPERTY(BlueprintReadWrite, config, EditAnywhere, Category = "USD options")
 	int32 SubdivisionLevel;
 
+	/* Describes if/how we should collect metadata from USD prims onto the assets and components we generate when importing */
+	UPROPERTY(BlueprintReadWrite, config, EditAnywhere, Category = "USD options")
+	FUsdMetadataImportOptions MetadataOptions;
+
 	/** Whether to use the specified StageOptions instead of the stage's own settings */
 	UPROPERTY( BlueprintReadWrite, config, EditAnywhere, Category = "USD options" )
 	bool bOverrideStageOptions;
diff --git a/Engine/Plugins/Importers/USDImporter/Source/USDUtilities/Private/USDConversionUtils.cpp b/Engine/Plugins/Importers/USDImporter/Source/USDUtilities/Private/USDConversionUtils.cpp
--- a/Engine/Plugins/Importers/USDImporter/Source/USDUtilities/Private/USDConversionUtils.cpp
+++ b/Engine/Plugins/Importers/USDImporter/Source/USDUtilities/Private/USDConversionUtils.cpp
@@ -3,8 +3,9 @@
 #include "USDConversionUtils.h"
 
 #include "USDAssetImportData.h"
-#include "USDDrawModeComponent.h"
+#include "USDAssetUserData.h"
 #include "USDClassesModule.h"
+#include "USDDrawModeComponent.h"
 #include "USDDuplicateType.h"
 #include "USDErrorUtils.h"
 #include "USDGeomMeshConversion.h"
@@ -1677,6 +1678,101 @@ void UsdUtils::SetAssetImportData(UObject* Asset, UAssetImportData* ImportData)
 #endif // WITH_EDITOR
 }
 
+UUsdAssetUserData* UsdUtils::GetAssetUserData(const UObject* Object, TSubclassOf<UUsdAssetUserData> Class)
+{
+	if (!Object)
+	{
+		return nullptr;
+	}
+
+	if (!Class)
+	{
+		Class = UUsdAssetUserData::StaticClass();
+	}
+
+	const IInterface_AssetUserData* AssetUserDataInterface = Cast<const IInterface_AssetUserData>(Object);
+	if (!AssetUserDataInterface)
+	{
+		UE_LOG(
+			LogUsd,
+			Warning,
+			TEXT("Tried getting AssetUserData from object '%s', but the class '%s' doesn't implement the AssetUserData interface!"),
+			*Object->GetPathName(),
+			*Object->GetClass()->GetName()
+		);
+		return nullptr;
+	}
+
+	// Const cast because there is no const access of asset user data on the interface
+	return Cast<UUsdAssetUserData>(const_cast<IInterface_AssetUserData*>(AssetUserDataInterface)->GetAssetUserDataOfClass(Class));
+}
+
+UUsdAssetUserData* UsdUtils::GetOrCreateAssetUserData(UObject* Object, TSubclassOf<UUsdAssetUserData> Class)
+{
+	if (!Object)
+	{
+		return nullptr;
+	}
+
+	if (!Class)
+	{
+		Class = UUsdAssetUserData::StaticClass();
+	}
+
+	IInterface_AssetUserData* AssetUserDataInterface = Cast<IInterface_AssetUserData>(Object);
+	if (!AssetUserDataInterface)
+	{
+		UE_LOG(
+			LogUsd,
+			Warning,
+			TEXT("Tried adding AssetUserData to object '%s', but it doesn't implement the AssetUserData interface!"),
+			*Object->GetPathName()
+		);
+		return nullptr;
+	}
+
+	UUsdAssetUserData* AssetUserData = Cast<UUsdAssetUserData>(AssetUserDataInterface->GetAssetUserDataOfClass(Class));
+	if (!AssetUserData)
+	{
+		// For now we're expecting objects to only have one instance of UUsdAssetUserData
+		ensure(!AssetUserDataInterface->HasAssetUserDataOfClass(UUsdAssetUserData::StaticClass()));
+
+		AssetUserData = NewObject<UUsdAssetUserData>(Object, Class, TEXT("UsdAssetUserData"));
+		AssetUserDataInterface->AddAssetUserData(AssetUserData);
+	}
+
+	return AssetUserData;
+}
+
+bool UsdUtils::SetAssetUserData(UObject* Object, UUsdAssetUserData* AssetUserData)
+{
+	if (!Object)
+	{
+		return false;
+	}
+
+	IInterface_AssetUserData* AssetUserDataInterface = Cast<IInterface_AssetUserData>(Object);
+	if (!AssetUserDataInterface)
+	{
+		UE_LOG(
+			LogUsd,
+			Warning,
+			TEXT("Tried adding AssetUserData to object '%s', but it doesn't implement the AssetUserData interface!"),
+			*Object->GetPathName()
+		);
+		return false;
+	}
+
+	while (AssetUserDataInterface->HasAssetUserDataOfClass(UUsdAssetUserData::StaticClass()))
+	{
+		UE_LOG(LogUsd, Log, TEXT("Removing old AssetUserData from object '%s' before adding a new one"), *Object->GetPathName());
+		AssetUserDataInterface->RemoveUserDataOfClass(UUsdAssetUserData::StaticClass());
+	}
+
+	AssetUserDataInterface->AddAssetUserData(AssetUserData);
+	return true;
+}
+
 namespace UE::UsdConversionUtils::Private
 {
 #if USE_USD_SDK
@@ -3144,6 +3240,52 @@ FUsdUnrealAssetInfo UsdUtils::GetPrimAssetInfo( const UE::FUsdPrim& Prim )
 	return Result;
 }
 
+#if USE_USD_SDK
+bool UsdUtils::ClearNonEssentialPrimMetadata(const pxr::UsdPrim& Prim)
+{
+	FScopedUsdAllocs Allocs;
+
+	pxr::SdfChangeBlock ChangeBlock;
+
+	// Note: This only returns top-level fields, and won't have a separate entry for values inside VtDictionaries
+	// or anything like that. This means this likely won't be that expensive, and we don't have to care about order
+	std::map<pxr::TfToken, pxr::VtValue, pxr::TfDictionaryLessThan> MetadataMap = Prim.GetAllAuthoredMetadata();
+
+	for (std::map<pxr::TfToken, pxr::VtValue, pxr::TfDictionaryLessThan>::const_iterator MetadataIter = MetadataMap.begin();
+		 MetadataIter != MetadataMap.end();
+		 ++MetadataIter)
+	{
+		const pxr::TfToken& FieldName = MetadataIter->first;
+
+		// We consider those "essential metadata", as removing them will mess with the prim definition
+		const static std::unordered_set<pxr::TfToken, pxr::TfHash> FieldsToSkip = {
+			pxr::SdfFieldKeys->Specifier,
+			pxr::SdfFieldKeys->TypeName
+		};
+		if (FieldsToSkip.count(FieldName) > 0)
+		{
+			continue;
+		}
+
+		const bool bSuccess = Prim.ClearMetadata(FieldName);
+
+		if (!bSuccess)
+		{
+			UE_LOG(
+				LogUsd,
+				Warning,
+				TEXT("Failed to clear metadata field '%s' from prim '%s'"),
+				*UsdToUnreal::ConvertToken(FieldName),
+				*UsdToUnreal::ConvertPath(Prim.GetPrimPath())
+			);
+			return false;
+		}
+	}
+
+	return true;
+}
+#endif // USE_USD_SDK
+
 void UsdUtils::CollectSchemaAnalytics(const UE::FUsdStage& Stage, const FString& EventName)
 {
 #if USE_USD_SDK
diff --git a/Engine/Plugins/Importers/USDImporter/Source/USDUtilities/Private/USDGeomMeshConversion.cpp b/Engine/Plugins/Importers/USDImporter/Source/USDUtilities/Private/USDGeomMeshConversion.cpp
--- a/Engine/Plugins/Importers/USDImporter/Source/USDUtilities/Private/USDGeomMeshConversion.cpp
+++ b/Engine/Plugins/Importers/USDImporter/Source/USDUtilities/Private/USDGeomMeshConversion.cpp
@@ -399,7 +399,7 @@ namespace UE::UsdGeomMeshConversion::Private
 					for ( int32 VertexIndex = 0; VertexIndex < VertexCount; ++VertexIndex )
 					{
 						FVector VertexPosition = (FVector)LODRenderMesh.VertexBuffers.PositionVertexBuffer.VertexPosition( VertexIndex );
-						PointsArray.push_back( UnrealToUsd::ConvertVector( StageInfo, VertexPosition ) );
+						PointsArray.push_back( UnrealToUsd::ConvertVectorFloat( StageInfo, VertexPosition ) );
 					}
 
 					Points.Set( PointsArray, TimeCode );
@@ -426,7 +426,7 @@ namespace UE::UsdGeomMeshConversion::Private
 					for ( int32 VertexIndex = 0; VertexIndex < VertexCount; ++VertexIndex )
 					{
 						FVector VertexNormal = (FVector4)LODRenderMesh.VertexBuffers.StaticMeshVertexBuffer.VertexTangentZ( VertexIndex );
-						Normals.push_back( UnrealToUsd::ConvertVector( StageInfo, VertexNormal ) );
+						Normals.push_back( UnrealToUsd::ConvertVectorFloat( StageInfo, VertexNormal ) );
 					}
 
 					NormalsAttribute.Set( Normals, TimeCode );
@@ -452,7 +452,7 @@ namespace UE::UsdGeomMeshConversion::Private
 							FVector2D TexCoord = FVector2D(LODRenderMesh.VertexBuffers.StaticMeshVertexBuffer.GetVertexUV( VertexIndex, TexCoordSourceIndex ));
 							TexCoord[ 1 ] = 1.f - TexCoord[ 1 ];
 
-							UVs.push_back( UnrealToUsd::ConvertVector( TexCoord ) );
+							UVs.push_back( UnrealToUsd::ConvertVectorFloat( TexCoord ) );
 						}
 
 						PrimvarST.Set( UVs, TimeCode );
@@ -658,7 +658,7 @@ namespace UE::UsdGeomMeshConversion::Private
 				for ( const FVertexID VertexID : MeshDescription.Vertices().GetElementIDs() )
 				{
 					FVector UEPosition = AdditionalTransform.TransformPosition( (FVector)VertexPositions[ VertexID ] );
-					PointsArray.push_back( UnrealToUsd::ConvertVector( StageInfo, UEPosition ) );
+					PointsArray.push_back( UnrealToUsd::ConvertVectorFloat( StageInfo, UEPosition ) );
 				}
 
 				Points.Set( PointsArray, TimeCode );
@@ -675,7 +675,7 @@ namespace UE::UsdGeomMeshConversion::Private
 				for ( const FVertexInstanceID InstanceID : MeshDescription.VertexInstances().GetElementIDs() )
 				{
 					FVector UENormal = (FVector)VertexInstanceNormals[ InstanceID ].GetSafeNormal();
-					Normals.push_back( UnrealToUsd::ConvertVector( StageInfo, UENormal ) );
+					Normals.push_back( UnrealToUsd::ConvertVectorFloat( StageInfo, UENormal ) );
 				}
 
 				NormalsAttribute.Set( Normals, TimeCode );
@@ -700,7 +700,7 @@ namespace UE::UsdGeomMeshConversion::Private
 					{
 						FVector2D UV = FVector2D(VertexInstanceUVs.Get( InstanceID, UVIndex ));
 						UV[ 1 ] = 1.f - UV[ 1 ];
-						UVs.push_back( UnrealToUsd::ConvertVector( UV ) );
+						UVs.push_back( UnrealToUsd::ConvertVectorFloat( UV ) );
 					}
 
 					PrimvarST.Set( UVs, TimeCode );
@@ -3917,7 +3917,7 @@ namespace UE::UsdGeometryCacheConversion::Private
 
 					for (int32 VertexIndex = 0; VertexIndex < VertexCount; ++VertexIndex)
 					{
-						PointsArray.push_back(UnrealToUsd::ConvertVector(StageInfo, (FVector)MeshData.Positions[VertexIndex]));
+						PointsArray.push_back(UnrealToUsd::ConvertVectorFloat(StageInfo, (FVector)MeshData.Positions[VertexIndex]));
 					}
 
 					Points.Set(PointsArray, TimeCode);
@@ -3945,7 +3945,7 @@ namespace UE::UsdGeometryCacheConversion::Private
 					for (int32 VertexIndex = 0; VertexIndex < VertexCount; ++VertexIndex)
 					{
 						FVector VertexNormal = MeshData.TangentsZ[VertexIndex].ToFVector();
-						Normals.push_back(UnrealToUsd::ConvertVector(StageInfo, VertexNormal));
+						Normals.push_back(UnrealToUsd::ConvertVectorFloat(StageInfo, VertexNormal));
 					}
 
 					NormalsAttribute.Set(Normals, TimeCode);
@@ -3970,7 +3970,7 @@ namespace UE::UsdGeometryCacheConversion::Private
 						FVector2D TexCoord = FVector2D(MeshData.TextureCoordinates[VertexIndex]);
 						TexCoord[1] = 1.f - TexCoord[1];
 
-						UVs.push_back(UnrealToUsd::ConvertVector(TexCoord));
+						UVs.push_back(UnrealToUsd::ConvertVectorFloat(TexCoord));
 					}
 
 					PrimvarST.Set(UVs, TimeCode);
@@ -4018,7 +4018,7 @@ namespace UE::UsdGeometryCacheConversion::Private
 					for (int32 VertexIndex = 0; VertexIndex < VertexCount; ++VertexIndex)
 					{
 						// The motion vectors in the MeshData are stored as unit per frame so convert it back to unit per second
-						Velocities.push_back(UnrealToUsd::ConvertVector(StageInfo, (FVector) -MeshData.MotionVectors[VertexIndex] * ExportContext.FrameRate));
+						Velocities.push_back(UnrealToUsd::ConvertVectorFloat(StageInfo, (FVector) -MeshData.MotionVectors[VertexIndex] * ExportContext.FrameRate));
 					}
 
 					VelocitiesAttribute.Set(Velocities, TimeCode);
