USD: Add support for importing and exporting prim metadata.

diff --git a/Engine/Plugins/Importers/USDImporter/Source/USDUtilities/Private/USDPrimConversion.cpp b/Engine/Plugins/Importers/USDImporter/Source/USDUtilities/Private/USDPrimConversion.cpp
--- a/Engine/Plugins/Importers/USDImporter/Source/USDUtilities/Private/USDPrimConversion.cpp
+++ b/Engine/Plugins/Importers/USDImporter/Source/USDUtilities/Private/USDPrimConversion.cpp
@@ -6,14 +6,15 @@
 #include "USDAssetCache2.h"
 #include "USDAssetUserData.h"
 #include "USDAttributeUtils.h"
-#include "USDDrawModeComponent.h"
 #include "USDConversionUtils.h"
+#include "USDDrawModeComponent.h"
 #include "USDLayerUtils.h"
 #include "USDLightConversion.h"
 #include "USDLog.h"
 #include "USDShadeConversion.h"
 #include "USDSkeletalDataConversion.h"
 #include "USDTypesConversion.h"
+#include "USDValueConversion.h"
 
 #include "UsdWrappers/UsdAttribute.h"
 #include "UsdWrappers/UsdGeomBBoxCache.h"
@@ -67,10 +68,16 @@
 #if USE_USD_SDK
 
 #include "USDIncludesStart.h"
+	#include "pxr/base/tf/stringUtils.h"
+	#include "pxr/base/vt/value.h"
 	#include "pxr/usd/kind/registry.h"
 	#include "pxr/usd/sdf/changeBlock.h"
+	#include "pxr/usd/sdf/path.h"
+	#include "pxr/usd/sdf/types.h"
 	#include "pxr/usd/usd/attribute.h"
 	#include "pxr/usd/usd/prim.h"
+	#include "pxr/usd/usd/primRange.h"
+	#include "pxr/usd/usd/stage.h"
 	#include "pxr/usd/usd/stage.h"
 	#include "pxr/usd/usd/timeCode.h"
 	#include "pxr/usd/usdGeom/camera.h"
@@ -86,6 +93,7 @@
 	#include "pxr/usd/usdGeom/scope.h"
 	#include "pxr/usd/usdGeom/sphere.h"
 	#include "pxr/usd/usdGeom/tokens.h"
+	#include "pxr/usd/usdGeom/tokens.h"
 	#include "pxr/usd/usdGeom/xform.h"
 	#include "pxr/usd/usdGeom/xformable.h"
 	#include "pxr/usd/usdGeom/xformCommonAPI.h"
@@ -106,6 +114,7 @@
 	#include "pxr/usd/usdSkel/root.h"
 	#include "pxr/usd/usdSkel/skeleton.h"
 	#include "pxr/usd/usdSkel/skeletonQuery.h"
+	#include "pxr/usd/usdUtils/stageCache.h"
 #include "USDIncludesEnd.h"
 
 static bool GConsiderAllPrimsHaveAnimatedBounds = false;
@@ -1780,6 +1789,247 @@ bool UsdToUnreal::ConvertBounds(
 	return true;
 }
 
+namespace UE::USDPrimConversionImpl::Private
+{
+	const static std::string UsdNamespaceDelimiter = UnrealToUsd::ConvertString(*UnrealIdentifiers::UsdNamespaceDelimiter).Get();
+
+	bool ShouldSkipField(const FString& FullFieldPath, const TArray<FString>& BlockedPrefixFilters, bool bInvertFilters)
+	{
+		if (bInvertFilters)
+		{
+			// Yes this can be simplified further as this code is just a copy paste of the case below,
+			// but splitting the cases should be quicker to understand
+			for (const FString& Prefix : BlockedPrefixFilters)
+			{
+				if (FullFieldPath.StartsWith(Prefix))
+				{
+					return false;
+				}
+			}
+
+			return true;
+		}
+		else
+		{
+			for (const FString& Prefix : BlockedPrefixFilters)
+			{
+				if (FullFieldPath.StartsWith(Prefix))
+				{
+					return true;
+				}
+			}
+
+			return false;
+		}
+	}
+
+	// Converts the entries within Dictionary into metadata entries within InOutPrimMetadata, using the provided filters
+	// and the additional FieldPathPrefix for the entry keys
+	void ConvertMetadataDictionary(
+		const pxr::VtDictionary& Dictionary,
+		const std::string& FieldPathPrefix,
+		FUsdPrimMetadata& InOutPrimMetadata,
+		const TArray<FString>& BlockedPrefixFilters = {},
+		bool bInvertFilters = false
+	)
+	{
+		FScopedUsdAllocs Allocs;
+
+		// We should only call this for nested dicts, at which point we should already have a field path prefix
+		ensure(!FieldPathPrefix.empty());
+
+		for (pxr::VtDictionary::const_iterator ValueIter = Dictionary.begin(); ValueIter != Dictionary.end(); ++ValueIter)
+		{
+			const std::string& DictFieldName = ValueIter->first;
+			const pxr::VtValue& DictFieldValue = ValueIter->second;
+
+			std::string FieldFullPath = FieldPathPrefix + UsdNamespaceDelimiter + DictFieldName;
+			FString FieldFullString = UsdToUnreal::ConvertString(FieldFullPath);
+
+			if (DictFieldValue.IsHolding<pxr::VtDictionary>())
+			{
+				ConvertMetadataDictionary(
+					DictFieldValue.UncheckedGet<pxr::VtDictionary>(),
+					FieldFullPath,
+					InOutPrimMetadata,
+					BlockedPrefixFilters,
+					bInvertFilters
+				);
+			}
+			else
+			{
+				// Note how we only check the filter when we have our *full* key path. It may seem wasteful to not do an early check in case
+				// the field is a dictionary, but consider this:
+				// 	Field: "customData:int"	AllowFilter: "customData"  		--> Should allow
+				// 	Field: "customData"		AllowFilter: "customData:int"  	--> Should... also allow? If we want to eventually allow the int
+				//                                                              we need to allow its parent dict too
+				//  Field: "abcde"			AllowFilter: "ab"  				--> Should allow
+				//  Field: "ab"				AllowFilter: "abcde"  			--> Should... not allow? It really doesn't start with that prefix...
+				// Our desired behavior changes when the prefix consists of an "incomplete path" to the key we're really interested in...
+				// The simplest way to handle that is probably to simply never early compare the path like in the second example at all, by
+				// only ever checking *full* paths against the filter, which is what we're doing here.
+				if (ShouldSkipField(FieldFullString, BlockedPrefixFilters, bInvertFilters))
+				{
+					continue;
+				}
+
+				FUsdMetadataValue& Metadata = InOutPrimMetadata.Metadata.FindOrAdd(FieldFullString);
+				Metadata.StringifiedValue = UsdUtils::Stringify(DictFieldValue);
+
+				// Prefer the SdfTypeNameToken over Value.GetTypeName() (the former is like "double3[]" and is the
+				// same you type out on the .usda files, while the latter matches the C++ type, like "VtArray<GfVec3d>")
+				if (pxr::SdfValueTypeName TypeName = pxr::SdfGetValueTypeNameForValue(DictFieldValue))
+				{
+					Metadata.TypeName = UsdToUnreal::ConvertToken(TypeName.GetAsToken());
+				}
+				else
+				{
+					Metadata.TypeName = UsdToUnreal::ConvertString(DictFieldValue.GetTypeName());
+				}
+			}
+		}
+	}
+
+	void CollectMetadataForPrim(
+		const pxr::UsdPrim& Prim,
+		FUsdCombinedPrimMetadata& InOutCombinedMetadata,
+		const TArray<FString>& BlockedPrefixFilters,
+		bool bInvertFilters
+	)
+	{
+		if (!Prim)
+		{
+			return;
+		}
+
+		FScopedUsdAllocs Allocs;
+
+		FUsdPrimMetadata* PrimMetadata = nullptr;
+
+		std::map<pxr::TfToken, pxr::VtValue, pxr::TfDictionaryLessThan> MetadataMap = Prim.GetAllAuthoredMetadata();
+		for (std::map<pxr::TfToken, pxr::VtValue, pxr::TfDictionaryLessThan>::const_iterator MetadataIter = MetadataMap.begin();
+			 MetadataIter != MetadataMap.end();
+			 ++MetadataIter)
+		{
+			const pxr::TfToken& FieldName = MetadataIter->first;
+			const pxr::VtValue& FieldValue = MetadataIter->second;
+
+			// There is no real point in keeping track of these as they are defined on every prim and should just match
+			// what the prim's definition is. It's probably a bad idea to author a value that differs from the prim
+			// definition too
+			const static std::unordered_set<pxr::TfToken, pxr::TfHash> FieldsToSkip = {
+				pxr::SdfFieldKeys->Specifier,
+				pxr::SdfFieldKeys->TypeName
+			};
+			if (FieldsToSkip.count(FieldName) > 0)
+			{
+				continue;
+			}
+
+			// We have a valid field we want to collect, so let's on-demand create a PrimMetadata entry.
+			// Creating on-demand prevents us from creating useless structs for prims without metadata.
+			if (!PrimMetadata)
+			{
+				FString PrimPath = UsdToUnreal::ConvertPath(Prim.GetPrimPath());
+				PrimMetadata = &InOutCombinedMetadata.PrimPathToMetadata.FindOrAdd(PrimPath);
+			}
+
+			if (FieldValue.IsHolding<pxr::VtDictionary>())
+			{
+				ConvertMetadataDictionary(
+					FieldValue.UncheckedGet<pxr::VtDictionary>(),
+					FieldName.GetString(),
+					*PrimMetadata,
+					BlockedPrefixFilters,
+					bInvertFilters
+				);
+			}
+			else
+			{
+				// c.f. the comment within ConvertMetadataDictionary
+				FString FieldNameString = UsdToUnreal::ConvertToken(FieldName);
+				if (ShouldSkipField(FieldNameString, BlockedPrefixFilters, bInvertFilters))
+				{
+					continue;
+				}
+
+				FUsdMetadataValue& Metadata = PrimMetadata->Metadata.FindOrAdd(FieldNameString);
+				Metadata.StringifiedValue = UsdUtils::Stringify(FieldValue);
+
+				if (pxr::SdfValueTypeName TypeName = pxr::SdfGetValueTypeNameForValue(FieldValue))
+				{
+					Metadata.TypeName = UsdToUnreal::ConvertToken(TypeName.GetAsToken());
+				}
+				else
+				{
+					Metadata.TypeName = UsdToUnreal::ConvertString(FieldValue.GetTypeName());
+				}
+			}
+		}
+	}
+};
+
+bool UsdToUnreal::ConvertMetadata(
+	const pxr::UsdPrim& Prim,
+	FUsdCombinedPrimMetadata& CombinedMetadata,
+	const TArray<FString>& BlockedPrefixFilters,
+	bool bInvertFilters,
+	bool bCollectFromEntireSubtrees
+)
+{
+	if (!Prim)
+	{
+		return false;
+	}
+
+	UE::USDPrimConversionImpl::Private::CollectMetadataForPrim(
+		Prim,
+		CombinedMetadata,
+		BlockedPrefixFilters,
+		bInvertFilters
+	);
+
+	if (bCollectFromEntireSubtrees)
+	{
+		pxr::UsdPrimRange PrimRange{Prim, pxr::UsdTraverseInstanceProxies()};
+		for (pxr::UsdPrimRange::iterator It = ++PrimRange.begin(); It != PrimRange.end(); ++It)
+		{
+			UE::USDPrimConversionImpl::Private::CollectMetadataForPrim(
+				*It,
+				CombinedMetadata,
+				BlockedPrefixFilters,
+				bInvertFilters
+			);
+		}
+	}
+
+	return true;
+}
+
+bool UsdToUnreal::ConvertMetadata(
+	const pxr::UsdPrim& Prim,
+	UUsdAssetUserData* AssetUserData,
+	const TArray<FString>& BlockedPrefixFilters,
+	bool bInvertFilters,
+	bool bCollectFromEntireSubtrees
+)
+{
+	if (!Prim || !AssetUserData)
+	{
+		return false;
+	}
+
+	FScopedUsdAllocs Allocs;
+
+	FString PrimPath = UsdToUnreal::ConvertPath(Prim.GetPrimPath());
+	pxr::UsdStageRefPtr StagePtr = Prim.GetStage();
+	FString StageIdentifier = UsdToUnreal::ConvertString(StagePtr->GetRootLayer()->GetIdentifier());
+
+	FUsdCombinedPrimMetadata& CombinedMetadata = AssetUserData->StageIdentifierToMetadata.FindOrAdd(StageIdentifier);
+
+	return UsdToUnreal::ConvertMetadata(Prim, CombinedMetadata, BlockedPrefixFilters, bInvertFilters, bCollectFromEntireSubtrees);
+}
+
 bool UnrealToUsd::ConvertCameraComponent( const UCineCameraComponent& CameraComponent, pxr::UsdPrim& Prim, double UsdTimeCode )
 {
 	FScopedUsdAllocs UsdAllocs;
@@ -3709,9 +3959,9 @@ bool UnrealToUsd::CreateSkeletalAnimationBaker( UE::FUsdPrim& SkeletonPrim, UE::
 				FTransform BoneTransform = LocalBoneTransforms[ BoneIndex ];
 				BoneTransform = UsdUtils::ConvertAxes( StageInfo.UpAxis == EUsdUpAxis::ZAxis, BoneTransform );
 
-				Translations[ BoneIndex ] = UnrealToUsd::ConvertVector( BoneTransform.GetTranslation() ) * ( 0.01f / StageInfo.MetersPerUnit );
-				Rotations[ BoneIndex ] = UnrealToUsd::ConvertQuat( BoneTransform.GetRotation() ).GetNormalized();
-				Scales[ BoneIndex ] = pxr::GfVec3h( UnrealToUsd::ConvertVector( BoneTransform.GetScale3D() ) );
+				Translations[ BoneIndex ] = UnrealToUsd::ConvertVectorFloat( BoneTransform.GetTranslation() ) * ( 0.01f / StageInfo.MetersPerUnit );
+				Rotations[ BoneIndex ] = UnrealToUsd::ConvertQuatFloat( BoneTransform.GetRotation() ).GetNormalized();
+				Scales[ BoneIndex ] = UnrealToUsd::ConvertVectorHalf( BoneTransform.GetScale3D() );
 			}
 
 			if ( Translations.size() > 0 )
@@ -4237,8 +4487,8 @@ UnrealToUsd::FPropertyTrackWriter UnrealToUsd::CreatePropertyTrackWriter( const
 				{
 					FScopedUsdAllocs Allocs;
 
-					pxr::GfVec3f UEBoundsMinUsdSpace = UnrealToUsd::ConvertVector(StageInfo, UEMinValue);
-					pxr::GfVec3f UEBoundsMaxUsdSpace = UnrealToUsd::ConvertVector(StageInfo, UEMaxValue);
+					pxr::GfVec3f UEBoundsMinUsdSpace = UnrealToUsd::ConvertVectorFloat(StageInfo, UEMinValue);
+					pxr::GfVec3f UEBoundsMaxUsdSpace = UnrealToUsd::ConvertVectorFloat(StageInfo, UEMaxValue);
 					pxr::GfVec3f UsdMin{
 						FMath::Min(UEBoundsMinUsdSpace[0], UEBoundsMaxUsdSpace[0]),
 						FMath::Min(UEBoundsMinUsdSpace[1], UEBoundsMaxUsdSpace[1]),
@@ -4727,8 +4977,8 @@ bool UnrealToUsd::ConvertBoundsComponent(const UUsdDrawModeComponent& BoundsComp
 
 		const FVector& UEBoundsMin = BoundsComponent.BoundsMin;
 		const FVector& UEBoundsMax = BoundsComponent.BoundsMax;
-		pxr::GfVec3f UEBoundsMinUsdSpace = UnrealToUsd::ConvertVector(StageInfo, UEBoundsMin);
-		pxr::GfVec3f UEBoundsMaxUsdSpace = UnrealToUsd::ConvertVector(StageInfo, UEBoundsMax);
+		pxr::GfVec3f UEBoundsMinUsdSpace = UnrealToUsd::ConvertVectorFloat(StageInfo, UEBoundsMin);
+		pxr::GfVec3f UEBoundsMaxUsdSpace = UnrealToUsd::ConvertVectorFloat(StageInfo, UEBoundsMax);
 		pxr::GfVec3f UsdMin{
 			FMath::Min(UEBoundsMinUsdSpace[0], UEBoundsMaxUsdSpace[0]),
 			FMath::Min(UEBoundsMinUsdSpace[1], UEBoundsMaxUsdSpace[1]),
@@ -4862,4 +5112,289 @@ bool UnrealToUsd::ConvertBoundsComponent(const UUsdDrawModeComponent& BoundsComp
 	return true;
 }
 
+namespace UE::USDPrimConversionImpl::Private
+{
+	FString PrimPathToNamespace(FString PrimPath)
+	{
+		const TCHAR* CharsToReplace = TEXT("/{}[]");
+		while (*CharsToReplace)
+		{
+			PrimPath.ReplaceCharInline(*CharsToReplace, **UnrealIdentifiers::UsdNamespaceDelimiter, ESearchCase::CaseSensitive);
+			++CharsToReplace;
+		}
+
+		// Make sure we don't start with a delimiter
+		while(PrimPath.RemoveFromStart(UnrealIdentifiers::UsdNamespaceDelimiter))
+		{
+		}
+		return PrimPath;
+	}
+
+	bool ConvertMetadataInternal(
+		const FUsdPrimMetadata& PrimMetadata,
+		const pxr::UsdPrim& Prim,
+		const TArray<FString>& BlockedPrefixFilters,
+		bool bInvertFilters,
+		const FString& NamespacePrefix = {}
+	)
+	{
+		using namespace UE::USDPrimConversionImpl::Private;
+
+		if (!Prim || PrimMetadata.Metadata.Num() == 0 || (bInvertFilters && BlockedPrefixFilters.Num() == 0))
+		{
+			return false;
+		}
+
+		FScopedUsdAllocs Allocs;
+
+		bool bSuccess = true;
+		for (const TPair<FString, FUsdMetadataValue>& MetadataPair : PrimMetadata.Metadata)
+		{
+			FString FullKeyPath = MetadataPair.Key;
+			const FUsdMetadataValue& MetadataValue = MetadataPair.Value;
+
+			if (MetadataValue.StringifiedValue.IsEmpty() || MetadataValue.TypeName.IsEmpty())
+			{
+				continue;
+			}
+
+			// It's likely always a bad idea to author these as they are automatically authored by just the prim
+			// definition itself and we'll likely run into trouble if we try writing anything that differs from it
+			const static TSet<FString> FieldsToSkip = {
+				UsdToUnreal::ConvertToken(pxr::SdfFieldKeys->Specifier),
+				UsdToUnreal::ConvertToken(pxr::SdfFieldKeys->TypeName)
+			};
+			if (FieldsToSkip.Contains(FullKeyPath))
+			{
+				continue;
+			}
+
+			// Note that here we always have full key paths, as we store these paths flattened out when we're in UE
+			if (ShouldSkipField(FullKeyPath, BlockedPrefixFilters, bInvertFilters))
+			{
+				continue;
+			}
+
+			const FString* TypeNameToUse = &MetadataValue.TypeName;
+
+			// Add the prim path prefix if we have any
+			// Only add the prefix now as we need to check the original key path against the filters
+			if (!NamespacePrefix.IsEmpty())
+			{
+				const static FString CustomDataPrefix = UsdToUnreal::ConvertToken(pxr::SdfFieldKeys->CustomData);
+
+				// e.g. "customData:fromSourcePrims:Root:MyXform:CollapsedMesh1:customData:myIntMetadataValue"
+				FullKeyPath =
+					CustomDataPrefix +
+					UnrealIdentifiers::UsdNamespaceDelimiter +
+					NamespacePrefix +
+					UnrealIdentifiers::UsdNamespaceDelimiter +
+					FullKeyPath;
+
+				// USD is fine with us authoring apiSchemas directly as top level metadata, but it can't understand the typename
+				// if it's in a nested dictionary. We're never going to be actively using that value as actual apiSchemas
+				// after that point anyway, so we may as well just keep that as a string and let the data make it to USD at
+				// least in some form if we need to.
+				// Note that we're already filtering apiSchemas when reading data from child prims, so this is mostly just for
+				// safety (given that the user can author all this manually) and edge cases (when assets are shared via the
+				// asset cache).
+				const static FString ApiSchemasToken = UsdToUnreal::ConvertToken(pxr::UsdTokens->apiSchemas);
+				const static FString StringTypeName = UsdToUnreal::ConvertToken(pxr::SdfValueTypeNames->String.GetAsToken());
+				if (MetadataPair.Key == ApiSchemasToken)
+				{
+					TypeNameToUse = &StringTypeName;
+				}
+			}
+
+			pxr::VtValue UnstringifiedValue;
+			bSuccess &= UsdUtils::Unstringify(
+				MetadataValue.StringifiedValue,
+				*TypeNameToUse,
+				UnstringifiedValue
+			);
+
+			if (!bSuccess)
+			{
+				UE_LOG(
+					LogUsd,
+					Warning,
+					TEXT("Failed to set metadata '%s' on prim '%s' as the value '%s' could not be parsed from string!"),
+					*FullKeyPath,
+					*UsdToUnreal::ConvertPath(Prim.GetPrimPath()),
+					*MetadataValue.StringifiedValue
+				);
+				break;
+			}
+
+			// If this is a key-value pair inside at least one dictionary we need to split the key into top-level
+			// dictionary name and "the rest"
+			pxr::TfToken TopLevelKeyName;
+			pxr::TfToken KeyPath;
+			int32 FirstColonIndex = FullKeyPath.Find(UnrealIdentifiers::UsdNamespaceDelimiter);
+			if (FirstColonIndex != INDEX_NONE)
+			{
+				// If our path was "first:second:third", this will put "first" on TopLevelKeyName, and "second:third" on KeyPath
+				FString UETopLevelKeyName = FullKeyPath.Left(FirstColonIndex);
+				FString UEKeyPath = FullKeyPath.Right(FullKeyPath.Len() - FirstColonIndex - 1);
+
+				TopLevelKeyName = UnrealToUsd::ConvertToken(*UETopLevelKeyName).Get();
+				KeyPath = UnrealToUsd::ConvertToken(*UEKeyPath).Get();
+			}
+			// If this is a top-level key-value pair we can use the MetadataFullPath directly as the key
+			else
+			{
+				TopLevelKeyName = UnrealToUsd::ConvertToken(*FullKeyPath).Get();
+			}
+
+			const bool bWillOverwrite = Prim.HasMetadataDictKey(TopLevelKeyName, KeyPath);
+			if (bWillOverwrite)
+			{
+				UE_LOG(
+					LogUsd,
+					Log,
+					TEXT("Overwriting metadata field '%s' on prim '%s'"),
+					*FullKeyPath,
+					*UsdToUnreal::ConvertPath(Prim.GetPrimPath())
+				);
+			}
+
+			// This will also nicely create the nested dictionaries that it needs on-demand
+			bSuccess &= Prim.SetMetadataByDictKey(TopLevelKeyName, KeyPath, UnstringifiedValue);
+
+			if (!bSuccess)
+			{
+				break;
+			}
+		}
+		return bSuccess;
+	}
+}
+
+bool UnrealToUsd::ConvertMetadata(
+	const FUsdCombinedPrimMetadata& CombinedPrimMetadata,
+	const pxr::UsdPrim& Prim,
+	const TArray<FString>& BlockedPrefixFilters,
+	bool bInvertFilters
+)
+{
+	using namespace UE::USDPrimConversionImpl::Private;
+
+	if (!Prim || CombinedPrimMetadata.PrimPathToMetadata.Num() == 0 || (bInvertFilters && BlockedPrefixFilters.Num() == 0))
+	{
+		return false;
+	}
+
+	FScopedUsdAllocs Allocs;
+
+	// In order to roundtrip metadata, we should try to output the metadata collected from the "main source prim" back
+	// out to the "Prim" we were provided here. In simple cases like a simple Mesh prim or a simple SkelRoot, this means we
+	// will output to the exported Mesh prim or SkelRoot the *exact* same metadata fields that were on the source prim,
+	// including "apiSchemas" and "kind" and so on, which is great! There are some other edge cases though.
+	//
+	// We use FUsdCombinedPrimMetadata to store metadata from multiple prims that ended up sharing the same generated asset
+	// (e.g. hash collision), but also to store metadata from all prims in the subtree that contributed to the asset, hash
+	// collision or not (LOD Mesh prims, skinned mesh prims for a Skeletal Mesh, collapsed Mesh prims, etc.).
+	//
+	// If we're in the latter case (source prim subtree), we can still try to find the "main source prim" by checking for
+	// a common ancestor to all, outputting the metadata from that common ancestor directly to "Prim", and outputting all
+	// the metadata for other prims in different metadata namespaces. That makes sure all the metadata makes it back out
+	// *somewhere*, but it also ensures that we fully roundtrip the metadata on the "main source prim".
+	//
+	// In the former case though (source prim hash collision generating single asset), there really is no "common ancestor",
+	// so the best we can do is to output *all* prim metadata in different namespaces. This should be a rare edge case,
+	// and the caller/user can always prevent that from happening by just ensuring FUsdCombinedPrimMetadata has a single
+	// PrimPath stored though.
+	FString CommonAncestor;
+	if (CombinedPrimMetadata.PrimPathToMetadata.Num() > 1)
+	{
+		TArray<FString> MetadataPrims;
+		CombinedPrimMetadata.PrimPathToMetadata.GetKeys(MetadataPrims);
+
+		// If we have multiple prims just because we collected metadata from an entire subtree, then our root prim will
+		// be the first one after we sort.
+		// Note that there is still the chance that in the process of collecting metadata from the source prim subtree,
+		// only one or more random child prim(s) had any metadata, while the actual "main source prim" didn't have any.
+		// In that we'll either end up "promoting" that child prim's metadata, or handling that case as if we were in the
+		// hash collision case mentioned above. Those are edge cases of edge cases though, and hopefully shouldn't cause
+		// any trouble either way (all the metadata is still going to be output just fine)
+		MetadataPrims.Sort();
+		const FString& PotentialAncestor = MetadataPrims[0];
+
+		bool bHasCommonAncestor = true;
+		for (int32 Index = 1; Index < MetadataPrims.Num(); ++Index)
+		{
+			const FString& SomeMetadataPrim = MetadataPrims[Index];
+			if (!SomeMetadataPrim.StartsWith(PotentialAncestor))
+			{
+				bHasCommonAncestor = false;
+				break;
+			}
+		}
+
+		if (bHasCommonAncestor)
+		{
+			CommonAncestor = PotentialAncestor;
+		}
+	}
+
+	bool bSuccess = true;
+	for (const TPair<FString, FUsdPrimMetadata>& PrimPathToMetadata : CombinedPrimMetadata.PrimPathToMetadata)
+	{
+		const FString& PrimPath = PrimPathToMetadata.Key;
+		const FUsdPrimMetadata& PrimMetadata = PrimPathToMetadata.Value;
+
+		// If this prim is not the common ancestor we need to output its metadata inside of a nested namespace
+		const bool bIsTopLevelPrim = (CombinedPrimMetadata.PrimPathToMetadata.Num() == 1) || (PrimPath == CommonAncestor);
+		const FString NamespacePrefix = bIsTopLevelPrim
+											? TEXT("")
+											: TEXT("fromSourcePrims") + UnrealIdentifiers::UsdNamespaceDelimiter + PrimPathToNamespace(PrimPath);
+
+		bSuccess &= ConvertMetadataInternal(PrimMetadata, Prim, BlockedPrefixFilters, bInvertFilters, NamespacePrefix);
+
+		if (!bSuccess)
+		{
+			break;
+		}
+	}
+
+	return bSuccess;
+}
+
+bool UnrealToUsd::ConvertMetadata(
+	const FUsdPrimMetadata& PrimMetadata,
+	const pxr::UsdPrim& Prim,
+	const TArray<FString>& BlockedPrefixFilters,
+	bool bInvertFilters
+)
+{
+	return UE::USDPrimConversionImpl::Private::ConvertMetadataInternal(PrimMetadata, Prim, BlockedPrefixFilters, bInvertFilters);
+}
+
+bool UnrealToUsd::ConvertMetadata(
+	const UUsdAssetUserData* AssetUserData,
+	const pxr::UsdPrim& Prim,
+	const TArray<FString>& BlockedPrefixFilters,
+	bool bInvertFilters
+)
+{
+	if (!AssetUserData || !Prim)
+	{
+		return false;
+	}
+
+	// In the general case we'll have a single stage in here, and also a single FUsdPrimMetadata entry inside of it.
+	// Here we coalesce all the metadata entries we have though. The inner ConvertMetadata call will warn
+	// about any overwriting metadata keys
+	for (const TPair<FString, FUsdCombinedPrimMetadata>& StageMetadataPair : AssetUserData->StageIdentifierToMetadata)
+	{
+		const bool bSuccess = ConvertMetadata(StageMetadataPair.Value, Prim, BlockedPrefixFilters, bInvertFilters);
+		if (!bSuccess)
+		{
+			return false;
+		}
+	}
+
+	return true;
+}
+
 #endif // #if USE_USD_SDK
diff --git a/Engine/Plugins/Importers/USDImporter/Source/USDUtilities/Private/USDShadeConversion.cpp b/Engine/Plugins/Importers/USDImporter/Source/USDUtilities/Private/USDShadeConversion.cpp
--- a/Engine/Plugins/Importers/USDImporter/Source/USDUtilities/Private/USDShadeConversion.cpp
+++ b/Engine/Plugins/Importers/USDImporter/Source/USDUtilities/Private/USDShadeConversion.cpp
@@ -1740,7 +1740,7 @@ namespace UE
 								case EMaterialProperty::MP_Tangent:
 								{
 									FVector ConvertedNormal{ ( *Samples )[ 0 ].ReinterpretAsLinear() };
-									ConstantLinearValue = UnrealToUsd::ConvertVector( StageInfo, ConvertedNormal );
+									ConstantLinearValue = UnrealToUsd::ConvertVectorFloat( StageInfo, ConvertedNormal );
 									break;
 								}
 								default:
diff --git a/Engine/Plugins/Importers/USDImporter/Source/USDUtilities/Private/USDSkeletalDataConversion.cpp b/Engine/Plugins/Importers/USDImporter/Source/USDUtilities/Private/USDSkeletalDataConversion.cpp
--- a/Engine/Plugins/Importers/USDImporter/Source/USDUtilities/Private/USDSkeletalDataConversion.cpp
+++ b/Engine/Plugins/Importers/USDImporter/Source/USDUtilities/Private/USDSkeletalDataConversion.cpp
@@ -792,7 +792,7 @@ namespace UnrealToUsdImpl
 
 					for ( int32 VertexIndex = 0; VertexIndex < VertexCount; ++VertexIndex )
 					{
-						PointsArray.push_back( UnrealToUsd::ConvertVector( StageInfo, (FVector)Vertices[ VertexIndex ].Position ) );
+						PointsArray.push_back( UnrealToUsd::ConvertVectorFloat( StageInfo, (FVector)Vertices[ VertexIndex ].Position ) );
 					}
 
 					Points.Set( PointsArray, TimeCode );
@@ -818,7 +818,7 @@ namespace UnrealToUsdImpl
 
 					for ( int32 VertexIndex = 0; VertexIndex < VertexCount; ++VertexIndex )
 					{
-						Normals.push_back( UnrealToUsd::ConvertVector( StageInfo, FVector4(Vertices[ VertexIndex ].TangentZ) ) );
+						Normals.push_back( UnrealToUsd::ConvertVectorFloat( StageInfo, FVector4(Vertices[ VertexIndex ].TangentZ) ) );
 					}
 
 					NormalsAttribute.Set( Normals, TimeCode );
@@ -842,7 +842,7 @@ namespace UnrealToUsdImpl
 							FVector2D TexCoord = FVector2D(Vertices[ VertexIndex ].UVs[ TexCoordSourceIndex ]);
 							TexCoord[ 1 ] = 1.f - TexCoord[ 1 ];
 
-							UVs.push_back( UnrealToUsd::ConvertVector( TexCoord ) );
+							UVs.push_back( UnrealToUsd::ConvertVectorFloat( TexCoord ) );
 						}
 
 						PrimvarST.Set( UVs, TimeCode );
@@ -1130,9 +1130,9 @@ namespace UnrealToUsdImpl
 		{
 			const FMorphTargetDelta& Delta = DeltaArray[ DeltaIndex ];
 
-			Offsets.push_back( UnrealToUsd::ConvertVector( StageInfo, (FVector)Delta.PositionDelta ) );
+			Offsets.push_back( UnrealToUsd::ConvertVectorFloat( StageInfo, (FVector)Delta.PositionDelta ) );
 			PointIndices.push_back( Delta.SourceIdx );
-			Normals.push_back( UnrealToUsd::ConvertVector( StageInfo, (FVector)Delta.TangentZDelta ) );
+			Normals.push_back( UnrealToUsd::ConvertVectorFloat( StageInfo, (FVector)Delta.TangentZDelta ) );
 		}
 
 		BlendShape.CreateOffsetsAttr().Set(Offsets, TimeCode);
@@ -3410,9 +3410,9 @@ bool UnrealToUsd::ConvertAnimSequence( UAnimSequence* AnimSequence, pxr::UsdPrim
 				FTransform BoneTransform = LocalBoneTransforms[ BoneIndex ];
 				BoneTransform = UsdUtils::ConvertAxes( StageInfo.UpAxis == EUsdUpAxis::ZAxis, BoneTransform );
 
-				Translations.push_back( UnrealToUsd::ConvertVector( BoneTransform.GetTranslation() ) );
-				Rotations.push_back( UnrealToUsd::ConvertQuat( BoneTransform.GetRotation() ).GetNormalized() );
-				Scales.push_back( pxr::GfVec3h( UnrealToUsd::ConvertVector( BoneTransform.GetScale3D() ) ) );
+				Translations.push_back( UnrealToUsd::ConvertVectorFloat( BoneTransform.GetTranslation() ) );
+				Rotations.push_back( UnrealToUsd::ConvertQuatFloat( BoneTransform.GetRotation() ).GetNormalized() );
+				Scales.push_back( UnrealToUsd::ConvertVectorHalf( BoneTransform.GetScale3D() ) );
 			}
 
 			TranslationsAttr.Set( Translations, pxr::UsdTimeCode( TimeCode ) );
@@ -3822,9 +3822,9 @@ bool UnrealToUsd::ConvertControlRigSection(
 				UsdTransform = UsdUtils::ConvertTransformToUsdSpace( StageInfo, ChildGlobal.GetRelativeTransform( ParentGlobal ) );
 			}
 
-			Translations[ RefSkeletonBoneIndex ] = UnrealToUsd::ConvertVector( UsdTransform.GetTranslation() );
-			Rotations[ RefSkeletonBoneIndex ] = UnrealToUsd::ConvertQuat( UsdTransform.GetRotation() ).GetNormalized();
-			Scales[ RefSkeletonBoneIndex ] = pxr::GfVec3h( UnrealToUsd::ConvertVector( UsdTransform.GetScale3D() ) );
+			Translations[ RefSkeletonBoneIndex ] = UnrealToUsd::ConvertVectorFloat( UsdTransform.GetTranslation() );
+			Rotations[ RefSkeletonBoneIndex ] = UnrealToUsd::ConvertQuatFloat( UsdTransform.GetRotation() ).GetNormalized();
+			Scales[ RefSkeletonBoneIndex ] = UnrealToUsd::ConvertVectorHalf( UsdTransform.GetScale3D() );
 		}
 
 		TranslationsAttr.Set( Translations, pxr::UsdTimeCode( UsdTimeCode ) );
diff --git a/Engine/Plugins/Importers/USDImporter/Source/USDUtilities/Private/USDTypesConversion.cpp b/Engine/Plugins/Importers/USDImporter/Source/USDUtilities/Private/USDTypesConversion.cpp
--- a/Engine/Plugins/Importers/USDImporter/Source/USDUtilities/Private/USDTypesConversion.cpp
+++ b/Engine/Plugins/Importers/USDImporter/Source/USDUtilities/Private/USDTypesConversion.cpp
@@ -83,13 +83,6 @@ namespace UsdToUnreal
 		return FString( ANSI_TO_TCHAR( InString.c_str() ) );
 	}
 
-	FString ConvertString( std::string&& InString )
-	{
-		TUsdStore< std::string > UsdString( MoveTemp( InString ) ); // Store the temporary so that it gets destroyed with the USD allocator
-
-		return FString( ANSI_TO_TCHAR( UsdString.Get().c_str() ) );
-	}
-
 	FString ConvertString( const char* InString )
 	{
 		return FString( ANSI_TO_TCHAR( InString ) );
@@ -110,13 +103,6 @@ namespace UsdToUnreal
 		return FName( InString.c_str() );
 	}
 
-	FName ConvertName( std::string&& InString )
-	{
-		TUsdStore< std::string > UsdString( MoveTemp( InString ) ); // Store the temporary so that it gets destroyed with the USD allocator
-
-		return FName( UsdString.Get().c_str() );
-	}
-
 	FString ConvertToken( const pxr::TfToken& Token )
 	{
 		return UsdToUnreal::ConvertString( Token.GetString() );
@@ -132,50 +118,52 @@ namespace UsdToUnreal
 		return FLinearColor( InValue[0], InValue[1], InValue[2], InValue[3] );
 	}
 
-	FVector2D ConvertVector( const pxr::GfVec2f& InValue )
+	FVector2D ConvertVector(const pxr::GfVec2h& InValue)
 	{
-		return FVector2D( InValue[0], InValue[1] );
+		return FVector2D(InValue[0], InValue[1]);
 	}
 
-	FVector ConvertVector( const pxr::GfVec3f& InValue )
+	FVector2D ConvertVector(const pxr::GfVec2f& InValue)
 	{
-		return FVector( InValue[0], InValue[1], InValue[2] );
+		return FVector2D(InValue[0], InValue[1]);
 	}
 
-	FVector ConvertVector(const pxr::GfVec3d& InValue)
+	FVector2D ConvertVector(const pxr::GfVec2d& InValue)
 	{
-		return FVector(InValue[0], InValue[1], InValue[2]);
+		return FVector2D(InValue[0], InValue[1]);
 	}
 
-	FVector ConvertVector( const FUsdStageInfo& StageInfo, const pxr::GfVec3f& InValue )
+	FIntPoint ConvertVector(const pxr::GfVec2i& InValue)
 	{
-		FVector Value = ConvertVector( InValue );
+		return FIntPoint(InValue[0], InValue[1]);
+	}
 
-		const float UEMetersPerUnit = 0.01f;
-		if ( !FMath::IsNearlyEqual( StageInfo.MetersPerUnit, UEMetersPerUnit ) )
-		{
-			Value *= StageInfo.MetersPerUnit / UEMetersPerUnit;
-		}
+	FVector ConvertVector(const pxr::GfVec3h& InValue)
+	{
+		return FVector(InValue[0], InValue[1], InValue[2]);
+	}
 
-		const bool bIsZUp = ( StageInfo.UpAxis == EUsdUpAxis::ZAxis );
+	FVector ConvertVector(const pxr::GfVec3f& InValue)
+	{
+		return FVector(InValue[0], InValue[1], InValue[2]);
+	}
 
-		if ( bIsZUp )
-		{
-			Value.Y = -Value.Y;
-		}
-		else
-		{
-			Swap( Value.Y, Value.Z );
-		}
+	FVector ConvertVector(const pxr::GfVec3d& InValue)
+	{
+		return FVector(InValue[0], InValue[1], InValue[2]);
+	}
 
-		return Value;
+	FIntVector ConvertVector(const pxr::GfVec3i& InValue)
+	{
+		return FIntVector(InValue[0], InValue[1], InValue[2]);
 	}
 
-	FVector ConvertVector(const FUsdStageInfo& StageInfo, const pxr::GfVec3d& InValue)
+	template<typename VecType>
+	FVector ConvertVectorInner(const FUsdStageInfo& StageInfo, const VecType& InValue)
 	{
 		FVector Value = ConvertVector(InValue);
 
-		const float UEMetersPerUnit = 0.01f;
+		const double UEMetersPerUnit = 0.01f;
 		if (!FMath::IsNearlyEqual(StageInfo.MetersPerUnit, UEMetersPerUnit))
 		{
 			Value *= StageInfo.MetersPerUnit / UEMetersPerUnit;
@@ -195,6 +183,58 @@ namespace UsdToUnreal
 		return Value;
 	}
 
+	FVector ConvertVector(const FUsdStageInfo& StageInfo, const pxr::GfVec3h& InValue)
+	{
+		return ConvertVectorInner(StageInfo, InValue);
+	}
+
+	FVector ConvertVector(const FUsdStageInfo& StageInfo, const pxr::GfVec3f& InValue)
+	{
+		return ConvertVectorInner(StageInfo, InValue);
+	}
+
+	FVector ConvertVector(const FUsdStageInfo& StageInfo, const pxr::GfVec3d& InValue)
+	{
+		return ConvertVectorInner(StageInfo, InValue);
+	}
+
+	FVector4 ConvertVector(const pxr::GfVec4h& InValue)
+	{
+		return FVector4(InValue[0], InValue[1], InValue[2], InValue[3]);
+	}
+
+	FVector4 ConvertVector(const pxr::GfVec4f& InValue)
+	{
+		return FVector4(InValue[0], InValue[1], InValue[2], InValue[3]);
+	}
+
+	FVector4 ConvertVector(const pxr::GfVec4d& InValue)
+	{
+		return FVector4(InValue[0], InValue[1], InValue[2], InValue[3]);
+	}
+
+	FIntVector4 ConvertVector(const pxr::GfVec4i& InValue)
+	{
+		return FIntVector4(InValue[0], InValue[1], InValue[2], InValue[3]);
+	}
+
+	FMatrix2D ConvertMatrix(const pxr::GfMatrix2d& Matrix)
+	{
+		FMatrix2D Result;
+		Result.Row0 = ConvertVector(Matrix.GetRow(0));
+		Result.Row1 = ConvertVector(Matrix.GetRow(1));
+		return Result;
+	}
+
+	FMatrix3D ConvertMatrix(const pxr::GfMatrix3d& Matrix)
+	{
+		FMatrix3D Result;
+		Result.Row0 = ConvertVector(Matrix.GetRow(0));
+		Result.Row1 = ConvertVector(Matrix.GetRow(1));
+		Result.Row2 = ConvertVector(Matrix.GetRow(2));
+		return Result;
+	}
+
 	FMatrix ConvertMatrix( const pxr::GfMatrix4d& Matrix )
 	{
 		FMatrix UnrealMatrix(
@@ -223,6 +263,24 @@ namespace UsdToUnreal
 		return Transform;
 	}
 
+	FQuat ConvertQuat(const pxr::GfQuath& InValue)
+	{
+		const pxr::GfVec3h& Imaginary = InValue.GetImaginary();
+		return FQuat{Imaginary[0], Imaginary[1], Imaginary[2], InValue.GetReal()};
+	}
+
+	FQuat ConvertQuat(const pxr::GfQuatf& InValue)
+	{
+		const pxr::GfVec3f& Imaginary = InValue.GetImaginary();
+		return FQuat{Imaginary[0], Imaginary[1], Imaginary[2], InValue.GetReal()};
+	}
+
+	FQuat ConvertQuat(const pxr::GfQuatd& InValue)
+	{
+		const pxr::GfVec3d& Imaginary = InValue.GetImaginary();
+		return FQuat{Imaginary[0], Imaginary[1], Imaginary[2], InValue.GetReal()};
+	}
+
 	float ConvertDistance( const FUsdStageInfo& StageInfo, float Value )
 	{
 		const float UEMetersPerUnit = 0.01f;
@@ -267,41 +325,140 @@ namespace UnrealToUsd
 		return ConvertColor( FLinearColor( InValue ) );
 	}
 
-	pxr::GfVec2f ConvertVector( const FVector2D& InValue )
+	// Deprecated
+	pxr::GfVec2f ConvertVector(const FVector2D& InValue)
 	{
-		return pxr::GfVec2f( InValue[0], InValue[1] );
+		return ConvertVectorFloat(InValue);
 	}
 
-	pxr::GfVec3f ConvertVector( const FVector& InValue )
+	// Deprecated
+	pxr::GfVec3f ConvertVector(const FVector& InValue)
 	{
-		return pxr::GfVec3f( InValue[0], InValue[1], InValue[2] );
+		return ConvertVectorFloat(InValue);
 	}
 
-	pxr::GfVec3f ConvertVector( const FUsdStageInfo& StageInfo, const FVector& InValue )
+	// Deprecated
+	pxr::GfVec3f ConvertVector(const FUsdStageInfo& StageInfo, const FVector& InValue)
 	{
-		pxr::GfVec3f Value = ConvertVector( InValue );
+		return ConvertVectorFloat(StageInfo, InValue);
+	}
 
-		const float UEMetersPerUnit = 0.01f;
-		if ( !FMath::IsNearlyEqual( StageInfo.MetersPerUnit, UEMetersPerUnit ) && !FMath::IsNearlyZero( StageInfo.MetersPerUnit ) )
+	pxr::GfVec2h ConvertVectorHalf(const FVector2D& InValue)
+	{
+		return pxr::GfVec2h(InValue[0], InValue[1]);
+	}
+
+	pxr::GfVec2f ConvertVectorFloat(const FVector2D& InValue)
+	{
+		return pxr::GfVec2f(InValue[0], InValue[1]);
+	}
+
+	pxr::GfVec2d ConvertVectorDouble(const FVector2D& InValue)
+	{
+		return pxr::GfVec2d(InValue[0], InValue[1]);
+	}
+
+	pxr::GfVec2i ConvertVectorInt(const FIntPoint& InValue)
+	{
+		return pxr::GfVec2i(InValue[0], InValue[1]);
+	}
+
+	pxr::GfVec3h ConvertVectorHalf(const FVector& InValue)
+	{
+		return pxr::GfVec3h(InValue[0], InValue[1], InValue[2]);
+	}
+
+	pxr::GfVec3f ConvertVectorFloat(const FVector& InValue)
+	{
+		return pxr::GfVec3f(InValue[0], InValue[1], InValue[2]);
+	}
+
+	pxr::GfVec3d ConvertVectorDouble(const FVector& InValue)
+	{
+		return pxr::GfVec3d(InValue[0], InValue[1], InValue[2]);
+	}
+
+	pxr::GfVec3i ConvertVectorInt(const FIntVector& InValue)
+	{
+		return pxr::GfVec3i(InValue[0], InValue[1], InValue[2]);
+	}
+
+	template<typename T>
+	T ConvertVectorInner(const FUsdStageInfo& StageInfo, T Value)
+	{
+		const double UEMetersPerUnit = 0.01f;
+		if (!FMath::IsNearlyEqual(StageInfo.MetersPerUnit, UEMetersPerUnit) && !FMath::IsNearlyZero(StageInfo.MetersPerUnit))
 		{
 			Value *= UEMetersPerUnit / StageInfo.MetersPerUnit;
 		}
 
-		const bool bIsZUp = ( StageInfo.UpAxis == EUsdUpAxis::ZAxis );
+		const bool bIsZUp = (StageInfo.UpAxis == EUsdUpAxis::ZAxis);
 
-		if ( bIsZUp )
+		if (bIsZUp)
 		{
 			Value[1] = -Value[1];
 		}
 		else
 		{
-			Swap( Value[1], Value[2] );
+			Swap(Value[1], Value[2]);
 		}
 
 		return Value;
 	}
 
-	pxr::GfMatrix4d ConvertMatrix( const FMatrix& Matrix )
+	pxr::GfVec3h ConvertVectorHalf(const FUsdStageInfo& StageInfo, const FVector& InValue)
+	{
+		return ConvertVectorInner(StageInfo, ConvertVectorHalf(InValue));
+	}
+
+	pxr::GfVec3f ConvertVectorFloat(const FUsdStageInfo& StageInfo, const FVector& InValue)
+	{
+		return ConvertVectorInner(StageInfo, ConvertVectorFloat(InValue));
+	}
+
+	pxr::GfVec3d ConvertVectorDouble(const FUsdStageInfo& StageInfo, const FVector& InValue)
+	{
+		return ConvertVectorInner(StageInfo, ConvertVectorDouble(InValue));
+	}
+
+	pxr::GfVec4h ConvertVectorHalf(const FVector4& InValue)
+	{
+		return pxr::GfVec4h(InValue[0], InValue[1], InValue[2], InValue[3]);
+	}
+
+	pxr::GfVec4f ConvertVectorFloat(const FVector4& InValue)
+	{
+		return pxr::GfVec4f(InValue[0], InValue[1], InValue[2], InValue[3]);
+	}
+
+	pxr::GfVec4d ConvertVectorDouble(const FVector4& InValue)
+	{
+		return pxr::GfVec4d(InValue[0], InValue[1], InValue[2], InValue[3]);
+	}
+
+	pxr::GfVec4i ConvertVectorInt(const FIntVector4& InValue)
+	{
+		return pxr::GfVec4i(InValue[0], InValue[1], InValue[2], InValue[3]);
+	}
+
+	pxr::GfMatrix2d ConvertMatrix(const FMatrix2D& Matrix)
+	{
+		pxr::GfMatrix2d UsdMatrix;
+		UsdMatrix.SetRow(0, ConvertVectorDouble(Matrix.Row0));
+		UsdMatrix.SetRow(1, ConvertVectorDouble(Matrix.Row1));
+		return UsdMatrix;
+	}
+
+	pxr::GfMatrix3d ConvertMatrix(const FMatrix3D& Matrix)
+	{
+		pxr::GfMatrix3d UsdMatrix;
+		UsdMatrix.SetRow(0, ConvertVectorDouble(Matrix.Row0));
+		UsdMatrix.SetRow(1, ConvertVectorDouble(Matrix.Row1));
+		UsdMatrix.SetRow(2, ConvertVectorDouble(Matrix.Row2));
+		return UsdMatrix;
+	}
+
+	pxr::GfMatrix4d ConvertMatrix(const FMatrix& Matrix)
 	{
 		pxr::GfMatrix4d UsdMatrix(
 			Matrix.M[0][0], Matrix.M[0][1], Matrix.M[0][2], Matrix.M[0][3],
@@ -313,11 +470,33 @@ namespace UnrealToUsd
 		return UsdMatrix;
 	}
 
-	pxr::GfQuatf ConvertQuat( const FQuat& InValue )
+	// Deprecated
+	pxr::GfQuatf ConvertQuat(const FQuat& InValue)
 	{
-		pxr::GfQuatf UsdQuat( InValue.W, InValue.X, InValue.Y, InValue.Z );
+		return ConvertQuatFloat(InValue);
+	}
+
+	pxr::GfQuath ConvertQuatHalf(const FQuat& InValue)
+	{
+		return pxr::GfQuath{
+			static_cast<pxr::GfHalf>(InValue.W),
+			static_cast<pxr::GfHalf>(InValue.X),
+			static_cast<pxr::GfHalf>(InValue.Y),
+			static_cast<pxr::GfHalf>(InValue.Z)};
+	}
 
-		return UsdQuat;
+	pxr::GfQuatf ConvertQuatFloat(const FQuat& InValue)
+	{
+		return pxr::GfQuatf{
+			static_cast<float>(InValue.W),
+			static_cast<float>(InValue.X),
+			static_cast<float>(InValue.Y),
+			static_cast<float>(InValue.Z)};
+	}
+
+	pxr::GfQuatd ConvertQuatDouble(const FQuat& InValue)
+	{
+		return pxr::GfQuatd{InValue.W, InValue.X, InValue.Y, InValue.Z};
 	}
 
 	pxr::GfMatrix4d ConvertTransform( const FUsdStageInfo& StageInfo, const FTransform& Transform )
