USD: Add support for importing and exporting prim metadata.

diff --git a/Engine/Plugins/Importers/USDImporter/Source/USDUtilities/Public/USDTypesConversion.h b/Engine/Plugins/Importers/USDImporter/Source/USDUtilities/Public/USDTypesConversion.h
--- a/Engine/Plugins/Importers/USDImporter/Source/USDUtilities/Public/USDTypesConversion.h
+++ b/Engine/Plugins/Importers/USDImporter/Source/USDUtilities/Public/USDTypesConversion.h
@@ -14,12 +14,24 @@
 #include "USDIncludesEnd.h"
 
 PXR_NAMESPACE_OPEN_SCOPE
+	class GfMatrix2d;
+	class GfMatrix3d;
 	class GfMatrix4d;
+	class GfQuatd;
 	class GfQuatf;
+	class GfQuath;
+	class GfVec2d;
 	class GfVec2f;
+	class GfVec2h;
+	class GfVec2i;
+	class GfVec3d;
 	class GfVec3f;
+	class GfVec3h;
+	class GfVec3i;
+	class GfVec4d;
 	class GfVec4f;
-	class GfVec3d;
+	class GfVec4h;
+	class GfVec4i;
 	class SdfPath;
 	class TfToken;
 
@@ -40,14 +52,12 @@ struct USDUTILITIES_API FUsdStageInfo
 namespace UsdToUnreal
 {
 	USDUTILITIES_API FString ConvertString( const std::string& InString );
-	USDUTILITIES_API FString ConvertString( std::string&& InString );
 	USDUTILITIES_API FString ConvertString( const char* InString );
 
 	USDUTILITIES_API FString ConvertPath( const pxr::SdfPath& Path );
 
 	USDUTILITIES_API FName ConvertName( const char* InString );
 	USDUTILITIES_API FName ConvertName( const std::string& InString );
-	USDUTILITIES_API FName ConvertName( std::string&& InString );
 
 	USDUTILITIES_API FString ConvertToken( const pxr::TfToken& Token );
 
@@ -55,15 +65,34 @@ namespace UsdToUnreal
 	USDUTILITIES_API FLinearColor ConvertColor( const pxr::GfVec3f& InValue );
 	USDUTILITIES_API FLinearColor ConvertColor( const pxr::GfVec4f& InValue );
 
-	USDUTILITIES_API FVector2D ConvertVector( const pxr::GfVec2f& InValue );
-	USDUTILITIES_API FVector ConvertVector( const pxr::GfVec3f& InValue );
+	USDUTILITIES_API FVector2D ConvertVector(const pxr::GfVec2h& InValue);
+	USDUTILITIES_API FVector2D ConvertVector(const pxr::GfVec2f& InValue);
+	USDUTILITIES_API FVector2D ConvertVector(const pxr::GfVec2d& InValue);
+	USDUTILITIES_API FIntPoint ConvertVector(const pxr::GfVec2i& InValue);
+
+	USDUTILITIES_API FVector ConvertVector(const pxr::GfVec3h& InValue);
+	USDUTILITIES_API FVector ConvertVector(const pxr::GfVec3f& InValue);
 	USDUTILITIES_API FVector ConvertVector(const pxr::GfVec3d& InValue);
-	USDUTILITIES_API FVector ConvertVector( const FUsdStageInfo& StageInfo, const pxr::GfVec3f& InValue );
+	USDUTILITIES_API FIntVector ConvertVector(const pxr::GfVec3i& InValue);
+	USDUTILITIES_API FVector ConvertVector(const FUsdStageInfo& StageInfo, const pxr::GfVec3h& InValue);
+	USDUTILITIES_API FVector ConvertVector(const FUsdStageInfo& StageInfo, const pxr::GfVec3f& InValue);
 	USDUTILITIES_API FVector ConvertVector(const FUsdStageInfo& StageInfo, const pxr::GfVec3d& InValue);
 
+	USDUTILITIES_API FVector4 ConvertVector(const pxr::GfVec4h& InValue);
+	USDUTILITIES_API FVector4 ConvertVector(const pxr::GfVec4f& InValue);
+	USDUTILITIES_API FVector4 ConvertVector(const pxr::GfVec4d& InValue);
+	USDUTILITIES_API FIntVector4 ConvertVector(const pxr::GfVec4i& InValue);
+
+	USDUTILITIES_API FMatrix2D ConvertMatrix(const pxr::GfMatrix2d& Matrix);
+	USDUTILITIES_API FMatrix3D ConvertMatrix(const pxr::GfMatrix3d& Matrix);
 	USDUTILITIES_API FMatrix ConvertMatrix( const pxr::GfMatrix4d& Matrix );
+
 	USDUTILITIES_API FTransform ConvertMatrix( const FUsdStageInfo& StageInfo, const pxr::GfMatrix4d& InMatrix );
 
+	USDUTILITIES_API FQuat ConvertQuat(const pxr::GfQuath& InValue);
+	USDUTILITIES_API FQuat ConvertQuat(const pxr::GfQuatf& InValue);
+	USDUTILITIES_API FQuat ConvertQuat(const pxr::GfQuatd& InValue);
+
 	/** Returns a distance in "UE units" (i.e. cm) */
 	USDUTILITIES_API float ConvertDistance( const FUsdStageInfo& StageInfo, float InValue );
 }
@@ -84,13 +113,41 @@ namespace UnrealToUsd
 	/** Assumes the input color is in sRGB space. Returns a color in linear space */
 	USDUTILITIES_API pxr::GfVec4f ConvertColor( const FColor& InValue );
 
-	USDUTILITIES_API pxr::GfVec2f ConvertVector( const FVector2D& InValue );
-	USDUTILITIES_API pxr::GfVec3f ConvertVector( const FVector& InValue );
-	USDUTILITIES_API pxr::GfVec3f ConvertVector( const FUsdStageInfo& StageInfo, const FVector& InValue );
-
-	USDUTILITIES_API pxr::GfMatrix4d ConvertMatrix( const FMatrix& Matrix );
-
-	USDUTILITIES_API pxr::GfQuatf ConvertQuat( const FQuat& InValue );
+	UE_DEPRECATED(5.4, "Please use either ConvertVectorHalf, ConvertVectorFloat or ConvertVectorDouble.")
+	USDUTILITIES_API pxr::GfVec2f ConvertVector(const FVector2D& InValue);
+	UE_DEPRECATED(5.4, "Please use either ConvertVectorHalf, ConvertVectorFloat or ConvertVectorDouble.")
+	USDUTILITIES_API pxr::GfVec3f ConvertVector(const FVector& InValue);
+	UE_DEPRECATED(5.4, "Please use either ConvertVectorHalf, ConvertVectorFloat or ConvertVectorDouble.")
+	USDUTILITIES_API pxr::GfVec3f ConvertVector(const FUsdStageInfo& StageInfo, const FVector& InValue);
+
+	USDUTILITIES_API pxr::GfVec2h ConvertVectorHalf(const FVector2D& InValue);
+	USDUTILITIES_API pxr::GfVec2f ConvertVectorFloat(const FVector2D& InValue);
+	USDUTILITIES_API pxr::GfVec2d ConvertVectorDouble(const FVector2D& InValue);
+	USDUTILITIES_API pxr::GfVec2i ConvertVectorInt(const FIntPoint& InValue);
+
+	USDUTILITIES_API pxr::GfVec3h ConvertVectorHalf(const FVector& InValue);
+	USDUTILITIES_API pxr::GfVec3f ConvertVectorFloat(const FVector& InValue);
+	USDUTILITIES_API pxr::GfVec3d ConvertVectorDouble(const FVector& InValue);
+	USDUTILITIES_API pxr::GfVec3i ConvertVectorInt(const FIntVector& InValue);
+	USDUTILITIES_API pxr::GfVec3h ConvertVectorHalf(const FUsdStageInfo& StageInfo, const FVector& InValue);
+	USDUTILITIES_API pxr::GfVec3f ConvertVectorFloat(const FUsdStageInfo& StageInfo, const FVector& InValue);
+	USDUTILITIES_API pxr::GfVec3d ConvertVectorDouble(const FUsdStageInfo& StageInfo, const FVector& InValue);
+
+	USDUTILITIES_API pxr::GfVec4h ConvertVectorHalf(const FVector4& InValue);
+	USDUTILITIES_API pxr::GfVec4f ConvertVectorFloat(const FVector4& InValue);
+	USDUTILITIES_API pxr::GfVec4d ConvertVectorDouble(const FVector4& InValue);
+	USDUTILITIES_API pxr::GfVec4i ConvertVectorInt(const FIntVector4& InValue);
+
+	USDUTILITIES_API pxr::GfMatrix2d ConvertMatrix(const FMatrix2D& Matrix);
+	USDUTILITIES_API pxr::GfMatrix3d ConvertMatrix(const FMatrix3D& Matrix);
+	USDUTILITIES_API pxr::GfMatrix4d ConvertMatrix(const FMatrix& Matrix);
+
+	UE_DEPRECATED(5.4, "Please use either ConvertQuatHalf, ConvertQuatFloat or ConvertQuatDouble.")
+	USDUTILITIES_API pxr::GfQuatf ConvertQuat(const FQuat& InValue);
+
+	USDUTILITIES_API pxr::GfQuath ConvertQuatHalf(const FQuat& InValue);
+	USDUTILITIES_API pxr::GfQuatf ConvertQuatFloat(const FQuat& InValue);
+	USDUTILITIES_API pxr::GfQuatd ConvertQuatDouble(const FQuat& InValue);
 
 	USDUTILITIES_API pxr::GfMatrix4d ConvertTransform( const FUsdStageInfo& StageInfo, const FTransform& Transform );
 
diff --git a/Engine/Plugins/Importers/USDImporter/Source/USDUtilities/Public/USDValueConversion.h b/Engine/Plugins/Importers/USDImporter/Source/USDUtilities/Public/USDValueConversion.h
--- a/Engine/Plugins/Importers/USDImporter/Source/USDUtilities/Public/USDValueConversion.h
+++ b/Engine/Plugins/Importers/USDImporter/Source/USDUtilities/Public/USDValueConversion.h
@@ -10,6 +10,23 @@
 #include "UsdWrappers/ForwardDeclarations.h"
 #include "UsdWrappers/UsdAttribute.h"
 
+#if USE_USD_SDK
+
+#include "USDIncludesStart.h"
+	#include "pxr/pxr.h"
+#include "USDIncludesEnd.h"
+
+PXR_NAMESPACE_OPEN_SCOPE
+	class SdfAssetPath;
+	class SdfTimeCode;
+	class TfToken;
+	class VtValue;
+PXR_NAMESPACE_CLOSE_SCOPE
+
+#endif // #if USE_USD_SDK
+
+struct FMatrix2D;
+struct FMatrix3D;
 namespace UE
 {
 	class FVtValue;
@@ -96,8 +113,163 @@ namespace UnrealToUsd
 
 namespace UsdUtils
 {
-	/** Uses USD to stringify the underlying pxr::VtValue */
-	USDUTILITIES_API FString Stringify( const UE::FVtValue& Value );
+	/**
+	 * Mostly uses USD to stringify the underlying value of the provided pxr::VtValue, but will
+	 * use some custom implementation depending on the type. For string-like types in particular we will mostly
+	 * use our own implementations, as they preserve quotes and delimiters in a way that we can unstringify later.
+	 */
+	USDUTILITIES_API FString Stringify(const UE::FVtValue& Value);
+#if USE_USD_SDK
+	USDUTILITIES_API FString Stringify(const pxr::VtValue& Value);
+#endif	  // USE_USD_SDK
+
+	// Converts the UE type into a string that can be set as USD metadata with the corresponding USD data type.
+	//
+	// We need these explicit functions here as well as the generic Stringify() because our Python/Blueprint libraries
+	// need to do these operations, but they are on modules that can't manipulate the USD types directly, and we need
+	// to be able to internally e.g. convert a TArray<FString> into a pxr::SdfListOp<pxr::TfToken> to stringify some of these.
+	//
+	// We provide Python/Blueprint libraries that can do these stringify/unstringify functions in the first place because
+	// otherwise a user that tried to manipulate metadata values from Python/Blueprint would be forced to either be
+	// content with just strings and not actual types, or to try unstringifying these types on their own (and it would
+	// not be trivial to unstringify an e.g. matrix3d[] on Blueprint)
+	USDUTILITIES_API FString StringifyAsBool(bool Value);
+	USDUTILITIES_API FString StringifyAsUChar(uint8 Value);
+	USDUTILITIES_API FString StringifyAsInt(int32 Value);
+	USDUTILITIES_API FString StringifyAsUInt(uint32 Value);
+	USDUTILITIES_API FString StringifyAsInt64(int64 Value);
+	USDUTILITIES_API FString StringifyAsUInt64(uint64 Value);
+	USDUTILITIES_API FString StringifyAsHalf(float Value);
+	USDUTILITIES_API FString StringifyAsFloat(float Value);
+	USDUTILITIES_API FString StringifyAsDouble(double Value);
+	USDUTILITIES_API FString StringifyAsTimeCode(double Value);
+	USDUTILITIES_API FString StringifyAsString(const FString& Value);
+	USDUTILITIES_API FString StringifyAsToken(const FString& Value);
+	USDUTILITIES_API FString StringifyAsAssetPath(const FString& Value);
+	USDUTILITIES_API FString StringifyAsMatrix2d(const FMatrix2D& Value);
+	USDUTILITIES_API FString StringifyAsMatrix3d(const FMatrix3D& Value);
+	USDUTILITIES_API FString StringifyAsMatrix4d(const FMatrix& Value);
+	USDUTILITIES_API FString StringifyAsQuatd(const FQuat& Value);
+	USDUTILITIES_API FString StringifyAsQuatf(const FQuat& Value);
+	USDUTILITIES_API FString StringifyAsQuath(const FQuat& Value);
+	USDUTILITIES_API FString StringifyAsDouble2(const FVector2D& Value);
+	USDUTILITIES_API FString StringifyAsFloat2(const FVector2D& Value);
+	USDUTILITIES_API FString StringifyAsHalf2(const FVector2D& Value);
+	USDUTILITIES_API FString StringifyAsInt2(const FIntPoint& Value);
+	USDUTILITIES_API FString StringifyAsDouble3(const FVector& Value);
+	USDUTILITIES_API FString StringifyAsFloat3(const FVector& Value);
+	USDUTILITIES_API FString StringifyAsHalf3(const FVector& Value);
+	USDUTILITIES_API FString StringifyAsInt3(const FIntVector& Value);
+	USDUTILITIES_API FString StringifyAsDouble4(const FVector4& Value);
+	USDUTILITIES_API FString StringifyAsFloat4(const FVector4& Value);
+	USDUTILITIES_API FString StringifyAsHalf4(const FVector4& Value);
+	USDUTILITIES_API FString StringifyAsInt4(const FIntVector4& Value);
+	USDUTILITIES_API FString StringifyAsBoolArray(const TArray<bool>& Value);
+	USDUTILITIES_API FString StringifyAsUCharArray(const TArray<uint8>& Value);
+	USDUTILITIES_API FString StringifyAsIntArray(const TArray<int32>& Value);
+	USDUTILITIES_API FString StringifyAsUIntArray(const TArray<uint32>& Value);
+	USDUTILITIES_API FString StringifyAsInt64Array(const TArray<int64>& Value);
+	USDUTILITIES_API FString StringifyAsUInt64Array(const TArray<uint64>& Value);
+	USDUTILITIES_API FString StringifyAsHalfArray(const TArray<float>& Value);
+	USDUTILITIES_API FString StringifyAsFloatArray(const TArray<float>& Value);
+	USDUTILITIES_API FString StringifyAsDoubleArray(const TArray<double>& Value);
+	USDUTILITIES_API FString StringifyAsTimeCodeArray(const TArray<double>& Value);
+	USDUTILITIES_API FString StringifyAsStringArray(const TArray<FString>& Value);
+	USDUTILITIES_API FString StringifyAsTokenArray(const TArray<FString>& Value);
+	USDUTILITIES_API FString StringifyAsAssetPathArray(const TArray<FString>& Value);
+	USDUTILITIES_API FString StringifyAsListOpTokens(const TArray<FString>& Value);
+	USDUTILITIES_API FString StringifyAsMatrix2dArray(const TArray<FMatrix2D>& Value);
+	USDUTILITIES_API FString StringifyAsMatrix3dArray(const TArray<FMatrix3D>& Value);
+	USDUTILITIES_API FString StringifyAsMatrix4dArray(const TArray<FMatrix>& Value);
+	USDUTILITIES_API FString StringifyAsQuatdArray(const TArray<FQuat>& Value);
+	USDUTILITIES_API FString StringifyAsQuatfArray(const TArray<FQuat>& Value);
+	USDUTILITIES_API FString StringifyAsQuathArray(const TArray<FQuat>& Value);
+	USDUTILITIES_API FString StringifyAsDouble2Array(const TArray<FVector2D>& Value);
+	USDUTILITIES_API FString StringifyAsFloat2Array(const TArray<FVector2D>& Value);
+	USDUTILITIES_API FString StringifyAsHalf2Array(const TArray<FVector2D>& Value);
+	USDUTILITIES_API FString StringifyAsInt2Array(const TArray<FIntPoint>& Value);
+	USDUTILITIES_API FString StringifyAsDouble3Array(const TArray<FVector>& Value);
+	USDUTILITIES_API FString StringifyAsFloat3Array(const TArray<FVector>& Value);
+	USDUTILITIES_API FString StringifyAsHalf3Array(const TArray<FVector>& Value);
+	USDUTILITIES_API FString StringifyAsInt3Array(const TArray<FIntVector>& Value);
+	USDUTILITIES_API FString StringifyAsDouble4Array(const TArray<FVector4>& Value);
+	USDUTILITIES_API FString StringifyAsFloat4Array(const TArray<FVector4>& Value);
+	USDUTILITIES_API FString StringifyAsHalf4Array(const TArray<FVector4>& Value);
+	USDUTILITIES_API FString StringifyAsInt4Array(const TArray<FIntVector4>& Value);
+
+#if USE_USD_SDK
+	/**
+	 * Unstringifies the data in String according to what TypeName it is.
+	 * The unstringified USD value is placed into a VtValue for type erasure.
+	 * Example: String is "(1, 0.5, 0.3)" and TypeName "float3" means Output will contain a pxr::GfVec3f(1, 0.5, 0.3) value
+	 */
+	USDUTILITIES_API bool Unstringify(const FString& String, const FString& TypeName, pxr::VtValue& Output);
+#endif // USE_USD_SDK
+
+	// c.f. comment above StringifyAsBool
+	USDUTILITIES_API bool UnstringifyAsBool(const FString& String, bool& OutValue);
+	USDUTILITIES_API bool UnstringifyAsUChar(const FString& String, uint8& OutValue);
+	USDUTILITIES_API bool UnstringifyAsInt(const FString& String, int32& OutValue);
+	USDUTILITIES_API bool UnstringifyAsUInt(const FString& String, uint32& OutValue);
+	USDUTILITIES_API bool UnstringifyAsInt64(const FString& String, int64& OutValue);
+	USDUTILITIES_API bool UnstringifyAsUInt64(const FString& String, uint64& OutValue);
+	USDUTILITIES_API bool UnstringifyAsHalf(const FString& String, float& OutValue);
+	USDUTILITIES_API bool UnstringifyAsFloat(const FString& String, float& OutValue);
+	USDUTILITIES_API bool UnstringifyAsDouble(const FString& String, double& OutValue);
+	USDUTILITIES_API bool UnstringifyAsTimeCode(const FString& String, double& OutValue);
+	USDUTILITIES_API bool UnstringifyAsString(const FString& String, FString& OutValue);
+	USDUTILITIES_API bool UnstringifyAsToken(const FString& String, FString& OutValue);
+	USDUTILITIES_API bool UnstringifyAsAssetPath(const FString& String, FString& OutValue);
+	USDUTILITIES_API bool UnstringifyAsMatrix2d(const FString& String, FMatrix2D& OutValue);
+	USDUTILITIES_API bool UnstringifyAsMatrix3d(const FString& String, FMatrix3D& OutValue);
+	USDUTILITIES_API bool UnstringifyAsMatrix4d(const FString& String, FMatrix& OutValue);
+	USDUTILITIES_API bool UnstringifyAsQuatd(const FString& String, FQuat& OutValue);
+	USDUTILITIES_API bool UnstringifyAsQuatf(const FString& String, FQuat& OutValue);
+	USDUTILITIES_API bool UnstringifyAsQuath(const FString& String, FQuat& OutValue);
+	USDUTILITIES_API bool UnstringifyAsDouble2(const FString& String, FVector2D& OutValue);
+	USDUTILITIES_API bool UnstringifyAsFloat2(const FString& String, FVector2D& OutValue);
+	USDUTILITIES_API bool UnstringifyAsHalf2(const FString& String, FVector2D& OutValue);
+	USDUTILITIES_API bool UnstringifyAsInt2(const FString& String, FIntPoint& OutValue);
+	USDUTILITIES_API bool UnstringifyAsDouble3(const FString& String, FVector& OutValue);
+	USDUTILITIES_API bool UnstringifyAsFloat3(const FString& String, FVector& OutValue);
+	USDUTILITIES_API bool UnstringifyAsHalf3(const FString& String, FVector& OutValue);
+	USDUTILITIES_API bool UnstringifyAsInt3(const FString& String, FIntVector& OutValue);
+	USDUTILITIES_API bool UnstringifyAsDouble4(const FString& String, FVector4& OutValue);
+	USDUTILITIES_API bool UnstringifyAsFloat4(const FString& String, FVector4& OutValue);
+	USDUTILITIES_API bool UnstringifyAsHalf4(const FString& String, FVector4& OutValue);
+	USDUTILITIES_API bool UnstringifyAsInt4(const FString& String, FIntVector4& OutValue);
+	USDUTILITIES_API bool UnstringifyAsBoolArray(const FString& String, TArray<bool>& OutValue);
+	USDUTILITIES_API bool UnstringifyAsUCharArray(const FString& String, TArray<uint8>& OutValue);
+	USDUTILITIES_API bool UnstringifyAsIntArray(const FString& String, TArray<int32>& OutValue);
+	USDUTILITIES_API bool UnstringifyAsUIntArray(const FString& String, TArray<uint32>& OutValue);
+	USDUTILITIES_API bool UnstringifyAsInt64Array(const FString& String, TArray<int64>& OutValue);
+	USDUTILITIES_API bool UnstringifyAsUInt64Array(const FString& String, TArray<uint64>& OutValue);
+	USDUTILITIES_API bool UnstringifyAsHalfArray(const FString& String, TArray<float>& OutValue);
+	USDUTILITIES_API bool UnstringifyAsFloatArray(const FString& String, TArray<float>& OutValue);
+	USDUTILITIES_API bool UnstringifyAsDoubleArray(const FString& String, TArray<double>& OutValue);
+	USDUTILITIES_API bool UnstringifyAsTimeCodeArray(const FString& String, TArray<double>& OutValue);
+	USDUTILITIES_API bool UnstringifyAsStringArray(const FString& String, TArray<FString>& OutValue);
+	USDUTILITIES_API bool UnstringifyAsTokenArray(const FString& String, TArray<FString>& OutValue);
+	USDUTILITIES_API bool UnstringifyAsAssetPathArray(const FString& String, TArray<FString>& OutValue);
+	USDUTILITIES_API bool UnstringifyAsListOpTokens(const FString& String, TArray<FString>& OutValue);
+	USDUTILITIES_API bool UnstringifyAsMatrix2dArray(const FString& String, TArray<FMatrix2D>& OutValue);
+	USDUTILITIES_API bool UnstringifyAsMatrix3dArray(const FString& String, TArray<FMatrix3D>& OutValue);
+	USDUTILITIES_API bool UnstringifyAsMatrix4dArray(const FString& String, TArray<FMatrix>& OutValue);
+	USDUTILITIES_API bool UnstringifyAsQuatdArray(const FString& String, TArray<FQuat>& OutValue);
+	USDUTILITIES_API bool UnstringifyAsQuatfArray(const FString& String, TArray<FQuat>& OutValue);
+	USDUTILITIES_API bool UnstringifyAsQuathArray(const FString& String, TArray<FQuat>& OutValue);
+	USDUTILITIES_API bool UnstringifyAsDouble2Array(const FString& String, TArray<FVector2D>& OutValue);
+	USDUTILITIES_API bool UnstringifyAsFloat2Array(const FString& String, TArray<FVector2D>& OutValue);
+	USDUTILITIES_API bool UnstringifyAsHalf2Array(const FString& String, TArray<FVector2D>& OutValue);
+	USDUTILITIES_API bool UnstringifyAsInt2Array(const FString& String, TArray<FIntPoint>& OutValue);
+	USDUTILITIES_API bool UnstringifyAsDouble3Array(const FString& String, TArray<FVector>& OutValue);
+	USDUTILITIES_API bool UnstringifyAsFloat3Array(const FString& String, TArray<FVector>& OutValue);
+	USDUTILITIES_API bool UnstringifyAsHalf3Array(const FString& String, TArray<FVector>& OutValue);
+	USDUTILITIES_API bool UnstringifyAsInt3Array(const FString& String, TArray<FIntVector>& OutValue);
+	USDUTILITIES_API bool UnstringifyAsDouble4Array(const FString& String, TArray<FVector4>& OutValue);
+	USDUTILITIES_API bool UnstringifyAsFloat4Array(const FString& String, TArray<FVector4>& OutValue);
+	USDUTILITIES_API bool UnstringifyAsHalf4Array(const FString& String, TArray<FVector4>& OutValue);
+	USDUTILITIES_API bool UnstringifyAsInt4Array(const FString& String, TArray<FIntVector4>& OutValue);
 
 	/**
 	 * Uses USD to quickly fetch the underlying value of the wrapped pxr::VtValue for fundamental data types.
@@ -126,4 +298,3 @@ namespace UsdUtils
 	 */
 	USDUTILITIES_API FString GetImpliedTypeName( const UE::FVtValue& Value );
 }
-
diff --git a/Engine/Plugins/Importers/USDImporter/Source/UnrealUSDWrapper/Private/UnrealUSDWrapper.cpp b/Engine/Plugins/Importers/USDImporter/Source/UnrealUSDWrapper/Private/UnrealUSDWrapper.cpp
--- a/Engine/Plugins/Importers/USDImporter/Source/UnrealUSDWrapper/Private/UnrealUSDWrapper.cpp
+++ b/Engine/Plugins/Importers/USDImporter/Source/UnrealUSDWrapper/Private/UnrealUSDWrapper.cpp
@@ -41,19 +41,20 @@
 #include "pxr/usd/ar/defineResolver.h"
 #include "pxr/usd/kind/registry.h"
 #include "pxr/usd/sdf/fileFormat.h"
+#include "pxr/usd/sdf/path.h"
 #include "pxr/usd/sdf/schema.h"
 #include "pxr/usd/usd/attribute.h"
 #include "pxr/usd/usd/common.h"
 #include "pxr/usd/usd/debugCodes.h"
 #include "pxr/usd/usd/modelAPI.h"
 #include "pxr/usd/usd/references.h"
-#include "pxr/usd/usd/schemaBase.h"
 #include "pxr/usd/usd/relationship.h"
+#include "pxr/usd/usd/schemaBase.h"
 #include "pxr/usd/usd/stage.h"
 #include "pxr/usd/usd/stageCacheContext.h"
-#include "pxr/usd/usd/usdFileFormat.h"
 #include "pxr/usd/usd/usdaFileFormat.h"
 #include "pxr/usd/usd/usdcFileFormat.h"
+#include "pxr/usd/usd/usdFileFormat.h"
 #include "pxr/usd/usd/usdzFileFormat.h"
 #include "pxr/usd/usd/variantSets.h"
 #include "pxr/usd/usdGeom/mesh.h"
@@ -448,6 +449,7 @@ FString UnrealIdentifiers::PrimvarsDisplayOpacity = ANSI_TO_TCHAR(pxr::UsdGeomTo
 FString UnrealIdentifiers::DoubleSided = ANSI_TO_TCHAR(pxr::UsdGeomTokens->doubleSided.GetString().c_str());
 FString UnrealIdentifiers::ModelDrawMode = ANSI_TO_TCHAR(pxr::UsdGeomTokens->modelDrawMode.GetString().c_str());
 FString UnrealIdentifiers::ModelApplyDrawMode = ANSI_TO_TCHAR(pxr::UsdGeomTokens->modelApplyDrawMode.GetString().c_str());
+FString UnrealIdentifiers::UsdNamespaceDelimiter = ANSI_TO_TCHAR(pxr::SdfPathTokens->namespaceDelimiter.GetString().c_str());
 #else
 FString UnrealIdentifiers::MaterialAllPurpose = TEXT( "" );
 FString UnrealIdentifiers::MaterialPreviewPurpose = TEXT( "preview" );
@@ -457,6 +459,7 @@ FString UnrealIdentifiers::PrimvarsDisplayOpacity = TEXT("primvars:displayOpacit
 FString UnrealIdentifiers::DoubleSided = TEXT("doubleSided");
 FString UnrealIdentifiers::ModelDrawMode = TEXT("model:drawMode");
 FString UnrealIdentifiers::ModelApplyDrawMode = TEXT("model:applyDrawMode");
+FString UnrealIdentifiers::UsdNamespaceDelimiter = TEXT(":");
 #endif // USE_USD_SDK
 
 FUsdDelegates::FUsdImportDelegate FUsdDelegates::OnPreUsdImport;
diff --git a/Engine/Plugins/Importers/USDImporter/Source/UnrealUSDWrapper/Public/UnrealUSDWrapper.h b/Engine/Plugins/Importers/USDImporter/Source/UnrealUSDWrapper/Public/UnrealUSDWrapper.h
--- a/Engine/Plugins/Importers/USDImporter/Source/UnrealUSDWrapper/Public/UnrealUSDWrapper.h
+++ b/Engine/Plugins/Importers/USDImporter/Source/UnrealUSDWrapper/Public/UnrealUSDWrapper.h
@@ -168,6 +168,41 @@ enum class EUsdListPosition : uint8
 	BackOfAppendList
 };
 
+/**
+ * Corresponds to pxr::GfMatrix2d. We don't expose any methods though, this is just to facilitate reading/writing
+ * these types from USD.
+ */
+USTRUCT(BlueprintType)
+struct FMatrix2D
+{
+	GENERATED_BODY()
+
+	UPROPERTY(EditAnywhere, BlueprintReadWrite, Category = "USD")
+	FVector2D Row0;
+
+	UPROPERTY(EditAnywhere, BlueprintReadWrite, Category = "USD")
+	FVector2D Row1;
+};
+
+/**
+ * Corresponds to pxr::GfMatrix3d. We don't expose any methods though, this is just to facilitate reading/writing
+ * these types from USD.
+ */
+USTRUCT(BlueprintType)
+struct FMatrix3D
+{
+	GENERATED_BODY()
+
+	UPROPERTY(EditAnywhere, BlueprintReadWrite, Category = "USD")
+	FVector Row0;
+
+	UPROPERTY(EditAnywhere, BlueprintReadWrite, Category = "USD")
+	FVector Row1;
+
+	UPROPERTY(EditAnywhere, BlueprintReadWrite, Category = "USD")
+	FVector Row2;
+};
+
 class IUnrealUSDWrapperModule : public IModuleInterface
 {
 };
@@ -410,6 +445,7 @@ namespace UnrealIdentifiers
 	extern UNREALUSDWRAPPER_API const TCHAR* Inherited;
 	extern UNREALUSDWRAPPER_API const TCHAR* IdentifierPrefix;
 
+	// The prim name we always use when exporting Skeletons inside SkelRoots to USD
 	extern UNREALUSDWRAPPER_API const TCHAR* ExportedSkeletonPrimName;
 
 	// USceneComponent properties
@@ -459,6 +495,9 @@ namespace UnrealIdentifiers
 	// Tokens from UsdGeomModelAPI that we need to reference from the UsdStageActor
 	extern UNREALUSDWRAPPER_API FString ModelDrawMode;
 	extern UNREALUSDWRAPPER_API FString ModelApplyDrawMode;
+
+	// The character used to separate property namespaces (usually just ':')
+	extern UNREALUSDWRAPPER_API FString UsdNamespaceDelimiter;
 }
 
 struct UNREALUSDWRAPPER_API FUsdDelegates
