USD: Restrict the ControlRigAPI schema to Skeleton and SkelRoot prims.

diff --git a/Engine/Plugins/Importers/USDImporter/Resources/UsdResources/Win64/plugins/unreal/resources/plugInfo.json b/Engine/Plugins/Importers/USDImporter/Resources/UsdResources/Win64/plugins/unreal/resources/plugInfo.json
--- a/Engine/Plugins/Importers/USDImporter/Resources/UsdResources/Win64/plugins/unreal/resources/plugInfo.json
+++ b/Engine/Plugins/Importers/USDImporter/Resources/UsdResources/Win64/plugins/unreal/resources/plugInfo.json
@@ -10,6 +10,10 @@
                         "alias": {
                             "UsdSchemaBase": "ControlRigAPI"
                         }, 
+                        "apiSchemaCanOnlyApplyTo": [
+                            "SkelRoot", 
+                            "Skeleton"
+                        ], 
                         "autoGenerated": true, 
                         "bases": [
                             "UsdAPISchemaBase"
