USD: Fix assert when refreshing the asset cache during engine exit.

diff --git a/Engine/Plugins/Importers/USDImporter/Source/USDClasses/Private/USDAssetCache2.cpp b/Engine/Plugins/Importers/USDImporter/Source/USDClasses/Private/USDAssetCache2.cpp
--- a/Engine/Plugins/Importers/USDImporter/Source/USDClasses/Private/USDAssetCache2.cpp
+++ b/Engine/Plugins/Importers/USDImporter/Source/USDClasses/Private/USDAssetCache2.cpp
@@ -58,7 +58,14 @@ namespace UE::AssetCache::Private
 		{
 			return false;
 		}
+
+		// Note we may fail to find a class even if the name is correct if we're running this on
+		// engine exit
 		UClass* AssetClass = FindClass(ClassName);
+		if (!AssetClass)
+		{
+			return false;
+		}
 
 		static TArray<const UClass*> AllowedPersistentClasses = {
 			UTexture2D::StaticClass(),
diff --git a/Engine/Plugins/Importers/USDImporter/Source/USDStage/Private/USDStageActor.cpp b/Engine/Plugins/Importers/USDImporter/Source/USDStage/Private/USDStageActor.cpp
--- a/Engine/Plugins/Importers/USDImporter/Source/USDStage/Private/USDStageActor.cpp
+++ b/Engine/Plugins/Importers/USDImporter/Source/USDStage/Private/USDStageActor.cpp
@@ -3170,6 +3170,12 @@ void AUsdStageActor::LoadUsdStage(bool bOpenIfNeeded)
 
 void AUsdStageActor::UnloadUsdStage()
 {
+	// No point doing any of this if we're unloading because we're exiting the engine altogether
+	if (IsEngineExitRequested())
+	{
+		return;
+	}
+
 	const bool bMarkDirty = false;
 	Modify(bMarkDirty);
 
