USD: Fix crash and missing keyframes when undo/redoing options on the control rig integration panel.

diff --git a/Engine/Plugins/Importers/USDImporter/Source/USDStage/Private/USDLevelSequenceHelper.cpp b/Engine/Plugins/Importers/USDImporter/Source/USDStage/Private/USDLevelSequenceHelper.cpp
--- a/Engine/Plugins/Importers/USDImporter/Source/USDStage/Private/USDLevelSequenceHelper.cpp
+++ b/Engine/Plugins/Importers/USDImporter/Source/USDStage/Private/USDLevelSequenceHelper.cpp
@@ -269,6 +269,8 @@ namespace UsdLevelSequenceHelperImpl
 			AnimSequence = NewObject<UAnimSequence>();
 			AnimSequence->SetSkeleton(SkeletalMeshComp->GetSkeletalMeshAsset()->GetSkeleton());
 
+			ExportOptions->bTransactRecording = false;
+
 			FMovieSceneSequenceIDRef Template = MovieSceneSequenceID::Root;
 			FMovieSceneSequenceTransform RootToLocalTransform;
 			bResult = MovieSceneToolHelpers::ExportToAnimSequence(AnimSequence, ExportOptions, MovieScene, Player, SkeletalMeshComp, Template, RootToLocalTransform);
@@ -2916,6 +2918,50 @@ void FUsdLevelSequenceHelperImpl::BlockMonitoringChangesForThisTransaction()
 
 void FUsdLevelSequenceHelperImpl::OnObjectTransacted(UObject* Object, const class FTransactionObjectEvent& Event)
 {
+	// Refresh the sequencer on the next tick, or else control rig sections will be missing their keyframes in some undo/redo scenarios.
+	// The repro for this is an extension of the one on UE-191861:
+	// 	- Open a stage with a SkelRoot
+	// 	- Open the stage actor's LevelSequence on the Sequencer
+	//    NOTE: After this point, do not select or interact with the sequencer in any way, just observe it
+	//  - Right-click the SkelRoot and add the ControlRigAPI
+	//  - Enable the option to Use FKControlRig
+	//  - Undo
+	//  - Redo
+	// At this point the track will be back, but the keyframes will be missing. Some interactions with the track at this point can cause a crash too.
+	// The really bizarre part is that *any transaction* after this will cause the keyframes to pop back up (selecting something, moving an unrelated
+	// object on the viewport, etc.).
+	//
+	// This is due to this mechanism on the Sequencer code where calls to MarkAsChanged (which is a member function of tracks, sections,
+	// MovieScene, etc. and is used to let the UI know it needs to refresh something) can be deferred.
+	// The thing that determines where a call is deferred or not is a global, private variable (check FScopedSignedObjectModifyDefer's implementation).
+	//
+	// I think something is causing this mechanism to be stuck deferring everything, or maybe it's some interaction with our code in some way
+	// (not sure at this point). But what I do know is that Sequencer.cpp also has this class FDeferredSignedObjectChangeHandler that listens
+	// to OnObjectTransacted and UndoRedo (much like we're doing right here) and has a member FScopedSignedObjectModifyDefer object (called
+	// "DeferTransactionChanges") that is destroyed when the transaction is complete/canceled. Once that happens, the deferred calls are flushed
+	// and the Sequencer refreshes. This is why *any transaction* causes the keyframes to be drawn back.
+	//
+	// Here we skip that middleman of needing an extra transaction and just flush it right now, to cause our keyframes to show up again.
+	// We could also check the Object to try to limit the scope of this trick, but an alternate repro for this involves deleting the control rig track
+	// and undo->redoing. In that case only the track object would transact, and you can imagine removing the entire binding, and maybe then only the
+	// binding would transact, etc., which would make a robust check on the Object annoying and difficult to maintain.
+	// Given that all this does is essentially refresh the Sequencer UI (and only if it had stuck deferred calls!), it's probably not the worst thing
+	// in the world to check it every undo/redo anyway.
+	//
+	// Annoyingly we also need to do this on the next tick though, because we need to make sure this runs this after FDeferredSignedObjectChangeHandler
+	// itself
+	if (Event.GetEventType() == ETransactionObjectEventType::UndoRedo)
+	{
+		FTSTicker::GetCoreTicker().AddTicker(FTickerDelegate::CreateLambda(
+			[](float Time)
+			{
+				const bool bInForceFlush = true;
+				UE::MovieScene::FScopedSignedObjectModifyDefer ForceFlush{bInForceFlush};
+				return false;
+			}
+		));
+	}
+
 	if (!MainLevelSequence || !IsMonitoringChanges() || !IsValid(Object) || !UsdStage || BlockedTransactionGuids.Contains(Event.GetTransactionId()))
 	{
 		return;
diff --git a/Engine/Plugins/Importers/USDImporter/Source/USDStage/Private/USDStageActor.cpp b/Engine/Plugins/Importers/USDImporter/Source/USDStage/Private/USDStageActor.cpp
--- a/Engine/Plugins/Importers/USDImporter/Source/USDStage/Private/USDStageActor.cpp
+++ b/Engine/Plugins/Importers/USDImporter/Source/USDStage/Private/USDStageActor.cpp
@@ -3433,6 +3433,8 @@ void AUsdStageActor::PostTransacted(const FTransactionObjectEvent& TransactionEv
 {
 	const TArray<FName>& ChangedProperties = TransactionEvent.GetChangedProperties();
 
+	// TODO: This check is broken... we care whether the stage actor has had a pending kill change,
+	// but we are checking whether *anything* has a pending kill change
 	if (TransactionEvent.HasPendingKillChange())
 	{
 		// Fires when being deleted in editor, redo delete
