USD: Modify non-native schemas analytics event to emit a concatenated string instead;

Fix event not tracking unknown schemas.

diff --git a/Engine/Plugins/Importers/USDImporter/Source/USDUtilities/Private/USDConversionUtils.cpp b/Engine/Plugins/Importers/USDImporter/Source/USDUtilities/Private/USDConversionUtils.cpp
--- a/Engine/Plugins/Importers/USDImporter/Source/USDUtilities/Private/USDConversionUtils.cpp
+++ b/Engine/Plugins/Importers/USDImporter/Source/USDUtilities/Private/USDConversionUtils.cpp
@@ -3294,7 +3294,7 @@ void UsdUtils::CollectSchemaAnalytics(const UE::FUsdStage& Stage, const FString&
 		return;
 	}
 
-	TMap<FString, int32> Counts;
+	TSet<FString> SeenSchemas;
 
 	{
 		FScopedUsdAllocs Allocs;
@@ -3310,11 +3310,12 @@ void UsdUtils::CollectSchemaAnalytics(const UE::FUsdStage& Stage, const FString&
 
 				if (!TypeNameStr.IsEmpty())
 				{
-					Counts.FindOrAdd(TypeNameStr) += 1;
+					SeenSchemas.Add(TypeNameStr);
 				}
 			}
 
-			for (const pxr::TfToken& AppliedSchema : PrimRangeIt->GetAppliedSchemas())
+			const pxr::UsdPrimTypeInfo& PrimTypeInfo = PrimRangeIt->GetPrimTypeInfo();
+			for (const pxr::TfToken& AppliedSchema : PrimTypeInfo.GetAppliedAPISchemas())
 			{
 				std::pair<pxr::TfToken, pxr::TfToken> Pair = pxr::UsdSchemaRegistry::GetTypeNameAndInstance(AppliedSchema);
 				FString TypeName = UsdToUnreal::ConvertToken(Pair.first);
@@ -3323,7 +3324,7 @@ void UsdUtils::CollectSchemaAnalytics(const UE::FUsdStage& Stage, const FString&
 				// an ensure or show a warning when analytics fails
 				if (!TypeName.IsEmpty())
 				{
-					Counts.FindOrAdd(TypeName) += 1;
+					SeenSchemas.Add(TypeName);
 				}
 			}
 		}
@@ -3426,20 +3427,22 @@ void UsdUtils::CollectSchemaAnalytics(const UE::FUsdStage& Stage, const FString&
 		"Xform",
 		"XformCommonAPI"};
 
-	TArray<FAnalyticsEventAttribute> EventAttributes;
-	EventAttributes.Reserve(Counts.Num());
 
-	for (const TPair<FString, int32>& Pair : Counts)
+	FString ConcatenatedSeenSchemas;
+	for (const FString& SchemaName : SeenSchemas)
 	{
 		// We only care about non-native schemas
-		if (!NativeSchemaNames.Contains(Pair.Key))
+		if (!NativeSchemaNames.Contains(SchemaName))
 		{
-			EventAttributes.Emplace(Pair.Key, Pair.Value);
+			ConcatenatedSeenSchemas += SchemaName + TEXT(", ");
 		}
 	}
+	ConcatenatedSeenSchemas.RemoveFromEnd(TEXT(", "));
 
-	if (EventAttributes.Num() > 0)
+	if (ConcatenatedSeenSchemas.Len() > 0)
 	{
+		TArray<FAnalyticsEventAttribute> EventAttributes;
+		EventAttributes.Emplace(TEXT("NonNativeSchemas"), ConcatenatedSeenSchemas);
 		IUsdClassesModule::SendAnalytics(MoveTemp(EventAttributes), FString::Printf(TEXT("%s.NonNativeSchemaCounts"), *EventName));
 	}
 #endif	  // USE_USD_SDK
