USD: Fix missing schema files from cl 30075823 for Mac and Linux.

#rnx

diff --git a/Engine/Plugins/Importers/USDImporter/Resources/UsdResources/Linux/plugins/unreal/resources/plugInfo.json b/Engine/Plugins/Importers/USDImporter/Resources/UsdResources/Linux/plugins/unreal/resources/plugInfo.json
--- a/Engine/Plugins/Importers/USDImporter/Resources/UsdResources/Linux/plugins/unreal/resources/plugInfo.json
+++ b/Engine/Plugins/Importers/USDImporter/Resources/UsdResources/Linux/plugins/unreal/resources/plugInfo.json
@@ -10,6 +10,10 @@
                         "alias": {
                             "UsdSchemaBase": "ControlRigAPI"
                         }, 
+                        "apiSchemaCanOnlyApplyTo": [
+                            "SkelRoot", 
+                            "Skeleton"
+                        ], 
                         "autoGenerated": true, 
                         "bases": [
                             "UsdAPISchemaBase"
diff --git a/Engine/Plugins/Importers/USDImporter/Resources/UsdResources/Mac/plugins/unreal/resources/plugInfo.json b/Engine/Plugins/Importers/USDImporter/Resources/UsdResources/Mac/plugins/unreal/resources/plugInfo.json
--- a/Engine/Plugins/Importers/USDImporter/Resources/UsdResources/Mac/plugins/unreal/resources/plugInfo.json
+++ b/Engine/Plugins/Importers/USDImporter/Resources/UsdResources/Mac/plugins/unreal/resources/plugInfo.json
@@ -10,6 +10,10 @@
                         "alias": {
                             "UsdSchemaBase": "ControlRigAPI"
                         }, 
+                        "apiSchemaCanOnlyApplyTo": [
+                            "SkelRoot", 
+                            "Skeleton"
+                        ], 
                         "autoGenerated": true, 
                         "bases": [
                             "UsdAPISchemaBase"
