USD: Add new UsdStageActor/import option "bReuseIdenticalAssets" to let the user control when to reuse assets by hash;

Temp patch to prevent race conditions when parsing two identical groom caches;
Fix GroomBinding and strands cache assets being renamed to just the exact same name as the groom when importing;
Fix SkeletalMesh physics assets and groom caches not being linked to their prims when an asset was reused from the asset cache.

diff --git a/Engine/Plugins/Importers/USDImporter/Source/USDSchemas/Private/Custom/MDLUSDShadeMaterialTranslator.cpp b/Engine/Plugins/Importers/USDImporter/Source/USDSchemas/Private/Custom/MDLUSDShadeMaterialTranslator.cpp
--- a/Engine/Plugins/Importers/USDImporter/Source/USDSchemas/Private/Custom/MDLUSDShadeMaterialTranslator.cpp
+++ b/Engine/Plugins/Importers/USDImporter/Source/USDSchemas/Private/Custom/MDLUSDShadeMaterialTranslator.cpp
@@ -149,8 +149,9 @@ void FMdlUsdShadeMaterialTranslator::CreateAssets()
 
 		const FString MdlFullName = MdlModuleName + TEXT("::") + MdlDefinitionName;
 		const FString MdlFullInstanceName = MdlFullName + TEXT("_Instance");
+		const FString HashPrefix = UsdUtils::GetAssetHashPrefix(GetPrim(), Context->bReuseIdenticalAssets);
 
-		UMaterialInterface* MdlMaterial = Cast< UMaterialInterface >( Context->AssetCache->GetCachedAsset( MdlFullName ) );
+		UMaterialInterface* MdlMaterial = Cast< UMaterialInterface >( Context->AssetCache->GetCachedAsset( HashPrefix + MdlFullName ) );
 		if ( !MdlMaterial )
 		{
 			FScopedUnrealAllocs UnrealAllocs;
@@ -174,13 +175,19 @@ void FMdlUsdShadeMaterialTranslator::CreateAssets()
 			}
 		}
 
-		UMaterialInstanceConstant* MdlMaterialInstance = Cast< UMaterialInstanceConstant >(Context->AssetCache->GetCachedAsset(MdlFullInstanceName));
+		UMaterialInstanceConstant* MdlMaterialInstance = Cast< UMaterialInstanceConstant >(Context->AssetCache->GetCachedAsset(HashPrefix + MdlFullInstanceName));
 		if ( !MdlMaterialInstance && MdlMaterial )
 		{
 			MdlMaterialInstance = NewObject< UMaterialInstanceConstant >(GetTransientPackage(), NAME_None, Context->ObjectFlags | RF_Transient);
 			MdlMaterialInstance->SetParentEditorOnly( MdlMaterial );
 
-			UsdToUnreal::ConvertShadeInputsToParameters( ShadeMaterial, *MdlMaterialInstance, Context->AssetCache.Get(), *Context->RenderContext.ToString() );
+			UsdToUnreal::ConvertShadeInputsToParameters(
+				ShadeMaterial,
+				*MdlMaterialInstance,
+				Context->AssetCache.Get(),
+				*Context->RenderContext.ToString(),
+				Context->bReuseIdenticalAssets
+			);
 
 			// We can't blindly recreate all component render states when a level is being added, because we may end up first creating
 			// render states for some components, and UWorld::AddToWorld calls FScene::AddPrimitive which expects the component to not have
@@ -199,8 +206,8 @@ void FMdlUsdShadeMaterialTranslator::CreateAssets()
 			UE::MDLShadeTranslatorImpl::Private::NotifyIfMaterialNeedsVirtualTextures( MdlMaterialInstance );
 		}
 
-		PostImportMaterial(MdlFullName, MdlMaterial);
-		PostImportMaterial(MdlFullInstanceName, MdlMaterialInstance);
+		PostImportMaterial(HashPrefix + MdlFullName, MdlMaterial);
+		PostImportMaterial(HashPrefix + MdlFullInstanceName, MdlMaterialInstance);
 	}
 	else
 	{
diff --git a/Engine/Plugins/Importers/USDImporter/Source/USDSchemas/Private/Custom/MaterialXUSDShadeMaterialTranslator.cpp b/Engine/Plugins/Importers/USDImporter/Source/USDSchemas/Private/Custom/MaterialXUSDShadeMaterialTranslator.cpp
--- a/Engine/Plugins/Importers/USDImporter/Source/USDSchemas/Private/Custom/MaterialXUSDShadeMaterialTranslator.cpp
+++ b/Engine/Plugins/Importers/USDImporter/Source/USDSchemas/Private/Custom/MaterialXUSDShadeMaterialTranslator.cpp
@@ -158,7 +158,13 @@ namespace UE::USDMaterialXTranslator::Private
 		return Hash.ToString();
 	}
 
-	bool TranslateMaterialXFile(const FString& MaterialXFilePath, const FString& FileHash, const pxr::UsdPrim& MaterialXReferencerPrim, UUsdAssetCache2& AssetCache)
+	bool TranslateMaterialXFile(
+		const FString& MaterialXFilePath,
+		const FString& FileHash,
+		const pxr::UsdPrim& MaterialXReferencerPrim,
+		UUsdAssetCache2& AssetCache,
+		bool bReuseIdenticalAssets
+	)
 	{
 		if (!FPaths::FileExists(MaterialXFilePath))
 		{
@@ -184,10 +190,12 @@ namespace UE::USDMaterialXTranslator::Private
 			TempPackagePath = FString::Printf(TEXT("/Engine/USDImporter/Transient/%d"), PackageSuffix++);
 		}
 
+		const FString HashPrefix = UsdUtils::GetAssetHashPrefix(MaterialXReferencerPrim, bReuseIdenticalAssets);
+
 		// Currently Interchange cannot import directly into the transient package, so we'll import into a temporary folder and then
 		// move everything over
 		InterchangeParameters.OnAssetsImportDoneNative.BindLambda(
-			[&AssetCache, &FileHash, &ReferencerPrimPath, &TempPackagePath, &MaterialXFilePath](const TArray<UObject*>& ImportedObjects)
+			[&AssetCache, &FileHash, &ReferencerPrimPath, &TempPackagePath, &MaterialXFilePath, &HashPrefix](const TArray<UObject*>& ImportedObjects)
 			{
 				UE_LOG(LogUsd, Verbose, TEXT("Translated '%d' assets from MaterialX file '%s'"), ImportedObjects.Num(), *MaterialXFilePath);
 
@@ -219,23 +227,23 @@ namespace UE::USDMaterialXTranslator::Private
 					{
 						// MaterialX names are unique, and can only have alphanumeric and the "_" character, so we should
 						// always have a solid enough mapping to assume UAsset name == Prim name == MaterialX name
-						FString NewMaterialHash = FileHash + TEXT("/") + Material->GetFName().GetPlainNameString();
+						FString NewPrefixedMaterialHash = HashPrefix + FileHash + TEXT("/") + Material->GetFName().GetPlainNameString();
 
 						// We generate all assets from the MaterialX file once, but it's possible we're just updating a
 						// single material prim here. If we were to cache all assets here, we'd potentially be trying
 						// to overwrite the existing assets that are being used by other prims and wouldn't otherwise
 						// be discarded, so make sure we don't do that
-						UMaterialInterface* ExistingMaterial = Cast<UMaterialInterface>(AssetCache.GetCachedAsset(NewMaterialHash));
+						UMaterialInterface* ExistingMaterial = Cast<UMaterialInterface>(AssetCache.GetCachedAsset(NewPrefixedMaterialHash));
 						if (!ExistingMaterial)
 						{
-							AssetCache.CacheAsset(NewMaterialHash, Material);
+							AssetCache.CacheAsset(NewPrefixedMaterialHash, Material);
 						}
 					}
 					else if (UTexture* Texture = Cast<UTexture>(ImportedObject))
 					{
 						const FString FilePath = Texture->AssetImportData ? Texture->AssetImportData->GetFirstFilename() : TEXT("");
 
-						FString TextureHash = UsdUtils::GetTextureHash(
+						FString PrefixedTextureHash = HashPrefix + UsdUtils::GetTextureHash(
 							FilePath,
 							Texture->SRGB,
 							Texture->CompressionSettings,
@@ -244,14 +252,14 @@ namespace UE::USDMaterialXTranslator::Private
 						);
 
 						// See comment on the analogous part of the Material case above
-						UTexture* ExistingTexture = Cast<UTexture>(AssetCache.GetCachedAsset(TextureHash));
+						UTexture* ExistingTexture = Cast<UTexture>(AssetCache.GetCachedAsset(PrefixedTextureHash));
 						if (!ExistingTexture)
 						{
 							UUsdAssetUserData* UserData = NewObject<UUsdAssetUserData>(Texture, TEXT("USDAssetUserData"));
 							UserData->PrimPaths = {ReferencerPrimPath};
 							Texture->AddAssetUserData(UserData);
 
-							AssetCache.CacheAsset(TextureHash, Texture);
+							AssetCache.CacheAsset(PrefixedTextureHash, Texture);
 						}
 					}
 					else
@@ -418,6 +426,7 @@ void FMaterialXUsdShadeMaterialTranslator::CreateAssets()
 	}
 
 	FString TargetHashSuffix = TEXT("/") + UsdToUnreal::ConvertString(Prim.GetName());
+	FString TargetHashPrefix = UsdUtils::GetAssetHashPrefix(Prim, Context->bReuseIdenticalAssets);
 
 	FString FoundMaterialAssetHash;
 	UMaterialInterface* ParsedMaterial = nullptr;
@@ -433,7 +442,7 @@ void FMaterialXUsdShadeMaterialTranslator::CreateAssets()
 			continue;
 		}
 
-		FString MaterialAssetHashForThisFile = MaterialXHash + TargetHashSuffix;
+		FString MaterialAssetHashForThisFile = TargetHashPrefix + MaterialXHash + TargetHashSuffix;
 		ParsedMaterial = Cast<UMaterialInterface>(Context->AssetCache->GetCachedAsset(MaterialAssetHashForThisFile));
 		if (ParsedMaterial)
 		{
@@ -455,7 +464,8 @@ void FMaterialXUsdShadeMaterialTranslator::CreateAssets()
 				MaterialXFilePath,
 				MaterialXFileHash,
 				MaterialXReferencerPrim,
-				*Context->AssetCache
+				*Context->AssetCache,
+				Context->bReuseIdenticalAssets
 			);
 
 			if (!bSuccess)
@@ -472,7 +482,7 @@ void FMaterialXUsdShadeMaterialTranslator::CreateAssets()
 			// Check if we found our target material when parsing this file
 			if (!ParsedMaterial)
 			{
-				FString MaterialAssetHashForThisFile = MaterialXFileHash + TargetHashSuffix;
+				FString MaterialAssetHashForThisFile = TargetHashPrefix + MaterialXFileHash + TargetHashSuffix;
 				ParsedMaterial = Cast<UMaterialInterface>(Context->AssetCache->GetCachedAsset(MaterialAssetHashForThisFile));
 				if (ParsedMaterial)
 				{
diff --git a/Engine/Plugins/Importers/USDImporter/Source/USDSchemas/Private/MeshTranslationImpl.cpp b/Engine/Plugins/Importers/USDImporter/Source/USDSchemas/Private/MeshTranslationImpl.cpp
--- a/Engine/Plugins/Importers/USDImporter/Source/USDSchemas/Private/MeshTranslationImpl.cpp
+++ b/Engine/Plugins/Importers/USDImporter/Source/USDSchemas/Private/MeshTranslationImpl.cpp
@@ -156,7 +156,9 @@ namespace UE::MeshTranslationImplInternal::Private
 		UMaterialInterface& Material,
 		const TMap<FString, int32>& MeshPrimvarToUVIndex,
 		UUsdAssetCache2* AssetCache,
-		FUsdInfoCache* InfoCache
+		FUsdInfoCache* InfoCache,
+		const FString& MaterialHashPrefix,
+		bool bReuseIdenticalAssets
 	)
 	{
 		UUsdMaterialAssetUserData* MaterialAssetUserData = Material.GetAssetUserData<UUsdMaterialAssetUserData>();
@@ -260,45 +262,53 @@ namespace UE::MeshTranslationImplInternal::Private
 			}
 		}
 
-		// Try finding a compatible material in the asset cache, if our original Material belongs to it
-		FString MaterialHash = AssetCache ? AssetCache->GetHashForAsset(&Material) : FString{};
-		const bool bMaterialBelongsToAssetCache = !MaterialHash.IsEmpty();
-		if (bMaterialBelongsToAssetCache)
+		FString ExistingHash = AssetCache->GetHashForAsset(&Material);
+		const bool bMaterialBelongsToAssetCache = !ExistingHash.IsEmpty();
+		if (!bMaterialBelongsToAssetCache)
 		{
-			// Generate a deterministic hash based on the original material hash and this primvar UVIndex assignment
-			CompatiblePrimvarAndUVIndexPairs.Sort(
-				[](const TPair<FString, int32>& LHS, const TPair<FString, int32>& RHS)
+			return nullptr;
+		}
+
+		// Generate a deterministic hash based on the original material hash and this primvar UVIndex assignment
+		CompatiblePrimvarAndUVIndexPairs.Sort(
+			[](const TPair<FString, int32>& LHS, const TPair<FString, int32>& RHS)
+			{
+				if (LHS.Key == RHS.Key)
 				{
-					if (LHS.Key == RHS.Key)
-					{
-						return LHS.Value < RHS.Value;
-					}
-					else
-					{
-						return LHS.Key < RHS.Key;
-					}
+					return LHS.Value < RHS.Value;
+				}
+				else
+				{
+					return LHS.Key < RHS.Key;
 				}
-			);
-			FSHAHash Hash;
-			FSHA1 SHA1;
-			SHA1.UpdateWithString(*MaterialHash, MaterialHash.Len());
-			for (const TPair<FString, int32>& Pair : CompatiblePrimvarAndUVIndexPairs)
-			{
-				SHA1.UpdateWithString(*Pair.Key, Pair.Key.Len());
-				SHA1.Update((const uint8*)&Pair.Value, sizeof(Pair.Value));
 			}
-			SHA1.Final();
-			SHA1.GetHash(&Hash.Hash[0]);
-			MaterialHash = Hash.ToString();
+		);
+		FSHAHash Hash;
+		FSHA1 SHA1;
+		SHA1.UpdateWithString(*ExistingHash, ExistingHash.Len());
+		for (const TPair<FString, int32>& Pair : CompatiblePrimvarAndUVIndexPairs)
+		{
+			SHA1.UpdateWithString(*Pair.Key, Pair.Key.Len());
+			SHA1.Update((const uint8*)&Pair.Value, sizeof(Pair.Value));
+		}
+		SHA1.Final();
+		SHA1.GetHash(&Hash.Hash[0]);
 
-			if (UMaterialInterface* ExistingCompatibleMaterial = Cast<UMaterialInterface>(AssetCache->GetCachedAsset(MaterialHash)))
-			{
-				UE_LOG(LogUsd, Verbose, TEXT("Found existing compatible Material '%s' on the asset cache with hash '%s'"),
-					*ExistingCompatibleMaterial->GetPathName(),
-					*MaterialHash
-				);
-				CompatibleMaterial = ExistingCompatibleMaterial;
-			}
+		// In theory we don't even need to add the prefix here because our ExistingHash will already have the same prefix...
+		// However for consistency it's probably for the best to have both assets have the same prefix, so you can tell
+		// from the hash that they originated from the same prim
+		FString PrefixedMaterialHash = MaterialHashPrefix + Hash.ToString();
+
+		if (UMaterialInterface* ExistingCompatibleMaterial = Cast<UMaterialInterface>(AssetCache->GetCachedAsset(PrefixedMaterialHash)))
+		{
+			UE_LOG(
+				LogUsd,
+				Verbose,
+				TEXT("Found existing compatible Material '%s' on the asset cache with hash '%s'"),
+				*ExistingCompatibleMaterial->GetPathName(),
+				*PrefixedMaterialHash
+			);
+			CompatibleMaterial = ExistingCompatibleMaterial;
 		}
 
 		// We have to create a brand new compatible material instance
@@ -321,7 +331,7 @@ namespace UE::MeshTranslationImplInternal::Private
 				*Material.GetPathName(),
 				*UsdUtils::StringifyMap(MaterialAssetUserData->ParameterToPrimvar),
 				*UsdUtils::StringifyMap(CompatiblePrimvarToUVIndex),
-				*MaterialHash
+				*PrefixedMaterialHash
 			);
 
 #if WITH_EDITOR
@@ -398,7 +408,7 @@ namespace UE::MeshTranslationImplInternal::Private
 		{
 			if (bMaterialBelongsToAssetCache)
 			{
-				AssetCache->CacheAsset(MaterialHash, CompatibleMaterial);
+				AssetCache->CacheAsset(PrefixedMaterialHash, CompatibleMaterial);
 			}
 
 			if (InfoCache)
@@ -419,7 +429,8 @@ TMap<const UsdUtils::FUsdPrimMaterialSlot*, UMaterialInterface*> MeshTranslation
 	const TArray<UsdUtils::FUsdPrimMaterialAssignmentInfo>& AssignmentInfo,
 	UUsdAssetCache2& AssetCache,
 	FUsdInfoCache& InfoCache,
-	EObjectFlags Flags
+	EObjectFlags Flags,
+	bool bReuseIdenticalAssets
 )
 {
 	FScopedUnrealAllocs Allocs;
@@ -450,8 +461,10 @@ TMap<const UsdUtils::FUsdPrimMaterialSlot*, UMaterialInterface*> MeshTranslation
 			{
 			case UsdUtils::EPrimAssignmentType::DisplayColor:
 			{
+				const FString PrefixedHash = UsdUtils::GetAssetHashPrefix(UsdPrim, bReuseIdenticalAssets) + Slot.MaterialSource;
+
 				// Try reusing an already created DisplayColor material
-				if (UMaterialInterface* ExistingMaterial = Cast<UMaterialInterface>(AssetCache.GetCachedAsset(Slot.MaterialSource)))
+				if (UMaterialInterface* ExistingMaterial = Cast<UMaterialInterface>(AssetCache.GetCachedAsset(PrefixedHash)))
 				{
 					Material = ExistingMaterial;
 				}
@@ -484,7 +497,7 @@ TMap<const UsdUtils::FUsdPrimMaterialSlot*, UMaterialInterface*> MeshTranslation
 						// We can only cache transient assets
 						MaterialInstance->SetFlags(RF_Transient);
 
-						AssetCache.CacheAsset(Slot.MaterialSource, MaterialInstance);
+						AssetCache.CacheAsset(PrefixedHash, MaterialInstance);
 						Material = MaterialInstance;
 					}
 				}
@@ -532,7 +545,8 @@ TMap<const UsdUtils::FUsdPrimMaterialSlot*, UMaterialInterface*> MeshTranslation
 					}
 				}
 
-				FString MaterialHash = Material ? AssetCache.GetHashForAsset(Material) : FString{};
+				FString PrefixedMaterialHash = Material ? AssetCache.GetHashForAsset(Material) : FString{};
+				FString HashPrefix = UsdUtils::GetAssetHashPrefix(UsdPrim.GetStage()->GetPrimAtPath(MaterialPrimPath), bReuseIdenticalAssets);
 
 				// Need to create a two-sided material on-demand, *before* we make it compatible:
 				// This because at runtime we can't just set the base property overrides, and just instead create a new
@@ -552,12 +566,12 @@ TMap<const UsdUtils::FUsdPrimMaterialSlot*, UMaterialInterface*> MeshTranslation
 						continue;
 					}
 
-					const FString OneSidedHash = AssetCache.GetHashForAsset(OneSidedMat);
-					const FString TwoSidedHash = OneSidedHash + UnrealIdentifiers::TwoSidedMaterialSuffix;
+					const FString PrefixedOneSidedHash = AssetCache.GetHashForAsset(OneSidedMat);
+					const FString PrefixedTwoSidedHash = PrefixedOneSidedHash + UnrealIdentifiers::TwoSidedMaterialSuffix;
 
 					// Check if for some reason we already have a two-sided material ready due to a complex scenario
 					// related to the global cache
-					UMaterialInterface* TwoSidedMat = Cast<UMaterialInterface>(AssetCache.GetCachedAsset(TwoSidedHash));
+					UMaterialInterface* TwoSidedMat = Cast<UMaterialInterface>(AssetCache.GetCachedAsset(PrefixedTwoSidedHash));
 					if (!TwoSidedMat)
 					{
 						TwoSidedMat =
@@ -568,7 +582,7 @@ TMap<const UsdUtils::FUsdPrimMaterialSlot*, UMaterialInterface*> MeshTranslation
 					{
 						TwoSidedMat->SetFlags(RF_Transient);
 						Material = TwoSidedMat;
-						MaterialHash = TwoSidedHash;
+						PrefixedMaterialHash = PrefixedTwoSidedHash;
 					}
 					else
 					{
@@ -581,7 +595,7 @@ TMap<const UsdUtils::FUsdPrimMaterialSlot*, UMaterialInterface*> MeshTranslation
 							*OneSidedMat->GetPathName()
 						);
 						Material = OneSidedMat;
-						MaterialHash = OneSidedHash;
+						PrefixedMaterialHash = PrefixedOneSidedHash;
 					}
 				}
 
@@ -589,7 +603,7 @@ TMap<const UsdUtils::FUsdPrimMaterialSlot*, UMaterialInterface*> MeshTranslation
 				{
 					// Cache the material to "ping it" as active, but also register two sided materials for the
 					// first time
-					AssetCache.CacheAsset(MaterialHash, Material);
+					AssetCache.CacheAsset(PrefixedMaterialHash, Material);
 					InfoCache.LinkAssetToPrim(UE::FSdfPath{*Slot.MaterialSource}, Material);
 
 					// Finally, try to make our generated material primvar-compatible. We do this last because this will
@@ -610,7 +624,9 @@ TMap<const UsdUtils::FUsdPrimMaterialSlot*, UMaterialInterface*> MeshTranslation
 								*Material,
 								MeshPrimvarToUVIndex,
 								&AssetCache,
-								&InfoCache
+								&InfoCache,
+								HashPrefix,
+								bReuseIdenticalAssets
 							);
 
 						if (CompatibleMaterial)
@@ -676,7 +692,8 @@ void MeshTranslationImpl::SetMaterialOverrides(
 	EObjectFlags Flags,
 	bool bInterpretLODs,
 	const FName& RenderContext,
-	const FName& MaterialPurpose
+	const FName& MaterialPurpose,
+	bool bReuseIdenticalAssets
 )
 {
 	FScopedUsdAllocs Allocs;
@@ -813,7 +830,8 @@ void MeshTranslationImpl::SetMaterialOverrides(
 			LODIndexToAssignments,
 			AssetCache,
 			InfoCache,
-			Flags
+			Flags,
+			bReuseIdenticalAssets
 		);
 	}
 
diff --git a/Engine/Plugins/Importers/USDImporter/Source/USDSchemas/Private/MeshTranslationImpl.h b/Engine/Plugins/Importers/USDImporter/Source/USDSchemas/Private/MeshTranslationImpl.h
--- a/Engine/Plugins/Importers/USDImporter/Source/USDSchemas/Private/MeshTranslationImpl.h
+++ b/Engine/Plugins/Importers/USDImporter/Source/USDSchemas/Private/MeshTranslationImpl.h
@@ -35,7 +35,8 @@ namespace MeshTranslationImpl
 		const TArray<UsdUtils::FUsdPrimMaterialAssignmentInfo>& AssignmentInfo,
 		UUsdAssetCache2& AssetCache,
 		FUsdInfoCache& InfoCache,
-		EObjectFlags Flags
+		EObjectFlags Flags,
+		bool bReuseIdenticalAssets
 	);
 
 	/**
@@ -52,7 +53,8 @@ namespace MeshTranslationImpl
 		EObjectFlags Flags,
 		bool bInterpretLODs,
 		const FName& RenderContext,
-		const FName& MaterialPurpose
+		const FName& MaterialPurpose,
+		bool bReuseIdenticalAssets
 	);
 
 	void RecordSourcePrimsForMaterialSlots(
diff --git a/Engine/Plugins/Importers/USDImporter/Source/USDSchemas/Private/USDGeomMeshTranslator.cpp b/Engine/Plugins/Importers/USDImporter/Source/USDSchemas/Private/USDGeomMeshTranslator.cpp
--- a/Engine/Plugins/Importers/USDImporter/Source/USDSchemas/Private/USDGeomMeshTranslator.cpp
+++ b/Engine/Plugins/Importers/USDImporter/Source/USDSchemas/Private/USDGeomMeshTranslator.cpp
@@ -196,7 +196,8 @@ namespace UsdGeomMeshTranslatorImpl
 		UUsdAssetCache2& AssetCache,
 		FUsdInfoCache* InfoCache,
 		float Time,
-		EObjectFlags Flags
+		EObjectFlags Flags,
+		bool bReuseIdenticalAssets
 	)
 	{
 		if (!InfoCache)
@@ -217,7 +218,8 @@ namespace UsdGeomMeshTranslatorImpl
 			LODIndexToMaterialInfo,
 			AssetCache,
 			*InfoCache,
-			Flags
+			Flags,
+			bReuseIdenticalAssets
 		);
 
 		uint32 StaticMeshSlotIndex = 0;
@@ -581,7 +583,14 @@ namespace UsdGeomMeshTranslatorImpl
 		FStaticMeshOperations::ComputeTangentsAndNormals(MeshDescription, Options);
 	}
 
-	UStaticMesh* CreateStaticMesh(TArray<FMeshDescription>& LODIndexToMeshDescription, FUsdSchemaTranslationContext& Context, const FString& MeshName, const bool bShouldEnableNanite, bool& bOutIsNew)
+	UStaticMesh* CreateStaticMesh(
+		const UE::FUsdPrim& Prim,
+		TArray<FMeshDescription>& LODIndexToMeshDescription,
+		FUsdSchemaTranslationContext& Context,
+		const FString& MeshName,
+		const bool bShouldEnableNanite,
+		bool& bOutIsNew
+	)
 	{
 		UStaticMesh* StaticMesh = nullptr;
 
@@ -626,9 +635,11 @@ namespace UsdGeomMeshTranslatorImpl
 			SHA1.GetHash(&AllLODHash.Hash[0]);
 		}
 
+		FString PrefixedAssetHash = UsdUtils::GetAssetHashPrefix(Prim, Context.bReuseIdenticalAssets) + AllLODHash.ToString();
+
 		if (Context.AssetCache)
 		{
-			StaticMesh = Cast< UStaticMesh >(Context.AssetCache->GetCachedAsset(AllLODHash.ToString()));
+			StaticMesh = Cast< UStaticMesh >(Context.AssetCache->GetCachedAsset(PrefixedAssetHash));
 		}
 
 		if (!StaticMesh && bHasValidMeshDescription)
@@ -668,7 +679,7 @@ namespace UsdGeomMeshTranslatorImpl
 
 			if (Context.AssetCache)
 			{
-				Context.AssetCache->CacheAsset(AllLODHash.ToString(), StaticMesh);
+				Context.AssetCache->CacheAsset(PrefixedAssetHash, StaticMesh);
 			}
 		}
 		else
@@ -1442,7 +1453,14 @@ void FBuildStaticMeshTaskChain::SetupTasks()
 
 			bool bIsNew = true;
 			const bool bShouldEnableNanite = UsdGeomMeshTranslatorImpl::ShouldEnableNanite(LODIndexToMeshDescription, LODIndexToMaterialInfo, *Context, GetPrim());
-			StaticMesh = UsdGeomMeshTranslatorImpl::CreateStaticMesh(LODIndexToMeshDescription, *Context, MeshName, bShouldEnableNanite, bIsNew);
+			StaticMesh = UsdGeomMeshTranslatorImpl::CreateStaticMesh(
+				GetPrim(),
+				LODIndexToMeshDescription,
+				*Context,
+				MeshName,
+				bShouldEnableNanite,
+				bIsNew
+			);
 
 			if (StaticMesh)
 			{
@@ -1499,7 +1517,8 @@ void FBuildStaticMeshTaskChain::SetupTasks()
 						*Context->AssetCache.Get(),
 						Context->InfoCache.Get(),
 						Context->Time,
-						Context->ObjectFlags
+						Context->ObjectFlags,
+						Context->bReuseIdenticalAssets
 					);
 
 #if WITH_EDITOR
diff --git a/Engine/Plugins/Importers/USDImporter/Source/USDSchemas/Private/USDGeomPointInstancerTranslator.cpp b/Engine/Plugins/Importers/USDImporter/Source/USDSchemas/Private/USDGeomPointInstancerTranslator.cpp
--- a/Engine/Plugins/Importers/USDImporter/Source/USDSchemas/Private/USDGeomPointInstancerTranslator.cpp
+++ b/Engine/Plugins/Importers/USDImporter/Source/USDSchemas/Private/USDGeomPointInstancerTranslator.cpp
@@ -541,7 +541,8 @@ void FUsdGeomPointInstancerTranslator::UpdateComponents(USceneComponent* PointIn
 					Context->ObjectFlags,
 					Context->bAllowInterpretingLODs,
 					Context->RenderContext,
-					Context->MaterialPurpose
+					Context->MaterialPurpose,
+					Context->bReuseIdenticalAssets
 				);
 			}
 		}
diff --git a/Engine/Plugins/Importers/USDImporter/Source/USDSchemas/Private/USDGeomPrimitiveTranslator.cpp b/Engine/Plugins/Importers/USDImporter/Source/USDSchemas/Private/USDGeomPrimitiveTranslator.cpp
--- a/Engine/Plugins/Importers/USDImporter/Source/USDSchemas/Private/USDGeomPrimitiveTranslator.cpp
+++ b/Engine/Plugins/Importers/USDImporter/Source/USDSchemas/Private/USDGeomPrimitiveTranslator.cpp
@@ -179,7 +179,8 @@ USceneComponent* FUsdGeomPrimitiveTranslator::CreateComponents()
 					Context->ObjectFlags,
 					Context->bAllowInterpretingLODs,
 					Context->RenderContext,
-					Context->MaterialPurpose
+					Context->MaterialPurpose,
+					Context->bReuseIdenticalAssets
 				);
 			}
 		}
diff --git a/Engine/Plugins/Importers/USDImporter/Source/USDSchemas/Private/USDGeomXformableTranslator.cpp b/Engine/Plugins/Importers/USDImporter/Source/USDSchemas/Private/USDGeomXformableTranslator.cpp
--- a/Engine/Plugins/Importers/USDImporter/Source/USDSchemas/Private/USDGeomXformableTranslator.cpp
+++ b/Engine/Plugins/Importers/USDImporter/Source/USDSchemas/Private/USDGeomXformableTranslator.cpp
@@ -342,7 +342,8 @@ USceneComponent* FUsdGeomXformableTranslator::CreateComponents()
 					Context->ObjectFlags,
 					Context->bAllowInterpretingLODs,
 					Context->RenderContext,
-					Context->MaterialPurpose
+					Context->MaterialPurpose,
+					Context->bReuseIdenticalAssets
 				);
 			}
 		}
@@ -958,8 +959,9 @@ void FUsdGeomXformableTranslator::CreateAlternativeDrawModeAssets(EUsdDrawMode D
 			}
 			else
 			{
-				const FString TextureHash = LexToString(FMD5Hash::HashFile(*ResolvedPath));
-				UTexture2D* Texture = Cast<UTexture2D>(Context->AssetCache->GetCachedAsset(TextureHash));
+				const FString HashPrefix = UsdUtils::GetAssetHashPrefix(GetPrim(), Context->bReuseIdenticalAssets);
+				const FString PrefixedTextureHash = HashPrefix + LexToString(FMD5Hash::HashFile(*ResolvedPath));
+				UTexture2D* Texture = Cast<UTexture2D>(Context->AssetCache->GetCachedAsset(PrefixedTextureHash));
 
 				if (!Texture)
 				{
@@ -970,7 +972,7 @@ void FUsdGeomXformableTranslator::CreateAlternativeDrawModeAssets(EUsdDrawMode D
 						Context->AssetCache.Get()
 					));
 
-					Context->AssetCache->CacheAsset(TextureHash, Texture);
+					Context->AssetCache->CacheAsset(PrefixedTextureHash, Texture);
 				}
 
 				// We link the textures to the prim, so that if the prim is reloaded the AUsdStageActor knows to potentially
diff --git a/Engine/Plugins/Importers/USDImporter/Source/USDSchemas/Private/USDGeometryCacheTranslator.cpp b/Engine/Plugins/Importers/USDImporter/Source/USDSchemas/Private/USDGeometryCacheTranslator.cpp
--- a/Engine/Plugins/Importers/USDImporter/Source/USDSchemas/Private/USDGeometryCacheTranslator.cpp
+++ b/Engine/Plugins/Importers/USDImporter/Source/USDSchemas/Private/USDGeometryCacheTranslator.cpp
@@ -65,11 +65,25 @@ static FAutoConsoleVariableRef CVarEnableSubdiv(
 
 namespace UsdGeometryCacheTranslatorImpl
 {
-	bool ProcessGeometryCacheMaterials(const pxr::UsdPrim& UsdPrim, const TArray<UsdUtils::FUsdPrimMaterialAssignmentInfo>& LODIndexToMaterialInfo,
-									   UGeometryCache& GeometryCache, UUsdAssetCache2& AssetCache, FUsdInfoCache* InfoCache, float Time, EObjectFlags Flags)
+	bool ProcessGeometryCacheMaterials(
+		const pxr::UsdPrim& UsdPrim,
+		const TArray<UsdUtils::FUsdPrimMaterialAssignmentInfo>& LODIndexToMaterialInfo,
+		UGeometryCache& GeometryCache,
+		UUsdAssetCache2& AssetCache,
+		FUsdInfoCache* InfoCache,
+		float Time,
+		EObjectFlags Flags,
+		bool bReuseIdenticalAssets
+	)
 	{
 		TMap<const UsdUtils::FUsdPrimMaterialSlot*, UMaterialInterface*> ResolvedMaterials = MeshTranslationImpl::ResolveMaterialAssignmentInfo(
-			UsdPrim, LODIndexToMaterialInfo, AssetCache, *InfoCache, Flags);
+			UsdPrim,
+			LODIndexToMaterialInfo,
+			AssetCache,
+			*InfoCache,
+			Flags,
+			bReuseIdenticalAssets
+		);
 
 		uint32 SlotIndex = 0;
 		bool bMaterialAssignementsHaveChanged = false;
@@ -347,9 +361,11 @@ namespace UsdGeometryCacheTranslatorImpl
 		return StreamableTrack;
 	}
 
-	UGeometryCache* CreateGeometryCache(const FString& RootPrimPath, const FMeshDescription& MeshDescription, const TArray<UE::FSdfPath>& MeshPaths, const TArray<int32>& MaterialOffsets,
+	UGeometryCache* CreateGeometryCache(const UE::FUsdPrim& RootPrim, const FMeshDescription& MeshDescription, const TArray<UE::FSdfPath>& MeshPaths, const TArray<int32>& MaterialOffsets,
 										TSharedRef<FUsdSchemaTranslationContext> Context, bool& bOutIsNew, float& StartOffsetTime)
 	{
+		FString RootPrimPath = RootPrim.GetPrimPath().GetString();
+
 		// Compute the asset hash from the merged mesh description
 		FSHA1 SHA1;
 		FSHAHash MeshHash = FStaticMeshOperations::ComputeSHAHash(MeshDescription);
@@ -372,7 +388,9 @@ namespace UsdGeometryCacheTranslatorImpl
 
 		FSHAHash GeoCacheHash;
 		SHA1.GetHash(&GeoCacheHash.Hash[0]);
-		UGeometryCache* GeometryCache = Cast<UGeometryCache>(Context->AssetCache->GetCachedAsset(GeoCacheHash.ToString()));
+		const FString PrefixedGeoCacheHash = UsdUtils::GetAssetHashPrefix(RootPrim, Context->bReuseIdenticalAssets) + GeoCacheHash.ToString();
+
+		UGeometryCache* GeometryCache = Cast<UGeometryCache>(Context->AssetCache->GetCachedAsset(PrefixedGeoCacheHash));
 
 		if (!GeometryCache)
 		{
@@ -414,7 +432,7 @@ namespace UsdGeometryCacheTranslatorImpl
 				Track->SetMatrixSamples(Mats, MatTimes);
 			}
 
-			Context->AssetCache->CacheAsset(GeoCacheHash.ToString(), GeometryCache);
+			Context->AssetCache->CacheAsset(PrefixedGeoCacheHash, GeometryCache);
 		}
 		else
 		{
@@ -774,8 +792,7 @@ void FGeometryCacheCreateAssetsTaskChain::SetupTasks()
 
 			bool bIsNew = true;
 			float StartTimeOffset = 0.0f;
-			const FString PrimPathString = PrimPath.GetString();
-			GeometryCache.Reset(UsdGeometryCacheTranslatorImpl::CreateGeometryCache(PrimPathString, LODIndexToMeshDescription[0], MeshPrimPaths, MaterialOffsets, Context, bIsNew, StartTimeOffset));
+			GeometryCache.Reset(UsdGeometryCacheTranslatorImpl::CreateGeometryCache(GetPrim(), LODIndexToMeshDescription[0], MeshPrimPaths, MaterialOffsets, Context, bIsNew, StartTimeOffset));
 
 			if (GeometryCache)
 			{
@@ -783,7 +800,7 @@ void FGeometryCacheCreateAssetsTaskChain::SetupTasks()
 				{
 					UserData->PrimvarToUVIndex = LODIndexToMaterialInfo[0].PrimvarToUVIndex;	// We use the same primvar mapping for all LODs
 					UserData->LayerStartOffsetSeconds = StartTimeOffset;
-					UserData->PrimPaths.AddUnique(PrimPathString);
+					UserData->PrimPaths.AddUnique(PrimPath.GetString());
 
 					if (Context->MetadataOptions.bCollectMetadata)
 					{
@@ -811,7 +828,8 @@ void FGeometryCacheCreateAssetsTaskChain::SetupTasks()
 						*Context->AssetCache.Get(),
 						Context->InfoCache.Get(),
 						Context->Time,
-						Context->ObjectFlags
+						Context->ObjectFlags,
+						Context->bReuseIdenticalAssets
 					);
 				}
 
@@ -923,7 +941,8 @@ USceneComponent* FUsdGeometryCacheTranslator::CreateComponents()
 					Context->ObjectFlags,
 					bAllowInterpretingLODs,
 					Context->RenderContext,
-					Context->MaterialPurpose
+					Context->MaterialPurpose,
+					Context->bReuseIdenticalAssets
 				);
 
 				// Check if the prim has the GroomBinding schema and setup the component and assets necessary to bind the groom to the GeometryCache
@@ -933,7 +952,8 @@ USceneComponent* FUsdGeometryCacheTranslator::CreateComponents()
 						GetPrim(),
 						*Context->AssetCache,
 						*Context->InfoCache,
-						Context->ObjectFlags
+						Context->ObjectFlags,
+						Context->bReuseIdenticalAssets
 					);
 
 					// For the groom binding to work, the GroomComponent must be a child of the SceneComponent
diff --git a/Engine/Plugins/Importers/USDImporter/Source/USDSchemas/Private/USDGroomTranslator.cpp b/Engine/Plugins/Importers/USDImporter/Source/USDSchemas/Private/USDGroomTranslator.cpp
--- a/Engine/Plugins/Importers/USDImporter/Source/USDSchemas/Private/USDGroomTranslator.cpp
+++ b/Engine/Plugins/Importers/USDImporter/Source/USDSchemas/Private/USDGroomTranslator.cpp
@@ -133,7 +133,7 @@ protected:
 	TStrongObjectPtr<UGroomImportOptions> ImportOptions;
 	TUniquePtr<FGroomCacheProcessor> GroomCacheProcessor;
 	FGroomAnimationInfo AnimInfo;
-	FSHAHash GroomCacheHash;
+	FString PrefixedGroomCacheHash;
 
 protected:
 	UE::FUsdPrim GetPrim() const { return Context->Stage.GetPrimAtPath(PrimPath); }
@@ -175,8 +175,10 @@ protected:
 
 				FSHAHash SHAHash = UE::UsdGroomTranslator::Private::ComputeHairDescriptionHash(HairDescription, ImportOptions->InterpolationSettings);
 
+				FString PrefixedAssetHash = UsdUtils::GetAssetHashPrefix(GetPrim(), Context->bReuseIdenticalAssets) + SHAHash.ToString();
+
 				const FString PrimPathString = PrimPath.GetString();
-				UGroomAsset* GroomAsset = Cast<UGroomAsset>(Context->AssetCache->GetCachedAsset(SHAHash.ToString()));
+				UGroomAsset* GroomAsset = Cast<UGroomAsset>(Context->AssetCache->GetCachedAsset(PrefixedAssetHash));
 				if (!GroomAsset)
 				{
 					FName AssetName = MakeUniqueObjectName(
@@ -196,7 +198,7 @@ protected:
 					GroomAsset = FHairStrandsImporter::ImportHair(HairImportContext, HairDescription, ExistingAsset);
 					if (GroomAsset)
 					{
-						Context->AssetCache->CacheAsset(SHAHash.ToString(), GroomAsset);
+						Context->AssetCache->CacheAsset(PrefixedAssetHash, GroomAsset);
 
 						if (UUsdAssetUserData* UserData = UsdUtils::GetOrCreateAssetUserData(GroomAsset))
 						{
@@ -251,10 +253,13 @@ protected:
 				SHA1.Update((const uint8*)&AnimInfo.Attributes, sizeof(AnimInfo.Attributes));
 				SHA1.Final();
 
-				SHA1.GetHash(GroomCacheHash.Hash);
+				FSHAHash Hash;
+				SHA1.GetHash(Hash.Hash);
+
+				PrefixedGroomCacheHash = UsdUtils::GetAssetHashPrefix(GetPrim(), Context->bReuseIdenticalAssets) + Hash.ToString();
 
 				bool bSuccess = true;
-				UGroomCache* GroomCache = Cast<UGroomCache>(Context->AssetCache->GetCachedAsset(GroomCacheHash.ToString()));
+				UGroomCache* GroomCache = Cast<UGroomCache>(Context->AssetCache->GetCachedAsset(PrefixedGroomCacheHash));
 				if (!GroomCache)
 				{
 					GroomCacheProcessor = MakeUnique<FGroomCacheProcessor>(EGroomCacheType::Strands, AnimInfo.Attributes);
@@ -323,7 +328,8 @@ protected:
 						GroomCacheProcessor->AddGroomSample(MoveTemp(HairGroupsData));
 					}
 				}
-				else if (Context->InfoCache)
+
+				if (Context->InfoCache && GroomCache)
 				{
 					Context->InfoCache->LinkAssetToPrim(PrimPath, GroomCache);
 				}
@@ -340,6 +346,17 @@ protected:
 		Then(ESchemaTranslationLaunchPolicy::Sync,
 			[this]() -> bool
 			{
+				// TEMP: This is a small trick to prevent two concurrent task chains from running into a hash collision in the asset cache.
+				// This is enough of a workaround because this is a Sync task, so we can guarantee only one of the competing task chains will be
+				// run at a time. Whichever wins gets to create the GroomCache, and the other will exit through this branch.
+				// It will likely be properly fixed before 5.4 is out, but check UE-201011 for more details.
+				UGroomCache* ExistingGroomCache = Cast<UGroomCache>(Context->AssetCache->GetCachedAsset(PrefixedGroomCacheHash));
+				if (ExistingGroomCache)
+				{
+					Context->InfoCache->LinkAssetToPrim(PrimPath, ExistingGroomCache);
+					return false;
+				}
+
 				const FString StrandsGroomCachePrimPath = UsdGroomTranslatorUtils::GetStrandsGroomCachePrimPath(PrimPath);
 				FHairImportContext HairImportContext(nullptr, GetTransientPackage(), nullptr, FName(), Context->ObjectFlags | RF_Public | RF_Transient);
 				FName UniqueName = MakeUniqueObjectName(
@@ -368,7 +385,7 @@ protected:
 						}
 					}
 
-					Context->AssetCache->CacheAsset(GroomCacheHash.ToString(), GroomCache);
+					Context->AssetCache->CacheAsset(PrefixedGroomCacheHash, GroomCache);
 					Context->InfoCache->LinkAssetToPrim(PrimPath, GroomCache);
 				}
 
diff --git a/Engine/Plugins/Importers/USDImporter/Source/USDSchemas/Private/USDGroomTranslatorUtils.cpp b/Engine/Plugins/Importers/USDImporter/Source/USDSchemas/Private/USDGroomTranslatorUtils.cpp
--- a/Engine/Plugins/Importers/USDImporter/Source/USDSchemas/Private/USDGroomTranslatorUtils.cpp
+++ b/Engine/Plugins/Importers/USDImporter/Source/USDSchemas/Private/USDGroomTranslatorUtils.cpp
@@ -195,7 +195,8 @@ namespace UsdGroomTranslatorUtils
 		const pxr::UsdPrim& Prim,
 		UUsdAssetCache2& AssetCache,
 		FUsdInfoCache& InfoCache,
-		EObjectFlags ObjectFlags
+		EObjectFlags ObjectFlags,
+		bool bReuseIdenticalAssets
 	)
 	{
 		// At this point, the prim (SkelRoot/Skeleton or GeomMesh) has already been checked to have the GroomBindingAPI,
@@ -237,9 +238,10 @@ namespace UsdGroomTranslatorUtils
 		Settings.TargetMesh = TargetMesh;
 		Settings.SourceMesh = GetGroomBindingSourceMesh(Prim, InfoCache, GroomBindingType);
 
+		const FString PrefixedBindingHash = UsdUtils::GetAssetHashPrefix(Prim, bReuseIdenticalAssets) + ComputeGroomBindingHash(Settings).ToString();
+
 		// Try to get the GroomBindingAsset from the cache
-		const FSHAHash SHAHash = ComputeGroomBindingHash(Settings);
-		UGroomBindingAsset* GroomBinding = Cast<UGroomBindingAsset>(AssetCache.GetCachedAsset(SHAHash.ToString()));
+		UGroomBindingAsset* GroomBinding = Cast<UGroomBindingAsset>(AssetCache.GetCachedAsset(PrefixedBindingHash));
 
 		const FString GroomBindingPath = FString::Printf(TEXT("%s_groombinding"), *PrimPath);
 		if (!GroomBinding)
@@ -248,7 +250,7 @@ namespace UsdGroomTranslatorUtils
 			GroomBinding = CreateGroomBindingAsset(GroomBindingPath, Settings, ObjectFlags);
 			if (GroomBinding)
 			{
-				AssetCache.CacheAsset(SHAHash.ToString(), GroomBinding);
+				AssetCache.CacheAsset(PrefixedBindingHash, GroomBinding);
 			}
 		}
 
diff --git a/Engine/Plugins/Importers/USDImporter/Source/USDSchemas/Private/USDGroomTranslatorUtils.h b/Engine/Plugins/Importers/USDImporter/Source/USDSchemas/Private/USDGroomTranslatorUtils.h
--- a/Engine/Plugins/Importers/USDImporter/Source/USDSchemas/Private/USDGroomTranslatorUtils.h
+++ b/Engine/Plugins/Importers/USDImporter/Source/USDSchemas/Private/USDGroomTranslatorUtils.h
@@ -27,7 +27,7 @@ namespace UE
 namespace UsdGroomTranslatorUtils
 {
 	// Create the groom binding asset for the given Prim with GroomBindingAPI and cache it in the AssetCache
-	void CreateGroomBindingAsset(const pxr::UsdPrim& Prim, UUsdAssetCache2& AssetCache, FUsdInfoCache& InfoCache, EObjectFlags ObjectFlags);
+	void CreateGroomBindingAsset(const pxr::UsdPrim& Prim, UUsdAssetCache2& AssetCache, FUsdInfoCache& InfoCache, EObjectFlags ObjectFlags, bool bReuseIdenticalAssets);
 
 	// Set the groom asset targeted by the given prim with GroomBindingAPI on the SceneComponent (with a GroomComponent as a direct child)
 	void SetGroomFromPrim(const pxr::UsdPrim& Prim, const FUsdInfoCache& InfoCache, USceneComponent* SceneComponent);
diff --git a/Engine/Plugins/Importers/USDImporter/Source/USDSchemas/Private/USDLuxLightTranslator.cpp b/Engine/Plugins/Importers/USDImporter/Source/USDSchemas/Private/USDLuxLightTranslator.cpp
--- a/Engine/Plugins/Importers/USDImporter/Source/USDSchemas/Private/USDLuxLightTranslator.cpp
+++ b/Engine/Plugins/Importers/USDImporter/Source/USDSchemas/Private/USDLuxLightTranslator.cpp
@@ -101,7 +101,7 @@ void FUsdLuxLightTranslator::UpdateComponents( USceneComponent* SceneComponent )
 	}
 	else if ( USkyLightComponent* SkyLightComponent = Cast< USkyLightComponent >( SceneComponent ) )
 	{
-		UsdToUnreal::ConvertDomeLight( Prim, *SkyLightComponent, Context->AssetCache.Get() );
+		UsdToUnreal::ConvertDomeLight(Prim, *SkyLightComponent, Context->AssetCache.Get(), Context->bReuseIdenticalAssets);
 		SkyLightComponent->Mobility = EComponentMobility::Movable; // We won't bake geometry in the sky light so it needs to be movable
 	}
 
diff --git a/Engine/Plugins/Importers/USDImporter/Source/USDSchemas/Private/USDShadeMaterialTranslator.cpp b/Engine/Plugins/Importers/USDImporter/Source/USDSchemas/Private/USDShadeMaterialTranslator.cpp
--- a/Engine/Plugins/Importers/USDImporter/Source/USDSchemas/Private/USDShadeMaterialTranslator.cpp
+++ b/Engine/Plugins/Importers/USDImporter/Source/USDSchemas/Private/USDShadeMaterialTranslator.cpp
@@ -267,13 +267,14 @@ void FUsdShadeMaterialTranslator::CreateAssets()
 		}
 	}
 
-	FString MaterialHashString = UsdUtils::HashShadeMaterial( ShadeMaterial, RenderContextToken ).ToString();
+	const FString PrefixedMaterialHash = UsdUtils::GetAssetHashPrefix(GetPrim(), Context->bReuseIdenticalAssets)
+										 + UsdUtils::HashShadeMaterial(ShadeMaterial, RenderContextToken).ToString();
 
 	UMaterialInterface* ConvertedMaterial = nullptr;
 
-	if ( Context->AssetCache )
+	if (Context->AssetCache)
 	{
-		ConvertedMaterial = Cast< UMaterialInterface >( Context->AssetCache->GetCachedAsset( MaterialHashString ) );
+		ConvertedMaterial = Cast<UMaterialInterface>(Context->AssetCache->GetCachedAsset(PrefixedMaterialHash));
 	}
 
 	if ( !ConvertedMaterial )
@@ -299,7 +300,8 @@ void FUsdShadeMaterialTranslator::CreateAssets()
 					ShadeMaterial,
 					*NewMaterial,
 					Context->AssetCache.Get(),
-					*Context->RenderContext.ToString()
+					*Context->RenderContext.ToString(),
+					Context->bReuseIdenticalAssets
 				);
 				if ( !bSuccess )
 				{
@@ -395,7 +397,13 @@ void FUsdShadeMaterialTranslator::CreateAssets()
 
 					NewMaterial->SetFlags(RF_Transient);
 
-					if (UsdToUnreal::ConvertMaterial(ShadeMaterial, *NewMaterial, Context->AssetCache.Get(), *Context->RenderContext.ToString()))
+					if (UsdToUnreal::ConvertMaterial(
+							ShadeMaterial,
+							*NewMaterial,
+							Context->AssetCache.Get(),
+							*Context->RenderContext.ToString(),
+							Context->bReuseIdenticalAssets
+						))
 					{
 						TSet<UTexture*> VTTextures;
 						TSet<UTexture*> NonVTTextures;
@@ -418,7 +426,7 @@ void FUsdShadeMaterialTranslator::CreateAssets()
 						// is what will actually swap our reference with a VT one if needed
 						if ( Context->AssetCache && Context->InfoCache )
 						{
-							Context->AssetCache->CacheAsset(MaterialHashString, NewMaterial);
+							Context->AssetCache->CacheAsset(PrefixedMaterialHash, NewMaterial);
 							Context->InfoCache->LinkAssetToPrim(PrimPath, NewMaterial);
 						}
 						for ( UTexture* Texture : VTTextures.Union( NonVTTextures ) )
@@ -437,14 +445,14 @@ void FUsdShadeMaterialTranslator::CreateAssets()
 
 						// We must go through the cache to fetch our result material here as UpgradeMaterialsAndTexturesToVT
 						// may have created a new MID for this material with a VT reference
-						ConvertedMaterial = Cast<UMaterialInterface>( Context->AssetCache->GetCachedAsset( MaterialHashString ) );
+						ConvertedMaterial = Cast<UMaterialInterface>(Context->AssetCache->GetCachedAsset(PrefixedMaterialHash));
 					}
 				}
 			}
 		}
 	}
 
-	PostImportMaterial(MaterialHashString, ConvertedMaterial);
+	PostImportMaterial(PrefixedMaterialHash, ConvertedMaterial);
 }
 
 bool FUsdShadeMaterialTranslator::CollapsesChildren( ECollapsingType CollapsingType ) const
@@ -457,7 +465,7 @@ bool FUsdShadeMaterialTranslator::CanBeCollapsed( ECollapsingType CollapsingType
 	return false;
 }
 
-void FUsdShadeMaterialTranslator::PostImportMaterial(const FString& MaterialHash, UMaterialInterface* ImportedMaterial)
+void FUsdShadeMaterialTranslator::PostImportMaterial(const FString& PrefixedMaterialHash, UMaterialInterface* ImportedMaterial)
 {
 	if (!ImportedMaterial || !Context->InfoCache || !Context->AssetCache)
 	{
@@ -482,7 +490,7 @@ void FUsdShadeMaterialTranslator::PostImportMaterial(const FString& MaterialHash
 
 	// Note that this needs to run even if we found this material in the asset cache already, otherwise we won't
 	// re-register the prim asset links when we reload a stage
-	Context->AssetCache->CacheAsset(MaterialHash, ImportedMaterial);
+	Context->AssetCache->CacheAsset(PrefixedMaterialHash, ImportedMaterial);
 	Context->InfoCache->LinkAssetToPrim(PrimPath, ImportedMaterial);
 
 	// Also link the textures to the same material prim.
@@ -519,7 +527,8 @@ void FUsdShadeMaterialTranslator::PostImportMaterial(const FString& MaterialHash
 					FilePath = Texture->GetName();
 				}
 
-				const FString TextureHash = UsdUtils::GetTextureHash(
+				const FString HashPrefix = UsdUtils::GetAssetHashPrefix(GetPrim(), Context->bReuseIdenticalAssets);
+				const FString PrefixedTextureHash = HashPrefix + UsdUtils::GetTextureHash(
 					FilePath,
 					Texture->SRGB,
 					Texture->CompressionSettings,
@@ -534,11 +543,11 @@ void FUsdShadeMaterialTranslator::PostImportMaterial(const FString& MaterialHash
 				// unfortunately recreate identical textures, and if we try to store them in here we'll run into a hash
 				// collision. These should be rare, however, and require repeatedly e.g. updating the Material prims
 				// generated by MaterialX or updating the MaterialX import options
-				UTexture* ExistingTexture = Cast<UTexture>(Context->AssetCache->GetCachedAsset(TextureHash));
+				UTexture* ExistingTexture = Cast<UTexture>(Context->AssetCache->GetCachedAsset(PrefixedTextureHash));
 				if (!ExistingTexture)
 				{
 					Texture->SetFlags(RF_Transient);
-					Context->AssetCache->CacheAsset(TextureHash, Texture);
+					Context->AssetCache->CacheAsset(PrefixedTextureHash, Texture);
 				}
 			}
 
diff --git a/Engine/Plugins/Importers/USDImporter/Source/USDSchemas/Private/USDSkelRootTranslator.cpp b/Engine/Plugins/Importers/USDImporter/Source/USDSchemas/Private/USDSkelRootTranslator.cpp
--- a/Engine/Plugins/Importers/USDImporter/Source/USDSchemas/Private/USDSkelRootTranslator.cpp
+++ b/Engine/Plugins/Importers/USDImporter/Source/USDSchemas/Private/USDSkelRootTranslator.cpp
@@ -80,8 +80,10 @@ namespace UsdSkelRootTranslatorImpl
 		USkeletalMesh* SkeletalMesh,
 		UUsdAssetCache2& AssetCache,
 		FUsdInfoCache& InfoCache,
-		float Time, EObjectFlags Flags,
-		bool bSkeletalMeshHasMorphTargets
+		float Time,
+		EObjectFlags Flags,
+		bool bSkeletalMeshHasMorphTargets,
+		bool bReuseIdenticalAssets
 	)
 	{
 		TRACE_CPUPROFILER_EVENT_SCOPE( UsdSkelRootTranslatorImpl::ProcessMaterials );
@@ -108,7 +110,8 @@ namespace UsdSkelRootTranslatorImpl
 			LODIndexToMaterialInfo,
 			AssetCache,
 			InfoCache,
-			Flags
+			Flags,
+			bReuseIdenticalAssets
 		);
 
 		bool bMaterialsHaveChanged = false;
@@ -653,7 +656,8 @@ namespace UsdSkelRootTranslatorImpl
 		EObjectFlags Flags,
 		bool bInterpretLODs,
 		const FName& RenderContext,
-		const FName& MaterialPurpose
+		const FName& MaterialPurpose,
+		bool bReuseIdenticalAssets
 	)
 	{
 		FScopedUsdAllocs Allocs;
@@ -802,7 +806,8 @@ namespace UsdSkelRootTranslatorImpl
 			LODIndexToAssignments,
 			AssetCache,
 			InfoCache,
-			Flags
+			Flags,
+			bReuseIdenticalAssets
 		);
 
 		// Compare resolved materials with existing assignments, and create overrides if we need to
@@ -973,11 +978,11 @@ namespace UsdSkelRootTranslatorImpl
 			SHA1.UpdateWithString(*PrimPath, PrimPath.Len());
 			SHA1.Final();
 			SHA1.GetHash(&Hash.Hash[0]);
-			const FString CacheKey = Hash.ToString();
+			const FString PrefixedAnimBPHash = UsdUtils::GetAssetHashPrefix(Prim, Context.bReuseIdenticalAssets) + Hash.ToString();
 
 			// Check if we can find an AnimBP for this prim in the asset cache (useful when doing Action->Import)
 			bool bReusedAnimBP = false;
-			if ( UAnimBlueprint* CachedAnimBP = Cast< UAnimBlueprint >( Context.AssetCache->GetCachedAsset( CacheKey ) ) )
+			if (UAnimBlueprint* CachedAnimBP = Cast<UAnimBlueprint>(Context.AssetCache->GetCachedAsset(PrefixedAnimBPHash)))
 			{
 				if ( CachedAnimBP->TargetSkeleton == Skeleton )
 				{
@@ -1012,7 +1017,7 @@ namespace UsdSkelRootTranslatorImpl
 
 				bNeedRecompile = true;
 
-				Context.AssetCache->CacheAsset( CacheKey, AnimBP );
+				Context.AssetCache->CacheAsset( PrefixedAnimBPHash, AnimBP );
 			}
 		}
 		// Path is pointing to an existing, persistent AnimBP
@@ -1250,7 +1255,7 @@ namespace UsdSkelRootTranslatorImpl
 		// Note that we want this to be case insensitive so that our UMorphTarget FNames are unique not only due to case differences
 		TSet<FString> UsedMorphTargetNames;
 		TUsdStore< pxr::UsdSkelCache > SkeletonCache;
-		FString SkeletalMeshHashString;
+		FString PrefixedSkelMeshHash;
 
 		// Don't keep a live reference to the prim because other translators may mutate the stage in an ExclusiveSync translation step, invalidating the reference
 		UE::FUsdPrim GetPrim() const { return Context->Stage.GetPrimAtPath( PrimPath ); }
@@ -1336,9 +1341,9 @@ namespace UsdSkelRootTranslatorImpl
 
 				FString SkelRootPath = PrimPath.GetString();
 				FSHAHash SkeletalMeshHash = UsdSkelRootTranslatorImpl::ComputeSHAHash( LODIndexToSkeletalMeshImportData, SkeletonBones, BlendShapes );
-				SkeletalMeshHashString = SkeletalMeshHash.ToString();
+				PrefixedSkelMeshHash = UsdUtils::GetAssetHashPrefix(GetPrim(), Context->bReuseIdenticalAssets) + SkeletalMeshHash.ToString();
 
-				USkeletalMesh* SkeletalMesh = Cast< USkeletalMesh >( Context->AssetCache->GetCachedAsset( SkeletalMeshHashString ) );
+				USkeletalMesh* SkeletalMesh = Cast<USkeletalMesh>(Context->AssetCache->GetCachedAsset(PrefixedSkelMeshHash));
 
 				bool bIsNew = false;
 				if ( !SkeletalMesh )
@@ -1377,7 +1382,8 @@ namespace UsdSkelRootTranslatorImpl
 							*Context->InfoCache.Get(),
 							Context->Time,
 							Context->ObjectFlags,
-							NewBlendShapes.Num() > 0
+							NewBlendShapes.Num() > 0,
+							Context->bReuseIdenticalAssets
 						);
 
 						if (bMaterialsHaveChanged)
@@ -1386,8 +1392,8 @@ namespace UsdSkelRootTranslatorImpl
 							SkeletalMesh->UpdateUVChannelData(bRebuildAll);
 						}
 
-						Context->AssetCache->CacheAsset( SkeletalMeshHashString, SkeletalMesh );
-						Context->AssetCache->CacheAsset( SkeletalMeshHashString + TEXT( "_Skeleton" ), SkeletalMesh->GetSkeleton() );
+						Context->AssetCache->CacheAsset( PrefixedSkelMeshHash, SkeletalMesh );
+						Context->AssetCache->CacheAsset( PrefixedSkelMeshHash + TEXT( "_Skeleton" ), SkeletalMesh->GetSkeleton() );
 
 						// The PreviewSkeletalMesh property on the Skeleton is a soft object path, that is set within
 						// UsdToUnreal::GetSkeletalMeshFromImportData before the SkeletalMesh is part of the cache. When we cache the
@@ -1403,22 +1409,24 @@ namespace UsdSkelRootTranslatorImpl
 
 					if ( bGeneratePhysicsAssets )
 					{
-						if ( !SkeletalMesh->GetPhysicsAsset() )
+						UPhysicsAsset* PhysicsAsset = SkeletalMesh->GetPhysicsAsset();
+						if (!PhysicsAsset)
 						{
-							UPhysicsAsset* PhysicsAsset = UsdSkelRootTranslatorImpl::GenerateAndAssignPhysicsAsset(
+							PhysicsAsset = UsdSkelRootTranslatorImpl::GenerateAndAssignPhysicsAsset(
 								SkeletalMesh,
 								Context->ObjectFlags
 							);
 
-							if ( PhysicsAsset )
+							if (PhysicsAsset)
 							{
-								Context->AssetCache->CacheAsset(
-									SkeletalMeshHashString + TEXT( "_PhysicsAsset" ),
-									PhysicsAsset
-								);
-								Context->InfoCache->LinkAssetToPrim(PrimPath, PhysicsAsset);
+								Context->AssetCache->CacheAsset(PrefixedSkelMeshHash + TEXT("_PhysicsAsset"), PhysicsAsset);
 							}
 						}
+
+						if (PhysicsAsset)
+						{
+							Context->InfoCache->LinkAssetToPrim(PrimPath, PhysicsAsset);
+						}
 					}
 					else
 					{
@@ -1560,10 +1568,10 @@ namespace UsdSkelRootTranslatorImpl
 						FSHAHash Hash = UsdSkelRootTranslatorImpl::ComputeSHAHash(
 							SkelQuery,
 							RootMotionPrim,
-							SkeletalMeshHashString
+							PrefixedSkelMeshHash
 						);
-						FString HashString = Hash.ToString();
-						UAnimSequence* AnimSequence = Cast< UAnimSequence >( Context->AssetCache->GetCachedAsset( HashString ) );
+						FString PrefixedSkelAnimHash = UsdUtils::GetAssetHashPrefix(SkelAnimationPrim, Context->bReuseIdenticalAssets) + Hash.ToString();
+						UAnimSequence* AnimSequence = Cast<UAnimSequence>(Context->AssetCache->GetCachedAsset(PrefixedSkelAnimHash));
 
 						if ( !AnimSequence || AnimSequence->GetSkeleton() != SkeletalMesh->GetSkeleton() )
 						{
@@ -1609,7 +1617,7 @@ namespace UsdSkelRootTranslatorImpl
 
 								AnimSequence->AddAssetUserData(UserData);
 
-								Context->AssetCache->CacheAsset( HashString, AnimSequence );
+								Context->AssetCache->CacheAsset( PrefixedSkelAnimHash, AnimSequence );
 							}
 							else
 							{
@@ -1680,7 +1688,8 @@ USceneComponent* FUsdSkelRootTranslator::CreateComponents()
 			GetPrim(),
 			*Context->AssetCache,
 			*Context->InfoCache,
-			Context->ObjectFlags
+			Context->ObjectFlags,
+			Context->bReuseIdenticalAssets
 		);
 
 		// For the groom binding to work, the GroomComponent must be a child of the SceneComponent
@@ -1786,7 +1795,8 @@ void FUsdSkelRootTranslator::UpdateComponents( USceneComponent* SceneComponent )
 				Context->ObjectFlags,
 				Context->bAllowInterpretingLODs,
 				Context->RenderContext,
-				Context->MaterialPurpose
+				Context->MaterialPurpose,
+				Context->bReuseIdenticalAssets
 			);
 		}
 	}
diff --git a/Engine/Plugins/Importers/USDImporter/Source/USDSchemas/Private/USDSkelSkeletonTranslator.cpp b/Engine/Plugins/Importers/USDImporter/Source/USDSchemas/Private/USDSkelSkeletonTranslator.cpp
--- a/Engine/Plugins/Importers/USDImporter/Source/USDSchemas/Private/USDSkelSkeletonTranslator.cpp
+++ b/Engine/Plugins/Importers/USDImporter/Source/USDSchemas/Private/USDSkelSkeletonTranslator.cpp
@@ -80,7 +80,8 @@ namespace UsdSkelSkeletonTranslatorImpl
 		FUsdInfoCache& InfoCache,
 		float Time,
 		EObjectFlags Flags,
-		bool bSkeletalMeshHasMorphTargets
+		bool bSkeletalMeshHasMorphTargets,
+		bool bReuseIdenticalAssets
 	)
 	{
 		TRACE_CPUPROFILER_EVENT_SCOPE(UsdSkelSkeletonTranslatorImpl::ProcessMaterials);
@@ -104,8 +105,14 @@ namespace UsdSkelSkeletonTranslatorImpl
 			*UsdToUnreal::ConvertPath(UsdPrim.GetPrimPath())
 		);
 
-		TMap<const UsdUtils::FUsdPrimMaterialSlot*, UMaterialInterface*>
-			ResolvedMaterials = MeshTranslationImpl::ResolveMaterialAssignmentInfo(UsdPrim, LODIndexToMaterialInfo, AssetCache, InfoCache, Flags);
+		TMap<const UsdUtils::FUsdPrimMaterialSlot*, UMaterialInterface*> ResolvedMaterials = MeshTranslationImpl::ResolveMaterialAssignmentInfo(
+			UsdPrim,
+			LODIndexToMaterialInfo,
+			AssetCache,
+			InfoCache,
+			Flags,
+			bReuseIdenticalAssets
+		);
 
 		bool bMaterialsHaveChanged = false;
 
@@ -715,7 +722,8 @@ namespace UsdSkelSkeletonTranslatorImpl
 		bool bInterpretLODs,
 		const FName& RenderContext,
 		const FName& MaterialPurpose,
-		const EUsdPurpose PurposesToLoad
+		const EUsdPurpose PurposesToLoad,
+		bool bReuseIdenticalAssets
 	)
 	{
 		FScopedUsdAllocs Allocs;
@@ -865,7 +873,8 @@ namespace UsdSkelSkeletonTranslatorImpl
 			LODIndexToAssignments,
 			AssetCache,
 			InfoCache,
-			Flags
+			Flags,
+			bReuseIdenticalAssets
 		);
 
 		// Compare resolved materials with existing assignments, and create overrides if we need to
@@ -1045,11 +1054,11 @@ namespace UsdSkelSkeletonTranslatorImpl
 			SHA1.UpdateWithString(*PrimPath, PrimPath.Len());
 			SHA1.Final();
 			SHA1.GetHash(&Hash.Hash[0]);
-			const FString CacheKey = Hash.ToString();
+			const FString PrefixedAnimBPHash = UsdUtils::GetAssetHashPrefix(SkeletonPrim, Context.bReuseIdenticalAssets) + Hash.ToString();
 
 			// Check if we can find an AnimBP for this prim in the asset cache (useful when doing Action->Import)
 			bool bReusedAnimBP = false;
-			if (UAnimBlueprint* CachedAnimBP = Cast<UAnimBlueprint>(Context.AssetCache->GetCachedAsset(CacheKey)))
+			if (UAnimBlueprint* CachedAnimBP = Cast<UAnimBlueprint>(Context.AssetCache->GetCachedAsset(PrefixedAnimBPHash)))
 			{
 				if (CachedAnimBP->TargetSkeleton == Skeleton)
 				{
@@ -1082,7 +1091,7 @@ namespace UsdSkelSkeletonTranslatorImpl
 
 				bNeedRecompile = true;
 
-				Context.AssetCache->CacheAsset(CacheKey, AnimBP);
+				Context.AssetCache->CacheAsset(PrefixedAnimBPHash, AnimBP);
 			}
 		}
 		// Path is pointing to an existing, persistent AnimBP
@@ -1348,7 +1357,7 @@ namespace UsdSkelSkeletonTranslatorImpl
 		TUsdStore<pxr::UsdSkelCache> SkelCache;
 		TUsdStore<pxr::UsdSkelRoot> ClosestParentSkelRoot;
 		TUsdStore<pxr::UsdSkelSkeletonQuery> SkeletonQuery;
-		FString SkeletalMeshHashString;
+		FString PrefixedSkelMeshHash;
 
 		// Don't keep a live reference to the prim because other translators may mutate the stage in an ExclusiveSync translation step, invalidating
 		// the reference
@@ -1402,6 +1411,33 @@ namespace UsdSkelSkeletonTranslatorImpl
 		SetupTasks();
 	}
 
+	// Right now parsing LODs involves flipping through variants, which invalidates some prims and references.
+	// The SkelSkeletonTranslator is especially vulnerable to this because the SkeletonBinding contains
+	// skinning queries that are all invalidated when we flip through variants, and the task chain holds on
+	// to the same SkeletonBinding throughout the entire chain...
+	// Here we'll refresh those references if needed.
+	// TODO: Find better way of handling LODs that doesn't require this mechanism.
+	void RefreshSkelReferencesIfNeeded(
+		const pxr::UsdSkelRoot& InSkelRootPrim,
+		const pxr::UsdSkelSkeleton& InSkeletonPrim,
+		pxr::UsdSkelCache& InOutSkelCache,
+		pxr::UsdSkelBinding& InOutSkelBinding,
+		pxr::UsdSkelSkeletonQuery& InOutSkeletonQuery
+	)
+	{
+		// If we still have valid skinning queries we know we still have valid references, as those are the first to break
+		for (const pxr::UsdSkelSkinningQuery& SkinningQuery : InOutSkelBinding.GetSkinningTargets())
+		{
+			if (pxr::UsdGeomMesh SkinningMesh = pxr::UsdGeomMesh(SkinningQuery.GetPrim()))
+			{
+				return;
+			}
+		}
+
+		InOutSkelCache.Populate(InSkelRootPrim, pxr::UsdTraverseInstanceProxies());
+		ensure(UsdUtils::GetSkelQueries(InSkelRootPrim, InSkeletonPrim, InOutSkelBinding, InOutSkeletonQuery, &InOutSkelCache));
+	}
+
 	void FSkelSkeletonCreateAssetsTaskChain::SetupTasks()
 	{
 		// To parse all LODs we need to actively switch variant sets to other variants (triggering prim loading/unloading and notices),
@@ -1411,69 +1447,59 @@ namespace UsdSkelSkeletonTranslatorImpl
 																	 : ESchemaTranslationLaunchPolicy::Async;
 
 		// Create SkeletalMeshImportData (Async or ExclusiveSync)
-		Do(LaunchPolicy,
-		   [this, bTryLODParsing]()
-		   {
-			   // No point in importing blend shapes if the import context doesn't want them
-			   UsdUtils::FBlendShapeMap* OutBlendShapes = Context->BlendShapesByPath ? &NewBlendShapes : nullptr;
-
-			   pxr::TfToken RenderContextToken = pxr::UsdShadeTokens->universalRenderContext;
-			   if (!Context->RenderContext.IsNone())
-			   {
-				   RenderContextToken = UnrealToUsd::ConvertToken(*Context->RenderContext.ToString()).Get();
-			   }
-
-			   pxr::TfToken MaterialPurposeToken = pxr::UsdShadeTokens->allPurpose;
-			   if (!Context->MaterialPurpose.IsNone())
-			   {
-				   MaterialPurposeToken = UnrealToUsd::ConvertToken(*Context->MaterialPurpose.ToString()).Get();
-			   }
-
-			   UsdToUnreal::FUsdMeshConversionOptions Options;
-			   Options.TimeCode = Context->Time;
-			   Options.RenderContext = RenderContextToken;
-			   Options.MaterialPurpose = MaterialPurposeToken;
-			   Options.SubdivisionLevel = Context->SubdivisionLevel;
-			   Options.PurposesToLoad = Context->PurposesToLoad;
-			   Options.bMergeIdenticalMaterialSlots = Context->bMergeIdenticalMaterialSlots;
-
-			   bool bContinueTaskChain = UsdSkelSkeletonTranslatorImpl::LoadAllSkeletalData(
-				   SkeletonBinding.Get(),
-				   SkelCache.Get(),
-				   LODIndexToSkeletalMeshImportData,
-				   LODIndexToMaterialInfo,
-				   LODMetadata,
-				   SkeletonBones,
-				   SkeletonName,
-				   OutBlendShapes,
-				   UsedMorphTargetNames,
-				   bTryLODParsing,
-				   Options,
-				   Context->MetadataOptions
-			   );
-
-			   // If we parsed LODs we could potentially have invalidated our references to our queries,
-			   // so let's refresh them before the downstream tasks try using them.
-			   if (bContinueTaskChain && bTryLODParsing)
-			   {
-				   // We are forced to repopulate here unfortunately, or else the queries/SkelCache will actually happily hold on to and return their
-				   // internal, invalid prim references...
-				   // We can't just keep prim paths of things instead (paths to skinned prims, paths to SkelAnimation prim, etc.) because parsing
-				   // these prims involves using those query objects to compute joint transforms and etc., and we can't create them ourselves.
-				   // This should only be temporary, as I believe we'll eventually have a cleverer way of handling LODs that doesn't invalidate
-				   // everything (e.g. open a separate stage with a population mask to only parse the LOD mesh, etc.)
-				   SkelCache.Get().Populate(ClosestParentSkelRoot.Get(), pxr::UsdTraverseInstanceProxies());
-				   ensure(UsdUtils::GetSkelQueries(
-					   ClosestParentSkelRoot.Get(),
-					   pxr::UsdSkelSkeleton{GetSkeletonPrim()},
-					   SkeletonBinding.Get(),
-					   SkeletonQuery.Get(),
-					   &SkelCache.Get()
-				   ));
-			   }
-
-			   return bContinueTaskChain;
-		   });
+		Do(
+			LaunchPolicy,
+			[this, bTryLODParsing]()
+			{
+				RefreshSkelReferencesIfNeeded(
+					ClosestParentSkelRoot.Get(),
+					pxr::UsdSkelSkeleton{GetSkeletonPrim()},
+					SkelCache.Get(),
+					SkeletonBinding.Get(),
+					SkeletonQuery.Get()
+				);
+
+				// No point in importing blend shapes if the import context doesn't want them
+				UsdUtils::FBlendShapeMap* OutBlendShapes = Context->BlendShapesByPath ? &NewBlendShapes : nullptr;
+
+				pxr::TfToken RenderContextToken = pxr::UsdShadeTokens->universalRenderContext;
+				if (!Context->RenderContext.IsNone())
+				{
+					RenderContextToken = UnrealToUsd::ConvertToken(*Context->RenderContext.ToString()).Get();
+				}
+
+				pxr::TfToken MaterialPurposeToken = pxr::UsdShadeTokens->allPurpose;
+				if (!Context->MaterialPurpose.IsNone())
+				{
+					MaterialPurposeToken = UnrealToUsd::ConvertToken(*Context->MaterialPurpose.ToString()).Get();
+				}
+
+				UsdToUnreal::FUsdMeshConversionOptions Options;
+				Options.TimeCode = Context->Time;
+				Options.RenderContext = RenderContextToken;
+				Options.MaterialPurpose = MaterialPurposeToken;
+				Options.SubdivisionLevel = Context->SubdivisionLevel;
+				Options.PurposesToLoad = Context->PurposesToLoad;
+				Options.bMergeIdenticalMaterialSlots = Context->bMergeIdenticalMaterialSlots;
+
+				const bool bContinueTaskChain = UsdSkelSkeletonTranslatorImpl::LoadAllSkeletalData(
+					SkeletonBinding.Get(),
+					SkelCache.Get(),
+					LODIndexToSkeletalMeshImportData,
+					LODIndexToMaterialInfo,
+					LODMetadata,
+					SkeletonBones,
+					SkeletonName,
+					OutBlendShapes,
+					UsedMorphTargetNames,
+					bTryLODParsing,
+					Options,
+					Context->MetadataOptions
+				);
+
+				return bContinueTaskChain;
+			}
+		);
 
 		// Create USkeletalMesh (Main thread)
 		Then(
@@ -1485,6 +1511,17 @@ namespace UsdSkelSkeletonTranslatorImpl
 					return false;
 				}
 
+				// We may have invalidated references with the previous task if it parsed LODs, so refresh them if needed.
+				// We'll assume that it's unlikely that those would be invalidated past this point though, as only the previous
+				// task is capable of invalidating them, and it is an ExclusiveSync task
+				RefreshSkelReferencesIfNeeded(
+					ClosestParentSkelRoot.Get(),
+					pxr::UsdSkelSkeleton{GetSkeletonPrim()},
+					SkelCache.Get(),
+					SkeletonBinding.Get(),
+					SkeletonQuery.Get()
+				);
+
 				UsdUtils::FBlendShapeMap* BlendShapes = Context->BlendShapesByPath ? &NewBlendShapes : nullptr;
 
 				FSHAHash SkeletalMeshHash = UsdSkelSkeletonTranslatorImpl::ComputeSHAHash(
@@ -1492,9 +1529,9 @@ namespace UsdSkelSkeletonTranslatorImpl
 					SkeletonBones,
 					BlendShapes
 				);
-				SkeletalMeshHashString = SkeletalMeshHash.ToString();
+				PrefixedSkelMeshHash = UsdUtils::GetAssetHashPrefix(GetSkeletonPrim(), Context->bReuseIdenticalAssets) + SkeletalMeshHash.ToString();
 
-				USkeletalMesh* SkeletalMesh = Cast<USkeletalMesh>(Context->AssetCache->GetCachedAsset(SkeletalMeshHashString));
+				USkeletalMesh* SkeletalMesh = Cast<USkeletalMesh>(Context->AssetCache->GetCachedAsset(PrefixedSkelMeshHash));
 
 				bool bIsNew = false;
 				if (!SkeletalMesh)
@@ -1569,7 +1606,8 @@ namespace UsdSkelSkeletonTranslatorImpl
 							*Context->InfoCache.Get(),
 							Context->Time,
 							Context->ObjectFlags,
-							NewBlendShapes.Num() > 0
+							NewBlendShapes.Num() > 0,
+							Context->bReuseIdenticalAssets
 						);
 
 						if (bMaterialsHaveChanged)
@@ -1578,8 +1616,8 @@ namespace UsdSkelSkeletonTranslatorImpl
 							SkeletalMesh->UpdateUVChannelData(bRebuildAll);
 						}
 
-						Context->AssetCache->CacheAsset(SkeletalMeshHashString, SkeletalMesh);
-						Context->AssetCache->CacheAsset(SkeletalMeshHashString + TEXT("_Skeleton"), SkeletalMesh->GetSkeleton());
+						Context->AssetCache->CacheAsset(PrefixedSkelMeshHash, SkeletalMesh);
+						Context->AssetCache->CacheAsset(PrefixedSkelMeshHash + TEXT("_Skeleton"), SkeletalMesh->GetSkeleton());
 
 						// The PreviewSkeletalMesh property on the Skeleton is a soft object path, that is set within
 						// UsdToUnreal::GetSkeletalMeshFromImportData before the SkeletalMesh is part of the cache. When we cache the
@@ -1592,19 +1630,24 @@ namespace UsdSkelSkeletonTranslatorImpl
 
 					if (bGeneratePhysicsAssets)
 					{
-						if (!SkeletalMesh->GetPhysicsAsset())
+						UPhysicsAsset* PhysicsAsset = SkeletalMesh->GetPhysicsAsset();
+						if (!PhysicsAsset)
 						{
-							UPhysicsAsset* PhysicsAsset = UsdSkelSkeletonTranslatorImpl::GenerateAndAssignPhysicsAsset(
+							PhysicsAsset = UsdSkelSkeletonTranslatorImpl::GenerateAndAssignPhysicsAsset(
 								SkeletalMesh,
 								Context->ObjectFlags
 							);
 
 							if (PhysicsAsset)
 							{
-								Context->AssetCache->CacheAsset(SkeletalMeshHashString + TEXT("_PhysicsAsset"), PhysicsAsset);
-								Context->InfoCache->LinkAssetToPrim(SkeletonPrimPath, PhysicsAsset);
+								Context->AssetCache->CacheAsset(PrefixedSkelMeshHash + TEXT("_PhysicsAsset"), PhysicsAsset);
 							}
 						}
+
+						if (PhysicsAsset)
+						{
+							Context->InfoCache->LinkAssetToPrim(SkeletonPrimPath, PhysicsAsset);
+						}
 					}
 					else
 					{
@@ -1749,9 +1792,13 @@ namespace UsdSkelSkeletonTranslatorImpl
 					}
 				}
 
-				FSHAHash Hash = UsdSkelSkeletonTranslatorImpl::ComputeSHAHash(SkeletonQuery.Get(), RootMotionPrim, SkeletalMeshHashString);
-				FString HashString = Hash.ToString();
-				UAnimSequence* AnimSequence = Cast<UAnimSequence>(Context->AssetCache->GetCachedAsset(HashString));
+				FSHAHash Hash = UsdSkelSkeletonTranslatorImpl::ComputeSHAHash(
+					SkeletonQuery.Get(),
+					RootMotionPrim,
+					PrefixedSkelMeshHash
+				);
+				FString PrefixedSkelAnimHash = UsdUtils::GetAssetHashPrefix(SkelAnimationPrim, Context->bReuseIdenticalAssets) + Hash.ToString();
+				UAnimSequence* AnimSequence = Cast<UAnimSequence>(Context->AssetCache->GetCachedAsset(PrefixedSkelAnimHash));
 
 				if (!AnimSequence || AnimSequence->GetSkeleton() != SkeletalMesh->GetSkeleton())
 				{
@@ -1814,7 +1861,7 @@ namespace UsdSkelSkeletonTranslatorImpl
 							}
 						}
 
-						Context->AssetCache->CacheAsset(HashString, AnimSequence);
+						Context->AssetCache->CacheAsset(PrefixedSkelAnimHash, AnimSequence);
 					}
 					else
 					{
@@ -1893,7 +1940,13 @@ USceneComponent* FUsdSkelSkeletonTranslator::CreateComponents()
 
 		if (PrimWithSchema)
 		{
-			UsdGroomTranslatorUtils::CreateGroomBindingAsset(PrimWithSchema, *Context->AssetCache, *Context->InfoCache, Context->ObjectFlags);
+			UsdGroomTranslatorUtils::CreateGroomBindingAsset(
+				PrimWithSchema,
+				*Context->AssetCache,
+				*Context->InfoCache,
+				Context->ObjectFlags,
+				Context->bReuseIdenticalAssets
+			);
 
 			// For the groom binding to work, the GroomComponent must be a child of the SceneComponent
 			// so the Context ParentComponent is set to the SceneComponent temporarily
@@ -2040,7 +2093,8 @@ void FUsdSkelSkeletonTranslator::UpdateComponents(USceneComponent* SceneComponen
 				Context->bAllowInterpretingLODs,
 				Context->RenderContext,
 				Context->MaterialPurpose,
-				Context->PurposesToLoad
+				Context->PurposesToLoad,
+				Context->bReuseIdenticalAssets
 			);
 		}
 	}
diff --git a/Engine/Plugins/Importers/USDImporter/Source/USDSchemas/Public/USDSchemaTranslator.h b/Engine/Plugins/Importers/USDImporter/Source/USDSchemas/Public/USDSchemaTranslator.h
--- a/Engine/Plugins/Importers/USDImporter/Source/USDSchemas/Public/USDSchemaTranslator.h
+++ b/Engine/Plugins/Importers/USDImporter/Source/USDSchemas/Public/USDSchemaTranslator.h
@@ -226,6 +226,13 @@ struct USDSCHEMAS_API FUsdSchemaTranslationContext : public TSharedFromThis< FUs
 	 */
 	bool bMergeIdenticalMaterialSlots = true;
 
+	/**
+	 * If true, whenever two prims would have generated identical UAssets (like identical StaticMeshes or materials) then only one instance of
+	 * that asset is generated, and the asset is shared by the components generated for both prims.
+	 * If false, we will always generate a dedicated asset for each prim.
+	 */
+	bool bReuseIdenticalAssets = true;
+
 	/**
 	 * If true, will cause us to collapse any point instancer prim into a single static mesh and static mesh component.
 	 * If false, will cause us to use HierarchicalInstancedStaticMeshComponents to replicate the instancing behavior.
diff --git a/Engine/Plugins/Importers/USDImporter/Source/USDSchemas/Public/USDShadeMaterialTranslator.h b/Engine/Plugins/Importers/USDImporter/Source/USDSchemas/Public/USDShadeMaterialTranslator.h
--- a/Engine/Plugins/Importers/USDImporter/Source/USDSchemas/Public/USDShadeMaterialTranslator.h
+++ b/Engine/Plugins/Importers/USDImporter/Source/USDSchemas/Public/USDShadeMaterialTranslator.h
@@ -23,7 +23,7 @@ public:
 	virtual TSet<UE::FSdfPath> CollectAuxiliaryPrims() const override;
 
 protected:
-	virtual void PostImportMaterial(const FString& MaterialHash, UMaterialInterface* ImportedMaterial);
+	virtual void PostImportMaterial(const FString& PrefixedMaterialHash, UMaterialInterface* ImportedMaterial);
 };
 
 #endif // #if USE_USD_SDK
diff --git a/Engine/Plugins/Importers/USDImporter/Source/USDStage/Private/USDStageActor.cpp b/Engine/Plugins/Importers/USDImporter/Source/USDStage/Private/USDStageActor.cpp
--- a/Engine/Plugins/Importers/USDImporter/Source/USDStage/Private/USDStageActor.cpp
+++ b/Engine/Plugins/Importers/USDImporter/Source/USDStage/Private/USDStageActor.cpp
@@ -157,6 +157,7 @@ struct FUsdStageActorImpl
 
 		TranslationContext->KindsToCollapse = (EUsdDefaultKind)StageActor->KindsToCollapse;
 		TranslationContext->bMergeIdenticalMaterialSlots = StageActor->bMergeIdenticalMaterialSlots;
+		TranslationContext->bReuseIdenticalAssets = StageActor->bReuseIdenticalAssets;
 
 		UE::FSdfPath UsdPrimPath(*PrimPath);
 		UUsdPrimTwin* ParentUsdPrimTwin = StageActor->GetRootPrimTwin()->Find(UsdPrimPath.GetParentPath().GetString());
@@ -571,6 +572,7 @@ struct FUsdStageActorImpl
 			EventAttributes.Emplace(TEXT("InterpolationType"), LexToString((uint8)StageActor->InterpolationType));
 			EventAttributes.Emplace(TEXT("KindsToCollapse"), LexToString(StageActor->KindsToCollapse));
 			EventAttributes.Emplace(TEXT("MergeIdenticalMaterialSlots"), LexToString(StageActor->bMergeIdenticalMaterialSlots));
+			EventAttributes.Emplace(TEXT("ReuseIdenticalAssets"), LexToString(StageActor->bReuseIdenticalAssets));
 			EventAttributes.Emplace(TEXT("PurposesToLoad"), LexToString(StageActor->PurposesToLoad));
 			EventAttributes.Emplace(TEXT("NaniteTriangleThreshold"), LexToString(StageActor->NaniteTriangleThreshold));
 			EventAttributes.Emplace(TEXT("RenderContext"), StageActor->RenderContext.ToString());
@@ -1058,6 +1060,7 @@ AUsdStageActor::AUsdStageActor()
 	, InterpolationType(EUsdInterpolationType::Linear)
 	, KindsToCollapse((int32)(EUsdDefaultKind::Component | EUsdDefaultKind::Subcomponent))
 	, bMergeIdenticalMaterialSlots(true)
+	, bReuseIdenticalAssets(true)
 	, PurposesToLoad((int32)EUsdPurpose::Proxy)
 	, NaniteTriangleThreshold((uint64)1000000)
 	, MaterialPurpose(*UnrealIdentifiers::MaterialPreviewPurpose)
@@ -2504,6 +2507,20 @@ void AUsdStageActor::SetMergeIdenticalMaterialSlots(bool bMerge)
 	LoadUsdStage();
 }
 
+void AUsdStageActor::SetReuseIdenticalAssets(bool bReuse)
+{
+	if (bReuse == bReuseIdenticalAssets)
+	{
+		return;
+	}
+
+	const bool bMarkDirty = false;
+	Modify(bMarkDirty);
+
+	bReuseIdenticalAssets = bReuse;
+	LoadUsdStage();
+}
+
 void AUsdStageActor::SetCollapseTopLevelPointInstancers(bool bCollapse)
 {
 }
@@ -4607,6 +4624,12 @@ void AUsdStageActor::HandlePropertyChangedEvent(FPropertyChangedEvent& PropertyC
 		bMergeIdenticalMaterialSlots = !bMergeIdenticalMaterialSlots;
 		SetMergeIdenticalMaterialSlots(bCorrectMergeMaterialSlots);
 	}
+	else if (PropertyName == GET_MEMBER_NAME_CHECKED(AUsdStageActor, bReuseIdenticalAssets))
+	{
+		const bool bCorrectReuse = bReuseIdenticalAssets;
+		bReuseIdenticalAssets = !bReuseIdenticalAssets;
+		SetReuseIdenticalAssets(bCorrectReuse);
+	}
 	else if (PropertyName == GET_MEMBER_NAME_CHECKED(AUsdStageActor, PurposesToLoad))
 	{
 		const int32 CorrectPurposesToLoad = PurposesToLoad;
diff --git a/Engine/Plugins/Importers/USDImporter/Source/USDStage/Public/USDStageActor.h b/Engine/Plugins/Importers/USDImporter/Source/USDStage/Public/USDStageActor.h
--- a/Engine/Plugins/Importers/USDImporter/Source/USDStage/Public/USDStageActor.h
+++ b/Engine/Plugins/Importers/USDImporter/Source/USDStage/Public/USDStageActor.h
@@ -79,6 +79,14 @@ public:
 	UPROPERTY(EditAnywhere, BlueprintReadOnly, Category = "USD", config )
 	bool bMergeIdenticalMaterialSlots;
 
+	/**
+	 * If true, whenever two prims would have generated identical UAssets (like identical StaticMeshes or materials) then only one instance of
+	 * that asset is generated, and the asset is shared by the components generated for both prims.
+	 * If false, we will always generate a dedicated asset for each prim.
+	 */
+	UPROPERTY(EditAnywhere, BlueprintReadOnly, Category = "USD", config )
+	bool bReuseIdenticalAssets;
+
 	/**
 	 * If true, will cause us to collapse any point instancer prim into a single static mesh and static mesh component.
 	 * If false, will cause us to use HierarchicalInstancedStaticMeshComponents to replicate the instancing behavior.
@@ -161,6 +169,9 @@ public:
 	UFUNCTION( BlueprintCallable, Category = "USD", meta = ( CallInEditor = "true" ) )
 	USDSTAGE_API void SetMergeIdenticalMaterialSlots( bool bMerge );
 
+	UFUNCTION(BlueprintCallable, Category = "USD", meta = (CallInEditor = "true"))
+	USDSTAGE_API void SetReuseIdenticalAssets(bool bReuse);
+
 	UE_DEPRECATED( 5.2, "This option is now controlled via the cvar 'USD.CollapseTopLevelPointInstancers'" )
 	UFUNCTION( BlueprintCallable, Category = "USD", meta = ( CallInEditor = "true" ) )
 	USDSTAGE_API void SetCollapseTopLevelPointInstancers( bool bCollapse );
diff --git a/Engine/Plugins/Importers/USDImporter/Source/USDStageEditor/Private/SUSDStage.cpp b/Engine/Plugins/Importers/USDImporter/Source/USDStageEditor/Private/SUSDStage.cpp
--- a/Engine/Plugins/Importers/USDImporter/Source/USDStageEditor/Private/SUSDStage.cpp
+++ b/Engine/Plugins/Importers/USDImporter/Source/USDStageEditor/Private/SUSDStage.cpp
@@ -1054,6 +1054,11 @@ void SUsdStage::FillOptionsMenu(FMenuBuilder& MenuBuilder)
 			LOCTEXT( "Collapsing_ToolTip", "Whether to try to combine individual assets and components of the same type on a Kind-per-Kind basis, like multiple Mesh prims into a single Static Mesh" ),
 			FNewMenuDelegate::CreateSP( this, &SUsdStage::FillCollapsingSubMenu ) );
 
+		MenuBuilder.AddSubMenu(
+			LOCTEXT("AssetReuse", "Asset reuse"),
+			LOCTEXT("AssetReuse_ToolTip", "How to behave when generating identical assets from different prims"),
+			FNewMenuDelegate::CreateSP(this, &SUsdStage::FillAssetReuseSubMenu));
+
 		MenuBuilder.AddSubMenu(
 			LOCTEXT( "InterpolationType", "Interpolation type" ),
 			LOCTEXT( "InterpolationType_ToolTip", "Whether to interpolate between time samples linearly or with 'held' (i.e. constant) interpolation" ),
@@ -1839,6 +1844,56 @@ void SUsdStage::FillCollapsingSubMenu( FMenuBuilder& MenuBuilder )
 	MenuBuilder.EndSection();
 }
 
+void SUsdStage::FillAssetReuseSubMenu( FMenuBuilder& MenuBuilder )
+{
+	MenuBuilder.AddMenuEntry(
+		LOCTEXT("ReuseIdenticalAssets", "Reuse identical assets"),
+		LOCTEXT(
+			"ReuseIdenticalAssets_ToolTip",
+			"If true, whenever two prims would have generated identical UAssets (like identical StaticMeshes or materials) then only one instance of "
+			"that asset is generated, and the asset is shared by the components generated for both prims. If false, we will generate a dedicated "
+			"asset for each prim."
+		),
+		FSlateIcon(),
+		FUIAction(
+			FExecuteAction::CreateLambda(
+				[this]()
+				{
+					if (AUsdStageActor* StageActor = GetStageActorOrCDO())
+					{
+						FScopedTransaction Transaction(FText::Format(
+							LOCTEXT("ReuseIdenticalAssetsTransaction", "Toggle bReuseIdenticalAssets on USD stage actor '{0}'"),
+							FText::FromString(StageActor->GetActorLabel())
+						));
+
+						// c.f. comment within AddKindToCollapseEntry just below
+						TGuardValue<bool> MaintainSelectionGuard(bUpdatingViewportSelection, true);
+
+						StageActor->SetReuseIdenticalAssets(!StageActor->bReuseIdenticalAssets);
+						if (StageActor->IsTemplate())
+						{
+							StageActor->SaveConfig();
+						}
+					}
+				}
+			),
+			FCanExecuteAction{},
+			FIsActionChecked::CreateLambda(
+				[this]()
+				{
+					if (AUsdStageActor* StageActor = GetStageActorOrCDO())
+					{
+						return StageActor->bReuseIdenticalAssets;
+					}
+					return false;
+				}
+			)
+		),
+		NAME_None,
+		EUserInterfaceActionType::ToggleButton
+	);
+}
+
 void SUsdStage::FillInterpolationTypeSubMenu(FMenuBuilder& MenuBuilder)
 {
 	MenuBuilder.AddMenuEntry(
diff --git a/Engine/Plugins/Importers/USDImporter/Source/USDStageEditor/Private/SUSDStage.h b/Engine/Plugins/Importers/USDImporter/Source/USDStageEditor/Private/SUSDStage.h
--- a/Engine/Plugins/Importers/USDImporter/Source/USDStageEditor/Private/SUSDStage.h
+++ b/Engine/Plugins/Importers/USDImporter/Source/USDStageEditor/Private/SUSDStage.h
@@ -92,6 +92,7 @@ protected:
 	void FillSubdivisionLevelSubMenu(FMenuBuilder& MenuBuilder);
 	void FillMetadataSubMenu(FMenuBuilder& MenuBuilder);
 	void FillCollapsingSubMenu( FMenuBuilder& MenuBuilder );
+	void FillAssetReuseSubMenu(FMenuBuilder& MenuBuilder);
 	void FillInterpolationTypeSubMenu( FMenuBuilder& MenuBuilder );
 	void FillSelectionSubMenu( FMenuBuilder& MenuBuilder );
 	void FillNaniteThresholdSubMenu( FMenuBuilder& MenuBuilder );
diff --git a/Engine/Plugins/Importers/USDImporter/Source/USDStageEditorViewModels/Private/USDStageViewModel.cpp b/Engine/Plugins/Importers/USDImporter/Source/USDStageEditorViewModels/Private/USDStageViewModel.cpp
--- a/Engine/Plugins/Importers/USDImporter/Source/USDStageEditorViewModels/Private/USDStageViewModel.cpp
+++ b/Engine/Plugins/Importers/USDImporter/Source/USDStageEditorViewModels/Private/USDStageViewModel.cpp
@@ -320,6 +320,7 @@ void FUsdStageViewModel::ImportStage( const TCHAR* TargetContentFolder, UUsdStag
 		ImportContext.ImportOptions->MetadataOptions = StageActor->MetadataOptions;
 		ImportContext.ImportOptions->KindsToCollapse = StageActor->KindsToCollapse;
 		ImportContext.ImportOptions->bMergeIdenticalMaterialSlots = StageActor->bMergeIdenticalMaterialSlots;
+		ImportContext.ImportOptions->bReuseIdenticalAssets = StageActor->bReuseIdenticalAssets;
 
 		ImportContext.bReadFromStageCache = true; // So that we import whatever the user has open right now, even if the file has changes
 
diff --git a/Engine/Plugins/Importers/USDImporter/Source/USDStageImporter/Private/USDStageImportOptions.cpp b/Engine/Plugins/Importers/USDImporter/Source/USDStageImporter/Private/USDStageImportOptions.cpp
--- a/Engine/Plugins/Importers/USDImporter/Source/USDStageImporter/Private/USDStageImportOptions.cpp
+++ b/Engine/Plugins/Importers/USDImporter/Source/USDStageImporter/Private/USDStageImportOptions.cpp
@@ -42,6 +42,7 @@ UUsdStageImportOptions::UUsdStageImportOptions(const FObjectInitializer& ObjectI
 
 	ExistingActorPolicy = EReplaceActorPolicy::Replace;
 	ExistingAssetPolicy = EReplaceAssetPolicy::Replace;
+	bReuseIdenticalAssets = true;
 
 	bPrimPathFolderStructure = false;
 	KindsToCollapse = ( int32 ) ( EUsdDefaultKind::Component | EUsdDefaultKind::Subcomponent );
@@ -120,6 +121,7 @@ void UsdUtils::AddAnalyticsAttributes(
 	InOutAttributes.Emplace( TEXT( "NumGroomInterpolationSettings" ), LexToString( Options.GroomInterpolationSettings.Num() ) );
 	InOutAttributes.Emplace( TEXT( "ReplaceActorPolicy" ), LexToString( ( uint8 ) Options.ExistingActorPolicy ) );
 	InOutAttributes.Emplace( TEXT( "ReplaceAssetPolicy" ), LexToString( ( uint8 ) Options.ExistingAssetPolicy ) );
+	InOutAttributes.Emplace(TEXT("ReuseIdenticalAssets"), LexToString(Options.bReuseIdenticalAssets));
 	InOutAttributes.Emplace( TEXT( "PrimPathFolderStructure" ), LexToString( Options.bPrimPathFolderStructure ) );
 	InOutAttributes.Emplace( TEXT( "KindsToCollapse" ), LexToString( Options.KindsToCollapse ) );
 	InOutAttributes.Emplace( TEXT( "MergeIdenticalMaterialSlots" ), LexToString( Options.bMergeIdenticalMaterialSlots ) );
diff --git a/Engine/Plugins/Importers/USDImporter/Source/USDStageImporter/Private/USDStageImporter.cpp b/Engine/Plugins/Importers/USDImporter/Source/USDStageImporter/Private/USDStageImporter.cpp
--- a/Engine/Plugins/Importers/USDImporter/Source/USDStageImporter/Private/USDStageImporter.cpp
+++ b/Engine/Plugins/Importers/USDImporter/Source/USDStageImporter/Private/USDStageImporter.cpp
@@ -493,6 +493,13 @@ namespace UsdStageImporterImpl
 				}
 			}
 		}
+		else if (Asset->IsA<UGroomAsset>() || Asset->IsA<UGroomCache>() || Asset->IsA<UGroomBindingAsset>())
+		{
+			// Keep the groom assets named like they originally are because they are not just named after their
+			// prim path, but may have additional suffixes like "_stands_cache" and "_groombinding" that we
+			// can't get from the prim path
+			AssetName = Asset->GetFName().GetPlainNameString();
+		}
 		else if (USkeletalMesh* SkMesh = Cast<USkeletalMesh>(Asset))
 		{
 			AssetPrefix = TEXT("SK_");
@@ -1965,6 +1972,7 @@ void UUsdStageImporter::ImportFromFile(FUsdStageImportContext& ImportContext)
 	TranslationContext->ParentComponent = ImportContext.SceneActor ? ImportContext.SceneActor->GetRootComponent() : nullptr;
 	TranslationContext->KindsToCollapse = ( EUsdDefaultKind ) ImportContext.ImportOptions->KindsToCollapse;
 	TranslationContext->bMergeIdenticalMaterialSlots = ImportContext.ImportOptions->bMergeIdenticalMaterialSlots;
+	TranslationContext->bReuseIdenticalAssets = ImportContext.ImportOptions->bReuseIdenticalAssets;
 	TranslationContext->bAllowInterpretingLODs = ImportContext.ImportOptions->bInterpretLODs;
 	TranslationContext->bAllowParsingSkeletalAnimations = ImportContext.ImportOptions->bImportGeometry && ImportContext.ImportOptions->bImportSkeletalAnimations;
 	TranslationContext->bAllowParsingGroomAssets = ImportContext.ImportOptions->bImportGroomAssets;
@@ -2121,6 +2129,7 @@ bool UUsdStageImporter::ReimportSingleAsset(
 	TranslationContext->MetadataOptions = ImportContext.ImportOptions->MetadataOptions;
 	TranslationContext->KindsToCollapse = ( EUsdDefaultKind ) ImportContext.ImportOptions->KindsToCollapse;
 	TranslationContext->bMergeIdenticalMaterialSlots = ImportContext.ImportOptions->bMergeIdenticalMaterialSlots;
+	TranslationContext->bReuseIdenticalAssets = ImportContext.ImportOptions->bReuseIdenticalAssets;
 	TranslationContext->bAllowInterpretingLODs = ImportContext.ImportOptions->bInterpretLODs;
 	TranslationContext->bAllowParsingSkeletalAnimations = ImportContext.ImportOptions->bImportGeometry && ImportContext.ImportOptions->bImportSkeletalAnimations;
 	TranslationContext->bAllowParsingGroomAssets = ImportContext.ImportOptions->bImportGroomAssets;
diff --git a/Engine/Plugins/Importers/USDImporter/Source/USDStageImporter/Public/USDStageImportOptions.h b/Engine/Plugins/Importers/USDImporter/Source/USDStageImporter/Public/USDStageImportOptions.h
--- a/Engine/Plugins/Importers/USDImporter/Source/USDStageImporter/Public/USDStageImportOptions.h
+++ b/Engine/Plugins/Importers/USDImporter/Source/USDStageImporter/Public/USDStageImportOptions.h
@@ -147,6 +147,13 @@ public:
 	UPROPERTY(BlueprintReadWrite, config, EditAnywhere, Category = "Collision")
 	EReplaceAssetPolicy ExistingAssetPolicy;
 
+	/**
+	 * If true, whenever two prims would have generated identical UAssets (like identical StaticMeshes or materials) then only one instance of
+	 * that asset is generated, and the asset is shared by the components generated for both prims.
+	 * If false, we will always generate a dedicated asset for each prim.
+	 */
+	UPROPERTY(BlueprintReadWrite, config, EditAnywhere, Category = "Collision")
+	bool bReuseIdenticalAssets;
 
 
 	/**
diff --git a/Engine/Plugins/Importers/USDImporter/Source/USDUtilities/Private/USDConversionUtils.cpp b/Engine/Plugins/Importers/USDImporter/Source/USDUtilities/Private/USDConversionUtils.cpp
--- a/Engine/Plugins/Importers/USDImporter/Source/USDUtilities/Private/USDConversionUtils.cpp
+++ b/Engine/Plugins/Importers/USDImporter/Source/USDUtilities/Private/USDConversionUtils.cpp
@@ -1773,6 +1773,28 @@ bool UsdUtils::SetAssetUserData(UObject* Object, UUsdAssetUserData* AssetUserDat
 	return true;
 }
 
+#if USE_USD_SDK
+FString UsdUtils::GetAssetHashPrefix(const pxr::UsdPrim& PrimForAsset, bool bReuseIdenticalAssets)
+{
+	if (!PrimForAsset || bReuseIdenticalAssets)
+	{
+		return FString{};
+	}
+
+	FString PrimPath = *UsdToUnreal::ConvertPath(PrimForAsset.GetPrimPath());
+	FString StageIdentifier = *UsdToUnreal::ConvertString(PrimForAsset.GetStage()->GetRootLayer()->GetIdentifier());
+
+	FSHA1 SHA1;
+	SHA1.UpdateWithString(*PrimPath, PrimPath.Len());
+	SHA1.UpdateWithString(*StageIdentifier, StageIdentifier.Len());
+
+	FSHAHash Hash;
+	SHA1.Final();
+	SHA1.GetHash(&Hash.Hash[0]);
+	return Hash.ToString() + TEXT("_");
+}
+#endif // WITH_EDITOR
+
 namespace UE::UsdConversionUtils::Private
 {
 #if USE_USD_SDK
diff --git a/Engine/Plugins/Importers/USDImporter/Source/USDUtilities/Private/USDLightConversion.cpp b/Engine/Plugins/Importers/USDImporter/Source/USDUtilities/Private/USDLightConversion.cpp
--- a/Engine/Plugins/Importers/USDImporter/Source/USDUtilities/Private/USDLightConversion.cpp
+++ b/Engine/Plugins/Importers/USDImporter/Source/USDUtilities/Private/USDLightConversion.cpp
@@ -220,7 +220,7 @@ bool UsdToUnreal::ConvertSphereLight( const pxr::UsdPrim& Prim, UPointLightCompo
 	return true;
 }
 
-bool UsdToUnreal::ConvertDomeLight( const pxr::UsdPrim& Prim, USkyLightComponent& LightComponent, UUsdAssetCache2* TexturesCache )
+bool UsdToUnreal::ConvertDomeLight( const pxr::UsdPrim& Prim, USkyLightComponent& LightComponent, UUsdAssetCache2* TexturesCache, bool bReuseIdenticalAssets )
 {
 	FScopedUsdAllocs UsdAllocs;
 
@@ -249,8 +249,9 @@ bool UsdToUnreal::ConvertDomeLight( const pxr::UsdPrim& Prim, USkyLightComponent
 		return true;
 	}
 
-	const FString DomeTextureHash = LexToString( FMD5Hash::HashFile( *ResolvedDomeTexturePath ) );
-	UTextureCube* Cubemap = Cast< UTextureCube >( TexturesCache ? TexturesCache->GetCachedAsset( DomeTextureHash ) : nullptr );
+	const FString PrefixedTextureHash = UsdUtils::GetAssetHashPrefix(Prim, bReuseIdenticalAssets)
+									+ LexToString(FMD5Hash::HashFile(*ResolvedDomeTexturePath));
+	UTextureCube* Cubemap = Cast<UTextureCube>(TexturesCache ? TexturesCache->GetCachedAsset(PrefixedTextureHash) : nullptr);
 
 	if ( !Cubemap )
 	{
@@ -258,7 +259,7 @@ bool UsdToUnreal::ConvertDomeLight( const pxr::UsdPrim& Prim, USkyLightComponent
 
 		if ( TexturesCache )
 		{
-			TexturesCache->CacheAsset( DomeTextureHash, Cubemap );
+			TexturesCache->CacheAsset( PrefixedTextureHash, Cubemap );
 		}
 	}
 
diff --git a/Engine/Plugins/Importers/USDImporter/Source/USDUtilities/Private/USDShadeConversion.cpp b/Engine/Plugins/Importers/USDImporter/Source/USDUtilities/Private/USDShadeConversion.cpp
--- a/Engine/Plugins/Importers/USDImporter/Source/USDUtilities/Private/USDShadeConversion.cpp
+++ b/Engine/Plugins/Importers/USDImporter/Source/USDUtilities/Private/USDShadeConversion.cpp
@@ -555,8 +555,9 @@ namespace UE
 				pxr::UsdShadeInput& ShadeInput,
 				TextureGroup LODGroup,
 				FParameterValue& OutValue,
-				UMaterialInterface* Material = nullptr,
-				UUsdAssetCache2* TexturesCache = nullptr
+				const UMaterialInterface* Material,
+				UUsdAssetCache2* TexturesCache,
+				bool bReuseIdenticalAssets
 			)
 			{
 				FScopedUsdAllocs UsdAllocs;
@@ -692,19 +693,20 @@ namespace UE
 							}
 						}
 
-						const FString TextureHash = UsdUtils::GetTextureHash( TexturePath, bSRGB, CompressionSettings, AddressX, AddressY );
+						const FString PrefixedTextureHash = UsdUtils::GetAssetHashPrefix(ShadeInput.GetPrim(), bReuseIdenticalAssets)
+													+ UsdUtils::GetTextureHash(TexturePath, bSRGB, CompressionSettings, AddressX, AddressY);
 
 						// We only actually want to retrieve the textures if we have a cache to put them in
 						if ( TexturesCache )
 						{
-							UTexture* Texture = Cast< UTexture >( TexturesCache->GetCachedAsset( TextureHash ) );
+							UTexture* Texture = Cast<UTexture>(TexturesCache->GetCachedAsset(PrefixedTextureHash));
 							if ( !Texture )
 							{
 								// Give the same prim path to the texture, so that it ends up imported right next to the material
 								FString MaterialPrimPath;
 								if ( Material )
 								{
-									if (UUsdAssetUserData* UserData = Material->GetAssetUserData<UUsdAssetUserData>())
+									if (UUsdAssetUserData* UserData = const_cast<UMaterialInterface*>(Material)->GetAssetUserData<UUsdAssetUserData>())
 									{
 										if (!UserData->PrimPaths.IsEmpty())
 										{
@@ -726,7 +728,7 @@ namespace UE
 									}
 									Texture->UpdateResource();
 
-									TexturesCache->CacheAsset( TextureHash, Texture );
+									TexturesCache->CacheAsset( PrefixedTextureHash, Texture );
 								}
 							}
 
@@ -792,8 +794,9 @@ namespace UE
 				const pxr::TfToken& InputName,
 				float DefaultValue,
 				FParameterValue& OutValue,
-				UMaterialInterface* Material = nullptr,
-				UUsdAssetCache2* TexturesCache = nullptr
+				const UMaterialInterface* Material,
+				UUsdAssetCache2* TexturesCache,
+				bool bReuseIdenticalAssets
 			)
 			{
 				FScopedUsdAllocs Allocs;
@@ -810,7 +813,7 @@ namespace UE
 				pxr::UsdShadeAttributeType AttributeType;
 				if ( pxr::UsdShadeConnectableAPI::GetConnectedSource( Input.GetAttr(), &Source, &SourceName, &AttributeType ) )
 				{
-					if (!GetTextureParameterValue(Input, TEXTUREGROUP_WorldSpecular, OutValue, Material, TexturesCache))
+					if (!GetTextureParameterValue(Input, TEXTUREGROUP_WorldSpecular, OutValue, Material, TexturesCache, bReuseIdenticalAssets))
 					{
 						// Check if we have a fallback input that we can use instead, since we don't have a valid texture value
 						if ( const pxr::UsdShadeInput FallbackInput = Source.GetInput( UnrealIdentifiers::Fallback ) )
@@ -825,7 +828,7 @@ namespace UE
 
 						// Recurse because the attribute may just be pointing at some other attribute that has the data
 						// (e.g. when shader input is just "hoisted" and connected to the parent material input)
-						return GetFloatParameterValue(Source, SourceName, DefaultValue, OutValue, Material, TexturesCache);
+						return GetFloatParameterValue(Source, SourceName, DefaultValue, OutValue, Material, TexturesCache, bReuseIdenticalAssets);
 					}
 				}
 				// No other node connected, so we must have some value
@@ -917,9 +920,10 @@ namespace UE
 				const pxr::TfToken& InputName,
 				const FLinearColor& DefaultValue,
 				FParameterValue& OutValue,
-				bool bIsNormalMap = false,
-				UMaterialInterface* Material = nullptr,
-				UUsdAssetCache2* TexturesCache = nullptr
+				bool bIsNormalMap,
+				const UMaterialInterface* Material,
+				UUsdAssetCache2* TexturesCache,
+				bool bReuseIdenticalAssets
 			)
 			{
 				FScopedUsdAllocs Allocs;
@@ -936,7 +940,14 @@ namespace UE
 				pxr::UsdShadeAttributeType AttributeType;
 				if ( pxr::UsdShadeConnectableAPI::GetConnectedSource( Input.GetAttr(), &Source, &SourceName, &AttributeType ) )
 				{
-					if (!GetTextureParameterValue(Input, bIsNormalMap ? TEXTUREGROUP_WorldNormalMap : TEXTUREGROUP_World, OutValue, Material, TexturesCache))
+					if (!GetTextureParameterValue(
+							Input,
+							bIsNormalMap ? TEXTUREGROUP_WorldNormalMap : TEXTUREGROUP_World,
+							OutValue,
+							Material,
+							TexturesCache,
+							bReuseIdenticalAssets
+						))
 					{
 						// Check whether this input receives its value through a connection to a
 						// primvar reader shader.
@@ -971,7 +982,16 @@ namespace UE
 						}
 
 						// This shader doesn't have anything: Traverse into the input connectable itself
-						return GetVec3ParameterValue(Source, SourceName, DefaultValue, OutValue, bIsNormalMap, Material, TexturesCache);
+						return GetVec3ParameterValue(
+							Source,
+							SourceName,
+							DefaultValue,
+							OutValue,
+							bIsNormalMap,
+							Material,
+							TexturesCache,
+							bReuseIdenticalAssets
+						);
 					}
 				}
 				// No other node connected, so we must have some value
@@ -2067,7 +2087,8 @@ bool UsdToUnreal::ConvertMaterial(
 	const pxr::UsdShadeMaterial& UsdShadeMaterial,
 	UMaterialInstance& Material,
 	UUsdAssetCache2* TexturesCache,
-	const TCHAR* RenderContext
+	const TCHAR* RenderContext,
+	bool bReuseIdenticalAssets
 )
 {
 	FScopedUsdAllocs UsdAllocs;
@@ -2101,49 +2122,116 @@ bool UsdToUnreal::ConvertMaterial(
 	const bool bIsNormalMap = true;
 
 	// Base color
-	if (UsdShadeConversionImpl::GetVec3ParameterValue(Connectable, UnrealIdentifiers::DiffuseColor, FLinearColor(0, 0, 0), ParameterValue, !bIsNormalMap, &Material, TexturesCache))
+	if (UsdShadeConversionImpl::GetVec3ParameterValue(
+			Connectable,
+			UnrealIdentifiers::DiffuseColor,
+			FLinearColor(0, 0, 0),
+			ParameterValue,
+			!bIsNormalMap,
+			&Material,
+			TexturesCache,
+			bReuseIdenticalAssets
+		))
 	{
 		MaterialParameters.Add(TEXT("BaseColor"), ParameterValue);
 	}
 
 	// Emissive color
-	if (UsdShadeConversionImpl::GetVec3ParameterValue(Connectable, UnrealIdentifiers::EmissiveColor, FLinearColor(0, 0, 0), ParameterValue, !bIsNormalMap, &Material, TexturesCache))
+	if (UsdShadeConversionImpl::GetVec3ParameterValue(
+			Connectable,
+			UnrealIdentifiers::EmissiveColor,
+			FLinearColor(0, 0, 0),
+			ParameterValue,
+			!bIsNormalMap,
+			&Material,
+			TexturesCache,
+			bReuseIdenticalAssets
+		))
 	{
 		MaterialParameters.Add(TEXT("EmissiveColor"), ParameterValue);
 	}
 
 	// Metallic
-	if (UsdShadeConversionImpl::GetFloatParameterValue(Connectable, UnrealIdentifiers::Metallic, 0.f, ParameterValue, &Material, TexturesCache))
+	if (UsdShadeConversionImpl::GetFloatParameterValue(
+			Connectable,
+			UnrealIdentifiers::Metallic,
+			0.f,
+			ParameterValue,
+			&Material,
+			TexturesCache,
+			bReuseIdenticalAssets
+		))
 	{
 		MaterialParameters.Add(TEXT("Metallic"), ParameterValue);
 	}
 
 	// Roughness
-	if (UsdShadeConversionImpl::GetFloatParameterValue(Connectable, UnrealIdentifiers::Roughness, 1.f, ParameterValue, &Material, TexturesCache))
+	if (UsdShadeConversionImpl::GetFloatParameterValue(
+			Connectable,
+			UnrealIdentifiers::Roughness,
+			1.f,
+			ParameterValue,
+			&Material,
+			TexturesCache,
+			bReuseIdenticalAssets
+		))
 	{
 		MaterialParameters.Add(TEXT("Roughness"), ParameterValue);
 	}
 
 	// Opacity
-	if (UsdShadeConversionImpl::GetFloatParameterValue(Connectable, UnrealIdentifiers::Opacity, 1.f, ParameterValue, &Material, TexturesCache))
+	if (UsdShadeConversionImpl::GetFloatParameterValue(
+			Connectable,
+			UnrealIdentifiers::Opacity,
+			1.f,
+			ParameterValue,
+			&Material,
+			TexturesCache,
+			bReuseIdenticalAssets
+		))
 	{
 		MaterialParameters.Add(TEXT("Opacity"), ParameterValue);
 	}
 
 	// Normal
-	if (UsdShadeConversionImpl::GetVec3ParameterValue(Connectable, UnrealIdentifiers::Normal, FLinearColor(0, 0, 1), ParameterValue, bIsNormalMap, &Material, TexturesCache))
+	if (UsdShadeConversionImpl::GetVec3ParameterValue(
+			Connectable,
+			UnrealIdentifiers::Normal,
+			FLinearColor(0, 0, 1),
+			ParameterValue,
+			bIsNormalMap,
+			&Material,
+			TexturesCache,
+			bReuseIdenticalAssets
+		))
 	{
 		MaterialParameters.Add(TEXT("Normal"), ParameterValue);
 	}
 
 	// Refraction
-	if (UsdShadeConversionImpl::GetFloatParameterValue(Connectable, UnrealIdentifiers::Refraction, 1.5f, ParameterValue, &Material, TexturesCache))
+	if (UsdShadeConversionImpl::GetFloatParameterValue(
+			Connectable,
+			UnrealIdentifiers::Refraction,
+			1.5f,
+			ParameterValue,
+			&Material,
+			TexturesCache,
+			bReuseIdenticalAssets
+		))
 	{
 		MaterialParameters.Add(TEXT("Refraction"), ParameterValue);
 	}
 
 	// Occlusion
-	if (UsdShadeConversionImpl::GetFloatParameterValue(Connectable, UnrealIdentifiers::Occlusion, 1.0f, ParameterValue, &Material, TexturesCache))
+	if (UsdShadeConversionImpl::GetFloatParameterValue(
+			Connectable,
+			UnrealIdentifiers::Occlusion,
+			1.0f,
+			ParameterValue,
+			&Material,
+			TexturesCache,
+			bReuseIdenticalAssets
+		))
 	{
 		MaterialParameters.Add(TEXT("AmbientOcclusion"), ParameterValue);
 	}
@@ -2213,7 +2301,8 @@ bool UsdToUnreal::ConvertMaterial(
 	const pxr::UsdShadeMaterial& UsdShadeMaterial,
 	UMaterial& Material,
 	UUsdAssetCache2* TexturesCache,
-	const TCHAR* RenderContext
+	const TCHAR* RenderContext,
+	bool bReuseIdenticalAssets
 )
 {
 #if WITH_EDITOR
@@ -2251,50 +2340,110 @@ bool UsdToUnreal::ConvertMaterial(
 	const bool bIsNormalMap = true;
 
 	// Base color
-	if (UsdShadeConversionImpl::GetVec3ParameterValue(Connectable, UnrealIdentifiers::DiffuseColor, FLinearColor(0, 0, 0), TempParameterValue, !bIsNormalMap, &Material, TexturesCache))
+	if (UsdShadeConversionImpl::GetVec3ParameterValue(
+			Connectable,
+			UnrealIdentifiers::DiffuseColor,
+			FLinearColor(0, 0, 0),
+			TempParameterValue,
+			!bIsNormalMap,
+			&Material,
+			TexturesCache,
+			bReuseIdenticalAssets
+		))
 	{
 		BaseColorParameter = TempParameterValue;
 	}
 
 	// Emissive color
-	if (UsdShadeConversionImpl::GetVec3ParameterValue(Connectable, UnrealIdentifiers::EmissiveColor, FLinearColor(0, 0, 0), TempParameterValue, !bIsNormalMap, &Material, TexturesCache))
+	if (UsdShadeConversionImpl::GetVec3ParameterValue(
+			Connectable,
+			UnrealIdentifiers::EmissiveColor,
+			FLinearColor(0, 0, 0),
+			TempParameterValue,
+			!bIsNormalMap,
+			&Material,
+			TexturesCache,
+			bReuseIdenticalAssets
+		))
 	{
 		EmissiveParameter = TempParameterValue;
 	}
 
 	// Metallic
-	if (UsdShadeConversionImpl::GetFloatParameterValue(Connectable, UnrealIdentifiers::Metallic, 0.f, TempParameterValue, &Material, TexturesCache))
+	if (UsdShadeConversionImpl::GetFloatParameterValue(
+			Connectable,
+			UnrealIdentifiers::Metallic,
+			0.f,
+			TempParameterValue,
+			&Material,
+			TexturesCache,
+			bReuseIdenticalAssets
+		))
 	{
 		MetallicParameter = TempParameterValue;
 	}
 
 	// Roughness
-	if (UsdShadeConversionImpl::GetFloatParameterValue(Connectable, UnrealIdentifiers::Roughness, 1.f, TempParameterValue, &Material, TexturesCache))
+	if (UsdShadeConversionImpl::GetFloatParameterValue(
+			Connectable,
+			UnrealIdentifiers::Roughness,
+			1.f,
+			TempParameterValue,
+			&Material,
+			TexturesCache,
+			bReuseIdenticalAssets
+		))
 	{
 		RoughnessParameter = TempParameterValue;
 	}
 
 	// Opacity
-	if (UsdShadeConversionImpl::GetFloatParameterValue(Connectable, UnrealIdentifiers::Opacity, 1.f, TempParameterValue, &Material, TexturesCache))
+	if (UsdShadeConversionImpl::GetFloatParameterValue(
+			Connectable,
+			UnrealIdentifiers::Opacity,
+			1.f,
+			TempParameterValue,
+			&Material,
+			TexturesCache,
+			bReuseIdenticalAssets
+		))
 	{
 		OpacityParameter = TempParameterValue;
 	}
 
 	// Normal
-	if (UsdShadeConversionImpl::GetVec3ParameterValue(Connectable, UnrealIdentifiers::Normal, FLinearColor(0, 0, 1), TempParameterValue, bIsNormalMap, &Material, TexturesCache))
+	if (UsdShadeConversionImpl::GetVec3ParameterValue(
+			Connectable,
+			UnrealIdentifiers::Normal,
+			FLinearColor(0, 0, 1),
+			TempParameterValue,
+			bIsNormalMap,
+			&Material,
+			TexturesCache,
+			bReuseIdenticalAssets
+		))
 	{
 		NormalParameter = TempParameterValue;
 	}
 
 	// Refraction
-	const bool bHasRefractionValue = UsdShadeConversionImpl::GetFloatParameterValue(Connectable, UnrealIdentifiers::Refraction, 1.5f, TempParameterValue, &Material, TexturesCache);
-	if (bHasRefractionValue || Material.BlendMode == BLEND_Translucent) // Force a 1.5 IOR if USD didn't specify a value, as it's USD's fallback value
+	const bool bHasRefractionValue = UsdShadeConversionImpl::GetFloatParameterValue(
+		Connectable,
+		UnrealIdentifiers::Refraction,
+		1.5f,
+		TempParameterValue,
+		&Material,
+		TexturesCache,
+		bReuseIdenticalAssets
+	);
+	if (bHasRefractionValue || Material.BlendMode == BLEND_Translucent)	   // Force a 1.5 IOR if USD didn't specify a value, as it's USD's fallback
+																		   // value
 	{
 		RefractionParameter = TempParameterValue;
 	}
 
 	// Ambient occlusion
-	if (UsdShadeConversionImpl::GetFloatParameterValue(Connectable, UnrealIdentifiers::Occlusion, 1.f, TempParameterValue, &Material, TexturesCache))
+	if (UsdShadeConversionImpl::GetFloatParameterValue(Connectable, UnrealIdentifiers::Occlusion, 1.f, TempParameterValue, &Material, TexturesCache, bReuseIdenticalAssets))
 	{
 		AmbientOcclusionParameter = TempParameterValue;
 	}
@@ -2512,7 +2661,13 @@ bool UsdToUnreal::ConvertMaterial(
 	return ConvertMaterial(UsdShadeMaterial, Material, TexturesCache, RenderContext);
 }
 
-bool UsdToUnreal::ConvertShadeInputsToParameters( const pxr::UsdShadeMaterial& UsdShadeMaterial, UMaterialInstance& MaterialInstance, UUsdAssetCache2* TexturesCache, const TCHAR* RenderContext )
+bool UsdToUnreal::ConvertShadeInputsToParameters(
+	const pxr::UsdShadeMaterial& UsdShadeMaterial,
+	UMaterialInstance& MaterialInstance,
+	UUsdAssetCache2* TexturesCache,
+	const TCHAR* RenderContext,
+	bool bReuseIdenticalAssets
+)
 {
 	FScopedUsdAllocs UsdAllocs;
 
@@ -2574,13 +2729,30 @@ bool UsdToUnreal::ConvertShadeInputsToParameters( const pxr::UsdShadeMaterial& U
 			ShadeInput.GetTypeName() == pxr::SdfValueTypeNames->Double ||
 			ShadeInput.GetTypeName() == pxr::SdfValueTypeNames->Half )
 		{
-			if ( UsdShadeConversionImpl::GetFloatParameterValue( Connectable, ShadeInput.GetBaseName(), 1.f, ParameterValue, &MaterialInstance, TexturesCache ) )
+			if (UsdShadeConversionImpl::GetFloatParameterValue(
+					Connectable,
+					ShadeInput.GetBaseName(),
+					1.f,
+					ParameterValue,
+					&MaterialInstance,
+					TexturesCache,
+					bReuseIdenticalAssets
+				))
 			{
 				UsdShadeConversionImpl::SetParameterValue(MaterialInstance, *DisplayName, ParameterValue, bForUsdPreviewSurface, Unused);
 				bHasMaterialInfo = true;
 			}
 		}
-		else if ( UsdShadeConversionImpl::GetVec3ParameterValue( Connectable, ShadeInput.GetBaseName(), FLinearColor( 0.f, 0.f, 0.f ), ParameterValue, bIsNormalMap, &MaterialInstance, TexturesCache ) )
+		else if (UsdShadeConversionImpl::GetVec3ParameterValue(
+					 Connectable,
+					 ShadeInput.GetBaseName(),
+					 FLinearColor(0.f, 0.f, 0.f),
+					 ParameterValue,
+					 bIsNormalMap,
+					 &MaterialInstance,
+					 TexturesCache,
+					 bReuseIdenticalAssets
+				 ))
 		{
 			UsdShadeConversionImpl::SetParameterValue(MaterialInstance, *DisplayName, ParameterValue, bForUsdPreviewSurface, Unused);
 			bHasMaterialInfo = true;
@@ -3135,8 +3307,19 @@ bool UsdUtils::IsMaterialTranslucent( const pxr::UsdShadeMaterial& UsdShadeMater
 	}
 	pxr::UsdShadeConnectableAPI Connectable{ SurfaceShader };
 
+	const UMaterialInterface* Material = nullptr;
+	UUsdAssetCache2* TexturesCache = nullptr;
+	const bool bReuseIdenticalAssets = true;
 	UsdShadeConversionImpl::FParameterValue ParameterValue;
-	bool bHasOpacityConnection = UsdShadeConversionImpl::GetFloatParameterValue( Connectable, UnrealIdentifiers::Opacity, 1.f, ParameterValue );
+	bool bHasOpacityConnection = UsdShadeConversionImpl::GetFloatParameterValue(
+		Connectable,
+		UnrealIdentifiers::Opacity,
+		1.f,
+		ParameterValue,
+		Material,
+		TexturesCache,
+		bReuseIdenticalAssets
+	);
 
 	// Don't check if the texture is nullptr here as we won't actually parse it yet. If the variant has this type we know it's meant to be bound to a texture
 	const bool bHasBoundTexture = ParameterValue.IsType<UsdShadeConversionImpl::FTextureParameterValue>();
diff --git a/Engine/Plugins/Importers/USDImporter/Source/USDUtilities/Public/USDConversionUtils.h b/Engine/Plugins/Importers/USDImporter/Source/USDUtilities/Public/USDConversionUtils.h
--- a/Engine/Plugins/Importers/USDImporter/Source/USDUtilities/Public/USDConversionUtils.h
+++ b/Engine/Plugins/Importers/USDImporter/Source/USDUtilities/Public/USDConversionUtils.h
@@ -339,6 +339,16 @@ namespace UsdUtils
 	 */
 	USDUTILITIES_API bool SetAssetUserData(UObject* Object, UUsdAssetUserData* AssetUserData);
 
+#if USE_USD_SDK
+	/**
+	 * Simple utility that generates a prefix that should be added to the hash when using it to cache/query an UsdAssetCache2 object.
+	 * The idea is to use this prefix alongside the hash whenever bReuseIdenticalAssets is false, so as to generate a new hash for each prim.
+	 * Whenever bReuseIdenticalAssets is true, the prefix will be the empty string.
+	 * Whenever bReuseIdenticalAssets is false, the prefix will be a SHA1 hash of the prim path and its stage identifier.
+	 */
+	USDUTILITIES_API FString GetAssetHashPrefix(const pxr::UsdPrim& PrimForAsset, bool bReuseIdenticalAssets);
+#endif // #if USE_USD_SDK
+
 	/** Adds a reference on Prim to the layer at AbsoluteFilePath */
 	USDUTILITIES_API void AddReference( UE::FUsdPrim& Prim, const TCHAR* AbsoluteFilePath, const UE::FSdfPath& TargetPrimPath = {}, double TimeCodeOffset = 0.0, double TimeCodeScale = 1.0 );
 
diff --git a/Engine/Plugins/Importers/USDImporter/Source/USDUtilities/Public/USDLightConversion.h b/Engine/Plugins/Importers/USDImporter/Source/USDUtilities/Public/USDLightConversion.h
--- a/Engine/Plugins/Importers/USDImporter/Source/USDUtilities/Public/USDLightConversion.h
+++ b/Engine/Plugins/Importers/USDImporter/Source/USDUtilities/Public/USDLightConversion.h
@@ -60,7 +60,7 @@ namespace UsdToUnreal
 	USDUTILITIES_API bool ConvertRectLight( const pxr::UsdPrim& Prim, URectLightComponent& LightComponent, double UsdTimeCode = UsdUtils::GetDefaultTimeCode() );
 	USDUTILITIES_API bool ConvertDiskLight( const pxr::UsdPrim& Prim, URectLightComponent& LightComponent, double UsdTimeCode = UsdUtils::GetDefaultTimeCode() );
 	USDUTILITIES_API bool ConvertSphereLight( const pxr::UsdPrim& Prim, UPointLightComponent& LightComponent, double UsdTimeCode = UsdUtils::GetDefaultTimeCode() );
-	USDUTILITIES_API bool ConvertDomeLight( const pxr::UsdPrim& Prim, USkyLightComponent& LightComponent, UUsdAssetCache2* TexturesCache );
+	USDUTILITIES_API bool ConvertDomeLight(const pxr::UsdPrim& Prim, USkyLightComponent& LightComponent, UUsdAssetCache2* TexturesCache, bool bReuseIdenticalAssets = true);
 	USDUTILITIES_API bool ConvertLuxShapingAPI( const pxr::UsdPrim& Prim, USpotLightComponent& LightComponent, double UsdTimeCode = UsdUtils::GetDefaultTimeCode() );
 
 PRAGMA_DISABLE_DEPRECATION_WARNINGS
diff --git a/Engine/Plugins/Importers/USDImporter/Source/USDUtilities/Public/USDShadeConversion.h b/Engine/Plugins/Importers/USDImporter/Source/USDUtilities/Public/USDShadeConversion.h
--- a/Engine/Plugins/Importers/USDImporter/Source/USDUtilities/Public/USDShadeConversion.h
+++ b/Engine/Plugins/Importers/USDImporter/Source/USDUtilities/Public/USDShadeConversion.h
@@ -53,10 +53,23 @@ namespace UsdToUnreal
 	 * @param Material - Output parameter that will be filled with the converted data. Only the versions that receive a dynamic material instance will work at runtime
 	 * @param TexturesCache - Cache to prevent importing a texture more than once
 	 * @param RenderContext - Which render context output to read from the UsdShadeMaterial
+	 * @param ReuseIdenticalAssets - Whether to reuse identical textures found in the TextureCache or to create dedicated textures for each material
 	 * @return Whether the conversion was successful or not.
 	 */
-	USDUTILITIES_API bool ConvertMaterial(const pxr::UsdShadeMaterial& UsdShadeMaterial, UMaterialInstance& Material, UUsdAssetCache2* TexturesCache = nullptr, const TCHAR* RenderContext = nullptr);
-	USDUTILITIES_API bool ConvertMaterial(const pxr::UsdShadeMaterial& UsdShadeMaterial, UMaterial& Material, UUsdAssetCache2* TexturesCache = nullptr, const TCHAR* RenderContext = nullptr);
+	USDUTILITIES_API bool ConvertMaterial(
+		const pxr::UsdShadeMaterial& UsdShadeMaterial,
+		UMaterialInstance& Material,
+		UUsdAssetCache2* TexturesCache = nullptr,
+		const TCHAR* RenderContext = nullptr,
+		bool bReuseIdenticalAssets = true
+	);
+	USDUTILITIES_API bool ConvertMaterial(
+		const pxr::UsdShadeMaterial& UsdShadeMaterial,
+		UMaterial& Material,
+		UUsdAssetCache2* TexturesCache = nullptr,
+		const TCHAR* RenderContext = nullptr,
+		bool bReuseIdenticalAssets = true
+	);
 
 	/**
 	 * Attemps to assign the values of the surface shader inputs to the MaterialInstance parameters by matching the inputs display names to the parameters names.
@@ -64,10 +77,17 @@ namespace UsdToUnreal
 	 * @param MaterialInstance - Material instance on which we will set the parameter values
 	 * @param TexturesCache - Cache to prevent importing a texture more than once
 	 * @param RenderContext - The USD render context to use when fetching the surface shader
+	 * @param ReuseIdenticalAssets - Whether to reuse identical textures found in the TextureCache or to create dedicated textures for each material
 	 * @return Whether the conversion was successful or not.
 	 *
 	 */
-	USDUTILITIES_API bool ConvertShadeInputsToParameters( const pxr::UsdShadeMaterial& UsdShadeMaterial, UMaterialInstance& MaterialInstance, UUsdAssetCache2* TexturesCache, const TCHAR* RenderContext = nullptr );
+	USDUTILITIES_API bool ConvertShadeInputsToParameters(
+		const pxr::UsdShadeMaterial& UsdShadeMaterial,
+		UMaterialInstance& MaterialInstance,
+		UUsdAssetCache2* TexturesCache,
+		const TCHAR* RenderContext = nullptr,
+		bool bReuseIdenticalAssets = true
+	);
 
 PRAGMA_DISABLE_DEPRECATION_WARNINGS
 	UE_DEPRECATED(5.3, "Please use the overload that doesn't use PrimvarToUVIndex anymore")
