USD: Fix issues with default USD asset cache dialog.

diff --git a/Engine/Plugins/Importers/USDImporter/Source/USDClassesEditor/Private/USDClassesEditorModule.cpp b/Engine/Plugins/Importers/USDImporter/Source/USDClassesEditor/Private/USDClassesEditorModule.cpp
--- a/Engine/Plugins/Importers/USDImporter/Source/USDClassesEditor/Private/USDClassesEditorModule.cpp
+++ b/Engine/Plugins/Importers/USDImporter/Source/USDClassesEditor/Private/USDClassesEditorModule.cpp
@@ -16,12 +16,17 @@
 UUsdAssetCache2* IUsdClassesEditorModule::ShowMissingDefaultAssetCacheDialog()
 {
 	UUsdAssetCache2* Result = nullptr;
-	bool bOutUserAccepted = false;
-	ShowMissingDefaultAssetCacheDialog(Result, bOutUserAccepted);
+	EDefaultAssetCacheDialogOption Unused = ShowMissingDefaultAssetCacheDialog(Result);
 	return Result;
 }
 
 void IUsdClassesEditorModule::ShowMissingDefaultAssetCacheDialog(UUsdAssetCache2*& OutCreatedCache, bool& bOutUserAccepted)
+{
+	EDefaultAssetCacheDialogOption Result = ShowMissingDefaultAssetCacheDialog(OutCreatedCache);
+	bOutUserAccepted = Result != EDefaultAssetCacheDialogOption::Cancel;
+}
+
+EDefaultAssetCacheDialogOption IUsdClassesEditorModule::ShowMissingDefaultAssetCacheDialog(UUsdAssetCache2*& OutCreatedCache)
 {
 	TSharedPtr<SWindow> ParentWindow;
 
@@ -48,16 +53,8 @@ void IUsdClassesEditorModule::ShowMissingDefaultAssetCacheDialog(UUsdAssetCache2
 	const bool bSlowTaskWindow = false;
 	FSlateApplication::Get().AddModalWindow(Window, ParentWindow, bSlowTaskWindow);
 
-	if (OptionsWindow->UserAccepted())
-	{
-		bOutUserAccepted = true;
-		OutCreatedCache = OptionsWindow->GetCreatedCache();
-	}
-	else
-	{
-		bOutUserAccepted = false;
-		OutCreatedCache = nullptr;
-	}
+	OutCreatedCache = OptionsWindow->GetCreatedCache();
+	return OptionsWindow->GetDialogOutcome();
 }
 
 class FUsdClassesEditorModule : public IUsdClassesEditorModule
diff --git a/Engine/Plugins/Importers/USDImporter/Source/USDClassesEditor/Private/USDDefaultAssetCacheDialog.cpp b/Engine/Plugins/Importers/USDImporter/Source/USDClassesEditor/Private/USDDefaultAssetCacheDialog.cpp
--- a/Engine/Plugins/Importers/USDImporter/Source/USDClassesEditor/Private/USDDefaultAssetCacheDialog.cpp
+++ b/Engine/Plugins/Importers/USDImporter/Source/USDClassesEditor/Private/USDDefaultAssetCacheDialog.cpp
@@ -105,7 +105,7 @@ void SUsdDefaultAssetCacheDialog::Construct( const FArguments& InArgs )
 {
 	ChosenCache = nullptr;
 	Window = InArgs._WidgetWindow;
-	bAccepted = false;
+	DialogOutcome = EDefaultAssetCacheDialogOption::Cancel;
 
 	FSlateFontInfo MessageFont(FAppStyle::GetFontStyle("StandardDialog.LargeFont"));
 
@@ -138,9 +138,9 @@ void SUsdDefaultAssetCacheDialog::Construct( const FArguments& InArgs )
 				.Padding(16.f, 0.f, 0.f, 0.f)
 				[
 					SNew(STextBlock)
-					.Text(LOCTEXT("DefaultAssetCacheDialogText", "This project has no default USD Asset Cache set.\n\nThe cache stores the assets generated when opening USD Stages, and shares them across all USD Stage Actors.\n\nDo you wish to set the default asset cache?"))
+					.Text(LOCTEXT("DefaultAssetCacheDialogText", "Opening a USD Stage with a Stage Actor requires a default USD Asset Cache.\n\nThe cache will be used to store and share the generated assets across all USD Stage Actors.\n\nHover on the buttons for more info about each of the proposed options."))
 					.Font(MessageFont)
-					.WrapTextAt(512.0f)
+					.WrapTextAt(650.0f)
 				]
 			]
 
@@ -242,7 +242,7 @@ FReply SUsdDefaultAssetCacheDialog::OnUseExisting()
 		ChosenCache = Cast<UUsdAssetCache2>(PickedAssets[0].GetAsset());
 	}
 
-	bAccepted = true;
+	DialogOutcome = EDefaultAssetCacheDialogOption::PickExisting;
 	if (Window.IsValid())
 	{
 		Window.Pin()->RequestDestroyWindow();
@@ -254,7 +254,7 @@ FReply SUsdDefaultAssetCacheDialog::OnCreateNew()
 {
 	ChosenCache = UE::DefaultCacheDialog::Private::CreateNewAssetCacheWithDialog();
 
-	bAccepted = true;
+	DialogOutcome = EDefaultAssetCacheDialogOption::CreateNew;
 	if (Window.IsValid())
 	{
 		Window.Pin()->RequestDestroyWindow();
@@ -264,7 +264,7 @@ FReply SUsdDefaultAssetCacheDialog::OnCreateNew()
 
 FReply SUsdDefaultAssetCacheDialog::OnDontCreate()
 {
-	bAccepted = true;
+	DialogOutcome = EDefaultAssetCacheDialogOption::DontUseDefault;
 	if (Window.IsValid())
 	{
 		Window.Pin()->RequestDestroyWindow();
@@ -281,9 +281,9 @@ FReply SUsdDefaultAssetCacheDialog::OnKeyDown( const FGeometry& MyGeometry, cons
 	return FReply::Unhandled();
 }
 
-bool SUsdDefaultAssetCacheDialog::UserAccepted() const
+EDefaultAssetCacheDialogOption SUsdDefaultAssetCacheDialog::GetDialogOutcome() const
 {
-	return bAccepted;
+	return DialogOutcome;
 }
 
 #undef LOCTEXT_NAMESPACE
\ No newline at end of file
diff --git a/Engine/Plugins/Importers/USDImporter/Source/USDClassesEditor/Private/USDDefaultAssetCacheDialog.h b/Engine/Plugins/Importers/USDImporter/Source/USDClassesEditor/Private/USDDefaultAssetCacheDialog.h
--- a/Engine/Plugins/Importers/USDImporter/Source/USDClassesEditor/Private/USDDefaultAssetCacheDialog.h
+++ b/Engine/Plugins/Importers/USDImporter/Source/USDClassesEditor/Private/USDDefaultAssetCacheDialog.h
@@ -2,6 +2,8 @@
 
 #pragma once
 
+#include "USDClassesEditorModule.h"
+
 #include "Widgets/SCompoundWidget.h"
 #include "Widgets/SWindow.h"
 
@@ -21,7 +23,7 @@ public:
 	UUsdAssetCache2* GetCreatedCache();
 	void Construct( const FArguments& InArgs );
 
-	bool UserAccepted() const;
+	EDefaultAssetCacheDialogOption GetDialogOutcome() const;
 
 private:
 	virtual bool SupportsKeyboardFocus() const override;
@@ -38,5 +40,5 @@ private:
 	TWeakPtr< SWindow > Window;
 
 	FText AcceptText;
-	bool bAccepted;
+	EDefaultAssetCacheDialogOption DialogOutcome = EDefaultAssetCacheDialogOption::Cancel;
 };
\ No newline at end of file
diff --git a/Engine/Plugins/Importers/USDImporter/Source/USDClassesEditor/Public/USDClassesEditorModule.h b/Engine/Plugins/Importers/USDImporter/Source/USDClassesEditor/Public/USDClassesEditorModule.h
--- a/Engine/Plugins/Importers/USDImporter/Source/USDClassesEditor/Public/USDClassesEditorModule.h
+++ b/Engine/Plugins/Importers/USDImporter/Source/USDClassesEditor/Public/USDClassesEditorModule.h
@@ -6,6 +6,14 @@
 
 class UUsdAssetCache2;
 
+enum class EDefaultAssetCacheDialogOption : uint8
+{
+	Cancel = 0,
+	PickExisting,
+	CreateNew,
+	DontUseDefault
+};
+
 class IUsdClassesEditorModule : public IModuleInterface
 {
 public:
@@ -13,5 +21,8 @@ public:
 	UE_DEPRECATED(5.3, "Use the other signature that also returns whether the user accepted the dialog or not")
 	USDCLASSESEDITOR_API static UUsdAssetCache2* ShowMissingDefaultAssetCacheDialog();
 
+	UE_DEPRECATED(5.4, "Use the other signature that also returns the dialog outcome as an Enum")
 	USDCLASSESEDITOR_API static void ShowMissingDefaultAssetCacheDialog(UUsdAssetCache2*& OutCreatedCache, bool& bOutUserAccepted);
+
+	USDCLASSESEDITOR_API static EDefaultAssetCacheDialogOption ShowMissingDefaultAssetCacheDialog(UUsdAssetCache2*& OutCreatedCache);
 };
\ No newline at end of file
diff --git a/Engine/Plugins/Importers/USDImporter/Source/USDStage/Private/USDStageActor.cpp b/Engine/Plugins/Importers/USDImporter/Source/USDStage/Private/USDStageActor.cpp
--- a/Engine/Plugins/Importers/USDImporter/Source/USDStage/Private/USDStageActor.cpp
+++ b/Engine/Plugins/Importers/USDImporter/Source/USDStage/Private/USDStageActor.cpp
@@ -7,9 +7,9 @@
 #include "USDAssetCache.h"
 #include "USDAssetImportData.h"
 #include "USDAssetUserData.h"
-#include "USDDrawModeComponent.h"
 #include "USDClassesModule.h"
 #include "USDConversionUtils.h"
+#include "USDDrawModeComponent.h"
 #include "USDDynamicBindingResolverLibrary.h"
 #include "USDErrorUtils.h"
 #include "USDGeomMeshConversion.h"
@@ -1427,6 +1427,21 @@ void AUsdStageActor::OnUsdObjectsChanged(const UsdUtils::FObjectChangesByPath& I
 	SetupAssetCacheIfNeeded();
 	if (!UsdAssetCache)
 	{
+		// If the user canceled out of providing an asset cache and we should be loading assets, we have no choice but to either change the stage
+		// state or close the stage, otherwise we'd be left in an invalid state. Closing the stage should be more visible though, so let's do that
+		if (StageState == EUsdStageState::OpenedAndLoaded)
+		{
+			UE_LOG(
+				LogUsd,
+				Warning,
+				TEXT("Closing the stage '%s' as no asset cache was provided, but the AUsdStageActor '%s' was set to open the stage and load assets. "
+					 "Either provide an asset cache or switch the stage actor to the 'Opened' state"),
+				*RootLayer.FilePath,
+				*GetPathName()
+			);
+			CloseUsdStage();
+		}
+
 		return;
 	}
 	FUsdScopedAssetCacheReferencer ScopedReferencer{UsdAssetCache, this};
@@ -2354,8 +2369,11 @@ void AUsdStageActor::SetRootLayer(const FString& RootFilePath)
 	UnloadUsdStage();
 	CloseUsdStage();
 	RootLayer.FilePath = RelativeFilePath;
-	OpenUsdStage();
-	LoadUsdStage();
+
+	// Don't call OpenUsdStage directly so that we can abort opening the stage in case the user cancels
+	// out of the missing asset cache dialog
+	const bool bOpenIfNeeded = true;
+	LoadUsdStage(bOpenIfNeeded);
 
 	// Do this here instead of on OpenUsdStage/LoadUsdStage as those also get called when changing any of
 	// our properties, like render context, material purpose, etc.
@@ -2385,8 +2403,10 @@ void AUsdStageActor::SetStageState(EUsdStageState NewState)
 	}
 	else if (StageState == EUsdStageState::OpenedAndLoaded)
 	{
-		OpenUsdStage();
-		LoadUsdStage();
+		// Don't call OpenUsdStage directly so that we can abort opening the stage in case the user cancels
+		// out of the missing asset cache dialog
+		const bool bOpenIfNeeded = true;
+		LoadUsdStage(bOpenIfNeeded);
 	}
 }
 
@@ -3026,21 +3046,10 @@ void AUsdStageActor::LoadUsdStage(bool bOpenIfNeeded)
 {
 	TRACE_CPUPROFILER_EVENT_SCOPE(AUsdStageActor::LoadUsdStage);
 
-	if (StageState != EUsdStageState::OpenedAndLoaded)
-	{
-		return;
-	}
-
-	if (!UsdStage && bOpenIfNeeded)
-	{
-		OpenUsdStage();
-		if (!UsdStage)
-		{
-			return;
-		}
-	}
-
-	// Ensure an asset cache if we're going to load something
+	// Ensure we have an asset cache before we open/load anything.
+	// In theory we don't need an asset cache for just opening the stage (especially if we have StageState that is not OpenAndLoaded),
+	// but that may not be immediately obvious to the user anyway, and it seems that the expected behavior is to abort loading the stage
+	// in case the user cancels out of the "missing asset cache" dialog
 	if (!RootLayer.FilePath.IsEmpty())
 	{
 		SetupAssetCacheIfNeeded();
@@ -3062,6 +3071,20 @@ void AUsdStageActor::LoadUsdStage(bool bOpenIfNeeded)
 		}
 	}
 
+	if (!UsdStage && bOpenIfNeeded)
+	{
+		OpenUsdStage();
+		if (!UsdStage)
+		{
+			return;
+		}
+	}
+
+	if (StageState != EUsdStageState::OpenedAndLoaded)
+	{
+		return;
+	}
+
 	double StartTime = FPlatformTime::Cycles64();
 
 	FScopedSlowTask SlowTask(1.f, LOCTEXT("LoadingUDStage", "Loading USD Stage"));
@@ -3274,6 +3297,7 @@ void AUsdStageActor::SetupAssetCacheIfNeeded()
 			}
 		}
 	}
+
 #if WITH_EDITOR
 	// Show a dialog to let the user create a new default asset cache somewhere
 	if (!UsdAssetCache && GIsEditor && !IsRunningCommandlet() && !IsTemplate())
@@ -3282,11 +3306,18 @@ void AUsdStageActor::SetupAssetCacheIfNeeded()
 		{
 			if (ProjectSettings->bShowCreateDefaultAssetCacheDialog)
 			{
-				bool bOutUserAccepted = false;
 				UUsdAssetCache2* NewCache = nullptr;
-				IUsdClassesEditorModule::ShowMissingDefaultAssetCacheDialog(NewCache, bOutUserAccepted);
 
-				if (bOutUserAccepted)
+				// Keep showing a dialog so that if the user clicks pick existing/create new and cancels, he can return to the dialog to pick
+				// something else
+				EDefaultAssetCacheDialogOption Outcome = EDefaultAssetCacheDialogOption::PickExisting;
+				while ((Outcome == EDefaultAssetCacheDialogOption::PickExisting || Outcome == EDefaultAssetCacheDialogOption::CreateNew) && !NewCache)
+				{
+					Outcome = IUsdClassesEditorModule::ShowMissingDefaultAssetCacheDialog(NewCache);
+				}
+
+				// We have an asset cache in some way: All good
+				if (NewCache)
 				{
 					ProjectSettings->DefaultAssetCache = NewCache;
 					ProjectSettings->SaveConfig();
@@ -3298,7 +3329,9 @@ void AUsdStageActor::SetupAssetCacheIfNeeded()
 						*UsdAssetCache->GetPathName()
 					);
 				}
-				else
+				// Return right now if the user canceled, so we don't fall back to the case below of creating
+				// a temp/transient asset cache
+				else if (Outcome == EDefaultAssetCacheDialogOption::Cancel)
 				{
 					return;
 				}
@@ -3306,6 +3339,7 @@ void AUsdStageActor::SetupAssetCacheIfNeeded()
 		}
 	}
 #endif
+
 	if (!UsdAssetCache)
 	{
 		UE_LOG(LogUsd, Warning, TEXT("USD Stage Actor '%s' had no previous USD Asset Cache and no default cache is specified on the project settings, so a temporary cache will be generated. For better performance, create a persistent USD Asset Cache asset and point to it with this actor's UsdAssetCache property."),
