USD: Fix warning when triggering light components to spawn during PIE.

diff --git a/Engine/Plugins/Importers/USDImporter/Source/USDSchemas/Private/USDGeomXformableTranslator.cpp b/Engine/Plugins/Importers/USDImporter/Source/USDSchemas/Private/USDGeomXformableTranslator.cpp
--- a/Engine/Plugins/Importers/USDImporter/Source/USDSchemas/Private/USDGeomXformableTranslator.cpp
+++ b/Engine/Plugins/Importers/USDImporter/Source/USDSchemas/Private/USDGeomXformableTranslator.cpp
@@ -653,9 +653,8 @@ void FUsdGeomXformableTranslator::UpdateComponents( USceneComponent* SceneCompon
 		SceneComponent->Modify();
 
 		// UsdToUnreal::ConvertXformable will set a new transform, which will emit warnings during PIE/Runtime if the component
-		// has Static mobility, so here we unregister, set the new transform value, and reregister below
-		const bool bStaticMobility = SceneComponent->Mobility == EComponentMobility::Static;
-		if ( bStaticMobility && SceneComponent->IsRegistered() )
+		// is not movable, so here we unregister, set the new transform value, and reregister below
+		if (SceneComponent->Mobility != EComponentMobility::Movable && SceneComponent->IsRegistered())
 		{
 			SceneComponent->UnregisterComponent();
 		}
