USD: Fix crash when enabling draw mode for a prim during PIE.

diff --git a/Engine/Plugins/Importers/USDImporter/Source/USDStage/Private/USDStageActor.cpp b/Engine/Plugins/Importers/USDImporter/Source/USDStage/Private/USDStageActor.cpp
--- a/Engine/Plugins/Importers/USDImporter/Source/USDStage/Private/USDStageActor.cpp
+++ b/Engine/Plugins/Importers/USDImporter/Source/USDStage/Private/USDStageActor.cpp
@@ -1097,6 +1097,10 @@ AUsdStageActor::AUsdStageActor()
 	Transactor = NewObject<UUsdTransactor>(this, UniqueName, EObjectFlags::RF_Transactional);
 	Transactor->Initialize(this);
 
+	// We never want to be without a valid BBoxCache or else we'll silently fail to compute bounds for all
+	// draw mode components we end up spawning
+	SetupBBoxCacheIfNeeded();
+
 	if (HasAuthorityOverStage())
 	{
 #if WITH_EDITOR
@@ -2169,12 +2173,9 @@ UUsdPrimTwin* AUsdStageActor::ExpandPrim(const UE::FUsdPrim& Prim, bool bResync,
 	}
 	else if (EUsdDrawMode DrawMode = UsdUtils::GetAppliedDrawMode(Prim); DrawMode != EUsdDrawMode::Default)
 	{
-		bHasAnimatedBounds = UsdUtils::HasAnimatedBounds(
-			Prim,
-			BBoxCache->GetIncludedPurposes(),
-			BBoxCache->GetUseExtentsHint(),
-			BBoxCache->GetIgnoreVisibility()
-		);
+		const bool bUseExtentsHint = true;
+		const bool bIgnoreVisibility = false;
+		bHasAnimatedBounds = UsdUtils::HasAnimatedBounds(Prim, static_cast<EUsdPurpose>(PurposesToLoad), bUseExtentsHint, bIgnoreVisibility);
 
 		if (bHasAnimatedBounds)
 		{
@@ -3119,12 +3120,6 @@ void AUsdStageActor::LoadUsdStage(bool bOpenIfNeeded)
 	{
 		InfoCache = MakeShared<FUsdInfoCache>();
 	}
-	if (!BBoxCache.IsValid())
-	{
-		const bool bUseExtentsHint = true;
-		const bool bIgnoreVisibility = false;
-		BBoxCache = MakeShared<UE::FUsdGeomBBoxCache>(Time, static_cast<EUsdPurpose>(PurposesToLoad), bUseExtentsHint, bIgnoreVisibility);
-	}
 
 	ReloadAnimations();
 
@@ -3350,6 +3345,18 @@ void AUsdStageActor::SetupAssetCacheIfNeeded()
 	}
 }
 
+void AUsdStageActor::SetupBBoxCacheIfNeeded()
+{
+	if (BBoxCache.IsValid())
+	{
+		return;
+	}
+
+	const bool bUseExtentsHint = true;
+	const bool bIgnoreVisibility = false;
+	BBoxCache = MakeShared<UE::FUsdGeomBBoxCache>(Time, static_cast<EUsdPurpose>(PurposesToLoad), bUseExtentsHint, bIgnoreVisibility);
+}
+
 UUsdPrimTwin* AUsdStageActor::GetRootPrimTwin()
 {
 	if (!RootUsdTwin)
diff --git a/Engine/Plugins/Importers/USDImporter/Source/USDStage/Public/USDStageActor.h b/Engine/Plugins/Importers/USDImporter/Source/USDStage/Public/USDStageActor.h
--- a/Engine/Plugins/Importers/USDImporter/Source/USDStage/Public/USDStageActor.h
+++ b/Engine/Plugins/Importers/USDImporter/Source/USDStage/Public/USDStageActor.h
@@ -353,6 +353,7 @@ protected:
 	void UnloadUsdStage();
 
 	void SetupAssetCacheIfNeeded();
+	void SetupBBoxCacheIfNeeded();
 
 	bool HasAuthorityOverStage() const;
 
