USD: Fix USD Stage Editor not updating when prims are modified during PIE.

diff --git a/Engine/Plugins/Importers/USDImporter/Source/USDStage/Private/USDStageActor.cpp b/Engine/Plugins/Importers/USDImporter/Source/USDStage/Private/USDStageActor.cpp
--- a/Engine/Plugins/Importers/USDImporter/Source/USDStage/Private/USDStageActor.cpp
+++ b/Engine/Plugins/Importers/USDImporter/Source/USDStage/Private/USDStageActor.cpp
@@ -1920,10 +1920,7 @@ void AUsdStageActor::OnUsdObjectsChanged(const UsdUtils::FObjectChangesByPath& I
 	// anymore, as that's how it refreshes whenever we delete a prim
 	for (const TPair<UE::FSdfPath, bool>& PrimChangedInfo : PrimsToUpdate)
 	{
-		if (HasAuthorityOverStage())
-		{
-			OnPrimChanged.Broadcast(PrimChangedInfo.Key.GetString(), PrimChangedInfo.Value);
-		}
+		OnPrimChanged.Broadcast(PrimChangedInfo.Key.GetString(), PrimChangedInfo.Value);
 	}
 
 	if (bNeedsAnimationReload)
