USD: Fixed simple collision for runtime support.

diff --git a/Engine/Plugins/Importers/USDImporter/Source/USDSchemas/Private/USDGeomMeshTranslator.cpp b/Engine/Plugins/Importers/USDImporter/Source/USDSchemas/Private/USDGeomMeshTranslator.cpp
--- a/Engine/Plugins/Importers/USDImporter/Source/USDSchemas/Private/USDGeomMeshTranslator.cpp
+++ b/Engine/Plugins/Importers/USDImporter/Source/USDSchemas/Private/USDGeomMeshTranslator.cpp
@@ -781,6 +781,115 @@ namespace UsdGeomMeshTranslatorImpl
 
 namespace UE::UsdCollision::Private
 {
+
+	class IMeshData
+	{
+	public:
+		virtual int32 GetNumVertices() const = 0;
+		virtual int32 GetNumIndices() const = 0;
+		virtual FVector3f GetVertexPosition(int32 Index) const = 0;
+		virtual int32 GetVertexIndex(int32 Index) const = 0;
+		virtual FBox GetBoundingBox() const = 0;
+		virtual bool IsValid() const = 0;
+		virtual ~IMeshData() {}
+	};
+
+	class FMeshDescriptionWrapper : public IMeshData
+	{
+	public:
+		FMeshDescriptionWrapper(const FMeshDescription& InMeshDescription)
+		: MeshDescription(InMeshDescription)
+		{
+		}
+
+		int32 GetNumVertices() const override { return MeshDescription.Vertices().Num(); }
+		int32 GetNumIndices() const override { return MeshDescription.VertexInstances().Num(); }
+		FVector3f GetVertexPosition(int32 Index) const override { return MeshDescription.GetVertexPositions()[Index]; }
+		int32 GetVertexIndex(int32 Index) const override { return MeshDescription.GetVertexInstanceVertex(Index); }
+		FBox GetBoundingBox() const override { return MeshDescription.ComputeBoundingBox(); }
+		bool IsValid() const override { return MeshDescription.Vertices().Num() > 0; };
+
+	private:
+		const FMeshDescription& MeshDescription;
+	};
+
+	class FRenderDataWrapper : public IMeshData
+	{
+	public:
+		FRenderDataWrapper(const FStaticMeshLODResources& InRenderData)
+		: RenderData(InRenderData)
+		{
+		}
+
+		int32 GetNumVertices() const override { return RenderData.GetNumVertices(); }
+		int32 GetNumIndices() const override { return RenderData.IndexBuffer.GetNumIndices(); }
+		FVector3f GetVertexPosition(int32 Index) const override { return RenderData.VertexBuffers.PositionVertexBuffer.VertexPosition(Index); }
+		int32 GetVertexIndex(int32 Index) const override { return RenderData.IndexBuffer.GetIndex(Index); }
+		bool IsValid() const override { return RenderData.GetNumVertices() > 0; };
+		FBox GetBoundingBox() const override 
+		{
+			if (!BoundingBox.IsValid)
+			{
+				for (int32 Index = 0; Index < RenderData.GetNumVertices(); ++Index)
+				{
+					BoundingBox += static_cast<FVector>(RenderData.VertexBuffers.PositionVertexBuffer.VertexPosition(Index));
+				}
+			}
+			return BoundingBox;
+		}
+
+	private:
+		const FStaticMeshLODResources& RenderData;
+		mutable FBox BoundingBox;
+	};
+
+	class FMeshData : public IMeshData
+	{
+	public:
+		FMeshData(const UStaticMesh& StaticMesh)
+		: MeshData(nullptr)
+		{
+#if WITH_EDITORONLY_DATA
+			FMeshDescription* MeshDescription = StaticMesh.GetMeshDescription(0);
+			if (!MeshDescription || MeshDescription->Vertices().Num() == 0)
+			{
+				return;
+			}
+
+			MeshData = new FMeshDescriptionWrapper(*MeshDescription);
+#else
+			if (StaticMesh.GetRenderData()->LODResources.Num() > 0)
+			{
+				MeshData = new FRenderDataWrapper(StaticMesh.GetRenderData()->LODResources[0]);
+			}
+#endif
+		}
+
+		FMeshData(const FMeshDescription& MeshDescription)
+		: MeshData(nullptr)
+		{
+			if (MeshDescription.Vertices().Num() > 0)
+			{
+				MeshData = new FMeshDescriptionWrapper(MeshDescription);
+			}
+		}
+
+		virtual ~FMeshData()
+		{
+			delete MeshData;
+		}
+
+		int32 GetNumVertices() const override { return IsValid() ? MeshData->GetNumVertices() : 0; }
+		int32 GetNumIndices() const override { return IsValid() ? MeshData->GetNumIndices() : 0; }
+		FVector3f GetVertexPosition(int32 Index) const override { return IsValid() ? MeshData->GetVertexPosition(Index) : FVector3f::ZeroVector; }
+		int32 GetVertexIndex(int32 Index) const override { return IsValid() ? MeshData->GetVertexIndex(Index) : 0; }
+		FBox GetBoundingBox() const override { return IsValid() ? MeshData->GetBoundingBox() : FBox(); }
+		bool IsValid() const override { return MeshData ? MeshData->IsValid() : false; }
+
+	private:
+		IMeshData* MeshData;
+	};
+
 	// Based on GeomFitUtils.cpp
 
 	// k-DOP (k-Discrete Oriented Polytopes) Direction Vectors
@@ -808,16 +917,15 @@ namespace UE::UsdCollision::Private
 		FVector(-RCP_SQRT2,  RCP_SQRT2, 0.f)
 	};
 
-	void GenerateKDopAsSimpleCollision(const FMeshDescription& MeshDescription, const TArray<FVector> &Dirs, FKAggregateGeom& CollisionShapes)
+	void GenerateKDopAsSimpleCollision(const FMeshData& MeshData, const TArray<FVector> &Dirs, FKAggregateGeom& CollisionShapes)
 	{
 		TArray<FVector> HullVertices;
-		TVertexAttributesConstRef<FVector3f> VertexPositions = MeshDescription.GetVertexPositions();
 		UE::Geometry::FitKDOPVertices3<double>(
 			Dirs, 
-			VertexPositions.GetNumElements(),
-			[&](int32 VertexId) 
+			MeshData.GetNumVertices(),
+			[&MeshData](int32 VertexId) 
 			{
-				return static_cast<FVector>(VertexPositions[VertexId]);
+				return static_cast<FVector>(MeshData.GetVertexPosition(VertexId));
 			},
 			HullVertices);
 
@@ -830,31 +938,29 @@ namespace UE::UsdCollision::Private
 		CollisionShapes.ConvexElems.Add(ConvexElem);
 	}
 
-	void CalcBoundingSphere(const FMeshDescription& MeshDescription, FSphere& Sphere)
+	void CalcBoundingSphere(const FMeshData& MeshData, FSphere& Sphere)
 	{
 		FBox Box;
 		FVector MinIx[3];
 		FVector MaxIx[3];
 
-		TVertexAttributesConstRef<FVector3f> VertexPositions = MeshDescription.GetVertexPositions();
-
 		bool bFirstVertex = true;
-		for (const FVertexID VertexID : MeshDescription.Vertices().GetElementIDs())
+		for (int32 VertexID = 0; VertexID < MeshData.GetNumVertices(); ++VertexID)
 		{
-			FVector Point = static_cast<FVector>(VertexPositions[VertexID]);
+			FVector Point = static_cast<FVector>(MeshData.GetVertexPosition(VertexID));
 			if (bFirstVertex)
 			{
 				// First, find AABB, remembering furthest points in each dir.
 				Box.Min = Point;
 				Box.Max = Box.Min;
 
-				MinIx[0] = static_cast<FVector>(VertexPositions[VertexID]);
-				MinIx[1] = static_cast<FVector>(VertexPositions[VertexID]);
-				MinIx[2] = static_cast<FVector>(VertexPositions[VertexID]);
+				MinIx[0] = Point;
+				MinIx[1] = Point;
+				MinIx[2] = Point;
 
-				MaxIx[0] = static_cast<FVector>(VertexPositions[VertexID]);
-				MaxIx[1] = static_cast<FVector>(VertexPositions[VertexID]);
-				MaxIx[2] = static_cast<FVector>(VertexPositions[VertexID]);
+				MaxIx[0] = Point;
+				MaxIx[1] = Point;
+				MaxIx[2] = Point;
 				bFirstVertex = false;
 				continue;
 			}
@@ -863,36 +969,36 @@ namespace UE::UsdCollision::Private
 			if (Point.X < Box.Min.X)
 			{
 				Box.Min.X = Point.X;
-				MinIx[0] = static_cast<FVector>(VertexPositions[VertexID]);
+				MinIx[0] = Point;
 			}
 			else if (Point.X > Box.Max.X)
 			{
 				Box.Max.X = Point.X;
-				MaxIx[0] = static_cast<FVector>(VertexPositions[VertexID]);
+				MaxIx[0] = Point;
 			}
 
 			// Y //
 			if (Point.Y < Box.Min.Y)
 			{
 				Box.Min.Y = Point.Y;
-				MinIx[1] = static_cast<FVector>(VertexPositions[VertexID]);
+				MinIx[1] = Point;
 			}
 			else if (Point.Y > Box.Max.Y)
 			{
 				Box.Max.Y = Point.Y;
-				MaxIx[1] = static_cast<FVector>(VertexPositions[VertexID]);
+				MaxIx[1] = Point;
 			}
 
 			// Z //
 			if (Point.Z < Box.Min.Z)
 			{
 				Box.Min.Z = Point.Z;
-				MinIx[2] = static_cast<FVector>(VertexPositions[VertexID]);
+				MinIx[2] = Point;
 			}
 			else if (Point.Z > Box.Max.Z)
 			{
 				Box.Max.Z = Point.Z;
-				MaxIx[2] = static_cast<FVector>(VertexPositions[VertexID]);
+				MaxIx[2] = Point;
 			}
 		}
 
@@ -918,9 +1024,9 @@ namespace UE::UsdCollision::Private
 		float Radius2 = FMath::Square(Radius);
 
 		// Now check each point lies within this sphere. If not - expand it a bit.
-		for (const FVertexID VertexID : MeshDescription.Vertices().GetElementIDs())
+		for (int32 VertexID = 0; VertexID < MeshData.GetNumVertices(); ++VertexID)
 		{
-			const FVector CenterToPoint = ((FVector)VertexPositions[VertexID]) - Sphere.Center;
+			const FVector CenterToPoint = static_cast<FVector>(MeshData.GetVertexPosition(VertexID)) - Sphere.Center;
 			const float CenterToPoint2 = CenterToPoint.SizeSquared();
 
 			// If this point is outside our current bounding sphere's radius
@@ -940,30 +1046,29 @@ namespace UE::UsdCollision::Private
 
 	// This is the one thats already used by unreal.
 	// Seems to do better with more symmetric input...
-	void CalcBoundingSphere2(const FMeshDescription& MeshDescription, FSphere& Sphere)
+	void CalcBoundingSphere2(const FMeshData& MeshData, FSphere& Sphere)
 	{
-		FVector Center = MeshDescription.ComputeBoundingBox().GetCenter();
+		FVector Center = MeshData.GetBoundingBox().GetCenter();
 
 		Sphere.Center = Center;
 		Sphere.W = 0.0f;
 
-		TVertexAttributesConstRef<FVector3f> VertexPositions = MeshDescription.GetVertexPositions();
-		for (const FVertexID VertexID : MeshDescription.Vertices().GetElementIDs())
+		for (int32 VertexID = 0; VertexID < MeshData.GetNumVertices(); ++VertexID)
 		{
-			float Dist2 = FVector::DistSquared((FVector)VertexPositions[VertexID], Sphere.Center);
+			float Dist2 = FVector::DistSquared(static_cast<FVector>(MeshData.GetVertexPosition(VertexID)), Sphere.Center);
 			if (Dist2 > Sphere.W)
 				Sphere.W = Dist2;
 		}
 		Sphere.W = FMath::Sqrt(Sphere.W);
 	}
 
-	void GenerateSphereAsSimpleCollision(const FMeshDescription& MeshDescription, FKAggregateGeom& CollisionShapes)
+	void GenerateSphereAsSimpleCollision(const FMeshData& MeshData, FKAggregateGeom& CollisionShapes)
 	{
 		FSphere Sphere, Sphere2, BestSphere;
 
 		// Calculate bounding sphere.
-		CalcBoundingSphere(MeshDescription, Sphere);
-		CalcBoundingSphere2(MeshDescription, Sphere2);
+		CalcBoundingSphere(MeshData, Sphere);
+		CalcBoundingSphere2(MeshData, Sphere2);
 
 		if(Sphere.W < Sphere2.W)
 			BestSphere = Sphere;
@@ -982,11 +1087,11 @@ namespace UE::UsdCollision::Private
 		CollisionShapes.SphereElems.Add(SphereElem);
 	}
 
-	void GenerateBoxAsSimpleCollision(const FMeshDescription& MeshDescription, FKAggregateGeom& CollisionShapes)
+	void GenerateBoxAsSimpleCollision(const FMeshData& MeshData, FKAggregateGeom& CollisionShapes)
 	{
 		// Calculate bounding Box.
 		FVector Center, Extents;
-		MeshDescription.ComputeBoundingBox().GetCenterAndExtents(Center, Extents);
+		MeshData.GetBoundingBox().GetCenterAndExtents(Center, Extents);
 
 		FKBoxElem BoxElem;
 		BoxElem.Center = Center;
@@ -996,10 +1101,10 @@ namespace UE::UsdCollision::Private
 		CollisionShapes.BoxElems.Add(BoxElem);
 	}
 
-	void CalcBoundingSphyl(const FMeshDescription& MeshDescription, FSphere& Sphere, float& Length, FRotator& Rotation)
+	void CalcBoundingSphyl(const FMeshData& MeshData, FSphere& Sphere, float& Length, FRotator& Rotation)
 	{
 		FVector Center, Extents;
-		MeshDescription.ComputeBoundingBox().GetCenterAndExtents(Center, Extents);
+		MeshData.GetBoundingBox().GetCenterAndExtents(Center, Extents);
 
 		Sphere.Center = Center;
 
@@ -1025,12 +1130,10 @@ namespace UE::UsdCollision::Private
 		float Radius = Extents.GetMax();
 		float Radius2 = FMath::Square(Radius);
 
-		TVertexAttributesConstRef<FVector3f> VertexPositions = MeshDescription.GetVertexPositions();
-
 		// Now check each point lies within this the radius. If not - expand it a bit.
-		for (const FVertexID VertexID : MeshDescription.Vertices().GetElementIDs())
+		for (int32 VertexID = 0; VertexID < MeshData.GetNumVertices(); ++VertexID)
 		{
-			FVector CenterToPoint = (static_cast<FVector>(VertexPositions[VertexID])) - Sphere.Center;
+			FVector CenterToPoint = static_cast<FVector>(MeshData.GetVertexPosition(VertexID)) - Sphere.Center;
 			CenterToPoint = Rotation.UnrotateVector(CenterToPoint);
 
 			const float PointRadius2 = CenterToPoint.SizeSquared2D();	// Ignore Z here...
@@ -1049,9 +1152,9 @@ namespace UE::UsdCollision::Private
 		float HalfLength = FMath::Max(0.0f, Extent - Radius);
 
 		// Now check each point lies within the length. If not - expand it a bit.
-		for (const FVertexID VertexID : MeshDescription.Vertices().GetElementIDs())
+		for (int32 VertexID = 0; VertexID < MeshData.GetNumVertices(); ++VertexID)
 		{
-			FVector CenterToPoint = (static_cast<FVector>(VertexPositions[VertexID])) - Sphere.Center;
+			FVector CenterToPoint = static_cast<FVector>(MeshData.GetVertexPosition(VertexID)) - Sphere.Center;
 			CenterToPoint = Rotation.UnrotateVector(CenterToPoint);
 
 			// If this point is outside our current bounding sphyl's length
@@ -1078,14 +1181,14 @@ namespace UE::UsdCollision::Private
 		Length = HalfLength * 2.0f;
 	}
 
-	void GenerateSphylAsSimpleCollision(const FMeshDescription& MeshDescription, FKAggregateGeom& CollisionShapes)
+	void GenerateSphylAsSimpleCollision(const FMeshData& MeshData, FKAggregateGeom& CollisionShapes)
 	{
 		FSphere Sphere;
 		float Length;
 		FRotator Rotation;
 
 		// Calculate bounding box.
-		CalcBoundingSphyl(MeshDescription, Sphere, Length, Rotation);
+		CalcBoundingSphyl(MeshData, Sphere, Length, Rotation);
 
 		// Don't use if radius is zero.
 		if (Sphere.W <= 0.f)
@@ -1107,15 +1210,14 @@ namespace UE::UsdCollision::Private
 		CollisionShapes.SphylElems.Add(SphylElem);
 	}
 
-	void ConvertMeshToSimpleCollision(const FMeshDescription& MeshDescription, FKAggregateGeom& CollisionShapes)
+	void ConvertMeshToSimpleCollision(const FMeshData& MeshData, FKAggregateGeom& CollisionShapes)
 	{
 		FKConvexElem ConvexElem;
-		ConvexElem.VertexData.Reserve(MeshDescription.Vertices().Num());
+		ConvexElem.VertexData.Reserve(MeshData.GetNumVertices());
 
-		TVertexAttributesConstRef<FVector3f> VertexPositions = MeshDescription.GetVertexPositions();
-		for (const FVertexID VertexID : MeshDescription.Vertices().GetElementIDs())
+		for (int32 VertexID = 0; VertexID < MeshData.GetNumVertices(); ++VertexID)
 		{
-			ConvexElem.VertexData.Add(static_cast<FVector>(VertexPositions[VertexID]));
+			ConvexElem.VertexData.Add(static_cast<FVector>(MeshData.GetVertexPosition(VertexID)));
 		}
 
 		// Note: UpdateElemBox also computes the convex hull indices
@@ -1241,9 +1343,8 @@ namespace UE::UsdCollision::Private
 			return;
 		}
 
-#if WITH_EDITORONLY_DATA
-		FMeshDescription* MeshDescription = StaticMesh.GetMeshDescription(0);
-		if (!MeshDescription || MeshDescription->Vertices().Num() == 0)
+		FMeshData MeshData(StaticMesh);
+		if (!MeshData.IsValid())
 		{
 			return;
 		}
@@ -1342,18 +1443,17 @@ namespace UE::UsdCollision::Private
 				// so the result would be not as good as convex decomposition anyway
 
 				TArray<FVector3f> Vertices;
-				Vertices.Reserve(MeshDescription->Vertices().Num());
-				TVertexAttributesConstRef<FVector3f> VertexPositions = MeshDescription->GetVertexPositions();
-				for (const FVertexID VertexID : MeshDescription->Vertices().GetElementIDs())
+				Vertices.Reserve(MeshData.GetNumVertices());
+				for (int32 VertexID = 0; VertexID < MeshData.GetNumVertices(); ++VertexID)
 				{
-					Vertices.Add(VertexPositions[VertexID]);
+					Vertices.Add(MeshData.GetVertexPosition(VertexID));
 				}
 
 				TArray<uint32> Indices;
-				Indices.Reserve(MeshDescription->VertexInstances().Num());
-				for (const FVertexInstanceID InstanceID : MeshDescription->VertexInstances().GetElementIDs())
+				Indices.Reserve(MeshData.GetNumIndices());
+				for (int32 InstanceID = 0; InstanceID < MeshData.GetNumIndices(); ++InstanceID)
 				{
-					Indices.Add(MeshDescription->GetVertexInstanceVertex(InstanceID).GetValue());
+					Indices.Add(MeshData.GetVertexIndex(InstanceID));
 				}
 
 				// Do not perform any action if we have invalid input
@@ -1368,27 +1468,27 @@ namespace UE::UsdCollision::Private
 #endif // WITH_EDITOR
 			case EUsdCollisionType::ConvexHull:
 			{
-				GenerateKDopAsSimpleCollision(*MeshDescription, KDopDir18, BodySetup->AggGeom);
+				GenerateKDopAsSimpleCollision(MeshData, KDopDir18, BodySetup->AggGeom);
 				break;
 			}
 			case EUsdCollisionType::Sphere:
 			{
-				GenerateSphereAsSimpleCollision(*MeshDescription, BodySetup->AggGeom);
+				GenerateSphereAsSimpleCollision(MeshData, BodySetup->AggGeom);
 				break;
 			}
 			case EUsdCollisionType::Cube:
 			{
-				GenerateBoxAsSimpleCollision(*MeshDescription, BodySetup->AggGeom);
+				GenerateBoxAsSimpleCollision(MeshData, BodySetup->AggGeom);
 				break;
 			}
 			case EUsdCollisionType::Capsule:
 			{
-				GenerateSphylAsSimpleCollision(*MeshDescription, BodySetup->AggGeom);
+				GenerateSphylAsSimpleCollision(MeshData, BodySetup->AggGeom);
 				break;
 			}
 			case EUsdCollisionType::CustomMesh:
 			{
-				ConvertMeshToSimpleCollision(*MeshDescription, BodySetup->AggGeom);
+				ConvertMeshToSimpleCollision(MeshData, BodySetup->AggGeom);
 				break;
 			}
 			default:
@@ -1397,11 +1497,12 @@ namespace UE::UsdCollision::Private
 
 		BodySetup->CreatePhysicsMeshes();
 
+#if WITH_EDITORONLY_DATA
 		StaticMesh.bCustomizedCollision = true;
+#endif
 
 		const bool bIsUpdate = true;
 		StaticMesh.CreateNavCollision(bIsUpdate);
-#endif // WITH_EDITORONLY_DATA
 	}
 }
 
