USD: Fixed possible crash when world outliner selection changes from a USD stage actor to a non-USD actor.

diff --git a/Engine/Plugins/Importers/USDImporter/Source/USDStageEditor/Private/SUSDStage.cpp b/Engine/Plugins/Importers/USDImporter/Source/USDStageEditor/Private/SUSDStage.cpp
--- a/Engine/Plugins/Importers/USDImporter/Source/USDStageEditor/Private/SUSDStage.cpp
+++ b/Engine/Plugins/Importers/USDImporter/Source/USDStageEditor/Private/SUSDStage.cpp
@@ -2652,7 +2652,8 @@ void SUsdStage::OnViewportSelectionChanged( UObject* NewSelection )
 
 		for ( AActor* SelectedActor : SelectedActors )
 		{
-			if ( SelectedActor )
+			// Some actors like AWorldDataLayer don't have a root component
+			if ( SelectedActor && SelectedActor->GetRootComponent() )
 			{
 				AllSelectedComponents.Add( SelectedActor->GetRootComponent() );
 			}
