USD: Prevent mesh prim with collision enabled from collapsing.

diff --git a/Engine/Plugins/Importers/USDImporter/Source/USDSchemas/Private/USDGeomMeshTranslator.cpp b/Engine/Plugins/Importers/USDImporter/Source/USDSchemas/Private/USDGeomMeshTranslator.cpp
--- a/Engine/Plugins/Importers/USDImporter/Source/USDSchemas/Private/USDGeomMeshTranslator.cpp
+++ b/Engine/Plugins/Importers/USDImporter/Source/USDSchemas/Private/USDGeomMeshTranslator.cpp
@@ -1948,6 +1948,21 @@ bool FUsdGeomMeshTranslator::CanBeCollapsed(ECollapsingType CollapsingType) cons
 		return false;
 	}
 
+	// Prevent collapse of mesh with collision enabled
+	if (pxr::UsdPhysicsCollisionAPI CollisionAPI{GetPrim()})
+	{
+		bool bIsCollisionEnabled = false;
+		if (pxr::UsdAttribute CollisionAttr = CollisionAPI.GetCollisionEnabledAttr())
+		{
+			CollisionAttr.Get(&bIsCollisionEnabled);
+		}
+
+		if (bIsCollisionEnabled)
+		{
+			return false;
+		}
+	}
+
 	// Prevent collapse of custom collision mesh
 	if (UsdUtils::IsCollisionMesh(Prim))
 	{
