USD: Fix nullptr deref when trying to handle draw mode prims that spawned no components.

diff --git a/Engine/Plugins/Importers/USDImporter/Source/USDStage/Private/USDStageActor.cpp b/Engine/Plugins/Importers/USDImporter/Source/USDStage/Private/USDStageActor.cpp
--- a/Engine/Plugins/Importers/USDImporter/Source/USDStage/Private/USDStageActor.cpp
+++ b/Engine/Plugins/Importers/USDImporter/Source/USDStage/Private/USDStageActor.cpp
@@ -2182,7 +2182,10 @@ UUsdPrimTwin* AUsdStageActor::ExpandPrim(const UE::FUsdPrim& Prim, bool bResync,
 
 			// Mark the component as animated right away because HasAnimatedBounds is expensive to call and
 			// we don't want to have to re-do it when creating the component
-			UsdPrimTwin->SceneComponent->SetMobility(EComponentMobility::Movable);
+			if (USceneComponent* PrimTwinsComponent = UsdPrimTwin->SceneComponent.Get())
+			{
+				PrimTwinsComponent->SetMobility(EComponentMobility::Movable);
+			}
 		}
 	}
 
