USD: Ensure when opening stages with indexed primvars.

diff --git a/Engine/Plugins/Importers/USDImporter/Source/USDUtilities/Private/USDGeomMeshConversion.cpp b/Engine/Plugins/Importers/USDImporter/Source/USDUtilities/Private/USDGeomMeshConversion.cpp
--- a/Engine/Plugins/Importers/USDImporter/Source/USDUtilities/Private/USDGeomMeshConversion.cpp
+++ b/Engine/Plugins/Importers/USDImporter/Source/USDUtilities/Private/USDGeomMeshConversion.cpp
@@ -2005,6 +2005,30 @@ namespace UE::UsdGeomMeshConversion::Private
 		return true;
 	}
 
+	void FlattenIndexedPrimvars(FUsdMeshData& InOutMeshData)
+	{
+		InOutMeshData.Points = ComputeFlattened(InOutMeshData.Points, InOutMeshData.PointIndices);
+		InOutMeshData.PointIndices = {};
+
+		InOutMeshData.Normals = ComputeFlattened(InOutMeshData.Normals, InOutMeshData.NormalIndices);
+		InOutMeshData.NormalIndices = {};
+
+		InOutMeshData.DisplayColors = ComputeFlattened(InOutMeshData.DisplayColors, InOutMeshData.DisplayColorIndices);
+		InOutMeshData.DisplayColorIndices = {};
+
+		InOutMeshData.DisplayOpacities = ComputeFlattened(InOutMeshData.DisplayOpacities, InOutMeshData.DisplayOpacityIndices);
+		InOutMeshData.DisplayOpacityIndices = {};
+
+		for (int32 UVSetIndex = 0; UVSetIndex < InOutMeshData.UVSets.Num(); ++UVSetIndex)
+		{
+			pxr::VtArray<pxr::GfVec2f>& UVs = InOutMeshData.UVSets[UVSetIndex];
+			pxr::VtArray<int>& UVIndices = InOutMeshData.UVSetIndices[UVSetIndex];
+
+			UVs = ComputeFlattened(UVs, UVIndices);
+			UVIndices = {};
+		}
+	}
+
 	bool ConvertMeshData(
 		const FUsdMeshData& InMeshData,
 		const FUsdStageInfo& InStageInfo,
@@ -2017,6 +2041,16 @@ namespace UE::UsdGeomMeshConversion::Private
 
 		FScopedUsdAllocs Allocs;
 
+		// ConvertMeshData can't handle indexed primvars! Make sure you call FlattenIndexedPrimvars beforehand
+		ensure(InMeshData.PointIndices.empty());
+		ensure(InMeshData.NormalIndices.empty());
+		ensure(InMeshData.DisplayColorIndices.empty());
+		ensure(InMeshData.DisplayOpacityIndices.empty());
+		for (int32 UVSetIndex = 0; UVSetIndex < InMeshData.UVSets.Num(); ++UVSetIndex)
+		{
+			ensure(InMeshData.UVSetIndices[UVSetIndex].empty());
+		}
+
 		// Material assignments
 		const TArray<UsdUtils::FUsdPrimMaterialSlot>& LocalMaterialSlots = InMeshData.LocalMaterialInfo.Slots;
 		const TArray<int32>& FaceMaterialIndices = InMeshData.LocalMaterialInfo.MaterialIndices;
@@ -2484,6 +2518,13 @@ bool UsdToUnreal::ConvertGeomMesh(
 		UsdGeomMeshImpl::SubdivideMeshData(UsdPrim, Options, MeshData);
 	}
 
+	// Make sure primvars are flattened before calling ConvertMeshData.
+	// We keep faceVarying indexed primvars within CollectMeshData as they are used for subdiv,
+	// and SubdivideMeshData will flatten them after subdivision.
+	// If we're not subdividing though we may still have some of these indexed primvars around,
+	// and ConvertMeshData can't handle them
+	UsdGeomMeshImpl::FlattenIndexedPrimvars(MeshData);
+
 	pxr::UsdStageRefPtr Stage = UsdPrim.GetStage();
 	const FUsdStageInfo StageInfo(Stage);
 	return UsdGeomMeshImpl::ConvertMeshData(MeshData, StageInfo, Options, OutMeshDescription, OutMaterialAssignments);
@@ -2830,6 +2871,8 @@ bool UsdToUnreal::ConvertGeomPrimitive(
 		MeshData.PointInterpolation = pxr::UsdGeomTokens->vertex;
 	}
 
+	FlattenIndexedPrimvars(MeshData);
+
 	pxr::UsdStageRefPtr Stage = InPrim.GetStage();
 	const FUsdStageInfo StageInfo(Stage);
 	return ConvertMeshData(MeshData, StageInfo, InOptions, InOutMeshDescription, InOutMaterialAssignments);
