USD: Fix initialization of some USTRUCT properties.

diff --git a/Engine/Plugins/Importers/USDImporter/Source/USDExporter/Public/USDMetadataExportOptions.h b/Engine/Plugins/Importers/USDImporter/Source/USDExporter/Public/USDMetadataExportOptions.h
--- a/Engine/Plugins/Importers/USDImporter/Source/USDExporter/Public/USDMetadataExportOptions.h
+++ b/Engine/Plugins/Importers/USDImporter/Source/USDExporter/Public/USDMetadataExportOptions.h
@@ -38,7 +38,7 @@ struct USDEXPORTER_API FUsdMetadataExportOptions
 	// When this is true, the "BlockedPrefixFilters" property is inverted, and describes prefixes to *allow*.
 	// In that case, entries are only allowed and exported if they match at least one of the provided prefixes.
 	UPROPERTY(EditAnywhere, BlueprintReadWrite, Category = "Metadata", config, meta = (EditCondition = "bExportAssetMetadata || bExportComponentMetadata"))
-	bool bInvertFilters;
+	bool bInvertFilters = false;
 };
 
 namespace UsdUtils
diff --git a/Engine/Plugins/Importers/USDImporter/Source/UnrealUSDWrapper/Public/UnrealUSDWrapper.h b/Engine/Plugins/Importers/USDImporter/Source/UnrealUSDWrapper/Public/UnrealUSDWrapper.h
--- a/Engine/Plugins/Importers/USDImporter/Source/UnrealUSDWrapper/Public/UnrealUSDWrapper.h
+++ b/Engine/Plugins/Importers/USDImporter/Source/UnrealUSDWrapper/Public/UnrealUSDWrapper.h
@@ -178,10 +178,10 @@ struct FMatrix2D
 	GENERATED_BODY()
 
 	UPROPERTY(EditAnywhere, BlueprintReadWrite, Category = "USD")
-	FVector2D Row0;
+	FVector2D Row0 = FVector2D(1.0, 0.0);
 
 	UPROPERTY(EditAnywhere, BlueprintReadWrite, Category = "USD")
-	FVector2D Row1;
+	FVector2D Row1 = FVector2D(0.0, 1.0);
 };
 
 /**
@@ -194,13 +194,13 @@ struct FMatrix3D
 	GENERATED_BODY()
 
 	UPROPERTY(EditAnywhere, BlueprintReadWrite, Category = "USD")
-	FVector Row0;
+	FVector Row0 = FVector(1.0, 0.0, 0.0);
 
 	UPROPERTY(EditAnywhere, BlueprintReadWrite, Category = "USD")
-	FVector Row1;
+	FVector Row1 = FVector(0.0, 1.0, 0.0);
 
 	UPROPERTY(EditAnywhere, BlueprintReadWrite, Category = "USD")
-	FVector Row2;
+	FVector Row2 = FVector(0.0, 0.0, 1.0);
 };
 
 class IUnrealUSDWrapperModule : public IModuleInterface
