USD: Fix allocator crash when adding a point instancer as a payload.

diff --git a/Engine/Plugins/Importers/USDImporter/Source/USDSchemas/Private/USDGeomPointInstancerTranslator.cpp b/Engine/Plugins/Importers/USDImporter/Source/USDSchemas/Private/USDGeomPointInstancerTranslator.cpp
--- a/Engine/Plugins/Importers/USDImporter/Source/USDSchemas/Private/USDGeomPointInstancerTranslator.cpp
+++ b/Engine/Plugins/Importers/USDImporter/Source/USDSchemas/Private/USDGeomPointInstancerTranslator.cpp
@@ -387,8 +387,6 @@ void FUsdGeomPointInstancerTranslator::UpdateComponents(USceneComponent* PointIn
 	// point this could be an UUsdDrawModeComponent
 	if (PointInstancerRootComponent->GetClass() == USceneComponent::StaticClass())
 	{
-		FScopedUsdAllocs UsdAllocs;
-
 		pxr::UsdPrim Prim = GetPrim();
 		pxr::UsdGeomPointInstancer PointInstancer( Prim );
 		if ( !PointInstancer )
@@ -396,8 +394,8 @@ void FUsdGeomPointInstancerTranslator::UpdateComponents(USceneComponent* PointIn
 			return;
 		}
 
-		pxr::SdfPathVector PrototypePaths;
-		if ( !PointInstancer.GetPrototypesRel().GetTargets( &PrototypePaths ) )
+		TUsdStore<pxr::SdfPathVector> PrototypePaths;
+		if ( !PointInstancer.GetPrototypesRel().GetTargets( &PrototypePaths.Get() ) )
 		{
 			return;
 		}
@@ -444,15 +442,16 @@ void FUsdGeomPointInstancerTranslator::UpdateComponents(USceneComponent* PointIn
 		}
 
 		TArray<TFuture<TTuple<UHierarchicalInstancedStaticMeshComponent*, TArray<FTransform>>>> Tasks;
-		FScopedSlowTask PrototypePathsSlowTask( ( float ) PrototypePaths.size(), LOCTEXT( "GeomPointUpdateComponents", "Updating HierarchicalInstancedStaticMeshComponents for point instancers" ) );
-		for ( int32 PrototypeIndex = 0; PrototypeIndex < PrototypePaths.size(); ++PrototypeIndex )
+		size_t NumPrototypePaths = PrototypePaths.Get().size();
+		FScopedSlowTask PrototypePathsSlowTask( ( float ) NumPrototypePaths, LOCTEXT( "GeomPointUpdateComponents", "Updating HierarchicalInstancedStaticMeshComponents for point instancers" ) );
+		for ( uint32 PrototypeIndex = 0; PrototypeIndex < NumPrototypePaths; ++PrototypeIndex )
 		{
 			PrototypePathsSlowTask.EnterProgressFrame();
 
-			pxr::SdfPath PrototypePath = PrototypePaths[PrototypeIndex];
+			pxr::SdfPath PrototypePath = PrototypePaths.Get()[PrototypeIndex];
 
-			pxr::UsdPrim PrototypeUsdPrim = Prim.GetStage()->GetPrimAtPath( PrototypePath );
-			if ( !PrototypeUsdPrim )
+			TUsdStore<pxr::UsdPrim> PrototypeUsdPrim = Prim.GetStage()->GetPrimAtPath( PrototypePath );
+			if ( !PrototypeUsdPrim.Get() )
 			{
 				UE_LOG(
 					LogUsd,
@@ -482,9 +481,9 @@ void FUsdGeomPointInstancerTranslator::UpdateComponents(USceneComponent* PointIn
 
 			// If our prototype was a LOD mesh we will have used the path of one of the actual LOD meshes to start the FGeomMeshCreateAssetsTaskChain,
 			// so we have to look for our resulting mesh with the same path
-			if ( Context->bAllowInterpretingLODs && UsdUtils::DoesPrimContainMeshLODs( PrototypeUsdPrim ) )
+			if ( Context->bAllowInterpretingLODs && UsdUtils::DoesPrimContainMeshLODs( PrototypeUsdPrim.Get() ) )
 			{
-				pxr::UsdPrimSiblingRange PrimRange = PrototypeUsdPrim.GetChildren();
+				pxr::UsdPrimSiblingRange PrimRange = PrototypeUsdPrim.Get().GetChildren();
 				for ( pxr::UsdPrimSiblingRange::iterator PrimRangeIt = PrimRange.begin(); PrimRangeIt != PrimRange.end(); ++PrimRangeIt )
 				{
 					const pxr::UsdPrim& Child = *PrimRangeIt;
@@ -532,7 +531,7 @@ void FUsdGeomPointInstancerTranslator::UpdateComponents(USceneComponent* PointIn
 				}
 
 				MeshTranslationImpl::SetMaterialOverrides(
-					PrototypeUsdPrim,
+					PrototypeUsdPrim.Get(),
 					ExistingAssignments,
 					*HISMComponent,
 					AssetCache,
