USD: Fix crash when trying to cook asset caches with persistent static meshes.

diff --git a/Engine/Plugins/Importers/USDImporter/Source/USDClasses/Private/USDAssetCache2.cpp b/Engine/Plugins/Importers/USDImporter/Source/USDClasses/Private/USDAssetCache2.cpp
--- a/Engine/Plugins/Importers/USDImporter/Source/USDClasses/Private/USDAssetCache2.cpp
+++ b/Engine/Plugins/Importers/USDImporter/Source/USDClasses/Private/USDAssetCache2.cpp
@@ -526,7 +526,7 @@ namespace UE::AssetCache::Private
 			// We only need these in the editor: When we're cooking we'll just save the cooked mesh data into the
 			// asset itself, so at runtime there is no need for an additional "build" step that we have to handle here
 
-			if (FApp::CanEverRender())
+			if (FApp::CanEverRender() || !FPlatformProperties::RequiresCookedData())
 			{
 				StaticMesh->SetRenderData(MakeUnique< FStaticMeshRenderData >());
 				ITargetPlatformManagerModule& TargetPlatformManager = GetTargetPlatformManagerRef();
diff --git a/Engine/Plugins/Importers/USDImporter/Source/USDSchemas/Private/USDGeomMeshTranslator.cpp b/Engine/Plugins/Importers/USDImporter/Source/USDSchemas/Private/USDGeomMeshTranslator.cpp
--- a/Engine/Plugins/Importers/USDImporter/Source/USDSchemas/Private/USDGeomMeshTranslator.cpp
+++ b/Engine/Plugins/Importers/USDImporter/Source/USDSchemas/Private/USDGeomMeshTranslator.cpp
@@ -756,26 +756,24 @@ namespace UsdGeomMeshTranslatorImpl
 		// For runtime builds, the analogue for this stuff is already done from within BuildFromMeshDescriptions
 		TRACE_CPUPROFILER_EVENT_SCOPE(UsdGeomMeshTranslatorImpl::PostBuildStaticMesh);
 
-		if (!FApp::CanEverRender())
+		if (FApp::CanEverRender() || !FPlatformProperties::RequiresCookedData())
 		{
-			return;
-		}
-
-		StaticMesh.InitResources();
+			StaticMesh.InitResources();
 
 #if WITH_EDITOR
-		// Fetch the MeshDescription from the StaticMesh because we'll have moved it away from LODIndexToMeshDescription CreateStaticMesh
-		if (const FMeshDescription* MeshDescription = StaticMesh.GetMeshDescription(0))
-		{
-			StaticMesh.GetRenderData()->Bounds = MeshDescription->GetBounds();
-		}
-		StaticMesh.CalculateExtendedBounds();
-		StaticMesh.ClearMeshDescriptions(); // Clear mesh descriptions to reduce memory usage, they are kept only in bulk data form
+			// Fetch the MeshDescription from the StaticMesh because we'll have moved it away from LODIndexToMeshDescription CreateStaticMesh
+			if (const FMeshDescription* MeshDescription = StaticMesh.GetMeshDescription(0))
+			{
+				StaticMesh.GetRenderData()->Bounds = MeshDescription->GetBounds();
+			}
+			StaticMesh.CalculateExtendedBounds();
+			StaticMesh.ClearMeshDescriptions(); // Clear mesh descriptions to reduce memory usage, they are kept only in bulk data form
 #else
-		// Fetch the MeshDescription from the imported LODIndexToMeshDescription as StaticMesh.GetMeshDescription is editor-only
-		StaticMesh.GetRenderData()->Bounds = LODIndexToMeshDescription[0].GetBounds();
-		StaticMesh.CalculateExtendedBounds();
+			// Fetch the MeshDescription from the imported LODIndexToMeshDescription as StaticMesh.GetMeshDescription is editor-only
+			StaticMesh.GetRenderData()->Bounds = LODIndexToMeshDescription[0].GetBounds();
+			StaticMesh.CalculateExtendedBounds();
 #endif // WITH_EDITOR
+		}
 	}
 }
 
diff --git a/Engine/Plugins/Importers/USDImporter/Source/USDSchemas/Private/USDGeomXformableTranslator.cpp b/Engine/Plugins/Importers/USDImporter/Source/USDSchemas/Private/USDGeomXformableTranslator.cpp
--- a/Engine/Plugins/Importers/USDImporter/Source/USDSchemas/Private/USDGeomXformableTranslator.cpp
+++ b/Engine/Plugins/Importers/USDImporter/Source/USDSchemas/Private/USDGeomXformableTranslator.cpp
@@ -679,7 +679,7 @@ void FUsdGeomXformableTranslator::UpdateComponents( USceneComponent* SceneCompon
 			{
 				// Need to make sure the mesh's resources are initialized here as it may have just been built in another thread
 				// Only do this if required though, as this mesh could using these resources currently (e.g. PIE and editor world sharing the mesh)
-				if ( PrimStaticMesh && !PrimStaticMesh->AreRenderingResourcesInitialized() && FApp::CanEverRender() )
+				if (PrimStaticMesh && !PrimStaticMesh->AreRenderingResourcesInitialized() && (FApp::CanEverRender() || !FPlatformProperties::RequiresCookedData()))
 				{
 					PrimStaticMesh->InitResources();
 				}
