USD: Fix stage state not being considered when handling USD notices.

diff --git a/Engine/Plugins/Importers/USDImporter/Source/USDStage/Private/USDStageActor.cpp b/Engine/Plugins/Importers/USDImporter/Source/USDStage/Private/USDStageActor.cpp
--- a/Engine/Plugins/Importers/USDImporter/Source/USDStage/Private/USDStageActor.cpp
+++ b/Engine/Plugins/Importers/USDImporter/Source/USDStage/Private/USDStageActor.cpp
@@ -1396,6 +1396,40 @@ void AUsdStageActor::OnUsdObjectsChanged(const UsdUtils::FObjectChangesByPath& I
 		return;
 	}
 
+	// If we're opened, we shouldn't have any actor, component or asset. We shouldn't even have a built info cache!
+	// This means we don't really need to do anything here, except trigger the USD Stage Editor to refresh.
+	// We'd otherwise go through a lot of work to find out the prims to update, but given that this is only for UI refresh
+	// and that we don't even have an info cache anyway, let's just have a simple loop over all prims mentioned in the
+	// notice and refresh the stage editor with them.
+	// Note that the stage editor only refreshes once per tick anyway, so this shouldn't even cause any unnecessary refresh spam
+	if (StageState == EUsdStageState::Opened)
+	{
+		TMap<UE::FSdfPath, bool> PrimsToUpdateOrResync;
+		for (const TPair<FString, TArray<UsdUtils::FSdfChangeListEntry>>& Change : InfoChanges)
+		{
+			const bool bIsResync = false;
+			PrimsToUpdateOrResync.Add({UE::FSdfPath(*Change.Key).StripAllVariantSelections(), bIsResync});
+		}
+		// Resyncs afterward so they overwrite
+		for (const TPair<FString, TArray<UsdUtils::FSdfChangeListEntry>>& Change : ResyncChanges)
+		{
+			const bool bIsResync = true;
+			PrimsToUpdateOrResync.Add({UE::FSdfPath(*Change.Key).StripAllVariantSelections(), bIsResync});
+		}
+
+		for (const TPair<UE::FSdfPath, bool>& PrimAndResync : PrimsToUpdateOrResync)
+		{
+			OnPrimChanged.Broadcast(PrimAndResync.Key.GetString(), PrimAndResync.Value);
+		}
+		return;
+	}
+	else if (StageState == EUsdStageState::Closed)
+	{
+		// If we're in the closed state we shouldn't have a stage, so we shouldn't ever get a notice
+		ensure(false);
+		return;
+	}
+
 	// Only update the transactor if we're listening to USD notices. Within OnObjectPropertyChanged we will stop listening when writing stage changes
 	// from our component changes, and this will also make sure we're not duplicating the events we store and replicate via multi-user: If a modification
 	// can be described purely via UObject changes, then those changes will be responsible for the whole modification and we won't record the corresponding
@@ -1433,18 +1467,15 @@ void AUsdStageActor::OnUsdObjectsChanged(const UsdUtils::FObjectChangesByPath& I
 	{
 		// If the user canceled out of providing an asset cache and we should be loading assets, we have no choice but to either change the stage
 		// state or close the stage, otherwise we'd be left in an invalid state. Closing the stage should be more visible though, so let's do that
-		if (StageState == EUsdStageState::OpenedAndLoaded)
-		{
-			UE_LOG(
-				LogUsd,
-				Warning,
-				TEXT("Closing the stage '%s' as no asset cache was provided, but the AUsdStageActor '%s' was set to open the stage and load assets. "
-					 "Either provide an asset cache or switch the stage actor to the 'Opened' state"),
-				*RootLayer.FilePath,
-				*GetPathName()
-			);
-			CloseUsdStage();
-		}
+		UE_LOG(
+			LogUsd,
+			Warning,
+			TEXT("Closing the stage '%s' as no asset cache was provided, but the AUsdStageActor '%s' was set to open the stage and load assets. "
+					"Either provide an asset cache or switch the stage actor to the 'Opened' state"),
+			*RootLayer.FilePath,
+			*GetPathName()
+		);
+		CloseUsdStage();
 
 		return;
 	}
