USD: Fix creation of memory-only stages when Stage State is 'Closed'.

diff --git a/Engine/Plugins/Importers/USDImporter/Source/USDStage/Private/USDStageActor.cpp b/Engine/Plugins/Importers/USDImporter/Source/USDStage/Private/USDStageActor.cpp
--- a/Engine/Plugins/Importers/USDImporter/Source/USDStage/Private/USDStageActor.cpp
+++ b/Engine/Plugins/Importers/USDImporter/Source/USDStage/Private/USDStageActor.cpp
@@ -1298,6 +1298,24 @@ void AUsdStageActor::NewStage()
 
 	NewStage.SetDefaultPrim(RootPrim);
 
+	// Call OpenStage to intentionally put the new stage within the usdutils stage cache if we're on "Closed" state.
+	// This is important because at least for now we want memory-only stages to always stick around after "closed",
+	// so we can undo/redo back into them.
+	// Yes, this is strange and will be removed eventually, but is needed now or else we'll get undo/redo crashes.
+	//
+	// Normally the new stage would naturally end up in the stage cache because SetRootLayer calls OpenUsdStage,
+	// which also calls OpenStage on our RootLayer path. The issue here is that OpenUsdStage will (correctly) not
+	// do anything if StageState is Closed, so here we need to cache that stage ourselves.
+	//
+	// This trick has the effect that clicking "New Stage" when StageState == Closed will stealthily open the stage
+	// and put it into the stage cache, but not open it *on the stage actor itself*. When changing stage state
+	// to e.g. "Opened", we'll try opening the stage on our RootLayer path and successfully end up opening that
+	// memory-only stage, as the RootLayer will contain its identifier.
+	if (StageState == EUsdStageState::Closed)
+	{
+		UnrealUSDWrapper::OpenStage(*StagePath, InitialLoadSet);
+	}
+
 	SetRootLayer(StagePath);
 #endif // USE_USD_SDK
 }
@@ -2887,7 +2905,10 @@ void AUsdStageActor::OpenUsdStage()
 
 	if (UsdStage)
 	{
-		UsdStage.SetEditTarget(UsdStage.GetRootLayer());
+		if (!UsdStage.IsEditTargetValid())
+		{
+			UsdStage.SetEditTarget(UsdStage.GetRootLayer());
+		}
 
 		UsdStage.SetInterpolationType(InterpolationType);
 
