USD: Reformat USDUtilities module.

#rnx

diff --git a/Engine/Plugins/Importers/USDImporter/Source/USDUtilities/Private/USDAttributeUtils.cpp b/Engine/Plugins/Importers/USDImporter/Source/USDUtilities/Private/USDAttributeUtils.cpp
--- a/Engine/Plugins/Importers/USDImporter/Source/USDUtilities/Private/USDAttributeUtils.cpp
+++ b/Engine/Plugins/Importers/USDImporter/Source/USDUtilities/Private/USDAttributeUtils.cpp
@@ -3,7 +3,6 @@
 #include "USDAttributeUtils.h"
 
 #include "UnrealUSDWrapper.h"
-#include "USDErrorUtils.h"
 #include "USDLayerUtils.h"
 #include "USDLog.h"
 #include "USDProjectSettings.h"
@@ -21,48 +20,43 @@
 
 #if USE_USD_SDK
 #include "USDIncludesStart.h"
-	#include "pxr/usd/pcp/cache.h"
-	#include "pxr/usd/pcp/primIndex.h"
-	#include "pxr/usd/pcp/propertyIndex.h"
-	#include "pxr/usd/sdf/layer.h"
-	#include "pxr/usd/sdf/layerUtils.h"
-	#include "pxr/usd/sdf/spec.h"
-	#include "pxr/usd/usd/attribute.h"
-	#include "pxr/usd/usd/editContext.h"
-	#include "pxr/usd/usd/stage.h"
-	#include "pxr/usd/usd/usdFileFormat.h"
+#include "pxr/usd/pcp/primIndex.h"
+#include "pxr/usd/sdf/layerUtils.h"
+#include "pxr/usd/usd/attribute.h"
+#include "pxr/usd/usd/editContext.h"
+#include "pxr/usd/usd/stage.h"
 #include "USDIncludesEnd.h"
-#endif // #if USE_USD_SDK
+#endif	  // #if USE_USD_SDK
 
 #define LOCTEXT_NAMESPACE "UsdAttributeUtils"
 
 namespace UsdUtils
 {
 #if USE_USD_SDK
-	const pxr::TfToken MutedToken = UnrealToUsd::ConvertToken( TEXT( "UE:Muted" ) ).Get();
-#endif // #if USE_USD_SDK
+	const pxr::TfToken MutedToken = UnrealToUsd::ConvertToken(TEXT("UE:Muted")).Get();
+#endif	  // #if USE_USD_SDK
 }
 
-bool UsdUtils::MuteAttribute( UE::FUsdAttribute& Attribute, const UE::FUsdStage& Stage )
+bool UsdUtils::MuteAttribute(UE::FUsdAttribute& Attribute, const UE::FUsdStage& Stage)
 {
 #if USE_USD_SDK
 	FScopedUsdAllocs Allocs;
 
-	const pxr::UsdAttribute& UsdAttribute = static_cast< const pxr::UsdAttribute& >( Attribute );
-	const pxr::UsdStageRefPtr UsdStage{ Stage };
-	if ( !UsdAttribute || !UsdStage )
+	const pxr::UsdAttribute& UsdAttribute = static_cast<const pxr::UsdAttribute&>(Attribute);
+	const pxr::UsdStageRefPtr UsdStage{Stage};
+	if (!UsdAttribute || !UsdStage)
 	{
 		return false;
 	}
 
-	pxr::SdfLayerRefPtr UEPersistentState = UsdUtils::GetUEPersistentStateSublayer( Stage );
-	if ( !UEPersistentState )
+	pxr::SdfLayerRefPtr UEPersistentState = UsdUtils::GetUEPersistentStateSublayer(Stage);
+	if (!UEPersistentState)
 	{
 		return false;
 	}
 
-	pxr::SdfLayerRefPtr UESessionState = UsdUtils::GetUESessionStateSublayer( Stage );
-	if ( !UESessionState )
+	pxr::SdfLayerRefPtr UESessionState = UsdUtils::GetUESessionStateSublayer(Stage);
+	if (!UESessionState)
 	{
 		return false;
 	}
@@ -71,22 +65,22 @@ bool UsdUtils::MuteAttribute( UE::FUsdAttribute& Attribute, const UE::FUsdStage&
 
 	// Mark it as muted on the persistent state
 	{
-		pxr::UsdEditContext Context( Stage, UEPersistentState );
+		pxr::UsdEditContext Context(Stage, UEPersistentState);
 
-		UsdAttribute.SetCustomDataByKey( MutedToken, pxr::VtValue{ true } );
+		UsdAttribute.SetCustomDataByKey(MutedToken, pxr::VtValue{true});
 	}
 
 	// Actually author the opinions that cause it to be muted on the session state
 	{
-		pxr::UsdEditContext Context( Stage, UESessionState );
+		pxr::UsdEditContext Context(Stage, UESessionState);
 
 		pxr::VtValue Value;
-		UsdAttribute.Get( &Value, pxr::UsdTimeCode::Default() );
+		UsdAttribute.Get(&Value, pxr::UsdTimeCode::Default());
 
 		// Clear the attribute so that it also gets rid of any time samples it may have
 		UsdAttribute.Clear();
 
-		if ( Value.IsEmpty() )
+		if (Value.IsEmpty())
 		{
 			// It doesn't have any default value, so just mute the attribute completely
 			UsdAttribute.Block();
@@ -94,41 +88,41 @@ bool UsdUtils::MuteAttribute( UE::FUsdAttribute& Attribute, const UE::FUsdStage&
 		else
 		{
 			// It has a default, non-animated value from a weaker opinion: Use that instead
-			UsdAttribute.Set( Value );
+			UsdAttribute.Set(Value);
 		}
 	}
 
 	return true;
 #else
 	return false;
-#endif // #if USE_USD_SDK
+#endif	  // #if USE_USD_SDK
 }
 
-bool UsdUtils::UnmuteAttribute( UE::FUsdAttribute& Attribute, const UE::FUsdStage& Stage )
+bool UsdUtils::UnmuteAttribute(UE::FUsdAttribute& Attribute, const UE::FUsdStage& Stage)
 {
 #if USE_USD_SDK
 	FScopedUsdAllocs Allocs;
 
-	pxr::UsdAttribute& UsdAttribute = static_cast< pxr::UsdAttribute& >( Attribute );
-	const pxr::UsdStageRefPtr UsdStage{ Stage };
-	if ( !UsdAttribute || !UsdStage )
+	pxr::UsdAttribute& UsdAttribute = static_cast<pxr::UsdAttribute&>(Attribute);
+	const pxr::UsdStageRefPtr UsdStage{Stage};
+	if (!UsdAttribute || !UsdStage)
 	{
 		return false;
 	}
 
-	if ( !IsAttributeMuted( Attribute, Stage ) )
+	if (!IsAttributeMuted(Attribute, Stage))
 	{
 		return true;
 	}
 
-	pxr::SdfLayerRefPtr UEPersistentState = UsdUtils::GetUEPersistentStateSublayer( Stage );
-	if ( !UEPersistentState )
+	pxr::SdfLayerRefPtr UEPersistentState = UsdUtils::GetUEPersistentStateSublayer(Stage);
+	if (!UEPersistentState)
 	{
 		return false;
 	}
 
-	pxr::SdfLayerRefPtr UESessionState = UsdUtils::GetUESessionStateSublayer( Stage );
-	if ( !UESessionState )
+	pxr::SdfLayerRefPtr UESessionState = UsdUtils::GetUESessionStateSublayer(Stage);
+	if (!UESessionState)
 	{
 		return false;
 	}
@@ -137,39 +131,39 @@ bool UsdUtils::UnmuteAttribute( UE::FUsdAttribute& Attribute, const UE::FUsdStag
 
 	// Remove the mute tag on the persistent state layer
 	{
-		pxr::UsdEditContext Context( Stage, UEPersistentState );
-		UsdAttribute.ClearCustomDataByKey( MutedToken );
+		pxr::UsdEditContext Context(Stage, UEPersistentState);
+		UsdAttribute.ClearCustomDataByKey(MutedToken);
 	}
 
 	// Clear our opinion of it on our session state layer
 	{
-		pxr::UsdEditContext Context( Stage, UESessionState );
+		pxr::UsdEditContext Context(Stage, UESessionState);
 		UsdAttribute.Clear();
 	}
 
 	return true;
 #else
 	return false;
-#endif // #if USE_USD_SDK
+#endif	  // #if USE_USD_SDK
 }
 
-bool UsdUtils::IsAttributeMuted( const UE::FUsdAttribute& Attribute, const UE::FUsdStage& Stage )
+bool UsdUtils::IsAttributeMuted(const UE::FUsdAttribute& Attribute, const UE::FUsdStage& Stage)
 {
 #if USE_USD_SDK
 	FScopedUsdAllocs Allocs;
 
-	const pxr::UsdAttribute& UsdAttribute = static_cast< const pxr::UsdAttribute& >( Attribute );
-	if ( !UsdAttribute )
+	const pxr::UsdAttribute& UsdAttribute = static_cast<const pxr::UsdAttribute&>(Attribute);
+	if (!UsdAttribute)
 	{
 		return false;
 	}
 
-	pxr::VtValue Data = UsdAttribute.GetCustomDataByKey( MutedToken );
-	if ( Data.IsHolding<bool>() )
+	pxr::VtValue Data = UsdAttribute.GetCustomDataByKey(MutedToken);
+	if (Data.IsHolding<bool>())
 	{
 		return Data.Get<bool>();
 	}
-#endif // #if USE_USD_SDK
+#endif	  // #if USE_USD_SDK
 
 	return false;
 }
@@ -237,7 +231,12 @@ void UsdUtils::NotifyIfOverriddenOpinion(const pxr::UsdProperty& Property)
 			const FText Text = LOCTEXT("OverridenOpinionText", "USD: Overridden opinion");
 
 			const FText SubText = FText::Format(
-				LOCTEXT("OverridenOpinionSubText", "Opinion authored for this attribute:\n\n{0}\n\nAt this layer:\n\n{1}\n\nIs overridden by another spec at this layer:\n\n{2}\n\nAnd so may not be visible on the composed stage. This means this edit may not be visible once the stage is reloaded."),
+				LOCTEXT(
+					"OverridenOpinionSubText",
+					"Opinion authored for this attribute:\n\n{0}\n\nAt this layer:\n\n{1}\n\nIs overridden by another spec at this "
+					"layer:\n\n{2}\n\nAnd so may not be visible on the composed stage. This means this edit may not be visible once the stage is "
+					"reloaded."
+				),
 				FText::FromString(UsdToUnreal::ConvertPath(Spec->GetPath())),
 				FText::FromString(UsdToUnreal::ConvertString(Layer->GetIdentifier())),
 				FText::FromString(UsdToUnreal::ConvertString(SpecLayer->GetIdentifier()))
@@ -263,17 +262,21 @@ void UsdUtils::NotifyIfOverriddenOpinion(const pxr::UsdProperty& Property)
 				Toast.ButtonDetails.Emplace(
 					LOCTEXT("OverridenOpinionMessageOk", "Ok"),
 					FText::GetEmpty(),
-					FSimpleDelegate::CreateLambda([]() {
-						if (TSharedPtr<SNotificationItem> PinnedNotification = Notification.Pin())
+					FSimpleDelegate::CreateLambda(
+						[]()
 						{
-							PinnedNotification->SetCompletionState(SNotificationItem::CS_Success);
-							PinnedNotification->ExpireAndFadeout();
+							if (TSharedPtr<SNotificationItem> PinnedNotification = Notification.Pin())
+							{
+								PinnedNotification->SetCompletionState(SNotificationItem::CS_Success);
+								PinnedNotification->ExpireAndFadeout();
+							}
 						}
-						})
+					)
 				);
 				// This is flipped because the default checkbox message is "Don't prompt again"
 				Toast.CheckBoxState = Settings->bShowOverriddenOpinionsWarning ? ECheckBoxState::Unchecked : ECheckBoxState::Checked;
-				Toast.CheckBoxStateChanged = FOnCheckStateChanged::CreateStatic([](ECheckBoxState NewState)
+				Toast.CheckBoxStateChanged = FOnCheckStateChanged::CreateStatic(
+					[](ECheckBoxState NewState)
 					{
 						if (UUsdProjectSettings* Settings = GetMutableDefault<UUsdProjectSettings>())
 						{
@@ -281,7 +284,8 @@ void UsdUtils::NotifyIfOverriddenOpinion(const pxr::UsdProperty& Property)
 							Settings->bShowOverriddenOpinionsWarning = NewState == ECheckBoxState::Unchecked;
 							Settings->SaveConfig();
 						}
-					});
+					}
+				);
 
 				// Only show one at a time
 				if (!Notification.IsValid())
@@ -368,7 +372,12 @@ bool UsdUtils::NotifyIfInstanceProxy(const pxr::UsdPrim& Prim)
 				}
 				case EUsdEditInInstanceBehavior::RemoveInstanceable:
 				{
-					UE_LOG(LogUsd, Log, TEXT("Removing all instanceable flags from ancestors of prim '%s'"), *UsdToUnreal::ConvertPath(Prim.GetPrimPath()));
+					UE_LOG(
+						LogUsd,
+						Log,
+						TEXT("Removing all instanceable flags from ancestors of prim '%s'"),
+						*UsdToUnreal::ConvertPath(Prim.GetPrimPath())
+					);
 					RemoveInstanceables();
 
 					// We shouldn't be instanceable now, so we can probably author whatever we wanted to author
@@ -399,7 +408,8 @@ bool UsdUtils::NotifyIfInstanceProxy(const pxr::UsdPrim& Prim)
 						LOCTEXT(
 							"AuthoringInsideInstanceSubText",
 							"Trying to author an opinion at or below prim:\n\n{0}\n\nThis prim is an instance proxy, and so cannot be "
-							"edited directly.\n\nIf you wish to modify just this prim, you can remove the instanceable flag for this prim and try again. If you wish to "
+							"edited directly.\n\nIf you wish to modify just this prim, you can remove the instanceable flag for this prim and try "
+							"again. If you wish to "
 							"modify all instances, please edit the prims referenced by prim '{1}' directly."
 						),
 						FText::FromString(UsdToUnreal::ConvertPath(Prim.GetPrimPath())),
@@ -469,9 +479,8 @@ bool UsdUtils::NotifyIfInstanceProxy(const pxr::UsdPrim& Prim)
 						)
 					);
 					// This is flipped because the default checkbox message is "Don't prompt again"
-					Toast.CheckBoxState = Settings->EditInInstanceableBehavior == EUsdEditInInstanceBehavior::ShowPrompt
-											  ? ECheckBoxState::Unchecked
-											  : ECheckBoxState::Checked;
+					Toast.CheckBoxState = Settings->EditInInstanceableBehavior == EUsdEditInInstanceBehavior::ShowPrompt ? ECheckBoxState::Unchecked
+																														 : ECheckBoxState::Checked;
 					Toast.CheckBoxStateChanged = FOnCheckStateChanged::CreateStatic(
 						[](ECheckBoxState NewState)
 						{
@@ -509,6 +518,6 @@ bool UsdUtils::NotifyIfInstanceProxy(const pxr::UsdPrim& Prim)
 
 	return false;
 }
-#endif // #if USE_USD_SDK
+#endif	  // #if USE_USD_SDK
 
-#undef LOCTEXT_NAMESPACE
\ No newline at end of file
+#undef LOCTEXT_NAMESPACE
