USD: Reformat USDUtilities module.

diff --git a/Engine/Plugins/Importers/USDImporter/Source/USDUtilities/Private/USDConversionUtils.cpp b/Engine/Plugins/Importers/USDImporter/Source/USDUtilities/Private/USDConversionUtils.cpp
--- a/Engine/Plugins/Importers/USDImporter/Source/USDUtilities/Private/USDConversionUtils.cpp
+++ b/Engine/Plugins/Importers/USDImporter/Source/USDUtilities/Private/USDConversionUtils.cpp
@@ -16,7 +16,6 @@
 #include "USDTypesConversion.h"
 #include "USDUnrealAssetInfo.h"
 
-#include "UsdWrappers/SdfLayer.h"
 #include "UsdWrappers/SdfPath.h"
 #include "UsdWrappers/UsdPrim.h"
 #include "UsdWrappers/UsdStage.h"
@@ -54,37 +53,36 @@
 
 #if USE_USD_SDK
 #include "USDIncludesStart.h"
-	#include "pxr/base/tf/stringUtils.h"
-	#include "pxr/base/tf/token.h"
-	#include "pxr/usd/kind/registry.h"
-	#include "pxr/usd/sdf/copyUtils.h"
-	#include "pxr/usd/usd/attribute.h"
-	#include "pxr/usd/usd/editContext.h"
-	#include "pxr/usd/usd/modelAPI.h"
-	#include "pxr/usd/usd/payloads.h"
-	#include "pxr/usd/usd/primCompositionQuery.h"
-	#include "pxr/usd/usd/primRange.h"
-	#include "pxr/usd/usd/stage.h"
-	#include "pxr/usd/usd/variantSets.h"
-	#include "pxr/usd/usdGeom/camera.h"
-	#include "pxr/usd/usdGeom/gprim.h"
-	#include "pxr/usd/usdGeom/imageable.h"
-	#include "pxr/usd/usdGeom/mesh.h"
-	#include "pxr/usd/usdGeom/metrics.h"
-	#include "pxr/usd/usdGeom/modelAPI.h"
-	#include "pxr/usd/usdGeom/primvar.h"
-	#include "pxr/usd/usdGeom/primvarsAPI.h"
-	#include "pxr/usd/usdLux/diskLight.h"
-	#include "pxr/usd/usdLux/distantLight.h"
-	#include "pxr/usd/usdLux/domeLight.h"
-	#include "pxr/usd/usdLux/rectLight.h"
-	#include "pxr/usd/usdLux/shapingAPI.h"
-	#include "pxr/usd/usdLux/sphereLight.h"
-	#include "pxr/usd/usdSkel/animation.h"
-	#include "pxr/usd/usdSkel/binding.h"
-	#include "pxr/usd/usdSkel/cache.h"
-	#include "pxr/usd/usdSkel/root.h"
-	#include "pxr/usd/usdSkel/skeletonQuery.h"
+#include "pxr/base/tf/stringUtils.h"
+#include "pxr/base/tf/token.h"
+#include "pxr/usd/kind/registry.h"
+#include "pxr/usd/sdf/copyUtils.h"
+#include "pxr/usd/usd/attribute.h"
+#include "pxr/usd/usd/editContext.h"
+#include "pxr/usd/usd/modelAPI.h"
+#include "pxr/usd/usd/payloads.h"
+#include "pxr/usd/usd/primCompositionQuery.h"
+#include "pxr/usd/usd/primRange.h"
+#include "pxr/usd/usd/stage.h"
+#include "pxr/usd/usd/variantSets.h"
+#include "pxr/usd/usdGeom/camera.h"
+#include "pxr/usd/usdGeom/gprim.h"
+#include "pxr/usd/usdGeom/imageable.h"
+#include "pxr/usd/usdGeom/mesh.h"
+#include "pxr/usd/usdGeom/metrics.h"
+#include "pxr/usd/usdGeom/modelAPI.h"
+#include "pxr/usd/usdGeom/primvar.h"
+#include "pxr/usd/usdGeom/primvarsAPI.h"
+#include "pxr/usd/usdLux/diskLight.h"
+#include "pxr/usd/usdLux/distantLight.h"
+#include "pxr/usd/usdLux/domeLight.h"
+#include "pxr/usd/usdLux/rectLight.h"
+#include "pxr/usd/usdLux/shapingAPI.h"
+#include "pxr/usd/usdLux/sphereLight.h"
+#include "pxr/usd/usdSkel/animation.h"
+#include "pxr/usd/usdSkel/cache.h"
+#include "pxr/usd/usdSkel/root.h"
+#include "pxr/usd/usdSkel/skeletonQuery.h"
 #include "USDIncludesEnd.h"
 
 #include <string>
@@ -95,60 +93,71 @@ static bool GParseUVSetsFromFloat2Primvars = true;
 static FAutoConsoleVariableRef CVarParseUVSetsFromFloat2Primvars(
 	TEXT("USD.ParseUVSetsFromFloat2Primvars"),
 	GParseUVSetsFromFloat2Primvars,
-	TEXT("Primvars with the 'texCoord2f' role will always be parsed when handling potential UV sets. If this cvar is enabled, we'll also handle primvars declared as just 'float2' however. You could disable this cvar if your pipeline emits many 'float2' primvars that you do not wish to be parsed as UV sets."));
+	TEXT("Primvars with the 'texCoord2f' role will always be parsed when handling potential UV sets. If this cvar is enabled, we'll also handle "
+		 "primvars declared as just 'float2' however. You could disable this cvar if your pipeline emits many 'float2' primvars that you do not wish "
+		 "to be parsed as UV sets.")
+);
 
 namespace USDConversionUtilsImpl
 {
 	/** Show some warnings if the UVSet primvars show some unsupported/problematic behavior */
-	void CheckUVSetPrimvars( TMap<int32, TArray<pxr::UsdGeomPrimvar>> UsablePrimvars, TMap<int32, TArray<pxr::UsdGeomPrimvar>> UsedPrimvars, const FString& MeshPath )
+	void CheckUVSetPrimvars(
+		TMap<int32, TArray<pxr::UsdGeomPrimvar>> UsablePrimvars,
+		TMap<int32, TArray<pxr::UsdGeomPrimvar>> UsedPrimvars,
+		const FString& MeshPath
+	)
 	{
 		// Show a warning if the mesh has a primvar that could be used as a UV set but will actually be ignored because it targets a UV set with index
 		// larger than MAX_STATIC_TEXCOORDS - 1
 		TArray<FString> IgnoredPrimvarNames;
-		for ( const TPair< int32, TArray<pxr::UsdGeomPrimvar> >& UsedPrimvar : UsedPrimvars )
+		for (const TPair<int32, TArray<pxr::UsdGeomPrimvar>>& UsedPrimvar : UsedPrimvars)
 		{
-			if ( UsedPrimvar.Key > MAX_STATIC_TEXCOORDS - 1 )
+			if (UsedPrimvar.Key > MAX_STATIC_TEXCOORDS - 1)
 			{
-				for ( const pxr::UsdGeomPrimvar& Primvar : UsedPrimvar.Value )
+				for (const pxr::UsdGeomPrimvar& Primvar : UsedPrimvar.Value)
 				{
-					IgnoredPrimvarNames.AddUnique( UsdToUnreal::ConvertToken( Primvar.GetBaseName() ) );
+					IgnoredPrimvarNames.AddUnique(UsdToUnreal::ConvertToken(Primvar.GetBaseName()));
 				}
 			}
 		}
-		for ( const TPair< int32, TArray<pxr::UsdGeomPrimvar> >& UsablePrimvar : UsablePrimvars )
+		for (const TPair<int32, TArray<pxr::UsdGeomPrimvar>>& UsablePrimvar : UsablePrimvars)
 		{
-			if ( UsablePrimvar.Key > MAX_STATIC_TEXCOORDS - 1 )
+			if (UsablePrimvar.Key > MAX_STATIC_TEXCOORDS - 1)
 			{
-				for ( const pxr::UsdGeomPrimvar& Primvar : UsablePrimvar.Value )
+				for (const pxr::UsdGeomPrimvar& Primvar : UsablePrimvar.Value)
 				{
 					// Only consider texcoord2f here because the user may have some other float2[] for some other reason
-					if ( Primvar.GetTypeName().GetRole() == pxr::SdfValueTypeNames->TexCoord2f.GetRole() )
+					if (Primvar.GetTypeName().GetRole() == pxr::SdfValueTypeNames->TexCoord2f.GetRole())
 					{
-						IgnoredPrimvarNames.AddUnique( UsdToUnreal::ConvertToken( Primvar.GetBaseName() ) );
+						IgnoredPrimvarNames.AddUnique(UsdToUnreal::ConvertToken(Primvar.GetBaseName()));
 					}
 				}
 			}
 		}
-		if ( IgnoredPrimvarNames.Num() > 0 )
+		if (IgnoredPrimvarNames.Num() > 0)
 		{
-			FString PrimvarNames = FString::Join( IgnoredPrimvarNames, TEXT( ", " ) );
+			FString PrimvarNames = FString::Join(IgnoredPrimvarNames, TEXT(", "));
 			FUsdLogManager::LogMessage(
 				EMessageSeverity::Warning,
 				FText::Format(
-					LOCTEXT( "TooHighUVIndex", "Mesh '{0}' has some valid UV set primvars ({1}) that will be ignored because they target an UV index larger than the highest supported ({2})" ),
-					FText::FromString( MeshPath ),
-					FText::FromString( PrimvarNames ),
+					LOCTEXT(
+						"TooHighUVIndex",
+						"Mesh '{0}' has some valid UV set primvars ({1}) that will be ignored because they target an UV index larger than the "
+						"highest supported ({2})"
+					),
+					FText::FromString(MeshPath),
+					FText::FromString(PrimvarNames),
 					MAX_STATIC_TEXCOORDS - 1
 				)
 			);
 		}
 
 		// Show a warning if the mesh does not contain the exact primvars the material wants
-		for ( const TPair< int32, TArray<pxr::UsdGeomPrimvar> >& UVAndPrimvars : UsedPrimvars )
+		for (const TPair<int32, TArray<pxr::UsdGeomPrimvar>>& UVAndPrimvars : UsedPrimvars)
 		{
 			const int32 UVIndex = UVAndPrimvars.Key;
 			const TArray<pxr::UsdGeomPrimvar>& UsedPrimvarsForIndex = UVAndPrimvars.Value;
-			if ( UsedPrimvarsForIndex.Num() < 1 )
+			if (UsedPrimvarsForIndex.Num() < 1)
 			{
 				continue;
 			}
@@ -157,43 +166,43 @@ namespace USDConversionUtilsImpl
 			const pxr::UsdGeomPrimvar& UsedPrimvar = UsedPrimvarsForIndex[0];
 
 			bool bFoundUsablePrimvar = false;
-			if ( const TArray<pxr::UsdGeomPrimvar>* FoundUsablePrimvars = UsablePrimvars.Find( UVIndex ) )
+			if (const TArray<pxr::UsdGeomPrimvar>* FoundUsablePrimvars = UsablePrimvars.Find(UVIndex))
 			{
 				// We will only ever use the first one, but will show more warnings in case there are multiple
-				if ( FoundUsablePrimvars->Contains( UsedPrimvar ) )
+				if (FoundUsablePrimvars->Contains(UsedPrimvar))
 				{
 					bFoundUsablePrimvar = true;
 				}
 			}
 
-			if ( !bFoundUsablePrimvar )
+			if (!bFoundUsablePrimvar)
 			{
 				FUsdLogManager::LogMessage(
 					EMessageSeverity::Warning,
 					FText::Format(
-						LOCTEXT( "DidNotFindPrimvar", "Could not find primvar '{0}' on mesh '{1}', used by its bound material" ),
-						FText::FromString( UsdToUnreal::ConvertString( UsedPrimvar.GetBaseName() ) ),
-						FText::FromString( MeshPath )
+						LOCTEXT("DidNotFindPrimvar", "Could not find primvar '{0}' on mesh '{1}', used by its bound material"),
+						FText::FromString(UsdToUnreal::ConvertString(UsedPrimvar.GetBaseName())),
+						FText::FromString(MeshPath)
 					)
 				);
 			}
 		}
 
 		// Show a warning if the mesh has multiple primvars that want to write to the same UV set (e.g. 'st', 'st_0' and 'st0' at the same time)
-		for ( const TPair< int32, TArray<pxr::UsdGeomPrimvar> >& UVAndPrimvars : UsablePrimvars )
+		for (const TPair<int32, TArray<pxr::UsdGeomPrimvar>>& UVAndPrimvars : UsablePrimvars)
 		{
 			const int32 UVIndex = UVAndPrimvars.Key;
 			const TArray<pxr::UsdGeomPrimvar>& Primvars = UVAndPrimvars.Value;
-			if ( Primvars.Num() > 1 )
+			if (Primvars.Num() > 1)
 			{
 				// Find out what primvar we'll actually end up using, as UsedPrimvars will take precedence. Note that in the best case scenario,
-				// UsablePrimvars will *contain* UsedPrimvars, so that really we're just picking which of the UsedPrimvars we'll choose. If we're not in that scenario,
-				// then we will show another warning about it
+				// UsablePrimvars will *contain* UsedPrimvars, so that really we're just picking which of the UsedPrimvars we'll choose. If we're not
+				// in that scenario, then we will show another warning about it
 				const pxr::UsdGeomPrimvar* UsedPrimvar = nullptr;
 				bool bUsedByMaterial = false;
-				if ( const TArray<pxr::UsdGeomPrimvar>* FoundUsedPrimvars = UsedPrimvars.Find( UVIndex ) )
+				if (const TArray<pxr::UsdGeomPrimvar>* FoundUsedPrimvars = UsedPrimvars.Find(UVIndex))
 				{
-					if ( FoundUsedPrimvars->Num() > 0 )
+					if (FoundUsedPrimvars->Num() > 0)
 					{
 						UsedPrimvar = &(*FoundUsedPrimvars)[0];
 						bUsedByMaterial = true;
@@ -207,11 +216,15 @@ namespace USDConversionUtilsImpl
 				FUsdLogManager::LogMessage(
 					EMessageSeverity::Warning,
 					FText::Format(
-						LOCTEXT( "MoreThanOnePrimvarForIndex", "Mesh '{0}' has more than one primvar used as UV set with index '{1}'. The UV set will use the values from primvar '{2}'{3}" ),
-						FText::FromString( MeshPath ),
+						LOCTEXT(
+							"MoreThanOnePrimvarForIndex",
+							"Mesh '{0}' has more than one primvar used as UV set with index '{1}'. The UV set will use the values from primvar "
+							"'{2}'{3}"
+						),
+						FText::FromString(MeshPath),
 						UVAndPrimvars.Key,
-						FText::FromString( UsdToUnreal::ConvertString( UsedPrimvar->GetBaseName() ) ),
-						bUsedByMaterial ? FText::FromString( TEXT( ", as its used by its bound material" ) ) : FText::GetEmpty()
+						FText::FromString(UsdToUnreal::ConvertString(UsedPrimvar->GetBaseName())),
+						bUsedByMaterial ? FText::FromString(TEXT(", as its used by its bound material")) : FText::GetEmpty()
 					)
 				);
 			}
@@ -220,42 +233,47 @@ namespace USDConversionUtilsImpl
 
 	// Shows a notification saying that some specs of the provided prims won't be duplicated due to being on external
 	// layers
-	void NotifySpecsWontBeDuplicated( const TArray<UE::FUsdPrim>& Prims )
+	void NotifySpecsWontBeDuplicated(const TArray<UE::FUsdPrim>& Prims)
 	{
-		if ( Prims.Num() == 0 )
+		if (Prims.Num() == 0)
 		{
 			return;
 		}
 
-		const FText Text = LOCTEXT( "IncompleteDuplicationText", "USD: Incomplete duplication" );
+		const FText Text = LOCTEXT("IncompleteDuplicationText", "USD: Incomplete duplication");
 
 		FString PrimNamesString;
 		const static FString Delimiter = TEXT(", ");
-		for( const UE::FUsdPrim& Prim : Prims )
+		for (const UE::FUsdPrim& Prim : Prims)
 		{
 			PrimNamesString += Prim.GetName().ToString() + Delimiter;
 		}
-		PrimNamesString.RemoveFromEnd( Delimiter );
+		PrimNamesString.RemoveFromEnd(Delimiter);
 
 		const int32 NumPrims = Prims.Num();
 
 		const FText SubText = FText::Format(
-			LOCTEXT( "IncompleteDuplicationSubText", "{0}|plural(one=This,other=These) duplicated {0}|plural(one=prim,other=prims):\n\n{1}\n\n{0}|plural(one=Has,other=Have) some specs within layers that are outside of the stage's local layer stack, and so will not be duplicated.\n\nIf you wish to modify referenced or payload layers, please open those layers as USD stages directly." ),
+			LOCTEXT(
+				"IncompleteDuplicationSubText",
+				"{0}|plural(one=This,other=These) duplicated {0}|plural(one=prim,other=prims):\n\n{1}\n\n{0}|plural(one=Has,other=Have) some specs "
+				"within layers that are outside of the stage's local layer stack, and so will not be duplicated.\n\nIf you wish to modify referenced "
+				"or payload layers, please open those layers as USD stages directly."
+			),
 			NumPrims,
-			FText::FromString( PrimNamesString )
+			FText::FromString(PrimNamesString)
 		);
 
-		UE_LOG( LogUsd, Warning, TEXT( "%s" ), *SubText.ToString().Replace( TEXT( "\n\n" ), TEXT( " " ) ) );
+		UE_LOG(LogUsd, Warning, TEXT("%s"), *SubText.ToString().Replace(TEXT("\n\n"), TEXT(" ")));
 
 		const UUsdProjectSettings* Settings = GetDefault<UUsdProjectSettings>();
-		if ( Settings && Settings->bShowWarningOnIncompleteDuplication )
+		if (Settings && Settings->bShowWarningOnIncompleteDuplication)
 		{
 			static TWeakPtr<SNotificationItem> Notification;
 
-			FNotificationInfo Toast( Text );
+			FNotificationInfo Toast(Text);
 			Toast.SubText = SubText;
-			Toast.Image = FCoreStyle::Get().GetBrush( TEXT( "MessageLog.Warning" ) );
-			Toast.CheckBoxText = LOCTEXT( "DontAskAgain", "Don't prompt again" );
+			Toast.Image = FCoreStyle::Get().GetBrush(TEXT("MessageLog.Warning"));
+			Toast.CheckBoxText = LOCTEXT("DontAskAgain", "Don't prompt again");
 			Toast.bUseLargeFont = false;
 			Toast.bFireAndForget = false;
 			Toast.FadeOutDuration = 0.0f;
@@ -263,102 +281,113 @@ namespace USDConversionUtilsImpl
 			Toast.bUseThrobber = false;
 			Toast.bUseSuccessFailIcons = false;
 			Toast.ButtonDetails.Emplace(
-				LOCTEXT( "OverridenOpinionMessageOk", "Ok" ),
+				LOCTEXT("OverridenOpinionMessageOk", "Ok"),
 				FText::GetEmpty(),
-				FSimpleDelegate::CreateLambda( []() {
-				if ( TSharedPtr<SNotificationItem> PinnedNotification = Notification.Pin() )
-				{
-					PinnedNotification->SetCompletionState( SNotificationItem::CS_Success );
-					PinnedNotification->ExpireAndFadeout();
-				}
-			} )
+				FSimpleDelegate::CreateLambda(
+					[]()
+					{
+						if (TSharedPtr<SNotificationItem> PinnedNotification = Notification.Pin())
+						{
+							PinnedNotification->SetCompletionState(SNotificationItem::CS_Success);
+							PinnedNotification->ExpireAndFadeout();
+						}
+					}
+				)
 			);
 			// This is flipped because the default checkbox message is "Don't prompt again"
 			Toast.CheckBoxState = Settings->bShowWarningOnIncompleteDuplication ? ECheckBoxState::Unchecked : ECheckBoxState::Checked;
-			Toast.CheckBoxStateChanged = FOnCheckStateChanged::CreateStatic( []( ECheckBoxState NewState )
-			{
-				if ( UUsdProjectSettings* Settings = GetMutableDefault<UUsdProjectSettings>() )
+			Toast.CheckBoxStateChanged = FOnCheckStateChanged::CreateStatic(
+				[](ECheckBoxState NewState)
 				{
-					// This is flipped because the default checkbox message is "Don't prompt again"
-					Settings->bShowWarningOnIncompleteDuplication = NewState == ECheckBoxState::Unchecked;
-					Settings->SaveConfig();
+					if (UUsdProjectSettings* Settings = GetMutableDefault<UUsdProjectSettings>())
+					{
+						// This is flipped because the default checkbox message is "Don't prompt again"
+						Settings->bShowWarningOnIncompleteDuplication = NewState == ECheckBoxState::Unchecked;
+						Settings->SaveConfig();
+					}
 				}
-			} );
+			);
 
 			// Only show one at a time
-			if ( !Notification.IsValid() )
+			if (!Notification.IsValid())
 			{
-				Notification = FSlateNotificationManager::Get().AddNotification( Toast );
+				Notification = FSlateNotificationManager::Get().AddNotification(Toast);
 			}
 
-			if ( TSharedPtr<SNotificationItem> PinnedNotification = Notification.Pin() )
+			if (TSharedPtr<SNotificationItem> PinnedNotification = Notification.Pin())
 			{
-				PinnedNotification->SetCompletionState( SNotificationItem::CS_Pending );
+				PinnedNotification->SetCompletionState(SNotificationItem::CS_Pending);
 			}
 		}
 	}
-}
+}	 // namespace USDConversionUtilsImpl
 
-template< typename ValueType >
-ValueType UsdUtils::GetUsdValue( const pxr::UsdAttribute& Attribute, pxr::UsdTimeCode TimeCode )
+template<typename ValueType>
+ValueType UsdUtils::GetUsdValue(const pxr::UsdAttribute& Attribute, pxr::UsdTimeCode TimeCode)
 {
 	ValueType Value{};
-	if ( Attribute )
+	if (Attribute)
 	{
-		Attribute.Get( &Value, TimeCode );
+		Attribute.Get(&Value, TimeCode);
 	}
 
 	return Value;
 }
 
 // Explicit template instantiation
-template USDUTILITIES_API bool							UsdUtils::GetUsdValue< bool >( const pxr::UsdAttribute& Attribute, pxr::UsdTimeCode TimeCode );
-template USDUTILITIES_API float							UsdUtils::GetUsdValue< float >( const pxr::UsdAttribute& Attribute, pxr::UsdTimeCode TimeCode );
-template USDUTILITIES_API double						UsdUtils::GetUsdValue< double >( const pxr::UsdAttribute& Attribute, pxr::UsdTimeCode TimeCode );
-template USDUTILITIES_API pxr::TfToken					UsdUtils::GetUsdValue< pxr::TfToken >( const pxr::UsdAttribute& Attribute, pxr::UsdTimeCode TimeCode );
-template USDUTILITIES_API pxr::GfVec3f					UsdUtils::GetUsdValue< pxr::GfVec3f >( const pxr::UsdAttribute& Attribute, pxr::UsdTimeCode TimeCode );
-template USDUTILITIES_API pxr::GfMatrix4d				UsdUtils::GetUsdValue< pxr::GfMatrix4d >( const pxr::UsdAttribute& Attribute, pxr::UsdTimeCode TimeCode );
-template USDUTILITIES_API pxr::SdfAssetPath				UsdUtils::GetUsdValue< pxr::SdfAssetPath >( const pxr::UsdAttribute& Attribute, pxr::UsdTimeCode TimeCode );
-template USDUTILITIES_API pxr::VtArray< pxr::GfVec3f >	UsdUtils::GetUsdValue< pxr::VtArray< pxr::GfVec3f > >( const pxr::UsdAttribute& Attribute, pxr::UsdTimeCode TimeCode );
-template USDUTILITIES_API pxr::VtArray< float >			UsdUtils::GetUsdValue< pxr::VtArray< float > >( const pxr::UsdAttribute& Attribute, pxr::UsdTimeCode TimeCode );
-template USDUTILITIES_API pxr::VtArray< int >			UsdUtils::GetUsdValue< pxr::VtArray< int > >( const pxr::UsdAttribute& Attribute, pxr::UsdTimeCode TimeCode );
-
-pxr::TfToken UsdUtils::GetUsdStageUpAxis( const pxr::UsdStageRefPtr& Stage )
-{
-	return pxr::UsdGeomGetStageUpAxis( Stage );
-}
-
-EUsdUpAxis UsdUtils::GetUsdStageUpAxisAsEnum( const pxr::UsdStageRefPtr& Stage )
-{
-	pxr::TfToken UpAxisToken = pxr::UsdGeomGetStageUpAxis( Stage );
+template USDUTILITIES_API bool UsdUtils::GetUsdValue<bool>(const pxr::UsdAttribute& Attribute, pxr::UsdTimeCode TimeCode);
+template USDUTILITIES_API float UsdUtils::GetUsdValue<float>(const pxr::UsdAttribute& Attribute, pxr::UsdTimeCode TimeCode);
+template USDUTILITIES_API double UsdUtils::GetUsdValue<double>(const pxr::UsdAttribute& Attribute, pxr::UsdTimeCode TimeCode);
+template USDUTILITIES_API pxr::TfToken UsdUtils::GetUsdValue<pxr::TfToken>(const pxr::UsdAttribute& Attribute, pxr::UsdTimeCode TimeCode);
+template USDUTILITIES_API pxr::GfVec3f UsdUtils::GetUsdValue<pxr::GfVec3f>(const pxr::UsdAttribute& Attribute, pxr::UsdTimeCode TimeCode);
+template USDUTILITIES_API pxr::GfMatrix4d UsdUtils::GetUsdValue<pxr::GfMatrix4d>(const pxr::UsdAttribute& Attribute, pxr::UsdTimeCode TimeCode);
+template USDUTILITIES_API pxr::SdfAssetPath UsdUtils::GetUsdValue<pxr::SdfAssetPath>(const pxr::UsdAttribute& Attribute, pxr::UsdTimeCode TimeCode);
+template USDUTILITIES_API pxr::VtArray<pxr::GfVec3f> UsdUtils::GetUsdValue<pxr::VtArray<pxr::GfVec3f>>(
+	const pxr::UsdAttribute& Attribute,
+	pxr::UsdTimeCode TimeCode
+);
+template USDUTILITIES_API pxr::VtArray<float> UsdUtils::GetUsdValue<pxr::VtArray<float>>(
+	const pxr::UsdAttribute& Attribute,
+	pxr::UsdTimeCode TimeCode
+);
+template USDUTILITIES_API pxr::VtArray<int> UsdUtils::GetUsdValue<pxr::VtArray<int>>(const pxr::UsdAttribute& Attribute, pxr::UsdTimeCode TimeCode);
+
+pxr::TfToken UsdUtils::GetUsdStageUpAxis(const pxr::UsdStageRefPtr& Stage)
+{
+	return pxr::UsdGeomGetStageUpAxis(Stage);
+}
+
+EUsdUpAxis UsdUtils::GetUsdStageUpAxisAsEnum(const pxr::UsdStageRefPtr& Stage)
+{
+	pxr::TfToken UpAxisToken = pxr::UsdGeomGetStageUpAxis(Stage);
 	return UpAxisToken == pxr::UsdGeomTokens->z ? EUsdUpAxis::ZAxis : EUsdUpAxis::YAxis;
 }
 
-void UsdUtils::SetUsdStageUpAxis( const pxr::UsdStageRefPtr& Stage, pxr::TfToken Axis )
+void UsdUtils::SetUsdStageUpAxis(const pxr::UsdStageRefPtr& Stage, pxr::TfToken Axis)
 {
-	pxr::UsdGeomSetStageUpAxis( Stage, Axis );
+	pxr::UsdGeomSetStageUpAxis(Stage, Axis);
 }
 
-void UsdUtils::SetUsdStageUpAxis( const pxr::UsdStageRefPtr& Stage, EUsdUpAxis Axis )
+void UsdUtils::SetUsdStageUpAxis(const pxr::UsdStageRefPtr& Stage, EUsdUpAxis Axis)
 {
 	pxr::TfToken UpAxisToken = Axis == EUsdUpAxis::ZAxis ? pxr::UsdGeomTokens->z : pxr::UsdGeomTokens->y;
-	SetUsdStageUpAxis( Stage, UpAxisToken );
+	SetUsdStageUpAxis(Stage, UpAxisToken);
 }
 
-float UsdUtils::GetUsdStageMetersPerUnit( const pxr::UsdStageRefPtr& Stage )
+float UsdUtils::GetUsdStageMetersPerUnit(const pxr::UsdStageRefPtr& Stage)
 {
-	return (float)pxr::UsdGeomGetStageMetersPerUnit( Stage );
+	return (float)pxr::UsdGeomGetStageMetersPerUnit(Stage);
 }
 
-void UsdUtils::SetUsdStageMetersPerUnit( const pxr::UsdStageRefPtr& Stage, float MetersPerUnit )
+void UsdUtils::SetUsdStageMetersPerUnit(const pxr::UsdStageRefPtr& Stage, float MetersPerUnit)
 {
-	if ( !Stage || !Stage->GetRootLayer() )
+	if (!Stage || !Stage->GetRootLayer())
 	{
 		return;
 	}
 
-	pxr::UsdEditContext( Stage, Stage->GetRootLayer() );
-	pxr::UsdGeomSetStageMetersPerUnit( Stage, MetersPerUnit );
+	pxr::UsdEditContext(Stage, Stage->GetRootLayer());
+	pxr::UsdGeomSetStageMetersPerUnit(Stage, MetersPerUnit);
 }
 
 int32 UsdUtils::GetUsdStageNumFrames(const pxr::UsdStageRefPtr& Stage)
@@ -367,61 +396,63 @@ int32 UsdUtils::GetUsdStageNumFrames(const pxr::UsdStageRefPtr& Stage)
 	return Stage ? FMath::Abs(FMath::CeilToInt32(Stage->GetEndTimeCode()) - FMath::FloorToInt32(Stage->GetStartTimeCode()) + 1) : 0;
 }
 
-bool UsdUtils::HasCompositionArcs( const pxr::UsdPrim& Prim )
+bool UsdUtils::HasCompositionArcs(const pxr::UsdPrim& Prim)
 {
-	if ( !Prim || !Prim.IsActive() )
+	if (!Prim || !Prim.IsActive())
 	{
 		return false;
 	}
 
-	return Prim.HasAuthoredReferences() || Prim.HasAuthoredPayloads() || Prim.HasAuthoredInherits() || Prim.HasAuthoredSpecializes() || Prim.HasVariantSets();
+	return Prim.HasAuthoredReferences() || Prim.HasAuthoredPayloads() || Prim.HasAuthoredInherits() || Prim.HasAuthoredSpecializes()
+		   || Prim.HasVariantSets();
 }
 
-bool UsdUtils::HasCompositionArcs( const pxr::SdfPrimSpecHandle& PrimSpec )
+bool UsdUtils::HasCompositionArcs(const pxr::SdfPrimSpecHandle& PrimSpec)
 {
-	if ( !PrimSpec || !PrimSpec->GetActive() )
+	if (!PrimSpec || !PrimSpec->GetActive())
 	{
 		return false;
 	}
 
-	return PrimSpec->HasReferences() || PrimSpec->HasPayloads() || PrimSpec->HasInheritPaths() || PrimSpec->HasSpecializes() || PrimSpec->HasVariantSetNames();
+	return PrimSpec->HasReferences() || PrimSpec->HasPayloads() || PrimSpec->HasInheritPaths() || PrimSpec->HasSpecializes()
+		   || PrimSpec->HasVariantSetNames();
 }
 
-UClass* UsdUtils::GetActorTypeForPrim( const pxr::UsdPrim& Prim )
+UClass* UsdUtils::GetActorTypeForPrim(const pxr::UsdPrim& Prim)
 {
 	// If we have this attribute and a valid child camera prim then we'll assume
 	// we correspond to the root scene component of an exported cine camera actor. Let's assume
 	// then that we have an actual ACineCameraActor class so that the schema translators can
 	// reuse the main UCineCameraComponent for the actual child camera prim
 	bool bIsCineCameraActorRootComponent = false;
-	if ( pxr::UsdAttribute Attr = Prim.GetAttribute( UnrealToUsd::ConvertToken( TEXT( "unrealCameraPrimName" ) ).Get() ) )
+	if (pxr::UsdAttribute Attr = Prim.GetAttribute(UnrealToUsd::ConvertToken(TEXT("unrealCameraPrimName")).Get()))
 	{
 		pxr::TfToken CameraComponentPrim;
-		if ( Attr.Get<pxr::TfToken>( &CameraComponentPrim ) )
+		if (Attr.Get<pxr::TfToken>(&CameraComponentPrim))
 		{
-			pxr::UsdPrim ChildCameraPrim = Prim.GetChild( CameraComponentPrim );
-			if ( ChildCameraPrim && ChildCameraPrim.IsA<pxr::UsdGeomCamera>() )
+			pxr::UsdPrim ChildCameraPrim = Prim.GetChild(CameraComponentPrim);
+			if (ChildCameraPrim && ChildCameraPrim.IsA<pxr::UsdGeomCamera>())
 			{
 				bIsCineCameraActorRootComponent = true;
 			}
 		}
 	}
 
-	if ( Prim.IsA< pxr::UsdGeomCamera >() || bIsCineCameraActorRootComponent )
+	if (Prim.IsA<pxr::UsdGeomCamera>() || bIsCineCameraActorRootComponent)
 	{
 		return ACineCameraActor::StaticClass();
 	}
-	else if ( Prim.IsA< pxr::UsdLuxDistantLight >() )
+	else if (Prim.IsA<pxr::UsdLuxDistantLight>())
 	{
 		return ADirectionalLight::StaticClass();
 	}
-	else if ( Prim.IsA< pxr::UsdLuxRectLight >() || Prim.IsA< pxr::UsdLuxDiskLight >() )
+	else if (Prim.IsA<pxr::UsdLuxRectLight>() || Prim.IsA<pxr::UsdLuxDiskLight>())
 	{
 		return ARectLight::StaticClass();
 	}
-	else if ( Prim.IsA< pxr::UsdLuxSphereLight >() )
+	else if (Prim.IsA<pxr::UsdLuxSphereLight>())
 	{
-		if ( Prim.HasAPI< pxr::UsdLuxShapingAPI >() )
+		if (Prim.HasAPI<pxr::UsdLuxShapingAPI>())
 		{
 			return ASpotLight::StaticClass();
 		}
@@ -430,7 +461,7 @@ UClass* UsdUtils::GetActorTypeForPrim( const pxr::UsdPrim& Prim )
 			return APointLight::StaticClass();
 		}
 	}
-	else if ( Prim.IsA< pxr::UsdLuxDomeLight >() )
+	else if (Prim.IsA<pxr::UsdLuxDomeLight>())
 	{
 		return ASkyLight::StaticClass();
 	}
@@ -440,9 +471,9 @@ UClass* UsdUtils::GetActorTypeForPrim( const pxr::UsdPrim& Prim )
 	}
 }
 
-UClass* UsdUtils::GetComponentTypeForPrim( const pxr::UsdPrim& Prim )
+UClass* UsdUtils::GetComponentTypeForPrim(const pxr::UsdPrim& Prim)
 {
-	if ( Prim.IsA< pxr::UsdSkelSkeleton >() )
+	if (Prim.IsA<pxr::UsdSkelSkeleton>())
 	{
 		return USkeletalMeshComponent::StaticClass();
 	}
@@ -450,21 +481,21 @@ UClass* UsdUtils::GetComponentTypeForPrim( const pxr::UsdPrim& Prim )
 	{
 		return UStaticMeshComponent::StaticClass();
 	}
-	else if ( Prim.IsA< pxr::UsdGeomCamera >() )
+	else if (Prim.IsA<pxr::UsdGeomCamera>())
 	{
 		return UCineCameraComponent::StaticClass();
 	}
-	else if ( Prim.IsA< pxr::UsdLuxDistantLight >() )
+	else if (Prim.IsA<pxr::UsdLuxDistantLight>())
 	{
 		return UDirectionalLightComponent::StaticClass();
 	}
-	else if ( Prim.IsA< pxr::UsdLuxRectLight >() || Prim.IsA< pxr::UsdLuxDiskLight >() )
+	else if (Prim.IsA<pxr::UsdLuxRectLight>() || Prim.IsA<pxr::UsdLuxDiskLight>())
 	{
 		return URectLightComponent::StaticClass();
 	}
-	else if ( Prim.IsA< pxr::UsdLuxSphereLight >() )
+	else if (Prim.IsA<pxr::UsdLuxSphereLight>())
 	{
-		if ( Prim.HasAPI< pxr::UsdLuxShapingAPI >() )
+		if (Prim.HasAPI<pxr::UsdLuxShapingAPI>())
 		{
 			return USpotLightComponent::StaticClass();
 		}
@@ -473,11 +504,11 @@ UClass* UsdUtils::GetComponentTypeForPrim( const pxr::UsdPrim& Prim )
 			return UPointLightComponent::StaticClass();
 		}
 	}
-	else if ( Prim.IsA< pxr::UsdLuxDomeLight >() )
+	else if (Prim.IsA<pxr::UsdLuxDomeLight>())
 	{
 		return USkyLightComponent::StaticClass();
 	}
-	else if ( Prim.IsA< pxr::UsdGeomXformable >() )
+	else if (Prim.IsA<pxr::UsdGeomXformable>())
 	{
 		return USceneComponent::StaticClass();
 	}
@@ -487,87 +518,87 @@ UClass* UsdUtils::GetComponentTypeForPrim( const pxr::UsdPrim& Prim )
 	}
 }
 
-FString UsdUtils::GetSchemaNameForComponent( const USceneComponent& Component )
+FString UsdUtils::GetSchemaNameForComponent(const USceneComponent& Component)
 {
 	AActor* OwnerActor = Component.GetOwner();
-	if ( OwnerActor->IsA<AInstancedFoliageActor>() )
+	if (OwnerActor->IsA<AInstancedFoliageActor>())
 	{
-		return TEXT( "PointInstancer" );
+		return TEXT("PointInstancer");
 	}
-	else if ( OwnerActor->IsA<ALandscapeProxy>() )
+	else if (OwnerActor->IsA<ALandscapeProxy>())
 	{
-		return TEXT( "Mesh" );
+		return TEXT("Mesh");
 	}
 
-	if ( Component.IsA<USkinnedMeshComponent>() )
+	if (Component.IsA<USkinnedMeshComponent>())
 	{
-		return TEXT( "SkelRoot" );
+		return TEXT("SkelRoot");
 	}
-	else if ( Component.IsA<UHierarchicalInstancedStaticMeshComponent>() )
+	else if (Component.IsA<UHierarchicalInstancedStaticMeshComponent>())
 	{
 		// The original HISM component becomes just a regular Xform prim, so that we can handle
 		// its children correctly. We'll manually create a new child PointInstancer prim to it
 		// however, and convert the HISM data onto that prim.
-		return TEXT( "Xform" );
+		return TEXT("Xform");
 	}
-	else if ( const UStaticMeshComponent* StaticMeshComponent = Cast<UStaticMeshComponent>( &Component ) )
+	else if (const UStaticMeshComponent* StaticMeshComponent = Cast<UStaticMeshComponent>(&Component))
 	{
 		UStaticMesh* Mesh = StaticMeshComponent->GetStaticMesh();
-		if ( Mesh && Mesh->GetNumLODs() > 1 )
+		if (Mesh && Mesh->GetNumLODs() > 1)
 		{
 			// Don't export 'Mesh' if we're going to export LODs, as those will also be Mesh prims.
 			// We need at least an Xform schema though as this component may still have a transform of its own
-			return TEXT( "Xform" );
+			return TEXT("Xform");
 		}
-		return TEXT( "Mesh" );
+		return TEXT("Mesh");
 	}
-	else if ( Component.IsA<UCineCameraComponent>() )
+	else if (Component.IsA<UCineCameraComponent>())
 	{
-		return TEXT( "Camera" );
+		return TEXT("Camera");
 	}
-	else if ( Component.IsA<UDirectionalLightComponent>() )
+	else if (Component.IsA<UDirectionalLightComponent>())
 	{
-		return TEXT( "DistantLight" );
+		return TEXT("DistantLight");
 	}
-	else if ( Component.IsA<URectLightComponent>() )
+	else if (Component.IsA<URectLightComponent>())
 	{
-		return TEXT( "RectLight" );
+		return TEXT("RectLight");
 	}
-	else if ( Component.IsA<UPointLightComponent>() )
+	else if (Component.IsA<UPointLightComponent>())
 	{
-		return TEXT( "SphereLight" );
+		return TEXT("SphereLight");
 	}
-	else if ( Component.IsA<USkyLightComponent>() )
+	else if (Component.IsA<USkyLightComponent>())
 	{
-		return TEXT( "DomeLight" );
+		return TEXT("DomeLight");
 	}
 
-	return TEXT( "Xform" );
+	return TEXT("Xform");
 }
 
-FString UsdUtils::GetPrimPathForObject( const UObject* ActorOrComponent, const FString& ParentPrimPath, bool bUseActorFolders )
+FString UsdUtils::GetPrimPathForObject(const UObject* ActorOrComponent, const FString& ParentPrimPath, bool bUseActorFolders)
 {
-	if ( !ActorOrComponent )
+	if (!ActorOrComponent)
 	{
 		return {};
 	}
 
 	// Get component and its owner actor
-	const USceneComponent* Component = Cast<const USceneComponent>( ActorOrComponent );
+	const USceneComponent* Component = Cast<const USceneComponent>(ActorOrComponent);
 	const AActor* Owner = nullptr;
-	if ( Component )
+	if (Component)
 	{
 		Owner = Component->GetOwner();
 	}
 	else
 	{
-		Owner = Cast<AActor>( ActorOrComponent );
-		if ( Owner )
+		Owner = Cast<AActor>(ActorOrComponent);
+		if (Owner)
 		{
 			Component = Owner->GetRootComponent();
 		}
 	}
-	if ( !Component || !Owner )
+	if (!Component || !Owner)
 	{
 		return {};
 	}
@@ -575,47 +606,47 @@ FString UsdUtils::GetPrimPathForObject( const UObject* ActorOrComponent, const F
 	// Get component name. Use actor label if the component is its root component
 	FString Path;
 #if WITH_EDITOR
-	if ( Component == Owner->GetRootComponent() )
+	if (Component == Owner->GetRootComponent())
 	{
 		Path = Owner->GetActorLabel();
 	}
 	else
-#endif // WITH_EDITOR
+#endif	  // WITH_EDITOR
 	{
 		Path = Component->GetName();
 	}
-	Path = UsdUtils::SanitizeUsdIdentifier( *Path );
+	Path = UsdUtils::SanitizeUsdIdentifier(*Path);
 
 	// Get a clean folder path string if we have and need one
 	FString FolderPathString;
 #if WITH_EDITOR
-	if ( bUseActorFolders && Component == Owner->GetRootComponent() )
+	if (bUseActorFolders && Component == Owner->GetRootComponent())
 	{
 		const FName& FolderPath = Owner->GetFolderPath();
-		if ( !FolderPath.IsNone() )
+		if (!FolderPath.IsNone())
 		{
 			FolderPathString = FolderPath.ToString();
 
 			TArray<FString> FolderSegments;
-			FolderPathString.ParseIntoArray( FolderSegments, TEXT( "/" ) );
+			FolderPathString.ParseIntoArray(FolderSegments, TEXT("/"));
 
-			for ( FString& Segment : FolderSegments )
+			for (FString& Segment : FolderSegments)
 			{
-				Segment = UsdUtils::SanitizeUsdIdentifier( *Segment );
+				Segment = UsdUtils::SanitizeUsdIdentifier(*Segment);
 			}
 
-			FolderPathString = FString::Join( FolderSegments, TEXT( "/" ) );
+			FolderPathString = FString::Join(FolderSegments, TEXT("/"));
 
-			if ( !FolderPathString.IsEmpty() )
+			if (!FolderPathString.IsEmpty())
 			{
 				Path = FolderPathString / Path;
 			}
 		}
 	}
-#endif // WITH_EDITOR
+#endif	  // WITH_EDITOR
 
 	// Get parent prim path if we need to
-	if ( !ParentPrimPath.IsEmpty() )
+	if (!ParentPrimPath.IsEmpty())
 	{
 		Path = ParentPrimPath / Path;
 	}
@@ -623,13 +654,13 @@ FString UsdUtils::GetPrimPathForObject( const UObject* ActorOrComponent, const F
 	{
 		FString FoundParentPath;
 
-		if ( USceneComponent* ParentComp = Component->GetAttachParent() )
+		if (USceneComponent* ParentComp = Component->GetAttachParent())
 		{
-			FoundParentPath = GetPrimPathForObject( ParentComp, TEXT( "" ), bUseActorFolders );
+			FoundParentPath = GetPrimPathForObject(ParentComp, TEXT(""), bUseActorFolders);
 		}
 		else
 		{
-			FoundParentPath = TEXT( "/Root" );
+			FoundParentPath = TEXT("/Root");
 		}
 
 		Path = FoundParentPath / Path;
@@ -638,42 +669,39 @@ FString UsdUtils::GetPrimPathForObject( const UObject* ActorOrComponent, const F
 	return Path;
 }
 
-TUsdStore< pxr::TfToken > UsdUtils::GetUVSetName( int32 UVChannelIndex )
+TUsdStore<pxr::TfToken> UsdUtils::GetUVSetName(int32 UVChannelIndex)
 {
 	FScopedUnrealAllocs UnrealAllocs;
 
 	FString UVSetName = TEXT("primvars:st");
 
-	if ( UVChannelIndex > 0 )
+	if (UVChannelIndex > 0)
 	{
-		UVSetName += LexToString( UVChannelIndex );
+		UVSetName += LexToString(UVChannelIndex);
 	}
 
-	TUsdStore< pxr::TfToken > UVSetNameToken = MakeUsdStore< pxr::TfToken >( UnrealToUsd::ConvertString( *UVSetName ).Get() );
+	TUsdStore<pxr::TfToken> UVSetNameToken = MakeUsdStore<pxr::TfToken>(UnrealToUsd::ConvertString(*UVSetName).Get());
 
 	return UVSetNameToken;
 }
 
-int32 UsdUtils::GetPrimvarUVIndex( FString PrimvarName )
+int32 UsdUtils::GetPrimvarUVIndex(FString PrimvarName)
 {
 	int32 Index = PrimvarName.Len();
-	while ( Index > 0 && PrimvarName[ Index - 1 ] >= '0' && PrimvarName[ Index - 1 ] <= '9' )
+	while (Index > 0 && PrimvarName[Index - 1] >= '0' && PrimvarName[Index - 1] <= '9')
 	{
 		--Index;
 	}
 
-	if ( Index < PrimvarName.Len() )
+	if (Index < PrimvarName.Len())
 	{
-		return FCString::Atoi( *PrimvarName.RightChop( Index ) );
+		return FCString::Atoi(*PrimvarName.RightChop(Index));
 	}
 
 	return 0;
 }
 
-TArray<TUsdStore<pxr::UsdGeomPrimvar>> UsdUtils::GetUVSetPrimvars(
-	const pxr::UsdPrim& UsdPrim,
-	int32 MaxNumPrimvars
-)
+TArray<TUsdStore<pxr::UsdGeomPrimvar>> UsdUtils::GetUVSetPrimvars(const pxr::UsdPrim& UsdPrim, int32 MaxNumPrimvars)
 {
 	if (!UsdPrim)
 	{
@@ -711,14 +739,18 @@ TArray<TUsdStore<pxr::UsdGeomPrimvar>> UsdUtils::GetUVSetPrimvars(
 		}
 	}
 
-	TexCoord2fPrimvars.Sort([](const TUsdStore<pxr::UsdGeomPrimvar>& A, const TUsdStore<pxr::UsdGeomPrimvar>& B)
-	{
-		return A.Get().GetName() < B.Get().GetName();
-	});
-	Float2Primvars.Sort([](const TUsdStore<pxr::UsdGeomPrimvar>& A, const TUsdStore<pxr::UsdGeomPrimvar>& B)
-	{
-		return A.Get().GetName() < B.Get().GetName();
-	});
+	TexCoord2fPrimvars.Sort(
+		[](const TUsdStore<pxr::UsdGeomPrimvar>& A, const TUsdStore<pxr::UsdGeomPrimvar>& B)
+		{
+			return A.Get().GetName() < B.Get().GetName();
+		}
+	);
+	Float2Primvars.Sort(
+		[](const TUsdStore<pxr::UsdGeomPrimvar>& A, const TUsdStore<pxr::UsdGeomPrimvar>& B)
+		{
+			return A.Get().GetName() < B.Get().GetName();
+		}
+	);
 
 	TArray<TUsdStore<pxr::UsdGeomPrimvar>> Result;
 	Result.Reserve(FMath::Min(TexCoord2fPrimvars.Num() + Float2Primvars.Num(), MaxNumPrimvars));
@@ -803,9 +835,7 @@ TArray<TUsdStore<pxr::UsdGeomPrimvar>> UsdUtils::AssemblePrimvarsIntoUVSets(
 	return PrimvarsByUVIndex;
 }
 
-TMap<FString, int32> UsdUtils::AssemblePrimvarsIntoPrimvarToUVIndexMap(
-	const TArray<TUsdStore<pxr::UsdGeomPrimvar>>& AllMeshUVPrimvars
-)
+TMap<FString, int32> UsdUtils::AssemblePrimvarsIntoPrimvarToUVIndexMap(const TArray<TUsdStore<pxr::UsdGeomPrimvar>>& AllMeshUVPrimvars)
 {
 	TMap<FString, int32> Result;
 	Result.Reserve(AllMeshUVPrimvars.Num());
@@ -822,10 +852,7 @@ TMap<FString, int32> UsdUtils::AssemblePrimvarsIntoPrimvarToUVIndexMap(
 	return Result;
 }
 
-TMap<FString, int32> UsdUtils::CombinePrimvarsIntoUVSets(
-	const TSet<FString>& AllPrimvars,
-	const TSet<FString>& PreferredPrimvars
-)
+TMap<FString, int32> UsdUtils::CombinePrimvarsIntoUVSets(const TSet<FString>& AllPrimvars, const TSet<FString>& PreferredPrimvars)
 {
 	TArray<FString> SortedPrimvars = AllPrimvars.Array();
 
@@ -872,7 +899,7 @@ bool UsdUtils::IsAnimated(const pxr::UsdPrim& Prim)
 	pxr::UsdGeomXformable Xformable(Prim);
 	if (Xformable)
 	{
-		std::vector< double > TimeSamples;
+		std::vector<double> TimeSamples;
 		Xformable.GetTimeSamples(&TimeSamples);
 
 		if (TimeSamples.size() > 0)
@@ -885,15 +912,15 @@ bool UsdUtils::IsAnimated(const pxr::UsdPrim& Prim)
 		// as a direct world transform", but when reading we'll manually recompute the relative transform to its parent anyway
 		// (for simplicity's sake). If that parent (or any of its ancestors) is being animated, we'll need to recompute this
 		// for every animation keyframe, which basically means we're animated too
-		if ( Xformable.GetResetXformStack() )
+		if (Xformable.GetResetXformStack())
 		{
 			pxr::UsdPrim AncestorPrim = Prim.GetParent();
-			while ( AncestorPrim && !AncestorPrim.IsPseudoRoot() )
+			while (AncestorPrim && !AncestorPrim.IsPseudoRoot())
 			{
-				if ( pxr::UsdGeomXformable AncestorXformable{ AncestorPrim } )
+				if (pxr::UsdGeomXformable AncestorXformable{AncestorPrim})
 				{
 					std::vector<double> AncestorTimeSamples;
-					if ( AncestorXformable.GetTimeSamples( &AncestorTimeSamples ) && AncestorTimeSamples.size() > 0 )
+					if (AncestorXformable.GetTimeSamples(&AncestorTimeSamples) && AncestorTimeSamples.size() > 0)
 					{
 						return true;
 					}
@@ -901,7 +928,7 @@ bool UsdUtils::IsAnimated(const pxr::UsdPrim& Prim)
 					// The exception is if our ancestor also wants to reset its xform stack (i.e. its transform is meant to be
 					// used as the world transform). In this case we don't need to care about higher up ancestors anymore, as
 					// their transforms wouldn't affect below this prim anyway
-					if ( AncestorXformable.GetResetXformStack() )
+					if (AncestorXformable.GetResetXformStack())
 					{
 						break;
 					}
@@ -912,11 +939,11 @@ bool UsdUtils::IsAnimated(const pxr::UsdPrim& Prim)
 		}
 	}
 
-	const std::vector< pxr::UsdAttribute >& Attributes = Prim.GetAttributes();
+	const std::vector<pxr::UsdAttribute>& Attributes = Prim.GetAttributes();
 	for (const pxr::UsdAttribute& Attribute : Attributes)
 	{
 		std::vector<double> TimeSamples;
-		if ( Attribute.GetTimeSamples( &TimeSamples ) && TimeSamples.size() > 0 )
+		if (Attribute.GetTimeSamples(&TimeSamples) && TimeSamples.size() > 0)
 		{
 			return true;
 		}
@@ -946,22 +973,22 @@ bool UsdUtils::IsAnimated(const pxr::UsdPrim& Prim)
 	return false;
 }
 
-bool UsdUtils::HasAnimatedVisibility( const pxr::UsdPrim& Prim )
+bool UsdUtils::HasAnimatedVisibility(const pxr::UsdPrim& Prim)
 {
-	if ( !Prim || !Prim.IsActive() )
+	if (!Prim || !Prim.IsActive())
 	{
 		return false;
 	}
 
 	FScopedUsdAllocs UsdAllocs;
 
-	pxr::UsdGeomImageable Imageable( Prim );
-	if ( Imageable )
+	pxr::UsdGeomImageable Imageable(Prim);
+	if (Imageable)
 	{
-		if ( pxr::UsdAttribute Attr = Imageable.GetVisibilityAttr() )
+		if (pxr::UsdAttribute Attr = Imageable.GetVisibilityAttr())
 		{
 			std::vector<double> TimeSamples;
-			if ( Attr.GetTimeSamples( &TimeSamples ) && TimeSamples.size() > 0 )
+			if (Attr.GetTimeSamples(&TimeSamples) && TimeSamples.size() > 0)
 			{
 				return true;
 			}
@@ -1170,14 +1197,9 @@ namespace UE::ConversionUtilsImpl::Private
 
 		return bHasAnimatedBounds;
 	}
-}
+}	 // namespace UE::ConversionUtilsImpl::Private
 
-bool UsdUtils::HasAnimatedBounds(
-	const pxr::UsdPrim& Prim,
-	EUsdPurpose IncludedPurposes,
-	bool bUseExtentsHint,
-	bool bIgnoreVisibility
-)
+bool UsdUtils::HasAnimatedBounds(const pxr::UsdPrim& Prim, EUsdPurpose IncludedPurposes, bool bUseExtentsHint, bool bIgnoreVisibility)
 {
 	// "ParentPrim" here because there are slight differences in behavior between handling the actual provided
 	// prim and another random prim in its subtree (for which bIsParentPrim will be 'false')
@@ -1224,15 +1246,15 @@ bool UsdUtils::GetAnimatedBoundsTimeSamples(
 	return bHasAnimatedBounds;
 }
 
-EUsdDefaultKind UsdUtils::GetDefaultKind( const pxr::UsdPrim& Prim )
+EUsdDefaultKind UsdUtils::GetDefaultKind(const pxr::UsdPrim& Prim)
 {
 	FScopedUsdAllocs Allocs;
 
-	pxr::UsdModelAPI Model{ pxr::UsdTyped( Prim ) };
+	pxr::UsdModelAPI Model{pxr::UsdTyped(Prim)};
 
 	EUsdDefaultKind Result = EUsdDefaultKind::None;
 
-	if ( !Model )
+	if (!Model)
 	{
 		return Result;
 	}
@@ -1240,27 +1262,27 @@ EUsdDefaultKind UsdUtils::GetDefaultKind( const pxr::UsdPrim& Prim )
 	// We need KindValidationNone here or else we get inconsistent results when a prim references another prim that is a component.
 	// For example, when referencing a component prim in another file, this returns 'true' if the referencer is a root prim,
 	// but false if the referencer is within another Xform prim, for whatever reason.
-	if ( Model.IsKind( pxr::KindTokens->model, pxr::UsdModelAPI::KindValidationNone ) )
+	if (Model.IsKind(pxr::KindTokens->model, pxr::UsdModelAPI::KindValidationNone))
 	{
 		Result |= EUsdDefaultKind::Model;
 	}
 
-	if ( Model.IsKind( pxr::KindTokens->component, pxr::UsdModelAPI::KindValidationNone ) )
+	if (Model.IsKind(pxr::KindTokens->component, pxr::UsdModelAPI::KindValidationNone))
 	{
 		Result |= EUsdDefaultKind::Component;
 	}
 
-	if ( Model.IsKind( pxr::KindTokens->group, pxr::UsdModelAPI::KindValidationNone ) )
+	if (Model.IsKind(pxr::KindTokens->group, pxr::UsdModelAPI::KindValidationNone))
 	{
 		Result |= EUsdDefaultKind::Group;
 	}
 
-	if ( Model.IsKind( pxr::KindTokens->assembly, pxr::UsdModelAPI::KindValidationNone ) )
+	if (Model.IsKind(pxr::KindTokens->assembly, pxr::UsdModelAPI::KindValidationNone))
 	{
 		Result |= EUsdDefaultKind::Assembly;
 	}
 
-	if ( Model.IsKind( pxr::KindTokens->subcomponent, pxr::UsdModelAPI::KindValidationNone ) )
+	if (Model.IsKind(pxr::KindTokens->subcomponent, pxr::UsdModelAPI::KindValidationNone))
 	{
 		Result |= EUsdDefaultKind::Subcomponent;
 	}
@@ -1268,24 +1290,24 @@ EUsdDefaultKind UsdUtils::GetDefaultKind( const pxr::UsdPrim& Prim )
 	return Result;
 }
 
-bool UsdUtils::SetDefaultKind( pxr::UsdPrim& Prim, EUsdDefaultKind NewKind )
+bool UsdUtils::SetDefaultKind(pxr::UsdPrim& Prim, EUsdDefaultKind NewKind)
 {
-	if ( !Prim )
+	if (!Prim)
 	{
 		return false;
 	}
 
 	FScopedUsdAllocs Allocs;
 
-	const int32 NewKindInt = static_cast< int32 >( NewKind );
-	const bool bSingleFlagSet = NewKindInt != 0 && ( NewKindInt & ( NewKindInt - 1 ) ) == 0;
-	if ( !bSingleFlagSet )
+	const int32 NewKindInt = static_cast<int32>(NewKind);
+	const bool bSingleFlagSet = NewKindInt != 0 && (NewKindInt & (NewKindInt - 1)) == 0;
+	if (!bSingleFlagSet)
 	{
 		return false;
 	}
 
 	pxr::TfToken NewKindToken;
-	switch ( NewKind )
+	switch (NewKind)
 	{
 		default:
 		case EUsdDefaultKind::Model:
@@ -1314,12 +1336,12 @@ bool UsdUtils::SetDefaultKind( pxr::UsdPrim& Prim, EUsdDefaultKind NewKind )
 			break;
 		}
 	}
-	if ( NewKindToken.IsEmpty() )
+	if (NewKindToken.IsEmpty())
 	{
 		return false;
 	}
 
-	return IUsdPrim::SetKind( Prim, NewKindToken );
+	return IUsdPrim::SetKind(Prim, NewKindToken);
 }
 
 EUsdDrawMode UsdUtils::GetAppliedDrawMode(const pxr::UsdPrim& Prim)
@@ -1399,36 +1421,53 @@ EUsdDrawMode UsdUtils::GetAppliedDrawMode(const pxr::UsdPrim& Prim)
 	}
 }
 
-TArray< TUsdStore< pxr::UsdPrim > > UsdUtils::GetAllPrimsOfType( const pxr::UsdPrim& StartPrim, const pxr::TfType& SchemaType, const TArray< TUsdStore< pxr::TfType > >& ExcludeSchemaTypes )
+TArray<TUsdStore<pxr::UsdPrim>> UsdUtils::GetAllPrimsOfType(
+	const pxr::UsdPrim& StartPrim,
+	const pxr::TfType& SchemaType,
+	const TArray<TUsdStore<pxr::TfType>>& ExcludeSchemaTypes
+)
 {
-    return GetAllPrimsOfType( StartPrim, SchemaType, []( const pxr::UsdPrim& ) { return false; }, ExcludeSchemaTypes );
+	return GetAllPrimsOfType(
+		StartPrim,
+		SchemaType,
+		[](const pxr::UsdPrim&)
+		{
+			return false;
+		},
+		ExcludeSchemaTypes
+	);
 }
 
-TArray< TUsdStore< pxr::UsdPrim > > UsdUtils::GetAllPrimsOfType( const pxr::UsdPrim& StartPrim, const pxr::TfType& SchemaType, TFunction< bool( const pxr::UsdPrim& ) > PruneChildren, const TArray< TUsdStore< pxr::TfType > >& ExcludeSchemaTypes )
+TArray<TUsdStore<pxr::UsdPrim>> UsdUtils::GetAllPrimsOfType(
+	const pxr::UsdPrim& StartPrim,
+	const pxr::TfType& SchemaType,
+	TFunction<bool(const pxr::UsdPrim&)> PruneChildren,
+	const TArray<TUsdStore<pxr::TfType>>& ExcludeSchemaTypes
+)
 {
-	TArray< TUsdStore< pxr::UsdPrim > > Result;
+	TArray<TUsdStore<pxr::UsdPrim>> Result;
 
-	pxr::UsdPrimRange PrimRange( StartPrim, pxr::UsdTraverseInstanceProxies() );
+	pxr::UsdPrimRange PrimRange(StartPrim, pxr::UsdTraverseInstanceProxies());
 
-	for ( pxr::UsdPrimRange::iterator PrimRangeIt = PrimRange.begin(); PrimRangeIt != PrimRange.end(); ++PrimRangeIt )
+	for (pxr::UsdPrimRange::iterator PrimRangeIt = PrimRange.begin(); PrimRangeIt != PrimRange.end(); ++PrimRangeIt)
 	{
 		bool bIsExcluded = false;
 
-		for ( const TUsdStore< pxr::TfType >& SchemaToExclude : ExcludeSchemaTypes )
+		for (const TUsdStore<pxr::TfType>& SchemaToExclude : ExcludeSchemaTypes)
 		{
-			if ( PrimRangeIt->IsA( SchemaToExclude.Get() ) )
+			if (PrimRangeIt->IsA(SchemaToExclude.Get()))
 			{
 				bIsExcluded = true;
 				break;
 			}
 		}
 
-		if ( !bIsExcluded && PrimRangeIt->IsA( SchemaType ) )
+		if (!bIsExcluded && PrimRangeIt->IsA(SchemaType))
 		{
-			Result.Add( *PrimRangeIt );
+			Result.Add(*PrimRangeIt);
 		}
 
-		if ( bIsExcluded || PruneChildren( *PrimRangeIt ) )
+		if (bIsExcluded || PruneChildren(*PrimRangeIt))
 		{
 			PrimRangeIt.PruneChildren();
 		}
@@ -1437,17 +1476,17 @@ TArray< TUsdStore< pxr::UsdPrim > > UsdUtils::GetAllPrimsOfType( const pxr::UsdP
 	return Result;
 }
 
-FString UsdUtils::GetAssetPathFromPrimPath( const FString& RootContentPath, const pxr::UsdPrim& Prim )
+FString UsdUtils::GetAssetPathFromPrimPath(const FString& RootContentPath, const pxr::UsdPrim& Prim)
 {
 	FString FinalPath;
 
-	auto GetEnclosingModelPrim = []( const pxr::UsdPrim& Prim ) -> pxr::UsdPrim
+	auto GetEnclosingModelPrim = [](const pxr::UsdPrim& Prim) -> pxr::UsdPrim
 	{
 		pxr::UsdPrim ModelPrim = Prim.GetParent();
 
-		while ( ModelPrim )
+		while (ModelPrim)
 		{
-			if ( IUsdPrim::IsKindChildOf( ModelPrim, "model" ) )
+			if (IUsdPrim::IsKindChildOf(ModelPrim, "model"))
 			{
 				break;
 			}
@@ -1460,92 +1499,104 @@ FString UsdUtils::GetAssetPathFromPrimPath( const FString& RootContentPath, cons
 		return ModelPrim.IsValid() ? ModelPrim : Prim;
 	};
 
-	const pxr::UsdPrim& ModelPrim = GetEnclosingModelPrim( Prim );
+	const pxr::UsdPrim& ModelPrim = GetEnclosingModelPrim(Prim);
 
-	const FString RawPrimName = UsdToUnreal::ConvertString( Prim.GetName() );
+	const FString RawPrimName = UsdToUnreal::ConvertString(Prim.GetName());
 
-	pxr::UsdModelAPI ModelApi = pxr::UsdModelAPI( ModelPrim );
+	pxr::UsdModelAPI ModelApi = pxr::UsdModelAPI(ModelPrim);
 
 	std::string RawAssetName;
-	ModelApi.GetAssetName( &RawAssetName );
+	ModelApi.GetAssetName(&RawAssetName);
 
-	FString AssetName = UsdToUnreal::ConvertString( RawAssetName );
+	FString AssetName = UsdToUnreal::ConvertString(RawAssetName);
 	FString MeshName = IUsdClassesModule::SanitizeObjectName(RawPrimName);
 
-	FString USDPath = UsdToUnreal::ConvertString( Prim.GetPrimPath().GetString().c_str() );
+	FString USDPath = UsdToUnreal::ConvertString(Prim.GetPrimPath().GetString().c_str());
 
 	pxr::SdfAssetPath AssetPath;
-	if ( ModelApi.GetAssetIdentifier( &AssetPath ) )
+	if (ModelApi.GetAssetIdentifier(&AssetPath))
 	{
 		std::string AssetIdentifier = AssetPath.GetAssetPath();
-		USDPath = UsdToUnreal::ConvertString( AssetIdentifier.c_str() );
+		USDPath = UsdToUnreal::ConvertString(AssetIdentifier.c_str());
 
-		USDPath = FPaths::ConvertRelativePathToFull( RootContentPath, USDPath );
+		USDPath = FPaths::ConvertRelativePathToFull(RootContentPath, USDPath);
 
-		FPackageName::TryConvertFilenameToLongPackageName( USDPath, USDPath );
-		USDPath.RemoveFromEnd( AssetName );
+		FPackageName::TryConvertFilenameToLongPackageName(USDPath, USDPath);
+		USDPath.RemoveFromEnd(AssetName);
 	}
 
 	FString VariantName;
 
-	if ( ModelPrim.HasVariantSets() )
+	if (ModelPrim.HasVariantSets())
 	{
-		pxr::UsdVariantSet ModelVariantSet = ModelPrim.GetVariantSet( "modelingVariant" );
-		if ( ModelVariantSet.IsValid() )
+		pxr::UsdVariantSet ModelVariantSet = ModelPrim.GetVariantSet("modelingVariant");
+		if (ModelVariantSet.IsValid())
 		{
 			std::string VariantSelection = ModelVariantSet.GetVariantSelection();
 
-			if ( VariantSelection.length() > 0 )
+			if (VariantSelection.length() > 0)
 			{
-				VariantName = UsdToUnreal::ConvertString( VariantSelection.c_str() );
+				VariantName = UsdToUnreal::ConvertString(VariantSelection.c_str());
 			}
 		}
 	}
 
-	if ( !VariantName.IsEmpty() )
+	if (!VariantName.IsEmpty())
 	{
 		USDPath = USDPath / VariantName;
 	}
 
-	USDPath.RemoveFromStart( TEXT("/") );
-	USDPath.RemoveFromEnd( RawPrimName );
+	USDPath.RemoveFromStart(TEXT("/"));
+	USDPath.RemoveFromEnd(RawPrimName);
 	FinalPath /= (USDPath / MeshName);
 
 	return FinalPath;
 }
-#endif // #if USE_USD_SDK
+#endif	  // #if USE_USD_SDK
 
-TArray< UE::FUsdPrim > UsdUtils::GetAllPrimsOfType( const UE::FUsdPrim& StartPrim, const TCHAR* SchemaName )
+TArray<UE::FUsdPrim> UsdUtils::GetAllPrimsOfType(const UE::FUsdPrim& StartPrim, const TCHAR* SchemaName)
 {
-	return GetAllPrimsOfType( StartPrim, SchemaName, []( const UE::FUsdPrim& ) { return false; } );
+	return GetAllPrimsOfType(
+		StartPrim,
+		SchemaName,
+		[](const UE::FUsdPrim&)
+		{
+			return false;
+		}
+	);
 }
 
-TArray< UE::FUsdPrim > UsdUtils::GetAllPrimsOfType( const UE::FUsdPrim& StartPrim, const TCHAR* SchemaName, TFunction< bool( const UE::FUsdPrim& ) > PruneChildren, const TArray<const TCHAR*>& ExcludeSchemaNames )
+TArray<UE::FUsdPrim> UsdUtils::GetAllPrimsOfType(
+	const UE::FUsdPrim& StartPrim,
+	const TCHAR* SchemaName,
+	TFunction<bool(const UE::FUsdPrim&)> PruneChildren,
+	const TArray<const TCHAR*>& ExcludeSchemaNames
+)
 {
-	TArray< UE::FUsdPrim > Result;
+	TArray<UE::FUsdPrim> Result;
 
 #if USE_USD_SDK
-	const pxr::TfType SchemaType = pxr::TfType::FindByName( TCHAR_TO_ANSI( SchemaName ) );
+	const pxr::TfType SchemaType = pxr::TfType::FindByName(TCHAR_TO_ANSI(SchemaName));
 
-	TArray< TUsdStore< pxr::TfType > > ExcludeSchemaTypes;
-	ExcludeSchemaTypes.Reserve( ExcludeSchemaNames.Num() );
-	for ( const TCHAR* ExcludeSchemaName : ExcludeSchemaNames )
+	TArray<TUsdStore<pxr::TfType>> ExcludeSchemaTypes;
+	ExcludeSchemaTypes.Reserve(ExcludeSchemaNames.Num());
+	for (const TCHAR* ExcludeSchemaName : ExcludeSchemaNames)
 	{
-		ExcludeSchemaTypes.Add( pxr::TfType( pxr::TfType::FindByName( TCHAR_TO_ANSI( ExcludeSchemaName ) ) ) );
+		ExcludeSchemaTypes.Add(pxr::TfType(pxr::TfType::FindByName(TCHAR_TO_ANSI(ExcludeSchemaName))));
 	}
 
-	auto UsdPruneChildren = [ &PruneChildren ]( const pxr::UsdPrim& ChildPrim ) -> bool
+	auto UsdPruneChildren = [&PruneChildren](const pxr::UsdPrim& ChildPrim) -> bool
 	{
-		return PruneChildren( UE::FUsdPrim( ChildPrim ) );
+		return PruneChildren(UE::FUsdPrim(ChildPrim));
 	};
 
-	TArray< TUsdStore< pxr::UsdPrim > > UsdResult = GetAllPrimsOfType( StartPrim, SchemaType, UsdPruneChildren, ExcludeSchemaTypes );
+	TArray<TUsdStore<pxr::UsdPrim>> UsdResult = GetAllPrimsOfType(StartPrim, SchemaType, UsdPruneChildren, ExcludeSchemaTypes);
 
-	for ( const TUsdStore< pxr::UsdPrim >& Prim : UsdResult )
+	for (const TUsdStore<pxr::UsdPrim>& Prim : UsdResult)
 	{
-		Result.Emplace( Prim.Get() );
+		Result.Emplace(Prim.Get());
 	}
-#endif // #if USE_USD_SDK
+#endif	  // #if USE_USD_SDK
 
 	return Result;
 }
@@ -1568,67 +1619,67 @@ double UsdUtils::GetEarliestTimeCode()
 #endif
 }
 
-UUsdAssetImportData* UsdUtils::GetAssetImportData( UObject* Asset )
+UUsdAssetImportData* UsdUtils::GetAssetImportData(UObject* Asset)
 {
 	UUsdAssetImportData* ImportData = nullptr;
 #if WITH_EDITORONLY_DATA
-	if ( UStaticMesh* Mesh = Cast<UStaticMesh>( Asset ) )
+	if (UStaticMesh* Mesh = Cast<UStaticMesh>(Asset))
 	{
-		ImportData = Cast<UUsdAssetImportData>( Mesh->AssetImportData );
+		ImportData = Cast<UUsdAssetImportData>(Mesh->AssetImportData);
 	}
-	else if ( USkeleton* Skeleton = Cast<USkeleton>( Asset ) )
+	else if (USkeleton* Skeleton = Cast<USkeleton>(Asset))
 	{
-		if ( USkeletalMesh* SkMesh = Skeleton->GetPreviewMesh() )
+		if (USkeletalMesh* SkMesh = Skeleton->GetPreviewMesh())
 		{
-			ImportData = Cast<UUsdAssetImportData>( SkMesh->GetAssetImportData() );
+			ImportData = Cast<UUsdAssetImportData>(SkMesh->GetAssetImportData());
 		}
 	}
-	else if ( UPhysicsAsset* PhysicsAsset = Cast<UPhysicsAsset>( Asset ) )
+	else if (UPhysicsAsset* PhysicsAsset = Cast<UPhysicsAsset>(Asset))
 	{
-		if ( USkeletalMesh* SkMesh = PhysicsAsset->GetPreviewMesh() )
+		if (USkeletalMesh* SkMesh = PhysicsAsset->GetPreviewMesh())
 		{
-			ImportData = Cast<UUsdAssetImportData>( SkMesh->GetAssetImportData() );
+			ImportData = Cast<UUsdAssetImportData>(SkMesh->GetAssetImportData());
 		}
 	}
-	else if ( UAnimBlueprint* AnimBP = Cast<UAnimBlueprint>( Asset ) )
+	else if (UAnimBlueprint* AnimBP = Cast<UAnimBlueprint>(Asset))
 	{
 		// We will always have a skeleton, but not necessarily we will have a preview mesh directly
 		// on the UAnimBlueprint
-		if ( USkeleton* AnimBPSkeleton = AnimBP->TargetSkeleton.Get() )
+		if (USkeleton* AnimBPSkeleton = AnimBP->TargetSkeleton.Get())
 		{
-			if ( USkeletalMesh* SkMesh = AnimBPSkeleton->GetPreviewMesh() )
+			if (USkeletalMesh* SkMesh = AnimBPSkeleton->GetPreviewMesh())
 			{
-				ImportData = Cast<UUsdAssetImportData>( SkMesh->GetAssetImportData() );
+				ImportData = Cast<UUsdAssetImportData>(SkMesh->GetAssetImportData());
 			}
 		}
 	}
-	else if ( USkeletalMesh* SkMesh = Cast<USkeletalMesh>( Asset ) )
+	else if (USkeletalMesh* SkMesh = Cast<USkeletalMesh>(Asset))
 	{
 		ImportData = Cast<UUsdAssetImportData>(SkMesh->GetAssetImportData());
 	}
-	else if ( UAnimSequence* SkelAnim = Cast<UAnimSequence>( Asset ) )
+	else if (UAnimSequence* SkelAnim = Cast<UAnimSequence>(Asset))
 	{
-		ImportData = Cast<UUsdAssetImportData>( SkelAnim->AssetImportData );
+		ImportData = Cast<UUsdAssetImportData>(SkelAnim->AssetImportData);
 	}
-	else if ( UMaterialInterface* Material = Cast<UMaterialInterface>( Asset ) )
+	else if (UMaterialInterface* Material = Cast<UMaterialInterface>(Asset))
 	{
-		ImportData = Cast<UUsdAssetImportData>( Material->AssetImportData );
+		ImportData = Cast<UUsdAssetImportData>(Material->AssetImportData);
 	}
-	else if ( UTexture* Texture = Cast<UTexture>( Asset ) )
+	else if (UTexture* Texture = Cast<UTexture>(Asset))
 	{
-		ImportData = Cast<UUsdAssetImportData>( Texture->AssetImportData );
+		ImportData = Cast<UUsdAssetImportData>(Texture->AssetImportData);
 	}
-	else if ( UGeometryCache* GeometryCache = Cast<UGeometryCache>( Asset ) )
+	else if (UGeometryCache* GeometryCache = Cast<UGeometryCache>(Asset))
 	{
-		ImportData = Cast<UUsdAssetImportData>( GeometryCache->AssetImportData );
+		ImportData = Cast<UUsdAssetImportData>(GeometryCache->AssetImportData);
 	}
-	else if ( UGroomAsset* Groom = Cast<UGroomAsset>( Asset ) )
+	else if (UGroomAsset* Groom = Cast<UGroomAsset>(Asset))
 	{
-		ImportData = Cast<UUsdAssetImportData>( Groom->AssetImportData );
+		ImportData = Cast<UUsdAssetImportData>(Groom->AssetImportData);
 	}
-	else if ( UGroomCache* GroomCache = Cast<UGroomCache>( Asset ) )
+	else if (UGroomCache* GroomCache = Cast<UGroomCache>(Asset))
 	{
-		ImportData = Cast<UUsdAssetImportData>( GroomCache->AssetImportData );
+		ImportData = Cast<UUsdAssetImportData>(GroomCache->AssetImportData);
 	}
 
 #endif
@@ -1675,7 +1726,7 @@ void UsdUtils::SetAssetImportData(UObject* Asset, UAssetImportData* ImportData)
 	{
 		GroomCache->AssetImportData = ImportData;
 	}
-#endif // WITH_EDITOR
+#endif	  // WITH_EDITOR
 }
 
 UUsdAssetUserData* UsdUtils::GetAssetUserData(const UObject* Object, TSubclassOf<UUsdAssetUserData> Class)
@@ -1793,7 +1844,7 @@ FString UsdUtils::GetAssetHashPrefix(const pxr::UsdPrim& PrimForAsset, bool bReu
 	SHA1.GetHash(&Hash.Hash[0]);
 	return Hash.ToString() + TEXT("_");
 }
-#endif // WITH_EDITOR
+#endif	  // WITH_EDITOR
 
 namespace UE::UsdConversionUtils::Private
 {
@@ -1806,24 +1857,24 @@ namespace UE::UsdConversionUtils::Private
 		pxr::SdfPath& OutPrimPath
 	)
 	{
-		if ( !Prim || !AbsoluteFilePath )
+		if (!Prim || !AbsoluteFilePath)
 		{
 			return false;
 		}
 
 		FScopedUsdAllocs UsdAllocs;
 
-		pxr::UsdPrim UsdPrim( Prim );
+		pxr::UsdPrim UsdPrim(Prim);
 
 		pxr::UsdStageRefPtr UsdStage = UsdPrim.GetStage();
-		if ( !UsdStage )
+		if (!UsdStage)
 		{
 			return false;
 		}
 
 		// Turn our layer path into a relative one
 		FString RelativePath = AbsoluteFilePath;
-		if ( !RelativePath.IsEmpty() )
+		if (!RelativePath.IsEmpty())
 		{
 			pxr::SdfLayerHandle EditLayer = UsdPrim.GetStage()->GetEditTarget().GetLayer();
 
@@ -1831,50 +1882,50 @@ namespace UE::UsdConversionUtils::Private
 
 			// If we're editing an in-memory stage our root layer may not have a path yet
 			// Giving an empty InRelativeTo to MakePathRelativeTo causes it to use the engine binary
-			if ( !RepositoryPath.empty() )
+			if (!RepositoryPath.empty())
 			{
-				FString LayerAbsolutePath = UsdToUnreal::ConvertString( RepositoryPath );
-				FPaths::MakePathRelativeTo( RelativePath, *LayerAbsolutePath );
+				FString LayerAbsolutePath = UsdToUnreal::ConvertString(RepositoryPath);
+				FPaths::MakePathRelativeTo(RelativePath, *LayerAbsolutePath);
 			}
 		}
 
 		// Get the target layer
 		pxr::SdfLayerRefPtr TargetLayer;
 		bool bIsInternalReference = false;
-		if ( RelativePath.IsEmpty() )
+		if (RelativePath.IsEmpty())
 		{
 			TargetLayer = UsdStage->GetRootLayer();
 			bIsInternalReference = true;
 		}
 		else
 		{
-			TargetLayer = pxr::SdfLayer::FindOrOpen( UnrealToUsd::ConvertString( AbsoluteFilePath ).Get() );
+			TargetLayer = pxr::SdfLayer::FindOrOpen(UnrealToUsd::ConvertString(AbsoluteFilePath).Get());
 		}
-		if ( !TargetLayer )
+		if (!TargetLayer)
 		{
 			return false;
 		}
 
 		// Get the target prim spec we want to reference
-		pxr::SdfPrimSpecHandle TargetPrimSpec = TargetLayer->GetPrimAtPath( TargetPrimPath );
-		if ( ( TargetPrimPath.IsEmpty() || !TargetPrimSpec ) && TargetLayer->HasDefaultPrim() )
+		pxr::SdfPrimSpecHandle TargetPrimSpec = TargetLayer->GetPrimAtPath(TargetPrimPath);
+		if ((TargetPrimPath.IsEmpty() || !TargetPrimSpec) && TargetLayer->HasDefaultPrim())
 		{
-			TargetPrimSpec = TargetLayer->GetPrimAtPath( pxr::SdfPath::AbsoluteRootPath().AppendChild( TargetLayer->GetDefaultPrim() ) );
+			TargetPrimSpec = TargetLayer->GetPrimAtPath(pxr::SdfPath::AbsoluteRootPath().AppendChild(TargetLayer->GetDefaultPrim()));
 		}
 
 		// Update UsdPrim's type so that it can handle the reference and be parsed by the proper translator
-		if ( !UsdPrim.GetTypeName().IsEmpty() && TargetPrimSpec )
+		if (!UsdPrim.GetTypeName().IsEmpty() && TargetPrimSpec)
 		{
 			pxr::TfToken TargetTypeName = TargetPrimSpec->GetTypeName();
-			pxr::TfType TargetPrimType = pxr::UsdSchemaRegistry::GetTypeFromName( TargetTypeName );
+			pxr::TfType TargetPrimType = pxr::UsdSchemaRegistry::GetTypeFromName(TargetTypeName);
 
-			if ( TargetPrimType.IsUnknown() )
+			if (TargetPrimType.IsUnknown())
 			{
 				UsdPrim.ClearTypeName();
 			}
-			else if ( !UsdPrim.IsA( TargetPrimType ) )
+			else if (!UsdPrim.IsA(TargetPrimType))
 			{
-				UsdPrim.SetTypeName( TargetTypeName );
+				UsdPrim.SetTypeName(TargetTypeName);
 			}
 		}
 
@@ -1882,13 +1933,19 @@ namespace UE::UsdConversionUtils::Private
 		// default prim is, so that the authored reference doesn't actually specify any prim name and just refers to the
 		// default prim by default. Otherwise if the default prim of the layer changed, we wouldn't update to the new prim
 		OutPrimPath = TargetPrimSpec && !TargetPrimPath.IsEmpty() ? TargetPrimSpec->GetPath() : pxr::SdfPath{};
-		OutRelativeFilePath = bIsInternalReference ? std::string() : UnrealToUsd::ConvertString( *RelativePath ).Get();
+		OutRelativeFilePath = bIsInternalReference ? std::string() : UnrealToUsd::ConvertString(*RelativePath).Get();
 		return true;
 	}
-#endif // #if USE_USD_SDK
-}
+#endif	  // #if USE_USD_SDK
+}	 // namespace UE::UsdConversionUtils::Private
 
-void UsdUtils::AddReference( UE::FUsdPrim& Prim, const TCHAR* AbsoluteFilePath, const UE::FSdfPath& TargetPrimPath, double TimeCodeOffset, double TimeCodeScale )
+void UsdUtils::AddReference(
+	UE::FUsdPrim& Prim,
+	const TCHAR* AbsoluteFilePath,
+	const UE::FSdfPath& TargetPrimPath,
+	double TimeCodeOffset,
+	double TimeCodeScale
+)
 {
 #if USE_USD_SDK
 	FScopedUsdAllocs UsdAllocs;
@@ -1907,18 +1964,14 @@ void UsdUtils::AddReference( UE::FUsdPrim& Prim, const TCHAR* AbsoluteFilePath,
 		FinalPrimPath
 	);
 
-	if( !bProceed )
+	if (!bProceed)
 	{
 		return;
 	}
 
-	pxr::UsdReferences References = pxr::UsdPrim{ Prim }.GetReferences();
-	References.AddReference(
-		RelativeLayerPath,
-		FinalPrimPath,
-		pxr::SdfLayerOffset{ TimeCodeOffset, TimeCodeScale }
-	);
-#endif // #if USE_USD_SDK
+	pxr::UsdReferences References = pxr::UsdPrim{Prim}.GetReferences();
+	References.AddReference(RelativeLayerPath, FinalPrimPath, pxr::SdfLayerOffset{TimeCodeOffset, TimeCodeScale});
+#endif	  // #if USE_USD_SDK
 }
 
 bool UsdUtils::GetReferenceFilePath(const UE::FUsdPrim& Prim, const FString& FileExtension, FString& OutReferenceFilePath)
@@ -1947,12 +2000,18 @@ bool UsdUtils::GetReferenceFilePath(const UE::FUsdPrim& Prim, const FString& Fil
 			}
 		}
 	}
-#endif // #if USE_USD_SDK
+#endif	  // #if USE_USD_SDK
 
 	return false;
 }
 
-void UsdUtils::AddPayload( UE::FUsdPrim& Prim, const TCHAR* AbsoluteFilePath, const UE::FSdfPath& TargetPrimPath, double TimeCodeOffset, double TimeCodeScale )
+void UsdUtils::AddPayload(
+	UE::FUsdPrim& Prim,
+	const TCHAR* AbsoluteFilePath,
+	const UE::FSdfPath& TargetPrimPath,
+	double TimeCodeOffset,
+	double TimeCodeScale
+)
 {
 #if USE_USD_SDK
 	FScopedUsdAllocs UsdAllocs;
@@ -1971,50 +2030,46 @@ void UsdUtils::AddPayload( UE::FUsdPrim& Prim, const TCHAR* AbsoluteFilePath, co
 		FinalPrimPath
 	);
 
-	if ( !bProceed )
+	if (!bProceed)
 	{
 		return;
 	}
 
-	pxr::UsdPayloads Payloads = pxr::UsdPrim{ Prim }.GetPayloads();
-	Payloads.AddPayload(
-		RelativeLayerPath,
-		FinalPrimPath,
-		pxr::SdfLayerOffset{ TimeCodeOffset, TimeCodeScale }
-	);
-#endif // #if USE_USD_SDK
+	pxr::UsdPayloads Payloads = pxr::UsdPrim{Prim}.GetPayloads();
+	Payloads.AddPayload(RelativeLayerPath, FinalPrimPath, pxr::SdfLayerOffset{TimeCodeOffset, TimeCodeScale});
+#endif	  // #if USE_USD_SDK
 }
 
-bool UsdUtils::RenamePrim( UE::FUsdPrim& Prim, const TCHAR* NewPrimName )
+bool UsdUtils::RenamePrim(UE::FUsdPrim& Prim, const TCHAR* NewPrimName)
 {
 #if USE_USD_SDK
 	FScopedUsdAllocs UsdAllocs;
 
-	if ( !Prim || !NewPrimName )
+	if (!Prim || !NewPrimName)
 	{
 		return false;
 	}
 
-	if ( Prim.GetName() == FName( NewPrimName ) )
+	if (Prim.GetName() == FName(NewPrimName))
 	{
 		return false;
 	}
 
-	pxr::UsdPrim PxrUsdPrim{ Prim };
-	pxr::UsdStageRefPtr PxrUsdStage{ Prim.GetStage() };
-	if ( !PxrUsdStage )
+	pxr::UsdPrim PxrUsdPrim{Prim};
+	pxr::UsdStageRefPtr PxrUsdStage{Prim.GetStage()};
+	if (!PxrUsdStage)
 	{
 		return false;
 	}
 
-	pxr::TfToken NewNameToken = UnrealToUsd::ConvertToken( NewPrimName ).Get();
-	pxr::SdfPath TargetPath = PxrUsdPrim.GetPrimPath().ReplaceName( NewNameToken );
+	pxr::TfToken NewNameToken = UnrealToUsd::ConvertToken(NewPrimName).Get();
+	pxr::SdfPath TargetPath = PxrUsdPrim.GetPrimPath().ReplaceName(NewNameToken);
 
 	std::unordered_set<std::string> LocalLayerIdentifiers;
 	const bool bIncludeSessionLayers = true;
-	for ( const pxr::SdfLayerHandle& Handle : PxrUsdStage->GetLayerStack( bIncludeSessionLayers ) )
+	for (const pxr::SdfLayerHandle& Handle : PxrUsdStage->GetLayerStack(bIncludeSessionLayers))
 	{
-		LocalLayerIdentifiers.insert( Handle->GetIdentifier() );
+		LocalLayerIdentifiers.insert(Handle->GetIdentifier());
 	}
 
 	std::vector<pxr::SdfPrimSpecHandle> SpecStack = PxrUsdPrim.GetPrimStack();
@@ -2026,15 +2081,15 @@ bool UsdUtils::RenamePrim( UE::FUsdPrim& Prim, const TCHAR* NewPrimName )
 	pxr::SdfNamespaceEditDetailVector Details;
 	int32 LastDetailsSize = 0;
 	bool bCanApply = true;
-	for ( const pxr::SdfPrimSpecHandle& Spec : SpecStack )
+	for (const pxr::SdfPrimSpecHandle& Spec : SpecStack)
 	{
-		if ( !Spec )
+		if (!Spec)
 		{
 			continue;
 		}
 
 		pxr::SdfPath SpecPath = Spec->GetPath();
-		if ( !SpecPath.IsPrimPath() )
+		if (!SpecPath.IsPrimPath())
 		{
 			// Especially when it comes to variants, we can have many different specs for a prim.
 			// e.g. we can simultaneously have "/Prim{Varset=}", "/Prim{Varset=Var}" and "/Prim" in there, and
@@ -2049,41 +2104,44 @@ bool UsdUtils::RenamePrim( UE::FUsdPrim& Prim, const TCHAR* NewPrimName )
 		// we'll end up renaming the prims within those layers too, which is not what we want: For reference/payloads it's as if
 		// we're just consuming the *contents* of the referenced prim, but we don't want to affect it. Another more drastic example:
 		// if we were to remove the referencer prim, we don't really want to delete the referenced prim within its layer
-		if ( LocalLayerIdentifiers.count( SpecLayer->GetIdentifier() ) == 0 )
+		if (LocalLayerIdentifiers.count(SpecLayer->GetIdentifier()) == 0)
 		{
 			continue;
 		}
 
 		pxr::SdfBatchNamespaceEdit BatchEdit;
-		BatchEdit.Add( pxr::SdfNamespaceEdit::Rename( SpecPath, NewNameToken ) );
+		BatchEdit.Add(pxr::SdfNamespaceEdit::Rename(SpecPath, NewNameToken));
 
 		int32 CurrentNumDetails = Details.size();
-		if ( SpecLayer->CanApply( BatchEdit, &Details ) != pxr::SdfNamespaceEditDetail::Result::Okay )
+		if (SpecLayer->CanApply(BatchEdit, &Details) != pxr::SdfNamespaceEditDetail::Result::Okay)
 		{
-			FString LayerIdentifier = UsdToUnreal::ConvertString( SpecLayer->GetIdentifier() );
+			FString LayerIdentifier = UsdToUnreal::ConvertString(SpecLayer->GetIdentifier());
 
 			// This error pushed something new into the Details vector. Get it as an error message
 			FString ErrorMessage;
-			if ( CurrentNumDetails != LastDetailsSize )
+			if (CurrentNumDetails != LastDetailsSize)
 			{
-				ErrorMessage = UsdToUnreal::ConvertString( Details[ CurrentNumDetails - 1 ].reason );
+				ErrorMessage = UsdToUnreal::ConvertString(Details[CurrentNumDetails - 1].reason);
 			}
 
-			ErrorMessages.Add( FString::Printf( TEXT( "\t%s: %s" ), *LayerIdentifier, *ErrorMessage ) );
+			ErrorMessages.Add(FString::Printf(TEXT("\t%s: %s"), *LayerIdentifier, *ErrorMessage));
 			bCanApply = false;
 			// Don't break so we can collect all error messages
 		}
 
 		LastDetailsSize = CurrentNumDetails;
-		Edits.Add( TPair<pxr::SdfLayerRefPtr, pxr::SdfBatchNamespaceEdit>{ SpecLayer, BatchEdit } );
+		Edits.Add(TPair<pxr::SdfLayerRefPtr, pxr::SdfBatchNamespaceEdit>{SpecLayer, BatchEdit});
 	}
 
-	if ( !bCanApply )
+	if (!bCanApply)
 	{
-		UE_LOG( LogUsd, Error, TEXT( "Failed to rename prim with path '%s' to name '%s'. Errors:\n%s" ),
+		UE_LOG(
+			LogUsd,
+			Error,
+			TEXT("Failed to rename prim with path '%s' to name '%s'. Errors:\n%s"),
 			*Prim.GetPrimPath().GetString(),
 			NewPrimName,
-			*FString::Join( ErrorMessages, TEXT( "\n" ) )
+			*FString::Join(ErrorMessages, TEXT("\n"))
 		);
 
 		return false;
@@ -2093,7 +2151,7 @@ bool UsdUtils::RenamePrim( UE::FUsdPrim& Prim, const TCHAR* NewPrimName )
 	{
 		pxr::SdfChangeBlock Block;
 
-		for ( const TPair<pxr::SdfLayerRefPtr, pxr::SdfBatchNamespaceEdit>& Pair : Edits )
+		for (const TPair<pxr::SdfLayerRefPtr, pxr::SdfBatchNamespaceEdit>& Pair : Edits)
 		{
 			const pxr::SdfLayerRefPtr& Layer = Pair.Key;
 			const pxr::SdfBatchNamespaceEdit& Edit = Pair.Value;
@@ -2101,23 +2159,26 @@ bool UsdUtils::RenamePrim( UE::FUsdPrim& Prim, const TCHAR* NewPrimName )
 			// Make sure that if the renamed prim is the layer's default prim, we also update that to match the
 			// prim's new name
 			pxr::UsdPrim ParentPrim = PxrUsdPrim.GetParent();
-			const bool bNeedToRenameDefaultPrim = ParentPrim && ParentPrim.IsPseudoRoot() && ( PxrUsdPrim.GetName() == Layer->GetDefaultPrim() );
+			const bool bNeedToRenameDefaultPrim = ParentPrim && ParentPrim.IsPseudoRoot() && (PxrUsdPrim.GetName() == Layer->GetDefaultPrim());
 
-			if ( !Layer->Apply( Edit ) )
+			if (!Layer->Apply(Edit))
 			{
 				// This should not be happening since CanApply was true, so stop doing whatever it is we're doing
-				UE_LOG( LogUsd, Error, TEXT( "Failed to rename prim with path '%s' to name '%s' in layer '%s'" ),
+				UE_LOG(
+					LogUsd,
+					Error,
+					TEXT("Failed to rename prim with path '%s' to name '%s' in layer '%s'"),
 					*Prim.GetPrimPath().GetString(),
 					NewPrimName,
-					*UsdToUnreal::ConvertString( Layer->GetIdentifier() )
+					*UsdToUnreal::ConvertString(Layer->GetIdentifier())
 				);
 
 				return false;
 			}
 
-			if ( bNeedToRenameDefaultPrim )
+			if (bNeedToRenameDefaultPrim)
 			{
-				Layer->SetDefaultPrim( NewNameToken );
+				Layer->SetDefaultPrim(NewNameToken);
 			}
 		}
 	}
@@ -2126,49 +2187,51 @@ bool UsdUtils::RenamePrim( UE::FUsdPrim& Prim, const TCHAR* NewPrimName )
 	// Here we override that with a SetActive opinion on the session layer, which will also trigger a new resync of that prim.
 	//
 	// We must send a separate notice for this (which is why this function can't be inside a change block) for two reasons:
-	// - In order to let the transactor know that this edit is done on the session layer (so that we can have our active=true opinion there and not save it to disk);
+	// - In order to let the transactor know that this edit is done on the session layer (so that we can have our active=true opinion there and not
+	// save it to disk);
 	// - Because after we apply the rename, usd *needs* to responds to notices in order to make the target path valid again. Until
 	//   it does so, we can't Get/Override/Define a prim at the target path at all, and so can't set it to active.
 	//
 	// We can't do this *before* we rename because if we already have a prim defined/overriden on "/Root/Target", then we
 	// can't apply a rename from a prim onto "/Root/Target": Meaning we'd lose all extra data we have on the prim on the session layer.
 	{
-		pxr::UsdEditContext EditContext{ PxrUsdStage, PxrUsdStage->GetSessionLayer() };
+		pxr::UsdEditContext EditContext{PxrUsdStage, PxrUsdStage->GetSessionLayer()};
 
-		if ( pxr::UsdPrim PostRenamePrim = PxrUsdStage->OverridePrim( TargetPath ) )
+		if (pxr::UsdPrim PostRenamePrim = PxrUsdStage->OverridePrim(TargetPath))
 		{
 			// We need to toggle it back and forth because whenever we undo a rename we'll rename our spec on the session layer
 			// back to the original path, and that spec *already* has an active=true opinion that we set during the first rename.
 			// This means that just setting it to active here wouldn't send any notice (because it already is). We need a new notice
-			// to update to the fact that the child prim is now active again (the rename notice is a resync, but it already comes with the prim set to inactive)
+			// to update to the fact that the child prim is now active again (the rename notice is a resync, but it already comes with the prim set to
+			// inactive)
 			pxr::SdfChangeBlock Block;
 			const bool bActive = true;
-			PostRenamePrim.SetActive( !bActive );
-			PostRenamePrim.SetActive( bActive );
+			PostRenamePrim.SetActive(!bActive);
+			PostRenamePrim.SetActive(bActive);
 		}
 	}
 
 	return true;
 #else
 	return false;
-#endif // #if USE_USD_SDK
+#endif	  // #if USE_USD_SDK
 }
 
-bool UsdUtils::RemoveNumberedSuffix( FString& Prefix )
+bool UsdUtils::RemoveNumberedSuffix(FString& Prefix)
 {
-	if ( Prefix.IsNumeric() )
+	if (Prefix.IsNumeric())
 	{
 		return false;
 	}
 
 	bool bRemoved = false;
 
-	FString LastChar = Prefix.Right( 1 );
-	while ( ( LastChar.IsNumeric() || LastChar == TEXT( "_" ) ) && Prefix.Len() > 1 )
+	FString LastChar = Prefix.Right(1);
+	while ((LastChar.IsNumeric() || LastChar == TEXT("_")) && Prefix.Len() > 1)
 	{
 		const bool bAllowShrinking = false;
-		Prefix.LeftChopInline( 1, bAllowShrinking );
-		LastChar = Prefix.Right( 1 );
+		Prefix.LeftChopInline(1, bAllowShrinking);
+		LastChar = Prefix.Right(1);
 
 		bRemoved = true;
 	}
@@ -2177,17 +2240,17 @@ bool UsdUtils::RemoveNumberedSuffix( FString& Prefix )
 	return bRemoved;
 }
 
-FString UsdUtils::GetUniqueName( FString Name, const TSet<FString>& UsedNames )
+FString UsdUtils::GetUniqueName(FString Name, const TSet<FString>& UsedNames)
 {
-	if ( !UsedNames.Contains( Name ) )
+	if (!UsedNames.Contains(Name))
 	{
 		return Name;
 	}
 
-	const bool bRemoved = RemoveNumberedSuffix( Name );
+	const bool bRemoved = RemoveNumberedSuffix(Name);
 
 	// Its possible that removing the suffix made it into a unique name already
-	if ( bRemoved && !UsedNames.Contains( Name ) )
+	if (bRemoved && !UsedNames.Contains(Name))
 	{
 		return Name;
 	}
@@ -2196,16 +2259,16 @@ FString UsdUtils::GetUniqueName( FString Name, const TSet<FString>& UsedNames )
 	FString Result;
 	do
 	{
-		Result = FString::Printf( TEXT( "%s_%d" ), *Name, Suffix++ );
-	} while ( UsedNames.Contains( Result ) );
+		Result = FString::Printf(TEXT("%s_%d"), *Name, Suffix++);
+	} while (UsedNames.Contains(Result));
 
 	return Result;
 }
 
 #if USE_USD_SDK
-FString UsdUtils::GetValidChildName( FString InName, const pxr::UsdPrim& ParentPrim )
+FString UsdUtils::GetValidChildName(FString InName, const pxr::UsdPrim& ParentPrim)
 {
-	if ( !ParentPrim )
+	if (!ParentPrim)
 	{
 		return {};
 	}
@@ -2213,84 +2276,84 @@ FString UsdUtils::GetValidChildName( FString InName, const pxr::UsdPrim& ParentP
 	FScopedUsdAllocs Allocs;
 
 	TSet<FString> UsedNames;
-	for ( const pxr::UsdPrim& Child : ParentPrim.GetChildren() )
+	for (const pxr::UsdPrim& Child : ParentPrim.GetChildren())
 	{
-		UsedNames.Add( UsdToUnreal::ConvertToken( Child.GetName() ) );
+		UsedNames.Add(UsdToUnreal::ConvertToken(Child.GetName()));
 	}
 
-	return GetUniqueName( SanitizeUsdIdentifier( *InName ), UsedNames );
+	return GetUniqueName(SanitizeUsdIdentifier(*InName), UsedNames);
 }
-#endif // USE_USD_SDK
+#endif	  // USE_USD_SDK
 
-FString UsdUtils::SanitizeUsdIdentifier( const TCHAR* InIdentifier )
+FString UsdUtils::SanitizeUsdIdentifier(const TCHAR* InIdentifier)
 {
 #if USE_USD_SDK
 	FScopedUsdAllocs Allocs;
 
-	std::string UsdInName = UnrealToUsd::ConvertString( InIdentifier ).Get();
-	std::string UsdValidName = pxr::TfMakeValidIdentifier( UsdInName );
+	std::string UsdInName = UnrealToUsd::ConvertString(InIdentifier).Get();
+	std::string UsdValidName = pxr::TfMakeValidIdentifier(UsdInName);
 
-	return UsdToUnreal::ConvertString( UsdValidName );
-#endif // USE_USD_SDK
+	return UsdToUnreal::ConvertString(UsdValidName);
+#endif	  // USE_USD_SDK
 
 	return InIdentifier;
 }
 
-void UsdUtils::MakeVisible( UE::FUsdPrim& Prim, double TimeCode )
+void UsdUtils::MakeVisible(UE::FUsdPrim& Prim, double TimeCode)
 {
 #if USE_USD_SDK
 	FScopedUsdAllocs Allocs;
 
-	pxr::UsdPrim PxrUsdPrim{ Prim };
-	if ( pxr::UsdGeomImageable Imageable{ PxrUsdPrim } )
+	pxr::UsdPrim PxrUsdPrim{Prim};
+	if (pxr::UsdGeomImageable Imageable{PxrUsdPrim})
 	{
-		Imageable.MakeVisible( TimeCode );
+		Imageable.MakeVisible(TimeCode);
 	}
-#endif // USE_USD_SDK
+#endif	  // USE_USD_SDK
 }
 
-void UsdUtils::MakeInvisible( UE::FUsdPrim& Prim, double TimeCode )
+void UsdUtils::MakeInvisible(UE::FUsdPrim& Prim, double TimeCode)
 {
 #if USE_USD_SDK
 	FScopedUsdAllocs Allocs;
 
-	pxr::UsdPrim PxrUsdPrim{ Prim };
-	if ( pxr::UsdGeomImageable Imageable{ PxrUsdPrim } )
+	pxr::UsdPrim PxrUsdPrim{Prim};
+	if (pxr::UsdGeomImageable Imageable{PxrUsdPrim})
 	{
-		Imageable.MakeInvisible( TimeCode );
+		Imageable.MakeInvisible(TimeCode);
 	}
-#endif // USE_USD_SDK
+#endif	  // USE_USD_SDK
 }
 
-bool UsdUtils::IsVisible( const UE::FUsdPrim& Prim, double TimeCode )
+bool UsdUtils::IsVisible(const UE::FUsdPrim& Prim, double TimeCode)
 {
 #if USE_USD_SDK
 	FScopedUsdAllocs Allocs;
 
-	pxr::UsdPrim PxrUsdPrim{ Prim };
-	if ( pxr::UsdGeomImageable Imageable{ PxrUsdPrim } )
+	pxr::UsdPrim PxrUsdPrim{Prim};
+	if (pxr::UsdGeomImageable Imageable{PxrUsdPrim})
 	{
-		return Imageable.ComputeVisibility( TimeCode ) == pxr::UsdGeomTokens->inherited;
+		return Imageable.ComputeVisibility(TimeCode) == pxr::UsdGeomTokens->inherited;
 	}
 
 	return true;
 #else
 	return false;
-#endif // USE_USD_SDK
+#endif	  // USE_USD_SDK
 }
 
-bool UsdUtils::HasInheritedVisibility( const UE::FUsdPrim& Prim, double TimeCode )
+bool UsdUtils::HasInheritedVisibility(const UE::FUsdPrim& Prim, double TimeCode)
 {
 #if USE_USD_SDK
 	FScopedUsdAllocs Allocs;
 
-	pxr::UsdPrim PxrUsdPrim{ Prim };
-	if ( pxr::UsdGeomImageable Imageable{ PxrUsdPrim } )
+	pxr::UsdPrim PxrUsdPrim{Prim};
+	if (pxr::UsdGeomImageable Imageable{PxrUsdPrim})
 	{
-		if ( pxr::UsdAttribute VisibilityAttr = Imageable.GetVisibilityAttr() )
+		if (pxr::UsdAttribute VisibilityAttr = Imageable.GetVisibilityAttr())
 		{
 			pxr::TfToken Visibility;
-			if ( !VisibilityAttr.Get<pxr::TfToken>( &Visibility, TimeCode ) )
+			if (!VisibilityAttr.Get<pxr::TfToken>(&Visibility, TimeCode))
 			{
 				return true;
 			}
@@ -2303,27 +2366,27 @@ bool UsdUtils::HasInheritedVisibility( const UE::FUsdPrim& Prim, double TimeCode
 	return true;
 #else
 	return false;
-#endif // USE_USD_SDK
+#endif	  // USE_USD_SDK
 }
 
-bool UsdUtils::HasInvisibleParent( const UE::FUsdPrim& Prim, const UE::FUsdPrim& RootPrim, double TimeCode )
+bool UsdUtils::HasInvisibleParent(const UE::FUsdPrim& Prim, const UE::FUsdPrim& RootPrim, double TimeCode)
 {
 #if USE_USD_SDK
 	FScopedUsdAllocs Allocs;
 
-	pxr::UsdPrim PxrUsdPrim{ Prim };
+	pxr::UsdPrim PxrUsdPrim{Prim};
 	pxr::UsdPrim Parent = PxrUsdPrim.GetParent();
 
-	while ( Parent && Parent != RootPrim )
+	while (Parent && Parent != RootPrim)
 	{
-		if ( pxr::UsdGeomImageable Imageable{ Parent } )
+		if (pxr::UsdGeomImageable Imageable{Parent})
 		{
-			if ( pxr::UsdAttribute VisibilityAttr = Imageable.GetVisibilityAttr() )
+			if (pxr::UsdAttribute VisibilityAttr = Imageable.GetVisibilityAttr())
 			{
 				pxr::TfToken Visibility;
-				if ( VisibilityAttr.Get<pxr::TfToken>( &Visibility, TimeCode ) )
+				if (VisibilityAttr.Get<pxr::TfToken>(&Visibility, TimeCode))
 				{
-					if ( Visibility == pxr::UsdGeomTokens->invisible )
+					if (Visibility == pxr::UsdGeomTokens->invisible)
 					{
 						return true;
 					}
@@ -2333,7 +2396,7 @@ bool UsdUtils::HasInvisibleParent( const UE::FUsdPrim& Prim, const UE::FUsdPrim&
 
 		Parent = Parent.GetParent();
 	}
-#endif // USE_USD_SDK
+#endif	  // USE_USD_SDK
 
 	return false;
 }
@@ -2374,20 +2437,20 @@ TArray<UE::FUsdPrim> UsdUtils::GetVisibleChildren(const UE::FUsdPrim& Prim, EUsd
 		}
 	};
 	RecursivelyCollectVisibleMeshes(Prim);
-#endif // USE_USD_SDK
+#endif	  // USE_USD_SDK
 
 	return VisiblePrims;
 }
 
-UE::FSdfPath UsdUtils::GetPrimSpecPathForLayer( const UE::FUsdPrim& Prim, const UE::FSdfLayer& Layer )
+UE::FSdfPath UsdUtils::GetPrimSpecPathForLayer(const UE::FUsdPrim& Prim, const UE::FSdfLayer& Layer)
 {
 	UE::FSdfPath Result;
 #if USE_USD_SDK
 	FScopedUsdAllocs Allocs;
 
-	pxr::UsdPrim UsdPrim{ Prim };
-	pxr::SdfLayerRefPtr UsdLayer{ Layer };
-	if ( !UsdPrim || !UsdLayer )
+	pxr::UsdPrim UsdPrim{Prim};
+	pxr::SdfLayerRefPtr UsdLayer{Layer};
+	if (!UsdPrim || !UsdLayer)
 	{
 		return Result;
 	}
@@ -2396,45 +2459,45 @@ UE::FSdfPath UsdUtils::GetPrimSpecPathForLayer( const UE::FUsdPrim& Prim, const
 	// "/Root{Varset=Var}Parent/Child{ChildSet=ChildVar}" and "/Root{Varset=Var}Parent/Child").
 	// This function needs to return a prim path with all of its variant selections (i.e. the last example above)
 	std::size_t LargestPathLength = 0;
-	for ( const pxr::SdfPrimSpecHandle& Spec : UsdPrim.GetPrimStack() )
+	for (const pxr::SdfPrimSpecHandle& Spec : UsdPrim.GetPrimStack())
 	{
-		if ( !Spec )
+		if (!Spec)
 		{
 			continue;
 		}
 
 		pxr::SdfPath SpecPath = Spec->GetPath();
-		if ( !SpecPath.IsPrimPath() )
+		if (!SpecPath.IsPrimPath())
 		{
 			continue;
 		}
 
-		if ( Spec->GetLayer() == UsdLayer )
+		if (Spec->GetLayer() == UsdLayer)
 		{
 			const std::size_t NewPathLength = Spec->GetPath().GetString().length();
-			if ( NewPathLength > LargestPathLength )
+			if (NewPathLength > LargestPathLength)
 			{
-				Result = UE::FSdfPath{ SpecPath };
+				Result = UE::FSdfPath{SpecPath};
 			}
 		}
 	}
 
-#endif // USE_USD_SDK
+#endif	  // USE_USD_SDK
 	return Result;
 }
 
-USDUTILITIES_API void UsdUtils::RemoveAllLocalPrimSpecs( const UE::FUsdPrim& Prim, const UE::FSdfLayer& Layer )
+USDUTILITIES_API void UsdUtils::RemoveAllLocalPrimSpecs(const UE::FUsdPrim& Prim, const UE::FSdfLayer& Layer)
 {
 #if USE_USD_SDK
 	FScopedUsdAllocs Allocs;
 
-	pxr::UsdPrim UsdPrim{ Prim };
-	if ( !UsdPrim )
+	pxr::UsdPrim UsdPrim{Prim};
+	if (!UsdPrim)
 	{
 		return;
 	}
 
-	pxr::SdfLayerRefPtr UsdLayer{ Layer };
+	pxr::SdfLayerRefPtr UsdLayer{Layer};
 	pxr::UsdStageRefPtr UsdStage = UsdPrim.GetStage();
 
 	std::unordered_set<std::string> LocalLayerIdentifiers;
@@ -2442,21 +2505,21 @@ USDUTILITIES_API void UsdUtils::RemoveAllLocalPrimSpecs( const UE::FUsdPrim& Pri
 	// We'll want to remove specs from the entire stage. We need to be careful though to only remove specs from the
 	// local layer stack. If a prim within the stage has a reference/payload to another layer and we remove the
 	// referencer prim, we don't want to end up removing the referenced/payload prim within its own layer too.
-	if ( !UsdLayer )
+	if (!UsdLayer)
 	{
 		const bool bIncludeSessionLayers = true;
-		for ( const pxr::SdfLayerHandle& Handle : UsdStage->GetLayerStack( bIncludeSessionLayers ) )
+		for (const pxr::SdfLayerHandle& Handle : UsdStage->GetLayerStack(bIncludeSessionLayers))
 		{
-			LocalLayerIdentifiers.insert( Handle->GetIdentifier() );
+			LocalLayerIdentifiers.insert(Handle->GetIdentifier());
 		}
 	}
 
 	const pxr::SdfPath TargetPath = UsdPrim.GetPrimPath();
 
-	for ( const pxr::SdfPrimSpecHandle& Spec : UsdPrim.GetPrimStack() )
+	for (const pxr::SdfPrimSpecHandle& Spec : UsdPrim.GetPrimStack())
 	{
 		// For whatever reason sometimes there are invalid specs in the layer stack, so we need to be careful
-		if ( !Spec )
+		if (!Spec)
 		{
 			continue;
 		}
@@ -2469,47 +2532,53 @@ USDUTILITIES_API void UsdUtils::RemoveAllLocalPrimSpecs( const UE::FUsdPrim& Pri
 		// the variant selection and look like '/PrimWithVarSet{VarSet=SomeVar}ChildPrim', but our TargetPath will
 		// just look like '/PrimWithVarSet/ChildPrim' instead. These do refer to the exact same prim on the stage
 		// though (when SomeVar is active at least), so we do want to remove both
-		if ( !SpecPath.IsPrimPath() || SpecPath.StripAllVariantSelections() != TargetPath )
+		if (!SpecPath.IsPrimPath() || SpecPath.StripAllVariantSelections() != TargetPath)
 		{
 			continue;
 		}
 
 		pxr::SdfLayerRefPtr SpecLayer = Spec->GetLayer();
-		if ( UsdLayer && SpecLayer != UsdLayer )
+		if (UsdLayer && SpecLayer != UsdLayer)
 		{
 			continue;
 		}
 
-		if ( !UsdLayer && LocalLayerIdentifiers.count( SpecLayer->GetIdentifier() ) == 0 )
+		if (!UsdLayer && LocalLayerIdentifiers.count(SpecLayer->GetIdentifier()) == 0)
 		{
 			continue;
 		}
 
-		UE_LOG( LogUsd, Log, TEXT( "Removing prim spec '%s' from layer '%s'" ), *UsdToUnreal::ConvertPath( SpecPath ), *UsdToUnreal::ConvertString( SpecLayer->GetIdentifier() ) );
-		pxr::UsdEditContext Context( UsdStage, SpecLayer );
-		UsdStage->RemovePrim( SpecPath );
+		UE_LOG(
+			LogUsd,
+			Log,
+			TEXT("Removing prim spec '%s' from layer '%s'"),
+			*UsdToUnreal::ConvertPath(SpecPath),
+			*UsdToUnreal::ConvertString(SpecLayer->GetIdentifier())
+		);
+		pxr::UsdEditContext Context(UsdStage, SpecLayer);
+		UsdStage->RemovePrim(SpecPath);
 	}
 
-#endif // USE_USD_SDK
+#endif	  // USE_USD_SDK
 }
 
-bool UsdUtils::CutPrims( const TArray<UE::FUsdPrim>& Prims )
+bool UsdUtils::CutPrims(const TArray<UE::FUsdPrim>& Prims)
 {
-	bool bCopied = UsdUtils::CopyPrims( Prims );
-	if ( !bCopied )
+	bool bCopied = UsdUtils::CopyPrims(Prims);
+	if (!bCopied)
 	{
 		return false;
 	}
 
-	for ( const UE::FUsdPrim& Prim : Prims )
+	for (const UE::FUsdPrim& Prim : Prims)
 	{
-		UsdUtils::RemoveAllLocalPrimSpecs( Prim );
+		UsdUtils::RemoveAllLocalPrimSpecs(Prim);
 	}
 
 	return true;
 }
 
-bool UsdUtils::CopyPrims( const TArray<UE::FUsdPrim>& Prims )
+bool UsdUtils::CopyPrims(const TArray<UE::FUsdPrim>& Prims)
 {
 	bool bCopiedSomething = false;
 
@@ -2517,49 +2586,49 @@ bool UsdUtils::CopyPrims( const TArray<UE::FUsdPrim>& Prims )
 	FScopedUsdAllocs Allocs;
 
 	pxr::UsdStageRefPtr UsdStage;
-	for ( const UE::FUsdPrim& Prim : Prims )
+	for (const UE::FUsdPrim& Prim : Prims)
 	{
-		if ( Prim )
+		if (Prim)
 		{
-			UsdStage = pxr::UsdStageRefPtr{ Prim.GetStage() };
-			if ( UsdStage )
+			UsdStage = pxr::UsdStageRefPtr{Prim.GetStage()};
+			if (UsdStage)
 			{
 				break;
 			}
 		}
 	}
-	if ( !UsdStage )
+	if (!UsdStage)
 	{
 		return false;
 	}
 
-	pxr::UsdStageRefPtr ClipboardStage = pxr::UsdStageRefPtr{ UnrealUSDWrapper::GetClipboardStage() };
-	if ( !ClipboardStage )
+	pxr::UsdStageRefPtr ClipboardStage = pxr::UsdStageRefPtr{UnrealUSDWrapper::GetClipboardStage()};
+	if (!ClipboardStage)
 	{
 		return false;
 	}
 
 	pxr::SdfLayerHandle ClipboardRoot = ClipboardStage->GetRootLayer();
-	if ( !ClipboardRoot )
+	if (!ClipboardRoot)
 	{
 		return false;
 	}
 
 	pxr::UsdStagePopulationMask Mask;
-	for ( const UE::FUsdPrim& Prim : Prims )
+	for (const UE::FUsdPrim& Prim : Prims)
 	{
-		if ( Prim )
+		if (Prim)
 		{
-			Mask.Add( pxr::SdfPath{ Prim.GetPrimPath() } );
+			Mask.Add(pxr::SdfPath{Prim.GetPrimPath()});
 		}
 	}
-	if ( Mask.IsEmpty() )
+	if (Mask.IsEmpty())
 	{
 		return false;
 	}
 
-	pxr::UsdStageRefPtr TempStage = pxr::UsdStage::OpenMasked( UsdStage->GetRootLayer(), Mask );
-	if ( !TempStage )
+	pxr::UsdStageRefPtr TempStage = pxr::UsdStage::OpenMasked(UsdStage->GetRootLayer(), Mask);
+	if (!TempStage)
 	{
 		return false;
 	}
@@ -2590,8 +2659,8 @@ bool UsdUtils::CopyPrims( const TArray<UE::FUsdPrim>& Prims )
 	}
 
 	const bool bAddSourceFileComment = false;
-	pxr::SdfLayerRefPtr FlattenedLayer = TempStage->Flatten( bAddSourceFileComment );
-	if ( !FlattenedLayer )
+	pxr::SdfLayerRefPtr FlattenedLayer = TempStage->Flatten(bAddSourceFileComment);
+	if (!FlattenedLayer)
 	{
 		return false;
 	}
@@ -2611,10 +2680,10 @@ bool UsdUtils::CopyPrims( const TArray<UE::FUsdPrim>& Prims )
 
 	TSet<FString> UsedNames;
 
-	for ( const UE::FUsdPrim& Prim : Prims )
+	for (const UE::FUsdPrim& Prim : Prims)
 	{
-		pxr::SdfPrimSpecHandle FlattenedPrim = FlattenedLayer->GetPrimAtPath( pxr::SdfPath{ Prim.GetPrimPath() } );
-		if ( !FlattenedPrim )
+		pxr::SdfPrimSpecHandle FlattenedPrim = FlattenedLayer->GetPrimAtPath(pxr::SdfPath{Prim.GetPrimPath()});
+		if (!FlattenedPrim)
 		{
 			continue;
 		}
@@ -2625,117 +2694,118 @@ bool UsdUtils::CopyPrims( const TArray<UE::FUsdPrim>& Prims )
 		// which could potentially yield stale results until USD actually emits the notices about these prims being
 		// added.
 		FString PrimName = Prim.GetName().ToString();
-		FString UniqueName = GetUniqueName( SanitizeUsdIdentifier( *PrimName ), UsedNames );
-		UsedNames.Add( UniqueName );
+		FString UniqueName = GetUniqueName(SanitizeUsdIdentifier(*PrimName), UsedNames);
+		UsedNames.Add(UniqueName);
 
 		const bool bSuccess = pxr::SdfCopySpec(
 			FlattenedLayer,
 			FlattenedPrim->GetPath(),
 			ClipboardRoot,
-			pxr::SdfPath::AbsoluteRootPath().AppendChild( UnrealToUsd::ConvertToken( *UniqueName ).Get() )
+			pxr::SdfPath::AbsoluteRootPath().AppendChild(UnrealToUsd::ConvertToken(*UniqueName).Get())
 		);
-		if ( !bSuccess )
+		if (!bSuccess)
 		{
 			continue;
 		}
 
 		bCopiedSomething = true;
-		UE_LOG( LogUsd, Log, TEXT( "Copied prim '%s' into the clipboard" ),
-			*Prim.GetPrimPath().GetString()
-		);
+		UE_LOG(LogUsd, Log, TEXT("Copied prim '%s' into the clipboard"), *Prim.GetPrimPath().GetString());
 	}
-#endif // USE_USD_SDK
+#endif	  // USE_USD_SDK
 
 	return bCopiedSomething;
 }
 
-TArray<UE::FSdfPath> UsdUtils::PastePrims( const UE::FUsdPrim& ParentPrim )
+TArray<UE::FSdfPath> UsdUtils::PastePrims(const UE::FUsdPrim& ParentPrim)
 {
 	TArray<UE::FSdfPath> Result;
 
 #if USE_USD_SDK
 	FScopedUsdAllocs Allocs;
 
-	pxr::UsdPrim UsdParentPrim{ ParentPrim };
-	if ( !UsdParentPrim )
+	pxr::UsdPrim UsdParentPrim{ParentPrim};
+	if (!UsdParentPrim)
 	{
 		return Result;
 	}
 
 	pxr::UsdStageRefPtr UsdStage = UsdParentPrim.GetStage();
-	if ( !UsdStage )
+	if (!UsdStage)
 	{
 		return Result;
 	}
 
-	pxr::UsdStageRefPtr ClipboardStage = pxr::UsdStageRefPtr{ UnrealUSDWrapper::GetClipboardStage() };
-	if ( !ClipboardStage )
+	pxr::UsdStageRefPtr ClipboardStage = pxr::UsdStageRefPtr{UnrealUSDWrapper::GetClipboardStage()};
+	if (!ClipboardStage)
 	{
 		return Result;
 	}
 
 	pxr::SdfLayerHandle ClipboardRoot = ClipboardStage->GetRootLayer();
-	if ( !ClipboardRoot )
+	if (!ClipboardRoot)
 	{
 		return Result;
 	}
 
 	pxr::UsdPrimSiblingRange PrimChildren = ClipboardStage->GetPseudoRoot().GetChildren();
-	int32 NumPrimsToPaste = std::distance( PrimChildren.begin(), PrimChildren.end() );
+	int32 NumPrimsToPaste = std::distance(PrimChildren.begin(), PrimChildren.end());
 
 	TArray<pxr::UsdPrim> PrimsToPaste;
-	PrimsToPaste.Reserve( NumPrimsToPaste );
-	for ( const pxr::UsdPrim& ClipboardPrim : ClipboardStage->GetPseudoRoot().GetChildren() )
+	PrimsToPaste.Reserve(NumPrimsToPaste);
+	for (const pxr::UsdPrim& ClipboardPrim : ClipboardStage->GetPseudoRoot().GetChildren())
 	{
-		PrimsToPaste.Add( ClipboardPrim );
+		PrimsToPaste.Add(ClipboardPrim);
 	}
 
 	pxr::SdfLayerHandle EditTarget = UsdStage->GetEditTarget().GetLayer();
-	if ( !EditTarget )
+	if (!EditTarget)
 	{
 		return Result;
 	}
 
 	TSet<FString> UsedNames;
-	for ( const pxr::UsdPrim& Child : ParentPrim.GetChildren() )
+	for (const pxr::UsdPrim& Child : ParentPrim.GetChildren())
 	{
-		UsedNames.Add( UsdToUnreal::ConvertToken( Child.GetName() ) );
+		UsedNames.Add(UsdToUnreal::ConvertToken(Child.GetName()));
 	}
 
-	Result.SetNum( NumPrimsToPaste );
-	for ( int32 Index = 0; Index < NumPrimsToPaste; ++Index )
+	Result.SetNum(NumPrimsToPaste);
+	for (int32 Index = 0; Index < NumPrimsToPaste; ++Index)
 	{
-		const pxr::UsdPrim& ClipboardPrim = PrimsToPaste[ Index ];
-		if ( !ClipboardPrim )
+		const pxr::UsdPrim& ClipboardPrim = PrimsToPaste[Index];
+		if (!ClipboardPrim)
 		{
 			continue;
 		}
 
-		const FString OriginalName = UsdToUnreal::ConvertToken( ClipboardPrim.GetName() );
-		FString ValidName = GetUniqueName( SanitizeUsdIdentifier( *OriginalName ), UsedNames );
-		UsedNames.Add( ValidName );
+		const FString OriginalName = UsdToUnreal::ConvertToken(ClipboardPrim.GetName());
+		FString ValidName = GetUniqueName(SanitizeUsdIdentifier(*OriginalName), UsedNames);
+		UsedNames.Add(ValidName);
 
-		pxr::SdfPath TargetSpecPath = UsdParentPrim.GetPath().AppendChild( UnrealToUsd::ConvertToken( *ValidName ).Get() );
+		pxr::SdfPath TargetSpecPath = UsdParentPrim.GetPath().AppendChild(UnrealToUsd::ConvertToken(*ValidName).Get());
 
 		// Ensure our parent prim spec exists, otherwise pxr::SdfCopySpec will fail
-		if ( !pxr::SdfCreatePrimInLayer( EditTarget, TargetSpecPath ) )
+		if (!pxr::SdfCreatePrimInLayer(EditTarget, TargetSpecPath))
 		{
 			continue;
 		}
 
-		if ( !pxr::SdfCopySpec( ClipboardRoot, ClipboardPrim.GetPath(), EditTarget, TargetSpecPath ) )
+		if (!pxr::SdfCopySpec(ClipboardRoot, ClipboardPrim.GetPath(), EditTarget, TargetSpecPath))
 		{
 			continue;
 		}
 
-		UE_LOG( LogUsd, Log, TEXT( "Pasted prim '%s' as a child of prim '%s' within the edit target '%s'" ),
+		UE_LOG(
+			LogUsd,
+			Log,
+			TEXT("Pasted prim '%s' as a child of prim '%s' within the edit target '%s'"),
 			*OriginalName,
-			*UsdToUnreal::ConvertPath( UsdParentPrim.GetPath() ),
-			*UsdToUnreal::ConvertString( EditTarget->GetIdentifier() )
+			*UsdToUnreal::ConvertPath(UsdParentPrim.GetPath()),
+			*UsdToUnreal::ConvertString(EditTarget->GetIdentifier())
 		);
-		Result[ Index ] = UE::FSdfPath{ TargetSpecPath };
+		Result[Index] = UE::FSdfPath{TargetSpecPath};
 	}
-#endif // USE_USD_SDK
+#endif	  // USE_USD_SDK
 
 	return Result;
 }
@@ -2743,20 +2813,20 @@ TArray<UE::FSdfPath> UsdUtils::PastePrims( const UE::FUsdPrim& ParentPrim )
 bool UsdUtils::CanPastePrims()
 {
 #if USE_USD_SDK
-	pxr::UsdStageRefPtr ClipboardStage = pxr::UsdStageRefPtr{ UnrealUSDWrapper::GetClipboardStage() };
-	if ( !ClipboardStage )
+	pxr::UsdStageRefPtr ClipboardStage = pxr::UsdStageRefPtr{UnrealUSDWrapper::GetClipboardStage()};
+	if (!ClipboardStage)
 	{
 		return false;
 	}
 
-	for ( const pxr::UsdPrim& ClipboardPrim : ClipboardStage->GetPseudoRoot().GetChildren() )
+	for (const pxr::UsdPrim& ClipboardPrim : ClipboardStage->GetPseudoRoot().GetChildren())
 	{
-		if ( ClipboardPrim )
+		if (ClipboardPrim)
 		{
 			return true;
 		}
 	}
-#endif // USE_USD_SDK
+#endif	  // USE_USD_SDK
 
 	return false;
 }
@@ -2764,94 +2834,94 @@ bool UsdUtils::CanPastePrims()
 void UsdUtils::ClearPrimClipboard()
 {
 #if USE_USD_SDK
-	pxr::UsdStageRefPtr ClipboardStage = pxr::UsdStageRefPtr{ UnrealUSDWrapper::GetClipboardStage() };
-	if ( !ClipboardStage )
+	pxr::UsdStageRefPtr ClipboardStage = pxr::UsdStageRefPtr{UnrealUSDWrapper::GetClipboardStage()};
+	if (!ClipboardStage)
 	{
 		return;
 	}
 
 	pxr::SdfLayerHandle ClipboardRoot = ClipboardStage->GetRootLayer();
-	if ( !ClipboardRoot )
+	if (!ClipboardRoot)
 	{
 		return;
 	}
 
 	ClipboardRoot->Clear();
-#endif // USE_USD_SDK
+#endif	  // USE_USD_SDK
 }
 
-TArray<UE::FSdfPath> UsdUtils::DuplicatePrims( const TArray<UE::FUsdPrim>& Prims, EUsdDuplicateType DuplicateType, const UE::FSdfLayer& TargetLayer )
+TArray<UE::FSdfPath> UsdUtils::DuplicatePrims(const TArray<UE::FUsdPrim>& Prims, EUsdDuplicateType DuplicateType, const UE::FSdfLayer& TargetLayer)
 {
 	TArray<UE::FSdfPath> Result;
-	Result.SetNum( Prims.Num() );
+	Result.SetNum(Prims.Num());
 
 #if USE_USD_SDK
 	FScopedUsdAllocs Allocs;
 
 	pxr::UsdStageRefPtr UsdStage;
-	for ( const UE::FUsdPrim& Prim : Prims )
+	for (const UE::FUsdPrim& Prim : Prims)
 	{
-		if ( Prim )
+		if (Prim)
 		{
-			UsdStage = pxr::UsdStageRefPtr{ Prim.GetStage() };
-			if ( UsdStage )
+			UsdStage = pxr::UsdStageRefPtr{Prim.GetStage()};
+			if (UsdStage)
 			{
 				break;
 			}
 		}
 	}
-	if ( !UsdStage )
+	if (!UsdStage)
 	{
 		return Result;
 	}
 
-	pxr::SdfLayerRefPtr UsdLayer{ TargetLayer };
+	pxr::SdfLayerRefPtr UsdLayer{TargetLayer};
 
 	// Figure out which layers we'll modify
 	std::unordered_set<pxr::SdfLayerHandle, pxr::TfHash> LayersThatCanBeAffected;
-	switch ( DuplicateType )
+	switch (DuplicateType)
 	{
 		case EUsdDuplicateType::FlattenComposedPrim:
 		case EUsdDuplicateType::SingleLayerSpecs:
 		{
-			if ( !UsdLayer )
+			if (!UsdLayer)
 			{
 				return Result;
 			}
 
-			LayersThatCanBeAffected.insert( UsdLayer );
+			LayersThatCanBeAffected.insert(UsdLayer);
 			break;
 		}
 		case EUsdDuplicateType::AllLocalLayerSpecs:
 		{
 			const bool bIncludeSessionLayers = true;
-			for ( const pxr::SdfLayerHandle& Handle : UsdStage->GetLayerStack( bIncludeSessionLayers ) )
+			for (const pxr::SdfLayerHandle& Handle : UsdStage->GetLayerStack(bIncludeSessionLayers))
 			{
-				LayersThatCanBeAffected.insert( Handle );
+				LayersThatCanBeAffected.insert(Handle);
 			}
 
 			// If any of our prims has specs on layers that are used by the stage but are not within the local layer
 			// stack, then warn the user that some of these specs will not be duplicated
 			{
 				TArray<UE::FUsdPrim> PrimsWithExternalSpecs;
-				for ( const UE::FUsdPrim& Prim : Prims )
+				for (const UE::FUsdPrim& Prim : Prims)
 				{
-					pxr::UsdPrim UsdPrim{ Prim };
-					if ( !UsdPrim )
+					pxr::UsdPrim UsdPrim{Prim};
+					if (!UsdPrim)
 					{
 						continue;
 					}
 
-					for ( const pxr::SdfPrimSpecHandle& Spec : UsdPrim.GetPrimStack() )
+					for (const pxr::SdfPrimSpecHandle& Spec : UsdPrim.GetPrimStack())
 					{
-						if ( Spec && LayersThatCanBeAffected.count( Spec->GetLayer() ) == 0 )
+						if (Spec && LayersThatCanBeAffected.count(Spec->GetLayer()) == 0)
 						{
 							PrimsWithExternalSpecs.Add(Prim);
 							break;
 						}
 					}
 				}
-				USDConversionUtilsImpl::NotifySpecsWontBeDuplicated( PrimsWithExternalSpecs );
+				USDConversionUtilsImpl::NotifySpecsWontBeDuplicated(PrimsWithExternalSpecs);
 			}
 			break;
 		}
@@ -2859,20 +2929,20 @@ TArray<UE::FSdfPath> UsdUtils::DuplicatePrims( const TArray<UE::FUsdPrim>& Prims
 
 	// If we're going to need to flatten, just flatten the stage once for all prims we'll duplicate
 	pxr::SdfLayerRefPtr FlattenedLayer = nullptr;
-	if ( DuplicateType == EUsdDuplicateType::FlattenComposedPrim )
+	if (DuplicateType == EUsdDuplicateType::FlattenComposedPrim)
 	{
 		pxr::UsdStagePopulationMask Mask;
-		for ( int32 Index = 0; Index < Prims.Num(); ++Index )
+		for (int32 Index = 0; Index < Prims.Num(); ++Index)
 		{
-			pxr::UsdPrim UsdPrim{ Prims[ Index ] };
-			if ( UsdPrim )
+			pxr::UsdPrim UsdPrim{Prims[Index]};
+			if (UsdPrim)
 			{
-				Mask.Add( UsdPrim.GetPath() );
+				Mask.Add(UsdPrim.GetPath());
 			}
 		}
 
-		pxr::UsdStageRefPtr TempStage = pxr::UsdStage::OpenMasked( UsdStage->GetRootLayer(), Mask );
-		if ( !TempStage )
+		pxr::UsdStageRefPtr TempStage = pxr::UsdStage::OpenMasked(UsdStage->GetRootLayer(), Mask);
+		if (!TempStage)
 		{
 			return Result;
 		}
@@ -2903,8 +2973,8 @@ TArray<UE::FSdfPath> UsdUtils::DuplicatePrims( const TArray<UE::FUsdPrim>& Prims
 		}
 
 		const bool bAddSourceFileComment = false;
-		FlattenedLayer = TempStage->Flatten( bAddSourceFileComment );
-		if ( !FlattenedLayer )
+		FlattenedLayer = TempStage->Flatten(bAddSourceFileComment);
+		if (!FlattenedLayer)
 		{
 			return Result;
 		}
@@ -2921,10 +2991,10 @@ TArray<UE::FSdfPath> UsdUtils::DuplicatePrims( const TArray<UE::FUsdPrim>& Prims
 		}
 	}
 
-	for ( int32 Index = 0; Index < Prims.Num(); ++Index )
+	for (int32 Index = 0; Index < Prims.Num(); ++Index)
 	{
-		pxr::UsdPrim UsdPrim{ Prims[ Index ] };
-		if ( !UsdPrim )
+		pxr::UsdPrim UsdPrim{Prims[Index]};
+		if (!UsdPrim)
 		{
 			continue;
 		}
@@ -2934,10 +3004,10 @@ TArray<UE::FSdfPath> UsdUtils::DuplicatePrims( const TArray<UE::FUsdPrim>& Prims
 		// Note: We won't actually use these in case we're flattening, but it makes the code a bit simpler to also
 		// do this while we're collecting LayersThatWillBeAffected below
 		std::vector<pxr::SdfPrimSpecHandle> SpecsToDuplicate;
-		SpecsToDuplicate.reserve( PrimSpecs.size() );
+		SpecsToDuplicate.reserve(PrimSpecs.size());
 
 		std::unordered_set<pxr::SdfLayerHandle, pxr::TfHash> LayersThatWillBeAffected;
-		LayersThatWillBeAffected.reserve( PrimSpecs.size() );
+		LayersThatWillBeAffected.reserve(PrimSpecs.size());
 
 		pxr::SdfPath TargetPath = UsdPrim.GetPrimPath();
 		for (int32 SpecIndex = PrimSpecs.size() - 1; SpecIndex >= 0; --SpecIndex)
@@ -2945,7 +3015,7 @@ TArray<UE::FSdfPath> UsdUtils::DuplicatePrims( const TArray<UE::FUsdPrim>& Prims
 			const pxr::SdfPrimSpecHandle& Spec = PrimSpecs[SpecIndex];
 
 			// For whatever reason sometimes there are invalid specs in the layer stack, so we need to be careful
-			if ( !Spec )
+			if (!Spec)
 			{
 				continue;
 			}
@@ -2963,13 +3033,13 @@ TArray<UE::FSdfPath> UsdUtils::DuplicatePrims( const TArray<UE::FUsdPrim>& Prims
 			}
 
 			pxr::SdfLayerHandle SpecLayerHandle = Spec->GetLayer();
-			if ( !SpecLayerHandle || LayersThatCanBeAffected.count( SpecLayerHandle ) == 0 )
+			if (!SpecLayerHandle || LayersThatCanBeAffected.count(SpecLayerHandle) == 0)
 			{
 				continue;
 			}
 
-			SpecsToDuplicate.push_back( Spec );
-			LayersThatWillBeAffected.insert( SpecLayerHandle );
+			SpecsToDuplicate.push_back(Spec);
+			LayersThatWillBeAffected.insert(SpecLayerHandle);
 		}
 
 		// Find a usable name for the new duplicate prim
@@ -2981,14 +3051,14 @@ TArray<UE::FSdfPath> UsdUtils::DuplicatePrims( const TArray<UE::FUsdPrim>& Prims
 			int32 Suffix = -1;
 
 			bool bFoundName = false;
-			while ( !bFoundName )
+			while (!bFoundName)
 			{
-				NewSpecPath = ParentPath.AppendElementString( SourcePrimName + "_" + std::to_string( ++Suffix ) );
+				NewSpecPath = ParentPath.AppendElementString(SourcePrimName + "_" + std::to_string(++Suffix));
 
 				bFoundName = true;
-				for ( const pxr::SdfLayerHandle& Layer : LayersThatWillBeAffected )
+				for (const pxr::SdfLayerHandle& Layer : LayersThatWillBeAffected)
 				{
-					if ( Layer->HasSpec( NewSpecPath ) )
+					if (Layer->HasSpec(NewSpecPath))
 					{
 						bFoundName = false;
 						break;
@@ -2998,93 +3068,106 @@ TArray<UE::FSdfPath> UsdUtils::DuplicatePrims( const TArray<UE::FUsdPrim>& Prims
 		}
 
 		// Actually do the duplication operation we chose
-		if ( DuplicateType == EUsdDuplicateType::FlattenComposedPrim && FlattenedLayer )
+		if (DuplicateType == EUsdDuplicateType::FlattenComposedPrim && FlattenedLayer)
 		{
-			pxr::SdfPrimSpecHandle FlattenedPrim = FlattenedLayer->GetPrimAtPath( UsdPrim.GetPath() );
-			if ( !FlattenedPrim )
+			pxr::SdfPrimSpecHandle FlattenedPrim = FlattenedLayer->GetPrimAtPath(UsdPrim.GetPath());
+			if (!FlattenedPrim)
 			{
 				return Result;
 			}
 
-			if ( !pxr::SdfJustCreatePrimInLayer( UsdLayer, NewSpecPath ) )
+			if (!pxr::SdfJustCreatePrimInLayer(UsdLayer, NewSpecPath))
 			{
-				UE_LOG( LogUsd, Warning, TEXT( "Failed to create prim and parent specs for path '%s' within layer '%s'" ),
-					*UsdToUnreal::ConvertPath( NewSpecPath ),
-					*UsdToUnreal::ConvertString( UsdLayer->GetIdentifier() )
+				UE_LOG(
+					LogUsd,
+					Warning,
+					TEXT("Failed to create prim and parent specs for path '%s' within layer '%s'"),
+					*UsdToUnreal::ConvertPath(NewSpecPath),
+					*UsdToUnreal::ConvertString(UsdLayer->GetIdentifier())
 				);
 				return Result;
 			}
 
-			if ( !pxr::SdfCopySpec( FlattenedLayer, FlattenedPrim->GetPath(), UsdLayer, NewSpecPath ) )
+			if (!pxr::SdfCopySpec(FlattenedLayer, FlattenedPrim->GetPath(), UsdLayer, NewSpecPath))
 			{
-				UE_LOG( LogUsd, Warning, TEXT( "Failed to copy flattened prim spec from '%s' onto path '%s' within layer '%s'" ),
-					*UsdToUnreal::ConvertPath( UsdPrim.GetPath() ),
-					*UsdToUnreal::ConvertPath( NewSpecPath ),
-					*UsdToUnreal::ConvertString( UsdLayer->GetIdentifier() )
+				UE_LOG(
+					LogUsd,
+					Warning,
+					TEXT("Failed to copy flattened prim spec from '%s' onto path '%s' within layer '%s'"),
+					*UsdToUnreal::ConvertPath(UsdPrim.GetPath()),
+					*UsdToUnreal::ConvertPath(NewSpecPath),
+					*UsdToUnreal::ConvertString(UsdLayer->GetIdentifier())
 				);
 				return Result;
 			}
 
-			UE_LOG( LogUsd, Log, TEXT( "Flattened prim '%s' onto spec '%s' at layer '%s'" ),
-				*UsdToUnreal::ConvertPath( UsdPrim.GetPath() ),
-				*UsdToUnreal::ConvertPath( NewSpecPath ),
-				*UsdToUnreal::ConvertString( UsdLayer->GetIdentifier() )
+			UE_LOG(
+				LogUsd,
+				Log,
+				TEXT("Flattened prim '%s' onto spec '%s' at layer '%s'"),
+				*UsdToUnreal::ConvertPath(UsdPrim.GetPath()),
+				*UsdToUnreal::ConvertPath(NewSpecPath),
+				*UsdToUnreal::ConvertString(UsdLayer->GetIdentifier())
 			);
 		}
 		else
 		{
-			for ( const pxr::SdfPrimSpecHandle& Spec : SpecsToDuplicate )
+			for (const pxr::SdfPrimSpecHandle& Spec : SpecsToDuplicate)
 			{
 				pxr::SdfPath SpecPath = Spec->GetPath();
 				pxr::SdfLayerHandle SpecLayerHandle = Spec->GetLayer();
 
-				UE_LOG( LogUsd, Log, TEXT( "Duplicating prim spec '%s' within layer '%s'" ),
-					*UsdToUnreal::ConvertPath( SpecPath ),
-					*UsdToUnreal::ConvertString( SpecLayerHandle->GetIdentifier() )
+				UE_LOG(
+					LogUsd,
+					Log,
+					TEXT("Duplicating prim spec '%s' within layer '%s'"),
+					*UsdToUnreal::ConvertPath(SpecPath),
+					*UsdToUnreal::ConvertString(SpecLayerHandle->GetIdentifier())
 				);
 
 				// Technically we shouldn't need to do this since we'll already do our changes on the Sdf level, however the
 				// USDTransactor will record these notices as belonging to the current edit target, and if that is not in sync
 				// with the layer that is actually changing, we won't be able to undo/redo the duplicate operation
-				pxr::UsdEditContext Context{ UsdStage, SpecLayerHandle };
+				pxr::UsdEditContext Context{UsdStage, SpecLayerHandle};
 
 				// Since we're duplicating a prim essentially as a sibling, parent specs should always exist.
 				// Let's ensure that though, just in case
-				if ( !pxr::SdfJustCreatePrimInLayer( SpecLayerHandle, NewSpecPath ) )
+				if (!pxr::SdfJustCreatePrimInLayer(SpecLayerHandle, NewSpecPath))
 				{
-					UE_LOG( LogUsd, Warning, TEXT( "Failed to create prim and parent specs for path '%s' within layer '%s'" ),
-						*UsdToUnreal::ConvertPath( NewSpecPath ),
-						*UsdToUnreal::ConvertString( SpecLayerHandle->GetIdentifier() )
+					UE_LOG(
+						LogUsd,
+						Warning,
+						TEXT("Failed to create prim and parent specs for path '%s' within layer '%s'"),
+						*UsdToUnreal::ConvertPath(NewSpecPath),
+						*UsdToUnreal::ConvertString(SpecLayerHandle->GetIdentifier())
 					);
 					continue;
 				}
 
-				pxr::SdfShouldCopyValueFn ShouldCopyValue =
-					[](pxr::SdfSpecType SpecType,
-					   const pxr::TfToken& Field,
-					   const pxr::SdfLayerHandle& SrcLayer,
-					   const pxr::SdfPath& SrcPath,
-					   bool FieldInSrc,
-					   const pxr::SdfLayerHandle& DstLayer,
-					   const pxr::SdfPath& DstPath,
-					   bool FieldInDst,
-					   boost::optional<pxr::VtValue>* ValueToCopy) -> bool
+				pxr::SdfShouldCopyValueFn ShouldCopyValue = [](pxr::SdfSpecType SpecType,
+															   const pxr::TfToken& Field,
+															   const pxr::SdfLayerHandle& SrcLayer,
+															   const pxr::SdfPath& SrcPath,
+															   bool FieldInSrc,
+															   const pxr::SdfLayerHandle& DstLayer,
+															   const pxr::SdfPath& DstPath,
+															   bool FieldInDst,
+															   boost::optional<pxr::VtValue>* ValueToCopy) -> bool
 				{
 					// Only copy a field over if it has a value. Otherwise it seems to clear the destination spec
 					// for nothing
 					return FieldInSrc;
 				};
 
-				pxr::SdfShouldCopyChildrenFn ShouldCopyChildren =
-					[](const pxr::TfToken& ChildrenField,
-					   const pxr::SdfLayerHandle& SrcLayer,
-					   const pxr::SdfPath& SrcPath,
-					   bool FieldInSrc,
-					   const pxr::SdfLayerHandle& DstLayer,
-					   const pxr::SdfPath& DstPath,
-					   bool FieldInDst,
-					   boost::optional<pxr::VtValue>* SrcChildren,
-					   boost::optional<pxr::VtValue>* DstChildren) -> bool
+				pxr::SdfShouldCopyChildrenFn ShouldCopyChildren = [](const pxr::TfToken& ChildrenField,
+																	 const pxr::SdfLayerHandle& SrcLayer,
+																	 const pxr::SdfPath& SrcPath,
+																	 bool FieldInSrc,
+																	 const pxr::SdfLayerHandle& DstLayer,
+																	 const pxr::SdfPath& DstPath,
+																	 bool FieldInDst,
+																	 boost::optional<pxr::VtValue>* SrcChildren,
+																	 boost::optional<pxr::VtValue>* DstChildren) -> bool
 				{
 					return true;
 				};
@@ -3105,20 +3188,20 @@ TArray<UE::FSdfPath> UsdUtils::DuplicatePrims( const TArray<UE::FUsdPrim>& Prims
 			}
 		}
 
-		Result[ Index ] = UE::FSdfPath{ NewSpecPath };
+		Result[Index] = UE::FSdfPath{NewSpecPath};
 	}
-#endif // USE_USD_SDK
+#endif	  // USE_USD_SDK
 
 	return Result;
 }
 
-void UsdUtils::SetPrimAssetInfo( UE::FUsdPrim& Prim, const FUsdUnrealAssetInfo& Info )
+void UsdUtils::SetPrimAssetInfo(UE::FUsdPrim& Prim, const FUsdUnrealAssetInfo& Info)
 {
 #if USE_USD_SDK
 	FScopedUsdAllocs Allocs;
 
-	pxr::UsdPrim UsdPrim{ Prim };
-	if ( !UsdPrim )
+	pxr::UsdPrim UsdPrim{Prim};
+	if (!UsdPrim)
 	{
 		return;
 	}
@@ -3126,75 +3209,60 @@ void UsdUtils::SetPrimAssetInfo( UE::FUsdPrim& Prim, const FUsdUnrealAssetInfo&
 	// Just fetch the dictionary already since we'll add custom keys anyway
 	pxr::VtDictionary AssetInfoDict = UsdPrim.GetAssetInfo();
 
-	if ( !Info.Name.IsEmpty() )
+	if (!Info.Name.IsEmpty())
 	{
-		AssetInfoDict.SetValueAtPath(
-			pxr::UsdModelAPIAssetInfoKeys->name,
-			pxr::VtValue{ UnrealToUsd::ConvertString( *Info.Name ).Get() }
-		);
+		AssetInfoDict.SetValueAtPath(pxr::UsdModelAPIAssetInfoKeys->name, pxr::VtValue{UnrealToUsd::ConvertString(*Info.Name).Get()});
 	}
 
-	if ( !Info.Identifier.IsEmpty() )
+	if (!Info.Identifier.IsEmpty())
 	{
 		AssetInfoDict.SetValueAtPath(
 			pxr::UsdModelAPIAssetInfoKeys->identifier,
-			pxr::VtValue{ pxr::SdfAssetPath{ UnrealToUsd::ConvertString( *Info.Identifier ).Get() } }
+			pxr::VtValue{pxr::SdfAssetPath{UnrealToUsd::ConvertString(*Info.Identifier).Get()}}
 		);
 	}
 
-	if ( !Info.Version.IsEmpty() )
+	if (!Info.Version.IsEmpty())
 	{
-		AssetInfoDict.SetValueAtPath(
-			pxr::UsdModelAPIAssetInfoKeys->version,
-			pxr::VtValue{ UnrealToUsd::ConvertString( *Info.Version ).Get() }
-		);
+		AssetInfoDict.SetValueAtPath(pxr::UsdModelAPIAssetInfoKeys->version, pxr::VtValue{UnrealToUsd::ConvertString(*Info.Version).Get()});
 	}
 
-	if ( !Info.UnrealContentPath.IsEmpty() )
+	if (!Info.UnrealContentPath.IsEmpty())
 	{
-		AssetInfoDict.SetValueAtPath(
-			UnrealIdentifiers::UnrealContentPath,
-			pxr::VtValue{ UnrealToUsd::ConvertString( *Info.UnrealContentPath ).Get() }
-		);
+		AssetInfoDict.SetValueAtPath(UnrealIdentifiers::UnrealContentPath, pxr::VtValue{UnrealToUsd::ConvertString(*Info.UnrealContentPath).Get()});
 	}
 
-	if ( !Info.UnrealAssetType.IsEmpty() )
+	if (!Info.UnrealAssetType.IsEmpty())
 	{
-		AssetInfoDict.SetValueAtPath(
-			UnrealIdentifiers::UnrealAssetType,
-			pxr::VtValue{ UnrealToUsd::ConvertString( *Info.UnrealAssetType ).Get() }
-		);
+		AssetInfoDict.SetValueAtPath(UnrealIdentifiers::UnrealAssetType, pxr::VtValue{UnrealToUsd::ConvertString(*Info.UnrealAssetType).Get()});
 	}
 
-	if ( !Info.UnrealExportTime.IsEmpty() )
+	if (!Info.UnrealExportTime.IsEmpty())
 	{
-		AssetInfoDict.SetValueAtPath(
-			UnrealIdentifiers::UnrealExportTime,
-			pxr::VtValue{ UnrealToUsd::ConvertString( *Info.UnrealExportTime ).Get() }
-		);
+		AssetInfoDict.SetValueAtPath(UnrealIdentifiers::UnrealExportTime, pxr::VtValue{UnrealToUsd::ConvertString(*Info.UnrealExportTime).Get()});
 	}
 
-	if ( !Info.UnrealEngineVersion.IsEmpty() )
+	if (!Info.UnrealEngineVersion.IsEmpty())
 	{
 		AssetInfoDict.SetValueAtPath(
 			UnrealIdentifiers::UnrealEngineVersion,
-			pxr::VtValue{ UnrealToUsd::ConvertString( *Info.UnrealEngineVersion ).Get() }
+			pxr::VtValue{UnrealToUsd::ConvertString(*Info.UnrealEngineVersion).Get()}
 		);
 	}
 
-	UsdPrim.SetAssetInfo( AssetInfoDict );
-#endif // USE_USD_SDK
+	UsdPrim.SetAssetInfo(AssetInfoDict);
+#endif	  // USE_USD_SDK
 }
 
-FUsdUnrealAssetInfo UsdUtils::GetPrimAssetInfo( const UE::FUsdPrim& Prim )
+FUsdUnrealAssetInfo UsdUtils::GetPrimAssetInfo(const UE::FUsdPrim& Prim)
 {
 	FUsdUnrealAssetInfo Result;
 
 #if USE_USD_SDK
 	FScopedUsdAllocs Allocs;
 
-	pxr::UsdPrim UsdPrim{ Prim };
-	if ( !UsdPrim )
+	pxr::UsdPrim UsdPrim{Prim};
+	if (!UsdPrim)
 	{
 		return Result;
 	}
@@ -3202,62 +3270,62 @@ FUsdUnrealAssetInfo UsdUtils::GetPrimAssetInfo( const UE::FUsdPrim& Prim )
 	// Just fetch the dictionary already since we'll fetch custom keys anyway
 	pxr::VtDictionary AssetInfoDict = UsdPrim.GetAssetInfo();
 
-	if ( const pxr::VtValue* Value = AssetInfoDict.GetValueAtPath( pxr::UsdModelAPIAssetInfoKeys->name ) )
+	if (const pxr::VtValue* Value = AssetInfoDict.GetValueAtPath(pxr::UsdModelAPIAssetInfoKeys->name))
 	{
-		if ( Value->IsHolding<std::string>() )
+		if (Value->IsHolding<std::string>())
 		{
-			Result.Name = UsdToUnreal::ConvertString( Value->Get<std::string>() );
+			Result.Name = UsdToUnreal::ConvertString(Value->Get<std::string>());
 		}
 	}
 
-	if ( const pxr::VtValue* Value = AssetInfoDict.GetValueAtPath( pxr::UsdModelAPIAssetInfoKeys->identifier ) )
+	if (const pxr::VtValue* Value = AssetInfoDict.GetValueAtPath(pxr::UsdModelAPIAssetInfoKeys->identifier))
 	{
-		if ( Value->IsHolding<pxr::SdfAssetPath>() )
+		if (Value->IsHolding<pxr::SdfAssetPath>())
 		{
-			Result.Identifier = UsdToUnreal::ConvertString( Value->Get<pxr::SdfAssetPath>().GetAssetPath() );
+			Result.Identifier = UsdToUnreal::ConvertString(Value->Get<pxr::SdfAssetPath>().GetAssetPath());
 		}
 	}
 
-	if ( const pxr::VtValue* Value = AssetInfoDict.GetValueAtPath( pxr::UsdModelAPIAssetInfoKeys->version ) )
+	if (const pxr::VtValue* Value = AssetInfoDict.GetValueAtPath(pxr::UsdModelAPIAssetInfoKeys->version))
 	{
-		if ( Value->IsHolding<std::string>() )
+		if (Value->IsHolding<std::string>())
 		{
-			Result.Version = UsdToUnreal::ConvertString( Value->Get<std::string>() );
+			Result.Version = UsdToUnreal::ConvertString(Value->Get<std::string>());
 		}
 	}
 
-	if ( const pxr::VtValue* Value = AssetInfoDict.GetValueAtPath( UnrealIdentifiers::UnrealContentPath ) )
+	if (const pxr::VtValue* Value = AssetInfoDict.GetValueAtPath(UnrealIdentifiers::UnrealContentPath))
 	{
-		if ( Value->IsHolding<std::string>() )
+		if (Value->IsHolding<std::string>())
 		{
-			Result.UnrealContentPath = UsdToUnreal::ConvertString( Value->Get<std::string>() );
+			Result.UnrealContentPath = UsdToUnreal::ConvertString(Value->Get<std::string>());
 		}
 	}
 
-	if ( const pxr::VtValue* Value = AssetInfoDict.GetValueAtPath( UnrealIdentifiers::UnrealAssetType ) )
+	if (const pxr::VtValue* Value = AssetInfoDict.GetValueAtPath(UnrealIdentifiers::UnrealAssetType))
 	{
-		if ( Value->IsHolding<std::string>() )
+		if (Value->IsHolding<std::string>())
 		{
-			Result.UnrealAssetType = UsdToUnreal::ConvertString( Value->Get<std::string>() );
+			Result.UnrealAssetType = UsdToUnreal::ConvertString(Value->Get<std::string>());
 		}
 	}
 
-	if ( const pxr::VtValue* Value = AssetInfoDict.GetValueAtPath( UnrealIdentifiers::UnrealExportTime ) )
+	if (const pxr::VtValue* Value = AssetInfoDict.GetValueAtPath(UnrealIdentifiers::UnrealExportTime))
 	{
-		if ( Value->IsHolding<std::string>() )
+		if (Value->IsHolding<std::string>())
 		{
-			Result.UnrealExportTime = UsdToUnreal::ConvertString( Value->Get<std::string>() );
+			Result.UnrealExportTime = UsdToUnreal::ConvertString(Value->Get<std::string>());
 		}
 	}
 
-	if ( const pxr::VtValue* Value = AssetInfoDict.GetValueAtPath( UnrealIdentifiers::UnrealEngineVersion ) )
+	if (const pxr::VtValue* Value = AssetInfoDict.GetValueAtPath(UnrealIdentifiers::UnrealEngineVersion))
 	{
-		if ( Value->IsHolding<std::string>() )
+		if (Value->IsHolding<std::string>())
 		{
-			Result.UnrealEngineVersion = UsdToUnreal::ConvertString( Value->Get<std::string>() );
+			Result.UnrealEngineVersion = UsdToUnreal::ConvertString(Value->Get<std::string>());
 		}
 	}
-#endif // USE_USD_SDK
+#endif	  // USE_USD_SDK
 
 	return Result;
 }
@@ -3280,10 +3348,7 @@ bool UsdUtils::ClearNonEssentialPrimMetadata(const pxr::UsdPrim& Prim)
 		const pxr::TfToken& FieldName = MetadataIter->first;
 
 		// We consider those "essential metadata", as removing them will mess with the prim definition
-		const static std::unordered_set<pxr::TfToken, pxr::TfHash> FieldsToSkip = {
-			pxr::SdfFieldKeys->Specifier,
-			pxr::SdfFieldKeys->TypeName
-		};
+		const static std::unordered_set<pxr::TfToken, pxr::TfHash> FieldsToSkip = {pxr::SdfFieldKeys->Specifier, pxr::SdfFieldKeys->TypeName};
 		if (FieldsToSkip.count(FieldName) > 0)
 		{
 			continue;
@@ -3306,7 +3371,7 @@ bool UsdUtils::ClearNonEssentialPrimMetadata(const pxr::UsdPrim& Prim)
 
 	return true;
 }
-#endif // USE_USD_SDK
+#endif	  // USE_USD_SDK
 
 void UsdUtils::CollectSchemaAnalytics(const UE::FUsdStage& Stage, const FString& EventName)
 {
@@ -3447,8 +3512,8 @@ void UsdUtils::CollectSchemaAnalytics(const UE::FUsdStage& Stage, const FString&
 		"Volume",
 		"VolumeLightAPI",
 		"Xform",
-		"XformCommonAPI"};
-
+		"XformCommonAPI"
+	};
 
 	FString ConcatenatedSeenSchemas;
 	for (const FString& SchemaName : SeenSchemas)
