USD: Reformat USDUtilities module.

diff --git a/Engine/Plugins/Importers/USDImporter/Source/USDUtilities/Private/USDErrorUtils.cpp b/Engine/Plugins/Importers/USDImporter/Source/USDUtilities/Private/USDErrorUtils.cpp
--- a/Engine/Plugins/Importers/USDImporter/Source/USDUtilities/Private/USDErrorUtils.cpp
+++ b/Engine/Plugins/Importers/USDImporter/Source/USDUtilities/Private/USDErrorUtils.cpp
@@ -11,10 +11,9 @@
 #if WITH_EDITOR
 #include "IMessageLogListing.h"
 #include "MessageLogModule.h"
-#endif // WITH_EDITOR
+#endif	  // WITH_EDITOR
 
 #if USE_USD_SDK
-
 #include "USDMemory.h"
 #include "USDTypesConversion.h"
 
@@ -23,11 +22,9 @@
 #include "Widgets/Notifications/SNotificationList.h"
 
 #include "USDIncludesStart.h"
-#include "pxr/pxr.h"
 #include "pxr/base/tf/errorMark.h"
 #include "USDIncludesEnd.h"
-
-#endif // #if USE_USD_SDK
+#endif	  // #if USE_USD_SDK
 
 #define LOCTEXT_NAMESPACE "USDErrorUtils"
 
@@ -67,9 +64,7 @@ namespace UsdUtils
 
 		TArray<FString> Errors;
 
-		for (pxr::TfErrorMark::Iterator ErrorIter = Mark.GetBegin();
-			 ErrorIter != Mark.GetEnd();
-			 ++ErrorIter)
+		for (pxr::TfErrorMark::Iterator ErrorIter = Mark.GetBegin(); ErrorIter != Mark.GetEnd(); ++ErrorIter)
 		{
 			std::string ErrorStr = ErrorIter->GetErrorCodeAsString();
 			ErrorStr += ": ";
@@ -92,7 +87,9 @@ namespace UsdUtils
 
 		if (bHadErrors)
 		{
-			FNotificationInfo ErrorToast(!ToastMessage.IsEmpty() ? ToastMessage : LOCTEXT("USDErrorsToast", "Encountered USD errors!\nCheck the Output Log for details."));
+			FNotificationInfo ErrorToast(
+				!ToastMessage.IsEmpty() ? ToastMessage : LOCTEXT("USDErrorsToast", "Encountered USD errors!\nCheck the Output Log for details.")
+			);
 
 			ErrorToast.ExpireDuration = 5.0f;
 			ErrorToast.bFireAndForget = true;
@@ -102,31 +99,33 @@ namespace UsdUtils
 
 		for (const FString& Error : Errors)
 		{
-			FUsdLogManager::LogMessage( FTokenizedMessage::Create( EMessageSeverity::Error, FText::FromString( Error ) ) );
+			FUsdLogManager::LogMessage(FTokenizedMessage::Create(EMessageSeverity::Error, FText::FromString(Error)));
 		}
 
 		return Errors.Num() > 0;
 	}
-}; // namespace UsdUtils
+};		  // namespace UsdUtils
 
-#else // #if USE_USD_SDK
+#else	  // #if USE_USD_SDK
 
 namespace UsdUtils
 {
 	void StartMonitoringErrors()
 	{
 	}
+
 	TArray<FString> GetErrorsAndStopMonitoring()
 	{
 		return {TEXT("USD SDK is not available!")};
 	}
+
 	bool ShowErrorsAndStopMonitoring(const FText& ToastMessage)
 	{
 		return false;
 	}
 }
 
-#endif // #if USE_USD_SDK
+#endif	  // #if USE_USD_SDK
 
 namespace UE
 {
@@ -138,44 +137,44 @@ namespace UE
 		}
 
 #if WITH_EDITOR
-		void FUsdMessageLog::Push( const TSharedRef< FTokenizedMessage >& Message )
+		void FUsdMessageLog::Push(const TSharedRef<FTokenizedMessage>& Message)
 		{
-			TokenizedMessages.Add( Message );
+			TokenizedMessages.Add(Message);
 		}
-#endif // WITH_EDITOR
+#endif	  // WITH_EDITOR
 
 		void FUsdMessageLog::Dump()
 		{
 #if WITH_EDITOR
-			FMessageLogModule& MessageLogModule = FModuleManager::LoadModuleChecked< FMessageLogModule >( "MessageLog" );
-			TSharedRef< IMessageLogListing > LogListing = MessageLogModule.GetLogListing( TEXT("USD") );
+			FMessageLogModule& MessageLogModule = FModuleManager::LoadModuleChecked<FMessageLogModule>("MessageLog");
+			TSharedRef<IMessageLogListing> LogListing = MessageLogModule.GetLogListing(TEXT("USD"));
 
-			if ( TokenizedMessages.Num() > 0 )
+			if (TokenizedMessages.Num() > 0)
 			{
 				// This crashes internally at runtime due to Slate not finding the correct brush...
 				// I think MessageLog is not supposed to be used at runtime, even though it technically can.
 				// At runtime we shouldn't ever push, but just in case
-				LogListing->AddMessages( TokenizedMessages );
-				LogListing->NotifyIfAnyMessages( LOCTEXT("Log", "There were some issues loading the USD Stage."), EMessageSeverity::Info );
+				LogListing->AddMessages(TokenizedMessages);
+				LogListing->NotifyIfAnyMessages(LOCTEXT("Log", "There were some issues loading the USD Stage."), EMessageSeverity::Info);
 				TokenizedMessages.Empty();
 			}
-#endif // WITH_EDITOR
+#endif	  // WITH_EDITOR
 		}
-	}
-}
+	}	  // namespace Internal
+}	 // namespace UE
 
-TOptional< UE::Internal::FUsdMessageLog > FUsdLogManager::MessageLog;
+TOptional<UE::Internal::FUsdMessageLog> FUsdLogManager::MessageLog;
 int32 FUsdLogManager::MessageLogRefCount;
 FCriticalSection FUsdLogManager::MessageLogLock;
 
 TSet<FString> LoggedMessages;
 
-void FUsdLogManager::LogMessage( EMessageSeverity::Type Severity, const FText& Message )
+void FUsdLogManager::LogMessage(EMessageSeverity::Type Severity, const FText& Message)
 {
-	LogMessage( FTokenizedMessage::Create( Severity, Message ) );
+	LogMessage(FTokenizedMessage::Create(Severity, Message));
 }
 
-void FUsdLogManager::LogMessage( const TSharedRef< FTokenizedMessage >& Message )
+void FUsdLogManager::LogMessage(const TSharedRef<FTokenizedMessage>& Message)
 {
 	bool bMessageProcessed = false;
 
@@ -198,43 +197,43 @@ void FUsdLogManager::LogMessage( const TSharedRef< FTokenizedMessage >& Message
 #endif
 	}
 
-	if ( !bMessageProcessed )
+	if (!bMessageProcessed)
 	{
-		if ( Message->GetSeverity() == EMessageSeverity::Error )
+		if (Message->GetSeverity() == EMessageSeverity::Error)
 		{
-			UE_LOG( LogUsd, Error, TEXT("%s"), *(Message->ToText().ToString()) );
+			UE_LOG(LogUsd, Error, TEXT("%s"), *(Message->ToText().ToString()));
 		}
-		else if ( Message->GetSeverity() == EMessageSeverity::Warning || Message->GetSeverity() == EMessageSeverity::PerformanceWarning )
+		else if (Message->GetSeverity() == EMessageSeverity::Warning || Message->GetSeverity() == EMessageSeverity::PerformanceWarning)
 		{
-			UE_LOG( LogUsd, Warning, TEXT("%s"), *(Message->ToText().ToString()) );
+			UE_LOG(LogUsd, Warning, TEXT("%s"), *(Message->ToText().ToString()));
 		}
 		else
 		{
-			UE_LOG( LogUsd, Log, TEXT("%s"), *(Message->ToText().ToString()) );
+			UE_LOG(LogUsd, Log, TEXT("%s"), *(Message->ToText().ToString()));
 		}
 	}
 }
 
 void FUsdLogManager::EnableMessageLog()
 {
-	FScopeLock Lock( &MessageLogLock );
+	FScopeLock Lock(&MessageLogLock);
 
 	LoggedMessages.Reset();
 
-	if ( ++MessageLogRefCount == 1 )
+	if (++MessageLogRefCount == 1)
 	{
 		MessageLog.Emplace();
 		UsdUtils::StartMonitoringErrors();
 	}
 
-	check( MessageLogRefCount < MAX_int32 );
+	check(MessageLogRefCount < MAX_int32);
 }
 
 void FUsdLogManager::DisableMessageLog()
 {
-	FScopeLock Lock( &MessageLogLock );
+	FScopeLock Lock(&MessageLogLock);
 
-	if ( --MessageLogRefCount == 0 )
+	if (--MessageLogRefCount == 0)
 	{
 		UsdUtils::ShowErrorsAndStopMonitoring();
 		MessageLog.Reset();
