USD: Reformat USDUtilities module.

diff --git a/Engine/Plugins/Importers/USDImporter/Source/USDUtilities/Private/USDGeomMeshConversion.cpp b/Engine/Plugins/Importers/USDImporter/Source/USDUtilities/Private/USDGeomMeshConversion.cpp
--- a/Engine/Plugins/Importers/USDImporter/Source/USDUtilities/Private/USDGeomMeshConversion.cpp
+++ b/Engine/Plugins/Importers/USDImporter/Source/USDUtilities/Private/USDGeomMeshConversion.cpp
@@ -36,55 +36,56 @@
 
 #if WITH_EDITOR
 #include "MaterialEditingLibrary.h"
-#endif // WITH_EDITOR
+#endif	  // WITH_EDITOR
 
 #include "USDIncludesStart.h"
-	#include "opensubdiv/far/primvarRefiner.h"
-	#include "opensubdiv/far/topologyRefiner.h"
-	#include "pxr/imaging/pxOsd/meshTopology.h"
-	#include "pxr/imaging/pxOsd/refinerFactory.h"
-	#include "pxr/usd/sdf/layer.h"
-	#include "pxr/usd/sdf/layerUtils.h"
-	#include "pxr/usd/sdf/types.h"
-	#include "pxr/usd/usd/editContext.h"
-	#include "pxr/usd/usd/prim.h"
-	#include "pxr/usd/usd/primRange.h"
-	#include "pxr/usd/usdGeom/capsule.h"
-	#include "pxr/usd/usdGeom/cone.h"
-	#include "pxr/usd/usdGeom/cube.h"
-	#include "pxr/usd/usdGeom/cylinder.h"
-	#include "pxr/usd/usdGeom/gprim.h"
-	#include "pxr/usd/usdGeom/mesh.h"
-	#include "pxr/usd/usdGeom/plane.h"
-	#include "pxr/usd/usdGeom/pointInstancer.h"
-	#include "pxr/usd/usdGeom/primvarsAPI.h"
-	#include "pxr/usd/usdGeom/sphere.h"
-	#include "pxr/usd/usdGeom/subset.h"
-	#include "pxr/usd/usdGeom/tokens.h"
-	#include "pxr/usd/usdPhysics/collisionAPI.h"
-	#include "pxr/usd/usdPhysics/meshCollisionAPI.h"
-	#include "pxr/usd/usdPhysics/tokens.h"
-	#include "pxr/usd/usdShade/material.h"
-	#include "pxr/usd/usdShade/materialBindingAPI.h"
-	#include "pxr/usd/usdShade/tokens.h"
-	#include "pxr/usd/usdSkel/bindingAPI.h"
-	#include "pxr/usd/usdSkel/root.h"
-	#include "pxr/usdImaging/usdImaging/implicitSurfaceMeshUtils.h"
+#include "opensubdiv/far/primvarRefiner.h"
+#include "opensubdiv/far/topologyRefiner.h"
+#include "pxr/imaging/pxOsd/meshTopology.h"
+#include "pxr/imaging/pxOsd/refinerFactory.h"
+#include "pxr/usd/sdf/layerUtils.h"
+#include "pxr/usd/sdf/types.h"
+#include "pxr/usd/usd/editContext.h"
+#include "pxr/usd/usd/prim.h"
+#include "pxr/usd/usd/primRange.h"
+#include "pxr/usd/usdGeom/capsule.h"
+#include "pxr/usd/usdGeom/cone.h"
+#include "pxr/usd/usdGeom/cube.h"
+#include "pxr/usd/usdGeom/cylinder.h"
+#include "pxr/usd/usdGeom/gprim.h"
+#include "pxr/usd/usdGeom/mesh.h"
+#include "pxr/usd/usdGeom/plane.h"
+#include "pxr/usd/usdGeom/pointInstancer.h"
+#include "pxr/usd/usdGeom/primvarsAPI.h"
+#include "pxr/usd/usdGeom/sphere.h"
+#include "pxr/usd/usdGeom/subset.h"
+#include "pxr/usd/usdGeom/tokens.h"
+#include "pxr/usd/usdPhysics/collisionAPI.h"
+#include "pxr/usd/usdPhysics/meshCollisionAPI.h"
+#include "pxr/usd/usdPhysics/tokens.h"
+#include "pxr/usd/usdShade/material.h"
+#include "pxr/usd/usdShade/materialBindingAPI.h"
+#include "pxr/usd/usdShade/tokens.h"
+#include "pxr/usd/usdSkel/bindingAPI.h"
+#include "pxr/usd/usdSkel/root.h"
+#include "pxr/usdImaging/usdImaging/implicitSurfaceMeshUtils.h"
 #include "USDIncludesEnd.h"
 
 #define LOCTEXT_NAMESPACE "USDGeomMeshConversion"
 
 static int32 GMaxInstancesPerPointInstancer = -1;
 static FAutoConsoleVariableRef CVarMaxInstancesPerPointInstancer(
-	TEXT( "USD.MaxInstancesPerPointInstancer" ),
+	TEXT("USD.MaxInstancesPerPointInstancer"),
 	GMaxInstancesPerPointInstancer,
-	TEXT( "We will only parse up to this many instances from any point instancer when reading from USD to UE. Set this to -1 to disable this limit." ) );
+	TEXT("We will only parse up to this many instances from any point instancer when reading from USD to UE. Set this to -1 to disable this limit.")
+);
 
 static bool GIgnoreNormalsWhenSubdividing = true;
 static FAutoConsoleVariableRef CVarIgnoreNormalsWhenSubdividing(
 	TEXT("USD.Subdiv.IgnoreNormalsWhenSubdividing"),
 	GIgnoreNormalsWhenSubdividing,
-	TEXT("This being true means that whenever we subdivide a mesh we fully ignore the authored normals (if any) and recompute new normals. If this is false we will try interpolating the normals during subdivision like a regular primvar")
+	TEXT("This being true means that whenever we subdivide a mesh we fully ignore the authored normals (if any) and recompute new normals. If this "
+		 "is false we will try interpolating the normals during subdivision like a regular primvar")
 );
 
 static const FString MaxUsdSubdivLevelCvarName = TEXT("USD.Subdiv.MaxSubdivLevel");
@@ -99,33 +100,33 @@ const FName MeshAttribute::VertexInstance::Velocity("Velocity");
 
 namespace UE::UsdGeomMeshConversion::Private
 {
-	static const FString DisplayColorID = TEXT( "!DisplayColor" );
+	static const FString DisplayColorID = TEXT("!DisplayColor");
 
 	// Dimensions used when generating Capsule meshes
 	static const float DefaultCapsuleMeshRadius = 0.25;
 	static const float DefaultCapsuleMeshHeight = 0.50;
 
-	int32 GetPrimValueIndex( const pxr::TfToken& InterpType, const int32 VertexIndex, const int32 VertexInstanceIndex, const int32 PolygonIndex )
+	int32 GetPrimValueIndex(const pxr::TfToken& InterpType, const int32 VertexIndex, const int32 VertexInstanceIndex, const int32 PolygonIndex)
 	{
-		if ( InterpType == pxr::UsdGeomTokens->vertex )
+		if (InterpType == pxr::UsdGeomTokens->vertex)
 		{
 			return VertexIndex;
 		}
-		else if ( InterpType == pxr::UsdGeomTokens->varying )
+		else if (InterpType == pxr::UsdGeomTokens->varying)
 		{
 			return VertexIndex;
 		}
-		else if ( InterpType == pxr::UsdGeomTokens->faceVarying )
+		else if (InterpType == pxr::UsdGeomTokens->faceVarying)
 		{
 			return VertexInstanceIndex;
 		}
-		else if ( InterpType == pxr::UsdGeomTokens->uniform )
+		else if (InterpType == pxr::UsdGeomTokens->uniform)
 		{
 			return PolygonIndex;
 		}
-		else /* if ( InterpType == pxr::UsdGeomTokens->constant ) */
+		else			 /* if ( InterpType == pxr::UsdGeomTokens->constant ) */
 		{
-			return 0; // return index 0 for constant or any other unsupported cases
+			return 0;	 // return index 0 for constant or any other unsupported cases
 		}
 	}
 
@@ -243,12 +244,7 @@ namespace UE::UsdGeomMeshConversion::Private
 				// The USD cylinder is aligned to the z axis by default
 				pxr::VtArray<pxr::GfVec3f> Points = pxr::UsdImagingGetUnitCylinderMeshPoints();
 
-				pxr::GfMatrix4d ZToXAxis{
-					0.0, 1.0, 0.0, 0.0,
-					0.0, 0.0, 1.0, 0.0,
-					1.0, 0.0, 0.0, 0.0,
-					0.0, 0.0, 0.0, 1.0
-				};
+				pxr::GfMatrix4d ZToXAxis{0.0, 1.0, 0.0, 0.0, 0.0, 0.0, 1.0, 0.0, 1.0, 0.0, 0.0, 0.0, 0.0, 0.0, 0.0, 1.0};
 
 				for (pxr::GfVec3f& Point : Points)
 				{
@@ -266,12 +262,7 @@ namespace UE::UsdGeomMeshConversion::Private
 				// The USD cylinder is aligned to the z axis by default
 				pxr::VtArray<pxr::GfVec3f> Points = pxr::UsdImagingGetUnitCylinderMeshPoints();
 
-				pxr::GfMatrix4d ZToYAxis{
-					0.0, 0.0, 1.0, 0.0,
-					1.0, 0.0, 0.0, 0.0,
-					0.0, 1.0, 0.0, 0.0,
-					0.0, 0.0, 0.0, 1.0
-				};
+				pxr::GfMatrix4d ZToYAxis{0.0, 0.0, 1.0, 0.0, 1.0, 0.0, 0.0, 0.0, 0.0, 1.0, 0.0, 0.0, 0.0, 0.0, 0.0, 1.0};
 
 				for (pxr::GfVec3f& Point : Points)
 				{
@@ -299,12 +290,7 @@ namespace UE::UsdGeomMeshConversion::Private
 				// The USD cone is aligned to the z axis by default
 				pxr::VtArray<pxr::GfVec3f> Points = pxr::UsdImagingGetUnitConeMeshPoints();
 
-				pxr::GfMatrix4d ZToXAxis{
-					0.0, 1.0, 0.0, 0.0,
-					0.0, 0.0, 1.0, 0.0,
-					1.0, 0.0, 0.0, 0.0,
-					0.0, 0.0, 0.0, 1.0
-				};
+				pxr::GfMatrix4d ZToXAxis{0.0, 1.0, 0.0, 0.0, 0.0, 0.0, 1.0, 0.0, 1.0, 0.0, 0.0, 0.0, 0.0, 0.0, 0.0, 1.0};
 
 				for (pxr::GfVec3f& Point : Points)
 				{
@@ -322,12 +308,7 @@ namespace UE::UsdGeomMeshConversion::Private
 				// The USD cone is aligned to the z axis by default
 				pxr::VtArray<pxr::GfVec3f> Points = pxr::UsdImagingGetUnitConeMeshPoints();
 
-				pxr::GfMatrix4d ZToYAxis{
-					0.0, 0.0, 1.0, 0.0,
-					1.0, 0.0, 0.0, 0.0,
-					0.0, 1.0, 0.0, 0.0,
-					0.0, 0.0, 0.0, 1.0
-				};
+				pxr::GfMatrix4d ZToYAxis{0.0, 0.0, 1.0, 0.0, 1.0, 0.0, 0.0, 0.0, 0.0, 1.0, 0.0, 0.0, 0.0, 0.0, 0.0, 1.0};
 
 				for (pxr::GfVec3f& Point : Points)
 				{
@@ -346,43 +327,43 @@ namespace UE::UsdGeomMeshConversion::Private
 		return {};
 	}
 
-	int32 GetLODIndexFromName( const std::string& Name )
+	int32 GetLODIndexFromName(const std::string& Name)
 	{
 		const std::string LODString = UnrealIdentifiers::LOD.GetString();
 
 		// True if Name does not start with "LOD"
-		if ( Name.rfind( LODString, 0 ) != 0 )
+		if (Name.rfind(LODString, 0) != 0)
 		{
 			return INDEX_NONE;
 		}
 
 		// After LODString there should be only numbers
-		if ( Name.find_first_not_of( "0123456789", LODString.size() ) != std::string::npos )
+		if (Name.find_first_not_of("0123456789", LODString.size()) != std::string::npos)
 		{
 			return INDEX_NONE;
 		}
 
 		const int Base = 10;
 		char** EndPtr = nullptr;
-		return std::strtol( Name.c_str() + LODString.size(), EndPtr, Base );
+		return std::strtol(Name.c_str() + LODString.size(), EndPtr, Base);
 	}
 
 	void ConvertStaticMeshLOD(
 		int32 LODIndex,
 		const FStaticMeshLODResources& LODRenderMesh,
 		pxr::UsdGeomMesh& UsdMesh,
-		const pxr::VtArray< std::string >& MaterialAssignments,
+		const pxr::VtArray<std::string>& MaterialAssignments,
 		const pxr::UsdTimeCode TimeCode,
 		pxr::UsdPrim MaterialPrim
 	)
 	{
 		pxr::UsdPrim MeshPrim = UsdMesh.GetPrim();
 		pxr::UsdStageRefPtr Stage = MeshPrim.GetStage();
-		if ( !Stage )
+		if (!Stage)
 		{
 			return;
 		}
-		const FUsdStageInfo StageInfo{ Stage };
+		const FUsdStageInfo StageInfo{Stage};
 
 		// Vertices
 		{
@@ -391,18 +372,18 @@ namespace UE::UsdGeomMeshConversion::Private
 			// Points
 			{
 				pxr::UsdAttribute Points = UsdMesh.CreatePointsAttr();
-				if ( Points )
+				if (Points)
 				{
-					pxr::VtArray< pxr::GfVec3f > PointsArray;
-					PointsArray.reserve( VertexCount );
+					pxr::VtArray<pxr::GfVec3f> PointsArray;
+					PointsArray.reserve(VertexCount);
 
-					for ( int32 VertexIndex = 0; VertexIndex < VertexCount; ++VertexIndex )
+					for (int32 VertexIndex = 0; VertexIndex < VertexCount; ++VertexIndex)
 					{
-						FVector VertexPosition = (FVector)LODRenderMesh.VertexBuffers.PositionVertexBuffer.VertexPosition( VertexIndex );
-						PointsArray.push_back( UnrealToUsd::ConvertVectorFloat( StageInfo, VertexPosition ) );
+						FVector VertexPosition = (FVector)LODRenderMesh.VertexBuffers.PositionVertexBuffer.VertexPosition(VertexIndex);
+						PointsArray.push_back(UnrealToUsd::ConvertVectorFloat(StageInfo, VertexPosition));
 					}
 
-					Points.Set( PointsArray, TimeCode );
+					Points.Set(PointsArray, TimeCode);
 				}
 			}
 
@@ -412,24 +393,24 @@ namespace UE::UsdGeomMeshConversion::Private
 				// actually subdivide (like usdview) will just discard authored normals and fully recompute them
 				// on-demand in case they have a valid subdivision scheme (which is the default state).
 				// Reference: https://graphics.pixar.com/usd/release/api/class_usd_geom_mesh.html#UsdGeom_Mesh_Normals
-				if ( pxr::UsdAttribute SubdivisionAttr = UsdMesh.CreateSubdivisionSchemeAttr() )
+				if (pxr::UsdAttribute SubdivisionAttr = UsdMesh.CreateSubdivisionSchemeAttr())
 				{
-					ensure( SubdivisionAttr.Set( pxr::UsdGeomTokens->none ) );
+					ensure(SubdivisionAttr.Set(pxr::UsdGeomTokens->none));
 				}
 
 				pxr::UsdAttribute NormalsAttribute = UsdMesh.CreateNormalsAttr();
-				if ( NormalsAttribute )
+				if (NormalsAttribute)
 				{
-					pxr::VtArray< pxr::GfVec3f > Normals;
-					Normals.reserve( VertexCount );
+					pxr::VtArray<pxr::GfVec3f> Normals;
+					Normals.reserve(VertexCount);
 
-					for ( int32 VertexIndex = 0; VertexIndex < VertexCount; ++VertexIndex )
+					for (int32 VertexIndex = 0; VertexIndex < VertexCount; ++VertexIndex)
 					{
-						FVector VertexNormal = (FVector4)LODRenderMesh.VertexBuffers.StaticMeshVertexBuffer.VertexTangentZ( VertexIndex );
-						Normals.push_back( UnrealToUsd::ConvertVectorFloat( StageInfo, VertexNormal ) );
+						FVector VertexNormal = (FVector4)LODRenderMesh.VertexBuffers.StaticMeshVertexBuffer.VertexTangentZ(VertexIndex);
+						Normals.push_back(UnrealToUsd::ConvertVectorFloat(StageInfo, VertexNormal));
 					}
 
-					NormalsAttribute.Set( Normals, TimeCode );
+					NormalsAttribute.Set(Normals, TimeCode);
 				}
 			}
 
@@ -437,54 +418,60 @@ namespace UE::UsdGeomMeshConversion::Private
 			{
 				const int32 TexCoordSourceCount = LODRenderMesh.VertexBuffers.StaticMeshVertexBuffer.GetNumTexCoords();
 
-				for ( int32 TexCoordSourceIndex = 0; TexCoordSourceIndex < TexCoordSourceCount; ++TexCoordSourceIndex )
+				for (int32 TexCoordSourceIndex = 0; TexCoordSourceIndex < TexCoordSourceCount; ++TexCoordSourceIndex)
 				{
-					pxr::TfToken UsdUVSetName = UsdUtils::GetUVSetName( TexCoordSourceIndex ).Get();
+					pxr::TfToken UsdUVSetName = UsdUtils::GetUVSetName(TexCoordSourceIndex).Get();
 
-					pxr::UsdGeomPrimvar PrimvarST = pxr::UsdGeomPrimvarsAPI(MeshPrim).CreatePrimvar( UsdUVSetName, pxr::SdfValueTypeNames->TexCoord2fArray, pxr::UsdGeomTokens->vertex );
+					pxr::UsdGeomPrimvar PrimvarST = pxr::UsdGeomPrimvarsAPI(MeshPrim).CreatePrimvar(
+						UsdUVSetName,
+						pxr::SdfValueTypeNames->TexCoord2fArray,
+						pxr::UsdGeomTokens->vertex
+					);
 
-					if ( PrimvarST )
+					if (PrimvarST)
 					{
 						pxr::VtVec2fArray UVs;
 
-						for ( int32 VertexIndex = 0; VertexIndex < VertexCount; ++VertexIndex )
+						for (int32 VertexIndex = 0; VertexIndex < VertexCount; ++VertexIndex)
 						{
-							FVector2D TexCoord = FVector2D(LODRenderMesh.VertexBuffers.StaticMeshVertexBuffer.GetVertexUV( VertexIndex, TexCoordSourceIndex ));
-							TexCoord[ 1 ] = 1.f - TexCoord[ 1 ];
+							FVector2D TexCoord = FVector2D(
+								LODRenderMesh.VertexBuffers.StaticMeshVertexBuffer.GetVertexUV(VertexIndex, TexCoordSourceIndex)
+							);
+							TexCoord[1] = 1.f - TexCoord[1];
 
-							UVs.push_back( UnrealToUsd::ConvertVectorFloat( TexCoord ) );
+							UVs.push_back(UnrealToUsd::ConvertVectorFloat(TexCoord));
 						}
 
-						PrimvarST.Set( UVs, TimeCode );
+						PrimvarST.Set(UVs, TimeCode);
 					}
 				}
 			}
 
 			// Vertex colors
-			if ( LODRenderMesh.bHasColorVertexData )
+			if (LODRenderMesh.bHasColorVertexData)
 			{
-				pxr::UsdGeomPrimvar DisplayColorPrimvar = UsdMesh.CreateDisplayColorPrimvar( pxr::UsdGeomTokens->vertex );
-				pxr::UsdGeomPrimvar DisplayOpacityPrimvar = UsdMesh.CreateDisplayOpacityPrimvar( pxr::UsdGeomTokens->vertex );
+				pxr::UsdGeomPrimvar DisplayColorPrimvar = UsdMesh.CreateDisplayColorPrimvar(pxr::UsdGeomTokens->vertex);
+				pxr::UsdGeomPrimvar DisplayOpacityPrimvar = UsdMesh.CreateDisplayOpacityPrimvar(pxr::UsdGeomTokens->vertex);
 
-				if ( DisplayColorPrimvar )
+				if (DisplayColorPrimvar)
 				{
-					pxr::VtArray< pxr::GfVec3f > DisplayColors;
-					DisplayColors.reserve( VertexCount );
+					pxr::VtArray<pxr::GfVec3f> DisplayColors;
+					DisplayColors.reserve(VertexCount);
 
-					pxr::VtArray< float > DisplayOpacities;
-					DisplayOpacities.reserve( VertexCount );
+					pxr::VtArray<float> DisplayOpacities;
+					DisplayOpacities.reserve(VertexCount);
 
-					for ( int32 VertexIndex = 0; VertexIndex < VertexCount; ++VertexIndex )
+					for (int32 VertexIndex = 0; VertexIndex < VertexCount; ++VertexIndex)
 					{
-						const FColor& VertexColor = LODRenderMesh.VertexBuffers.ColorVertexBuffer.VertexColor( VertexIndex );
+						const FColor& VertexColor = LODRenderMesh.VertexBuffers.ColorVertexBuffer.VertexColor(VertexIndex);
 
-						pxr::GfVec4f Color = UnrealToUsd::ConvertColor( VertexColor );
-						DisplayColors.push_back( pxr::GfVec3f( Color[ 0 ], Color[ 1 ], Color[ 2 ] ) );
-						DisplayOpacities.push_back( Color[ 3 ] );
+						pxr::GfVec4f Color = UnrealToUsd::ConvertColor(VertexColor);
+						DisplayColors.push_back(pxr::GfVec3f(Color[0], Color[1], Color[2]));
+						DisplayOpacities.push_back(Color[3]);
 					}
 
-					DisplayColorPrimvar.Set( DisplayColors, TimeCode );
-					DisplayOpacityPrimvar.Set( DisplayOpacities, TimeCode );
+					DisplayColorPrimvar.Set(DisplayColors, TimeCode);
+					DisplayOpacityPrimvar.Set(DisplayOpacities, TimeCode);
 				}
 			}
 		}
@@ -497,17 +484,17 @@ namespace UE::UsdGeomMeshConversion::Private
 			{
 				pxr::UsdAttribute FaceCountsAttribute = UsdMesh.CreateFaceVertexCountsAttr();
 
-				if ( FaceCountsAttribute )
+				if (FaceCountsAttribute)
 				{
-					pxr::VtArray< int > FaceVertexCounts;
-					FaceVertexCounts.reserve( FaceCount );
+					pxr::VtArray<int> FaceVertexCounts;
+					FaceVertexCounts.reserve(FaceCount);
 
-					for ( int32 FaceIndex = 0; FaceIndex < FaceCount; ++FaceIndex )
+					for (int32 FaceIndex = 0; FaceIndex < FaceCount; ++FaceIndex)
 					{
-						FaceVertexCounts.push_back( 3 );
+						FaceVertexCounts.push_back(3);
 					}
 
-					FaceCountsAttribute.Set( FaceVertexCounts, TimeCode );
+					FaceCountsAttribute.Set(FaceVertexCounts, TimeCode);
 				}
 			}
 
@@ -515,20 +502,20 @@ namespace UE::UsdGeomMeshConversion::Private
 			{
 				pxr::UsdAttribute FaceVertexIndicesAttribute = UsdMesh.GetFaceVertexIndicesAttr();
 
-				if ( FaceVertexIndicesAttribute )
+				if (FaceVertexIndicesAttribute)
 				{
 					FIndexArrayView Indices = LODRenderMesh.IndexBuffer.GetArrayView();
-					ensure( Indices.Num() == FaceCount * 3 );
+					ensure(Indices.Num() == FaceCount * 3);
 
-					pxr::VtArray< int > FaceVertexIndices;
-					FaceVertexIndices.reserve( FaceCount * 3 );
+					pxr::VtArray<int> FaceVertexIndices;
+					FaceVertexIndices.reserve(FaceCount * 3);
 
-					for ( int32 Index = 0; Index < FaceCount * 3; ++Index )
+					for (int32 Index = 0; Index < FaceCount * 3; ++Index)
 					{
-						FaceVertexIndices.push_back( Indices[ Index ] );
+						FaceVertexIndices.push_back(Indices[Index]);
 					}
 
-					FaceVertexIndicesAttribute.Set( FaceVertexIndices, TimeCode );
+					FaceVertexIndicesAttribute.Set(FaceVertexIndices, TimeCode);
 				}
 			}
 		}
@@ -537,105 +524,116 @@ namespace UE::UsdGeomMeshConversion::Private
 		{
 			bool bHasUEMaterialAssignements = false;
 
-			pxr::VtArray< std::string > UnrealMaterialsForLOD;
-			for ( const FStaticMeshSection& Section : LODRenderMesh.Sections )
+			pxr::VtArray<std::string> UnrealMaterialsForLOD;
+			for (const FStaticMeshSection& Section : LODRenderMesh.Sections)
 			{
-				if ( Section.MaterialIndex >= 0 && Section.MaterialIndex < MaterialAssignments.size() )
+				if (Section.MaterialIndex >= 0 && Section.MaterialIndex < MaterialAssignments.size())
 				{
-					UnrealMaterialsForLOD.push_back( MaterialAssignments[ Section.MaterialIndex ] );
+					UnrealMaterialsForLOD.push_back(MaterialAssignments[Section.MaterialIndex]);
 					bHasUEMaterialAssignements = true;
 				}
 				else
 				{
 					// Keep unrealMaterials with the same number of elements as our MaterialIndices expect
-					UnrealMaterialsForLOD.push_back( "" );
+					UnrealMaterialsForLOD.push_back("");
 				}
 			}
 
 			// This LOD has a single material assignment, just add an unrealMaterials attribute to the mesh prim
-			if ( bHasUEMaterialAssignements && UnrealMaterialsForLOD.size() == 1 )
+			if (bHasUEMaterialAssignements && UnrealMaterialsForLOD.size() == 1)
 			{
-				if ( pxr::UsdAttribute UEMaterialsAttribute = MaterialPrim.CreateAttribute( UnrealIdentifiers::MaterialAssignment, pxr::SdfValueTypeNames->String ) )
+				if (pxr::UsdAttribute UEMaterialsAttribute = MaterialPrim.CreateAttribute(
+						UnrealIdentifiers::MaterialAssignment,
+						pxr::SdfValueTypeNames->String
+					))
 				{
-					UEMaterialsAttribute.Set( UnrealMaterialsForLOD[ 0 ] );
+					UEMaterialsAttribute.Set(UnrealMaterialsForLOD[0]);
 				}
 			}
 			// Multiple material assignments to the same LOD (and so the same mesh prim). Need to create a GeomSubset for each UE mesh section
-			else if ( UnrealMaterialsForLOD.size() > 1 )
+			else if (UnrealMaterialsForLOD.size() > 1)
 			{
 				// Need to fetch all triangles of a section, and add their indices
-				for ( int32 SectionIndex = 0; SectionIndex < LODRenderMesh.Sections.Num(); ++SectionIndex )
+				for (int32 SectionIndex = 0; SectionIndex < LODRenderMesh.Sections.Num(); ++SectionIndex)
 				{
-					const FStaticMeshSection& Section = LODRenderMesh.Sections[ SectionIndex ];
+					const FStaticMeshSection& Section = LODRenderMesh.Sections[SectionIndex];
 
 					// Note that we will continue on even if we have no material assignment, so as to satisfy the "partition" family condition (below)
 					std::string SectionMaterial;
-					if ( Section.MaterialIndex >= 0 && Section.MaterialIndex < MaterialAssignments.size() )
+					if (Section.MaterialIndex >= 0 && Section.MaterialIndex < MaterialAssignments.size())
 					{
-						SectionMaterial = MaterialAssignments[ Section.MaterialIndex ];
+						SectionMaterial = MaterialAssignments[Section.MaterialIndex];
 					}
 
 					pxr::UsdPrim GeomSubsetPrim = Stage->DefinePrim(
-						MeshPrim.GetPath().AppendPath( pxr::SdfPath( "Section" + std::to_string( SectionIndex ) ) ),
-						UnrealToUsd::ConvertToken( TEXT( "GeomSubset" ) ).Get()
+						MeshPrim.GetPath().AppendPath(pxr::SdfPath("Section" + std::to_string(SectionIndex))),
+						UnrealToUsd::ConvertToken(TEXT("GeomSubset")).Get()
 					);
 
 					// MaterialPrim may be in another stage, so we may need another GeomSubset there
 					pxr::UsdPrim MaterialGeomSubsetPrim = GeomSubsetPrim;
-					if ( MaterialPrim.GetStage() != MeshPrim.GetStage() )
+					if (MaterialPrim.GetStage() != MeshPrim.GetStage())
 					{
 						MaterialGeomSubsetPrim = MaterialPrim.GetStage()->OverridePrim(
-							MaterialPrim.GetPath().AppendPath( pxr::SdfPath( "Section" + std::to_string( SectionIndex ) ) )
+							MaterialPrim.GetPath().AppendPath(pxr::SdfPath("Section" + std::to_string(SectionIndex)))
 						);
 					}
 
-					pxr::UsdGeomSubset GeomSubsetSchema{ GeomSubsetPrim };
+					pxr::UsdGeomSubset GeomSubsetSchema{GeomSubsetPrim};
 
 					// Element type attribute
 					pxr::UsdAttribute ElementTypeAttr = GeomSubsetSchema.CreateElementTypeAttr();
-					ElementTypeAttr.Set( pxr::UsdGeomTokens->face, TimeCode );
+					ElementTypeAttr.Set(pxr::UsdGeomTokens->face, TimeCode);
 
 					// Indices attribute
 					const uint32 TriangleCount = Section.NumTriangles;
-					const uint32 FirstTriangleIndex = Section.FirstIndex / 3; // FirstIndex is the first *vertex* instance index
+					const uint32 FirstTriangleIndex = Section.FirstIndex / 3;	 // FirstIndex is the first *vertex* instance index
 					FIndexArrayView VertexInstances = LODRenderMesh.IndexBuffer.GetArrayView();
 					pxr::VtArray<int> IndicesAttrValue;
-					for ( uint32 TriangleIndex = FirstTriangleIndex; TriangleIndex - FirstTriangleIndex < TriangleCount; ++TriangleIndex )
+					for (uint32 TriangleIndex = FirstTriangleIndex; TriangleIndex - FirstTriangleIndex < TriangleCount; ++TriangleIndex)
 					{
 						// Note that we add VertexInstances in sequence to the usda file for the faceVertexInstances attribute, which
 						// also constitutes our triangle order
-						IndicesAttrValue.push_back( static_cast< int >( TriangleIndex ) );
+						IndicesAttrValue.push_back(static_cast<int>(TriangleIndex));
 					}
 
 					pxr::UsdAttribute IndicesAttr = GeomSubsetSchema.CreateIndicesAttr();
-					IndicesAttr.Set( IndicesAttrValue, TimeCode );
+					IndicesAttr.Set(IndicesAttrValue, TimeCode);
 
 					// Family name attribute
 					pxr::UsdAttribute FamilyNameAttr = GeomSubsetSchema.CreateFamilyNameAttr();
-					FamilyNameAttr.Set( pxr::UsdShadeTokens->materialBind, TimeCode );
+					FamilyNameAttr.Set(pxr::UsdShadeTokens->materialBind, TimeCode);
 
 					// Family type
-					pxr::UsdGeomSubset::SetFamilyType( UsdMesh, pxr::UsdShadeTokens->materialBind, pxr::UsdGeomTokens->partition );
+					pxr::UsdGeomSubset::SetFamilyType(UsdMesh, pxr::UsdShadeTokens->materialBind, pxr::UsdGeomTokens->partition);
 
 					// unrealMaterial attribute
-					if ( pxr::UsdAttribute UEMaterialsAttribute = MaterialGeomSubsetPrim.CreateAttribute( UnrealIdentifiers::MaterialAssignment, pxr::SdfValueTypeNames->String ) )
+					if (pxr::UsdAttribute UEMaterialsAttribute = MaterialGeomSubsetPrim.CreateAttribute(
+							UnrealIdentifiers::MaterialAssignment,
+							pxr::SdfValueTypeNames->String
+						))
 					{
-						UEMaterialsAttribute.Set( UnrealMaterialsForLOD[ SectionIndex ] );
+						UEMaterialsAttribute.Set(UnrealMaterialsForLOD[SectionIndex]);
 					}
 				}
 			}
 		}
 	}
 
-	bool ConvertMeshDescription( const FMeshDescription& MeshDescription, pxr::UsdGeomMesh& UsdMesh, const FMatrix& AdditionalTransform, const pxr::UsdTimeCode TimeCode )
+	bool ConvertMeshDescription(
+		const FMeshDescription& MeshDescription,
+		pxr::UsdGeomMesh& UsdMesh,
+		const FMatrix& AdditionalTransform,
+		const pxr::UsdTimeCode TimeCode
+	)
 	{
 		pxr::UsdPrim MeshPrim = UsdMesh.GetPrim();
 		pxr::UsdStageRefPtr Stage = MeshPrim.GetStage();
-		if ( !Stage )
+		if (!Stage)
 		{
 			return false;
 		}
-		const FUsdStageInfo StageInfo{ Stage };
+		const FUsdStageInfo StageInfo{Stage};
 
 		FStaticMeshConstAttributes Attributes(MeshDescription);
 		TVertexAttributesConstRef<FVector3f> VertexPositions = Attributes.GetVertexPositions();
@@ -650,36 +648,36 @@ namespace UE::UsdGeomMeshConversion::Private
 
 		// Points
 		{
-			if ( pxr::UsdAttribute Points = UsdMesh.CreatePointsAttr() )
+			if (pxr::UsdAttribute Points = UsdMesh.CreatePointsAttr())
 			{
-				pxr::VtArray< pxr::GfVec3f > PointsArray;
-				PointsArray.reserve( VertexCount );
+				pxr::VtArray<pxr::GfVec3f> PointsArray;
+				PointsArray.reserve(VertexCount);
 
-				for ( const FVertexID VertexID : MeshDescription.Vertices().GetElementIDs() )
+				for (const FVertexID VertexID : MeshDescription.Vertices().GetElementIDs())
 				{
-					FVector UEPosition = AdditionalTransform.TransformPosition( (FVector)VertexPositions[ VertexID ] );
-					PointsArray.push_back( UnrealToUsd::ConvertVectorFloat( StageInfo, UEPosition ) );
+					FVector UEPosition = AdditionalTransform.TransformPosition((FVector)VertexPositions[VertexID]);
+					PointsArray.push_back(UnrealToUsd::ConvertVectorFloat(StageInfo, UEPosition));
 				}
 
-				Points.Set( PointsArray, TimeCode );
+				Points.Set(PointsArray, TimeCode);
 			}
 		}
 
 		// Normals
 		{
-			if ( pxr::UsdAttribute NormalsAttribute = UsdMesh.CreateNormalsAttr() )
+			if (pxr::UsdAttribute NormalsAttribute = UsdMesh.CreateNormalsAttr())
 			{
-				pxr::VtArray< pxr::GfVec3f > Normals;
-				Normals.reserve( VertexInstanceCount );
+				pxr::VtArray<pxr::GfVec3f> Normals;
+				Normals.reserve(VertexInstanceCount);
 
-				for ( const FVertexInstanceID InstanceID : MeshDescription.VertexInstances().GetElementIDs() )
+				for (const FVertexInstanceID InstanceID : MeshDescription.VertexInstances().GetElementIDs())
 				{
-					FVector UENormal = (FVector)VertexInstanceNormals[ InstanceID ].GetSafeNormal();
-					Normals.push_back( UnrealToUsd::ConvertVectorFloat( StageInfo, UENormal ) );
+					FVector UENormal = (FVector)VertexInstanceNormals[InstanceID].GetSafeNormal();
+					Normals.push_back(UnrealToUsd::ConvertVectorFloat(StageInfo, UENormal));
 				}
 
-				NormalsAttribute.Set( Normals, TimeCode );
-				UsdMesh.SetNormalsInterpolation( pxr::UsdGeomTokens->faceVarying );
+				NormalsAttribute.Set(Normals, TimeCode);
+				UsdMesh.SetNormalsInterpolation(pxr::UsdGeomTokens->faceVarying);
 			}
 		}
 
@@ -687,50 +685,51 @@ namespace UE::UsdGeomMeshConversion::Private
 		{
 			int32 NumUVs = VertexInstanceUVs.GetNumChannels();
 
-			for ( int32 UVIndex = 0; UVIndex < NumUVs; ++UVIndex )
+			for (int32 UVIndex = 0; UVIndex < NumUVs; ++UVIndex)
 			{
-				pxr::TfToken UsdUVSetName = UsdUtils::GetUVSetName( UVIndex ).Get();
+				pxr::TfToken UsdUVSetName = UsdUtils::GetUVSetName(UVIndex).Get();
 
-				pxr::UsdGeomPrimvar PrimvarST = pxr::UsdGeomPrimvarsAPI(MeshPrim).CreatePrimvar( UsdUVSetName, pxr::SdfValueTypeNames->TexCoord2fArray, pxr::UsdGeomTokens->vertex );
-				if ( PrimvarST )
+				pxr::UsdGeomPrimvar PrimvarST = pxr::UsdGeomPrimvarsAPI(MeshPrim)
+													.CreatePrimvar(UsdUVSetName, pxr::SdfValueTypeNames->TexCoord2fArray, pxr::UsdGeomTokens->vertex);
+				if (PrimvarST)
 				{
 					pxr::VtVec2fArray UVs;
 
-					for ( const FVertexInstanceID InstanceID : MeshDescription.VertexInstances().GetElementIDs() )
+					for (const FVertexInstanceID InstanceID : MeshDescription.VertexInstances().GetElementIDs())
 					{
-						FVector2D UV = FVector2D(VertexInstanceUVs.Get( InstanceID, UVIndex ));
-						UV[ 1 ] = 1.f - UV[ 1 ];
-						UVs.push_back( UnrealToUsd::ConvertVectorFloat( UV ) );
+						FVector2D UV = FVector2D(VertexInstanceUVs.Get(InstanceID, UVIndex));
+						UV[1] = 1.f - UV[1];
+						UVs.push_back(UnrealToUsd::ConvertVectorFloat(UV));
 					}
 
-					PrimvarST.Set( UVs, TimeCode );
-					PrimvarST.SetInterpolation( pxr::UsdGeomTokens->faceVarying );
+					PrimvarST.Set(UVs, TimeCode);
+					PrimvarST.SetInterpolation(pxr::UsdGeomTokens->faceVarying);
 				}
 			}
 		}
 
 		// Vertex colors
-		if ( VertexInstanceColors.GetNumElements() > 0 )
+		if (VertexInstanceColors.GetNumElements() > 0)
 		{
-			pxr::UsdGeomPrimvar DisplayColorPrimvar = UsdMesh.CreateDisplayColorPrimvar( pxr::UsdGeomTokens->faceVarying );
-			pxr::UsdGeomPrimvar DisplayOpacityPrimvar = UsdMesh.CreateDisplayOpacityPrimvar( pxr::UsdGeomTokens->faceVarying );
-			if ( DisplayColorPrimvar && DisplayOpacityPrimvar )
+			pxr::UsdGeomPrimvar DisplayColorPrimvar = UsdMesh.CreateDisplayColorPrimvar(pxr::UsdGeomTokens->faceVarying);
+			pxr::UsdGeomPrimvar DisplayOpacityPrimvar = UsdMesh.CreateDisplayOpacityPrimvar(pxr::UsdGeomTokens->faceVarying);
+			if (DisplayColorPrimvar && DisplayOpacityPrimvar)
 			{
-				pxr::VtArray< pxr::GfVec3f > DisplayColors;
-				DisplayColors.reserve( VertexInstanceCount );
+				pxr::VtArray<pxr::GfVec3f> DisplayColors;
+				DisplayColors.reserve(VertexInstanceCount);
 
-				pxr::VtArray< float > DisplayOpacities;
-				DisplayOpacities.reserve( VertexInstanceCount );
+				pxr::VtArray<float> DisplayOpacities;
+				DisplayOpacities.reserve(VertexInstanceCount);
 
-				for ( const FVertexInstanceID InstanceID : MeshDescription.VertexInstances().GetElementIDs() )
+				for (const FVertexInstanceID InstanceID : MeshDescription.VertexInstances().GetElementIDs())
 				{
-					pxr::GfVec4f Color = UnrealToUsd::ConvertColor( FLinearColor( VertexInstanceColors[ InstanceID ] ) );
-					DisplayColors.push_back( pxr::GfVec3f( Color[ 0 ], Color[ 1 ], Color[ 2 ] ) );
-					DisplayOpacities.push_back( Color[ 3 ] );
+					pxr::GfVec4f Color = UnrealToUsd::ConvertColor(FLinearColor(VertexInstanceColors[InstanceID]));
+					DisplayColors.push_back(pxr::GfVec3f(Color[0], Color[1], Color[2]));
+					DisplayOpacities.push_back(Color[3]);
 				}
 
-				DisplayColorPrimvar.Set( DisplayColors, TimeCode );
-				DisplayOpacityPrimvar.Set( DisplayOpacities, TimeCode );
+				DisplayColorPrimvar.Set(DisplayColors, TimeCode);
+				DisplayOpacityPrimvar.Set(DisplayOpacities, TimeCode);
 			}
 		}
 
@@ -739,25 +738,25 @@ namespace UE::UsdGeomMeshConversion::Private
 			pxr::UsdAttribute FaceCountsAttribute = UsdMesh.CreateFaceVertexCountsAttr();
 			pxr::UsdAttribute FaceVertexIndicesAttribute = UsdMesh.GetFaceVertexIndicesAttr();
 
-			pxr::VtArray< int > FaceVertexCounts;
-			FaceVertexCounts.reserve( FaceCount );
+			pxr::VtArray<int> FaceVertexCounts;
+			FaceVertexCounts.reserve(FaceCount);
 
-			pxr::VtArray< int > FaceVertexIndices;
+			pxr::VtArray<int> FaceVertexIndices;
 
-			for ( FPolygonID PolygonID : MeshDescription.Polygons().GetElementIDs() )
+			for (FPolygonID PolygonID : MeshDescription.Polygons().GetElementIDs())
 			{
-				const TArray<FVertexInstanceID>& PolygonVertexInstances = MeshDescription.GetPolygonVertexInstances( PolygonID );
-				FaceVertexCounts.push_back( static_cast< int >( PolygonVertexInstances.Num() ) );
+				const TArray<FVertexInstanceID>& PolygonVertexInstances = MeshDescription.GetPolygonVertexInstances(PolygonID);
+				FaceVertexCounts.push_back(static_cast<int>(PolygonVertexInstances.Num()));
 
-				for ( FVertexInstanceID VertexInstanceID : PolygonVertexInstances )
+				for (FVertexInstanceID VertexInstanceID : PolygonVertexInstances)
 				{
-					int32 VertexIndex = MeshDescription.GetVertexInstanceVertex( VertexInstanceID ).GetValue();
-					FaceVertexIndices.push_back( static_cast< int >( VertexIndex ) );
+					int32 VertexIndex = MeshDescription.GetVertexInstanceVertex(VertexInstanceID).GetValue();
+					FaceVertexIndices.push_back(static_cast<int>(VertexIndex));
 				}
 			}
 
-			FaceCountsAttribute.Set( FaceVertexCounts, TimeCode );
-			FaceVertexIndicesAttribute.Set( FaceVertexIndices, TimeCode );
+			FaceCountsAttribute.Set(FaceVertexCounts, TimeCode);
+			FaceVertexIndicesAttribute.Set(FaceVertexIndices, TimeCode);
 		}
 
 		return true;
@@ -773,14 +772,14 @@ namespace UE::UsdGeomMeshConversion::Private
 	)
 	{
 		// Ignore meshes from disabled purposes
-		if ( !EnumHasAllFlags( Options.PurposesToLoad, IUsdPrim::GetPurpose( Prim ) ) )
+		if (!EnumHasAllFlags(Options.PurposesToLoad, IUsdPrim::GetPurpose(Prim)))
 		{
 			return true;
 		}
 
 		FTransform ChildTransform = Options.AdditionalTransform;
 
-		if ( !bIsFirstPrim )
+		if (!bIsFirstPrim)
 		{
 			// Ignore invisible child meshes.
 			//
@@ -801,22 +800,22 @@ namespace UE::UsdGeomMeshConversion::Private
 			//
 			// Note that if we were to ever manually set any of these back to visible again via the editor, the visibility changes are
 			// now resyncs and we'll reparse this entire asset, which will give us the chance to add them back to the collapsed mesh.
-			if ( pxr::UsdGeomImageable UsdGeomImageable = pxr::UsdGeomImageable( Prim ) )
+			if (pxr::UsdGeomImageable UsdGeomImageable = pxr::UsdGeomImageable(Prim))
 			{
-				if ( pxr::UsdAttribute VisibilityAttr = UsdGeomImageable.GetVisibilityAttr() )
+				if (pxr::UsdAttribute VisibilityAttr = UsdGeomImageable.GetVisibilityAttr())
 				{
 					pxr::TfToken VisibilityToken;
-					if ( VisibilityAttr.Get( &VisibilityToken ) && VisibilityToken == pxr::UsdGeomTokens->invisible )
+					if (VisibilityAttr.Get(&VisibilityToken) && VisibilityToken == pxr::UsdGeomTokens->invisible)
 					{
 						return true;
 					}
 				}
 			}
 
-			if ( pxr::UsdGeomXformable Xformable = pxr::UsdGeomXformable( Prim ) )
+			if (pxr::UsdGeomXformable Xformable = pxr::UsdGeomXformable(Prim))
 			{
 				FTransform LocalChildTransform;
-				UsdToUnreal::ConvertXformable( Prim.GetStage(), Xformable, LocalChildTransform, Options.TimeCode.GetValue() );
+				UsdToUnreal::ConvertXformable(Prim.GetStage(), Xformable, LocalChildTransform, Options.TimeCode.GetValue());
 
 				ChildTransform = LocalChildTransform * Options.AdditionalTransform;
 			}
@@ -829,15 +828,15 @@ namespace UE::UsdGeomMeshConversion::Private
 		// children afterwards, its fine to overwrite Options.AdditionalTransform. We do have to put it back to our
 		// original value after we're done though, as calls to sibling prims that would run after this call would need
 		// the original AdditionalTransform in place. The alternative is to copy the entire options object...
-		TGuardValue<FTransform> Guard{ Options.AdditionalTransform, ChildTransform };
+		TGuardValue<FTransform> Guard{Options.AdditionalTransform, ChildTransform};
 
-		if ( pxr::UsdGeomMesh Mesh = pxr::UsdGeomMesh( Prim ) )
+		if (pxr::UsdGeomMesh Mesh = pxr::UsdGeomMesh(Prim))
 		{
 			// We never want to glob up *skinned* meshes inside SkelRoots, as those presumably will be handled by the
 			// UsdSkelSkeletonTranslator and the skeletal data code path.
 			if (!GIsEditor || !bIsInsideSkelRoot || !Prim.HasAPI<pxr::UsdSkelBindingAPI>())
 			{
-				bSuccess = UsdToUnreal::ConvertGeomMesh( Mesh, OutMeshDescription, OutMaterialAssignments, Options );
+				bSuccess = UsdToUnreal::ConvertGeomMesh(Mesh, OutMeshDescription, OutMaterialAssignments, Options);
 			}
 		}
 		// Check for a cube/capsule/etc. ConvertGeomPrimitive will internally check for all specific types
@@ -845,24 +844,19 @@ namespace UE::UsdGeomMeshConversion::Private
 		{
 			bSuccess = UsdToUnreal::ConvertGeomPrimitive(Prim, OutMeshDescription, OutMaterialAssignments, Options);
 		}
-		else if ( pxr::UsdGeomPointInstancer PointInstancer = pxr::UsdGeomPointInstancer{ Prim } )
+		else if (pxr::UsdGeomPointInstancer PointInstancer = pxr::UsdGeomPointInstancer{Prim})
 		{
-			bSuccess = UsdToUnreal::ConvertPointInstancerToMesh(
-				PointInstancer,
-				OutMeshDescription,
-				OutMaterialAssignments,
-				Options
-			);
+			bSuccess = UsdToUnreal::ConvertPointInstancerToMesh(PointInstancer, OutMeshDescription, OutMaterialAssignments, Options);
 
 			// We never want to step into point instancers when fetching prims for drawing
 			bTraverseChildren = false;
 		}
 
-		if ( bTraverseChildren )
+		if (bTraverseChildren)
 		{
-			for ( const pxr::UsdPrim& ChildPrim : Prim.GetFilteredChildren( pxr::UsdTraverseInstanceProxies() ) )
+			for (const pxr::UsdPrim& ChildPrim : Prim.GetFilteredChildren(pxr::UsdTraverseInstanceProxies()))
 			{
-				if ( !bSuccess )
+				if (!bSuccess)
 				{
 					break;
 				}
@@ -935,13 +929,7 @@ namespace UE::UsdGeomMeshConversion::Private
 				}
 
 				const bool bChildIsFirstPrim = false;
-				RecursivelyCollectPrimvars(
-					PrototypeUsdPrim,
-					Options,
-					InOutAllPrimvars,
-					InOutPreferredPrimvars,
-					bChildIsFirstPrim
-				);
+				RecursivelyCollectPrimvars(PrototypeUsdPrim, Options, InOutAllPrimvars, InOutPreferredPrimvars, bChildIsFirstPrim);
 			}
 
 			// We never want to step into point instancers when fetching prims for drawing
@@ -949,10 +937,7 @@ namespace UE::UsdGeomMeshConversion::Private
 		}
 		if (Prim)
 		{
-			TArray<TUsdStore<pxr::UsdGeomPrimvar>> Primvars = UsdUtils::GetUVSetPrimvars(
-				Prim,
-				TNumericLimits<int32>::Max()
-			);
+			TArray<TUsdStore<pxr::UsdGeomPrimvar>> Primvars = UsdUtils::GetUVSetPrimvars(Prim, TNumericLimits<int32>::Max());
 
 			for (const TUsdStore<pxr::UsdGeomPrimvar>& Primvar : Primvars)
 			{
@@ -974,13 +959,7 @@ namespace UE::UsdGeomMeshConversion::Private
 			for (const pxr::UsdPrim& ChildPrim : Prim.GetFilteredChildren(pxr::UsdTraverseInstanceProxies()))
 			{
 				const bool bChildIsFirstPrim = false;
-				RecursivelyCollectPrimvars(
-					ChildPrim,
-					Options,
-					InOutAllPrimvars,
-					InOutPreferredPrimvars,
-					bChildIsFirstPrim
-				);
+				RecursivelyCollectPrimvars(ChildPrim, Options, InOutAllPrimvars, InOutPreferredPrimvars, bChildIsFirstPrim);
 			}
 		}
 	}
@@ -998,18 +977,9 @@ namespace UE::UsdGeomMeshConversion::Private
 		TSet<FString> AllPrimvars;
 		TSet<FString> PreferredPrimvars;
 
-		RecursivelyCollectPrimvars(
-			RootPrim,
-			Options,
-			AllPrimvars,
-			PreferredPrimvars,
-			bIsFirstPrim
-		);
+		RecursivelyCollectPrimvars(RootPrim, Options, AllPrimvars, PreferredPrimvars, bIsFirstPrim);
 
-		return UsdUtils::CombinePrimvarsIntoUVSets(
-			AllPrimvars,
-			PreferredPrimvars
-		);
+		return UsdUtils::CombinePrimvarsIntoUVSets(AllPrimvars, PreferredPrimvars);
 	}
 
 	// Unconverted, raw USD mesh data to convert into a MeshDescription
@@ -1281,7 +1251,6 @@ namespace UE::UsdGeomMeshConversion::Private
 		return true;
 	}
 
-
 	// OpenSubdiv expects the data elements of its buffers to implement a simple interface,
 	// so here we wrap the datatypes we'll be interpolating with that interface
 	struct FSubdivVec2f
@@ -1297,6 +1266,7 @@ namespace UE::UsdGeomMeshConversion::Private
 		{
 			Data += Src * Weight;
 		}
+
 		void AddWithWeight(const FSubdivVec2f& Src, float Weight)
 		{
 			Data += Src.Data * Weight;
@@ -1316,6 +1286,7 @@ namespace UE::UsdGeomMeshConversion::Private
 		{
 			Data += Src * Weight;
 		}
+
 		void AddWithWeight(const FSubdivVec3f& Src, float Weight)
 		{
 			Data += Src.Data * Weight;
@@ -1335,6 +1306,7 @@ namespace UE::UsdGeomMeshConversion::Private
 		{
 			Data += Src * Weight;
 		}
+
 		void AddWithWeight(const FSubdivInt& Src, float Weight)
 		{
 			Data += Src.Data * Weight;
@@ -1354,6 +1326,7 @@ namespace UE::UsdGeomMeshConversion::Private
 		{
 			Data += Src * Weight;
 		}
+
 		void AddWithWeight(const FSubdivFloat& Src, float Weight)
 		{
 			Data += Src.Data * Weight;
@@ -1428,11 +1401,7 @@ namespace UE::UsdGeomMeshConversion::Private
 		Swap(FlattenedValues, SharedValuesArray);
 	};
 
-	bool SubdivideMeshData(
-		const pxr::UsdPrim& Prim,
-		const UsdToUnreal::FUsdMeshConversionOptions& Options,
-		FUsdMeshData& InOutMeshData
-	)
+	bool SubdivideMeshData(const pxr::UsdPrim& Prim, const UsdToUnreal::FUsdMeshConversionOptions& Options, FUsdMeshData& InOutMeshData)
 	{
 		// References:
 		// - USD's HdSt_OsdTopologyComputation::Resolve
@@ -1463,7 +1432,8 @@ namespace UE::UsdGeomMeshConversion::Private
 			UE_LOG(
 				LogUsd,
 				Log,
-				TEXT("Cancelling out of subdividing mesh '%s' due to target subdivision level being %d after clamping (it needs to be at least 1 for a round of subdivision)"),
+				TEXT("Cancelling out of subdividing mesh '%s' due to target subdivision level being %d after clamping (it needs to be at least 1 for "
+					 "a round of subdivision)"),
 				*InOutMeshData.SourcePrimPath,
 				TargetSubdivLevel
 			);
@@ -1538,9 +1508,10 @@ namespace UE::UsdGeomMeshConversion::Private
 			// "Normals should not be authored on a subdivision mesh, since subdivision algorithms define their own normals.
 			// They should only be authored for polygonal meshes (subdivisionScheme = "none")."
 			// There is no free normal computation to be had from OpenSubdiv subdivision algoriths as far as I can tell however.
-			// We'd have to compute them manually as in [REDACTED_PROFILE_URL]/OpenSubdiv/blob/release/tutorials/far/tutorial_2_3/far_tutorial_2_3.cpp
-			// If that is the case, we may as well just ignore normals here and let RepairNormalsAndTangents fix it,
-			// since it will need to run it to compute tangents anyway
+			// We'd have to compute them manually as in
+			// [REDACTED_PROFILE_URL]/OpenSubdiv/blob/release/tutorials/far/tutorial_2_3/far_tutorial_2_3.cpp If that is the case,
+			// we may as well just ignore normals here and let RepairNormalsAndTangents fix it, since it will need to run it to compute tangents
+			// anyway
 
 			InOutMeshData.Normals = {};
 			InOutMeshData.NormalIndices = {};
@@ -1662,7 +1633,8 @@ namespace UE::UsdGeomMeshConversion::Private
 			InOutMeshData.CreaseLengths,
 			InOutMeshData.CreaseSharpnesses,
 			InOutMeshData.CornerIndices,
-			InOutMeshData.CornerSharpnesses};
+			InOutMeshData.CornerSharpnesses
+		};
 
 		pxr::PxOsdMeshTopology Topology{
 			InOutMeshData.SubdivScheme,
@@ -1670,7 +1642,8 @@ namespace UE::UsdGeomMeshConversion::Private
 			InOutMeshData.FaceVertexCounts,
 			InOutMeshData.FaceVertexIndices,
 			InOutMeshData.HoleIndices,
-			SubdivTags};
+			SubdivTags
+		};
 
 		std::vector<pxr::VtArray<int>> FaceVaryingTopologies;
 		FaceVaryingTopologies.resize(FaceVaryingChannelCounter);
@@ -1729,12 +1702,13 @@ namespace UE::UsdGeomMeshConversion::Private
 		pxr::VtArray<pxr::GfVec3f> TempVelocities;
 		pxr::VtArray<pxr::GfVec3f> TempDisplayColors;
 		pxr::VtArray<float> TempDisplayOpacities;
-		TArray<int32> TempMaterialIndices; // Using a TArray saves us a memcpy when outputting results
+		TArray<int32> TempMaterialIndices;	  // Using a TArray saves us a memcpy when outputting results
 		TArray<pxr::VtArray<pxr::GfVec2f>> TempUVSets;
 
 		// Resize the target buffers to be large enough to hold all refinements *simultaneously* (one next to the other).
 		// This is great because we can just read/write to the same buffer as we iteratively refine
-		TFunction<size_t(pxr::TfToken, int32)> GetTotalNumElements = [&TopologyRefiner](pxr::TfToken InterpolationType, int32 FaceVaryingChannel) -> size_t
+		TFunction<size_t(pxr::TfToken, int32)> GetTotalNumElements =
+			[&TopologyRefiner](pxr::TfToken InterpolationType, int32 FaceVaryingChannel) -> size_t
 		{
 			// The "GetNumXTotal()" functions also include space for the source data as well.
 			// In our case we'll keep the source data on the actual source arrays so we don't have to
@@ -1771,7 +1745,7 @@ namespace UE::UsdGeomMeshConversion::Private
 		TempVelocities.resize(GetTotalNumElements(InOutMeshData.VelocityInterpolation, VelocitiesFaceVaryingChannel));
 		TempDisplayColors.resize(GetTotalNumElements(InOutMeshData.DisplayColorInterpolation, DisplayColorsFaceVaryingChannel));
 		TempDisplayOpacities.resize(GetTotalNumElements(InOutMeshData.DisplayOpacityInterpolation, DisplayOpacitiesFaceVaryingChannel));
-		TempMaterialIndices.SetNum(GetTotalNumElements(MaterialIndicesInterpolation, /*FaceVaryingChannel*/ 0));	 // Always 'uniform'
+		TempMaterialIndices.SetNum(GetTotalNumElements(MaterialIndicesInterpolation, /*FaceVaryingChannel*/ 0));	// Always 'uniform'
 		TempUVSets.SetNum(NumUVSets);
 		for (int32 UVSetIndex = 0; UVSetIndex < NumUVSets; ++UVSetIndex)
 		{
@@ -1869,16 +1843,35 @@ namespace UE::UsdGeomMeshConversion::Private
 		};
 
 		// Actually refine all of our attributes/primvars
-		// Inspired by [REDACTED_PROFILE_URL]/OpenSubdiv/blob/7d0ab5530feef693ac0a920585b5c663b80773b3/tutorials/far/tutorial_2_2/far_tutorial_2_2.cpp#L293
+		// Inspired by
+		// [REDACTED_PROFILE_URL]/OpenSubdiv/blob/7d0ab5530feef693ac0a920585b5c663b80773b3/tutorials/far/tutorial_2_2/far_tutorial_2_2.cpp#L293
 		// but avoiding the initial copy from the source data arrays
 		for (int32 CurrentLevel = 1; CurrentLevel <= TargetSubdivLevel; ++CurrentLevel)
 		{
 			InterpolateAttribute(SrcPointsPtr, DstPointsPtr, InOutMeshData.PointInterpolation, CurrentLevel, PointsFaceVaryingChannel);
 			InterpolateAttribute(SrcNormalsPtr, DstNormalsPtr, InOutMeshData.NormalInterpolation, CurrentLevel, NormalsFaceVaryingChannel);
 			InterpolateAttribute(SrcVelocitiesPtr, DstVelocitiesPtr, InOutMeshData.VelocityInterpolation, CurrentLevel, VelocitiesFaceVaryingChannel);
-			InterpolateAttribute(SrcDisplayColorsPtr, DstDisplayColorsPtr, InOutMeshData.DisplayColorInterpolation, CurrentLevel, DisplayColorsFaceVaryingChannel);
-			InterpolateAttribute(SrcDisplayOpacitiesPtr, DstDisplayOpacitiesPtr, InOutMeshData.DisplayOpacityInterpolation, CurrentLevel, DisplayOpacitiesFaceVaryingChannel);
-			InterpolateAttribute(SrcMaterialIndicesPtr, DstMaterialIndicesPtr, MaterialIndicesInterpolation, CurrentLevel, /*FaceVaryingChannel*/ 0); // Always 'uniform'
+			InterpolateAttribute(
+				SrcDisplayColorsPtr,
+				DstDisplayColorsPtr,
+				InOutMeshData.DisplayColorInterpolation,
+				CurrentLevel,
+				DisplayColorsFaceVaryingChannel
+			);
+			InterpolateAttribute(
+				SrcDisplayOpacitiesPtr,
+				DstDisplayOpacitiesPtr,
+				InOutMeshData.DisplayOpacityInterpolation,
+				CurrentLevel,
+				DisplayOpacitiesFaceVaryingChannel
+			);
+			InterpolateAttribute(
+				SrcMaterialIndicesPtr,
+				DstMaterialIndicesPtr,
+				MaterialIndicesInterpolation,
+				CurrentLevel,
+				/*FaceVaryingChannel*/ 0
+			);	  // Always 'uniform'
 			for (int32 UVSetIndex = 0; UVSetIndex < NumUVSets; ++UVSetIndex)
 			{
 				InterpolateAttribute(
@@ -1947,7 +1940,7 @@ namespace UE::UsdGeomMeshConversion::Private
 			OpenSubdiv::Far::ConstIndexArray Face = FinalLevel.GetFaceVertices(FaceIndex);
 			TempFaceVertexCounts[FaceIndex] = Face.size();
 
-			for(int32 FaceVertexIndex = 0; FaceVertexIndex < Face.size(); ++FaceVertexIndex)
+			for (int32 FaceVertexIndex = 0; FaceVertexIndex < Face.size(); ++FaceVertexIndex)
 			{
 				TempFaceVertexIndices.push_back(Face[FaceVertexIndex]);
 			}
@@ -2276,10 +2269,10 @@ namespace UE::UsdGeomMeshConversion::Private
 			}
 
 			TVertexInstanceAttributesRef<FVector3f> MeshDescriptionNormals = StaticMeshAttributes.GetVertexInstanceNormals();
-			TVertexInstanceAttributesRef<FVector3f>
-				MeshDescriptionVelocities = OutMeshDescription.VertexInstanceAttributes().GetAttributesRef<FVector3f>(
-					MeshAttribute::VertexInstance::Velocity
-				);
+			TVertexInstanceAttributesRef<FVector3f> MeshDescriptionVelocities = OutMeshDescription.VertexInstanceAttributes()
+																					.GetAttributesRef<FVector3f>(
+																						MeshAttribute::VertexInstance::Velocity
+																					);
 
 			OutMeshDescription.ReserveNewVertexInstances(InMeshData.FaceVertexCounts.size() * 3);
 			OutMeshDescription.ReserveNewPolygons(InMeshData.FaceVertexCounts.size());
@@ -2472,7 +2465,7 @@ namespace UE::UsdGeomMeshConversion::Private
 
 		return true;
 	}
-}
+}	 // namespace UE::UsdGeomMeshConversion::Private
 namespace UsdGeomMeshImpl = UE::UsdGeomMeshConversion::Private;
 
 namespace UsdToUnreal
@@ -2480,12 +2473,12 @@ namespace UsdToUnreal
 	const FUsdMeshConversionOptions FUsdMeshConversionOptions::DefaultOptions;
 
 	FUsdMeshConversionOptions::FUsdMeshConversionOptions()
-		: AdditionalTransform( FTransform::Identity )
-		, PurposesToLoad( EUsdPurpose::Render )
-		, RenderContext( pxr::UsdShadeTokens->universalRenderContext )
-		, MaterialPurpose( pxr::UsdShadeTokens->allPurpose )
-		, TimeCode( pxr::UsdTimeCode::EarliestTime() )
-		, bMergeIdenticalMaterialSlots( true )
+		: AdditionalTransform(FTransform::Identity)
+		, PurposesToLoad(EUsdPurpose::Render)
+		, RenderContext(pxr::UsdShadeTokens->universalRenderContext)
+		, MaterialPurpose(pxr::UsdShadeTokens->allPurpose)
+		, TimeCode(pxr::UsdTimeCode::EarliestTime())
+		, bMergeIdenticalMaterialSlots(true)
 		, SubdivisionLevel(0)
 	{
 	}
@@ -2537,9 +2530,9 @@ bool UsdToUnreal::ConvertPointInstancerToMesh(
 	const FUsdMeshConversionOptions& Options
 )
 {
-	TRACE_CPUPROFILER_EVENT_SCOPE( UsdToUnreal::ConvertPointInstancerToMesh );
+	TRACE_CPUPROFILER_EVENT_SCOPE(UsdToUnreal::ConvertPointInstancerToMesh);
 
-	if ( !PointInstancer )
+	if (!PointInstancer)
 	{
 		return false;
 	}
@@ -2565,7 +2558,7 @@ bool UsdToUnreal::ConvertPointInstancerToMesh(
 			const pxr::UsdRelationship& Prototypes = PointInstancer.GetPrototypesRel();
 
 			pxr::SdfPathVector UsdPrototypePaths;
-			if ( !Prototypes.GetTargets( &UsdPrototypePaths ) )
+			if (!Prototypes.GetTargets(&UsdPrototypePaths))
 			{
 				return false;
 			}
@@ -2578,32 +2571,35 @@ bool UsdToUnreal::ConvertPointInstancerToMesh(
 
 			Allocs.Reset();
 			PrototypePaths.Reserve(NumPrototypes);
-			for(const pxr::SdfPath& UsdPath : UsdPrototypePaths)
+			for (const pxr::SdfPath& UsdPath : UsdPrototypePaths)
 			{
 				PrototypePaths.Add(UE::FSdfPath{UsdPath});
 			}
 			Allocs.Emplace();
 		}
 
-		PrototypeMeshDescriptions.SetNum( NumPrototypes );
-		PrototypeMaterialAssignments.SetNum( NumPrototypes );
-		PrototypePolygonGroupRemapping.SetNum( NumPrototypes );
+		PrototypeMeshDescriptions.SetNum(NumPrototypes);
+		PrototypeMaterialAssignments.SetNum(NumPrototypes);
+		PrototypePolygonGroupRemapping.SetNum(NumPrototypes);
 
 		// Our AdditionalTransform should be applied after even the instance transforms, we don't want to apply it
 		// directly to our prototypes
 		FUsdMeshConversionOptions OptionsCopy = Options;
 		OptionsCopy.AdditionalTransform = FTransform::Identity;
 
-		for ( uint32 PrototypeIndex = 0; PrototypeIndex < NumPrototypes; ++PrototypeIndex )
+		for (uint32 PrototypeIndex = 0; PrototypeIndex < NumPrototypes; ++PrototypeIndex)
 		{
-			const UE::FSdfPath& PrototypePath = PrototypePaths[ PrototypeIndex ];
+			const UE::FSdfPath& PrototypePath = PrototypePaths[PrototypeIndex];
 
-			UE::FUsdPrim PrototypeUsdPrim = Stage.GetPrimAtPath( PrototypePath );
-			if ( !PrototypeUsdPrim )
+			UE::FUsdPrim PrototypeUsdPrim = Stage.GetPrimAtPath(PrototypePath);
+			if (!PrototypeUsdPrim)
 			{
-				UE_LOG( LogUsd, Warning, TEXT( "Failed to find prototype '%s' for PointInstancer '%s' within ConvertPointInstancerToMesh" ),
-					*UsdToUnreal::ConvertPath( PrototypePath ),
-					*UsdToUnreal::ConvertPath( PointInstancer.GetPrim().GetPrimPath() )
+				UE_LOG(
+					LogUsd,
+					Warning,
+					TEXT("Failed to find prototype '%s' for PointInstancer '%s' within ConvertPointInstancerToMesh"),
+					*UsdToUnreal::ConvertPath(PrototypePath),
+					*UsdToUnreal::ConvertPath(PointInstancer.GetPrim().GetPrimPath())
 				);
 				continue;
 			}
@@ -2611,8 +2607,8 @@ bool UsdToUnreal::ConvertPointInstancerToMesh(
 			const bool bSkipRootPrimTransformAndVisibility = false;
 			ConvertGeomMeshHierarchy(
 				PrototypeUsdPrim,
-				PrototypeMeshDescriptions[ PrototypeIndex ],
-				PrototypeMaterialAssignments[ PrototypeIndex ],
+				PrototypeMeshDescriptions[PrototypeIndex],
+				PrototypeMaterialAssignments[PrototypeIndex],
 				OptionsCopy,
 				bSkipRootPrimTransformAndVisibility
 			);
@@ -2623,43 +2619,45 @@ bool UsdToUnreal::ConvertPointInstancerToMesh(
 	// Sets up PrototypePolygonGroupRemapping so that our new faces are remapped from the prototype's mesh description polygon groups to the
 	// combined mesh description's polygon groups when AppendMeshDescription is called.
 	// Note: We always setup our mesh description polygon groups in the same order as the material assignment slots, so this is not so complicated
-	for ( uint32 PrototypeIndex = 0; PrototypeIndex < NumPrototypes; ++PrototypeIndex )
+	for (uint32 PrototypeIndex = 0; PrototypeIndex < NumPrototypes; ++PrototypeIndex)
 	{
-		UsdUtils::FUsdPrimMaterialAssignmentInfo& PrototypeMaterialAssignment = PrototypeMaterialAssignments[ PrototypeIndex ];
-		TMap<FPolygonGroupID, FPolygonGroupID>& PrototypeToCombinedMeshPolygonGroupMap = PrototypePolygonGroupRemapping[ PrototypeIndex ];
+		UsdUtils::FUsdPrimMaterialAssignmentInfo& PrototypeMaterialAssignment = PrototypeMaterialAssignments[PrototypeIndex];
+		TMap<FPolygonGroupID, FPolygonGroupID>& PrototypeToCombinedMeshPolygonGroupMap = PrototypePolygonGroupRemapping[PrototypeIndex];
 
-		if ( Options.bMergeIdenticalMaterialSlots )
+		if (Options.bMergeIdenticalMaterialSlots)
 		{
 			// Build a map of our existing slots since we can hash the entire slot, and our incoming mesh may have an arbitrary number of new slots
-			TMap< UsdUtils::FUsdPrimMaterialSlot, int32 > CombinedMaterialSlotsToIndex;
-			for ( int32 Index = 0; Index < OutMaterialAssignments.Slots.Num(); ++Index )
+			TMap<UsdUtils::FUsdPrimMaterialSlot, int32> CombinedMaterialSlotsToIndex;
+			for (int32 Index = 0; Index < OutMaterialAssignments.Slots.Num(); ++Index)
 			{
-				const UsdUtils::FUsdPrimMaterialSlot& Slot = OutMaterialAssignments.Slots[ Index ];
-				CombinedMaterialSlotsToIndex.Add( Slot, Index );
+				const UsdUtils::FUsdPrimMaterialSlot& Slot = OutMaterialAssignments.Slots[Index];
+				CombinedMaterialSlotsToIndex.Add(Slot, Index);
 			}
 
-			for ( int32 PrototypeMaterialSlotIndex = 0; PrototypeMaterialSlotIndex < PrototypeMaterialAssignment.Slots.Num(); ++PrototypeMaterialSlotIndex )
+			for (int32 PrototypeMaterialSlotIndex = 0; PrototypeMaterialSlotIndex < PrototypeMaterialAssignment.Slots.Num();
+				 ++PrototypeMaterialSlotIndex)
 			{
-				const UsdUtils::FUsdPrimMaterialSlot& LocalSlot = PrototypeMaterialAssignment.Slots[ PrototypeMaterialSlotIndex ];
-				if ( int32* ExistingCombinedIndex = CombinedMaterialSlotsToIndex.Find( LocalSlot ) )
+				const UsdUtils::FUsdPrimMaterialSlot& LocalSlot = PrototypeMaterialAssignment.Slots[PrototypeMaterialSlotIndex];
+				if (int32* ExistingCombinedIndex = CombinedMaterialSlotsToIndex.Find(LocalSlot))
 				{
-					PrototypeToCombinedMeshPolygonGroupMap.Add( PrototypeMaterialSlotIndex, *ExistingCombinedIndex );
+					PrototypeToCombinedMeshPolygonGroupMap.Add(PrototypeMaterialSlotIndex, *ExistingCombinedIndex);
 				}
 				else
 				{
-					OutMaterialAssignments.Slots.Add( LocalSlot );
-					PrototypeToCombinedMeshPolygonGroupMap.Add( PrototypeMaterialSlotIndex, OutMaterialAssignments.Slots.Num() - 1 );
+					OutMaterialAssignments.Slots.Add(LocalSlot);
+					PrototypeToCombinedMeshPolygonGroupMap.Add(PrototypeMaterialSlotIndex, OutMaterialAssignments.Slots.Num() - 1);
 				}
 			}
 		}
 		else
 		{
 			const int32 NumExistingMaterialSlots = OutMaterialAssignments.Slots.Num();
-			OutMaterialAssignments.Slots.Append( PrototypeMaterialAssignment.Slots );
+			OutMaterialAssignments.Slots.Append(PrototypeMaterialAssignment.Slots);
 
-			for ( int32 PrototypeMaterialSlotIndex = 0; PrototypeMaterialSlotIndex < PrototypeMaterialAssignment.Slots.Num(); ++PrototypeMaterialSlotIndex )
+			for (int32 PrototypeMaterialSlotIndex = 0; PrototypeMaterialSlotIndex < PrototypeMaterialAssignment.Slots.Num();
+				 ++PrototypeMaterialSlotIndex)
 			{
-				PrototypeToCombinedMeshPolygonGroupMap.Add( PrototypeMaterialSlotIndex, NumExistingMaterialSlots + PrototypeMaterialSlotIndex );
+				PrototypeToCombinedMeshPolygonGroupMap.Add(PrototypeMaterialSlotIndex, NumExistingMaterialSlots + PrototypeMaterialSlotIndex);
 			}
 		}
 	}
@@ -2667,34 +2665,38 @@ bool UsdToUnreal::ConvertPointInstancerToMesh(
 	// Make sure we have the polygon groups we expect. Appending the mesh descriptions will not create new polygon groups if we're using a
 	// PolygonGroupsDelegate, which we will
 	const int32 NumExistingPolygonGroups = OutMeshDescription.PolygonGroups().Num();
-	for ( int32 NumMissingPolygonGroups = OutMaterialAssignments.Slots.Num() - NumExistingPolygonGroups; NumMissingPolygonGroups > 0; --NumMissingPolygonGroups )
+	for (int32 NumMissingPolygonGroups = OutMaterialAssignments.Slots.Num() - NumExistingPolygonGroups; NumMissingPolygonGroups > 0;
+		 --NumMissingPolygonGroups)
 	{
 		OutMeshDescription.CreatePolygonGroup();
 	}
 
 	// Double-check our target mesh description has the attributes we need
-	FStaticMeshAttributes StaticMeshAttributes( OutMeshDescription );
+	FStaticMeshAttributes StaticMeshAttributes(OutMeshDescription);
 	StaticMeshAttributes.Register();
 
 	// Append mesh descriptions
-	FUsdStageInfo StageInfo{ PointInstancer.GetPrim().GetStage() };
-	for ( uint32 PrototypeIndex = 0; PrototypeIndex < NumPrototypes; ++PrototypeIndex )
+	FUsdStageInfo StageInfo{PointInstancer.GetPrim().GetStage()};
+	for (uint32 PrototypeIndex = 0; PrototypeIndex < NumPrototypes; ++PrototypeIndex)
 	{
-		const FMeshDescription& PrototypeMeshDescription = PrototypeMeshDescriptions[ PrototypeIndex ];
+		const FMeshDescription& PrototypeMeshDescription = PrototypeMeshDescriptions[PrototypeIndex];
 
 		// We may generate some empty meshes in case a prototype is invisible, for example
-		if ( PrototypeMeshDescription.IsEmpty() )
+		if (PrototypeMeshDescription.IsEmpty())
 		{
 			continue;
 		}
 
 		TArray<FTransform> InstanceTransforms;
-		bool bSuccess = UsdUtils::GetPointInstancerTransforms( StageInfo, PointInstancer, PrototypeIndex, Options.TimeCode, InstanceTransforms );
-		if ( !bSuccess )
+		bool bSuccess = UsdUtils::GetPointInstancerTransforms(StageInfo, PointInstancer, PrototypeIndex, Options.TimeCode, InstanceTransforms);
+		if (!bSuccess)
 		{
-			UE_LOG( LogUsd, Error, TEXT( "Failed to retrieve point instancer transforms for prototype index '%u' of point instancer '%s'" ),
+			UE_LOG(
+				LogUsd,
+				Error,
+				TEXT("Failed to retrieve point instancer transforms for prototype index '%u' of point instancer '%s'"),
 				PrototypeIndex,
-				*UsdToUnreal::ConvertPath( PointInstancer.GetPrim().GetPrimPath() )
+				*UsdToUnreal::ConvertPath(PointInstancer.GetPrim().GetPrimPath())
 			);
 
 			continue;
@@ -2702,25 +2704,26 @@ bool UsdToUnreal::ConvertPointInstancerToMesh(
 
 		const int32 NumInstances = InstanceTransforms.Num();
 
-		OutMeshDescription.ReserveNewVertices( PrototypeMeshDescription.Vertices().Num() * NumInstances );
-		OutMeshDescription.ReserveNewVertexInstances( PrototypeMeshDescription.VertexInstances().Num() * NumInstances );
-		OutMeshDescription.ReserveNewEdges( PrototypeMeshDescription.Edges().Num() * NumInstances );
-		OutMeshDescription.ReserveNewTriangles( PrototypeMeshDescription.Triangles().Num() * NumInstances );
+		OutMeshDescription.ReserveNewVertices(PrototypeMeshDescription.Vertices().Num() * NumInstances);
+		OutMeshDescription.ReserveNewVertexInstances(PrototypeMeshDescription.VertexInstances().Num() * NumInstances);
+		OutMeshDescription.ReserveNewEdges(PrototypeMeshDescription.Edges().Num() * NumInstances);
+		OutMeshDescription.ReserveNewTriangles(PrototypeMeshDescription.Triangles().Num() * NumInstances);
 
 		FStaticMeshOperations::FAppendSettings Settings;
 		Settings.PolygonGroupsDelegate = FAppendPolygonGroupsDelegate::CreateLambda(
-			[&PrototypePolygonGroupRemapping, PrototypeIndex]( const FMeshDescription& SourceMesh, FMeshDescription& TargetMesh, PolygonGroupMap& RemapPolygonGroups )
+			[&PrototypePolygonGroupRemapping,
+			 PrototypeIndex](const FMeshDescription& SourceMesh, FMeshDescription& TargetMesh, PolygonGroupMap& RemapPolygonGroups)
 			{
-				RemapPolygonGroups = PrototypePolygonGroupRemapping[ PrototypeIndex ];
+				RemapPolygonGroups = PrototypePolygonGroupRemapping[PrototypeIndex];
 			}
 		);
 
 		// TODO: Maybe we should make a new overload of AppendMeshDescriptions that can do this more efficiently, since all we need is to change the
 		// transform repeatedly?
-		for ( const FTransform& Transform : InstanceTransforms )
+		for (const FTransform& Transform : InstanceTransforms)
 		{
 			Settings.MeshTransform = Transform * Options.AdditionalTransform;
-			FStaticMeshOperations::AppendMeshDescription( PrototypeMeshDescription, OutMeshDescription, Settings );
+			FStaticMeshOperations::AppendMeshDescription(PrototypeMeshDescription, OutMeshDescription, Settings);
 		}
 	}
 
@@ -2735,14 +2738,14 @@ bool UsdToUnreal::ConvertGeomMeshHierarchy(
 	bool bSkipRootPrimTransformAndVisibility
 )
 {
-	TRACE_CPUPROFILER_EVENT_SCOPE( UsdToUnreal::ConvertGeomMeshHierarchy );
+	TRACE_CPUPROFILER_EVENT_SCOPE(UsdToUnreal::ConvertGeomMeshHierarchy);
 
-	if ( !Prim )
+	if (!Prim)
 	{
 		return false;
 	}
 
-	FStaticMeshAttributes StaticMeshAttributes( OutMeshDescription );
+	FStaticMeshAttributes StaticMeshAttributes(OutMeshDescription);
 	StaticMeshAttributes.Register();
 
 	// Pass a copy down so that we can repeatedly overwrite the AdditionalTransform and still
@@ -2752,11 +2755,7 @@ bool UsdToUnreal::ConvertGeomMeshHierarchy(
 	// Prepass to figure out the best primvars to use for the entire collapsed mesh UV sets
 	if (OutMaterialAssignments.PrimvarToUVIndex.Num() == 0)
 	{
-		OutMaterialAssignments.PrimvarToUVIndex = UsdGeomMeshImpl::CollectSubtreePrimvars(
-			Prim,
-			Options,
-			bSkipRootPrimTransformAndVisibility
-		);
+		OutMaterialAssignments.PrimvarToUVIndex = UsdGeomMeshImpl::CollectSubtreePrimvars(Prim, Options, bSkipRootPrimTransformAndVisibility);
 	}
 
 	const bool bIsInSkelRoot = static_cast<bool>(UsdUtils::GetClosestParentSkelRoot(Prim));
@@ -2878,11 +2877,7 @@ bool UsdToUnreal::ConvertGeomPrimitive(
 	return ConvertMeshData(MeshData, StageInfo, InOptions, InOutMeshDescription, InOutMaterialAssignments);
 }
 
-bool UsdToUnreal::ConvertGeomPrimitiveTransform(
-	const pxr::UsdPrim& InPrim,
-	const pxr::UsdTimeCode& InTimeCode,
-	FTransform& OutTransform
-)
+bool UsdToUnreal::ConvertGeomPrimitiveTransform(const pxr::UsdPrim& InPrim, const pxr::UsdTimeCode& InTimeCode, FTransform& OutTransform)
 {
 	TRACE_CPUPROFILER_EVENT_SCOPE(UsdToUnreal::ConvertGeomPrimitive);
 
@@ -2896,35 +2891,19 @@ bool UsdToUnreal::ConvertGeomPrimitiveTransform(
 	pxr::UsdStageRefPtr Stage = InPrim.GetStage();
 	const FUsdStageInfo StageInfo(Stage);
 
-	TFunction<pxr::GfMatrix4d(pxr::TfToken, double, double)> GetScalingTransform =
-		[](pxr::TfToken Axis, double Longitudinal, double Transversal)
+	TFunction<pxr::GfMatrix4d(pxr::TfToken, double, double)> GetScalingTransform = [](pxr::TfToken Axis, double Longitudinal, double Transversal)
 	{
 		if (Axis == pxr::UsdGeomTokens->x)
 		{
-			return pxr::GfMatrix4d(
-				Longitudinal, 		  0.0, 		   0.0, 0.0,
-						 0.0, Transversal, 		   0.0, 0.0,
-						 0.0, 		  0.0, Transversal, 0.0,
-						 0.0, 		  0.0, 		   0.0, 1.0
-			);
+			return pxr::GfMatrix4d(Longitudinal, 0.0, 0.0, 0.0, 0.0, Transversal, 0.0, 0.0, 0.0, 0.0, Transversal, 0.0, 0.0, 0.0, 0.0, 1.0);
 		}
 		else if (Axis == pxr::UsdGeomTokens->y)
 		{
-			return pxr::GfMatrix4d(
-				Transversal, 		  0.0, 		   0.0, 0.0,
-						0.0, Longitudinal, 		   0.0, 0.0,
-						0.0, 		  0.0, Transversal, 0.0,
-						0.0, 		  0.0, 		   0.0, 1.0
-			);
+			return pxr::GfMatrix4d(Transversal, 0.0, 0.0, 0.0, 0.0, Longitudinal, 0.0, 0.0, 0.0, 0.0, Transversal, 0.0, 0.0, 0.0, 0.0, 1.0);
 		}
 		else
 		{
-			return pxr::GfMatrix4d(
-				Transversal, 		 0.0, 		   0.0, 0.0,
-						0.0, Transversal, 		   0.0, 0.0,
-						0.0, 		 0.0, Longitudinal, 0.0,
-						0.0, 		 0.0, 		   0.0, 1.0
-			);
+			return pxr::GfMatrix4d(Transversal, 0.0, 0.0, 0.0, 0.0, Transversal, 0.0, 0.0, 0.0, 0.0, Longitudinal, 0.0, 0.0, 0.0, 0.0, 1.0);
 		}
 	};
 
@@ -3077,18 +3056,18 @@ bool UsdToUnreal::ConvertGeomPrimitiveTransform(
 }
 
 PRAGMA_DISABLE_DEPRECATION_WARNINGS
-UMaterialInstanceDynamic* UsdUtils::CreateDisplayColorMaterialInstanceDynamic( const UsdUtils::FDisplayColorMaterial& DisplayColorDescription )
+UMaterialInstanceDynamic* UsdUtils::CreateDisplayColorMaterialInstanceDynamic(const UsdUtils::FDisplayColorMaterial& DisplayColorDescription)
 {
 	const UUsdProjectSettings* Settings = GetDefault<UUsdProjectSettings>();
-	if ( !Settings )
+	if (!Settings)
 	{
 		return nullptr;
 	}
 
 	const FSoftObjectPath* ParentPathPtr = nullptr;
-	if ( DisplayColorDescription.bHasOpacity )
+	if (DisplayColorDescription.bHasOpacity)
 	{
-		if ( DisplayColorDescription.bIsDoubleSided )
+		if (DisplayColorDescription.bIsDoubleSided)
 		{
 			ParentPathPtr = &Settings->ReferenceDisplayColorAndOpacityTwoSidedMaterial;
 		}
@@ -3099,7 +3078,7 @@ UMaterialInstanceDynamic* UsdUtils::CreateDisplayColorMaterialInstanceDynamic( c
 	}
 	else
 	{
-		if ( DisplayColorDescription.bIsDoubleSided )
+		if (DisplayColorDescription.bIsDoubleSided)
 		{
 			ParentPathPtr = &Settings->ReferenceDisplayColorTwoSidedMaterial;
 		}
@@ -3108,23 +3087,24 @@ UMaterialInstanceDynamic* UsdUtils::CreateDisplayColorMaterialInstanceDynamic( c
 			ParentPathPtr = &Settings->ReferenceDisplayColorMaterial;
 		}
 	}
-	if ( !ParentPathPtr )
+	if (!ParentPathPtr)
 	{
 		return nullptr;
 	}
 
-	if ( UMaterialInterface* ParentMaterial = Cast< UMaterialInterface >( ParentPathPtr->TryLoad() ) )
+	if (UMaterialInterface* ParentMaterial = Cast<UMaterialInterface>(ParentPathPtr->TryLoad()))
 	{
 		FName AssetName = MakeUniqueObjectName(
 			GetTransientPackage(),
 			UMaterialInstanceConstant::StaticClass(),
-			*FString::Printf( TEXT( "DisplayColor_%s_%s" ),
-				DisplayColorDescription.bHasOpacity ? TEXT( "Opacity" ) : TEXT( "NoOpacity" ),
-				DisplayColorDescription.bIsDoubleSided ? TEXT( "DoubleSided" ) : TEXT( "SingleSided" )
+			*FString::Printf(
+				TEXT("DisplayColor_%s_%s"),
+				DisplayColorDescription.bHasOpacity ? TEXT("Opacity") : TEXT("NoOpacity"),
+				DisplayColorDescription.bIsDoubleSided ? TEXT("DoubleSided") : TEXT("SingleSided")
 			)
 		);
 
-		if ( UMaterialInstanceDynamic* NewMaterial = UMaterialInstanceDynamic::Create( ParentMaterial, GetTransientPackage(), AssetName ) )
+		if (UMaterialInstanceDynamic* NewMaterial = UMaterialInstanceDynamic::Create(ParentMaterial, GetTransientPackage(), AssetName))
 		{
 			return NewMaterial;
 		}
@@ -3133,19 +3113,19 @@ UMaterialInstanceDynamic* UsdUtils::CreateDisplayColorMaterialInstanceDynamic( c
 	return nullptr;
 }
 
-UMaterialInstanceConstant* UsdUtils::CreateDisplayColorMaterialInstanceConstant( const UsdUtils::FDisplayColorMaterial& DisplayColorDescription )
+UMaterialInstanceConstant* UsdUtils::CreateDisplayColorMaterialInstanceConstant(const UsdUtils::FDisplayColorMaterial& DisplayColorDescription)
 {
 #if WITH_EDITOR
 	const UUsdProjectSettings* Settings = GetDefault<UUsdProjectSettings>();
-	if ( !Settings )
+	if (!Settings)
 	{
 		return nullptr;
 	}
 
 	const FSoftObjectPath* ParentPathPtr = nullptr;
-	if ( DisplayColorDescription.bHasOpacity )
+	if (DisplayColorDescription.bHasOpacity)
 	{
-		if ( DisplayColorDescription.bIsDoubleSided )
+		if (DisplayColorDescription.bIsDoubleSided)
 		{
 			ParentPathPtr = &Settings->ReferenceDisplayColorAndOpacityTwoSidedMaterial;
 		}
@@ -3156,7 +3136,7 @@ UMaterialInstanceConstant* UsdUtils::CreateDisplayColorMaterialInstanceConstant(
 	}
 	else
 	{
-		if ( DisplayColorDescription.bIsDoubleSided )
+		if (DisplayColorDescription.bIsDoubleSided)
 		{
 			ParentPathPtr = &Settings->ReferenceDisplayColorTwoSidedMaterial;
 		}
@@ -3166,26 +3146,28 @@ UMaterialInstanceConstant* UsdUtils::CreateDisplayColorMaterialInstanceConstant(
 		}
 	}
 
-	if ( UMaterialInterface* ParentMaterial = Cast< UMaterialInterface >( ParentPathPtr->TryLoad() ) )
+	if (UMaterialInterface* ParentMaterial = Cast<UMaterialInterface>(ParentPathPtr->TryLoad()))
 	{
 		FName AssetName = MakeUniqueObjectName(
 			GetTransientPackage(),
 			UMaterialInstanceConstant::StaticClass(),
-			*FString::Printf( TEXT( "DisplayColor_%s_%s" ),
-				DisplayColorDescription.bHasOpacity ? TEXT( "Opacity" ) : TEXT( "NoOpacity" ),
-				DisplayColorDescription.bIsDoubleSided ? TEXT( "DoubleSided" ) : TEXT( "SingleSided" )
+			*FString::Printf(
+				TEXT("DisplayColor_%s_%s"),
+				DisplayColorDescription.bHasOpacity ? TEXT("Opacity") : TEXT("NoOpacity"),
+				DisplayColorDescription.bIsDoubleSided ? TEXT("DoubleSided") : TEXT("SingleSided")
 			)
 		);
 
-		if ( UMaterialInstanceConstant* MaterialInstance = NewObject< UMaterialInstanceConstant >( GetTransientPackage(), AssetName, RF_NoFlags ) )
+		if (UMaterialInstanceConstant* MaterialInstance = NewObject<UMaterialInstanceConstant>(GetTransientPackage(), AssetName, RF_NoFlags))
 		{
-			UMaterialEditingLibrary::SetMaterialInstanceParent( MaterialInstance, ParentMaterial );
+			UMaterialEditingLibrary::SetMaterialInstanceParent(MaterialInstance, ParentMaterial);
 			return MaterialInstance;
 		}
 	}
-#endif // WITH_EDITOR
+#endif	  // WITH_EDITOR
 	return nullptr;
 }
+
 PRAGMA_ENABLE_DEPRECATION_WARNINGS
 
 UsdUtils::FUsdPrimMaterialAssignmentInfo UsdUtils::GetPrimMaterialAssignments(
@@ -3196,59 +3178,65 @@ UsdUtils::FUsdPrimMaterialAssignmentInfo UsdUtils::GetPrimMaterialAssignments(
 	const pxr::TfToken& MaterialPurpose
 )
 {
-	if ( !UsdPrim )
+	if (!UsdPrim)
 	{
 		return {};
 	}
 
-	auto FetchFirstUEMaterialFromAttribute = []( const pxr::UsdPrim& UsdPrim, const pxr::UsdTimeCode TimeCode ) -> TOptional<FString>
+	auto FetchFirstUEMaterialFromAttribute = [](const pxr::UsdPrim& UsdPrim, const pxr::UsdTimeCode TimeCode) -> TOptional<FString>
 	{
 		FString ValidPackagePath;
-		if ( pxr::UsdAttribute MaterialAttribute = UsdPrim.GetAttribute( UnrealIdentifiers::MaterialAssignment ) )
+		if (pxr::UsdAttribute MaterialAttribute = UsdPrim.GetAttribute(UnrealIdentifiers::MaterialAssignment))
 		{
 			std::string UEMaterial;
-			if ( MaterialAttribute.Get( &UEMaterial, TimeCode ) && UEMaterial.size() > 0)
+			if (MaterialAttribute.Get(&UEMaterial, TimeCode) && UEMaterial.size() > 0)
 			{
-				ValidPackagePath = UsdToUnreal::ConvertString( UEMaterial );
+				ValidPackagePath = UsdToUnreal::ConvertString(UEMaterial);
 			}
 		}
 
-		if ( !ValidPackagePath.IsEmpty() )
+		if (!ValidPackagePath.IsEmpty())
 		{
 			// We can't TryLoad() or LoadObject<> this right now as we may be in an Async thread.
 			// The FAssetData may not be ready yet however, in case we're loading a stage right when launching the
 			// editor, so here we just settle for finding any valid object
-			FSoftObjectPath SoftObjectPath{ ValidPackagePath };
-			if ( SoftObjectPath.IsValid() )
+			FSoftObjectPath SoftObjectPath{ValidPackagePath};
+			if (SoftObjectPath.IsValid())
 			{
 				return ValidPackagePath;
 			}
 
-			UE_LOG( LogUsd, Warning, TEXT( "Could not find a valid material at path '%s', targetted by prim '%s's unrealMaterial attribute. Material assignment will fallback to USD materials and display color data." ),
-				*ValidPackagePath, *UsdToUnreal::ConvertPath( UsdPrim.GetPath() ) );
+			UE_LOG(
+				LogUsd,
+				Warning,
+				TEXT("Could not find a valid material at path '%s', targetted by prim '%s's unrealMaterial attribute. Material assignment will "
+					 "fallback to USD materials and display color data."),
+				*ValidPackagePath,
+				*UsdToUnreal::ConvertPath(UsdPrim.GetPath())
+			);
 		}
 
 		return {};
 	};
 
-	auto FetchMaterialByComputingBoundMaterial = [ &RenderContext, &MaterialPurpose ]( const pxr::UsdPrim& UsdPrim ) -> TOptional<FString>
+	auto FetchMaterialByComputingBoundMaterial = [&RenderContext, &MaterialPurpose](const pxr::UsdPrim& UsdPrim) -> TOptional<FString>
 	{
-		pxr::UsdShadeMaterialBindingAPI BindingAPI( UsdPrim );
-		pxr::UsdShadeMaterial ShadeMaterial = BindingAPI.ComputeBoundMaterial( MaterialPurpose );
-		if ( !ShadeMaterial )
+		pxr::UsdShadeMaterialBindingAPI BindingAPI(UsdPrim);
+		pxr::UsdShadeMaterial ShadeMaterial = BindingAPI.ComputeBoundMaterial(MaterialPurpose);
+		if (!ShadeMaterial)
 		{
 			return {};
 		}
 
 		// Ignore this material if UsdToUnreal::ConvertMaterial would as well
-		pxr::UsdShadeShader SurfaceShader = ShadeMaterial.ComputeSurfaceSource( RenderContext );
-		if ( !SurfaceShader )
+		pxr::UsdShadeShader SurfaceShader = ShadeMaterial.ComputeSurfaceSource(RenderContext);
+		if (!SurfaceShader)
 		{
 			return {};
 		}
 
 		pxr::UsdPrim ShadeMaterialPrim = ShadeMaterial.GetPrim();
-		if ( ShadeMaterialPrim )
+		if (ShadeMaterialPrim)
 		{
 			const std::string ShadingEngineName = ShadeMaterialPrim.GetPrimPath().GetString();
 			if (!ShadingEngineName.empty())
@@ -3266,32 +3254,32 @@ UsdUtils::FUsdPrimMaterialAssignmentInfo UsdUtils::GetPrimMaterialAssignments(
 	{
 		pxr::VtArray<int> FaceVertexCounts = UsdGeomMeshImpl::GetFaceVertexCounts(UsdPrim, TimeCode);
 		NumFaces = FaceVertexCounts.size();
-		if ( NumFaces < 1 )
+		if (NumFaces < 1)
 		{
 			return Result;
 		}
 
-		if ( bProvideMaterialIndices )
+		if (bProvideMaterialIndices)
 		{
 			// Note how we're defaulting to slot zero here, which is our "main assignment"
-			Result.MaterialIndices.SetNumZeroed( NumFaces );
+			Result.MaterialIndices.SetNumZeroed(NumFaces);
 		}
 	}
 
 	bool bIsDoubleSided = false;
-	if ( pxr::UsdGeomMesh Mesh = pxr::UsdGeomMesh{ UsdPrim } )
+	if (pxr::UsdGeomMesh Mesh = pxr::UsdGeomMesh{UsdPrim})
 	{
-		if ( pxr::UsdAttribute Attr = Mesh.GetDoubleSidedAttr() )
+		if (pxr::UsdAttribute Attr = Mesh.GetDoubleSidedAttr())
 		{
 			pxr::VtValue AttrValue;
-			if ( Attr.Get( &AttrValue ) && AttrValue.IsHolding<bool>() )
+			if (Attr.Get(&AttrValue) && AttrValue.IsHolding<bool>())
 			{
 				bIsDoubleSided = AttrValue.UncheckedGet<bool>();
 			}
 		}
 	}
 
-	FString MeshPrimPath = UsdToUnreal::ConvertPath( UsdPrim.GetPath() );
+	FString MeshPrimPath = UsdToUnreal::ConvertPath(UsdPrim.GetPath());
 
 	bool bNeedsMainAssignment = true;
 
@@ -3543,113 +3531,130 @@ UsdUtils::FUsdPrimMaterialAssignmentInfo UsdUtils::GetPrimMaterialAssignments(
 	return Result;
 }
 
-TArray<FString> UsdUtils::GetMaterialUsers( const UE::FUsdPrim& MaterialPrim, FName MaterialPurpose )
+TArray<FString> UsdUtils::GetMaterialUsers(const UE::FUsdPrim& MaterialPrim, FName MaterialPurpose)
 {
 	TArray<FString> Result;
 
 	FScopedUsdAllocs Allocs;
 
-	pxr::UsdPrim UsdMaterialPrim{ MaterialPrim };
-	if ( !UsdMaterialPrim || !UsdMaterialPrim.IsA<pxr::UsdShadeMaterial>() )
+	pxr::UsdPrim UsdMaterialPrim{MaterialPrim};
+	if (!UsdMaterialPrim || !UsdMaterialPrim.IsA<pxr::UsdShadeMaterial>())
 	{
 		return Result;
 	}
 
 	pxr::TfToken MaterialPurposeToken = pxr::UsdShadeTokens->allPurpose;
-	if ( !MaterialPurpose.IsNone() )
+	if (!MaterialPurpose.IsNone())
 	{
-		MaterialPurposeToken = UnrealToUsd::ConvertToken( *MaterialPurpose.ToString() ).Get();
+		MaterialPurposeToken = UnrealToUsd::ConvertToken(*MaterialPurpose.ToString()).Get();
 	}
 
 	pxr::UsdStageRefPtr UsdStage = UsdMaterialPrim.GetStage();
 
-	pxr::UsdPrimRange PrimRange = pxr::UsdPrimRange::Stage( UsdStage, pxr::UsdTraverseInstanceProxies() );
-	for ( pxr::UsdPrimRange::iterator PrimRangeIt = PrimRange.begin(); PrimRangeIt != PrimRange.end(); ++PrimRangeIt )
+	pxr::UsdPrimRange PrimRange = pxr::UsdPrimRange::Stage(UsdStage, pxr::UsdTraverseInstanceProxies());
+	for (pxr::UsdPrimRange::iterator PrimRangeIt = PrimRange.begin(); PrimRangeIt != PrimRange.end(); ++PrimRangeIt)
 	{
 		pxr::UsdPrim Prim = *PrimRangeIt;
 
-		if ( !Prim.HasAPI<pxr::UsdShadeMaterialBindingAPI>() )
+		if (!Prim.HasAPI<pxr::UsdShadeMaterialBindingAPI>())
 		{
 			continue;
 		}
 
-		pxr::UsdShadeMaterialBindingAPI BindingAPI( Prim );
-		pxr::UsdShadeMaterial ShadeMaterial = BindingAPI.ComputeBoundMaterial( MaterialPurposeToken );
-		if ( !ShadeMaterial )
+		pxr::UsdShadeMaterialBindingAPI BindingAPI(Prim);
+		pxr::UsdShadeMaterial ShadeMaterial = BindingAPI.ComputeBoundMaterial(MaterialPurposeToken);
+		if (!ShadeMaterial)
 		{
 			continue;
 		}
 
 		pxr::UsdPrim ShadeMaterialPrim = ShadeMaterial.GetPrim();
-		if ( ShadeMaterialPrim == UsdMaterialPrim )
+		if (ShadeMaterialPrim == UsdMaterialPrim)
 		{
-			Result.Add( UsdToUnreal::ConvertPath( Prim.GetPrimPath() ) );
+			Result.Add(UsdToUnreal::ConvertPath(Prim.GetPrimPath()));
 		}
 	}
 
 	return Result;
 }
 
-bool UnrealToUsd::ConvertStaticMesh( const UStaticMesh* StaticMesh, pxr::UsdPrim& UsdPrim, const pxr::UsdTimeCode TimeCode, UE::FUsdStage* StageForMaterialAssignments, int32 LowestMeshLOD, int32 HighestMeshLOD )
+bool UnrealToUsd::ConvertStaticMesh(
+	const UStaticMesh* StaticMesh,
+	pxr::UsdPrim& UsdPrim,
+	const pxr::UsdTimeCode TimeCode,
+	UE::FUsdStage* StageForMaterialAssignments,
+	int32 LowestMeshLOD,
+	int32 HighestMeshLOD
+)
 {
 	FScopedUsdAllocs UsdAllocs;
 
 	pxr::UsdStageRefPtr Stage = UsdPrim.GetStage();
-	if ( !Stage )
+	if (!Stage)
 	{
 		return false;
 	}
 
-	const FUsdStageInfo StageInfo( Stage );
+	const FUsdStageInfo StageInfo(Stage);
 
 	int32 NumLODs = StaticMesh->GetNumLODs();
-	if ( NumLODs < 1 )
+	if (NumLODs < 1)
 	{
 		return false;
 	}
 
 	// Make sure they're both >= 0 (the options dialog slider is clamped, but this may be called directly)
-	LowestMeshLOD = FMath::Clamp( LowestMeshLOD, 0, NumLODs -1 );
-	HighestMeshLOD = FMath::Clamp( HighestMeshLOD, 0, NumLODs - 1 );
+	LowestMeshLOD = FMath::Clamp(LowestMeshLOD, 0, NumLODs - 1);
+	HighestMeshLOD = FMath::Clamp(HighestMeshLOD, 0, NumLODs - 1);
 
 	// Make sure Lowest <= Highest
-	int32 Temp = FMath::Min( LowestMeshLOD, HighestMeshLOD );
-	HighestMeshLOD = FMath::Max( LowestMeshLOD, HighestMeshLOD );
+	int32 Temp = FMath::Min(LowestMeshLOD, HighestMeshLOD);
+	HighestMeshLOD = FMath::Max(LowestMeshLOD, HighestMeshLOD);
 	LowestMeshLOD = Temp;
 
 	// Make sure it's at least 1 LOD level
-	NumLODs = FMath::Max( HighestMeshLOD - LowestMeshLOD + 1, 1 );
+	NumLODs = FMath::Max(HighestMeshLOD - LowestMeshLOD + 1, 1);
 
 	// If exporting a Nanite mesh, just use the lowest LOD and write the unrealNanite override attribute,
 	// this way we can guarantee it will have Nanite when it's imported back.
 #if WITH_EDITOR
-	if ( StaticMesh->IsNaniteEnabled())
+	if (StaticMesh->IsNaniteEnabled())
 	{
-		if ( NumLODs > 1 )
+		if (NumLODs > 1)
 		{
-			UE_LOG( LogUsd, Log, TEXT( "Not exporting multiple LODs for mesh '%s' onto prim '%s' since the mesh has Nanite enabled: LOD '%d' will be used and the '%s' attribute will be written out instead" ),
+			UE_LOG(
+				LogUsd,
+				Log,
+				TEXT("Not exporting multiple LODs for mesh '%s' onto prim '%s' since the mesh has Nanite enabled: LOD '%d' will be used and the '%s' "
+					 "attribute will be written out instead"),
 				*StaticMesh->GetName(),
-				*UsdToUnreal::ConvertPath( UsdPrim.GetPrimPath() ),
+				*UsdToUnreal::ConvertPath(UsdPrim.GetPrimPath()),
 				LowestMeshLOD,
-				*UsdToUnreal::ConvertToken( UnrealIdentifiers::UnrealNaniteOverride )
+				*UsdToUnreal::ConvertToken(UnrealIdentifiers::UnrealNaniteOverride)
 			);
 		}
 
 		HighestMeshLOD = LowestMeshLOD;
 		NumLODs = 1;
 
-		if ( pxr::UsdAttribute Attr = UsdPrim.CreateAttribute( UnrealIdentifiers::UnrealNaniteOverride, pxr::SdfValueTypeNames->Token ) )
+		if (pxr::UsdAttribute Attr = UsdPrim.CreateAttribute(UnrealIdentifiers::UnrealNaniteOverride, pxr::SdfValueTypeNames->Token))
 		{
-			Attr.Set( UnrealIdentifiers::UnrealNaniteOverrideEnable );
+			Attr.Set(UnrealIdentifiers::UnrealNaniteOverrideEnable);
 			UsdUtils::NotifyIfOverriddenOpinion(Attr);
 		}
 	}
-#endif // WITH_EDITOR
+#endif	  // WITH_EDITOR
 
 	pxr::UsdVariantSets VariantSets = UsdPrim.GetVariantSets();
-	if ( NumLODs > 1 && VariantSets.HasVariantSet( UnrealIdentifiers::LOD ) )
-	{
-		UE_LOG( LogUsd, Error, TEXT("Failed to export higher LODs for mesh '%s', as the target prim already has a variant set named '%s'!"), *StaticMesh->GetName(), *UsdToUnreal::ConvertToken( UnrealIdentifiers::LOD ) );
+	if (NumLODs > 1 && VariantSets.HasVariantSet(UnrealIdentifiers::LOD))
+	{
+		UE_LOG(
+			LogUsd,
+			Error,
+			TEXT("Failed to export higher LODs for mesh '%s', as the target prim already has a variant set named '%s'!"),
+			*StaticMesh->GetName(),
+			*UsdToUnreal::ConvertToken(UnrealIdentifiers::LOD)
+		);
 		NumLODs = 1;
 	}
 
@@ -3660,22 +3665,22 @@ bool UnrealToUsd::ConvertStaticMesh( const UStaticMesh* StaticMesh, pxr::UsdPrim
 
 	// Collect all material assignments, referenced by the sections' material indices
 	bool bHasMaterialAssignments = false;
-	pxr::VtArray< std::string > MaterialAssignments;
-	for(const FStaticMaterial& StaticMaterial : StaticMesh->GetStaticMaterials())
+	pxr::VtArray<std::string> MaterialAssignments;
+	for (const FStaticMaterial& StaticMaterial : StaticMesh->GetStaticMaterials())
 	{
 		FString AssignedMaterialPathName;
-		if ( UMaterialInterface* Material = StaticMaterial.MaterialInterface )
+		if (UMaterialInterface* Material = StaticMaterial.MaterialInterface)
 		{
-			if ( Material->GetOutermost() != GetTransientPackage() )
+			if (Material->GetOutermost() != GetTransientPackage())
 			{
 				AssignedMaterialPathName = Material->GetPathName();
 				bHasMaterialAssignments = true;
 			}
 		}
 
-		MaterialAssignments.push_back( UnrealToUsd::ConvertString( *AssignedMaterialPathName ).Get() );
+		MaterialAssignments.push_back(UnrealToUsd::ConvertString(*AssignedMaterialPathName).Get());
 	}
-	if ( !bHasMaterialAssignments )
+	if (!bHasMaterialAssignments)
 	{
 		// Prevent creation of the unrealMaterials attribute in case we don't have any assignments at all
 		MaterialAssignments.clear();
@@ -3688,24 +3693,24 @@ bool UnrealToUsd::ConvertStaticMesh( const UStaticMesh* StaticMesh, pxr::UsdPrim
 	// here, so would our referencer and we wouldn't be able to put a transform on it
 	UsdPrim = Stage->DefinePrim(UsdPrim.GetPath(), UnrealToUsd::ConvertToken(bExportMultipleLODs ? TEXT("Xform") : TEXT("Mesh")).Get());
 
-	for ( int32 LODIndex = LowestMeshLOD; LODIndex <= HighestMeshLOD; ++LODIndex )
+	for (int32 LODIndex = LowestMeshLOD; LODIndex <= HighestMeshLOD; ++LODIndex)
 	{
-		const FStaticMeshLODResources& RenderMesh = StaticMesh->GetLODForExport( LODIndex );
+		const FStaticMeshLODResources& RenderMesh = StaticMesh->GetLODForExport(LODIndex);
 
 		// Verify the integrity of the static mesh.
-		if ( RenderMesh.VertexBuffers.StaticMeshVertexBuffer.GetNumVertices() == 0 )
+		if (RenderMesh.VertexBuffers.StaticMeshVertexBuffer.GetNumVertices() == 0)
 		{
 			continue;
 		}
 
-		if ( RenderMesh.Sections.Num() == 0 )
+		if (RenderMesh.Sections.Num() == 0)
 		{
 			continue;
 		}
 
 		// LOD0, LOD1, etc
-		std::string VariantName = UnrealIdentifiers::LOD.GetString() + UnrealToUsd::ConvertString( *LexToString( LODIndex ) ).Get();
-		if ( LowestLODAdded.size() == 0 )
+		std::string VariantName = UnrealIdentifiers::LOD.GetString() + UnrealToUsd::ConvertString(*LexToString(LODIndex)).Get();
+		if (LowestLODAdded.size() == 0)
 		{
 			LowestLODAdded = VariantName;
 		}
@@ -3713,25 +3718,25 @@ bool UnrealToUsd::ConvertStaticMesh( const UStaticMesh* StaticMesh, pxr::UsdPrim
 		pxr::SdfPath LODPrimPath = ParentPrimPath.AppendPath(pxr::SdfPath(VariantName));
 
 		// Enable the variant edit context, if we are creating variant LODs
-		TOptional< pxr::UsdEditContext > EditContext;
-		if ( bExportMultipleLODs )
+		TOptional<pxr::UsdEditContext> EditContext;
+		if (bExportMultipleLODs)
 		{
-			pxr::UsdVariantSet VariantSet = VariantSets.GetVariantSet( UnrealIdentifiers::LOD );
+			pxr::UsdVariantSet VariantSet = VariantSets.GetVariantSet(UnrealIdentifiers::LOD);
 
-			if ( !VariantSet.AddVariant( VariantName ) )
+			if (!VariantSet.AddVariant(VariantName))
 			{
 				continue;
 			}
 
-			VariantSet.SetVariantSelection( VariantName );
-			EditContext.Emplace( VariantSet.GetVariantEditContext() );
+			VariantSet.SetVariantSelection(VariantName);
+			EditContext.Emplace(VariantSet.GetVariantEditContext());
 		}
 
 		// Author material bindings on the dedicated stage if we have one
 		pxr::UsdStageRefPtr MaterialStage;
-		if ( StageForMaterialAssignments )
+		if (StageForMaterialAssignments)
 		{
-			MaterialStage = static_cast< pxr::UsdStageRefPtr >( *StageForMaterialAssignments );
+			MaterialStage = static_cast<pxr::UsdStageRefPtr>(*StageForMaterialAssignments);
 		}
 		else
 		{
@@ -3740,26 +3745,26 @@ bool UnrealToUsd::ConvertStaticMesh( const UStaticMesh* StaticMesh, pxr::UsdPrim
 
 		pxr::UsdGeomMesh TargetMesh;
 		pxr::UsdPrim MaterialPrim = UsdPrim;
-		if ( bExportMultipleLODs )
+		if (bExportMultipleLODs)
 		{
 			// Add the mesh data to a child prim with the Mesh schema
-			pxr::UsdPrim UsdLODPrim = Stage->DefinePrim( LODPrimPath, UnrealToUsd::ConvertToken( TEXT("Mesh") ).Get() );
-			TargetMesh = pxr::UsdGeomMesh{ UsdLODPrim };
+			pxr::UsdPrim UsdLODPrim = Stage->DefinePrim(LODPrimPath, UnrealToUsd::ConvertToken(TEXT("Mesh")).Get());
+			TargetMesh = pxr::UsdGeomMesh{UsdLODPrim};
 
-			MaterialPrim = MaterialStage->OverridePrim( LODPrimPath );
+			MaterialPrim = MaterialStage->OverridePrim(LODPrimPath);
 		}
 		else
 		{
-			TargetMesh = pxr::UsdGeomMesh{ UsdPrim };
+			TargetMesh = pxr::UsdGeomMesh{UsdPrim};
 
-			MaterialPrim = MaterialStage->OverridePrim( UsdPrim.GetPath() );
+			MaterialPrim = MaterialStage->OverridePrim(UsdPrim.GetPath());
 		}
 
-		UsdGeomMeshImpl::ConvertStaticMeshLOD( LODIndex, RenderMesh, TargetMesh, MaterialAssignments, TimeCode, MaterialPrim );
+		UsdGeomMeshImpl::ConvertStaticMeshLOD(LODIndex, RenderMesh, TargetMesh, MaterialAssignments, TimeCode, MaterialPrim);
 	}
 
 	// Reset variant set to start with the lowest lod selected
-	if ( bExportMultipleLODs )
+	if (bExportMultipleLODs)
 	{
 		VariantSets.GetVariantSet(UnrealIdentifiers::LOD).SetVariantSelection(LowestLODAdded);
 	}
@@ -3767,30 +3772,38 @@ bool UnrealToUsd::ConvertStaticMesh( const UStaticMesh* StaticMesh, pxr::UsdPrim
 	return true;
 }
 
-bool UnrealToUsd::ConvertMeshDescriptions( const TArray<FMeshDescription>& LODIndexToMeshDescription, pxr::UsdPrim& UsdPrim, const FMatrix& AdditionalTransform, const pxr::UsdTimeCode TimeCode )
+bool UnrealToUsd::ConvertMeshDescriptions(
+	const TArray<FMeshDescription>& LODIndexToMeshDescription,
+	pxr::UsdPrim& UsdPrim,
+	const FMatrix& AdditionalTransform,
+	const pxr::UsdTimeCode TimeCode
+)
 {
 	FScopedUsdAllocs UsdAllocs;
 
 	pxr::UsdStageRefPtr Stage = UsdPrim.GetStage();
-	if ( !Stage )
+	if (!Stage)
 	{
 		return false;
 	}
 
-	const FUsdStageInfo StageInfo( Stage );
+	const FUsdStageInfo StageInfo(Stage);
 
 	int32 NumLODs = LODIndexToMeshDescription.Num();
-	if ( NumLODs < 1 )
+	if (NumLODs < 1)
 	{
 		return false;
 	}
 
 	pxr::UsdVariantSets VariantSets = UsdPrim.GetVariantSets();
-	if ( NumLODs > 1 && VariantSets.HasVariantSet( UnrealIdentifiers::LOD ) )
-	{
-		UE_LOG( LogUsd, Error, TEXT( "Failed to convert higher mesh description LODs for prim '%s', as the target prim already has a variant set named '%s'!" ),
-			*UsdToUnreal::ConvertPath( UsdPrim.GetPath() ),
-			*UsdToUnreal::ConvertToken( UnrealIdentifiers::LOD )
+	if (NumLODs > 1 && VariantSets.HasVariantSet(UnrealIdentifiers::LOD))
+	{
+		UE_LOG(
+			LogUsd,
+			Error,
+			TEXT("Failed to convert higher mesh description LODs for prim '%s', as the target prim already has a variant set named '%s'!"),
+			*UsdToUnreal::ConvertPath(UsdPrim.GetPath()),
+			*UsdToUnreal::ConvertToken(UnrealIdentifiers::LOD)
 		);
 		NumLODs = 1;
 	}
@@ -3803,55 +3816,55 @@ bool UnrealToUsd::ConvertMeshDescriptions( const TArray<FMeshDescription>& LODIn
 	// Check the comment on the analogous line on ConvertStaticMesh
 	UsdPrim = Stage->DefinePrim(UsdPrim.GetPath(), UnrealToUsd::ConvertToken(bExportMultipleLODs ? TEXT("Xform") : TEXT("Mesh")).Get());
 
-	for ( int32 LODIndex = 0; LODIndex < NumLODs; ++LODIndex )
+	for (int32 LODIndex = 0; LODIndex < NumLODs; ++LODIndex)
 	{
-		const FMeshDescription& MeshDescription = LODIndexToMeshDescription[ LODIndex ];
+		const FMeshDescription& MeshDescription = LODIndexToMeshDescription[LODIndex];
 
 		// LOD0, LOD1, etc
-		std::string VariantName = UnrealIdentifiers::LOD.GetString() + UnrealToUsd::ConvertString( *LexToString( LODIndex ) ).Get();
-		if ( LowestLODAdded.size() == 0 )
+		std::string VariantName = UnrealIdentifiers::LOD.GetString() + UnrealToUsd::ConvertString(*LexToString(LODIndex)).Get();
+		if (LowestLODAdded.size() == 0)
 		{
 			LowestLODAdded = VariantName;
 		}
 
-		pxr::SdfPath LODPrimPath = ParentPrimPath.AppendPath( pxr::SdfPath( VariantName ) );
+		pxr::SdfPath LODPrimPath = ParentPrimPath.AppendPath(pxr::SdfPath(VariantName));
 
 		// Enable the variant edit context, if we are creating variant LODs
-		TOptional< pxr::UsdEditContext > EditContext;
-		if ( bExportMultipleLODs )
+		TOptional<pxr::UsdEditContext> EditContext;
+		if (bExportMultipleLODs)
 		{
-			pxr::UsdVariantSet VariantSet = VariantSets.GetVariantSet( UnrealIdentifiers::LOD );
-			if ( !VariantSet.AddVariant( VariantName ) )
+			pxr::UsdVariantSet VariantSet = VariantSets.GetVariantSet(UnrealIdentifiers::LOD);
+			if (!VariantSet.AddVariant(VariantName))
 			{
 				continue;
 			}
 
-			VariantSet.SetVariantSelection( VariantName );
-			EditContext.Emplace( VariantSet.GetVariantEditContext() );
+			VariantSet.SetVariantSelection(VariantName);
+			EditContext.Emplace(VariantSet.GetVariantEditContext());
 		}
 
 		pxr::UsdGeomMesh TargetMesh;
-		if ( bExportMultipleLODs )
+		if (bExportMultipleLODs)
 		{
 			// Add the mesh data to a child prim with the Mesh schema
-			pxr::UsdPrim UsdLODPrim = Stage->DefinePrim( LODPrimPath, UnrealToUsd::ConvertToken( TEXT( "Mesh" ) ).Get() );
-			TargetMesh = pxr::UsdGeomMesh{ UsdLODPrim };
+			pxr::UsdPrim UsdLODPrim = Stage->DefinePrim(LODPrimPath, UnrealToUsd::ConvertToken(TEXT("Mesh")).Get());
+			TargetMesh = pxr::UsdGeomMesh{UsdLODPrim};
 		}
 		else
 		{
-			TargetMesh = pxr::UsdGeomMesh{ UsdPrim };
+			TargetMesh = pxr::UsdGeomMesh{UsdPrim};
 		}
 
-		if ( !UsdGeomMeshImpl::ConvertMeshDescription( MeshDescription, TargetMesh, AdditionalTransform, TimeCode ) )
+		if (!UsdGeomMeshImpl::ConvertMeshDescription(MeshDescription, TargetMesh, AdditionalTransform, TimeCode))
 		{
 			return false;
 		}
 	}
 
 	// Reset variant set to start with the lowest lod selected
-	if ( bExportMultipleLODs )
+	if (bExportMultipleLODs)
 	{
-		VariantSets.GetVariantSet( UnrealIdentifiers::LOD ).SetVariantSelection( LowestLODAdded );
+		VariantSets.GetVariantSet(UnrealIdentifiers::LOD).SetVariantSelection(LowestLODAdded);
 	}
 
 	return true;
@@ -3873,7 +3886,14 @@ namespace UE::UsdGeometryCacheConversion::Private
 
 		// ... and adjusting the indices with the proper offset
 		InOutFlattenedMeshData.Indices.Reserve(InOutFlattenedMeshData.Indices.Num() + InMeshData.Indices.Num());
-		Algo::Transform(InMeshData.Indices, InOutFlattenedMeshData.Indices, [VertexIndexOffset](uint32 Index) { return Index + VertexIndexOffset; });
+		Algo::Transform(
+			InMeshData.Indices,
+			InOutFlattenedMeshData.Indices,
+			[VertexIndexOffset](uint32 Index)
+			{
+				return Index + VertexIndexOffset;
+			}
+		);
 
 		// Same with the BatchInfo's StartIndex, which describes where each mesh section starts
 		for (const FGeometryCacheMeshBatchInfo& BatchInfo : InMeshData.BatchesInfo)
@@ -4002,7 +4022,8 @@ namespace UE::UsdGeometryCacheConversion::Private
 				const int32 TexCoordSourceIndex = 0;
 				pxr::TfToken UsdUVSetName = UsdUtils::GetUVSetName(TexCoordSourceIndex).Get();
 
-				pxr::UsdGeomPrimvar PrimvarST = pxr::UsdGeomPrimvarsAPI(MeshPrim).CreatePrimvar(UsdUVSetName, pxr::SdfValueTypeNames->TexCoord2fArray, pxr::UsdGeomTokens->vertex);
+				pxr::UsdGeomPrimvar PrimvarST = pxr::UsdGeomPrimvarsAPI(MeshPrim)
+													.CreatePrimvar(UsdUVSetName, pxr::SdfValueTypeNames->TexCoord2fArray, pxr::UsdGeomTokens->vertex);
 
 				if (PrimvarST)
 				{
@@ -4061,13 +4082,14 @@ namespace UE::UsdGeometryCacheConversion::Private
 					for (int32 VertexIndex = 0; VertexIndex < VertexCount; ++VertexIndex)
 					{
 						// The motion vectors in the MeshData are stored as unit per frame so convert it back to unit per second
-						Velocities.push_back(UnrealToUsd::ConvertVectorFloat(StageInfo, (FVector) -MeshData.MotionVectors[VertexIndex] * ExportContext.FrameRate));
+						Velocities.push_back(
+							UnrealToUsd::ConvertVectorFloat(StageInfo, (FVector)-MeshData.MotionVectors[VertexIndex] * ExportContext.FrameRate)
+						);
 					}
 
 					VelocitiesAttribute.Set(Velocities, TimeCode);
 				}
 			}
-
 		}
 
 		// Faces
@@ -4124,7 +4146,10 @@ namespace UE::UsdGeometryCacheConversion::Private
 			// This mesh has a single material assignment, just add an unrealMaterials attribute to the mesh prim
 			if (MaterialAssignments.size() == 1)
 			{
-				if (pxr::UsdAttribute UEMaterialsAttribute = MaterialPrim.CreateAttribute(UnrealIdentifiers::MaterialAssignment, pxr::SdfValueTypeNames->String))
+				if (pxr::UsdAttribute UEMaterialsAttribute = MaterialPrim.CreateAttribute(
+						UnrealIdentifiers::MaterialAssignment,
+						pxr::SdfValueTypeNames->String
+					))
 				{
 					UEMaterialsAttribute.Set(MaterialAssignments[0]);
 				}
@@ -4167,9 +4192,7 @@ namespace UE::UsdGeometryCacheConversion::Private
 					pxr::UsdPrim MaterialGeomSubsetPrim = GeomSubsetPrim;
 					if (MaterialPrim.GetStage() != MeshPrim.GetStage())
 					{
-						MaterialGeomSubsetPrim = MaterialPrim.GetStage()->OverridePrim(
-							MaterialPrim.GetPath().AppendPath(PrimPath)
-						);
+						MaterialGeomSubsetPrim = MaterialPrim.GetStage()->OverridePrim(MaterialPrim.GetPath().AppendPath(PrimPath));
 					}
 
 					pxr::UsdGeomSubset GeomSubsetSchema{GeomSubsetPrim};
@@ -4183,7 +4206,7 @@ namespace UE::UsdGeometryCacheConversion::Private
 
 						// Indices attribute
 						const uint32 TriangleCount = Section.NumTriangles;
-						const uint32 FirstTriangleIndex = Section.StartIndex / 3; // StartIndex is the first *vertex* instance index
+						const uint32 FirstTriangleIndex = Section.StartIndex / 3;	 // StartIndex is the first *vertex* instance index
 						pxr::VtArray<int> IndicesAttrValue;
 						for (uint32 TriangleIndex = FirstTriangleIndex; TriangleIndex - FirstTriangleIndex < TriangleCount; ++TriangleIndex)
 						{
@@ -4206,7 +4229,10 @@ namespace UE::UsdGeometryCacheConversion::Private
 						pxr::UsdGeomSubset::SetFamilyType(UsdMesh, pxr::UsdShadeTokens->materialBind, pxr::UsdGeomTokens->partition);
 
 						// unrealMaterial attribute
-						if (pxr::UsdAttribute UEMaterialsAttribute = MaterialGeomSubsetPrim.CreateAttribute(UnrealIdentifiers::MaterialAssignment, pxr::SdfValueTypeNames->String))
+						if (pxr::UsdAttribute UEMaterialsAttribute = MaterialGeomSubsetPrim.CreateAttribute(
+								UnrealIdentifiers::MaterialAssignment,
+								pxr::SdfValueTypeNames->String
+							))
 						{
 							UEMaterialsAttribute.Set(SectionMaterial);
 						}
@@ -4215,7 +4241,7 @@ namespace UE::UsdGeometryCacheConversion::Private
 			}
 		}
 	}
-}
+}	 // namespace UE::UsdGeometryCacheConversion::Private
 
 bool UnrealToUsd::ConvertGeometryCache(const UGeometryCache* GeometryCache, pxr::UsdPrim& UsdPrim, UE::FUsdStage* StageForMaterialAssignments)
 {
@@ -4235,7 +4261,7 @@ bool UnrealToUsd::ConvertGeometryCache(const UGeometryCache* GeometryCache, pxr:
 
 	// Collect all material assignments, referenced by the sections' material indices
 	bool bHasMaterialAssignments = false;
-	pxr::VtArray< std::string > MaterialAssignments;
+	pxr::VtArray<std::string> MaterialAssignments;
 	for (const UMaterialInterface* Material : GeometryCache->Materials)
 	{
 		FString AssignedMaterialPathName;
@@ -4337,11 +4363,11 @@ TOptional<UsdUtils::FDisplayColorMaterial> UsdUtils::ExtractDisplayColorMaterial
 
 namespace UE::UsdGeomMeshConversion::Private
 {
-	bool DoesPrimContainMeshLODsInternal( const pxr::UsdPrim& Prim )
+	bool DoesPrimContainMeshLODsInternal(const pxr::UsdPrim& Prim)
 	{
 		FScopedUsdAllocs Allocs;
 
-		if ( !Prim )
+		if (!Prim)
 		{
 			return false;
 		}
@@ -4349,35 +4375,35 @@ namespace UE::UsdGeomMeshConversion::Private
 		const std::string LODString = UnrealIdentifiers::LOD.GetString();
 
 		pxr::UsdVariantSets VariantSets = Prim.GetVariantSets();
-		if ( !VariantSets.HasVariantSet( LODString ) )
+		if (!VariantSets.HasVariantSet(LODString))
 		{
 			return false;
 		}
 
-		std::string Selection = VariantSets.GetVariantSet( LODString ).GetVariantSelection();
-		int32 LODIndex = UsdGeomMeshImpl::GetLODIndexFromName( Selection );
-		if ( LODIndex == INDEX_NONE )
+		std::string Selection = VariantSets.GetVariantSet(LODString).GetVariantSelection();
+		int32 LODIndex = UsdGeomMeshImpl::GetLODIndexFromName(Selection);
+		if (LODIndex == INDEX_NONE)
 		{
 			return false;
 		}
 
 		return true;
 	}
-}
+}	 // namespace UE::UsdGeomMeshConversion::Private
 
-bool UsdUtils::DoesPrimContainMeshLODs( const pxr::UsdPrim& Prim )
+bool UsdUtils::DoesPrimContainMeshLODs(const pxr::UsdPrim& Prim)
 {
-	const bool bHasValidVariantSetup = UE::UsdGeomMeshConversion::Private::DoesPrimContainMeshLODsInternal( Prim );
-	if ( bHasValidVariantSetup )
+	const bool bHasValidVariantSetup = UE::UsdGeomMeshConversion::Private::DoesPrimContainMeshLODsInternal(Prim);
+	if (bHasValidVariantSetup)
 	{
 		FScopedUsdAllocs Allocs;
 
 		// Check if it has at least one mesh too
 		pxr::UsdPrimSiblingRange PrimRange = Prim.GetChildren();
-		for ( pxr::UsdPrimSiblingRange::iterator PrimRangeIt = PrimRange.begin(); PrimRangeIt != PrimRange.end(); ++PrimRangeIt )
+		for (pxr::UsdPrimSiblingRange::iterator PrimRangeIt = PrimRange.begin(); PrimRangeIt != PrimRange.end(); ++PrimRangeIt)
 		{
 			const pxr::UsdPrim& Child = *PrimRangeIt;
-			if ( pxr::UsdGeomMesh ChildMesh{ Child } )
+			if (pxr::UsdGeomMesh ChildMesh{Child})
 			{
 				return true;
 				break;
@@ -4388,12 +4414,12 @@ bool UsdUtils::DoesPrimContainMeshLODs( const pxr::UsdPrim& Prim )
 	return false;
 }
 
-bool UsdUtils::IsGeomMeshALOD( const pxr::UsdPrim& UsdMeshPrim )
+bool UsdUtils::IsGeomMeshALOD(const pxr::UsdPrim& UsdMeshPrim)
 {
 	FScopedUsdAllocs Allocs;
 
-	pxr::UsdGeomMesh UsdMesh{ UsdMeshPrim };
-	if ( !UsdMesh )
+	pxr::UsdGeomMesh UsdMesh{UsdMeshPrim};
+	if (!UsdMesh)
 	{
 		return false;
 	}
@@ -4406,7 +4432,7 @@ bool UsdUtils::IsGeomMeshALOD( const pxr::UsdPrim& UsdMeshPrim )
 	// as the user puts more than one Mesh prim inside of the variant set prim, we're already in a "garbage in"
 	// scenario, and will likely generate some garbage in turn. We'll emit a bunch of warning for that though.
 
-	return UE::UsdGeomMeshConversion::Private::DoesPrimContainMeshLODsInternal( UsdMeshPrim.GetParent() );
+	return UE::UsdGeomMeshConversion::Private::DoesPrimContainMeshLODsInternal(UsdMeshPrim.GetParent());
 }
 
 bool UsdUtils::IsCollisionMesh(const pxr::UsdPrim& UsdPrim)
@@ -4472,22 +4498,22 @@ bool UsdUtils::IsCollisionMesh(const pxr::UsdPrim& UsdPrim)
 	return true;
 }
 
-int32 UsdUtils::GetNumberOfLODVariants( const pxr::UsdPrim& Prim )
+int32 UsdUtils::GetNumberOfLODVariants(const pxr::UsdPrim& Prim)
 {
 	FScopedUsdAllocs Allocs;
 
 	const std::string LODString = UnrealIdentifiers::LOD.GetString();
 
 	pxr::UsdVariantSets VariantSets = Prim.GetVariantSets();
-	if ( !VariantSets.HasVariantSet( LODString ) )
+	if (!VariantSets.HasVariantSet(LODString))
 	{
 		return 1;
 	}
 
-	return VariantSets.GetVariantSet( LODString ).GetVariantNames().size();
+	return VariantSets.GetVariantSet(LODString).GetVariantNames().size();
 }
 
-bool UsdUtils::IterateLODMeshes( const pxr::UsdPrim& ParentPrim, TFunction<bool( const pxr::UsdGeomMesh & LODMesh, int32 LODIndex )> Func )
+bool UsdUtils::IterateLODMeshes(const pxr::UsdPrim& ParentPrim, TFunction<bool(const pxr::UsdGeomMesh& LODMesh, int32 LODIndex)> Func)
 {
 	if (!ParentPrim)
 	{
@@ -4584,17 +4610,17 @@ void UsdUtils::ReplaceUnrealMaterialsWithBaked(
 
 	struct FMaterialScopePrim
 	{
-		FMaterialScopePrim( pxr::UsdStageRefPtr ScopeStage, pxr::UsdPrim ParentPrim )
+		FMaterialScopePrim(pxr::UsdStageRefPtr ScopeStage, pxr::UsdPrim ParentPrim)
 		{
-			pxr::SdfPath Path = ParentPrim.GetPrimPath().AppendPath( pxr::SdfPath{ "Materials" } );
-			Prim = ScopeStage->DefinePrim( Path, UnrealToUsd::ConvertToken( TEXT( "Scope" ) ).Get() );
+			pxr::SdfPath Path = ParentPrim.GetPrimPath().AppendPath(pxr::SdfPath{"Materials"});
+			Prim = ScopeStage->DefinePrim(Path, UnrealToUsd::ConvertToken(TEXT("Scope")).Get());
 
 			// Initialize our UsedPrimNames correctly, so we can guarantee we'll never have name collisions
-			if ( Prim )
+			if (Prim)
 			{
-				for ( pxr::UsdPrim Child : Prim.GetFilteredChildren( pxr::UsdTraverseInstanceProxies( pxr::UsdPrimAllPrimsPredicate ) ) )
+				for (pxr::UsdPrim Child : Prim.GetFilteredChildren(pxr::UsdTraverseInstanceProxies(pxr::UsdPrimAllPrimsPredicate)))
 				{
-					UsedPrimNames.Add( UsdToUnreal::ConvertToken( Child.GetName() ) );
+					UsedPrimNames.Add(UsdToUnreal::ConvertToken(Child.GetName()));
 				}
 			}
 		}
@@ -4603,9 +4629,10 @@ void UsdUtils::ReplaceUnrealMaterialsWithBaked(
 		TSet<FString> UsedPrimNames;
 		TMap<FString, pxr::UsdPrim> BakedFileNameToMatPrim;
 	};
+
 	TOptional<FMaterialScopePrim> StageMatScope;
 
-	pxr::UsdStageRefPtr UsdStage{ Stage };
+	pxr::UsdStageRefPtr UsdStage{Stage};
 
 	// UsedLayers here instead of layer stack because we may be exporting using payloads, and payload layers
 	// don't show up on the layer stack list but do show up on the UsedLayers list
@@ -4615,296 +4642,296 @@ void UsdUtils::ReplaceUnrealMaterialsWithBaked(
 	// This handles Mesh prims as well as GeomSubset prims.
 	// Note how we receive the stage as an argument instead of capturing it from the outer scope:
 	// This ensures the inner function doesn't hold a reference to the stage
-	TFunction<void( pxr::UsdStageRefPtr StageToTraverse, pxr::UsdPrim Prim, TOptional<FMaterialScopePrim>& MatPrimScope, TOptional<pxr::UsdVariantSet> OuterVariantSet )> TraverseForMaterialReplacement;
+	TFunction<void(
+		pxr::UsdStageRefPtr StageToTraverse,
+		pxr::UsdPrim Prim,
+		TOptional<FMaterialScopePrim> & MatPrimScope,
+		TOptional<pxr::UsdVariantSet> OuterVariantSet
+	)>
+		TraverseForMaterialReplacement;
 	TraverseForMaterialReplacement =
-		[
-			&TraverseForMaterialReplacement,
-			&LayersToTraverse,
-			&LayerToAuthorIn,
-			&BakedMaterials,
-			bIsAssetLayer,
-			bUsePayload,
-			&StageMatScope
-		]
-		(
+		[&TraverseForMaterialReplacement, &LayersToTraverse, &LayerToAuthorIn, &BakedMaterials, bIsAssetLayer, bUsePayload, &StageMatScope](
 			pxr::UsdStageRefPtr StageToTraverse,
 			pxr::UsdPrim Prim,
 			TOptional<FMaterialScopePrim>& MatPrimScope,
 			TOptional<pxr::UsdVariantSet> OuterVariantSet
 		)
+	{
+		// Recurse into children before doing anything as we may need to parse LODs
+		pxr::UsdVariantSet VarSet = Prim.GetVariantSet(UnrealIdentifiers::LOD);
+		std::vector<std::string> LODs = VarSet.GetVariantNames();
+		if (LODs.size() > 0)
 		{
-			// Recurse into children before doing anything as we may need to parse LODs
-			pxr::UsdVariantSet VarSet = Prim.GetVariantSet( UnrealIdentifiers::LOD );
-			std::vector<std::string> LODs = VarSet.GetVariantNames();
-			if ( LODs.size() > 0 )
-			{
-				TOptional<std::string> OriginalSelection = VarSet.HasAuthoredVariantSelection() ? VarSet.GetVariantSelection() : TOptional<std::string>{};
+			TOptional<std::string> OriginalSelection = VarSet.HasAuthoredVariantSelection() ? VarSet.GetVariantSelection() : TOptional<std::string>{};
 
-				// Prims within variant sets can't have relationships to prims outside the scope of the prim that
-				// contains the variant set itself. This means we'll need a new material scope prim if we're stepping
-				// into a variant within an asset layer, so that any material proxy prims we author are contained within it.
-				// Note that we only do this for asset layers: If we're parsing the root layer, any LOD variant sets we can step into
-				// are brought in via references to asset files, and we know that referenced subtree only has relationships to
-				// things within that same subtree ( which will be entirely brought in to the root layer ). This means we can
-				// just keep inner_mat_prim_scope as None and default to using the layer's mat scope prim if we need one
-				TOptional<FMaterialScopePrim> InnerMatPrimScope = bIsAssetLayer ? FMaterialScopePrim{ StageToTraverse, Prim } : TOptional<FMaterialScopePrim>{};
+			// Prims within variant sets can't have relationships to prims outside the scope of the prim that
+			// contains the variant set itself. This means we'll need a new material scope prim if we're stepping
+			// into a variant within an asset layer, so that any material proxy prims we author are contained within it.
+			// Note that we only do this for asset layers: If we're parsing the root layer, any LOD variant sets we can step into
+			// are brought in via references to asset files, and we know that referenced subtree only has relationships to
+			// things within that same subtree ( which will be entirely brought in to the root layer ). This means we can
+			// just keep inner_mat_prim_scope as None and default to using the layer's mat scope prim if we need one
+			TOptional<FMaterialScopePrim> InnerMatPrimScope = bIsAssetLayer ? FMaterialScopePrim{StageToTraverse, Prim}
+																			: TOptional<FMaterialScopePrim>{};
 
-				// Switch into each of the LOD variants the prim has, and recurse into the child prims
-				for ( const std::string& Variant : LODs )
+			// Switch into each of the LOD variants the prim has, and recurse into the child prims
+			for (const std::string& Variant : LODs)
+			{
 				{
-					{
-						pxr::UsdEditContext Context{ StageToTraverse, StageToTraverse->GetSessionLayer() };
-						VarSet.SetVariantSelection( Variant );
-					}
-
-					for ( const pxr::UsdPrim& Child : Prim.GetChildren() )
-					{
-						TraverseForMaterialReplacement( StageToTraverse, Child, InnerMatPrimScope, VarSet );
-					}
+					pxr::UsdEditContext Context{StageToTraverse, StageToTraverse->GetSessionLayer()};
+					VarSet.SetVariantSelection(Variant);
 				}
 
-				// Restore the variant selection to what it originally was
-				pxr::UsdEditContext Context{ StageToTraverse, StageToTraverse->GetSessionLayer() };
-				if ( OriginalSelection.IsSet() )
+				for (const pxr::UsdPrim& Child : Prim.GetChildren())
 				{
-					VarSet.SetVariantSelection( OriginalSelection.GetValue() );
-				}
-				else
-				{
-					VarSet.ClearVariantSelection();
+					TraverseForMaterialReplacement(StageToTraverse, Child, InnerMatPrimScope, VarSet);
 				}
 			}
+
+			// Restore the variant selection to what it originally was
+			pxr::UsdEditContext Context{StageToTraverse, StageToTraverse->GetSessionLayer()};
+			if (OriginalSelection.IsSet())
+			{
+				VarSet.SetVariantSelection(OriginalSelection.GetValue());
+			}
 			else
 			{
-				for ( const pxr::UsdPrim& Child : Prim.GetChildren() )
-				{
-					TraverseForMaterialReplacement( StageToTraverse, Child, MatPrimScope, OuterVariantSet );
-				}
+				VarSet.ClearVariantSelection();
 			}
-
-			// Don't try fetching attributes from the pseudoroot as we'll obviously never have a material binding here
-			// and we may get some USD warnings
-			if ( Prim.IsPseudoRoot() )
+		}
+		else
+		{
+			for (const pxr::UsdPrim& Child : Prim.GetChildren())
 			{
-				return;
+				TraverseForMaterialReplacement(StageToTraverse, Child, MatPrimScope, OuterVariantSet);
 			}
+		}
+
+		// Don't try fetching attributes from the pseudoroot as we'll obviously never have a material binding here
+		// and we may get some USD warnings
+		if (Prim.IsPseudoRoot())
+		{
+			return;
+		}
 
-			std::string UnrealMaterialAttrAssetPath;
-			FString UnrealMaterialPrimAssetPath;
+		std::string UnrealMaterialAttrAssetPath;
+		FString UnrealMaterialPrimAssetPath;
 
-			pxr::UsdAttribute UnrealMaterialAttr = Prim.GetAttribute( UnrealIdentifiers::MaterialAssignment );
-			pxr::UsdShadeMaterial UnrealMaterial;
+		pxr::UsdAttribute UnrealMaterialAttr = Prim.GetAttribute(UnrealIdentifiers::MaterialAssignment);
+		pxr::UsdShadeMaterial UnrealMaterial;
 
-			pxr::UsdShadeMaterialBindingAPI MaterialBindingAPI{ Prim };
-			if ( MaterialBindingAPI )
+		pxr::UsdShadeMaterialBindingAPI MaterialBindingAPI{Prim};
+		if (MaterialBindingAPI)
+		{
+			// We always emit UnrealMaterials with allpurpose bindings, so we can use default arguments for
+			// ComputeBoundMaterial
+			if (pxr::UsdShadeMaterial BoundMaterial = MaterialBindingAPI.ComputeBoundMaterial())
 			{
-				// We always emit UnrealMaterials with allpurpose bindings, so we can use default arguments for
-				// ComputeBoundMaterial
-				if ( pxr::UsdShadeMaterial BoundMaterial = MaterialBindingAPI.ComputeBoundMaterial() )
-				{
-					UnrealMaterial = BoundMaterial;
+				UnrealMaterial = BoundMaterial;
 
-					TOptional<FString> ExistingUEAssetReference = UsdUtils::GetUnrealSurfaceOutput( UnrealMaterial.GetPrim() );
-					if ( ExistingUEAssetReference.IsSet() )
-					{
-						UnrealMaterialPrimAssetPath = MoveTemp( ExistingUEAssetReference.GetValue() );
-					}
+				TOptional<FString> ExistingUEAssetReference = UsdUtils::GetUnrealSurfaceOutput(UnrealMaterial.GetPrim());
+				if (ExistingUEAssetReference.IsSet())
+				{
+					UnrealMaterialPrimAssetPath = MoveTemp(ExistingUEAssetReference.GetValue());
 				}
 			}
+		}
 
-			if ( !UnrealMaterial && ( !UnrealMaterialAttr || !UnrealMaterialAttr.Get<std::string>( &UnrealMaterialAttrAssetPath ) ) )
-			{
-				return;
-			}
+		if (!UnrealMaterial && (!UnrealMaterialAttr || !UnrealMaterialAttr.Get<std::string>(&UnrealMaterialAttrAssetPath)))
+		{
+			return;
+		}
 
-			pxr::UsdPrim UnrealMaterialPrim = UnrealMaterial.GetPrim();
+		pxr::UsdPrim UnrealMaterialPrim = UnrealMaterial.GetPrim();
 
-			// Prioritize the Unreal material since import will do so too
-			FString UnrealMaterialAssetPath = UnrealMaterialPrimAssetPath.IsEmpty()
-				? UsdToUnreal::ConvertString( UnrealMaterialAttrAssetPath )
-				: UnrealMaterialPrimAssetPath;
+		// Prioritize the Unreal material since import will do so too
+		FString UnrealMaterialAssetPath = UnrealMaterialPrimAssetPath.IsEmpty() ? UsdToUnreal::ConvertString(UnrealMaterialAttrAssetPath)
+																				: UnrealMaterialPrimAssetPath;
 
-			FString BakedFilename = BakedMaterials.FindRef( UnrealMaterialAssetPath );
+		FString BakedFilename = BakedMaterials.FindRef(UnrealMaterialAssetPath);
 
-			// If we have a valid UE asset but just haven't baked it, something went wrong: Just leave everything alone and abort
-			if ( !UnrealMaterialAssetPath.IsEmpty() && BakedFilename.IsEmpty() )
-			{
-				return;
-			}
+		// If we have a valid UE asset but just haven't baked it, something went wrong: Just leave everything alone and abort
+		if (!UnrealMaterialAssetPath.IsEmpty() && BakedFilename.IsEmpty())
+		{
+			return;
+		}
 
-			pxr::SdfPath UnrealMaterialAttrPath = UnrealMaterialAttr ? UnrealMaterialAttr.GetPath() : pxr::SdfPath{};
-			pxr::SdfPath UnrealMaterialPrimPath = UnrealMaterial ? UnrealMaterialPrim.GetPrimPath() : pxr::SdfPath{};
+		pxr::SdfPath UnrealMaterialAttrPath = UnrealMaterialAttr ? UnrealMaterialAttr.GetPath() : pxr::SdfPath{};
+		pxr::SdfPath UnrealMaterialPrimPath = UnrealMaterial ? UnrealMaterialPrim.GetPrimPath() : pxr::SdfPath{};
 
-			// Find out if we need to remove / author material bindings within an actual variant or outside of it, as an over.
-			// We don't do this when using payloads because our override prims aren't inside the actual LOD variants : They just
-			// directly override a mesh called e.g. 'LOD3' as if it's a child prim, so that the override automatically only
-			// does anything when we happen to have the variant that enables the LOD3 Mesh
-			const bool bAuthorInsideVariants = OuterVariantSet.IsSet() && bIsAssetLayer && !bUsePayload;
+		// Find out if we need to remove / author material bindings within an actual variant or outside of it, as an over.
+		// We don't do this when using payloads because our override prims aren't inside the actual LOD variants : They just
+		// directly override a mesh called e.g. 'LOD3' as if it's a child prim, so that the override automatically only
+		// does anything when we happen to have the variant that enables the LOD3 Mesh
+		const bool bAuthorInsideVariants = OuterVariantSet.IsSet() && bIsAssetLayer && !bUsePayload;
 
-			if ( bAuthorInsideVariants )
+		if (bAuthorInsideVariants)
+		{
+			pxr::UsdVariantSet& OuterVariantSetValue = OuterVariantSet.GetValue();
+			pxr::SdfPath VarPrimPath = OuterVariantSetValue.GetPrim().GetPath();
+			pxr::SdfPath VarPrimPathWithVar = VarPrimPath.AppendVariantSelection(
+				OuterVariantSetValue.GetName(),
+				OuterVariantSetValue.GetVariantSelection()
+			);
+
+			if (UnrealMaterialAttrPath.HasPrefix(VarPrimPath))
 			{
-				pxr::UsdVariantSet& OuterVariantSetValue = OuterVariantSet.GetValue();
-				pxr::SdfPath VarPrimPath = OuterVariantSetValue.GetPrim().GetPath();
-				pxr::SdfPath VarPrimPathWithVar = VarPrimPath.AppendVariantSelection( OuterVariantSetValue.GetName(), OuterVariantSetValue.GetVariantSelection() );
+				// This builds a path like '/MyMesh{LOD=LOD0}LOD0.unrealMaterial',
+				// or '/MyMesh{LOD=LOD0}LOD0/Section1.unrealMaterial'.This is required because we'll query the layer
+				// for a spec path below, and this path must contain the variant selection in it, which the path returned
+				// from attr.GetPath() doesn't contain
+				UnrealMaterialAttrPath = UnrealMaterialAttrPath.ReplacePrefix(VarPrimPath, VarPrimPathWithVar);
+			}
 
-				if ( UnrealMaterialAttrPath.HasPrefix( VarPrimPath ) )
-				{
-					// This builds a path like '/MyMesh{LOD=LOD0}LOD0.unrealMaterial',
-					// or '/MyMesh{LOD=LOD0}LOD0/Section1.unrealMaterial'.This is required because we'll query the layer
-					// for a spec path below, and this path must contain the variant selection in it, which the path returned
-					// from attr.GetPath() doesn't contain
-					UnrealMaterialAttrPath = UnrealMaterialAttrPath.ReplacePrefix( VarPrimPath, VarPrimPathWithVar );
-				}
+			if (UnrealMaterialPrimPath.HasPrefix(VarPrimPath))
+			{
+				UnrealMaterialPrimPath = UnrealMaterialPrimPath.ReplacePrefix(VarPrimPath, VarPrimPathWithVar);
+			}
+		}
 
-				if ( UnrealMaterialPrimPath.HasPrefix( VarPrimPath ) )
-				{
-					UnrealMaterialPrimPath = UnrealMaterialPrimPath.ReplacePrefix( VarPrimPath, VarPrimPathWithVar );
-				}
+		// We always want to replace things in whatever layer they were authored, and not just override with
+		// a stronger opinion, so search through all sublayers to find the ones with the specs we are targeting
+		for (pxr::SdfLayerHandle Layer : LayersToTraverse)
+		{
+			pxr::SdfAttributeSpecHandle UnrealMaterialAttrSpec = Layer->GetAttributeAtPath(UnrealMaterialAttrPath);
+			pxr::SdfPrimSpecHandle UnrealMaterialPrimSpec = Layer->GetPrimAtPath(UnrealMaterialPrimPath);
+			if (!UnrealMaterialAttrSpec && !UnrealMaterialPrimSpec)
+			{
+				continue;
 			}
 
-			// We always want to replace things in whatever layer they were authored, and not just override with
-			// a stronger opinion, so search through all sublayers to find the ones with the specs we are targeting
-			for ( pxr::SdfLayerHandle Layer : LayersToTraverse )
+			pxr::UsdEditContext Context{StageToTraverse, Layer};
+
+			// It was just an empty UE asset path, so just cancel now as our BakedFilename can't possibly be useful
+			if (UnrealMaterialAssetPath.IsEmpty())
 			{
-				pxr::SdfAttributeSpecHandle UnrealMaterialAttrSpec = Layer->GetAttributeAtPath( UnrealMaterialAttrPath );
-				pxr::SdfPrimSpecHandle UnrealMaterialPrimSpec = Layer->GetPrimAtPath( UnrealMaterialPrimPath );
-				if ( !UnrealMaterialAttrSpec && !UnrealMaterialPrimSpec )
-				{
-					continue;
-				}
+				continue;
+			}
 
-				pxr::UsdEditContext Context{ StageToTraverse, Layer };
+			// Get the proxy prim for the material within this layer
+			// (or create one outside the variant edit context)
+			pxr::UsdPrim MatPrim;
+			{
+				pxr::UsdEditContext MatContext(StageToTraverse, pxr::SdfLayerRefPtr{LayerToAuthorIn});
 
-				// It was just an empty UE asset path, so just cancel now as our BakedFilename can't possibly be useful
-				if ( UnrealMaterialAssetPath.IsEmpty() )
+				// We are already referencing an unreal material prim: Let's just augment it with a reference to the baked
+				// material usd asset layer.
+				// Note how this will likely not be within MatPrimScope but instead will be a child of the Mesh/GeomSubset.
+				// This is fine, and in the future we'll likely exclusively do this since it will handle mesh-specific
+				// material baking much better, as it will allow even having separate bakes for each LOD
+				if (UnrealMaterial && UnrealMaterialPrimSpec)
 				{
-					continue;
-				}
+					MatPrim = UnrealMaterialPrim;
 
-				// Get the proxy prim for the material within this layer
-				// (or create one outside the variant edit context)
-				pxr::UsdPrim MatPrim;
-				{
-					pxr::UsdEditContext MatContext( StageToTraverse, pxr::SdfLayerRefPtr{ LayerToAuthorIn } );
-
-					// We are already referencing an unreal material prim: Let's just augment it with a reference to the baked
-					// material usd asset layer.
-					// Note how this will likely not be within MatPrimScope but instead will be a child of the Mesh/GeomSubset.
-					// This is fine, and in the future we'll likely exclusively do this since it will handle mesh-specific
-					// material baking much better, as it will allow even having separate bakes for each LOD
-					if ( UnrealMaterial && UnrealMaterialPrimSpec )
-					{
-						MatPrim = UnrealMaterialPrim;
+					bool bAlreadyHasReference = false;
 
-						bool bAlreadyHasReference = false;
+					// Make sure we don't reference it more than once. This shouldn't be needed since we'll only ever run into
+					// these unreal material prims once per Mesh/GeomSubset, but when creating MatScopePrims we can guarantee we
+					// add a reference only once by adding it along with the Material prim creation, so it would be nice to be able to
+					// guarantee it here as well
+					pxr::SdfReferencesProxy ReferencesProxy = UnrealMaterialPrimSpec->GetReferenceList();
+					for (const pxr::SdfReference& UsdReference : ReferencesProxy.GetAddedOrExplicitItems())
+					{
+						FString ReferencedFilePath = UsdToUnreal::ConvertString(UsdReference.GetAssetPath());
+						FString LayerPath = UsdToUnreal::ConvertString(Layer->GetRealPath());
 
-						// Make sure we don't reference it more than once. This shouldn't be needed since we'll only ever run into
-						// these unreal material prims once per Mesh/GeomSubset, but when creating MatScopePrims we can guarantee we
-						// add a reference only once by adding it along with the Material prim creation, so it would be nice to be able to
-						// guarantee it here as well
-						pxr::SdfReferencesProxy ReferencesProxy = UnrealMaterialPrimSpec->GetReferenceList();
-						for ( const pxr::SdfReference& UsdReference : ReferencesProxy.GetAddedOrExplicitItems() )
+						if (!LayerPath.IsEmpty())
 						{
-							FString ReferencedFilePath = UsdToUnreal::ConvertString( UsdReference.GetAssetPath() );
-							FString LayerPath = UsdToUnreal::ConvertString( Layer->GetRealPath() );
-
-							if ( !LayerPath.IsEmpty() )
-							{
-								ReferencedFilePath = FPaths::ConvertRelativePathToFull( LayerPath, ReferencedFilePath );
-							}
-
-							if ( FPaths::IsSamePath( ReferencedFilePath, BakedFilename ) )
-							{
-								bAlreadyHasReference = true;
-								break;
-							}
+							ReferencedFilePath = FPaths::ConvertRelativePathToFull(LayerPath, ReferencedFilePath);
 						}
 
-						if ( !bAlreadyHasReference )
+						if (FPaths::IsSamePath(ReferencedFilePath, BakedFilename))
 						{
-							UE::FUsdPrim UEMatPrim{ MatPrim };
-							UsdUtils::AddReference( UEMatPrim, *BakedFilename );
+							bAlreadyHasReference = true;
+							break;
 						}
 					}
-					// Need a MatScopePrim authored somewhere within this layer
-					else
+
+					if (!bAlreadyHasReference)
 					{
-						FMaterialScopePrim* MatPrimScopePtr = nullptr;
+						UE::FUsdPrim UEMatPrim{MatPrim};
+						UsdUtils::AddReference(UEMatPrim, *BakedFilename);
+					}
+				}
+				// Need a MatScopePrim authored somewhere within this layer
+				else
+				{
+					FMaterialScopePrim* MatPrimScopePtr = nullptr;
 
-						if ( MatPrimScope.IsSet() )
-						{
-							MatPrimScopePtr = &MatPrimScope.GetValue();
-						}
-						else
+					if (MatPrimScope.IsSet())
+					{
+						MatPrimScopePtr = &MatPrimScope.GetValue();
+					}
+					else
+					{
+						// On-demand create a *single* material scope prim for the stage, if we're not inside a variant set
+						if (!StageMatScope.IsSet())
 						{
-							// On-demand create a *single* material scope prim for the stage, if we're not inside a variant set
-							if ( !StageMatScope.IsSet() )
-							{
-								// If a prim from a stage references another layer, USD's composition will effectively
-								// paste the default prim of the referenced layer over the referencing prim. Because of
-								// this, the subprims within the hierarchy of that default prim can't ever have
-								// relationships to other prims outside that of that same hierarchy, as those prims
-								// will not be present on the referencing stage at all. This is why we author our stage
-								// materials scope under the default prim, and not the pseudoroot
-								StageMatScope = FMaterialScopePrim{ StageToTraverse, StageToTraverse->GetDefaultPrim() };
-							}
-							MatPrimScopePtr = &StageMatScope.GetValue();
+							// If a prim from a stage references another layer, USD's composition will effectively
+							// paste the default prim of the referenced layer over the referencing prim. Because of
+							// this, the subprims within the hierarchy of that default prim can't ever have
+							// relationships to other prims outside that of that same hierarchy, as those prims
+							// will not be present on the referencing stage at all. This is why we author our stage
+							// materials scope under the default prim, and not the pseudoroot
+							StageMatScope = FMaterialScopePrim{StageToTraverse, StageToTraverse->GetDefaultPrim()};
 						}
+						MatPrimScopePtr = &StageMatScope.GetValue();
+					}
 
-						// This should never happen
-						if ( !ensure( MatPrimScopePtr ) )
-						{
-							continue;
-						}
+					// This should never happen
+					if (!ensure(MatPrimScopePtr))
+					{
+						continue;
+					}
 
-						// We already have a material proxy prim for this UE material within MatPrimScope, so just reuse it
-						if ( pxr::UsdPrim* FoundPrim = MatPrimScopePtr->BakedFileNameToMatPrim.Find( BakedFilename ) )
-						{
-							MatPrim = *FoundPrim;
-						}
-						// Create a new material proxy prim for this UE material within MatPrimScope
-						else
-						{
-							FString MatName = FPaths::GetBaseFilename( UnrealMaterialAssetPath );
-							MatName = UsdToUnreal::ConvertString( pxr::TfMakeValidIdentifier( UnrealToUsd::ConvertString( *MatName ).Get() ) );
-							FString MatPrimName = UsdUtils::GetUniqueName( MatName, MatPrimScopePtr->UsedPrimNames );
-							MatPrimScopePtr->UsedPrimNames.Add( MatPrimName );
-
-							MatPrim = StageToTraverse->DefinePrim(
-								MatPrimScopePtr->Prim.GetPath().AppendChild( UnrealToUsd::ConvertToken( *MatPrimName ).Get() ),
-								UnrealToUsd::ConvertToken( TEXT( "Material" ) ).Get()
-							);
+					// We already have a material proxy prim for this UE material within MatPrimScope, so just reuse it
+					if (pxr::UsdPrim* FoundPrim = MatPrimScopePtr->BakedFileNameToMatPrim.Find(BakedFilename))
+					{
+						MatPrim = *FoundPrim;
+					}
+					// Create a new material proxy prim for this UE material within MatPrimScope
+					else
+					{
+						FString MatName = FPaths::GetBaseFilename(UnrealMaterialAssetPath);
+						MatName = UsdToUnreal::ConvertString(pxr::TfMakeValidIdentifier(UnrealToUsd::ConvertString(*MatName).Get()));
+						FString MatPrimName = UsdUtils::GetUniqueName(MatName, MatPrimScopePtr->UsedPrimNames);
+						MatPrimScopePtr->UsedPrimNames.Add(MatPrimName);
+
+						MatPrim = StageToTraverse->DefinePrim(
+							MatPrimScopePtr->Prim.GetPath().AppendChild(UnrealToUsd::ConvertToken(*MatPrimName).Get()),
+							UnrealToUsd::ConvertToken(TEXT("Material")).Get()
+						);
 
-							// We should only keep track and reuse the material proxy prims that we create within the MatPrimScope, not
-							// the ones we have appropriated from within Mesh/GeomSubset from being UnrealPrims
-							MatPrimScopePtr->BakedFileNameToMatPrim.Add( BakedFilename, MatPrim );
+						// We should only keep track and reuse the material proxy prims that we create within the MatPrimScope, not
+						// the ones we have appropriated from within Mesh/GeomSubset from being UnrealPrims
+						MatPrimScopePtr->BakedFileNameToMatPrim.Add(BakedFilename, MatPrim);
 
-							UE::FUsdPrim UEMatPrim{ MatPrim };
-							UsdUtils::AddReference( UEMatPrim, *BakedFilename );
-						}
+						UE::FUsdPrim UEMatPrim{MatPrim};
+						UsdUtils::AddReference(UEMatPrim, *BakedFilename);
 					}
 				}
+			}
 
-				// Make sure we have a binding to the material prim and the material binding API
-				if( pxr::UsdShadeMaterial MaterialToBind{ MatPrim } )
+			// Make sure we have a binding to the material prim and the material binding API
+			if (pxr::UsdShadeMaterial MaterialToBind{MatPrim})
+			{
+				TOptional<pxr::UsdEditContext> VarContext;
+				if (bAuthorInsideVariants)
 				{
-					TOptional<pxr::UsdEditContext> VarContext;
-					if ( bAuthorInsideVariants )
-					{
-						VarContext.Emplace( OuterVariantSet.GetValue().GetVariantEditContext() );
-					}
+					VarContext.Emplace(OuterVariantSet.GetValue().GetVariantEditContext());
+				}
 
-					if ( pxr::UsdShadeMaterialBindingAPI AppliedMaterialBindingAPI = pxr::UsdShadeMaterialBindingAPI::Apply( Prim ) )
-					{
-						AppliedMaterialBindingAPI.Bind( MaterialToBind );
-					}
+				if (pxr::UsdShadeMaterialBindingAPI AppliedMaterialBindingAPI = pxr::UsdShadeMaterialBindingAPI::Apply(Prim))
+				{
+					AppliedMaterialBindingAPI.Bind(MaterialToBind);
 				}
 			}
-		};
+		}
+	};
 
 	pxr::UsdPrim Root = Stage.GetPseudoRoot();
 	TOptional<FMaterialScopePrim> Empty;
-	TraverseForMaterialReplacement( UsdStage, Root, Empty, {} );
+	TraverseForMaterialReplacement(UsdStage, Root, Empty, {});
 }
 
 void UsdUtils::ReplaceUnrealMaterialsWithBaked(
@@ -4916,13 +4943,7 @@ void UsdUtils::ReplaceUnrealMaterialsWithBaked(
 	bool bRemoveUnrealMaterials
 )
 {
-	ReplaceUnrealMaterialsWithBaked(
-		Stage,
-		LayerToAuthorIn,
-		BakedMaterials,
-		bIsAssetLayer,
-		bUsePayload
-	);
+	ReplaceUnrealMaterialsWithBaked(Stage, LayerToAuthorIn, BakedMaterials, bIsAssetLayer, bUsePayload);
 }
 
 namespace UE::UsdGeomMeshConversion::Private
@@ -4964,7 +4985,7 @@ namespace UE::UsdGeomMeshConversion::Private
 			MD5.Update(reinterpret_cast<const uint8*>(Token.data()), Token.size());
 		}
 	}
-}
+}	 // namespace UE::UsdGeomMeshConversion::Private
 
 FString UsdUtils::HashGeomMeshPrim(const UE::FUsdStage& Stage, const FString& PrimPath, double TimeCode)
 {
@@ -4994,14 +5015,14 @@ void UsdUtils::HashGeomMeshPrim(const UE::FUsdStage& Stage, const FString& PrimP
 
 	FScopedUsdAllocs Allocs;
 
-	UsdPrim UsdPrim = pxr::UsdPrim( Stage.GetPrimAtPath( UE::FSdfPath( *PrimPath ) ) );
-	if ( !UsdPrim )
+	UsdPrim UsdPrim = pxr::UsdPrim(Stage.GetPrimAtPath(UE::FSdfPath(*PrimPath)));
+	if (!UsdPrim)
 	{
 		return;
 	}
 
-	UsdGeomMesh UsdMesh( UsdPrim );
-	if ( !UsdMesh )
+	UsdGeomMesh UsdMesh(UsdPrim);
+	if (!UsdMesh)
 	{
 		return;
 	}
@@ -5033,10 +5054,10 @@ void UsdUtils::HashGeomMeshPrim(const UE::FUsdStage& Stage, const FString& PrimP
 
 	// TODO: This is not providing render context or material purpose, so it will never consider float2f primvars
 	// for the hash, which could be an issue in very exotic cases
-	TArray< TUsdStore< UsdGeomPrimvar > > PrimvarsByUVIndex = UsdUtils::GetUVSetPrimvars( UsdPrim );
-	for ( int32 UVChannelIndex = 0; UVChannelIndex < PrimvarsByUVIndex.Num(); ++UVChannelIndex )
+	TArray<TUsdStore<UsdGeomPrimvar>> PrimvarsByUVIndex = UsdUtils::GetUVSetPrimvars(UsdPrim);
+	for (int32 UVChannelIndex = 0; UVChannelIndex < PrimvarsByUVIndex.Num(); ++UVChannelIndex)
 	{
-		if ( !PrimvarsByUVIndex.IsValidIndex( UVChannelIndex ) )
+		if (!PrimvarsByUVIndex.IsValidIndex(UVChannelIndex))
 		{
 			break;
 		}
@@ -5045,7 +5066,13 @@ void UsdUtils::HashGeomMeshPrim(const UE::FUsdStage& Stage, const FString& PrimP
 	}
 }
 
-bool UsdUtils::GetPointInstancerTransforms(const FUsdStageInfo& StageInfo, const pxr::UsdGeomPointInstancer& PointInstancer, const int32 ProtoIndex, pxr::UsdTimeCode EvalTime, TArray<FTransform>& OutInstanceTransforms)
+bool UsdUtils::GetPointInstancerTransforms(
+	const FUsdStageInfo& StageInfo,
+	const pxr::UsdGeomPointInstancer& PointInstancer,
+	const int32 ProtoIndex,
+	pxr::UsdTimeCode EvalTime,
+	TArray<FTransform>& OutInstanceTransforms
+)
 {
 	TRACE_CPUPROFILER_EVENT_SCOPE(GetPointInstancerTransforms);
 
@@ -5069,8 +5096,8 @@ bool UsdUtils::GetPointInstancerTransforms(const FUsdStageInfo& StageInfo, const
 	int32 Index = 0;
 
 	const int32 NumInstances = GMaxInstancesPerPointInstancer >= 0
-		? FMath::Min(static_cast<int32>(UsdInstanceTransforms.size()), GMaxInstancesPerPointInstancer)
-		: static_cast<int32>(UsdInstanceTransforms.size());
+								   ? FMath::Min(static_cast<int32>(UsdInstanceTransforms.size()), GMaxInstancesPerPointInstancer)
+								   : static_cast<int32>(UsdInstanceTransforms.size());
 
 	{
 		FScopedUnrealAllocs UnrealAllocs;
@@ -5293,4 +5320,4 @@ void UsdUtils::AuthorIdentityTransformGprimAttributes(const pxr::UsdPrim& UsdPri
 
 #undef LOCTEXT_NAMESPACE
 
-#endif // #if USE_USD_SDK
+#endif	  // #if USE_USD_SDK
