USD: Reformat USDUtilities module.

diff --git a/Engine/Plugins/Importers/USDImporter/Source/USDUtilities/Private/USDGroomConversion.cpp b/Engine/Plugins/Importers/USDImporter/Source/USDUtilities/Private/USDGroomConversion.cpp
--- a/Engine/Plugins/Importers/USDImporter/Source/USDUtilities/Private/USDGroomConversion.cpp
+++ b/Engine/Plugins/Importers/USDImporter/Source/USDUtilities/Private/USDGroomConversion.cpp
@@ -91,8 +91,16 @@ namespace UE::UsdGroomConversion::Private
 	};
 
 	/** Convert the given primvar to hair attribute in the proper scope */
-	template <typename USDType, typename UEType>
-	void ConvertPrimvar(const pxr::UsdGeomPrimvar& Primvar, const pxr::UsdTimeCode& TimeCode, int32 StartStrandID, int32 NumStrands, int32 StartVertexID, int32 NumVertices, FHairDescription& HairDescription)
+	template<typename USDType, typename UEType>
+	void ConvertPrimvar(
+		const pxr::UsdGeomPrimvar& Primvar,
+		const pxr::UsdTimeCode& TimeCode,
+		int32 StartStrandID,
+		int32 NumStrands,
+		int32 StartVertexID,
+		int32 NumVertices,
+		FHairDescription& HairDescription
+	)
 	{
 		FString PrimvarName = UsdToUnreal::ConvertToken(Primvar.GetPrimvarName());
 		FName AttributeName(*PrimvarName);
@@ -128,7 +136,9 @@ namespace UE::UsdGroomConversion::Private
 			bool bIsConstantValue = Interpolation == pxr::UsdGeomTokens->constant;
 			for (int32 StrandIndex = 0; StrandIndex < NumStrands; ++StrandIndex)
 			{
-				StrandAttributeRef[FStrandID(StartStrandID + StrandIndex)] = TConverter<UEType>::ConvertType(Values[bIsConstantValue ? 0 : StrandIndex]);
+				StrandAttributeRef[FStrandID(StartStrandID + StrandIndex)] = TConverter<UEType>::ConvertType(
+					Values[bIsConstantValue ? 0 : StrandIndex]
+				);
 			}
 		}
 		else if (Interpolation == pxr::UsdGeomTokens->vertex || NumValues == NumVertices)
@@ -148,7 +158,15 @@ namespace UE::UsdGroomConversion::Private
 	}
 
 	/** Convert the given primvars to hair attributes */
-	void ConvertPrimvars(const std::vector<pxr::UsdGeomPrimvar>& Attributes, const pxr::UsdTimeCode& TimeCode, int32 StartStrandID, int32 NumStrands, int32 StartVertexID, int32 NumVertices, FHairDescription& HairDescription)
+	void ConvertPrimvars(
+		const std::vector<pxr::UsdGeomPrimvar>& Attributes,
+		const pxr::UsdTimeCode& TimeCode,
+		int32 StartStrandID,
+		int32 NumStrands,
+		int32 StartVertexID,
+		int32 NumVertices,
+		FHairDescription& HairDescription
+	)
 	{
 		for (const pxr::UsdGeomPrimvar& Primvar : Attributes)
 		{
@@ -208,11 +226,17 @@ namespace UE::UsdGroomConversion::Private
 			}
 		}
 	}
-}
+}	 // namespace UE::UsdGroomConversion::Private
 
 namespace UsdToUnreal
 {
-	bool ConvertGroomHierarchy(const pxr::UsdPrim& Prim, const pxr::UsdTimeCode& TimeCode, const FTransform& ParentTransform, FHairDescription& HairDescription, FGroomAnimationInfo* AnimInfo)
+	bool ConvertGroomHierarchy(
+		const pxr::UsdPrim& Prim,
+		const pxr::UsdTimeCode& TimeCode,
+		const FTransform& ParentTransform,
+		FHairDescription& HairDescription,
+		FGroomAnimationInfo* AnimInfo
+	)
 	{
 		TRACE_CPUPROFILER_EVENT_SCOPE(UsdToUnreal::ConvertGroomHierarchy);
 
@@ -296,11 +320,14 @@ namespace UsdToUnreal
 				pxr::TfToken Interpolation = Curves.GetWidthsInterpolation();
 
 				const float UEMetersPerUnit = 0.01f;
-				const float Scale = !FMath::IsNearlyEqual(StageInfo.MetersPerUnit, UEMetersPerUnit) ? StageInfo.MetersPerUnit / UEMetersPerUnit : 1.0f;
+				const float Scale = !FMath::IsNearlyEqual(StageInfo.MetersPerUnit, UEMetersPerUnit) ? StageInfo.MetersPerUnit / UEMetersPerUnit
+																									: 1.0f;
 				if (Interpolation == pxr::UsdGeomTokens->uniform || Interpolation == pxr::UsdGeomTokens->constant)
 				{
 					const float ConstWidth = WidthsArray[0] * Scale;
-					TStrandAttributesRef<float> WidthStrandAttributeRef = HairDescription.StrandAttributes().GetAttributesRef<float>(HairAttribute::Strand::Width);
+					TStrandAttributesRef<float> WidthStrandAttributeRef = HairDescription.StrandAttributes().GetAttributesRef<float>(
+						HairAttribute::Strand::Width
+					);
 					if (!WidthStrandAttributeRef.IsValid())
 					{
 						HairDescription.StrandAttributes().RegisterAttribute<float>(HairAttribute::Strand::Width);
@@ -309,12 +336,16 @@ namespace UsdToUnreal
 
 					for (int32 Index = 0; Index < NumCurves; ++Index)
 					{
-						WidthStrandAttributeRef[FStrandID(StartStrandID + Index)] = Interpolation == pxr::UsdGeomTokens->constant ? ConstWidth : WidthsArray[Index] * Scale;
+						WidthStrandAttributeRef[FStrandID(StartStrandID + Index)] = Interpolation == pxr::UsdGeomTokens->constant
+																						? ConstWidth
+																						: WidthsArray[Index] * Scale;
 					}
 				}
 				else if (Interpolation == pxr::UsdGeomTokens->vertex)
 				{
-					TVertexAttributesRef<float> VertexAttributeRef = HairDescription.VertexAttributes().GetAttributesRef<float>(HairAttribute::Vertex::Width);
+					TVertexAttributesRef<float> VertexAttributeRef = HairDescription.VertexAttributes().GetAttributesRef<float>(
+						HairAttribute::Vertex::Width
+					);
 					if (!VertexAttributeRef.IsValid())
 					{
 						HairDescription.VertexAttributes().RegisterAttribute<float>(HairAttribute::Vertex::Width);
@@ -333,7 +364,15 @@ namespace UsdToUnreal
 			std::vector<pxr::UsdGeomPrimvar> Primvars = PrimvarsAPI.GetPrimvars();
 			if (!Primvars.empty())
 			{
-				UE::UsdGroomConversion::Private::ConvertPrimvars(Primvars, TimeCode, StartStrandID, NumCurves, StartVertexID, NumPoints, HairDescription);
+				UE::UsdGroomConversion::Private::ConvertPrimvars(
+					Primvars,
+					TimeCode,
+					StartStrandID,
+					NumCurves,
+					StartVertexID,
+					NumPoints,
+					HairDescription
+				);
 			}
 
 			// Groom-scope attributes are not currently extracted since they are not strictly necessary
@@ -367,6 +406,6 @@ namespace UsdToUnreal
 
 		return bSuccess;
 	}
-}
+}	 // namespace UsdToUnreal
 
-#endif // #if USE_USD_SDK
+#endif	  // #if USE_USD_SDK
diff --git a/Engine/Plugins/Importers/USDImporter/Source/USDUtilities/Private/USDIntegrationUtils.cpp b/Engine/Plugins/Importers/USDImporter/Source/USDUtilities/Private/USDIntegrationUtils.cpp
--- a/Engine/Plugins/Importers/USDImporter/Source/USDUtilities/Private/USDIntegrationUtils.cpp
+++ b/Engine/Plugins/Importers/USDImporter/Source/USDUtilities/Private/USDIntegrationUtils.cpp
@@ -4,7 +4,6 @@
 
 #include "UnrealUSDWrapper.h"
 #include "USDAttributeUtils.h"
-#include "USDLog.h"
 #include "USDTypesConversion.h"
 #include "UsdWrappers/SdfChangeBlock.h"
 #include "UsdWrappers/UsdAttribute.h"
@@ -15,50 +14,47 @@
 
 #if USE_USD_SDK
 #include "USDIncludesStart.h"
-	#include "pxr/usd/sdf/primSpec.h"
-	#include "pxr/usd/usd/attribute.h"
-	#include "pxr/usd/usd/prim.h"
-	#include "pxr/usd/usd/tokens.h"
-	#include "pxr/usd/usdGeom/xformable.h"
-	#include "pxr/usd/usdSkel/root.h"
+#include "pxr/usd/usd/attribute.h"
+#include "pxr/usd/usd/prim.h"
+#include "pxr/usd/usd/tokens.h"
 #include "USDIncludesEnd.h"
 
 #define LOCTEXT_NAMESPACE "USDIntegrationUtils"
 
 namespace UE::UsdIntegrationUtils::Private
 {
-	bool CanApplyOrRemoveSchema( pxr::UsdPrim Prim, pxr::TfToken SchemaName, bool bApplySchema )
+	bool CanApplyOrRemoveSchema(pxr::UsdPrim Prim, pxr::TfToken SchemaName, bool bApplySchema)
 	{
-		if ( !Prim || Prim.IsPseudoRoot() )
+		if (!Prim || Prim.IsPseudoRoot())
 		{
 			return false;
 		}
 
 		FScopedUsdAllocs UsdAllocs;
 
-		pxr::TfType Schema = pxr::UsdSchemaRegistry::GetTypeFromSchemaTypeName( SchemaName );
-		if ( !ensure( static_cast< bool >( Schema ) ) )
+		pxr::TfType Schema = pxr::UsdSchemaRegistry::GetTypeFromSchemaTypeName(SchemaName);
+		if (!ensure(static_cast<bool>(Schema)))
 		{
 			return false;
 		}
 
 		// Check if the schema is compatible with this prim
-		if ( bApplySchema && !Prim.CanApplyAPI( Schema ) )
+		if (bApplySchema && !Prim.CanApplyAPI(Schema))
 		{
 			return false;
 		}
 
 		pxr::UsdStageRefPtr Stage = Prim.GetStage();
 		pxr::SdfLayerRefPtr EditTarget = Stage->GetEditTarget().GetLayer();
-		if ( !Stage || !EditTarget )
+		if (!Stage || !EditTarget)
 		{
 			return false;
 		}
 
 		bool bAlreadyHasSchema = false;
-		if ( EditTarget == Stage->GetRootLayer() )
+		if (EditTarget == Stage->GetRootLayer())
 		{
-			bAlreadyHasSchema = Prim.HasAPI( Schema );
+			bAlreadyHasSchema = Prim.HasAPI(Schema);
 		}
 		else
 		{
@@ -73,8 +69,8 @@ namespace UE::UsdIntegrationUtils::Private
 			// (because the composed prim does not have the schema) but adding the schema would do nothing,
 			// as the sublayer prim spec already has the schema...
 
-			pxr::SdfPrimSpecHandle PrimSpecOnLayer = EditTarget->GetPrimAtPath( Prim.GetPath() );
-			if ( !PrimSpecOnLayer )
+			pxr::SdfPrimSpecHandle PrimSpecOnLayer = EditTarget->GetPrimAtPath(Prim.GetPath());
+			if (!PrimSpecOnLayer)
 			{
 				// We can always just make a new 'over' on this layer
 				return true;
@@ -82,77 +78,77 @@ namespace UE::UsdIntegrationUtils::Private
 
 			// We can only add the schema if its not present on this layer, but we should let USD sort the
 			// Ops out because we could have Ops to add/delete/prepend/append the same schema on the same prim...
-			pxr::SdfListOp<pxr::TfToken> Ops = PrimSpecOnLayer->GetInfo( pxr::UsdTokens->apiSchemas ).Get<pxr::SdfTokenListOp>();
+			pxr::SdfListOp<pxr::TfToken> Ops = PrimSpecOnLayer->GetInfo(pxr::UsdTokens->apiSchemas).Get<pxr::SdfTokenListOp>();
 			std::vector<pxr::TfToken> AppliedOps;
-			Ops.ApplyOperations( &AppliedOps );
-			bAlreadyHasSchema = std::find( AppliedOps.begin(), AppliedOps.end(), SchemaName ) != AppliedOps.end();
+			Ops.ApplyOperations(&AppliedOps);
+			bAlreadyHasSchema = std::find(AppliedOps.begin(), AppliedOps.end(), SchemaName) != AppliedOps.end();
 		}
 
 		return bApplySchema != bAlreadyHasSchema;
 	}
-}
+}	 // namespace UE::UsdIntegrationUtils::Private
 
-bool UsdUtils::PrimHasSchema( const pxr::UsdPrim& Prim, const pxr::TfToken& SchemaToken )
+bool UsdUtils::PrimHasSchema(const pxr::UsdPrim& Prim, const pxr::TfToken& SchemaToken)
 {
-	if ( !Prim )
+	if (!Prim)
 	{
 		return false;
 	}
 
 	FScopedUsdAllocs Allocs;
 
-	pxr::TfType Schema = pxr::UsdSchemaRegistry::GetTypeFromSchemaTypeName( SchemaToken );
-	ensure( static_cast<bool>( Schema ) );
+	pxr::TfType Schema = pxr::UsdSchemaRegistry::GetTypeFromSchemaTypeName(SchemaToken);
+	ensure(static_cast<bool>(Schema));
 
-	return Prim.HasAPI( Schema );
+	return Prim.HasAPI(Schema);
 }
 
-bool UsdUtils::CanApplySchema( pxr::UsdPrim Prim, pxr::TfToken SchemaName )
+bool UsdUtils::CanApplySchema(pxr::UsdPrim Prim, pxr::TfToken SchemaName)
 {
 	const bool bApplySchema = true;
-	return UE::UsdIntegrationUtils::Private::CanApplyOrRemoveSchema( Prim, SchemaName, bApplySchema );
+	return UE::UsdIntegrationUtils::Private::CanApplyOrRemoveSchema(Prim, SchemaName, bApplySchema);
 }
 
-bool UsdUtils::ApplySchema( const pxr::UsdPrim& Prim, const pxr::TfToken& SchemaToken )
+bool UsdUtils::ApplySchema(const pxr::UsdPrim& Prim, const pxr::TfToken& SchemaToken)
 {
-	if ( !Prim )
+	if (!Prim)
 	{
 		return false;
 	}
 
 	FScopedUsdAllocs Allocs;
 
-	pxr::TfType Schema = pxr::UsdSchemaRegistry::GetTypeFromSchemaTypeName( SchemaToken );
-	ensure( static_cast<bool>( Schema ) );
+	pxr::TfType Schema = pxr::UsdSchemaRegistry::GetTypeFromSchemaTypeName(SchemaToken);
+	ensure(static_cast<bool>(Schema));
 
-	return Prim.ApplyAPI( Schema );
+	return Prim.ApplyAPI(Schema);
 }
 
-bool UsdUtils::CanRemoveSchema( pxr::UsdPrim Prim, pxr::TfToken SchemaName )
+bool UsdUtils::CanRemoveSchema(pxr::UsdPrim Prim, pxr::TfToken SchemaName)
 {
 	const bool bApplySchema = false;
-	return UE::UsdIntegrationUtils::Private::CanApplyOrRemoveSchema( Prim, SchemaName, bApplySchema );
+	return UE::UsdIntegrationUtils::Private::CanApplyOrRemoveSchema(Prim, SchemaName, bApplySchema);
 }
 
-bool UsdUtils::RemoveSchema( const pxr::UsdPrim& Prim, const pxr::TfToken& SchemaToken )
+bool UsdUtils::RemoveSchema(const pxr::UsdPrim& Prim, const pxr::TfToken& SchemaToken)
 {
-	if ( !Prim )
+	if (!Prim)
 	{
 		return false;
 	}
 
 	FScopedUsdAllocs Allocs;
 
-	pxr::TfType Schema = pxr::UsdSchemaRegistry::GetTypeFromSchemaTypeName( SchemaToken );
-	ensure( static_cast< bool >( Schema ) );
+	pxr::TfType Schema = pxr::UsdSchemaRegistry::GetTypeFromSchemaTypeName(SchemaToken);
+	ensure(static_cast<bool>(Schema));
 
-	return Prim.RemoveAPI( Schema );
+	return Prim.RemoveAPI(Schema);
 }
 
-void UnrealToUsd::ConvertLiveLinkProperties( const UActorComponent* InComponent, pxr::UsdPrim& InOutPrim )
+void UnrealToUsd::ConvertLiveLinkProperties(const UActorComponent* InComponent, pxr::UsdPrim& InOutPrim)
 {
 #if WITH_EDITOR
-	if ( !InOutPrim || !UsdUtils::PrimHasSchema( InOutPrim, UnrealIdentifiers::LiveLinkAPI ) )
+	if (!InOutPrim || !UsdUtils::PrimHasSchema(InOutPrim, UnrealIdentifiers::LiveLinkAPI))
 	{
 		return;
 	}
@@ -167,42 +163,42 @@ void UnrealToUsd::ConvertLiveLinkProperties( const UActorComponent* InComponent,
 	}
 
 	// Skeletal LiveLink case
-	if ( const USkeletalMeshComponent* SkeletalComponent = Cast<USkeletalMeshComponent>( InComponent ) )
+	if (const USkeletalMeshComponent* SkeletalComponent = Cast<USkeletalMeshComponent>(InComponent))
 	{
-		if ( pxr::UsdAttribute Attr = InOutPrim.CreateAttribute( UnrealIdentifiers::UnrealAnimBlueprintPath, pxr::SdfValueTypeNames->String ) )
+		if (pxr::UsdAttribute Attr = InOutPrim.CreateAttribute(UnrealIdentifiers::UnrealAnimBlueprintPath, pxr::SdfValueTypeNames->String))
 		{
 			FString AnimBPPath = SkeletalComponent->AnimClass && SkeletalComponent->AnimClass->ClassGeneratedBy
-				? SkeletalComponent->AnimClass->ClassGeneratedBy->GetPathName()
-				: FString();
+									 ? SkeletalComponent->AnimClass->ClassGeneratedBy->GetPathName()
+									 : FString();
 
-			Attr.Set( UnrealToUsd::ConvertString( *AnimBPPath ).Get(), pxr::UsdTimeCode::Default() );
+			Attr.Set(UnrealToUsd::ConvertString(*AnimBPPath).Get(), pxr::UsdTimeCode::Default());
 			UsdUtils::NotifyIfOverriddenOpinion(Attr);
 		}
 
-		if ( pxr::UsdAttribute Attr = InOutPrim.CreateAttribute( UnrealIdentifiers::UnrealLiveLinkEnabled, pxr::SdfValueTypeNames->Bool ) )
+		if (pxr::UsdAttribute Attr = InOutPrim.CreateAttribute(UnrealIdentifiers::UnrealLiveLinkEnabled, pxr::SdfValueTypeNames->Bool))
 		{
 			const bool bEnabled = SkeletalComponent->GetAnimationMode() == EAnimationMode::AnimationBlueprint;
-			Attr.Set( bEnabled, pxr::UsdTimeCode::Default() );
+			Attr.Set(bEnabled, pxr::UsdTimeCode::Default());
 			UsdUtils::NotifyIfOverriddenOpinion(Attr);
 		}
 	}
 	// Non-skeletal LiveLink case
-	else if ( const ULiveLinkComponentController* InController = Cast<ULiveLinkComponentController>( InComponent ) )
+	else if (const ULiveLinkComponentController* InController = Cast<ULiveLinkComponentController>(InComponent))
 	{
-		if ( pxr::UsdAttribute Attr = InOutPrim.CreateAttribute( UnrealIdentifiers::UnrealLiveLinkSubjectName, pxr::SdfValueTypeNames->String ) )
+		if (pxr::UsdAttribute Attr = InOutPrim.CreateAttribute(UnrealIdentifiers::UnrealLiveLinkSubjectName, pxr::SdfValueTypeNames->String))
 		{
-			Attr.Set( UnrealToUsd::ConvertString( *InController->SubjectRepresentation.Subject.Name.ToString() ).Get(), pxr::UsdTimeCode::Default() );
+			Attr.Set(UnrealToUsd::ConvertString(*InController->SubjectRepresentation.Subject.Name.ToString()).Get(), pxr::UsdTimeCode::Default());
 			UsdUtils::NotifyIfOverriddenOpinion(Attr);
 		}
 
-		if ( pxr::UsdAttribute Attr = InOutPrim.CreateAttribute( UnrealIdentifiers::UnrealLiveLinkEnabled, pxr::SdfValueTypeNames->Bool ) )
+		if (pxr::UsdAttribute Attr = InOutPrim.CreateAttribute(UnrealIdentifiers::UnrealLiveLinkEnabled, pxr::SdfValueTypeNames->Bool))
 		{
-			Attr.Set( InController->bEvaluateLiveLink, pxr::UsdTimeCode::Default() );
+			Attr.Set(InController->bEvaluateLiveLink, pxr::UsdTimeCode::Default());
 			UsdUtils::NotifyIfOverriddenOpinion(Attr);
 		}
 	}
-#endif // WITH_EDITOR
+#endif	  // WITH_EDITOR
 }
 
 #undef LOCTEXT_NAMESPACE
-#endif // USE_USD_SDK
+#endif	  // USE_USD_SDK
diff --git a/Engine/Plugins/Importers/USDImporter/Source/USDUtilities/Private/USDLayerUtils.cpp b/Engine/Plugins/Importers/USDImporter/Source/USDUtilities/Private/USDLayerUtils.cpp
--- a/Engine/Plugins/Importers/USDImporter/Source/USDUtilities/Private/USDLayerUtils.cpp
+++ b/Engine/Plugins/Importers/USDImporter/Source/USDUtilities/Private/USDLayerUtils.cpp
@@ -15,26 +15,24 @@
 
 #include "Framework/Application/SlateApplication.h"
 #include "Misc/Paths.h"
-#include "Widgets/SWidget.h"
 
 #if WITH_EDITOR
-	#include "DesktopPlatformModule.h"
-#endif // WITH_EDITOR
+#include "DesktopPlatformModule.h"
+#endif	  // WITH_EDITOR
 
 #if USE_USD_SDK
 
 #include "USDIncludesStart.h"
-	#include "pxr/usd/pcp/layerStack.h"
-	#include "pxr/usd/sdf/attributeSpec.h"
-	#include "pxr/usd/sdf/fileFormat.h"
-	#include "pxr/usd/sdf/layer.h"
-	#include "pxr/usd/sdf/layerUtils.h"
-	#include "pxr/usd/sdf/textFileFormat.h"
-	#include "pxr/usd/usd/editContext.h"
-	#include "pxr/usd/usd/modelAPI.h"
-	#include "pxr/usd/usd/primCompositionQuery.h"
-	#include "pxr/usd/usd/stage.h"
-	#include "pxr/usd/usdGeom/xform.h"
+#include "pxr/usd/pcp/layerStack.h"
+#include "pxr/usd/sdf/fileFormat.h"
+#include "pxr/usd/sdf/layer.h"
+#include "pxr/usd/sdf/layerUtils.h"
+#include "pxr/usd/sdf/textFileFormat.h"
+#include "pxr/usd/usd/editContext.h"
+#include "pxr/usd/usd/modelAPI.h"
+#include "pxr/usd/usd/primCompositionQuery.h"
+#include "pxr/usd/usd/stage.h"
+#include "pxr/usd/usdGeom/xform.h"
 #include "USDIncludesEnd.h"
 
 #define LOCTEXT_NAMESPACE "USDLayerUtils"
@@ -42,108 +40,99 @@
 namespace UE::UsdLayerUtilsImpl::Private
 {
 	/**
-		* Adapted from flattenUtils.cpp::_FixAssetPaths, except that we only handle actual AssetPaths here as layer/prim paths
-		* will be remapped via Layer.UpdateCompositionAssetDependency().
-		* Returns whether anything was remapped
-		*/
-	bool FixAssetPaths( const pxr::SdfLayerHandle& SourceLayer, pxr::VtValue* Value )
+	 * Adapted from flattenUtils.cpp::_FixAssetPaths, except that we only handle actual AssetPaths here as layer/prim paths
+	 * will be remapped via Layer.UpdateCompositionAssetDependency().
+	 * Returns whether anything was remapped
+	 */
+	bool FixAssetPaths(const pxr::SdfLayerHandle& SourceLayer, pxr::VtValue* Value)
 	{
-		if ( Value->IsHolding<pxr::SdfAssetPath>() )
+		if (Value->IsHolding<pxr::SdfAssetPath>())
 		{
 			pxr::SdfAssetPath AssetPath;
-			Value->Swap( AssetPath );
-			AssetPath = pxr::SdfAssetPath( SourceLayer->ComputeAbsolutePath( AssetPath.GetAssetPath() ) );
-			Value->Swap( AssetPath );
+			Value->Swap(AssetPath);
+			AssetPath = pxr::SdfAssetPath(SourceLayer->ComputeAbsolutePath(AssetPath.GetAssetPath()));
+			Value->Swap(AssetPath);
 			return true;
 		}
-		else if ( Value->IsHolding<pxr::VtArray<pxr::SdfAssetPath>>() )
+		else if (Value->IsHolding<pxr::VtArray<pxr::SdfAssetPath>>())
 		{
 			pxr::VtArray<pxr::SdfAssetPath> PathArray;
-			Value->Swap( PathArray );
-			for ( pxr::SdfAssetPath& AssetPath : PathArray )
+			Value->Swap(PathArray);
+			for (pxr::SdfAssetPath& AssetPath : PathArray)
 			{
-				AssetPath = pxr::SdfAssetPath( SourceLayer->ComputeAbsolutePath( AssetPath.GetAssetPath() ) );
+				AssetPath = pxr::SdfAssetPath(SourceLayer->ComputeAbsolutePath(AssetPath.GetAssetPath()));
 			}
-			Value->Swap( PathArray );
+			Value->Swap(PathArray);
 			return true;
 		}
 
 		return false;
 	}
-}
+}	 // namespace UE::UsdLayerUtilsImpl::Private
 
-FText UsdUtils::ToText( ECanInsertSublayerResult Result )
+FText UsdUtils::ToText(ECanInsertSublayerResult Result)
 {
-	switch ( Result )
-	{
-	default:
-	case ECanInsertSublayerResult::Success:
-		return FText::GetEmpty();
-		break;
-	case ECanInsertSublayerResult::ErrorSubLayerNotFound:
-		return LOCTEXT( "CanAddSubLayerNotFound", "SubLayer not found!" );
-		break;
-	case ECanInsertSublayerResult::ErrorSubLayerInvalid:
-		return LOCTEXT( "CanAddSubLayerInvalid", "SubLayer is invalid!" );
-		break;
-	case ECanInsertSublayerResult::ErrorSubLayerIsParentLayer:
-		return LOCTEXT( "CanAddSubLayerIsParent", "SubLayer is the same as the parent layer!" );
-		break;
-	case ECanInsertSublayerResult::ErrorCycleDetected:
-		return LOCTEXT( "CanAddSubLayerCycle", "Cycles detected!" );
-		break;
+	switch (Result)
+	{
+		default:
+		case ECanInsertSublayerResult::Success:
+			return FText::GetEmpty();
+			break;
+		case ECanInsertSublayerResult::ErrorSubLayerNotFound:
+			return LOCTEXT("CanAddSubLayerNotFound", "SubLayer not found!");
+			break;
+		case ECanInsertSublayerResult::ErrorSubLayerInvalid:
+			return LOCTEXT("CanAddSubLayerInvalid", "SubLayer is invalid!");
+			break;
+		case ECanInsertSublayerResult::ErrorSubLayerIsParentLayer:
+			return LOCTEXT("CanAddSubLayerIsParent", "SubLayer is the same as the parent layer!");
+			break;
+		case ECanInsertSublayerResult::ErrorCycleDetected:
+			return LOCTEXT("CanAddSubLayerCycle", "Cycles detected!");
+			break;
 	}
 }
 
-UsdUtils::ECanInsertSublayerResult UsdUtils::CanInsertSubLayer(
-	const pxr::SdfLayerRefPtr& ParentLayer,
-	const TCHAR* SubLayerIdentifier
-)
+UsdUtils::ECanInsertSublayerResult UsdUtils::CanInsertSubLayer(const pxr::SdfLayerRefPtr& ParentLayer, const TCHAR* SubLayerIdentifier)
 {
-	if ( !SubLayerIdentifier )
+	if (!SubLayerIdentifier)
 	{
 		return ECanInsertSublayerResult::ErrorSubLayerNotFound;
 	}
 
 	FScopedUsdAllocs Allocs;
 
-	pxr::SdfLayerRefPtr SubLayer = pxr::SdfLayer::FindOrOpen( UnrealToUsd::ConvertString( SubLayerIdentifier ).Get() );
-	if ( !SubLayer )
+	pxr::SdfLayerRefPtr SubLayer = pxr::SdfLayer::FindOrOpen(UnrealToUsd::ConvertString(SubLayerIdentifier).Get());
+	if (!SubLayer)
 	{
 		return ECanInsertSublayerResult::ErrorSubLayerNotFound;
 	}
 
-	if ( SubLayer == ParentLayer )
+	if (SubLayer == ParentLayer)
 	{
 		return ECanInsertSublayerResult::ErrorSubLayerIsParentLayer;
 	}
 
 	// We can't climb through ancestors of ParentLayer, so we have to open sublayer and see if parentlayer is a
 	// descendant of *it* in order to detect cycles
-	TFunction< ECanInsertSublayerResult( const pxr::SdfLayerRefPtr& ) > CanAddSubLayerRecursive;
-	CanAddSubLayerRecursive = [ParentLayer, &CanAddSubLayerRecursive](
-		const pxr::SdfLayerRefPtr& CurrentParent
-	) -> ECanInsertSublayerResult
+	TFunction<ECanInsertSublayerResult(const pxr::SdfLayerRefPtr&)> CanAddSubLayerRecursive;
+	CanAddSubLayerRecursive = [ParentLayer, &CanAddSubLayerRecursive](const pxr::SdfLayerRefPtr& CurrentParent) -> ECanInsertSublayerResult
 	{
-		for ( const std::string& SubLayerPath : CurrentParent->GetSubLayerPaths() )
+		for (const std::string& SubLayerPath : CurrentParent->GetSubLayerPaths())
 		{
 			// This may seem expensive, but keep in mind the main use case for this (at least for now) is for checking
 			// during layer drag and drop, where all of these layers are actually already open anyway
-			pxr::SdfLayerRefPtr ChildSubLayer = pxr::SdfLayer::FindOrOpenRelativeToLayer(
-				CurrentParent,
-				SubLayerPath
-			);
+			pxr::SdfLayerRefPtr ChildSubLayer = pxr::SdfLayer::FindOrOpenRelativeToLayer(CurrentParent, SubLayerPath);
 
-			if ( !ChildSubLayer )
+			if (!ChildSubLayer)
 			{
 				return ECanInsertSublayerResult::ErrorSubLayerInvalid;
 			}
 
-			ECanInsertSublayerResult RecursiveResult = ChildSubLayer == ParentLayer
-				 ? ECanInsertSublayerResult::ErrorCycleDetected
-				 : CanAddSubLayerRecursive( ChildSubLayer );
+			ECanInsertSublayerResult RecursiveResult = ChildSubLayer == ParentLayer ? ECanInsertSublayerResult::ErrorCycleDetected
+																					: CanAddSubLayerRecursive(ChildSubLayer);
 
-			if ( RecursiveResult != ECanInsertSublayerResult::Success )
+			if (RecursiveResult != ECanInsertSublayerResult::Success)
 			{
 				return RecursiveResult;
 			}
@@ -152,116 +141,121 @@ UsdUtils::ECanInsertSublayerResult UsdUtils::CanInsertSubLayer(
 		return ECanInsertSublayerResult::Success;
 	};
 
-	return CanAddSubLayerRecursive( SubLayer );
+	return CanAddSubLayerRecursive(SubLayer);
 }
 
-bool UsdUtils::InsertSubLayer( const pxr::SdfLayerRefPtr& ParentLayer, const TCHAR* SubLayerFile, int32 Index, double OffsetTimeCodes, double TimeCodesScale )
+bool UsdUtils::InsertSubLayer(
+	const pxr::SdfLayerRefPtr& ParentLayer,
+	const TCHAR* SubLayerFile,
+	int32 Index,
+	double OffsetTimeCodes,
+	double TimeCodesScale
+)
 {
-	if ( !ParentLayer )
+	if (!ParentLayer)
 	{
 		return false;
 	}
 
 	FString RelativeSubLayerPath = SubLayerFile;
-	MakePathRelativeToLayer( UE::FSdfLayer{ ParentLayer }, RelativeSubLayerPath );
+	MakePathRelativeToLayer(UE::FSdfLayer{ParentLayer}, RelativeSubLayerPath);
 
 	// If the relative path is just the same as the clean name (e.g. Layer.usda wants to add Layer.usda as a sublayer) then
 	// just stop here as that is always an error
-	FString ParentLayerPath = UsdToUnreal::ConvertString( ParentLayer->GetRealPath() );
-	if ( FPaths::GetCleanFilename( ParentLayerPath ) == RelativeSubLayerPath )
+	FString ParentLayerPath = UsdToUnreal::ConvertString(ParentLayer->GetRealPath());
+	if (FPaths::GetCleanFilename(ParentLayerPath) == RelativeSubLayerPath)
 	{
-		UE_LOG( LogUsd, Warning, TEXT( "Tried to add layer '%s' as a sublayer of itself!" ), *ParentLayerPath );
+		UE_LOG(LogUsd, Warning, TEXT("Tried to add layer '%s' as a sublayer of itself!"), *ParentLayerPath);
 		return false;
 	}
 
 	FScopedUsdAllocs UsdAllocs;
 
-	ParentLayer->InsertSubLayerPath( UnrealToUsd::ConvertString( *RelativeSubLayerPath ).Get(), Index );
+	ParentLayer->InsertSubLayerPath(UnrealToUsd::ConvertString(*RelativeSubLayerPath).Get(), Index);
 
-	if ( !FMath::IsNearlyZero( OffsetTimeCodes ) || !FMath::IsNearlyEqual( TimeCodesScale, 1.0 ) )
+	if (!FMath::IsNearlyZero(OffsetTimeCodes) || !FMath::IsNearlyEqual(TimeCodesScale, 1.0))
 	{
-		if ( Index == -1 )
+		if (Index == -1)
 		{
 			Index = ParentLayer->GetNumSubLayerPaths() - 1;
 		}
 
-		ParentLayer->SetSubLayerOffset( pxr::SdfLayerOffset{ OffsetTimeCodes, TimeCodesScale }, Index );
+		ParentLayer->SetSubLayerOffset(pxr::SdfLayerOffset{OffsetTimeCodes, TimeCodesScale}, Index);
 	}
 
 	return true;
 }
 
 #if WITH_EDITOR
-TOptional< FString > UsdUtils::BrowseUsdFile( EBrowseFileMode Mode )
+TOptional<FString> UsdUtils::BrowseUsdFile(EBrowseFileMode Mode)
 {
 	IDesktopPlatform* DesktopPlatform = FDesktopPlatformModule::Get();
 
-	if ( DesktopPlatform == nullptr )
+	if (DesktopPlatform == nullptr)
 	{
 		return {};
 	}
 
-	TArray< FString > OutFiles;
+	TArray<FString> OutFiles;
 
 	// When browsing files for the purposes of opening a stage or saving layers,
 	// we offer the native USD file formats as options. Browsing files in order
 	// to use them as the targets of composition arcs (e.g. sublayers,
 	// references, payloads, etc.) also offers any plugin file formats that are
 	// registered.
-	TArray< FString > SupportedExtensions = ( Mode == EBrowseFileMode::Composition ) ?
-		UnrealUSDWrapper::GetAllSupportedFileFormats() :
-		UnrealUSDWrapper::GetNativeFileFormats();
-	if ( SupportedExtensions.Num() == 0 )
+	TArray<FString> SupportedExtensions = (Mode == EBrowseFileMode::Composition) ? UnrealUSDWrapper::GetAllSupportedFileFormats()
+																				 : UnrealUSDWrapper::GetNativeFileFormats();
+	if (SupportedExtensions.Num() == 0)
 	{
 		UE_LOG(LogUsd, Error, TEXT("No file extensions supported by the USD SDK!"));
 		return {};
 	}
 
-	if ( Mode == EBrowseFileMode::Save )
+	if (Mode == EBrowseFileMode::Save)
 	{
 		// USD 21.08 doesn't yet support saving to USDZ, so instead of allowing this option and leading to an error we'll just hide it
-		SupportedExtensions.Remove( TEXT( "usdz" ) );
+		SupportedExtensions.Remove(TEXT("usdz"));
 	}
 
-	FString JoinedExtensions = FString::Join( SupportedExtensions, TEXT( ";*." ) ); // Combine "usd" and "usda" into "usd; *.usda"
-	FString FileTypes = FString::Printf( TEXT( "Universal Scene Description files (*.%s)|*.%s|" ), *JoinedExtensions, *JoinedExtensions );
-	for ( const FString& SupportedExtension : SupportedExtensions )
+	FString JoinedExtensions = FString::Join(SupportedExtensions, TEXT(";*."));	   // Combine "usd" and "usda" into "usd; *.usda"
+	FString FileTypes = FString::Printf(TEXT("Universal Scene Description files (*.%s)|*.%s|"), *JoinedExtensions, *JoinedExtensions);
+	for (const FString& SupportedExtension : SupportedExtensions)
 	{
 		// The '(*.%s)' on the actual name (before the '|') is not optional: We need the name part to be different for each format
 		// or else the options will overwrite each other on the Mac
-		FileTypes += FString::Printf( TEXT( "Universal Scene Description file (*.%s)|*.%s|" ), *SupportedExtension, *SupportedExtension );
+		FileTypes += FString::Printf(TEXT("Universal Scene Description file (*.%s)|*.%s|"), *SupportedExtension, *SupportedExtension);
 	}
 
-	switch ( Mode )
+	switch (Mode)
 	{
-		case EBrowseFileMode::Open :
-		case EBrowseFileMode::Composition :
+		case EBrowseFileMode::Open:
+		case EBrowseFileMode::Composition:
 		{
-			if ( !DesktopPlatform->OpenFileDialog(
-				FSlateApplication::Get().FindBestParentWindowHandleForDialogs( nullptr ),
-				LOCTEXT( "ChooseFile", "Choose file").ToString(),
-				TEXT(""),
-				TEXT(""),
-				*FileTypes,
-				EFileDialogFlags::None,
-				OutFiles
-			) )
+			if (!DesktopPlatform->OpenFileDialog(
+					FSlateApplication::Get().FindBestParentWindowHandleForDialogs(nullptr),
+					LOCTEXT("ChooseFile", "Choose file").ToString(),
+					TEXT(""),
+					TEXT(""),
+					*FileTypes,
+					EFileDialogFlags::None,
+					OutFiles
+				))
 			{
 				return {};
 			}
 			break;
 		}
-		case EBrowseFileMode::Save :
+		case EBrowseFileMode::Save:
 		{
-			if ( !DesktopPlatform->SaveFileDialog(
-				FSlateApplication::Get().FindBestParentWindowHandleForDialogs( nullptr ),
-				LOCTEXT( "ChooseFile", "Choose file").ToString(),
-				TEXT(""),
-				TEXT(""),
-				*FileTypes,
-				EFileDialogFlags::None,
-				OutFiles
-			) )
+			if (!DesktopPlatform->SaveFileDialog(
+					FSlateApplication::Get().FindBestParentWindowHandleForDialogs(nullptr),
+					LOCTEXT("ChooseFile", "Choose file").ToString(),
+					TEXT(""),
+					TEXT(""),
+					*FileTypes,
+					EFileDialogFlags::None,
+					OutFiles
+				))
 			{
 				return {};
 			}
@@ -271,59 +265,63 @@ TOptional< FString > UsdUtils::BrowseUsdFile( EBrowseFileMode Mode )
 			break;
 	}
 
-	if ( OutFiles.Num() > 0 )
+	if (OutFiles.Num() > 0)
 	{
 		// Always make this an absolute path because it may try generating a relative path to the engine binary if it can
-		return FPaths::ConvertRelativePathToFull( OutFiles[ 0 ] );
+		return FPaths::ConvertRelativePathToFull(OutFiles[0]);
 	}
 
 	return {};
 }
 
-#endif // WITH_EDITOR
+#endif	  // WITH_EDITOR
 
-FString UsdUtils::MakePathRelativeToProjectDir( const FString& Path )
+FString UsdUtils::MakePathRelativeToProjectDir(const FString& Path)
 {
-	FString PathConverted = FPaths::ConvertRelativePathToFull( Path );
+	FString PathConverted = FPaths::ConvertRelativePathToFull(Path);
 
 	// Mirror behavior of RelativeToGameDir meta tag on the stage actor's RootLayer
-	if ( FPaths::IsUnderDirectory( PathConverted, FPaths::ProjectDir() ) )
+	if (FPaths::IsUnderDirectory(PathConverted, FPaths::ProjectDir()))
 	{
-		FPaths::MakePathRelativeTo( PathConverted, *FPaths::ProjectDir() );
+		FPaths::MakePathRelativeTo(PathConverted, *FPaths::ProjectDir());
 	}
 
 	return PathConverted;
 }
 
-TUsdStore< pxr::SdfLayerRefPtr > UsdUtils::CreateNewLayer( TUsdStore< pxr::UsdStageRefPtr > UsdStage, const TUsdStore<pxr::SdfLayerRefPtr>& ParentLayer, const TCHAR* LayerFilePath )
+TUsdStore<pxr::SdfLayerRefPtr> UsdUtils::CreateNewLayer(
+	TUsdStore<pxr::UsdStageRefPtr> UsdStage,
+	const TUsdStore<pxr::SdfLayerRefPtr>& ParentLayer,
+	const TCHAR* LayerFilePath
+)
 {
 	FScopedUsdAllocs UsdAllocs;
 
-	std::string UsdLayerFilePath = UnrealToUsd::ConvertString( *FPaths::ConvertRelativePathToFull( LayerFilePath ) ).Get();
+	std::string UsdLayerFilePath = UnrealToUsd::ConvertString(*FPaths::ConvertRelativePathToFull(LayerFilePath)).Get();
 
-	pxr::SdfLayerRefPtr LayerRef = pxr::SdfLayer::CreateNew( UsdLayerFilePath );
+	pxr::SdfLayerRefPtr LayerRef = pxr::SdfLayer::CreateNew(UsdLayerFilePath);
 
-	if ( !LayerRef )
+	if (!LayerRef)
 	{
 		return {};
 	}
 
 	// New layer needs to be created and in the stage layer stack before we can edit it
-	UsdUtils::InsertSubLayer( ParentLayer.Get(), LayerFilePath );
+	UsdUtils::InsertSubLayer(ParentLayer.Get(), LayerFilePath);
 
 	UsdUtils::StartMonitoringErrors();
-	pxr::UsdEditContext UsdEditContext( UsdStage.Get(), LayerRef );
+	pxr::UsdEditContext UsdEditContext(UsdStage.Get(), LayerRef);
 
 	// Create default prim
-	FString PrimPath = TEXT("/") + FPaths::GetBaseFilename( UsdToUnreal::ConvertString( LayerRef->GetDisplayName() ) );
+	FString PrimPath = TEXT("/") + FPaths::GetBaseFilename(UsdToUnreal::ConvertString(LayerRef->GetDisplayName()));
 
-	pxr::SdfPath UsdPrimPath = UnrealToUsd::ConvertPath( *PrimPath ).Get();
-	pxr::UsdGeomXform DefaultPrim = pxr::UsdGeomXform::Define( UsdStage.Get(), UsdPrimPath );
+	pxr::SdfPath UsdPrimPath = UnrealToUsd::ConvertPath(*PrimPath).Get();
+	pxr::UsdGeomXform DefaultPrim = pxr::UsdGeomXform::Define(UsdStage.Get(), UsdPrimPath);
 
-	if ( DefaultPrim)
+	if (DefaultPrim)
 	{
 		// Set default prim
-		LayerRef->SetDefaultPrim( DefaultPrim.GetPrim().GetName() );
+		LayerRef->SetDefaultPrim(DefaultPrim.GetPrim().GetName());
 	}
 
 	bool bHadErrors = UsdUtils::ShowErrorsAndStopMonitoring();
@@ -333,12 +331,12 @@ TUsdStore< pxr::SdfLayerRefPtr > UsdUtils::CreateNewLayer( TUsdStore< pxr::UsdSt
 		return {};
 	}
 
-	return TUsdStore< pxr::SdfLayerRefPtr >( LayerRef );
+	return TUsdStore<pxr::SdfLayerRefPtr>(LayerRef);
 }
 
-UE::FSdfLayer UsdUtils::FindLayerForPrim( const pxr::UsdPrim& Prim )
+UE::FSdfLayer UsdUtils::FindLayerForPrim(const pxr::UsdPrim& Prim)
 {
-	if ( !Prim )
+	if (!Prim)
 	{
 		return {};
 	}
@@ -347,80 +345,80 @@ UE::FSdfLayer UsdUtils::FindLayerForPrim( const pxr::UsdPrim& Prim )
 
 	// Use this instead of UsdPrimCompositionQuery as that one can simply fail in some scenarios
 	// (e.g. empty parent layer pointing at a sublayer with a prim, where it fails to provide the sublayer arc's layer)
-	for ( const pxr::SdfPrimSpecHandle& Handle : Prim.GetPrimStack() )
+	for (const pxr::SdfPrimSpecHandle& Handle : Prim.GetPrimStack())
 	{
-		if ( Handle )
+		if (Handle)
 		{
-			if ( pxr::SdfLayerHandle Layer = Handle->GetLayer() )
+			if (pxr::SdfLayerHandle Layer = Handle->GetLayer())
 			{
-				return UE::FSdfLayer( Layer );
+				return UE::FSdfLayer(Layer);
 			}
 		}
 	}
 
-	return UE::FSdfLayer( Prim.GetStage()->GetRootLayer() );
+	return UE::FSdfLayer(Prim.GetStage()->GetRootLayer());
 }
 
-UE::FSdfLayer UsdUtils::FindLayerForAttribute( const pxr::UsdAttribute& Attribute, double TimeCode )
+UE::FSdfLayer UsdUtils::FindLayerForAttribute(const pxr::UsdAttribute& Attribute, double TimeCode)
 {
-	if ( !Attribute )
+	if (!Attribute)
 	{
 		return {};
 	}
 
 	FScopedUsdAllocs UsdAllocs;
 
-	for ( const pxr::SdfPropertySpecHandle& PropertySpec : Attribute.GetPropertyStack( TimeCode ) )
+	for (const pxr::SdfPropertySpecHandle& PropertySpec : Attribute.GetPropertyStack(TimeCode))
 	{
-		if ( PropertySpec->HasDefaultValue() || PropertySpec->GetLayer()->GetNumTimeSamplesForPath( PropertySpec->GetPath() ) > 0 )
+		if (PropertySpec->HasDefaultValue() || PropertySpec->GetLayer()->GetNumTimeSamplesForPath(PropertySpec->GetPath()) > 0)
 		{
-			return UE::FSdfLayer( PropertySpec->GetLayer() );
+			return UE::FSdfLayer(PropertySpec->GetLayer());
 		}
 	}
 
 	return {};
 }
 
-UE::FSdfLayer UsdUtils::FindLayerForAttributes( const TArray<UE::FUsdAttribute>& Attributes, double TimeCode, bool bIncludeSessionLayers )
+UE::FSdfLayer UsdUtils::FindLayerForAttributes(const TArray<UE::FUsdAttribute>& Attributes, double TimeCode, bool bIncludeSessionLayers)
 {
 	FScopedUsdAllocs UsdAllocs;
 
 	TMap<FString, UE::FSdfLayer> IdentifierToLayers;
-	IdentifierToLayers.Reserve( Attributes.Num() );
+	IdentifierToLayers.Reserve(Attributes.Num());
 
 	pxr::UsdStageRefPtr Stage;
-	for ( const UE::FUsdAttribute& Attribute : Attributes )
+	for (const UE::FUsdAttribute& Attribute : Attributes)
 	{
-		if ( Attribute )
+		if (Attribute)
 		{
-			if ( UE::FSdfLayer Layer = UsdUtils::FindLayerForAttribute( Attribute, TimeCode ) )
+			if (UE::FSdfLayer Layer = UsdUtils::FindLayerForAttribute(Attribute, TimeCode))
 			{
-				IdentifierToLayers.Add( Layer.GetIdentifier(), Layer );
+				IdentifierToLayers.Add(Layer.GetIdentifier(), Layer);
 
-				if ( !Stage )
+				if (!Stage)
 				{
-					Stage = pxr::UsdStageRefPtr{ Attribute.GetPrim().GetStage() };
+					Stage = pxr::UsdStageRefPtr{Attribute.GetPrim().GetStage()};
 				}
 			}
 		}
 	}
 
-	if ( !Stage || IdentifierToLayers.Num() == 0)
+	if (!Stage || IdentifierToLayers.Num() == 0)
 	{
 		return {};
 	}
 
-	if ( IdentifierToLayers.Num() == 1 )
+	if (IdentifierToLayers.Num() == 1)
 	{
 		return IdentifierToLayers.CreateIterator().Value();
 	}
 
 	// Iterate through the layer stack in strong to weak order, and return the first of those layers
 	// that is actually one of the attribute layers
-	for ( const pxr::SdfLayerHandle& LayerHandle : Stage->GetLayerStack( bIncludeSessionLayers ) )
+	for (const pxr::SdfLayerHandle& LayerHandle : Stage->GetLayerStack(bIncludeSessionLayers))
 	{
-		FString Identifier = UsdToUnreal::ConvertString( LayerHandle->GetIdentifier() );
-		if ( UE::FSdfLayer* AttributeLayer = IdentifierToLayers.Find( Identifier ) )
+		FString Identifier = UsdToUnreal::ConvertString(LayerHandle->GetIdentifier());
+		if (UE::FSdfLayer* AttributeLayer = IdentifierToLayers.Find(Identifier))
 		{
 			return *AttributeLayer;
 		}
@@ -429,48 +427,48 @@ UE::FSdfLayer UsdUtils::FindLayerForAttributes( const TArray<UE::FUsdAttribute>&
 	return {};
 }
 
-UE::FSdfLayer UsdUtils::FindLayerForSubLayerPath( const UE::FSdfLayer& RootLayer, const FStringView& SubLayerPath )
+UE::FSdfLayer UsdUtils::FindLayerForSubLayerPath(const UE::FSdfLayer& RootLayer, const FStringView& SubLayerPath)
 {
-	const FString RelativeLayerPath = UE::FSdfLayerUtils::SdfComputeAssetPathRelativeToLayer( RootLayer, SubLayerPath.GetData() );
+	const FString RelativeLayerPath = UE::FSdfLayerUtils::SdfComputeAssetPathRelativeToLayer(RootLayer, SubLayerPath.GetData());
 
-	return UE::FSdfLayer::FindOrOpen( *RelativeLayerPath );
+	return UE::FSdfLayer::FindOrOpen(*RelativeLayerPath);
 }
 
-bool UsdUtils::SetRefOrPayloadLayerOffset( pxr::UsdPrim& Prim, const UE::FSdfLayerOffset& LayerOffset )
+bool UsdUtils::SetRefOrPayloadLayerOffset(pxr::UsdPrim& Prim, const UE::FSdfLayerOffset& LayerOffset)
 {
 	FScopedUsdAllocs UsdAllocs;
 
-	pxr::UsdPrimCompositionQuery PrimCompositionQuery( Prim );
-	std::vector< pxr::UsdPrimCompositionQueryArc > CompositionArcs = PrimCompositionQuery.GetCompositionArcs();
+	pxr::UsdPrimCompositionQuery PrimCompositionQuery(Prim);
+	std::vector<pxr::UsdPrimCompositionQueryArc> CompositionArcs = PrimCompositionQuery.GetCompositionArcs();
 
-	for ( const pxr::UsdPrimCompositionQueryArc& CompositionArc : CompositionArcs )
+	for (const pxr::UsdPrimCompositionQueryArc& CompositionArc : CompositionArcs)
 	{
-		if ( CompositionArc.GetArcType() == pxr::PcpArcTypeReference )
+		if (CompositionArc.GetArcType() == pxr::PcpArcTypeReference)
 		{
 			pxr::SdfReferenceEditorProxy ReferenceEditor;
 			pxr::SdfReference OldReference;
 
-			if ( CompositionArc.GetIntroducingListEditor( &ReferenceEditor, &OldReference ) )
+			if (CompositionArc.GetIntroducingListEditor(&ReferenceEditor, &OldReference))
 			{
 				pxr::SdfReference NewReference = OldReference;
-				NewReference.SetLayerOffset( pxr::SdfLayerOffset( LayerOffset.Offset, LayerOffset.Scale ) );
+				NewReference.SetLayerOffset(pxr::SdfLayerOffset(LayerOffset.Offset, LayerOffset.Scale));
 
-				ReferenceEditor.ReplaceItemEdits( OldReference, NewReference );
+				ReferenceEditor.ReplaceItemEdits(OldReference, NewReference);
 
 				return true;
 			}
 		}
-		else if ( CompositionArc.GetArcType() == pxr::PcpArcTypePayload )
+		else if (CompositionArc.GetArcType() == pxr::PcpArcTypePayload)
 		{
 			pxr::SdfPayloadEditorProxy PayloadEditor;
 			pxr::SdfPayload OldPayload;
 
-			if ( CompositionArc.GetIntroducingListEditor( &PayloadEditor, &OldPayload ) )
+			if (CompositionArc.GetIntroducingListEditor(&PayloadEditor, &OldPayload))
 			{
 				pxr::SdfPayload NewPayload = OldPayload;
-				NewPayload.SetLayerOffset( pxr::SdfLayerOffset( LayerOffset.Offset, LayerOffset.Scale ) );
+				NewPayload.SetLayerOffset(pxr::SdfLayerOffset(LayerOffset.Offset, LayerOffset.Scale));
 
-				PayloadEditor.ReplaceItemEdits( OldPayload, NewPayload );
+				PayloadEditor.ReplaceItemEdits(OldPayload, NewPayload);
 
 				return true;
 			}
@@ -480,23 +478,23 @@ bool UsdUtils::SetRefOrPayloadLayerOffset( pxr::UsdPrim& Prim, const UE::FSdfLay
 	return false;
 }
 
-UE::FSdfLayerOffset UsdUtils::GetLayerToStageOffset( const pxr::UsdAttribute& Attribute )
+UE::FSdfLayerOffset UsdUtils::GetLayerToStageOffset(const pxr::UsdAttribute& Attribute)
 {
 	// Inspired by pxr::_GetLayerToStageOffset
 
-	UE::FSdfLayer AttributeLayer = UsdUtils::FindLayerForAttribute( Attribute, pxr::UsdTimeCode::EarliestTime().GetValue() );
+	UE::FSdfLayer AttributeLayer = UsdUtils::FindLayerForAttribute(Attribute, pxr::UsdTimeCode::EarliestTime().GetValue());
 
 	FScopedUsdAllocs UsdAllocs;
 
-	pxr::UsdResolveInfo ResolveInfo = Attribute.GetResolveInfo( pxr::UsdTimeCode::EarliestTime() );
+	pxr::UsdResolveInfo ResolveInfo = Attribute.GetResolveInfo(pxr::UsdTimeCode::EarliestTime());
 	pxr::PcpNodeRef Node = ResolveInfo.GetNode();
-	if ( !Node )
+	if (!Node)
 	{
 		return UE::FSdfLayerOffset();
 	}
 
 	const pxr::PcpMapExpression& MapToRoot = Node.GetMapToRoot();
-	if ( MapToRoot.IsNull() )
+	if (MapToRoot.IsNull())
 	{
 		return UE::FSdfLayerOffset();
 	}
@@ -505,15 +503,17 @@ UE::FSdfLayerOffset UsdUtils::GetLayerToStageOffset( const pxr::UsdAttribute& At
 
 	pxr::SdfLayerOffset LocalOffset = NodeToRootNodeOffset;
 
-	if ( const pxr::SdfLayerOffset* LayerToRootLayerOffset = ResolveInfo.GetNode().GetLayerStack()->GetLayerOffsetForLayer( pxr::SdfLayerRefPtr( AttributeLayer ) ) )
+	if (const pxr::SdfLayerOffset* LayerToRootLayerOffset = ResolveInfo.GetNode().GetLayerStack()->GetLayerOffsetForLayer(
+			pxr::SdfLayerRefPtr(AttributeLayer)
+		))
 	{
 		LocalOffset = LocalOffset * (*LayerToRootLayerOffset);
 	}
 
-	return UE::FSdfLayerOffset( LocalOffset.GetOffset(), LocalOffset.GetScale() );
+	return UE::FSdfLayerOffset(LocalOffset.GetOffset(), LocalOffset.GetScale());
 }
 
-UE::FSdfLayerOffset UsdUtils::GetPrimToStageOffset( const UE::FUsdPrim& Prim )
+UE::FSdfLayerOffset UsdUtils::GetPrimToStageOffset(const UE::FUsdPrim& Prim)
 {
 	// In most cases all we care about is an offset from the prim's layer to the stage, but it is also possible for a prim
 	// to directly reference another layer with an offset and scale as well, and this function will pick up on that. Example:
@@ -529,83 +529,83 @@ UE::FSdfLayerOffset UsdUtils::GetPrimToStageOffset( const UE::FUsdPrim& Prim )
 
 	FScopedUsdAllocs Allocs;
 
-	UE::FSdfLayer StrongestLayerForPrim = UsdUtils::FindLayerForPrim( Prim );
+	UE::FSdfLayer StrongestLayerForPrim = UsdUtils::FindLayerForPrim(Prim);
 
-	pxr::UsdPrim UsdPrim{ Prim };
+	pxr::UsdPrim UsdPrim{Prim};
 
-	pxr::UsdPrimCompositionQuery PrimCompositionQuery( UsdPrim );
+	pxr::UsdPrimCompositionQuery PrimCompositionQuery(UsdPrim);
 	pxr::UsdPrimCompositionQuery::Filter Filter;
 	Filter.hasSpecsFilter = pxr::UsdPrimCompositionQuery::HasSpecsFilter::HasSpecs;
-	PrimCompositionQuery.SetFilter( Filter );
+	PrimCompositionQuery.SetFilter(Filter);
 
-	for ( const pxr::UsdPrimCompositionQueryArc& CompositionArc : PrimCompositionQuery.GetCompositionArcs() )
+	for (const pxr::UsdPrimCompositionQueryArc& CompositionArc : PrimCompositionQuery.GetCompositionArcs())
 	{
-		if ( pxr::PcpNodeRef Node = CompositionArc.GetTargetNode() )
+		if (pxr::PcpNodeRef Node = CompositionArc.GetTargetNode())
 		{
 			pxr::SdfLayerOffset Offset;
 
 			// This part of the offset will handle direct prim references
 			const pxr::PcpMapExpression& MapToRoot = Node.GetMapToRoot();
-			if ( !MapToRoot.IsNull() )
+			if (!MapToRoot.IsNull())
 			{
 				Offset = MapToRoot.GetTimeOffset();
 			}
 
-			if ( const pxr::SdfLayerOffset* LayerOffset = Node.GetLayerStack()->GetLayerOffsetForLayer( pxr::SdfLayerRefPtr{ StrongestLayerForPrim } ) )
+			if (const pxr::SdfLayerOffset* LayerOffset = Node.GetLayerStack()->GetLayerOffsetForLayer(pxr::SdfLayerRefPtr{StrongestLayerForPrim}))
 			{
 				Offset = Offset * (*LayerOffset);
 			}
 
-			return UE::FSdfLayerOffset{ Offset.GetOffset(), Offset.GetScale() };
+			return UE::FSdfLayerOffset{Offset.GetOffset(), Offset.GetScale()};
 		}
 	}
 
 	return UE::FSdfLayerOffset{};
 }
 
-void UsdUtils::AddTimeCodeRangeToLayer( const pxr::SdfLayerRefPtr& Layer, double StartTimeCode, double EndTimeCode )
+void UsdUtils::AddTimeCodeRangeToLayer(const pxr::SdfLayerRefPtr& Layer, double StartTimeCode, double EndTimeCode)
 {
 	FScopedUsdAllocs UsdAllocs;
 
-	if ( !Layer )
+	if (!Layer)
 	{
-		FUsdLogManager::LogMessage( EMessageSeverity::Warning, LOCTEXT( "AddTimeCode_InvalidLayer", "Trying to set timecodes on an invalid layer." ) );
+		FUsdLogManager::LogMessage(EMessageSeverity::Warning, LOCTEXT("AddTimeCode_InvalidLayer", "Trying to set timecodes on an invalid layer."));
 		return;
 	}
 
 	// The HasTimeCode check is needed or else we can't author anything with a StartTimeCode lower than the default of 0
-	if ( StartTimeCode < Layer->GetStartTimeCode() || !Layer->HasStartTimeCode() )
+	if (StartTimeCode < Layer->GetStartTimeCode() || !Layer->HasStartTimeCode())
 	{
-		Layer->SetStartTimeCode( StartTimeCode );
+		Layer->SetStartTimeCode(StartTimeCode);
 	}
 
-	if ( EndTimeCode > Layer->GetEndTimeCode() || !Layer->HasEndTimeCode() )
+	if (EndTimeCode > Layer->GetEndTimeCode() || !Layer->HasEndTimeCode())
 	{
-		Layer->SetEndTimeCode( EndTimeCode );
+		Layer->SetEndTimeCode(EndTimeCode);
 	}
 }
 
-void UsdUtils::MakePathRelativeToLayer( const UE::FSdfLayer& Layer, FString& Path )
+void UsdUtils::MakePathRelativeToLayer(const UE::FSdfLayer& Layer, FString& Path)
 {
 #if USE_USD_SDK
 	FScopedUsdAllocs UsdAllocs;
 
-	if ( pxr::SdfLayerRefPtr UsdLayer = (pxr::SdfLayerRefPtr)Layer )
+	if (pxr::SdfLayerRefPtr UsdLayer = (pxr::SdfLayerRefPtr)Layer)
 	{
 		std::string RepositoryPath = UsdLayer->GetRepositoryPath().empty() ? UsdLayer->GetRealPath() : UsdLayer->GetRepositoryPath();
-		FString LayerAbsolutePath = UsdToUnreal::ConvertString( RepositoryPath );
-		if ( !LayerAbsolutePath.IsEmpty() )
+		FString LayerAbsolutePath = UsdToUnreal::ConvertString(RepositoryPath);
+		if (!LayerAbsolutePath.IsEmpty())
 		{
-			FPaths::MakePathRelativeTo( Path, *LayerAbsolutePath );
+			FPaths::MakePathRelativeTo(Path, *LayerAbsolutePath);
 		}
 	}
-#endif // #if USE_USD_SDK
+#endif	  // #if USE_USD_SDK
 }
 
-UE::FSdfLayer UsdUtils::GetUEPersistentStateSublayer( const UE::FUsdStage& Stage, bool bCreateIfNeeded )
+UE::FSdfLayer UsdUtils::GetUEPersistentStateSublayer(const UE::FUsdStage& Stage, bool bCreateIfNeeded)
 {
 	UE::FSdfLayer StateLayer;
-	if ( !Stage )
+	if (!Stage)
 	{
 		return StateLayer;
 	}
@@ -617,24 +617,24 @@ UE::FSdfLayer UsdUtils::GetUEPersistentStateSublayer( const UE::FUsdStage& Stage
 	FString PathPart;
 	FString FilenamePart;
 	FString ExtensionPart;
-	FPaths::Split( Stage.GetRootLayer().GetRealPath(), PathPart, FilenamePart, ExtensionPart );
+	FPaths::Split(Stage.GetRootLayer().GetRealPath(), PathPart, FilenamePart, ExtensionPart);
 
-	FString ExpectedStateLayerPath = FPaths::Combine( PathPart, FString::Printf( TEXT( "%s-UE-persistent-state.%s" ), *FilenamePart, *ExtensionPart ) );
-	FPaths::NormalizeFilename( ExpectedStateLayerPath );
+	FString ExpectedStateLayerPath = FPaths::Combine(PathPart, FString::Printf(TEXT("%s-UE-persistent-state.%s"), *FilenamePart, *ExtensionPart));
+	FPaths::NormalizeFilename(ExpectedStateLayerPath);
 
-	StateLayer = UE::FSdfLayer::FindOrOpen( *ExpectedStateLayerPath );
+	StateLayer = UE::FSdfLayer::FindOrOpen(*ExpectedStateLayerPath);
 
-	if ( !StateLayer && bCreateIfNeeded )
+	if (!StateLayer && bCreateIfNeeded)
 	{
 		StateLayer = pxr::SdfLayer::New(
-			pxr::SdfFileFormat::FindById( pxr::SdfTextFileFormatTokens->Id ),
-			UnrealToUsd::ConvertString( *ExpectedStateLayerPath ).Get()
+			pxr::SdfFileFormat::FindById(pxr::SdfTextFileFormatTokens->Id),
+			UnrealToUsd::ConvertString(*ExpectedStateLayerPath).Get()
 		);
 	}
 
 	// Add the layer as a sublayer of the session layer, in the right location
 	// Always check this because we need to do this even if we just loaded an existing state layer from disk
-	if ( StateLayer )
+	if (StateLayer)
 	{
 		UE::FSdfLayer SessionLayer = Stage.GetSessionLayer();
 
@@ -642,69 +642,70 @@ UE::FSdfLayer UsdUtils::GetUEPersistentStateSublayer( const UE::FUsdStage& Stage
 		// Note that we intentionally only guarantee the UEPersistentLayer is weaker than the UESessionLayer when inserting,
 		// so that the user may reorder these if they want, for whatever reason
 		bool bNeedsToBeAdded = true;
-		for ( const FString& Path : SessionLayer.GetSubLayerPaths() )
+		for (const FString& Path : SessionLayer.GetSubLayerPaths())
 		{
-			if ( FPaths::IsSamePath( Path, ExpectedStateLayerPath ) )
+			if (FPaths::IsSamePath(Path, ExpectedStateLayerPath))
 			{
 				bNeedsToBeAdded = false;
 				break;
 			}
 		}
 
-		if ( bNeedsToBeAdded )
+		if (bNeedsToBeAdded)
 		{
 			// Always add it at the back, so it's weaker than the session layer
-			InsertSubLayer( static_cast< pxr::SdfLayerRefPtr& >( SessionLayer ), *ExpectedStateLayerPath );
+			InsertSubLayer(static_cast<pxr::SdfLayerRefPtr&>(SessionLayer), *ExpectedStateLayerPath);
 		}
 	}
 
 	return StateLayer;
 }
 
-UE::FSdfLayer UsdUtils::GetUESessionStateSublayer( const UE::FUsdStage& Stage, bool bCreateIfNeeded )
+UE::FSdfLayer UsdUtils::GetUESessionStateSublayer(const UE::FUsdStage& Stage, bool bCreateIfNeeded)
 {
 	UE::FSdfLayer StateLayer;
-	if ( !Stage )
+	if (!Stage)
 	{
 		return StateLayer;
 	}
 
 	FScopedUsdAllocs Allocs;
 
-	const pxr::UsdStageRefPtr UsdStage{ Stage };
+	const pxr::UsdStageRefPtr UsdStage{Stage};
 	pxr::SdfLayerRefPtr UsdSessionLayer = UsdStage->GetSessionLayer();
 
 	FString PathPart;
 	FString FilenamePart;
 	FString ExtensionPart;
-	FPaths::Split( Stage.GetRootLayer().GetRealPath(), PathPart, FilenamePart, ExtensionPart );
+	FPaths::Split(Stage.GetRootLayer().GetRealPath(), PathPart, FilenamePart, ExtensionPart);
 
-	FString ExpectedStateLayerDisplayName = FString::Printf( TEXT( "%s-UE-session-state.%s" ), *FilenamePart, *ExtensionPart );
-	FPaths::NormalizeFilename( ExpectedStateLayerDisplayName );
+	FString ExpectedStateLayerDisplayName = FString::Printf(TEXT("%s-UE-session-state.%s"), *FilenamePart, *ExtensionPart);
+	FPaths::NormalizeFilename(ExpectedStateLayerDisplayName);
 
-	std::string UsdExpectedStateLayerDisplayName = UnrealToUsd::ConvertString( *ExpectedStateLayerDisplayName ).Get();
+	std::string UsdExpectedStateLayerDisplayName = UnrealToUsd::ConvertString(*ExpectedStateLayerDisplayName).Get();
 
 	// Check if we already have an existing utils layer in this stage
 	std::string ExistingUESessionStateIdentifier;
 	{
 		std::unordered_set<std::string> SessionLayerSubLayerIdentifiers;
-		for ( const std::string& SubLayerIdentifier : UsdSessionLayer->GetSubLayerPaths() )
+		for (const std::string& SubLayerIdentifier : UsdSessionLayer->GetSubLayerPaths())
 		{
-			SessionLayerSubLayerIdentifiers.insert( SubLayerIdentifier );
+			SessionLayerSubLayerIdentifiers.insert(SubLayerIdentifier);
 		}
-		if ( SessionLayerSubLayerIdentifiers.size() > 0 )
+		if (SessionLayerSubLayerIdentifiers.size() > 0)
 		{
 			const bool bIncludeSessionLayers = true;
-			for ( const pxr::SdfLayerHandle& Layer : UsdStage->GetLayerStack( bIncludeSessionLayers ) )
+			for (const pxr::SdfLayerHandle& Layer : UsdStage->GetLayerStack(bIncludeSessionLayers))
 			{
 				// All session layers always come before the root layer
-				if ( Layer == UsdStage->GetRootLayer() )
+				if (Layer == UsdStage->GetRootLayer())
 				{
 					break;
 				}
 
 				const std::string& Identifier = Layer->GetIdentifier();
-				if ( Layer->IsAnonymous() && Layer->GetDisplayName() == UsdExpectedStateLayerDisplayName && SessionLayerSubLayerIdentifiers.count( Identifier ) > 0 )
+				if (Layer->IsAnonymous() && Layer->GetDisplayName() == UsdExpectedStateLayerDisplayName
+					&& SessionLayerSubLayerIdentifiers.count(Identifier) > 0)
 				{
 					ExistingUESessionStateIdentifier = Identifier;
 					break;
@@ -713,17 +714,18 @@ UE::FSdfLayer UsdUtils::GetUESessionStateSublayer( const UE::FUsdStage& Stage, b
 		}
 	}
 
-	if ( ExistingUESessionStateIdentifier.size() > 0 )
+	if (ExistingUESessionStateIdentifier.size() > 0)
 	{
-		StateLayer = UE::FSdfLayer::FindOrOpen( *UsdToUnreal::ConvertString( ExistingUESessionStateIdentifier ) );
+		StateLayer = UE::FSdfLayer::FindOrOpen(*UsdToUnreal::ConvertString(ExistingUESessionStateIdentifier));
 	}
 
 	// We only need to add as sublayer when creating the StateLayer layers, because they are always transient and never saved/loaded from disk
 	// so if it exists already, it was created right here, where we add it as a sublayer
-	if ( !StateLayer && bCreateIfNeeded )
+	if (!StateLayer && bCreateIfNeeded)
 	{
-		pxr::SdfLayerRefPtr UsdStateLayer = pxr::SdfLayer::CreateAnonymous( UsdExpectedStateLayerDisplayName );
-		UsdSessionLayer->InsertSubLayerPath( UsdStateLayer->GetIdentifier(), 0 ); // Always add it at the front, so it's stronger than the persistent layer
+		pxr::SdfLayerRefPtr UsdStateLayer = pxr::SdfLayer::CreateAnonymous(UsdExpectedStateLayerDisplayName);
+		UsdSessionLayer->InsertSubLayerPath(UsdStateLayer->GetIdentifier(), 0);	   // Always add it at the front, so it's stronger than the persistent
+																				   // layer
 
 		StateLayer = UsdStateLayer;
 	}
@@ -731,41 +733,41 @@ UE::FSdfLayer UsdUtils::GetUESessionStateSublayer( const UE::FUsdStage& Stage, b
 	return StateLayer;
 }
 
-UE::FSdfLayer UsdUtils::FindLayerForIdentifier( const TCHAR* Identifier, const UE::FUsdStage& Stage )
+UE::FSdfLayer UsdUtils::FindLayerForIdentifier(const TCHAR* Identifier, const UE::FUsdStage& Stage)
 {
 	FScopedUsdAllocs UsdAllocs;
 
-	std::string IdentifierStr = UnrealToUsd::ConvertString( Identifier ).Get();
-	if ( pxr::SdfLayer::IsAnonymousLayerIdentifier( IdentifierStr ) )
+	std::string IdentifierStr = UnrealToUsd::ConvertString(Identifier).Get();
+	if (pxr::SdfLayer::IsAnonymousLayerIdentifier(IdentifierStr))
 	{
-		std::string DisplayName = pxr::SdfLayer::GetDisplayNameFromIdentifier( IdentifierStr );
+		std::string DisplayName = pxr::SdfLayer::GetDisplayNameFromIdentifier(IdentifierStr);
 
-		if ( pxr::UsdStageRefPtr UsdStage = static_cast< pxr::UsdStageRefPtr >( Stage ) )
+		if (pxr::UsdStageRefPtr UsdStage = static_cast<pxr::UsdStageRefPtr>(Stage))
 		{
 			const bool bIncludeSessionLayers = true;
-			for ( const pxr::SdfLayerHandle& Layer : UsdStage->GetLayerStack( bIncludeSessionLayers ) )
+			for (const pxr::SdfLayerHandle& Layer : UsdStage->GetLayerStack(bIncludeSessionLayers))
 			{
-				if ( Layer->GetDisplayName() == DisplayName )
+				if (Layer->GetDisplayName() == DisplayName)
 				{
-					return UE::FSdfLayer{ Layer };
+					return UE::FSdfLayer{Layer};
 				}
 			}
 		}
 	}
 	else
 	{
-		if ( pxr::SdfLayerRefPtr Layer = pxr::SdfLayer::FindOrOpen( IdentifierStr ) )
+		if (pxr::SdfLayerRefPtr Layer = pxr::SdfLayer::FindOrOpen(IdentifierStr))
 		{
-			return UE::FSdfLayer{ Layer };
+			return UE::FSdfLayer{Layer};
 		}
 	}
 
 	return UE::FSdfLayer{};
 }
 
-bool UsdUtils::IsSessionLayerWithinStage( const pxr::SdfLayerRefPtr& Layer, const pxr::UsdStageRefPtr& Stage )
+bool UsdUtils::IsSessionLayerWithinStage(const pxr::SdfLayerRefPtr& Layer, const pxr::UsdStageRefPtr& Stage)
 {
-	if ( !Layer || !Stage )
+	if (!Layer || !Stage)
 	{
 		return false;
 	}
@@ -773,16 +775,16 @@ bool UsdUtils::IsSessionLayerWithinStage( const pxr::SdfLayerRefPtr& Layer, cons
 	pxr::SdfLayerRefPtr RootLayer = Stage->GetRootLayer();
 
 	const bool bIncludeSessionLayers = true;
-	for ( const pxr::SdfLayerHandle& ExistingLayer : Stage->GetLayerStack( bIncludeSessionLayers ) )
+	for (const pxr::SdfLayerHandle& ExistingLayer : Stage->GetLayerStack(bIncludeSessionLayers))
 	{
 		// All session layers come before the root layer within the layer stack
 		// Break before we compare with Layer because if Layer is the actual stage's RootLayer we want to return false
-		if ( ExistingLayer == RootLayer )
+		if (ExistingLayer == RootLayer)
 		{
 			break;
 		}
 
-		if ( ExistingLayer == Layer )
+		if (ExistingLayer == Layer)
 		{
 			return true;
 		}
@@ -791,47 +793,47 @@ bool UsdUtils::IsSessionLayerWithinStage( const pxr::SdfLayerRefPtr& Layer, cons
 	return false;
 }
 
-void UsdUtils::ConvertAssetRelativePathsToAbsolute( UE::FSdfLayer& LayerToConvert, const UE::FSdfLayer& AnchorLayer )
+void UsdUtils::ConvertAssetRelativePathsToAbsolute(UE::FSdfLayer& LayerToConvert, const UE::FSdfLayer& AnchorLayer)
 {
 	FScopedUsdAllocs Allocs;
 
-	pxr::SdfLayerRefPtr UsdLayer{ LayerToConvert };
+	pxr::SdfLayerRefPtr UsdLayer{LayerToConvert};
 
 	UsdLayer->Traverse(
 		pxr::SdfPath::AbsoluteRootPath(),
-		[ &UsdLayer, &AnchorLayer ]( const pxr::SdfPath& Path )
+		[&UsdLayer, &AnchorLayer](const pxr::SdfPath& Path)
 		{
-			pxr::SdfSpecType SpecType = UsdLayer->GetSpecType( Path );
-			if ( SpecType != pxr::SdfSpecTypeAttribute )
+			pxr::SdfSpecType SpecType = UsdLayer->GetSpecType(Path);
+			if (SpecType != pxr::SdfSpecTypeAttribute)
 			{
 				return;
 			}
 
 			pxr::VtValue LayerValue;
-			if ( !UsdLayer->HasField( Path, pxr::SdfFieldKeys->Default, &LayerValue ) )
+			if (!UsdLayer->HasField(Path, pxr::SdfFieldKeys->Default, &LayerValue))
 			{
 				return;
 			}
 
-			if ( UE::UsdLayerUtilsImpl::Private::FixAssetPaths( pxr::SdfLayerRefPtr{ AnchorLayer }, &LayerValue ) )
+			if (UE::UsdLayerUtilsImpl::Private::FixAssetPaths(pxr::SdfLayerRefPtr{AnchorLayer}, &LayerValue))
 			{
-				UsdLayer->SetField( Path, pxr::SdfFieldKeys->Default, LayerValue );
+				UsdLayer->SetField(Path, pxr::SdfFieldKeys->Default, LayerValue);
 			}
 		}
 	);
 }
 
-int32 UsdUtils::GetSdfLayerNumFrames( const pxr::SdfLayerRefPtr& Layer )
+int32 UsdUtils::GetSdfLayerNumFrames(const pxr::SdfLayerRefPtr& Layer)
 {
-	if ( !Layer )
+	if (!Layer)
 	{
 		return 0;
 	}
 
 	// USD time code range is inclusive on both ends
-	return FMath::Abs( FMath::CeilToInt32( Layer->GetEndTimeCode() ) - FMath::FloorToInt32( Layer->GetStartTimeCode() ) + 1 );
+	return FMath::Abs(FMath::CeilToInt32(Layer->GetEndTimeCode()) - FMath::FloorToInt32(Layer->GetStartTimeCode()) + 1);
 }
 
 #undef LOCTEXT_NAMESPACE
 
-#endif // #if USE_USD_SDK
+#endif	  // #if USE_USD_SDK
