USD: Reformat USDUtilities module.

diff --git a/Engine/Plugins/Importers/USDImporter/Source/USDUtilities/Private/USDLightConversion.cpp b/Engine/Plugins/Importers/USDImporter/Source/USDUtilities/Private/USDLightConversion.cpp
--- a/Engine/Plugins/Importers/USDImporter/Source/USDUtilities/Private/USDLightConversion.cpp
+++ b/Engine/Plugins/Importers/USDImporter/Source/USDUtilities/Private/USDLightConversion.cpp
@@ -10,11 +10,8 @@
 #include "USDShadeConversion.h"
 #include "USDTypesConversion.h"
 
-#include "UsdWrappers/SdfLayer.h"
-#include "UsdWrappers/SdfPath.h"
 #include "UsdWrappers/UsdAttribute.h"
 #include "UsdWrappers/UsdPrim.h"
-#include "UsdWrappers/UsdStage.h"
 
 #include "Components/DirectionalLightComponent.h"
 #include "Components/LightComponent.h"
@@ -30,13 +27,13 @@
 #if USE_USD_SDK
 
 #include "USDIncludesStart.h"
-	#include "pxr/usd/usdLux/diskLight.h"
-	#include "pxr/usd/usdLux/distantLight.h"
-	#include "pxr/usd/usdLux/domeLight.h"
-	#include "pxr/usd/usdLux/lightAPI.h"
-	#include "pxr/usd/usdLux/rectLight.h"
-	#include "pxr/usd/usdLux/shapingAPI.h"
-	#include "pxr/usd/usdLux/sphereLight.h"
+#include "pxr/usd/usdLux/diskLight.h"
+#include "pxr/usd/usdLux/distantLight.h"
+#include "pxr/usd/usdLux/domeLight.h"
+#include "pxr/usd/usdLux/lightAPI.h"
+#include "pxr/usd/usdLux/rectLight.h"
+#include "pxr/usd/usdLux/shapingAPI.h"
+#include "pxr/usd/usdLux/sphereLight.h"
 #include "USDIncludesEnd.h"
 
 namespace LightConversionImpl
@@ -47,185 +44,201 @@ namespace LightConversionImpl
 	 * E.g. Sun in the sky has ~0.53 degree angular diameter -> 6.720407093551621e-05 sr
 	 * Source: https://en.wikipedia.org/wiki/Solid_angle#Cone,_spherical_cap,_hemisphere
 	 */
-	float SourceAngleToSteradian( float SourceAngleDeg )
+	float SourceAngleToSteradian(float SourceAngleDeg)
 	{
-		return 2.0f * PI * ( 1.0f - FMath::Cos( FMath::DegreesToRadians( SourceAngleDeg / 2.0f ) ) );
+		return 2.0f * PI * (1.0f - FMath::Cos(FMath::DegreesToRadians(SourceAngleDeg / 2.0f)));
 	}
 
 	// As far as we can tell in Aug 2021, the best approximation for a standard in USD light units is to
 	// use Nits all the time. This doesn't make a lot of sense in some cases but here we try our best to convert to Nits
-	float ConvertIntensityToNits( float Intensity, float Steradians, float AreaInSqMeters, ELightUnits SourceUnits )
+	float ConvertIntensityToNits(float Intensity, float Steradians, float AreaInSqMeters, ELightUnits SourceUnits)
 	{
-		switch ( SourceUnits )
+		switch (SourceUnits)
 		{
-		case ELightUnits::Candelas:
-			// Nit = candela / area
-			return Intensity / AreaInSqMeters;
-			break;
-		case ELightUnits::Lumens:
-			// Nit = lumen / ( sr * area );
-			// https://docs.unrealengine.com/en-US/Engine/Rendering/LightingAndShadows/PhysicalLightUnits/index.html#point,spot,andrectlights
-			return Intensity / ( Steradians * AreaInSqMeters );
-			break;
-		case ELightUnits::EV:
-			// Nit = luminance (cd/m2)
-			return EV100ToLuminance(Intensity);
-			break;
-		case ELightUnits::Unitless:
-			// Nit = ( unitless / 625 ) / area = candela / area
-			// https://docs.unrealengine.com/en-US/Engine/Rendering/LightingAndShadows/PhysicalLightUnits/index.html#point,spot,andrectlights
-			return ( Intensity / 625.0f ) / AreaInSqMeters;
-			break;
-		default:
-			break;
+			case ELightUnits::Candelas:
+				// Nit = candela / area
+				return Intensity / AreaInSqMeters;
+				break;
+			case ELightUnits::Lumens:
+				// Nit = lumen / ( sr * area );
+				// https://docs.unrealengine.com/en-US/Engine/Rendering/LightingAndShadows/PhysicalLightUnits/index.html#point,spot,andrectlights
+				return Intensity / (Steradians * AreaInSqMeters);
+				break;
+			case ELightUnits::EV:
+				// Nit = luminance (cd/m2)
+				return EV100ToLuminance(Intensity);
+				break;
+			case ELightUnits::Unitless:
+				// Nit = ( unitless / 625 ) / area = candela / area
+				// https://docs.unrealengine.com/en-US/Engine/Rendering/LightingAndShadows/PhysicalLightUnits/index.html#point,spot,andrectlights
+				return (Intensity / 625.0f) / AreaInSqMeters;
+				break;
+			default:
+				break;
 		}
 
 		return Intensity;
 	}
 
 	// Copied from USpotLightComponent::GetCosHalfConeAngle, so we don't need a component to do the same math
-	float GetSpotLightCosHalfConeAngle( float OuterConeAngle, float InnerConeAngle )
+	float GetSpotLightCosHalfConeAngle(float OuterConeAngle, float InnerConeAngle)
 	{
-		const float ClampedInnerConeAngle = FMath::Clamp( InnerConeAngle, 0.0f, 89.0f ) * ( float ) PI / 180.0f;
-		const float HalfConeAngle = FMath::Clamp( OuterConeAngle * ( float ) PI / 180.0f, ClampedInnerConeAngle + 0.001f, 89.0f * ( float ) PI / 180.0f + 0.001f );
-		return FMath::Cos( HalfConeAngle );
+		const float ClampedInnerConeAngle = FMath::Clamp(InnerConeAngle, 0.0f, 89.0f) * (float)PI / 180.0f;
+		const float HalfConeAngle = FMath::Clamp(
+			OuterConeAngle * (float)PI / 180.0f,
+			ClampedInnerConeAngle + 0.001f,
+			89.0f * (float)PI / 180.0f + 0.001f
+		);
+		return FMath::Cos(HalfConeAngle);
 	}
-}
+}	 // namespace LightConversionImpl
 
-bool UsdToUnreal::ConvertLight( const pxr::UsdPrim& Prim, ULightComponentBase& LightComponentBase, double UsdTimeCode )
+bool UsdToUnreal::ConvertLight(const pxr::UsdPrim& Prim, ULightComponentBase& LightComponentBase, double UsdTimeCode)
 {
 	FScopedUsdAllocs Allocs;
 
-	const pxr::UsdLuxLightAPI LightAPI( Prim );
-	if ( !LightAPI )
+	const pxr::UsdLuxLightAPI LightAPI(Prim);
+	if (!LightAPI)
 	{
 		return false;
 	}
 
-	const float UsdIntensity = UsdUtils::GetUsdValue< float >( LightAPI.GetIntensityAttr(), UsdTimeCode );
-	const float UsdExposure = UsdUtils::GetUsdValue< float >( LightAPI.GetExposureAttr(), UsdTimeCode );
-	const pxr::GfVec3f UsdColor = UsdUtils::GetUsdValue< pxr::GfVec3f >( LightAPI.GetColorAttr(), UsdTimeCode );
+	const float UsdIntensity = UsdUtils::GetUsdValue<float>(LightAPI.GetIntensityAttr(), UsdTimeCode);
+	const float UsdExposure = UsdUtils::GetUsdValue<float>(LightAPI.GetExposureAttr(), UsdTimeCode);
+	const pxr::GfVec3f UsdColor = UsdUtils::GetUsdValue<pxr::GfVec3f>(LightAPI.GetColorAttr(), UsdTimeCode);
 
 	const bool bSRGB = true;
-	LightComponentBase.LightColor = UsdToUnreal::ConvertColor( UsdColor ).ToFColor( bSRGB );
-	LightComponentBase.Intensity = UsdToUnreal::ConvertLightIntensityAttr( UsdIntensity, UsdExposure );
+	LightComponentBase.LightColor = UsdToUnreal::ConvertColor(UsdColor).ToFColor(bSRGB);
+	LightComponentBase.Intensity = UsdToUnreal::ConvertLightIntensityAttr(UsdIntensity, UsdExposure);
 
-	if ( ULightComponent* LightComponent = Cast< ULightComponent >( &LightComponentBase ) )
+	if (ULightComponent* LightComponent = Cast<ULightComponent>(&LightComponentBase))
 	{
-		LightComponent->bUseTemperature = UsdUtils::GetUsdValue< bool >( LightAPI.GetEnableColorTemperatureAttr(), UsdTimeCode );
-		LightComponent->Temperature = UsdUtils::GetUsdValue< float >( LightAPI.GetColorTemperatureAttr(), UsdTimeCode );
+		LightComponent->bUseTemperature = UsdUtils::GetUsdValue<bool>(LightAPI.GetEnableColorTemperatureAttr(), UsdTimeCode);
+		LightComponent->Temperature = UsdUtils::GetUsdValue<float>(LightAPI.GetColorTemperatureAttr(), UsdTimeCode);
 	}
 
 	return true;
 }
 
-bool UsdToUnreal::ConvertDistantLight( const pxr::UsdPrim& Prim, UDirectionalLightComponent& LightComponent, double UsdTimeCode )
+bool UsdToUnreal::ConvertDistantLight(const pxr::UsdPrim& Prim, UDirectionalLightComponent& LightComponent, double UsdTimeCode)
 {
 	FScopedUsdAllocs Allocs;
 
-	pxr::UsdLuxDistantLight DistantLight{ Prim };
-	if ( !DistantLight )
+	pxr::UsdLuxDistantLight DistantLight{Prim};
+	if (!DistantLight)
 	{
 		return false;
 	}
 
-	LightComponent.LightSourceAngle = UsdUtils::GetUsdValue< float >( DistantLight.GetAngleAttr(), UsdTimeCode );
+	LightComponent.LightSourceAngle = UsdUtils::GetUsdValue<float>(DistantLight.GetAngleAttr(), UsdTimeCode);
 
 	return true;
 }
 
-bool UsdToUnreal::ConvertRectLight( const pxr::UsdPrim& Prim, URectLightComponent& LightComponent, double UsdTimeCode )
+bool UsdToUnreal::ConvertRectLight(const pxr::UsdPrim& Prim, URectLightComponent& LightComponent, double UsdTimeCode)
 {
 	FScopedUsdAllocs Allocs;
 
-	pxr::UsdLuxRectLight RectLight{ Prim };
-	if ( !RectLight )
+	pxr::UsdLuxRectLight RectLight{Prim};
+	if (!RectLight)
 	{
 		return false;
 	}
 
-	const FUsdStageInfo StageInfo( Prim.GetStage() );
+	const FUsdStageInfo StageInfo(Prim.GetStage());
 
-	const float UsdIntensity = UsdUtils::GetUsdValue< float >( RectLight.GetIntensityAttr(), UsdTimeCode );
-	const float UsdExposure = UsdUtils::GetUsdValue< float >( RectLight.GetExposureAttr(), UsdTimeCode );
-	const float UsdWidth = UsdUtils::GetUsdValue< float >( RectLight.GetWidthAttr(), UsdTimeCode );
-	const float UsdHeight = UsdUtils::GetUsdValue< float >( RectLight.GetHeightAttr(), UsdTimeCode );
+	const float UsdIntensity = UsdUtils::GetUsdValue<float>(RectLight.GetIntensityAttr(), UsdTimeCode);
+	const float UsdExposure = UsdUtils::GetUsdValue<float>(RectLight.GetExposureAttr(), UsdTimeCode);
+	const float UsdWidth = UsdUtils::GetUsdValue<float>(RectLight.GetWidthAttr(), UsdTimeCode);
+	const float UsdHeight = UsdUtils::GetUsdValue<float>(RectLight.GetHeightAttr(), UsdTimeCode);
 
-	LightComponent.SourceWidth = UsdToUnreal::ConvertDistance( StageInfo, UsdWidth );
-	LightComponent.SourceHeight = UsdToUnreal::ConvertDistance( StageInfo, UsdHeight );
-	LightComponent.Intensity = UsdToUnreal::ConvertRectLightIntensityAttr( UsdIntensity, UsdExposure, UsdWidth, UsdHeight, StageInfo );
+	LightComponent.SourceWidth = UsdToUnreal::ConvertDistance(StageInfo, UsdWidth);
+	LightComponent.SourceHeight = UsdToUnreal::ConvertDistance(StageInfo, UsdHeight);
+	LightComponent.Intensity = UsdToUnreal::ConvertRectLightIntensityAttr(UsdIntensity, UsdExposure, UsdWidth, UsdHeight, StageInfo);
 	LightComponent.IntensityUnits = ELightUnits::Lumens;
 
 	return true;
 }
 
-bool UsdToUnreal::ConvertDiskLight( const pxr::UsdPrim& Prim, URectLightComponent& LightComponent, double UsdTimeCode )
+bool UsdToUnreal::ConvertDiskLight(const pxr::UsdPrim& Prim, URectLightComponent& LightComponent, double UsdTimeCode)
 {
 	FScopedUsdAllocs Allocs;
 
-	pxr::UsdLuxDiskLight DiskLight{ Prim };
-	if ( !DiskLight )
+	pxr::UsdLuxDiskLight DiskLight{Prim};
+	if (!DiskLight)
 	{
 		return false;
 	}
 
-	const FUsdStageInfo StageInfo( Prim.GetStage() );
+	const FUsdStageInfo StageInfo(Prim.GetStage());
 
-	const float UsdIntensity = UsdUtils::GetUsdValue< float >( DiskLight.GetIntensityAttr(), UsdTimeCode );
-	const float UsdExposure = UsdUtils::GetUsdValue< float >( DiskLight.GetExposureAttr(), UsdTimeCode );
-	const float UsdRadius = UsdUtils::GetUsdValue< float >( DiskLight.GetRadiusAttr(), UsdTimeCode );
+	const float UsdIntensity = UsdUtils::GetUsdValue<float>(DiskLight.GetIntensityAttr(), UsdTimeCode);
+	const float UsdExposure = UsdUtils::GetUsdValue<float>(DiskLight.GetExposureAttr(), UsdTimeCode);
+	const float UsdRadius = UsdUtils::GetUsdValue<float>(DiskLight.GetRadiusAttr(), UsdTimeCode);
 
-	LightComponent.SourceWidth = UsdToUnreal::ConvertDistance( StageInfo, UsdRadius ) * 2.f;
+	LightComponent.SourceWidth = UsdToUnreal::ConvertDistance(StageInfo, UsdRadius) * 2.f;
 	LightComponent.SourceHeight = LightComponent.SourceWidth;
-	LightComponent.Intensity = UsdToUnreal::ConvertDiskLightIntensityAttr( UsdIntensity, UsdExposure, UsdRadius, StageInfo );
+	LightComponent.Intensity = UsdToUnreal::ConvertDiskLightIntensityAttr(UsdIntensity, UsdExposure, UsdRadius, StageInfo);
 	LightComponent.IntensityUnits = ELightUnits::Lumens;
 
 	return true;
 }
 
-bool UsdToUnreal::ConvertSphereLight( const pxr::UsdPrim& Prim, UPointLightComponent& LightComponent, double UsdTimeCode )
+bool UsdToUnreal::ConvertSphereLight(const pxr::UsdPrim& Prim, UPointLightComponent& LightComponent, double UsdTimeCode)
 {
 	FScopedUsdAllocs Allocs;
 
-	pxr::UsdLuxSphereLight SphereLight{ Prim };
-	if ( !SphereLight )
+	pxr::UsdLuxSphereLight SphereLight{Prim};
+	if (!SphereLight)
 	{
 		return false;
 	}
 
-	const FUsdStageInfo StageInfo( Prim.GetStage() );
+	const FUsdStageInfo StageInfo(Prim.GetStage());
 
-	const float UsdIntensity = UsdUtils::GetUsdValue< float >( SphereLight.GetIntensityAttr(), UsdTimeCode );
-	const float UsdExposure = UsdUtils::GetUsdValue< float >( SphereLight.GetExposureAttr(), UsdTimeCode );
-	const float UsdRadius = UsdUtils::GetUsdValue< float >( SphereLight.GetRadiusAttr(), UsdTimeCode );
+	const float UsdIntensity = UsdUtils::GetUsdValue<float>(SphereLight.GetIntensityAttr(), UsdTimeCode);
+	const float UsdExposure = UsdUtils::GetUsdValue<float>(SphereLight.GetExposureAttr(), UsdTimeCode);
+	const float UsdRadius = UsdUtils::GetUsdValue<float>(SphereLight.GetRadiusAttr(), UsdTimeCode);
 
-	if ( pxr::UsdLuxShapingAPI ShapingAPI{ Prim } )
+	if (pxr::UsdLuxShapingAPI ShapingAPI{Prim})
 	{
-		const float UsdConeAngle = UsdUtils::GetUsdValue< float >( ShapingAPI.GetShapingConeAngleAttr(), UsdTimeCode );
-		const float UsdConeSoftness = UsdUtils::GetUsdValue< float >( ShapingAPI.GetShapingConeSoftnessAttr(), UsdTimeCode );
+		const float UsdConeAngle = UsdUtils::GetUsdValue<float>(ShapingAPI.GetShapingConeAngleAttr(), UsdTimeCode);
+		const float UsdConeSoftness = UsdUtils::GetUsdValue<float>(ShapingAPI.GetShapingConeSoftnessAttr(), UsdTimeCode);
 
 		float InnerConeAngle = 0.0f;
-		const float OuterConeAngle = UsdToUnreal::ConvertConeAngleSoftnessAttr( UsdConeAngle, UsdConeSoftness, InnerConeAngle );
+		const float OuterConeAngle = UsdToUnreal::ConvertConeAngleSoftnessAttr(UsdConeAngle, UsdConeSoftness, InnerConeAngle);
 
-		LightComponent.Intensity = UsdToUnreal::ConvertLuxShapingAPIIntensityAttr( UsdIntensity, UsdExposure, UsdRadius, UsdConeAngle, UsdConeSoftness, StageInfo );
+		LightComponent.Intensity = UsdToUnreal::ConvertLuxShapingAPIIntensityAttr(
+			UsdIntensity,
+			UsdExposure,
+			UsdRadius,
+			UsdConeAngle,
+			UsdConeSoftness,
+			StageInfo
+		);
 	}
 	else
 	{
-		LightComponent.Intensity = UsdToUnreal::ConvertSphereLightIntensityAttr( UsdIntensity, UsdExposure, UsdRadius, StageInfo );
+		LightComponent.Intensity = UsdToUnreal::ConvertSphereLightIntensityAttr(UsdIntensity, UsdExposure, UsdRadius, StageInfo);
 	}
 
 	LightComponent.IntensityUnits = ELightUnits::Lumens;
-	LightComponent.SourceRadius = UsdToUnreal::ConvertDistance( StageInfo, UsdRadius );
+	LightComponent.SourceRadius = UsdToUnreal::ConvertDistance(StageInfo, UsdRadius);
 
 	return true;
 }
 
-bool UsdToUnreal::ConvertDomeLight( const pxr::UsdPrim& Prim, USkyLightComponent& LightComponent, UUsdAssetCache2* TexturesCache, bool bReuseIdenticalAssets )
+bool UsdToUnreal::ConvertDomeLight(
+	const pxr::UsdPrim& Prim,
+	USkyLightComponent& LightComponent,
+	UUsdAssetCache2* TexturesCache,
+	bool bReuseIdenticalAssets
+)
 {
 	FScopedUsdAllocs UsdAllocs;
 
-	pxr::UsdLuxDomeLight DomeLight{ Prim };
-	if ( !DomeLight )
+	pxr::UsdLuxDomeLight DomeLight{Prim};
+	if (!DomeLight)
 	{
 		return false;
 	}
@@ -233,37 +246,48 @@ bool UsdToUnreal::ConvertDomeLight( const pxr::UsdPrim& Prim, USkyLightComponent
 	// Revert the allocator in case we end up creating a texture on the ansi allocator or something like that
 	FScopedUnrealAllocs UEAllocs;
 
-	const FString ResolvedDomeTexturePath = UsdUtils::GetResolvedTexturePath( DomeLight.GetTextureFileAttr() );
-	if ( ResolvedDomeTexturePath.IsEmpty() )
+	const FString ResolvedDomeTexturePath = UsdUtils::GetResolvedTexturePath(DomeLight.GetTextureFileAttr());
+	if (ResolvedDomeTexturePath.IsEmpty())
 	{
 		FScopedUsdAllocs Allocs;
 
 		pxr::SdfAssetPath TextureAssetPath;
-		DomeLight.GetTextureFileAttr().Get< pxr::SdfAssetPath >( &TextureAssetPath );
+		DomeLight.GetTextureFileAttr().Get<pxr::SdfAssetPath>(&TextureAssetPath);
 
 		// Show a good warning for this because it's easy to pick some cubemap asset from the engine (that usually don't come with the
 		// source texture) and have the dome light silently not work again
-		FString TargetAssetPath = UsdToUnreal::ConvertString( TextureAssetPath.GetAssetPath() );
-		UE_LOG( LogUsd, Warning, TEXT( "Failed to find texture '%s' used for UsdLuxDomeLight '%s'!" ), *TargetAssetPath, *UsdToUnreal::ConvertPath( DomeLight.GetPrim().GetPath() ) );
+		FString TargetAssetPath = UsdToUnreal::ConvertString(TextureAssetPath.GetAssetPath());
+		UE_LOG(
+			LogUsd,
+			Warning,
+			TEXT("Failed to find texture '%s' used for UsdLuxDomeLight '%s'!"),
+			*TargetAssetPath,
+			*UsdToUnreal::ConvertPath(DomeLight.GetPrim().GetPath())
+		);
 
 		return true;
 	}
 
 	const FString PrefixedTextureHash = UsdUtils::GetAssetHashPrefix(Prim, bReuseIdenticalAssets)
-									+ LexToString(FMD5Hash::HashFile(*ResolvedDomeTexturePath));
+										+ LexToString(FMD5Hash::HashFile(*ResolvedDomeTexturePath));
 	UTextureCube* Cubemap = Cast<UTextureCube>(TexturesCache ? TexturesCache->GetCachedAsset(PrefixedTextureHash) : nullptr);
 
-	if ( !Cubemap )
+	if (!Cubemap)
 	{
-		Cubemap = Cast< UTextureCube >( UsdUtils::CreateTexture( DomeLight.GetTextureFileAttr(), UsdToUnreal::ConvertPath( DomeLight.GetPrim().GetPath() ), TEXTUREGROUP_Skybox, TexturesCache ) );
+		Cubemap = Cast<UTextureCube>(UsdUtils::CreateTexture(
+			DomeLight.GetTextureFileAttr(),
+			UsdToUnreal::ConvertPath(DomeLight.GetPrim().GetPath()),
+			TEXTUREGROUP_Skybox,
+			TexturesCache
+		));
 
-		if ( TexturesCache )
+		if (TexturesCache)
 		{
-			TexturesCache->CacheAsset( PrefixedTextureHash, Cubemap );
+			TexturesCache->CacheAsset(PrefixedTextureHash, Cubemap);
 		}
 	}
 
-	if ( Cubemap )
+	if (Cubemap)
 	{
 		LightComponent.Cubemap = Cubemap;
 		LightComponent.SourceType = ESkyLightSourceType::SLS_SpecifiedCubemap;
@@ -272,29 +296,29 @@ bool UsdToUnreal::ConvertDomeLight( const pxr::UsdPrim& Prim, USkyLightComponent
 	return true;
 }
 
-bool UsdToUnreal::ConvertLuxShapingAPI( const pxr::UsdPrim& Prim, USpotLightComponent& LightComponent, double UsdTimeCode )
+bool UsdToUnreal::ConvertLuxShapingAPI(const pxr::UsdPrim& Prim, USpotLightComponent& LightComponent, double UsdTimeCode)
 {
 	FScopedUsdAllocs Allocs;
 
-	if ( !Prim.HasAPI<pxr::UsdLuxShapingAPI>() )
+	if (!Prim.HasAPI<pxr::UsdLuxShapingAPI>())
 	{
 		return false;
 	}
 
-	pxr::UsdLuxShapingAPI ShapingAPI{ Prim };
-	if ( !ShapingAPI )
+	pxr::UsdLuxShapingAPI ShapingAPI{Prim};
+	if (!ShapingAPI)
 	{
 		return false;
 	}
 
-	const float UsdConeAngle = UsdUtils::GetUsdValue< float >( ShapingAPI.GetShapingConeAngleAttr(), UsdTimeCode );
-	const float UsdConeSoftness = UsdUtils::GetUsdValue< float >( ShapingAPI.GetShapingConeSoftnessAttr(), UsdTimeCode );
+	const float UsdConeAngle = UsdUtils::GetUsdValue<float>(ShapingAPI.GetShapingConeAngleAttr(), UsdTimeCode);
+	const float UsdConeSoftness = UsdUtils::GetUsdValue<float>(ShapingAPI.GetShapingConeSoftnessAttr(), UsdTimeCode);
 
 	float InnerConeAngle = 0.0f;
-	const float OuterConeAngle = UsdToUnreal::ConvertConeAngleSoftnessAttr( UsdConeAngle, UsdConeSoftness, InnerConeAngle );
+	const float OuterConeAngle = UsdToUnreal::ConvertConeAngleSoftnessAttr(UsdConeAngle, UsdConeSoftness, InnerConeAngle);
 
-	LightComponent.SetInnerConeAngle( InnerConeAngle );
-	LightComponent.SetOuterConeAngle( OuterConeAngle );
+	LightComponent.SetInnerConeAngle(InnerConeAngle);
+	LightComponent.SetOuterConeAngle(OuterConeAngle);
 
 	return true;
 }
@@ -305,87 +329,102 @@ bool UsdToUnreal::ConvertDomeLight(const pxr::UsdPrim& Prim, USkyLightComponent&
 	return UsdToUnreal::ConvertDomeLight(Prim, LightComponent, NewCache);
 }
 
-float UsdToUnreal::ConvertLightIntensityAttr( float UsdIntensity, float UsdExposure )
+float UsdToUnreal::ConvertLightIntensityAttr(float UsdIntensity, float UsdExposure)
 {
-	return UsdIntensity * FMath::Exp2( UsdExposure );
+	return UsdIntensity * FMath::Exp2(UsdExposure);
 }
 
-float UsdToUnreal::ConvertDistantLightIntensityAttr( float UsdIntensity, float UsdExposure )
+float UsdToUnreal::ConvertDistantLightIntensityAttr(float UsdIntensity, float UsdExposure)
 {
-	return UsdToUnreal::ConvertLightIntensityAttr( UsdIntensity, UsdExposure );
+	return UsdToUnreal::ConvertLightIntensityAttr(UsdIntensity, UsdExposure);
 }
 
-float UsdToUnreal::ConvertRectLightIntensityAttr( float UsdIntensity, float UsdExposure, float UsdWidth, float UsdHeight, const FUsdStageInfo& StageInfo )
+float UsdToUnreal::ConvertRectLightIntensityAttr(
+	float UsdIntensity,
+	float UsdExposure,
+	float UsdWidth,
+	float UsdHeight,
+	const FUsdStageInfo& StageInfo
+)
 {
-	float UEWidth = UsdToUnreal::ConvertDistance( StageInfo, UsdWidth );
-	float UEHeight= UsdToUnreal::ConvertDistance( StageInfo, UsdHeight);
+	float UEWidth = UsdToUnreal::ConvertDistance(StageInfo, UsdWidth);
+	float UEHeight = UsdToUnreal::ConvertDistance(StageInfo, UsdHeight);
 
-	const float AreaInSqMeters = ( UEWidth / 100.f ) * ( UEHeight / 100.f );
+	const float AreaInSqMeters = (UEWidth / 100.f) * (UEHeight / 100.f);
 
-	// Only use PI instead of 2PI because URectLightComponent::SetLightBrightness will use just PI and not 2PI for lumen conversions, due to a cosine distribution
-	// c.f. UActorFactoryRectLight::PostSpawnActor, and the PI factor between candela and lumen for rect lights on https://docs.unrealengine.com/en-US/BuildingWorlds/LightingAndShadows/PhysicalLightUnits/index.html#point,spot,andrectlights
-	return ConvertLightIntensityAttr( UsdIntensity, UsdExposure ) * PI * AreaInSqMeters; // Lumen = Nits * (PI sr for area light) * Area
+	// Only use PI instead of 2PI because URectLightComponent::SetLightBrightness will use just PI and not 2PI for lumen conversions, due to a cosine
+	// distribution c.f. UActorFactoryRectLight::PostSpawnActor, and the PI factor between candela and lumen for rect lights on
+	// https://docs.unrealengine.com/en-US/BuildingWorlds/LightingAndShadows/PhysicalLightUnits/index.html#point,spot,andrectlights
+	return ConvertLightIntensityAttr(UsdIntensity, UsdExposure) * PI * AreaInSqMeters;	  // Lumen = Nits * (PI sr for area light) * Area
 }
 
-float UsdToUnreal::ConvertDiskLightIntensityAttr( float UsdIntensity, float UsdExposure, float UsdRadius, const FUsdStageInfo& StageInfo )
+float UsdToUnreal::ConvertDiskLightIntensityAttr(float UsdIntensity, float UsdExposure, float UsdRadius, const FUsdStageInfo& StageInfo)
 {
-	const float Radius = UsdToUnreal::ConvertDistance( StageInfo, UsdRadius );
+	const float Radius = UsdToUnreal::ConvertDistance(StageInfo, UsdRadius);
 
-	const float AreaInSqMeters = PI * FMath::Square( Radius / 100.f );
+	const float AreaInSqMeters = PI * FMath::Square(Radius / 100.f);
 
-	// Only use PI instead of 2PI because URectLightComponent::SetLightBrightness will use just PI and not 2PI for lumen conversions, due to a cosine distribution
-	// c.f. UActorFactoryRectLight::PostSpawnActor, and the PI factor between candela and lumen for rect lights on https://docs.unrealengine.com/en-US/BuildingWorlds/LightingAndShadows/PhysicalLightUnits/index.html#point,spot,andrectlights
-	return ConvertLightIntensityAttr( UsdIntensity, UsdExposure ) * PI * AreaInSqMeters; // Lumen = Nits * (PI sr for area light) * Area
+	// Only use PI instead of 2PI because URectLightComponent::SetLightBrightness will use just PI and not 2PI for lumen conversions, due to a cosine
+	// distribution c.f. UActorFactoryRectLight::PostSpawnActor, and the PI factor between candela and lumen for rect lights on
+	// https://docs.unrealengine.com/en-US/BuildingWorlds/LightingAndShadows/PhysicalLightUnits/index.html#point,spot,andrectlights
+	return ConvertLightIntensityAttr(UsdIntensity, UsdExposure) * PI * AreaInSqMeters;	  // Lumen = Nits * (PI sr for area light) * Area
 }
 
-float UsdToUnreal::ConvertSphereLightIntensityAttr( float UsdIntensity, float UsdExposure, float UsdRadius, const FUsdStageInfo& StageInfo )
+float UsdToUnreal::ConvertSphereLightIntensityAttr(float UsdIntensity, float UsdExposure, float UsdRadius, const FUsdStageInfo& StageInfo)
 {
-	float Radius = UsdToUnreal::ConvertDistance( StageInfo, UsdRadius );
+	float Radius = UsdToUnreal::ConvertDistance(StageInfo, UsdRadius);
 
 	float SolidAngle = 4.f * PI;
 
 	// Using solid angle for this area is possibly incorrect, but using Nits for point lights also doesn't make much sense in the first place either,
 	// but we must do it for consistency with USD
-	const float AreaInSqMeters = FMath::Max( SolidAngle * FMath::Square( Radius / 100.f ), KINDA_SMALL_NUMBER );
+	const float AreaInSqMeters = FMath::Max(SolidAngle * FMath::Square(Radius / 100.f), KINDA_SMALL_NUMBER);
 
-	return ConvertLightIntensityAttr( UsdIntensity, UsdExposure ) * SolidAngle * AreaInSqMeters; // Lumen = Nits * SolidAngle * Area
+	return ConvertLightIntensityAttr(UsdIntensity, UsdExposure) * SolidAngle * AreaInSqMeters;	  // Lumen = Nits * SolidAngle * Area
 }
 
-float UsdToUnreal::ConvertLuxShapingAPIIntensityAttr( float UsdIntensity, float UsdExposure, float UsdRadius, float UsdConeAngle, float UsdConeSoftness, const FUsdStageInfo& StageInfo )
+float UsdToUnreal::ConvertLuxShapingAPIIntensityAttr(
+	float UsdIntensity,
+	float UsdExposure,
+	float UsdRadius,
+	float UsdConeAngle,
+	float UsdConeSoftness,
+	const FUsdStageInfo& StageInfo
+)
 {
-	float Radius = UsdToUnreal::ConvertDistance( StageInfo, UsdRadius );
+	float Radius = UsdToUnreal::ConvertDistance(StageInfo, UsdRadius);
 
 	float InnerConeAngle = 0.0f;
-	float OuterConeAngle = ConvertConeAngleSoftnessAttr( UsdConeAngle, UsdConeSoftness, InnerConeAngle );
+	float OuterConeAngle = ConvertConeAngleSoftnessAttr(UsdConeAngle, UsdConeSoftness, InnerConeAngle);
 
 	// c.f. USpotLightComponent::ComputeLightBrightness
-	float SolidAngle = 2.f * PI * ( 1.0f - LightConversionImpl::GetSpotLightCosHalfConeAngle( OuterConeAngle, InnerConeAngle ) );
+	float SolidAngle = 2.f * PI * (1.0f - LightConversionImpl::GetSpotLightCosHalfConeAngle(OuterConeAngle, InnerConeAngle));
 
 	// Using solid angle for this area is possibly incorrect, but using Nits for point lights also doesn't make much sense in the first place either,
 	// but we must do it for consistency with USD
-	const float AreaInSqMeters = FMath::Max( SolidAngle * FMath::Square( Radius / 100.f ), KINDA_SMALL_NUMBER );
+	const float AreaInSqMeters = FMath::Max(SolidAngle * FMath::Square(Radius / 100.f), KINDA_SMALL_NUMBER);
 
-	return ConvertLightIntensityAttr( UsdIntensity, UsdExposure ) * SolidAngle * AreaInSqMeters; // Lumen = Nits * SolidAngle * Area
+	return ConvertLightIntensityAttr(UsdIntensity, UsdExposure) * SolidAngle * AreaInSqMeters;	  // Lumen = Nits * SolidAngle * Area
 }
 
-float UsdToUnreal::ConvertConeAngleSoftnessAttr( float UsdConeAngle, float UsdConeSoftness, float& OutInnerConeAngle )
+float UsdToUnreal::ConvertConeAngleSoftnessAttr(float UsdConeAngle, float UsdConeSoftness, float& OutInnerConeAngle)
 {
-	OutInnerConeAngle = UsdConeAngle * ( 1.0f - UsdConeSoftness );
+	OutInnerConeAngle = UsdConeAngle * (1.0f - UsdConeSoftness);
 	float OuterConeAngle = UsdConeAngle;
 	return OuterConeAngle;
 }
 
-bool UnrealToUsd::ConvertLightComponent( const ULightComponentBase& LightComponent, pxr::UsdPrim& Prim, double UsdTimeCode )
+bool UnrealToUsd::ConvertLightComponent(const ULightComponentBase& LightComponent, pxr::UsdPrim& Prim, double UsdTimeCode)
 {
 	FScopedUsdAllocs UsdAllocs;
 
-	if ( !Prim )
+	if (!Prim)
 	{
 		return false;
 	}
 
-	pxr::UsdLuxLightAPI LightAPI( Prim );
-	if ( !LightAPI )
+	pxr::UsdLuxLightAPI LightAPI(Prim);
+	if (!LightAPI)
 	{
 		return false;
 	}
@@ -395,56 +434,56 @@ bool UnrealToUsd::ConvertLightComponent( const ULightComponentBase& LightCompone
 		return false;
 	}
 
-	if ( pxr::UsdAttribute Attr = LightAPI.CreateIntensityAttr() )
+	if (pxr::UsdAttribute Attr = LightAPI.CreateIntensityAttr())
 	{
-		Attr.Set<float>( LightComponent.Intensity, UsdTimeCode );
+		Attr.Set<float>(LightComponent.Intensity, UsdTimeCode);
 		UsdUtils::NotifyIfOverriddenOpinion(Attr);
 	}
 
 	// When converting into UE we multiply intensity and exposure together, so when writing back we just
 	// put everything in intensity. USD also multiplies those two together, meaning it should end up the same
-	if ( pxr::UsdAttribute Attr = LightAPI.CreateExposureAttr() )
+	if (pxr::UsdAttribute Attr = LightAPI.CreateExposureAttr())
 	{
-		Attr.Set<float>( 0.0f, UsdTimeCode );
+		Attr.Set<float>(0.0f, UsdTimeCode);
 		UsdUtils::NotifyIfOverriddenOpinion(Attr);
 	}
 
-	if ( const ULightComponent* DerivedLightComponent = Cast<ULightComponent>( &LightComponent ) )
+	if (const ULightComponent* DerivedLightComponent = Cast<ULightComponent>(&LightComponent))
 	{
-		if ( pxr::UsdAttribute Attr = LightAPI.CreateEnableColorTemperatureAttr() )
+		if (pxr::UsdAttribute Attr = LightAPI.CreateEnableColorTemperatureAttr())
 		{
-			Attr.Set<bool>( DerivedLightComponent->bUseTemperature, UsdTimeCode );
+			Attr.Set<bool>(DerivedLightComponent->bUseTemperature, UsdTimeCode);
 			UsdUtils::NotifyIfOverriddenOpinion(Attr);
 		}
 
-		if ( pxr::UsdAttribute Attr = LightAPI.CreateColorTemperatureAttr() )
+		if (pxr::UsdAttribute Attr = LightAPI.CreateColorTemperatureAttr())
 		{
-			Attr.Set<float>( DerivedLightComponent->Temperature, UsdTimeCode );
+			Attr.Set<float>(DerivedLightComponent->Temperature, UsdTimeCode);
 			UsdUtils::NotifyIfOverriddenOpinion(Attr);
 		}
 	}
 
-	if ( pxr::UsdAttribute Attr = LightAPI.CreateColorAttr() )
+	if (pxr::UsdAttribute Attr = LightAPI.CreateColorAttr())
 	{
-		pxr::GfVec4f LinearColor = UnrealToUsd::ConvertColor( LightComponent.LightColor );
-		Attr.Set<pxr::GfVec3f>( pxr::GfVec3f( LinearColor[0], LinearColor[1], LinearColor[2] ), UsdTimeCode );
+		pxr::GfVec4f LinearColor = UnrealToUsd::ConvertColor(LightComponent.LightColor);
+		Attr.Set<pxr::GfVec3f>(pxr::GfVec3f(LinearColor[0], LinearColor[1], LinearColor[2]), UsdTimeCode);
 		UsdUtils::NotifyIfOverriddenOpinion(Attr);
 	}
 
 	return true;
 }
 
-bool UnrealToUsd::ConvertDirectionalLightComponent( const UDirectionalLightComponent& LightComponent, pxr::UsdPrim& Prim, double UsdTimeCode )
+bool UnrealToUsd::ConvertDirectionalLightComponent(const UDirectionalLightComponent& LightComponent, pxr::UsdPrim& Prim, double UsdTimeCode)
 {
 	FScopedUsdAllocs UsdAllocs;
 
-	if ( !Prim )
+	if (!Prim)
 	{
 		return false;
 	}
 
-	pxr::UsdLuxDistantLight Light{ Prim };
-	if ( !Light )
+	pxr::UsdLuxDistantLight Light{Prim};
+	if (!Light)
 	{
 		return false;
 	}
@@ -454,9 +493,9 @@ bool UnrealToUsd::ConvertDirectionalLightComponent( const UDirectionalLightCompo
 		return false;
 	}
 
-	if ( pxr::UsdAttribute Attr = Light.CreateAngleAttr() )
+	if (pxr::UsdAttribute Attr = Light.CreateAngleAttr())
 	{
-		Attr.Set<float>( LightComponent.LightSourceAngle, UsdTimeCode );
+		Attr.Set<float>(LightComponent.LightSourceAngle, UsdTimeCode);
 		UsdUtils::NotifyIfOverriddenOpinion(Attr);
 	}
 
@@ -466,17 +505,17 @@ bool UnrealToUsd::ConvertDirectionalLightComponent( const UDirectionalLightCompo
 	return true;
 }
 
-bool UnrealToUsd::ConvertRectLightComponent( const URectLightComponent& LightComponent, pxr::UsdPrim& Prim, double UsdTimeCode )
+bool UnrealToUsd::ConvertRectLightComponent(const URectLightComponent& LightComponent, pxr::UsdPrim& Prim, double UsdTimeCode)
 {
 	FScopedUsdAllocs UsdAllocs;
 
-	if ( !Prim )
+	if (!Prim)
 	{
 		return false;
 	}
 
-	pxr::UsdLuxLightAPI LightAPI( Prim );
-	if ( !LightAPI )
+	pxr::UsdLuxLightAPI LightAPI(Prim);
+	if (!LightAPI)
 	{
 		return false;
 	}
@@ -486,36 +525,36 @@ bool UnrealToUsd::ConvertRectLightComponent( const URectLightComponent& LightCom
 		return false;
 	}
 
-	FUsdStageInfo StageInfo( Prim.GetStage() );
+	FUsdStageInfo StageInfo(Prim.GetStage());
 
 	// Disk light
 	float AreaInSqMeters = 1.0f;
-	if ( pxr::UsdLuxDiskLight DiskLight{ Prim } )
+	if (pxr::UsdLuxDiskLight DiskLight{Prim})
 	{
 		// Averaging and converting "diameter" to "radius"
-		const float Radius = ( LightComponent.SourceWidth + LightComponent.SourceHeight ) / 2.0f / 2.0f;
-		AreaInSqMeters = PI * FMath::Square( Radius / 100.0f );
+		const float Radius = (LightComponent.SourceWidth + LightComponent.SourceHeight) / 2.0f / 2.0f;
+		AreaInSqMeters = PI * FMath::Square(Radius / 100.0f);
 
-		if ( pxr::UsdAttribute Attr = DiskLight.CreateRadiusAttr() )
+		if (pxr::UsdAttribute Attr = DiskLight.CreateRadiusAttr())
 		{
-			Attr.Set<float>( UnrealToUsd::ConvertDistance( StageInfo, Radius ), UsdTimeCode );
+			Attr.Set<float>(UnrealToUsd::ConvertDistance(StageInfo, Radius), UsdTimeCode);
 			UsdUtils::NotifyIfOverriddenOpinion(Attr);
 		}
 	}
 	// Rect light
-	else if ( pxr::UsdLuxRectLight RectLight{ Prim } )
+	else if (pxr::UsdLuxRectLight RectLight{Prim})
 	{
-		AreaInSqMeters = ( LightComponent.SourceWidth / 100.0f ) * ( LightComponent.SourceHeight / 100.0f );
+		AreaInSqMeters = (LightComponent.SourceWidth / 100.0f) * (LightComponent.SourceHeight / 100.0f);
 
-		if ( pxr::UsdAttribute Attr = RectLight.CreateWidthAttr() )
+		if (pxr::UsdAttribute Attr = RectLight.CreateWidthAttr())
 		{
-			Attr.Set<float>( UnrealToUsd::ConvertDistance( StageInfo, LightComponent.SourceWidth ), UsdTimeCode );
+			Attr.Set<float>(UnrealToUsd::ConvertDistance(StageInfo, LightComponent.SourceWidth), UsdTimeCode);
 			UsdUtils::NotifyIfOverriddenOpinion(Attr);
 		}
 
-		if ( pxr::UsdAttribute Attr = RectLight.CreateHeightAttr() )
+		if (pxr::UsdAttribute Attr = RectLight.CreateHeightAttr())
 		{
-			Attr.Set<float>( UnrealToUsd::ConvertDistance( StageInfo, LightComponent.SourceHeight ), UsdTimeCode );
+			Attr.Set<float>(UnrealToUsd::ConvertDistance(StageInfo, LightComponent.SourceHeight), UsdTimeCode);
 			UsdUtils::NotifyIfOverriddenOpinion(Attr);
 		}
 	}
@@ -525,64 +564,64 @@ bool UnrealToUsd::ConvertRectLightComponent( const URectLightComponent& LightCom
 	}
 
 	// Common for both
-	if ( pxr::UsdAttribute Attr = LightAPI.CreateIntensityAttr() )
+	if (pxr::UsdAttribute Attr = LightAPI.CreateIntensityAttr())
 	{
 		float FinalIntensityNits = 1.0f;
-		float OldIntensity = UsdUtils::GetUsdValue< float >( Attr, UsdTimeCode );
+		float OldIntensity = UsdUtils::GetUsdValue<float>(Attr, UsdTimeCode);
 
 		// Area light with no area probably shouldn't emit any light?
 		// It's not possible to set width/height less than 1 via the Details panel anyway, but just in case
-		if ( FMath::IsNearlyZero( AreaInSqMeters ) )
+		if (FMath::IsNearlyZero(AreaInSqMeters))
 		{
 			OldIntensity = 0.0f;
 		}
 
-		AreaInSqMeters = FMath::Max( AreaInSqMeters, KINDA_SMALL_NUMBER );
+		AreaInSqMeters = FMath::Max(AreaInSqMeters, KINDA_SMALL_NUMBER);
 
-		switch ( LightComponent.IntensityUnits )
+		switch (LightComponent.IntensityUnits)
 		{
-		case ELightUnits::Candelas:
-			// Nit = candela / area
-			FinalIntensityNits = OldIntensity / AreaInSqMeters;
-			break;
-		case ELightUnits::Lumens:
-			// Nit = lumen / ( sr * area ); For area lights sr is technically 2PI, but we cancel that with an
-			// extra factor of 2.0 here because URectLightComponent::SetLightBrightness uses just PI and not 2PI as steradian
-			// due to some cosine distribution. This also matches the PI factor between candelas and lumen for rect lights on
-			// https://docs.unrealengine.com/en-US/Engine/Rendering/LightingAndShadows/PhysicalLightUnits/index.html#point,spot,andrectlights
-			FinalIntensityNits = OldIntensity / ( PI * AreaInSqMeters );
-			break;
-		case ELightUnits::EV:
-			// Nit = luminance (cd/m2)
-			FinalIntensityNits = EV100ToLuminance(OldIntensity);
-			break;
-		case ELightUnits::Unitless:
-			// Nit = (unitless/625) / area = candela / area
-			// https://docs.unrealengine.com/en-US/Engine/Rendering/LightingAndShadows/PhysicalLightUnits/index.html#point,spot,andrectlights
-			FinalIntensityNits = ( OldIntensity / 625.0f ) / AreaInSqMeters;
-			break;
-		default:
-			break;
+			case ELightUnits::Candelas:
+				// Nit = candela / area
+				FinalIntensityNits = OldIntensity / AreaInSqMeters;
+				break;
+			case ELightUnits::Lumens:
+				// Nit = lumen / ( sr * area ); For area lights sr is technically 2PI, but we cancel that with an
+				// extra factor of 2.0 here because URectLightComponent::SetLightBrightness uses just PI and not 2PI as steradian
+				// due to some cosine distribution. This also matches the PI factor between candelas and lumen for rect lights on
+				// https://docs.unrealengine.com/en-US/Engine/Rendering/LightingAndShadows/PhysicalLightUnits/index.html#point,spot,andrectlights
+				FinalIntensityNits = OldIntensity / (PI * AreaInSqMeters);
+				break;
+			case ELightUnits::EV:
+				// Nit = luminance (cd/m2)
+				FinalIntensityNits = EV100ToLuminance(OldIntensity);
+				break;
+			case ELightUnits::Unitless:
+				// Nit = (unitless/625) / area = candela / area
+				// https://docs.unrealengine.com/en-US/Engine/Rendering/LightingAndShadows/PhysicalLightUnits/index.html#point,spot,andrectlights
+				FinalIntensityNits = (OldIntensity / 625.0f) / AreaInSqMeters;
+				break;
+			default:
+				break;
 		}
 
-		Attr.Set<float>( FinalIntensityNits, UsdTimeCode );
+		Attr.Set<float>(FinalIntensityNits, UsdTimeCode);
 		UsdUtils::NotifyIfOverriddenOpinion(Attr);
 	}
 
 	return true;
 }
 
-bool UnrealToUsd::ConvertPointLightComponent( const UPointLightComponent& LightComponent, pxr::UsdPrim& Prim, double UsdTimeCode )
+bool UnrealToUsd::ConvertPointLightComponent(const UPointLightComponent& LightComponent, pxr::UsdPrim& Prim, double UsdTimeCode)
 {
 	FScopedUsdAllocs Allocs;
 
-	if ( !Prim )
+	if (!Prim)
 	{
 		return false;
 	}
 
-	pxr::UsdLuxSphereLight Light{ Prim };
-	if ( !Light )
+	pxr::UsdLuxSphereLight Light{Prim};
+	if (!Light)
 	{
 		return false;
 	}
@@ -592,76 +631,76 @@ bool UnrealToUsd::ConvertPointLightComponent( const UPointLightComponent& LightC
 		return false;
 	}
 
-	FUsdStageInfo StageInfo( Prim.GetStage() );
+	FUsdStageInfo StageInfo(Prim.GetStage());
 
-	if ( pxr::UsdAttribute Attr = Light.CreateRadiusAttr() )
+	if (pxr::UsdAttribute Attr = Light.CreateRadiusAttr())
 	{
-		Attr.Set<float>( UnrealToUsd::ConvertDistance( StageInfo, LightComponent.SourceRadius ), UsdTimeCode );
+		Attr.Set<float>(UnrealToUsd::ConvertDistance(StageInfo, LightComponent.SourceRadius), UsdTimeCode);
 		UsdUtils::NotifyIfOverriddenOpinion(Attr);
 	}
 
-	if ( pxr::UsdAttribute Attr = Light.CreateTreatAsPointAttr() )
+	if (pxr::UsdAttribute Attr = Light.CreateTreatAsPointAttr())
 	{
-		Attr.Set<bool>( FMath::IsNearlyZero( LightComponent.SourceRadius ), UsdTimeCode );
+		Attr.Set<bool>(FMath::IsNearlyZero(LightComponent.SourceRadius), UsdTimeCode);
 		UsdUtils::NotifyIfOverriddenOpinion(Attr);
 	}
 
 	float SolidAngle = 4.f * PI;
-	if ( const USpotLightComponent* SpotLightComponent = Cast<const USpotLightComponent>( &LightComponent ) )
+	if (const USpotLightComponent* SpotLightComponent = Cast<const USpotLightComponent>(&LightComponent))
 	{
-		SolidAngle = 2.f * PI * ( 1.0f - SpotLightComponent->GetCosHalfConeAngle() );
+		SolidAngle = 2.f * PI * (1.0f - SpotLightComponent->GetCosHalfConeAngle());
 	}
 
 	// It doesn't make much physical sense to use nits for point lights in this way, but USD light intensities are always in nits
 	// so we must do something. We do the analogue on the UsdToUnreal conversion, at least. Also using the solid angle for the
 	// area calculation is possibly incorrect, but I think it depends on the chosen convention
-	const float AreaInSqMeters = FMath::Max( SolidAngle * FMath::Square( LightComponent.SourceRadius / 100.f ), KINDA_SMALL_NUMBER );
-	if ( pxr::UsdAttribute Attr = Light.CreateIntensityAttr() )
+	const float AreaInSqMeters = FMath::Max(SolidAngle * FMath::Square(LightComponent.SourceRadius / 100.f), KINDA_SMALL_NUMBER);
+	if (pxr::UsdAttribute Attr = Light.CreateIntensityAttr())
 	{
 		float FinalIntensityNits = 1.0f;
-		const float OldIntensity = UsdUtils::GetUsdValue< float >( Attr, UsdTimeCode );
+		const float OldIntensity = UsdUtils::GetUsdValue<float>(Attr, UsdTimeCode);
 
-		switch ( LightComponent.IntensityUnits )
+		switch (LightComponent.IntensityUnits)
 		{
-		case ELightUnits::Candelas:
-			// Nit = candela / area
-			FinalIntensityNits = OldIntensity / AreaInSqMeters;
-			break;
-		case ELightUnits::Lumens:
-			// Nit = lumen / (sr * area); For a full sphere sr = 4PI
-			FinalIntensityNits = OldIntensity / ( SolidAngle * AreaInSqMeters );
-			break;
-		case ELightUnits::EV:
-			// Nit = luminance (cd/m2)
-			FinalIntensityNits = EV100ToLuminance(OldIntensity);
-			break;
-		case ELightUnits::Unitless:
-			// Nit = (unitless/625) / area = candela / area
-			// https://docs.unrealengine.com/en-US/Engine/Rendering/LightingAndShadows/PhysicalLightUnits/index.html#point,spot,andrectlights
-			FinalIntensityNits = ( OldIntensity / 625.0f ) / AreaInSqMeters;
-			break;
-		default:
-			break;
+			case ELightUnits::Candelas:
+				// Nit = candela / area
+				FinalIntensityNits = OldIntensity / AreaInSqMeters;
+				break;
+			case ELightUnits::Lumens:
+				// Nit = lumen / (sr * area); For a full sphere sr = 4PI
+				FinalIntensityNits = OldIntensity / (SolidAngle * AreaInSqMeters);
+				break;
+			case ELightUnits::EV:
+				// Nit = luminance (cd/m2)
+				FinalIntensityNits = EV100ToLuminance(OldIntensity);
+				break;
+			case ELightUnits::Unitless:
+				// Nit = (unitless/625) / area = candela / area
+				// https://docs.unrealengine.com/en-US/Engine/Rendering/LightingAndShadows/PhysicalLightUnits/index.html#point,spot,andrectlights
+				FinalIntensityNits = (OldIntensity / 625.0f) / AreaInSqMeters;
+				break;
+			default:
+				break;
 		}
 
-		Attr.Set<float>( FinalIntensityNits, UsdTimeCode );
+		Attr.Set<float>(FinalIntensityNits, UsdTimeCode);
 		UsdUtils::NotifyIfOverriddenOpinion(Attr);
 	}
 
 	return true;
 }
 
-bool UnrealToUsd::ConvertSkyLightComponent( const USkyLightComponent& LightComponent, pxr::UsdPrim& Prim, double UsdTimeCode )
+bool UnrealToUsd::ConvertSkyLightComponent(const USkyLightComponent& LightComponent, pxr::UsdPrim& Prim, double UsdTimeCode)
 {
 	FScopedUsdAllocs Allocs;
 
-	if ( !Prim )
+	if (!Prim)
 	{
 		return false;
 	}
 
-	pxr::UsdLuxDomeLight Light{ Prim };
-	if ( !Light )
+	pxr::UsdLuxDomeLight Light{Prim};
+	if (!Light)
 	{
 		return false;
 	}
@@ -672,41 +711,45 @@ bool UnrealToUsd::ConvertSkyLightComponent( const USkyLightComponent& LightCompo
 	}
 
 #if WITH_EDITORONLY_DATA
-	FUsdStageInfo StageInfo( Prim.GetStage() );
+	FUsdStageInfo StageInfo(Prim.GetStage());
 
-	if ( pxr::UsdAttribute Attr = Light.CreateTextureFileAttr() )
+	if (pxr::UsdAttribute Attr = Light.CreateTextureFileAttr())
 	{
-		if ( UTextureCube* TextureCube = LightComponent.Cubemap )
+		if (UTextureCube* TextureCube = LightComponent.Cubemap)
 		{
-			if ( UAssetImportData* AssetImportData = TextureCube->AssetImportData )
+			if (UAssetImportData* AssetImportData = TextureCube->AssetImportData)
 			{
 				FString FilePath = AssetImportData->GetFirstFilename();
-				if ( !FPaths::FileExists( FilePath ) )
+				if (!FPaths::FileExists(FilePath))
 				{
-					UE_LOG(LogUsd, Warning, TEXT("Used '%s' as cubemap when converting SkyLightComponent '%s' onto prim '%s', but the cubemap does not exist on the filesystem!"),
+					UE_LOG(
+						LogUsd,
+						Warning,
+						TEXT("Used '%s' as cubemap when converting SkyLightComponent '%s' onto prim '%s', but the cubemap does not exist on the "
+							 "filesystem!"),
 						*FilePath,
 						*LightComponent.GetPathName(),
 						*UsdToUnreal::ConvertPath(Prim.GetPrimPath())
 					);
 				}
 
-				UsdUtils::MakePathRelativeToLayer( UE::FSdfLayer{ Prim.GetStage()->GetEditTarget().GetLayer() }, FilePath );
-				Attr.Set<pxr::SdfAssetPath>( pxr::SdfAssetPath{ UnrealToUsd::ConvertString( *FilePath ).Get() }, UsdTimeCode );
+				UsdUtils::MakePathRelativeToLayer(UE::FSdfLayer{Prim.GetStage()->GetEditTarget().GetLayer()}, FilePath);
+				Attr.Set<pxr::SdfAssetPath>(pxr::SdfAssetPath{UnrealToUsd::ConvertString(*FilePath).Get()}, UsdTimeCode);
 				UsdUtils::NotifyIfOverriddenOpinion(Attr);
 			}
 		}
 	}
-#endif //  #if WITH_EDITORONLY_DATA
+#endif	  //  #if WITH_EDITORONLY_DATA
 
 	return true;
 }
 
-bool UnrealToUsd::ConvertSpotLightComponent( const USpotLightComponent& LightComponent, pxr::UsdPrim& Prim, double UsdTimeCode )
+bool UnrealToUsd::ConvertSpotLightComponent(const USpotLightComponent& LightComponent, pxr::UsdPrim& Prim, double UsdTimeCode)
 {
 	FScopedUsdAllocs Allocs;
 
-	pxr::UsdLuxShapingAPI ShapingAPI = pxr::UsdLuxShapingAPI::Apply( Prim );
-	if ( !ShapingAPI )
+	pxr::UsdLuxShapingAPI ShapingAPI = pxr::UsdLuxShapingAPI::Apply(Prim);
+	if (!ShapingAPI)
 	{
 		return false;
 	}
@@ -716,108 +759,124 @@ bool UnrealToUsd::ConvertSpotLightComponent( const USpotLightComponent& LightCom
 		return false;
 	}
 
-	if ( pxr::UsdAttribute ConeAngleAttr = ShapingAPI.CreateShapingConeAngleAttr() )
+	if (pxr::UsdAttribute ConeAngleAttr = ShapingAPI.CreateShapingConeAngleAttr())
 	{
-		ConeAngleAttr.Set<float>( LightComponent.OuterConeAngle, UsdTimeCode );
+		ConeAngleAttr.Set<float>(LightComponent.OuterConeAngle, UsdTimeCode);
 		UsdUtils::NotifyIfOverriddenOpinion(ConeAngleAttr);
 	}
 
-	// As of March 2021 there doesn't seem to be a consensus on what softness means, according to https://groups.google.com/g/usd-interest/c/A6bc4OZjSB0
-	// We approximate the best look here by trying to convert from inner/outer cone angle to softness according to the renderman docs
-	if ( pxr::UsdAttribute SoftnessAttr = ShapingAPI.CreateShapingConeSoftnessAttr() )
+	// As of March 2021 there doesn't seem to be a consensus on what softness means, according to
+	// https://groups.google.com/g/usd-interest/c/A6bc4OZjSB0 We approximate the best look here by trying to convert from inner/outer cone angle to
+	// softness according to the renderman docs
+	if (pxr::UsdAttribute SoftnessAttr = ShapingAPI.CreateShapingConeSoftnessAttr())
 	{
 		// Keep in [0, 1] range, where 1 is maximum softness, i.e. inner cone angle is zero
-		const float Softness = FMath::IsNearlyZero( LightComponent.OuterConeAngle ) ? 0.0 : 1.0f - LightComponent.InnerConeAngle / LightComponent.OuterConeAngle;
-		SoftnessAttr.Set<float>( Softness, UsdTimeCode );
+		const float Softness = FMath::IsNearlyZero(LightComponent.OuterConeAngle)
+								   ? 0.0
+								   : 1.0f - LightComponent.InnerConeAngle / LightComponent.OuterConeAngle;
+		SoftnessAttr.Set<float>(Softness, UsdTimeCode);
 		UsdUtils::NotifyIfOverriddenOpinion(SoftnessAttr);
 	}
 
 	return true;
 }
 
-float UnrealToUsd::ConvertLightIntensityProperty( float Intensity )
+float UnrealToUsd::ConvertLightIntensityProperty(float Intensity)
 {
 	return Intensity;
 }
 
-float UnrealToUsd::ConvertRectLightIntensityProperty( float Intensity, float Width, float Height, const FUsdStageInfo& StageInfo, ELightUnits SourceUnits )
+float UnrealToUsd::ConvertRectLightIntensityProperty(
+	float Intensity,
+	float Width,
+	float Height,
+	const FUsdStageInfo& StageInfo,
+	ELightUnits SourceUnits
+)
 {
-	float UsdIntensity = UnrealToUsd::ConvertLightIntensityProperty( Intensity );
+	float UsdIntensity = UnrealToUsd::ConvertLightIntensityProperty(Intensity);
 
-	float AreaInSqMeters = ( Width / 100.0f ) * ( Height / 100.0f );
+	float AreaInSqMeters = (Width / 100.0f) * (Height / 100.0f);
 
-	if ( FMath::IsNearlyZero( AreaInSqMeters ) )
+	if (FMath::IsNearlyZero(AreaInSqMeters))
 	{
 		UsdIntensity = 0.0f;
 	}
 
-	AreaInSqMeters = FMath::Max( AreaInSqMeters, KINDA_SMALL_NUMBER );
+	AreaInSqMeters = FMath::Max(AreaInSqMeters, KINDA_SMALL_NUMBER);
 
 	// For area lights sr is technically 2PI, but we cancel that with an
 	// extra factor of 2.0 here because URectLightComponent::SetLightBrightness uses just PI and not 2PI as steradian
 	// due to some cosine distribution. This also matches the PI factor between candelas and lumen for rect lights on
 	// https://docs.unrealengine.com/en-US/Engine/Rendering/LightingAndShadows/PhysicalLightUnits/index.html#point,spot,andrectlights
-	return LightConversionImpl::ConvertIntensityToNits( UsdIntensity, PI, AreaInSqMeters, SourceUnits );
+	return LightConversionImpl::ConvertIntensityToNits(UsdIntensity, PI, AreaInSqMeters, SourceUnits);
 }
 
-float UnrealToUsd::ConvertRectLightIntensityProperty( float Intensity, float Radius, const FUsdStageInfo& StageInfo, ELightUnits SourceUnits )
+float UnrealToUsd::ConvertRectLightIntensityProperty(float Intensity, float Radius, const FUsdStageInfo& StageInfo, ELightUnits SourceUnits)
 {
-	float UsdIntensity = UnrealToUsd::ConvertLightIntensityProperty( Intensity );
+	float UsdIntensity = UnrealToUsd::ConvertLightIntensityProperty(Intensity);
 
-	float AreaInSqMeters = PI * FMath::Square( Radius / 100.0f );
+	float AreaInSqMeters = PI * FMath::Square(Radius / 100.0f);
 
-	if ( FMath::IsNearlyZero( AreaInSqMeters ) )
+	if (FMath::IsNearlyZero(AreaInSqMeters))
 	{
 		UsdIntensity = 0.0f;
 	}
 
-	AreaInSqMeters = FMath::Max( AreaInSqMeters, KINDA_SMALL_NUMBER );
+	AreaInSqMeters = FMath::Max(AreaInSqMeters, KINDA_SMALL_NUMBER);
 
 	// For area lights sr is technically 2PI, but we cancel that with an
 	// extra factor of 2.0 here because URectLightComponent::SetLightBrightness uses just PI and not 2PI as steradian
 	// due to some cosine distribution. This also matches the PI factor between candelas and lumen for rect lights on
 	// https://docs.unrealengine.com/en-US/Engine/Rendering/LightingAndShadows/PhysicalLightUnits/index.html#point,spot,andrectlights
-	return LightConversionImpl::ConvertIntensityToNits( UsdIntensity, PI, AreaInSqMeters, SourceUnits );
+	return LightConversionImpl::ConvertIntensityToNits(UsdIntensity, PI, AreaInSqMeters, SourceUnits);
 }
 
-float UnrealToUsd::ConvertPointLightIntensityProperty( float Intensity, float SourceRadius, const FUsdStageInfo& StageInfo, ELightUnits SourceUnits )
+float UnrealToUsd::ConvertPointLightIntensityProperty(float Intensity, float SourceRadius, const FUsdStageInfo& StageInfo, ELightUnits SourceUnits)
 {
-	const float UsdIntensity = UnrealToUsd::ConvertLightIntensityProperty( Intensity );
+	const float UsdIntensity = UnrealToUsd::ConvertLightIntensityProperty(Intensity);
 
 	const float SolidAngle = 4.f * PI;
 
 	// It doesn't make much physical sense to use nits for point lights in this way, but USD light intensities are always in nits
 	// so we must do something. We do the analogue on the UsdToUnreal conversion, at least. Also using the solid angle for the
 	// area calculation is possibly incorrect, but I think it depends on the chosen convention
-	const float AreaInSqMeters = FMath::Max( SolidAngle * FMath::Square( SourceRadius / 100.f ), KINDA_SMALL_NUMBER );
+	const float AreaInSqMeters = FMath::Max(SolidAngle * FMath::Square(SourceRadius / 100.f), KINDA_SMALL_NUMBER);
 
-	return LightConversionImpl::ConvertIntensityToNits( UsdIntensity, SolidAngle, AreaInSqMeters, SourceUnits );
+	return LightConversionImpl::ConvertIntensityToNits(UsdIntensity, SolidAngle, AreaInSqMeters, SourceUnits);
 }
 
-float UnrealToUsd::ConvertSpotLightIntensityProperty( float Intensity, float OuterConeAngle, float InnerConeAngle, float SourceRadius, const FUsdStageInfo& StageInfo, ELightUnits SourceUnits )
+float UnrealToUsd::ConvertSpotLightIntensityProperty(
+	float Intensity,
+	float OuterConeAngle,
+	float InnerConeAngle,
+	float SourceRadius,
+	const FUsdStageInfo& StageInfo,
+	ELightUnits SourceUnits
+)
 {
-	const float UsdIntensity = UnrealToUsd::ConvertLightIntensityProperty( Intensity );
+	const float UsdIntensity = UnrealToUsd::ConvertLightIntensityProperty(Intensity);
 
 	// c.f. USpotLightComponent::ComputeLightBrightness
-	const float SolidAngle = 2.f * PI * ( 1.0f - LightConversionImpl::GetSpotLightCosHalfConeAngle( OuterConeAngle, InnerConeAngle ) );
+	const float SolidAngle = 2.f * PI * (1.0f - LightConversionImpl::GetSpotLightCosHalfConeAngle(OuterConeAngle, InnerConeAngle));
 
 	// It doesn't make much physical sense to use nits for point lights in this way, but USD light intensities are always in nits
 	// so we must do something. We do the analogue on the UsdToUnreal conversion, at least. Also using the solid angle for the
 	// area calculation is possibly incorrect, but I think it depends on the chosen convention
-	const float AreaInSqMeters = FMath::Max( SolidAngle * FMath::Square( SourceRadius / 100.f ), KINDA_SMALL_NUMBER );
+	const float AreaInSqMeters = FMath::Max(SolidAngle * FMath::Square(SourceRadius / 100.f), KINDA_SMALL_NUMBER);
 
-	return LightConversionImpl::ConvertIntensityToNits( UsdIntensity, SolidAngle, AreaInSqMeters, SourceUnits );
+	return LightConversionImpl::ConvertIntensityToNits(UsdIntensity, SolidAngle, AreaInSqMeters, SourceUnits);
 }
 
-float UnrealToUsd::ConvertOuterConeAngleProperty( float OuterConeAngle )
+float UnrealToUsd::ConvertOuterConeAngleProperty(float OuterConeAngle)
 {
 	return OuterConeAngle;
 }
 
-float UnrealToUsd::ConvertInnerConeAngleProperty( float InnerConeAngle, float OuterConeAngle )
+float UnrealToUsd::ConvertInnerConeAngleProperty(float InnerConeAngle, float OuterConeAngle)
 {
 	// Keep in [0, 1] range, where 1 is maximum softness, i.e. inner cone angle is zero
-	return FMath::IsNearlyZero( OuterConeAngle ) ? 0.0 : 1.0f - InnerConeAngle / OuterConeAngle;
+	return FMath::IsNearlyZero(OuterConeAngle) ? 0.0 : 1.0f - InnerConeAngle / OuterConeAngle;
 }
 
-#endif // #if USE_USD_SDK
+#endif	  // #if USE_USD_SDK
