USD: Reformat USDUtilities module.

diff --git a/Engine/Plugins/Importers/USDImporter/Source/USDUtilities/Private/USDPrimConversion.cpp b/Engine/Plugins/Importers/USDImporter/Source/USDUtilities/Private/USDPrimConversion.cpp
--- a/Engine/Plugins/Importers/USDImporter/Source/USDUtilities/Private/USDPrimConversion.cpp
+++ b/Engine/Plugins/Importers/USDImporter/Source/USDUtilities/Private/USDPrimConversion.cpp
@@ -55,73 +55,72 @@
 #include "Sections/MovieSceneBoolSection.h"
 #include "Sections/MovieSceneColorSection.h"
 #include "Sections/MovieSceneFloatSection.h"
-#include "Sections/MovieSceneVisibilitySection.h"
 #include "Sections/MovieSceneVectorSection.h"
+#include "Sections/MovieSceneVisibilitySection.h"
 #include "Tracks/MovieScene3DTransformTrack.h"
 #include "Tracks/MovieSceneBoolTrack.h"
 #include "Tracks/MovieSceneColorTrack.h"
 #include "Tracks/MovieSceneFloatTrack.h"
 #include "Tracks/MovieScenePropertyTrack.h"
-#include "Tracks/MovieSceneVisibilityTrack.h"
 #include "Tracks/MovieSceneVectorTrack.h"
+#include "Tracks/MovieSceneVisibilityTrack.h"
 
 #if USE_USD_SDK
 
 #include "USDIncludesStart.h"
-	#include "pxr/base/tf/stringUtils.h"
-	#include "pxr/base/vt/value.h"
-	#include "pxr/usd/kind/registry.h"
-	#include "pxr/usd/sdf/changeBlock.h"
-	#include "pxr/usd/sdf/path.h"
-	#include "pxr/usd/sdf/types.h"
-	#include "pxr/usd/usd/attribute.h"
-	#include "pxr/usd/usd/prim.h"
-	#include "pxr/usd/usd/primRange.h"
-	#include "pxr/usd/usd/stage.h"
-	#include "pxr/usd/usd/stage.h"
-	#include "pxr/usd/usd/timeCode.h"
-	#include "pxr/usd/usdGeom/camera.h"
-	#include "pxr/usd/usdGeom/capsule.h"
-	#include "pxr/usd/usdGeom/cone.h"
-	#include "pxr/usd/usdGeom/cube.h"
-	#include "pxr/usd/usdGeom/cylinder.h"
-	#include "pxr/usd/usdGeom/imageable.h"
-	#include "pxr/usd/usdGeom/mesh.h"
-	#include "pxr/usd/usdGeom/modelAPI.h"
-	#include "pxr/usd/usdGeom/plane.h"
-	#include "pxr/usd/usdGeom/pointInstancer.h"
-	#include "pxr/usd/usdGeom/scope.h"
-	#include "pxr/usd/usdGeom/sphere.h"
-	#include "pxr/usd/usdGeom/tokens.h"
-	#include "pxr/usd/usdGeom/tokens.h"
-	#include "pxr/usd/usdGeom/xform.h"
-	#include "pxr/usd/usdGeom/xformable.h"
-	#include "pxr/usd/usdGeom/xformCommonAPI.h"
-	#include "pxr/usd/usdLux/diskLight.h"
-	#include "pxr/usd/usdLux/distantLight.h"
-	#include "pxr/usd/usdLux/lightAPI.h"
-	#include "pxr/usd/usdLux/rectLight.h"
-	#include "pxr/usd/usdLux/shapingAPI.h"
-	#include "pxr/usd/usdLux/sphereLight.h"
-	#include "pxr/usd/usdLux/tokens.h"
-	#include "pxr/usd/usdShade/connectableAPI.h"
-	#include "pxr/usd/usdShade/material.h"
-	#include "pxr/usd/usdShade/materialBindingAPI.h"
-	#include "pxr/usd/usdShade/shader.h"
-	#include "pxr/usd/usdShade/tokens.h"
-	#include "pxr/usd/usdSkel/animation.h"
-	#include "pxr/usd/usdSkel/binding.h"
-	#include "pxr/usd/usdSkel/root.h"
-	#include "pxr/usd/usdSkel/skeleton.h"
-	#include "pxr/usd/usdSkel/skeletonQuery.h"
-	#include "pxr/usd/usdUtils/stageCache.h"
+#include "pxr/base/tf/stringUtils.h"
+#include "pxr/base/vt/value.h"
+#include "pxr/usd/kind/registry.h"
+#include "pxr/usd/sdf/changeBlock.h"
+#include "pxr/usd/sdf/path.h"
+#include "pxr/usd/sdf/types.h"
+#include "pxr/usd/usd/attribute.h"
+#include "pxr/usd/usd/prim.h"
+#include "pxr/usd/usd/primRange.h"
+#include "pxr/usd/usd/stage.h"
+#include "pxr/usd/usd/timeCode.h"
+#include "pxr/usd/usdGeom/camera.h"
+#include "pxr/usd/usdGeom/capsule.h"
+#include "pxr/usd/usdGeom/cone.h"
+#include "pxr/usd/usdGeom/cube.h"
+#include "pxr/usd/usdGeom/cylinder.h"
+#include "pxr/usd/usdGeom/imageable.h"
+#include "pxr/usd/usdGeom/mesh.h"
+#include "pxr/usd/usdGeom/modelAPI.h"
+#include "pxr/usd/usdGeom/plane.h"
+#include "pxr/usd/usdGeom/pointInstancer.h"
+#include "pxr/usd/usdGeom/scope.h"
+#include "pxr/usd/usdGeom/sphere.h"
+#include "pxr/usd/usdGeom/tokens.h"
+#include "pxr/usd/usdGeom/xform.h"
+#include "pxr/usd/usdGeom/xformable.h"
+#include "pxr/usd/usdGeom/xformCommonAPI.h"
+#include "pxr/usd/usdLux/diskLight.h"
+#include "pxr/usd/usdLux/distantLight.h"
+#include "pxr/usd/usdLux/lightAPI.h"
+#include "pxr/usd/usdLux/rectLight.h"
+#include "pxr/usd/usdLux/shapingAPI.h"
+#include "pxr/usd/usdLux/sphereLight.h"
+#include "pxr/usd/usdLux/tokens.h"
+#include "pxr/usd/usdShade/connectableAPI.h"
+#include "pxr/usd/usdShade/material.h"
+#include "pxr/usd/usdShade/materialBindingAPI.h"
+#include "pxr/usd/usdShade/shader.h"
+#include "pxr/usd/usdSkel/animation.h"
+#include "pxr/usd/usdSkel/binding.h"
+#include "pxr/usd/usdSkel/root.h"
+#include "pxr/usd/usdSkel/skeleton.h"
+#include "pxr/usd/usdSkel/skeletonQuery.h"
+#include "pxr/usd/usdUtils/stageCache.h"
 #include "USDIncludesEnd.h"
 
 static bool GConsiderAllPrimsHaveAnimatedBounds = false;
 static FAutoConsoleVariableRef CVarConsiderAllPrimsHaveAnimatedBounds(
 	TEXT("USD.Bounds.ConsiderAllPrimsHaveAnimatedBounds"),
 	GConsiderAllPrimsHaveAnimatedBounds,
-	TEXT("When active prevents USD from caching computed bounds between timeSamples for any prim, which allows us to force it to recompute accurate bounds for cases it does not naturally consider animated (e.g. for animated Mesh points, skeletal animation, etc.). Warning: This can be extremely expensive!")
+	TEXT("When active prevents USD from caching computed bounds between timeSamples for any prim, which allows us to force it to recompute accurate "
+		 "bounds for cases it does not naturally consider animated (e.g. for animated Mesh points, skeletal animation, etc.). Warning: This can be "
+		 "extremely expensive!")
 );
 
 namespace UE
@@ -134,27 +133,27 @@ namespace UE
 			// it already has an 'unreal' render context output and the expected structure, or by creating a new Material prim
 			// that fulfills those requirements.
 			// Doesn't write to 'unrealMaterial' at all, as we intend on deprecating it in the future.
-			void AuthorMaterialOverride( pxr::UsdPrim& MeshPrim, const FString& UEMaterialAssetPath )
+			void AuthorMaterialOverride(pxr::UsdPrim& MeshPrim, const FString& UEMaterialAssetPath)
 			{
-				if ( !MeshPrim )
+				if (!MeshPrim)
 				{
 					return;
 				}
 
 				FScopedUsdAllocs UsdAllocs;
 
-				pxr::UsdShadeMaterialBindingAPI BindingAPI = pxr::UsdShadeMaterialBindingAPI::Apply( MeshPrim );
+				pxr::UsdShadeMaterialBindingAPI BindingAPI = pxr::UsdShadeMaterialBindingAPI::Apply(MeshPrim);
 
 				// If this mesh prim already has a binding to a *child* material with the 'unreal' render context,
 				// just write our material there and early out
-				if ( pxr::UsdShadeMaterial ShadeMaterial = BindingAPI.ComputeBoundMaterial() )
+				if (pxr::UsdShadeMaterial ShadeMaterial = BindingAPI.ComputeBoundMaterial())
 				{
 					// Only consider this material reusable if its within the same layer as the edit target, otherwise we'll
 					// prefer to author something else that can be fully defined on the edit target. This is handy when we're
 					// exporting a level, as we'll ensure we're making these prims on the MaterialOverrides layer
 					pxr::SdfLayerRefPtr EditTarget = MeshPrim.GetStage()->GetEditTarget().GetLayer();
-					pxr::SdfLayerRefPtr MaterialPrimLayer = static_cast< pxr::SdfLayerRefPtr >( UsdUtils::FindLayerForPrim( ShadeMaterial.GetPrim() ) );
-					if ( EditTarget == MaterialPrimLayer )
+					pxr::SdfLayerRefPtr MaterialPrimLayer = static_cast<pxr::SdfLayerRefPtr>(UsdUtils::FindLayerForPrim(ShadeMaterial.GetPrim()));
+					if (EditTarget == MaterialPrimLayer)
 					{
 						// We need to try reusing these materials or else we'd write a new material prim every time we change
 						// the override in UE, but we also run the risk of modifying a material that is used by multiple prims
@@ -162,11 +161,11 @@ namespace UE
 						// material if it is a child of MeshPrim already, and always to author our material prims as children
 						std::string MaterialPath = ShadeMaterial.GetPrim().GetPath().GetString();
 						std::string MeshPrimPath = MeshPrim.GetPath().GetString();
-						if ( MaterialPath.rfind( MeshPrimPath, 0 ) == 0 )
+						if (MaterialPath.rfind(MeshPrimPath, 0) == 0)
 						{
-							if ( pxr::UsdPrim MaterialPrim = ShadeMaterial.GetPrim() )
+							if (pxr::UsdPrim MaterialPrim = ShadeMaterial.GetPrim())
 							{
-								UsdUtils::SetUnrealSurfaceOutput( MaterialPrim, UEMaterialAssetPath );
+								UsdUtils::SetUnrealSurfaceOutput(MaterialPrim, UEMaterialAssetPath);
 								return;
 							}
 						}
@@ -178,47 +177,50 @@ namespace UE
 				// Material prim to use for multiple overrides: This because in the future we'll want to have a separate material
 				// bake for each mesh (to make sure we get vertex color effects, etc.), and so we'd have multiple baked .usda material
 				// asset layers for each UE material, and we'd want each mesh/section/LOD to refer to its own anyway
-				FString ChildMaterialName = TEXT( "UnrealMaterial" );
-				if ( pxr::UsdPrim ExistingPrim = MeshPrim.GetChild( UnrealToUsd::ConvertToken( *ChildMaterialName ).Get() ) )
+				FString ChildMaterialName = TEXT("UnrealMaterial");
+				if (pxr::UsdPrim ExistingPrim = MeshPrim.GetChild(UnrealToUsd::ConvertToken(*ChildMaterialName).Get()))
 				{
 					// Get a unique name for a new prim. Don't even try checking if this prim is usable as the material binding,
 					// because if it was the material binding for this mesh we would have already used it above, when fetching the ExistingShader.
 					// If we're here, we don't know what this prim is about
 					TSet<FString> UsedNames;
-					for ( pxr::UsdPrim Child : MeshPrim.GetFilteredChildren( pxr::UsdTraverseInstanceProxies( pxr::UsdPrimAllPrimsPredicate ) ) )
+					for (pxr::UsdPrim Child : MeshPrim.GetFilteredChildren(pxr::UsdTraverseInstanceProxies(pxr::UsdPrimAllPrimsPredicate)))
 					{
-						UsedNames.Add( UsdToUnreal::ConvertToken( Child.GetName() ) );
+						UsedNames.Add(UsdToUnreal::ConvertToken(Child.GetName()));
 					}
 
-					ChildMaterialName = UsdUtils::GetUniqueName( ChildMaterialName, UsedNames );
+					ChildMaterialName = UsdUtils::GetUniqueName(ChildMaterialName, UsedNames);
 				}
 
 				pxr::UsdStageRefPtr Stage = MeshPrim.GetStage();
 				pxr::SdfPath MeshPath = MeshPrim.GetPath();
-				pxr::SdfPath MaterialPath = MeshPath.AppendChild( UnrealToUsd::ConvertToken( *ChildMaterialName ).Get() );
+				pxr::SdfPath MaterialPath = MeshPath.AppendChild(UnrealToUsd::ConvertToken(*ChildMaterialName).Get());
 
-				pxr::UsdShadeMaterial ChildMaterial = pxr::UsdShadeMaterial::Define( Stage, MaterialPath );
-				if ( !ChildMaterial )
+				pxr::UsdShadeMaterial ChildMaterial = pxr::UsdShadeMaterial::Define(Stage, MaterialPath);
+				if (!ChildMaterial)
 				{
-					UE_LOG( LogUsd, Warning, TEXT( "Failed to author material prim '%s' when trying to write '%s's material override '%s' to USD" ),
-						*UsdToUnreal::ConvertPath( MaterialPath ),
-						*UsdToUnreal::ConvertPath( MeshPrim.GetPath() ),
+					UE_LOG(
+						LogUsd,
+						Warning,
+						TEXT("Failed to author material prim '%s' when trying to write '%s's material override '%s' to USD"),
+						*UsdToUnreal::ConvertPath(MaterialPath),
+						*UsdToUnreal::ConvertPath(MeshPrim.GetPath()),
 						*UEMaterialAssetPath
 					);
 					return;
 				}
 
-				if ( pxr::UsdPrim MaterialPrim = ChildMaterial.GetPrim() )
+				if (pxr::UsdPrim MaterialPrim = ChildMaterial.GetPrim())
 				{
-					UsdUtils::SetUnrealSurfaceOutput( MaterialPrim, UEMaterialAssetPath );
+					UsdUtils::SetUnrealSurfaceOutput(MaterialPrim, UEMaterialAssetPath);
 
-					BindingAPI.Bind( ChildMaterial );
+					BindingAPI.Bind(ChildMaterial);
 				}
 			}
 
 			// On the current edit target, will set the Xformable's op order to a single "xformOp:transform",
 			// create the corresponding attribute, and return the op
-			pxr::UsdGeomXformOp ForceMatrixXform( pxr::UsdGeomXformable& Xformable )
+			pxr::UsdGeomXformOp ForceMatrixXform(pxr::UsdGeomXformable& Xformable)
 			{
 				FScopedUsdAllocs Allocs;
 
@@ -236,9 +238,12 @@ namespace UE
 				Xformable.ClearXformOpOrder();
 
 				// Find details about the transform attribute related to the default transform type xform op
-				pxr::TfToken TransformAttrName = pxr::UsdGeomXformOp::GetOpName( pxr::UsdGeomXformOp::TypeTransform );
-				const pxr::SdfValueTypeName& TransformAttrTypeName = pxr::UsdGeomXformOp::GetValueTypeName( pxr::UsdGeomXformOp::TypeTransform, pxr::UsdGeomXformOp::PrecisionDouble );
-				if ( TransformAttrName.IsEmpty() || !TransformAttrTypeName )
+				pxr::TfToken TransformAttrName = pxr::UsdGeomXformOp::GetOpName(pxr::UsdGeomXformOp::TypeTransform);
+				const pxr::SdfValueTypeName& TransformAttrTypeName = pxr::UsdGeomXformOp::GetValueTypeName(
+					pxr::UsdGeomXformOp::TypeTransform,
+					pxr::UsdGeomXformOp::PrecisionDouble
+				);
+				if (TransformAttrName.IsEmpty() || !TransformAttrTypeName)
 				{
 					return {};
 				}
@@ -246,33 +251,38 @@ namespace UE
 				// Create the transform attribute that would match the default transform type xform op
 				const bool bCustom = false;
 				pxr::UsdPrim UsdPrim = Xformable.GetPrim();
-				pxr::UsdAttribute TransformAttr = UsdPrim.CreateAttribute( TransformAttrName, TransformAttrTypeName, bCustom );
-				if ( !TransformAttr )
+				pxr::UsdAttribute TransformAttr = UsdPrim.CreateAttribute(TransformAttrName, TransformAttrTypeName, bCustom);
+				if (!TransformAttr)
 				{
 					return {};
 				}
 
 				// Now that the attribute is created, use it to create the corresponding pxr::UsdGeomXformOp
 				const bool bIsInverseOp = false;
-				pxr::UsdGeomXformOp NewOp{ TransformAttr, bIsInverseOp };
-				if ( !NewOp )
+				pxr::UsdGeomXformOp NewOp{TransformAttr, bIsInverseOp};
+				if (!NewOp)
 				{
 					return {};
 				}
 
 				// Store the Op name on an array that will be our new op order value
 				pxr::VtTokenArray NewOps;
-				NewOps.push_back( NewOp.GetOpName() );
-				Xformable.CreateXformOpOrderAttr().Set( NewOps );
+				NewOps.push_back(NewOp.GetOpName());
+				Xformable.CreateXformOpOrderAttr().Set(NewOps);
 
 				return NewOp;
 			}
 
 			// Turns OutTransform into the UE-space relative (local to parent) transform for Xformable, paying attention to if it
 			// or any of its ancestors has the '!resetXformStack!' xformOp.
-			void GetPrimConvertedRelativeTransform( pxr::UsdGeomXformable Xformable, double UsdTimeCode, FTransform& OutTransform, bool bIgnoreLocalTransform = false )
+			void GetPrimConvertedRelativeTransform(
+				pxr::UsdGeomXformable Xformable,
+				double UsdTimeCode,
+				FTransform& OutTransform,
+				bool bIgnoreLocalTransform = false
+			)
 			{
-				if ( !Xformable )
+				if (!Xformable)
 				{
 					return;
 				}
@@ -283,16 +293,16 @@ namespace UE
 				pxr::UsdStageRefPtr UsdStage = UsdPrim.GetStage();
 
 				bool bResetTransformStack = false;
-				if( bIgnoreLocalTransform )
+				if (bIgnoreLocalTransform)
 				{
 					FTransform Dummy;
-					UsdToUnreal::ConvertXformable( UsdStage, Xformable, Dummy, UsdTimeCode, &bResetTransformStack );
+					UsdToUnreal::ConvertXformable(UsdStage, Xformable, Dummy, UsdTimeCode, &bResetTransformStack);
 
 					OutTransform = FTransform::Identity;
 				}
 				else
 				{
-					UsdToUnreal::ConvertXformable( UsdStage, Xformable, OutTransform, UsdTimeCode, &bResetTransformStack );
+					UsdToUnreal::ConvertXformable(UsdStage, Xformable, OutTransform, UsdTimeCode, &bResetTransformStack);
 				}
 
 				// If we have the resetXformStack op on this prim's xformOpOrder we have to essentially use its transform
@@ -318,23 +328,29 @@ namespace UE
 				// resetXformStack through all USD workflows, and mean we'd have to *write out* resetXformStack when writing/exporting
 				// absolute transform components, and also convert between them when the user toggles between relative/absolute manually,
 				// which is probably worse than just baking it as relative transforms on first read and forgetting about it.
-				if ( bResetTransformStack )
+				if (bResetTransformStack)
 				{
 					FTransform ParentToWorld = FTransform::Identity;
 
 					pxr::UsdPrim AncestorPrim = UsdPrim.GetParent();
-					while ( AncestorPrim && !AncestorPrim.IsPseudoRoot() )
+					while (AncestorPrim && !AncestorPrim.IsPseudoRoot())
 					{
 						FTransform AncestorTransform = FTransform::Identity;
 						bool bAncestorResetTransformStack = false;
-						UsdToUnreal::ConvertXformable( UsdStage, pxr::UsdGeomXformable{ AncestorPrim }, AncestorTransform, UsdTimeCode, &bAncestorResetTransformStack );
+						UsdToUnreal::ConvertXformable(
+							UsdStage,
+							pxr::UsdGeomXformable{AncestorPrim},
+							AncestorTransform,
+							UsdTimeCode,
+							&bAncestorResetTransformStack
+						);
 
 						ParentToWorld = ParentToWorld * AncestorTransform;
 
 						// If we find a parent that also has the resetXformStack, then we're in luck: That transform value will be its world
 						// transform already, so we can stop concatenating stuff. Yes, on the component-side of things we'd have done the same
 						// thing of making a fake relative transform for it, but the end result would have been the same final world transform
-						if ( bAncestorResetTransformStack )
+						if (bAncestorResetTransformStack)
 						{
 							break;
 						}
@@ -343,28 +359,39 @@ namespace UE
 					}
 
 					const FVector& Scale = ParentToWorld.GetScale3D();
-					if ( !FMath::IsNearlyEqual( Scale.X, Scale.Y ) || !FMath::IsNearlyEqual( Scale.X, Scale.Z ) )
+					if (!FMath::IsNearlyEqual(Scale.X, Scale.Y) || !FMath::IsNearlyEqual(Scale.X, Scale.Z))
 					{
-						UE_LOG( LogUsd, Warning, TEXT( "Inverting transform with non-uniform scaling '%s' when computing relative transform for prim '%s'! Result will likely be incorrect, since FTransforms can't invert non-uniform scalings. You can work around this by baking your non-uniform scaling transform into the vertices, or by not using the !resetXformStack! Xform op." ),
+						UE_LOG(
+							LogUsd,
+							Warning,
+							TEXT("Inverting transform with non-uniform scaling '%s' when computing relative transform for prim '%s'! Result will "
+								 "likely be incorrect, since FTransforms can't invert non-uniform scalings. You can work around this by baking your "
+								 "non-uniform scaling transform into the vertices, or by not using the !resetXformStack! Xform op."),
 							*Scale.ToString(),
-							*UsdToUnreal::ConvertPath( UsdPrim.GetPrimPath() )
+							*UsdToUnreal::ConvertPath(UsdPrim.GetPrimPath())
 						);
 					}
 
 					// Multiplying with matrices here helps mitigate the issues encountered with non-uniform scaling, however it will stil
 					// never be perfect, as it is not possible to generate an FTransform that can properly invert a complex transform with non-uniform
 					// scaling when just multiplying them (which is what downstream code within USceneComponent will do).
-					OutTransform = FTransform{ OutTransform.ToMatrixWithScale() * ParentToWorld.ToInverseMatrixWithScale() };
+					OutTransform = FTransform{OutTransform.ToMatrixWithScale() * ParentToWorld.ToInverseMatrixWithScale()};
 				}
 			}
-		}
-	}
-}
-
-bool UsdToUnreal::ConvertXformable( const pxr::UsdStageRefPtr& Stage, const pxr::UsdTyped& Schema, FTransform& OutTransform, double EvalTime, bool* bOutResetTransformStack )
+		}	 // namespace Private
+	}		 // namespace USDPrimConversionImpl
+}	 // namespace UE
+
+bool UsdToUnreal::ConvertXformable(
+	const pxr::UsdStageRefPtr& Stage,
+	const pxr::UsdTyped& Schema,
+	FTransform& OutTransform,
+	double EvalTime,
+	bool* bOutResetTransformStack
+)
 {
-	pxr::UsdGeomXformable Xformable( Schema );
-	if ( !Xformable )
+	pxr::UsdGeomXformable Xformable(Schema);
+	if (!Xformable)
 	{
 		return false;
 	}
@@ -380,44 +407,44 @@ bool UsdToUnreal::ConvertXformable( const pxr::UsdStageRefPtr& Stage, const pxr:
 	pxr::GfMatrix4d UsdMatrix;
 	bool bResetXformStack = false;
 	bool* bResetXformStackPtr = bOutResetTransformStack ? bOutResetTransformStack : &bResetXformStack;
-	Xformable.GetLocalTransformation( &UsdMatrix, bResetXformStackPtr, EvalTime );
+	Xformable.GetLocalTransformation(&UsdMatrix, bResetXformStackPtr, EvalTime);
 
-	FUsdStageInfo StageInfo( Stage );
-	OutTransform = OutTransform * UsdToUnreal::ConvertMatrix( StageInfo, UsdMatrix );
+	FUsdStageInfo StageInfo(Stage);
+	OutTransform = OutTransform * UsdToUnreal::ConvertMatrix(StageInfo, UsdMatrix);
 
-	const bool bPrimIsLight = Xformable.GetPrim().HasAPI< pxr::UsdLuxLightAPI >();
+	const bool bPrimIsLight = Xformable.GetPrim().HasAPI<pxr::UsdLuxLightAPI>();
 
 	// Extra rotation to match different camera facing direction convention
 	// Note: The camera space is always Y-up, yes, but this is not what this is: This is the camera's transform wrt the stage,
 	// which follows the stage up axis
-	if ( Xformable.GetPrim().IsA< pxr::UsdGeomCamera >() || bPrimIsLight )
+	if (Xformable.GetPrim().IsA<pxr::UsdGeomCamera>() || bPrimIsLight)
 	{
-		if ( StageInfo.UpAxis == EUsdUpAxis::YAxis )
+		if (StageInfo.UpAxis == EUsdUpAxis::YAxis)
 		{
-			OutTransform = FTransform( FRotator( 0.0f, -90.f, 0.0f ) ) * OutTransform;
+			OutTransform = FTransform(FRotator(0.0f, -90.f, 0.0f)) * OutTransform;
 		}
 		else
 		{
-			OutTransform = FTransform( FRotator( -90.0f, -90.f, 0.0f ) ) * OutTransform;
+			OutTransform = FTransform(FRotator(-90.0f, -90.f, 0.0f)) * OutTransform;
 		}
 	}
 	// Invert the compensation applied to our parents, in case they're a camera or a light
-	if ( pxr::UsdPrim Parent = Xformable.GetPrim().GetParent() )
+	if (pxr::UsdPrim Parent = Xformable.GetPrim().GetParent())
 	{
-		const bool bParentIsLight = Parent.HasAPI< pxr::UsdLuxLightAPI >();
+		const bool bParentIsLight = Parent.HasAPI<pxr::UsdLuxLightAPI>();
 
 		// If bResetXFormStack is true, then the prim's local transform will be used directly as the world transform, and we will
 		// already invert the parent transform fully, regardless of what it is. This means it doesn't really matter if our parent
 		// has a camera/light compensation or not, and so we don't have to have the explicit inverse compensation here anyway!
-		if ( !(*bResetXformStackPtr) && ( Parent.IsA< pxr::UsdGeomCamera >() || bParentIsLight ) )
+		if (!(*bResetXformStackPtr) && (Parent.IsA<pxr::UsdGeomCamera>() || bParentIsLight))
 		{
-			if ( StageInfo.UpAxis == EUsdUpAxis::YAxis )
+			if (StageInfo.UpAxis == EUsdUpAxis::YAxis)
 			{
-				OutTransform = OutTransform * FTransform( FRotator( 0.0f, -90.f, 0.0f ).GetInverse() );
+				OutTransform = OutTransform * FTransform(FRotator(0.0f, -90.f, 0.0f).GetInverse());
 			}
 			else
 			{
-				OutTransform = OutTransform * FTransform( FRotator( -90.0f, -90.f, 0.0f ).GetInverse() );
+				OutTransform = OutTransform * FTransform(FRotator(-90.0f, -90.f, 0.0f).GetInverse());
 			}
 		}
 	}
@@ -425,7 +452,13 @@ bool UsdToUnreal::ConvertXformable( const pxr::UsdStageRefPtr& Stage, const pxr:
 	return true;
 }
 
-void UsdToUnreal::PropagateTransform(const pxr::UsdStageRefPtr& Stage, const pxr::UsdPrim& Root, const pxr::UsdPrim& Leaf, double EvalTime, FTransform& OutTransform)
+void UsdToUnreal::PropagateTransform(
+	const pxr::UsdStageRefPtr& Stage,
+	const pxr::UsdPrim& Root,
+	const pxr::UsdPrim& Leaf,
+	double EvalTime,
+	FTransform& OutTransform
+)
 {
 	FScopedUsdAllocs UsdAllocs;
 
@@ -463,93 +496,117 @@ void UsdToUnreal::PropagateTransform(const pxr::UsdStageRefPtr& Stage, const pxr
 	}
 }
 
-bool UsdToUnreal::ConvertXformable( const pxr::UsdStageRefPtr& Stage, const pxr::UsdTyped& Schema, USceneComponent& SceneComponent, double EvalTime, bool bUsePrimTransform )
+bool UsdToUnreal::ConvertXformable(
+	const pxr::UsdStageRefPtr& Stage,
+	const pxr::UsdTyped& Schema,
+	USceneComponent& SceneComponent,
+	double EvalTime,
+	bool bUsePrimTransform
+)
 {
-	pxr::UsdGeomXformable Xformable( Schema );
-	if ( !Xformable )
+	pxr::UsdGeomXformable Xformable(Schema);
+	if (!Xformable)
 	{
 		return false;
 	}
 
-	TRACE_CPUPROFILER_EVENT_SCOPE( UsdToUnreal::ConvertXformable );
+	TRACE_CPUPROFILER_EVENT_SCOPE(UsdToUnreal::ConvertXformable);
 
 	FScopedUsdAllocs UsdAllocs;
 
 	// Transform
 	FTransform Transform;
-	UE::USDPrimConversionImpl::Private::GetPrimConvertedRelativeTransform( Xformable, EvalTime, Transform, !bUsePrimTransform );
-	SceneComponent.SetRelativeTransform( Transform );
+	UE::USDPrimConversionImpl::Private::GetPrimConvertedRelativeTransform(Xformable, EvalTime, Transform, !bUsePrimTransform);
+	SceneComponent.SetRelativeTransform(Transform);
 
 	SceneComponent.Modify();
 
 	// Computed (effective) visibility
-	const bool bIsHidden = ( Xformable.ComputeVisibility( EvalTime ) == pxr::UsdGeomTokens->invisible );
-	SceneComponent.SetHiddenInGame( bIsHidden );
+	const bool bIsHidden = (Xformable.ComputeVisibility(EvalTime) == pxr::UsdGeomTokens->invisible);
+	SceneComponent.SetHiddenInGame(bIsHidden);
 
 	// Per-prim visibility
-	bool bIsInvisible = false; // Default to 'inherited'
-	if ( pxr::UsdAttribute VisibilityAttr = Xformable.GetVisibilityAttr() )
+	bool bIsInvisible = false;	  // Default to 'inherited'
+	if (pxr::UsdAttribute VisibilityAttr = Xformable.GetVisibilityAttr())
 	{
 		pxr::TfToken Value;
-		if ( VisibilityAttr.Get( &Value, EvalTime ) )
+		if (VisibilityAttr.Get(&Value, EvalTime))
 		{
 			bIsInvisible = Value == pxr::UsdGeomTokens->invisible;
 		}
 	}
-	if ( bIsInvisible )
+	if (bIsInvisible)
 	{
-		SceneComponent.ComponentTags.AddUnique( UnrealIdentifiers::Invisible );
-		SceneComponent.ComponentTags.Remove( UnrealIdentifiers::Inherited );
+		SceneComponent.ComponentTags.AddUnique(UnrealIdentifiers::Invisible);
+		SceneComponent.ComponentTags.Remove(UnrealIdentifiers::Inherited);
 	}
 	else
 	{
-		SceneComponent.ComponentTags.Remove( UnrealIdentifiers::Invisible );
-		SceneComponent.ComponentTags.AddUnique( UnrealIdentifiers::Inherited );
+		SceneComponent.ComponentTags.Remove(UnrealIdentifiers::Invisible);
+		SceneComponent.ComponentTags.AddUnique(UnrealIdentifiers::Inherited);
 	}
 
 	return true;
 }
 
-bool UsdToUnreal::ConvertGeomCamera( const UE::FUsdPrim& Prim, UCineCameraComponent& CameraComponent, double UsdTimeCode )
+bool UsdToUnreal::ConvertGeomCamera(const UE::FUsdPrim& Prim, UCineCameraComponent& CameraComponent, double UsdTimeCode)
 {
 	FScopedUsdAllocs Allocs;
 
-	pxr::UsdPrim UsdPrim{ Prim };
-	pxr::UsdGeomCamera GeomCamera{ UsdPrim };
-	if ( !GeomCamera )
+	pxr::UsdPrim UsdPrim{Prim};
+	pxr::UsdGeomCamera GeomCamera{UsdPrim};
+	if (!GeomCamera)
 	{
 		return false;
 	}
 
 	UE::FUsdStage Stage = Prim.GetStage();
-	FUsdStageInfo StageInfo( Stage );
+	FUsdStageInfo StageInfo(Stage);
 
-	CameraComponent.CurrentFocalLength = UsdToUnreal::ConvertDistance( StageInfo, UsdUtils::GetUsdValue< float >( GeomCamera.GetFocalLengthAttr(), UsdTimeCode ) );
+	CameraComponent.CurrentFocalLength = UsdToUnreal::ConvertDistance(
+		StageInfo,
+		UsdUtils::GetUsdValue<float>(GeomCamera.GetFocalLengthAttr(), UsdTimeCode)
+	);
 
-	CameraComponent.FocusSettings.ManualFocusDistance = UsdToUnreal::ConvertDistance( StageInfo, UsdUtils::GetUsdValue< float >( GeomCamera.GetFocusDistanceAttr(), UsdTimeCode ) );
+	CameraComponent.FocusSettings.ManualFocusDistance = UsdToUnreal::ConvertDistance(
+		StageInfo,
+		UsdUtils::GetUsdValue<float>(GeomCamera.GetFocusDistanceAttr(), UsdTimeCode)
+	);
 
-	if ( FMath::IsNearlyZero( CameraComponent.FocusSettings.ManualFocusDistance ) )
+	if (FMath::IsNearlyZero(CameraComponent.FocusSettings.ManualFocusDistance))
 	{
 		CameraComponent.FocusSettings.FocusMethod = ECameraFocusMethod::DoNotOverride;
 	}
 
-	CameraComponent.CurrentAperture = UsdUtils::GetUsdValue< float >( GeomCamera.GetFStopAttr(), UsdTimeCode );
+	CameraComponent.CurrentAperture = UsdUtils::GetUsdValue<float>(GeomCamera.GetFStopAttr(), UsdTimeCode);
 
-	CameraComponent.Filmback.SensorWidth = UsdToUnreal::ConvertDistance( StageInfo, UsdUtils::GetUsdValue< float >( GeomCamera.GetHorizontalApertureAttr(), UsdTimeCode ) );
-	CameraComponent.Filmback.SensorHeight = UsdToUnreal::ConvertDistance( StageInfo, UsdUtils::GetUsdValue< float >( GeomCamera.GetVerticalApertureAttr(), UsdTimeCode ) );
+	CameraComponent.Filmback.SensorWidth = UsdToUnreal::ConvertDistance(
+		StageInfo,
+		UsdUtils::GetUsdValue<float>(GeomCamera.GetHorizontalApertureAttr(), UsdTimeCode)
+	);
+	CameraComponent.Filmback.SensorHeight = UsdToUnreal::ConvertDistance(
+		StageInfo,
+		UsdUtils::GetUsdValue<float>(GeomCamera.GetVerticalApertureAttr(), UsdTimeCode)
+	);
 
 	return true;
 }
 
-bool UsdToUnreal::ConvertBoolTimeSamples( const UE::FUsdStage& Stage, const TArray<double>& UsdTimeSamples, const TFunction<bool( double )>& ReaderFunc, UMovieSceneBoolTrack& MovieSceneTrack, const FMovieSceneSequenceTransform& SequenceTransform )
+bool UsdToUnreal::ConvertBoolTimeSamples(
+	const UE::FUsdStage& Stage,
+	const TArray<double>& UsdTimeSamples,
+	const TFunction<bool(double)>& ReaderFunc,
+	UMovieSceneBoolTrack& MovieSceneTrack,
+	const FMovieSceneSequenceTransform& SequenceTransform
+)
 {
-	if ( !ReaderFunc )
+	if (!ReaderFunc)
 	{
 		return false;
 	}
 
-	const UMovieScene* MovieScene = MovieSceneTrack.GetTypedOuter< UMovieScene >();
-	if ( !MovieScene )
+	const UMovieScene* MovieScene = MovieSceneTrack.GetTypedOuter<UMovieScene>();
+	if (!MovieScene)
 	{
 		return false;
 	}
@@ -559,65 +616,71 @@ bool UsdToUnreal::ConvertBoolTimeSamples( const UE::FUsdStage& Stage, const TArr
 
 	FScopedUsdAllocs Allocs;
 
-	pxr::UsdStageRefPtr UsdStage{ Stage };
-	FUsdStageInfo StageInfo{ Stage };
+	pxr::UsdStageRefPtr UsdStage{Stage};
+	FUsdStageInfo StageInfo{Stage};
 
-	TArray< FFrameNumber > FrameNumbers;
-	FrameNumbers.Reserve( UsdTimeSamples.Num() );
+	TArray<FFrameNumber> FrameNumbers;
+	FrameNumbers.Reserve(UsdTimeSamples.Num());
 
-	TArray< bool > SectionValues;
-	SectionValues.Reserve( UsdTimeSamples.Num() );
+	TArray<bool> SectionValues;
+	SectionValues.Reserve(UsdTimeSamples.Num());
 
 	const double StageTimeCodesPerSecond = UsdStage->GetTimeCodesPerSecond();
-	const FFrameRate StageFrameRate( StageTimeCodesPerSecond, 1 );
+	const FFrameRate StageFrameRate(StageTimeCodesPerSecond, 1);
 
 	double LastTimeSample = TNumericLimits<double>::Lowest();
-	for ( const double UsdTimeSample : UsdTimeSamples )
+	for (const double UsdTimeSample : UsdTimeSamples)
 	{
 		// We never want to evaluate the same time twice
-		if ( FMath::IsNearlyEqual( UsdTimeSample, LastTimeSample ) )
+		if (FMath::IsNearlyEqual(UsdTimeSample, LastTimeSample))
 		{
 			continue;
 		}
 		LastTimeSample = UsdTimeSample;
 
-		int32 FrameNumber = FMath::FloorToInt( UsdTimeSample );
+		int32 FrameNumber = FMath::FloorToInt(UsdTimeSample);
 		float SubFrameNumber = UsdTimeSample - FrameNumber;
 
-		FFrameTime FrameTime( FrameNumber, SubFrameNumber );
+		FFrameTime FrameTime(FrameNumber, SubFrameNumber);
 
-		FFrameTime KeyFrameTime = FFrameRate::TransformTime( FrameTime, StageFrameRate, Resolution );
+		FFrameTime KeyFrameTime = FFrameRate::TransformTime(FrameTime, StageFrameRate, Resolution);
 		KeyFrameTime *= SequenceTransform;
-		FrameNumbers.Add( KeyFrameTime.GetFrame() );
+		FrameNumbers.Add(KeyFrameTime.GetFrame());
 
-		bool UEValue = ReaderFunc( UsdTimeSample );
-		SectionValues.Emplace_GetRef( UEValue );
+		bool UEValue = ReaderFunc(UsdTimeSample);
+		SectionValues.Emplace_GetRef(UEValue);
 	}
 
 	bool bSectionAdded = false;
-	UMovieSceneBoolSection* Section = Cast< UMovieSceneBoolSection >( MovieSceneTrack.FindOrAddSection( 0, bSectionAdded ) );
+	UMovieSceneBoolSection* Section = Cast<UMovieSceneBoolSection>(MovieSceneTrack.FindOrAddSection(0, bSectionAdded));
 	Section->EvalOptions.CompletionMode = EMovieSceneCompletionMode::KeepState;
 
 	TMovieSceneChannelData<bool> Data = Section->GetChannel().GetData();
 	Data.Reset();
-	for ( int32 KeyIndex = 0; KeyIndex < FrameNumbers.Num(); ++KeyIndex )
+	for (int32 KeyIndex = 0; KeyIndex < FrameNumbers.Num(); ++KeyIndex)
 	{
-		Data.AddKey( FrameNumbers[ KeyIndex ], SectionValues[ KeyIndex ] );
+		Data.AddKey(FrameNumbers[KeyIndex], SectionValues[KeyIndex]);
 	}
 
-	Section->SetRange( Section->GetAutoSizeRange().Get( TRange<FFrameNumber>::Empty() ) );
+	Section->SetRange(Section->GetAutoSizeRange().Get(TRange<FFrameNumber>::Empty()));
 
 	return true;
 }
 
-bool UsdToUnreal::ConvertBoolTimeSamples(const UE::FUsdStage& Stage, const TArray<double>& UsdTimeSamples, const TFunction<bool(double)>& ReaderFunc, UMovieSceneVisibilityTrack& MovieSceneTrack, const FMovieSceneSequenceTransform& SequenceTransform)
+bool UsdToUnreal::ConvertBoolTimeSamples(
+	const UE::FUsdStage& Stage,
+	const TArray<double>& UsdTimeSamples,
+	const TFunction<bool(double)>& ReaderFunc,
+	UMovieSceneVisibilityTrack& MovieSceneTrack,
+	const FMovieSceneSequenceTransform& SequenceTransform
+)
 {
 	if (!ReaderFunc)
 	{
 		return false;
 	}
 
-	const UMovieScene* MovieScene = MovieSceneTrack.GetTypedOuter< UMovieScene >();
+	const UMovieScene* MovieScene = MovieSceneTrack.GetTypedOuter<UMovieScene>();
 	if (!MovieScene)
 	{
 		return false;
@@ -628,13 +691,13 @@ bool UsdToUnreal::ConvertBoolTimeSamples(const UE::FUsdStage& Stage, const TArra
 
 	FScopedUsdAllocs Allocs;
 
-	pxr::UsdStageRefPtr UsdStage{ Stage };
-	FUsdStageInfo StageInfo{ Stage };
+	pxr::UsdStageRefPtr UsdStage{Stage};
+	FUsdStageInfo StageInfo{Stage};
 
-	TArray< FFrameNumber > FrameNumbers;
+	TArray<FFrameNumber> FrameNumbers;
 	FrameNumbers.Reserve(UsdTimeSamples.Num());
 
-	TArray< bool > SectionValues;
+	TArray<bool> SectionValues;
 	SectionValues.Reserve(UsdTimeSamples.Num());
 
 	const double StageTimeCodesPerSecond = UsdStage->GetTimeCodesPerSecond();
@@ -664,7 +727,7 @@ bool UsdToUnreal::ConvertBoolTimeSamples(const UE::FUsdStage& Stage, const TArra
 	}
 
 	bool bSectionAdded = false;
-	UMovieSceneVisibilitySection* Section = Cast< UMovieSceneVisibilitySection >(MovieSceneTrack.FindOrAddSection(0, bSectionAdded));
+	UMovieSceneVisibilitySection* Section = Cast<UMovieSceneVisibilitySection>(MovieSceneTrack.FindOrAddSection(0, bSectionAdded));
 	Section->EvalOptions.CompletionMode = EMovieSceneCompletionMode::KeepState;
 
 	TMovieSceneChannelData<bool> Data = Section->GetChannel().GetData();
@@ -679,15 +742,21 @@ bool UsdToUnreal::ConvertBoolTimeSamples(const UE::FUsdStage& Stage, const TArra
 	return true;
 }
 
-bool UsdToUnreal::ConvertFloatTimeSamples( const UE::FUsdStage& Stage, const TArray<double>& UsdTimeSamples, const TFunction<float( double )>& ReaderFunc, UMovieSceneFloatTrack& MovieSceneTrack, const FMovieSceneSequenceTransform& SequenceTransform )
+bool UsdToUnreal::ConvertFloatTimeSamples(
+	const UE::FUsdStage& Stage,
+	const TArray<double>& UsdTimeSamples,
+	const TFunction<float(double)>& ReaderFunc,
+	UMovieSceneFloatTrack& MovieSceneTrack,
+	const FMovieSceneSequenceTransform& SequenceTransform
+)
 {
-	if ( !ReaderFunc )
+	if (!ReaderFunc)
 	{
 		return false;
 	}
 
-	const UMovieScene* MovieScene = MovieSceneTrack.GetTypedOuter< UMovieScene >();
-	if ( !MovieScene )
+	const UMovieScene* MovieScene = MovieSceneTrack.GetTypedOuter<UMovieScene>();
+	if (!MovieScene)
 	{
 		return false;
 	}
@@ -697,67 +766,75 @@ bool UsdToUnreal::ConvertFloatTimeSamples( const UE::FUsdStage& Stage, const TAr
 
 	FScopedUsdAllocs Allocs;
 
-	pxr::UsdStageRefPtr UsdStage{ Stage };
-	FUsdStageInfo StageInfo{ Stage };
+	pxr::UsdStageRefPtr UsdStage{Stage};
+	FUsdStageInfo StageInfo{Stage};
 
-	TArray< FFrameNumber > FrameNumbers;
-	FrameNumbers.Reserve( UsdTimeSamples.Num() );
+	TArray<FFrameNumber> FrameNumbers;
+	FrameNumbers.Reserve(UsdTimeSamples.Num());
 
-	TArray< FMovieSceneFloatValue > SectionValues;
-	SectionValues.Reserve( UsdTimeSamples.Num() );
+	TArray<FMovieSceneFloatValue> SectionValues;
+	SectionValues.Reserve(UsdTimeSamples.Num());
 
 	const double StageTimeCodesPerSecond = UsdStage->GetTimeCodesPerSecond();
-	const FFrameRate StageFrameRate( StageTimeCodesPerSecond, 1 );
+	const FFrameRate StageFrameRate(StageTimeCodesPerSecond, 1);
 
-	const ERichCurveInterpMode InterpMode = ( UsdStage->GetInterpolationType() == pxr::UsdInterpolationTypeLinear ) ? ERichCurveInterpMode::RCIM_Linear : ERichCurveInterpMode::RCIM_Constant;
+	const ERichCurveInterpMode InterpMode = (UsdStage->GetInterpolationType() == pxr::UsdInterpolationTypeLinear)
+												? ERichCurveInterpMode::RCIM_Linear
+												: ERichCurveInterpMode::RCIM_Constant;
 
 	double LastTimeSample = TNumericLimits<double>::Lowest();
-	for ( const double UsdTimeSample : UsdTimeSamples )
+	for (const double UsdTimeSample : UsdTimeSamples)
 	{
 		// We never want to evaluate the same time twice
-		if ( FMath::IsNearlyEqual( UsdTimeSample, LastTimeSample ) )
+		if (FMath::IsNearlyEqual(UsdTimeSample, LastTimeSample))
 		{
 			continue;
 		}
 		LastTimeSample = UsdTimeSample;
 
-		int32 FrameNumber = FMath::FloorToInt( UsdTimeSample );
+		int32 FrameNumber = FMath::FloorToInt(UsdTimeSample);
 		float SubFrameNumber = UsdTimeSample - FrameNumber;
 
-		FFrameTime FrameTime( FrameNumber, SubFrameNumber );
+		FFrameTime FrameTime(FrameNumber, SubFrameNumber);
 
-		FFrameTime KeyFrameTime = FFrameRate::TransformTime( FrameTime, StageFrameRate, Resolution );
+		FFrameTime KeyFrameTime = FFrameRate::TransformTime(FrameTime, StageFrameRate, Resolution);
 		KeyFrameTime *= SequenceTransform;
-		FrameNumbers.Add( KeyFrameTime.GetFrame() );
+		FrameNumbers.Add(KeyFrameTime.GetFrame());
 
-		float UEValue = ReaderFunc( UsdTimeSample );
-		SectionValues.Emplace_GetRef( UEValue ).InterpMode = InterpMode;
+		float UEValue = ReaderFunc(UsdTimeSample);
+		SectionValues.Emplace_GetRef(UEValue).InterpMode = InterpMode;
 	}
 
 	bool bSectionAdded = false;
-	UMovieSceneFloatSection* Section = Cast< UMovieSceneFloatSection >( MovieSceneTrack.FindOrAddSection( 0, bSectionAdded ) );
+	UMovieSceneFloatSection* Section = Cast<UMovieSceneFloatSection>(MovieSceneTrack.FindOrAddSection(0, bSectionAdded));
 	Section->EvalOptions.CompletionMode = EMovieSceneCompletionMode::KeepState;
 
-	TArrayView< FMovieSceneFloatChannel* > Channels = Section->GetChannelProxy().GetChannels< FMovieSceneFloatChannel >();
-	if ( Channels.Num() > 0 )
+	TArrayView<FMovieSceneFloatChannel*> Channels = Section->GetChannelProxy().GetChannels<FMovieSceneFloatChannel>();
+	if (Channels.Num() > 0)
 	{
-		Channels[ 0 ]->Set( FrameNumbers, SectionValues );
+		Channels[0]->Set(FrameNumbers, SectionValues);
 	}
 
-	Section->SetRange( Section->GetAutoSizeRange().Get( TRange<FFrameNumber>::Empty() ) );
+	Section->SetRange(Section->GetAutoSizeRange().Get(TRange<FFrameNumber>::Empty()));
 
 	return true;
 }
 
-bool UsdToUnreal::ConvertColorTimeSamples( const UE::FUsdStage& Stage, const TArray<double>& UsdTimeSamples, const TFunction<FLinearColor( double )>& ReaderFunc, UMovieSceneColorTrack& MovieSceneTrack, const FMovieSceneSequenceTransform& SequenceTransform )
+bool UsdToUnreal::ConvertColorTimeSamples(
+	const UE::FUsdStage& Stage,
+	const TArray<double>& UsdTimeSamples,
+	const TFunction<FLinearColor(double)>& ReaderFunc,
+	UMovieSceneColorTrack& MovieSceneTrack,
+	const FMovieSceneSequenceTransform& SequenceTransform
+)
 {
-	if ( !ReaderFunc )
+	if (!ReaderFunc)
 	{
 		return false;
 	}
 
-	const UMovieScene* MovieScene = MovieSceneTrack.GetTypedOuter< UMovieScene >();
-	if ( !MovieScene )
+	const UMovieScene* MovieScene = MovieSceneTrack.GetTypedOuter<UMovieScene>();
+	if (!MovieScene)
 	{
 		return false;
 	}
@@ -767,81 +844,89 @@ bool UsdToUnreal::ConvertColorTimeSamples( const UE::FUsdStage& Stage, const TAr
 
 	FScopedUsdAllocs Allocs;
 
-	pxr::UsdStageRefPtr UsdStage{ Stage };
-	FUsdStageInfo StageInfo{ Stage };
+	pxr::UsdStageRefPtr UsdStage{Stage};
+	FUsdStageInfo StageInfo{Stage};
 
-	TArray< FFrameNumber > FrameNumbers;
-	FrameNumbers.Reserve( UsdTimeSamples.Num() );
+	TArray<FFrameNumber> FrameNumbers;
+	FrameNumbers.Reserve(UsdTimeSamples.Num());
 
-	TArray< FMovieSceneFloatValue > RedValues;
-	TArray< FMovieSceneFloatValue > GreenValues;
-	TArray< FMovieSceneFloatValue > BlueValues;
-	TArray< FMovieSceneFloatValue > AlphaValues;
-	RedValues.Reserve( UsdTimeSamples.Num() );
-	GreenValues.Reserve( UsdTimeSamples.Num() );
-	BlueValues.Reserve( UsdTimeSamples.Num() );
-	AlphaValues.Reserve( UsdTimeSamples.Num() );
+	TArray<FMovieSceneFloatValue> RedValues;
+	TArray<FMovieSceneFloatValue> GreenValues;
+	TArray<FMovieSceneFloatValue> BlueValues;
+	TArray<FMovieSceneFloatValue> AlphaValues;
+	RedValues.Reserve(UsdTimeSamples.Num());
+	GreenValues.Reserve(UsdTimeSamples.Num());
+	BlueValues.Reserve(UsdTimeSamples.Num());
+	AlphaValues.Reserve(UsdTimeSamples.Num());
 
 	const double StageTimeCodesPerSecond = UsdStage->GetTimeCodesPerSecond();
-	const FFrameRate StageFrameRate( StageTimeCodesPerSecond, 1 );
+	const FFrameRate StageFrameRate(StageTimeCodesPerSecond, 1);
 
-	const ERichCurveInterpMode InterpMode = ( UsdStage->GetInterpolationType() == pxr::UsdInterpolationTypeLinear ) ? ERichCurveInterpMode::RCIM_Linear : ERichCurveInterpMode::RCIM_Constant;
+	const ERichCurveInterpMode InterpMode = (UsdStage->GetInterpolationType() == pxr::UsdInterpolationTypeLinear)
+												? ERichCurveInterpMode::RCIM_Linear
+												: ERichCurveInterpMode::RCIM_Constant;
 
 	double LastTimeSample = TNumericLimits<double>::Lowest();
-	for ( const double UsdTimeSample : UsdTimeSamples )
+	for (const double UsdTimeSample : UsdTimeSamples)
 	{
 		// We never want to evaluate the same time twice
-		if ( FMath::IsNearlyEqual( UsdTimeSample, LastTimeSample ) )
+		if (FMath::IsNearlyEqual(UsdTimeSample, LastTimeSample))
 		{
 			continue;
 		}
 		LastTimeSample = UsdTimeSample;
 
-		int32 FrameNumber = FMath::FloorToInt( UsdTimeSample );
+		int32 FrameNumber = FMath::FloorToInt(UsdTimeSample);
 		float SubFrameNumber = UsdTimeSample - FrameNumber;
 
-		FFrameTime FrameTime( FrameNumber, SubFrameNumber );
+		FFrameTime FrameTime(FrameNumber, SubFrameNumber);
 
-		FFrameTime KeyFrameTime = FFrameRate::TransformTime( FrameTime, StageFrameRate, Resolution );
+		FFrameTime KeyFrameTime = FFrameRate::TransformTime(FrameTime, StageFrameRate, Resolution);
 		KeyFrameTime *= SequenceTransform;
-		FrameNumbers.Add( KeyFrameTime.GetFrame() );
+		FrameNumbers.Add(KeyFrameTime.GetFrame());
 
-		FLinearColor UEValue = ReaderFunc( UsdTimeSample );
-		RedValues.Emplace_GetRef( UEValue.R ).InterpMode = InterpMode;
-		GreenValues.Emplace_GetRef( UEValue.G ).InterpMode = InterpMode;
-		BlueValues.Emplace_GetRef( UEValue.B ).InterpMode = InterpMode;
-		AlphaValues.Emplace_GetRef( UEValue.A ).InterpMode = InterpMode;
+		FLinearColor UEValue = ReaderFunc(UsdTimeSample);
+		RedValues.Emplace_GetRef(UEValue.R).InterpMode = InterpMode;
+		GreenValues.Emplace_GetRef(UEValue.G).InterpMode = InterpMode;
+		BlueValues.Emplace_GetRef(UEValue.B).InterpMode = InterpMode;
+		AlphaValues.Emplace_GetRef(UEValue.A).InterpMode = InterpMode;
 	}
 
 	bool bSectionAdded = false;
-	UMovieSceneColorSection* Section = Cast< UMovieSceneColorSection >( MovieSceneTrack.FindOrAddSection( 0, bSectionAdded ) );
+	UMovieSceneColorSection* Section = Cast<UMovieSceneColorSection>(MovieSceneTrack.FindOrAddSection(0, bSectionAdded));
 	Section->EvalOptions.CompletionMode = EMovieSceneCompletionMode::KeepState;
 
-	TArrayView< FMovieSceneFloatChannel* > Channels = Section->GetChannelProxy().GetChannels< FMovieSceneFloatChannel >();
-	if ( Channels.Num() != 4 )
+	TArrayView<FMovieSceneFloatChannel*> Channels = Section->GetChannelProxy().GetChannels<FMovieSceneFloatChannel>();
+	if (Channels.Num() != 4)
 	{
 		return false;
 	}
 
-	Channels[ 0 ]->Set( FrameNumbers, RedValues );
-	Channels[ 1 ]->Set( FrameNumbers, GreenValues );
-	Channels[ 2 ]->Set( FrameNumbers, BlueValues );
-	Channels[ 3 ]->Set( FrameNumbers, AlphaValues );
+	Channels[0]->Set(FrameNumbers, RedValues);
+	Channels[1]->Set(FrameNumbers, GreenValues);
+	Channels[2]->Set(FrameNumbers, BlueValues);
+	Channels[3]->Set(FrameNumbers, AlphaValues);
 
-	Section->SetRange( Section->GetAutoSizeRange().Get( TRange<FFrameNumber>::Empty() ) );
+	Section->SetRange(Section->GetAutoSizeRange().Get(TRange<FFrameNumber>::Empty()));
 
 	return true;
 }
 
-bool UsdToUnreal::ConvertTransformTimeSamples( const UE::FUsdStage& Stage, const TArray<double>& UsdTimeSamples, const TFunction<FTransform( double )>& ReaderFunc, UMovieScene3DTransformTrack& MovieSceneTrack, const FMovieSceneSequenceTransform& SequenceTransform )
+bool UsdToUnreal::ConvertTransformTimeSamples(
+	const UE::FUsdStage& Stage,
+	const TArray<double>& UsdTimeSamples,
+	const TFunction<FTransform(double)>& ReaderFunc,
+	UMovieScene3DTransformTrack& MovieSceneTrack,
+	const FMovieSceneSequenceTransform& SequenceTransform
+)
 {
-	if ( !ReaderFunc )
+	if (!ReaderFunc)
 	{
 		return false;
 	}
 
-	const UMovieScene* MovieScene = MovieSceneTrack.GetTypedOuter< UMovieScene >();
-	if ( !MovieScene )
+	const UMovieScene* MovieScene = MovieSceneTrack.GetTypedOuter<UMovieScene>();
+	if (!MovieScene)
 	{
 		return false;
 	}
@@ -851,100 +936,102 @@ bool UsdToUnreal::ConvertTransformTimeSamples( const UE::FUsdStage& Stage, const
 
 	FScopedUsdAllocs Allocs;
 
-	pxr::UsdStageRefPtr UsdStage{ Stage };
-	FUsdStageInfo StageInfo{ Stage };
+	pxr::UsdStageRefPtr UsdStage{Stage};
+	FUsdStageInfo StageInfo{Stage};
 
-	TArray< FFrameNumber > FrameNumbers;
-	FrameNumbers.Reserve( UsdTimeSamples.Num() );
+	TArray<FFrameNumber> FrameNumbers;
+	FrameNumbers.Reserve(UsdTimeSamples.Num());
 
-	TArray< FMovieSceneDoubleValue > LocationXValues;
-	TArray< FMovieSceneDoubleValue > LocationYValues;
-	TArray< FMovieSceneDoubleValue > LocationZValues;
+	TArray<FMovieSceneDoubleValue> LocationXValues;
+	TArray<FMovieSceneDoubleValue> LocationYValues;
+	TArray<FMovieSceneDoubleValue> LocationZValues;
 
-	TArray< FMovieSceneDoubleValue > RotationXValues;
-	TArray< FMovieSceneDoubleValue > RotationYValues;
-	TArray< FMovieSceneDoubleValue > RotationZValues;
+	TArray<FMovieSceneDoubleValue> RotationXValues;
+	TArray<FMovieSceneDoubleValue> RotationYValues;
+	TArray<FMovieSceneDoubleValue> RotationZValues;
 
-	TArray< FMovieSceneDoubleValue > ScaleXValues;
-	TArray< FMovieSceneDoubleValue > ScaleYValues;
-	TArray< FMovieSceneDoubleValue > ScaleZValues;
+	TArray<FMovieSceneDoubleValue> ScaleXValues;
+	TArray<FMovieSceneDoubleValue> ScaleYValues;
+	TArray<FMovieSceneDoubleValue> ScaleZValues;
 
-	LocationXValues.Reserve( UsdTimeSamples.Num() );
-	LocationYValues.Reserve( UsdTimeSamples.Num() );
-	LocationZValues.Reserve( UsdTimeSamples.Num() );
+	LocationXValues.Reserve(UsdTimeSamples.Num());
+	LocationYValues.Reserve(UsdTimeSamples.Num());
+	LocationZValues.Reserve(UsdTimeSamples.Num());
 
-	RotationXValues.Reserve( UsdTimeSamples.Num() );
-	RotationYValues.Reserve( UsdTimeSamples.Num() );
-	RotationZValues.Reserve( UsdTimeSamples.Num() );
+	RotationXValues.Reserve(UsdTimeSamples.Num());
+	RotationYValues.Reserve(UsdTimeSamples.Num());
+	RotationZValues.Reserve(UsdTimeSamples.Num());
 
-	ScaleXValues.Reserve( UsdTimeSamples.Num() );
-	ScaleYValues.Reserve( UsdTimeSamples.Num() );
-	ScaleZValues.Reserve( UsdTimeSamples.Num() );
+	ScaleXValues.Reserve(UsdTimeSamples.Num());
+	ScaleYValues.Reserve(UsdTimeSamples.Num());
+	ScaleZValues.Reserve(UsdTimeSamples.Num());
 
 	const double StageTimeCodesPerSecond = UsdStage->GetTimeCodesPerSecond();
-	const FFrameRate StageFrameRate( StageTimeCodesPerSecond, 1 );
+	const FFrameRate StageFrameRate(StageTimeCodesPerSecond, 1);
 
-	const ERichCurveInterpMode InterpMode = ( UsdStage->GetInterpolationType() == pxr::UsdInterpolationTypeLinear ) ? ERichCurveInterpMode::RCIM_Linear : ERichCurveInterpMode::RCIM_Constant;
+	const ERichCurveInterpMode InterpMode = (UsdStage->GetInterpolationType() == pxr::UsdInterpolationTypeLinear)
+												? ERichCurveInterpMode::RCIM_Linear
+												: ERichCurveInterpMode::RCIM_Constant;
 
 	double LastTimeSample = TNumericLimits<double>::Lowest();
-	for ( const double UsdTimeSample : UsdTimeSamples )
+	for (const double UsdTimeSample : UsdTimeSamples)
 	{
 		// We never want to evaluate the same time twice
-		if ( FMath::IsNearlyEqual( UsdTimeSample, LastTimeSample ) )
+		if (FMath::IsNearlyEqual(UsdTimeSample, LastTimeSample))
 		{
 			continue;
 		}
 		LastTimeSample = UsdTimeSample;
 
-		int32 FrameNumber = FMath::FloorToInt( UsdTimeSample );
+		int32 FrameNumber = FMath::FloorToInt(UsdTimeSample);
 		float SubFrameNumber = UsdTimeSample - FrameNumber;
 
-		FFrameTime FrameTime( FrameNumber, SubFrameNumber );
+		FFrameTime FrameTime(FrameNumber, SubFrameNumber);
 
-		FFrameTime KeyFrameTime = FFrameRate::TransformTime( FrameTime, StageFrameRate, Resolution );
+		FFrameTime KeyFrameTime = FFrameRate::TransformTime(FrameTime, StageFrameRate, Resolution);
 		KeyFrameTime *= SequenceTransform;
-		FrameNumbers.Add( KeyFrameTime.GetFrame() );
+		FrameNumbers.Add(KeyFrameTime.GetFrame());
 
-		FTransform UEValue = ReaderFunc( UsdTimeSample );
+		FTransform UEValue = ReaderFunc(UsdTimeSample);
 		FVector Location = UEValue.GetLocation();
 		FRotator Rotator = UEValue.Rotator();
 		FVector Scale = UEValue.GetScale3D();
 
-		LocationXValues.Emplace_GetRef( Location.X ).InterpMode = InterpMode;
-		LocationYValues.Emplace_GetRef( Location.Y ).InterpMode = InterpMode;
-		LocationZValues.Emplace_GetRef( Location.Z ).InterpMode = InterpMode;
+		LocationXValues.Emplace_GetRef(Location.X).InterpMode = InterpMode;
+		LocationYValues.Emplace_GetRef(Location.Y).InterpMode = InterpMode;
+		LocationZValues.Emplace_GetRef(Location.Z).InterpMode = InterpMode;
 
-		RotationXValues.Emplace_GetRef( Rotator.Roll ).InterpMode = InterpMode;
-		RotationYValues.Emplace_GetRef( Rotator.Pitch ).InterpMode = InterpMode;
-		RotationZValues.Emplace_GetRef( Rotator.Yaw ).InterpMode = InterpMode;
+		RotationXValues.Emplace_GetRef(Rotator.Roll).InterpMode = InterpMode;
+		RotationYValues.Emplace_GetRef(Rotator.Pitch).InterpMode = InterpMode;
+		RotationZValues.Emplace_GetRef(Rotator.Yaw).InterpMode = InterpMode;
 
-		ScaleXValues.Emplace_GetRef( Scale.X ).InterpMode = InterpMode;
-		ScaleYValues.Emplace_GetRef( Scale.Y ).InterpMode = InterpMode;
-		ScaleZValues.Emplace_GetRef( Scale.Z ).InterpMode = InterpMode;
+		ScaleXValues.Emplace_GetRef(Scale.X).InterpMode = InterpMode;
+		ScaleYValues.Emplace_GetRef(Scale.Y).InterpMode = InterpMode;
+		ScaleZValues.Emplace_GetRef(Scale.Z).InterpMode = InterpMode;
 	}
 
 	bool bSectionAdded = false;
-	UMovieScene3DTransformSection* Section = Cast< UMovieScene3DTransformSection >( MovieSceneTrack.FindOrAddSection( 0, bSectionAdded ) );
+	UMovieScene3DTransformSection* Section = Cast<UMovieScene3DTransformSection>(MovieSceneTrack.FindOrAddSection(0, bSectionAdded));
 	Section->EvalOptions.CompletionMode = EMovieSceneCompletionMode::KeepState;
-	Section->SetRange( TRange< FFrameNumber >::All() );
+	Section->SetRange(TRange<FFrameNumber>::All());
 
-	TArrayView< FMovieSceneDoubleChannel* > Channels = Section->GetChannelProxy().GetChannels< FMovieSceneDoubleChannel >();
-	if ( Channels.Num() < 9 )
+	TArrayView<FMovieSceneDoubleChannel*> Channels = Section->GetChannelProxy().GetChannels<FMovieSceneDoubleChannel>();
+	if (Channels.Num() < 9)
 	{
 		return false;
 	}
 
-	Channels[ 0 ]->Set( FrameNumbers, LocationXValues );
-	Channels[ 1 ]->Set( FrameNumbers, LocationYValues );
-	Channels[ 2 ]->Set( FrameNumbers, LocationZValues );
+	Channels[0]->Set(FrameNumbers, LocationXValues);
+	Channels[1]->Set(FrameNumbers, LocationYValues);
+	Channels[2]->Set(FrameNumbers, LocationZValues);
 
-	Channels[ 3 ]->Set( FrameNumbers, RotationXValues );
-	Channels[ 4 ]->Set( FrameNumbers, RotationYValues );
-	Channels[ 5 ]->Set( FrameNumbers, RotationZValues );
+	Channels[3]->Set(FrameNumbers, RotationXValues);
+	Channels[4]->Set(FrameNumbers, RotationYValues);
+	Channels[5]->Set(FrameNumbers, RotationZValues);
 
-	Channels[ 6 ]->Set( FrameNumbers, ScaleXValues );
-	Channels[ 7 ]->Set( FrameNumbers, ScaleYValues );
-	Channels[ 8 ]->Set( FrameNumbers, ScaleZValues );
+	Channels[6]->Set(FrameNumbers, ScaleXValues);
+	Channels[7]->Set(FrameNumbers, ScaleYValues);
+	Channels[8]->Set(FrameNumbers, ScaleZValues);
 
 	return true;
 }
@@ -1055,7 +1142,9 @@ bool UsdToUnreal::ConvertBoundsTimeSamples(
 		{
 			FVector UESpaceUSDMin = UsdToUnreal::ConvertVector(StageInfo, Box.GetMin());
 			FVector UESpaceUSDMax = UsdToUnreal::ConvertVector(StageInfo, Box.GetMax());
-			UEBox = FBox{TArray<FVector>{UESpaceUSDMin, UESpaceUSDMax}};
+			UEBox = FBox{
+				TArray<FVector>{UESpaceUSDMin, UESpaceUSDMax}
+			};
 		}
 
 		MinXValues.Emplace_GetRef(UEBox.Min.X).InterpMode = InterpMode;
@@ -1094,94 +1183,103 @@ bool UsdToUnreal::ConvertBoundsTimeSamples(
 	return true;
 }
 
-UsdToUnreal::FPropertyTrackReader UsdToUnreal::CreatePropertyTrackReader( const UE::FUsdPrim& Prim, const FName& PropertyPath, bool bIgnorePrimLocalTransform )
+UsdToUnreal::FPropertyTrackReader UsdToUnreal::CreatePropertyTrackReader(
+	const UE::FUsdPrim& Prim,
+	const FName& PropertyPath,
+	bool bIgnorePrimLocalTransform
+)
 {
 	UsdToUnreal::FPropertyTrackReader Reader;
 
 	FScopedUsdAllocs Allocs;
 
-	pxr::UsdPrim UsdPrim{ Prim };
+	pxr::UsdPrim UsdPrim{Prim};
 	pxr::UsdStageRefPtr UsdStage = UsdPrim.GetStage();
-	FUsdStageInfo StageInfo{ UsdStage };
+	FUsdStageInfo StageInfo{UsdStage};
 
-	if ( pxr::UsdGeomXformable Xformable{ UsdPrim } )
+	if (pxr::UsdGeomXformable Xformable{UsdPrim})
 	{
-		if ( PropertyPath == UnrealIdentifiers::TransformPropertyName )
+		if (PropertyPath == UnrealIdentifiers::TransformPropertyName)
 		{
 			FTransform Default = FTransform::Identity;
-			UE::USDPrimConversionImpl::Private::GetPrimConvertedRelativeTransform( Xformable, UsdUtils::GetDefaultTimeCode(), Default, bIgnorePrimLocalTransform );
+			UE::USDPrimConversionImpl::Private::GetPrimConvertedRelativeTransform(
+				Xformable,
+				UsdUtils::GetDefaultTimeCode(),
+				Default,
+				bIgnorePrimLocalTransform
+			);
 
-			Reader.TransformReader = [UsdStage, Xformable, Default, bIgnorePrimLocalTransform]( double UsdTimeCode )
+			Reader.TransformReader = [UsdStage, Xformable, Default, bIgnorePrimLocalTransform](double UsdTimeCode)
 			{
 				FTransform Result = Default;
-				UE::USDPrimConversionImpl::Private::GetPrimConvertedRelativeTransform( Xformable, UsdTimeCode, Result, bIgnorePrimLocalTransform );
+				UE::USDPrimConversionImpl::Private::GetPrimConvertedRelativeTransform(Xformable, UsdTimeCode, Result, bIgnorePrimLocalTransform);
 				return Result;
 			};
 			return Reader;
 		}
 	}
 
-	if ( pxr::UsdGeomImageable Imageable{ UsdPrim } )
+	if (pxr::UsdGeomImageable Imageable{UsdPrim})
 	{
-		if ( PropertyPath == UnrealIdentifiers::HiddenInGamePropertyName )
+		if (PropertyPath == UnrealIdentifiers::HiddenInGamePropertyName)
 		{
-			if ( pxr::UsdAttribute Attr = Imageable.GetVisibilityAttr() )
+			if (pxr::UsdAttribute Attr = Imageable.GetVisibilityAttr())
 			{
 				pxr::TfToken Default = pxr::UsdGeomTokens->inherited;
-				Attr.Get<pxr::TfToken>( &Default );
+				Attr.Get<pxr::TfToken>(&Default);
 
-				Reader.BoolReader = [Imageable]( double UsdTimeCode )
+				Reader.BoolReader = [Imageable](double UsdTimeCode)
 				{
 					// The property is "HiddenInGame" but it will end up in a visibility track, which is just a bool track,
 					// where true means visible
-					return Imageable.ComputeVisibility( UsdTimeCode ) == pxr::UsdGeomTokens->inherited;
+					return Imageable.ComputeVisibility(UsdTimeCode) == pxr::UsdGeomTokens->inherited;
 				};
 				return Reader;
 			}
 		}
 	}
 
-	if ( pxr::UsdGeomCamera Camera{ UsdPrim } )
+	if (pxr::UsdGeomCamera Camera{UsdPrim})
 	{
 		bool bConvertDistance = true;
 		pxr::UsdAttribute Attr;
 
-		if ( PropertyPath == UnrealIdentifiers::CurrentFocalLengthPropertyName )
+		if (PropertyPath == UnrealIdentifiers::CurrentFocalLengthPropertyName)
 		{
 			Attr = Camera.GetFocalLengthAttr();
 		}
-		else if ( PropertyPath == UnrealIdentifiers::ManualFocusDistancePropertyName )
+		else if (PropertyPath == UnrealIdentifiers::ManualFocusDistancePropertyName)
 		{
 			Attr = Camera.GetFocusDistanceAttr();
 		}
-		else if ( PropertyPath == UnrealIdentifiers::CurrentAperturePropertyName )
+		else if (PropertyPath == UnrealIdentifiers::CurrentAperturePropertyName)
 		{
 			bConvertDistance = false;
 			Attr = Camera.GetFStopAttr();
 		}
-		else if ( PropertyPath == UnrealIdentifiers::SensorWidthPropertyName )
+		else if (PropertyPath == UnrealIdentifiers::SensorWidthPropertyName)
 		{
 			Attr = Camera.GetHorizontalApertureAttr();
 		}
-		else if ( PropertyPath == UnrealIdentifiers::SensorHeightPropertyName )
+		else if (PropertyPath == UnrealIdentifiers::SensorHeightPropertyName)
 		{
 			Attr = Camera.GetVerticalApertureAttr();
 		}
 
-		if ( Attr )
+		if (Attr)
 		{
-			if ( bConvertDistance )
+			if (bConvertDistance)
 			{
 				float Default = 0.0;
-				Attr.Get<float>( &Default );
-				Default = UsdToUnreal::ConvertDistance( StageInfo, Default );
+				Attr.Get<float>(&Default);
+				Default = UsdToUnreal::ConvertDistance(StageInfo, Default);
 
-				Reader.FloatReader = [Attr, Default, StageInfo]( double UsdTimeCode )
+				Reader.FloatReader = [Attr, Default, StageInfo](double UsdTimeCode)
 				{
 					float Result = Default;
-					if ( Attr.Get<float>( &Result, UsdTimeCode ) )
+					if (Attr.Get<float>(&Result, UsdTimeCode))
 					{
-						Result = UsdToUnreal::ConvertDistance( StageInfo, Result );
+						Result = UsdToUnreal::ConvertDistance(StageInfo, Result);
 					}
 
 					return Result;
@@ -1191,36 +1289,36 @@ UsdToUnreal::FPropertyTrackReader UsdToUnreal::CreatePropertyTrackReader( const
 			else
 			{
 				float Default = 0.0;
-				Attr.Get<float>( &Default );
+				Attr.Get<float>(&Default);
 
-				Reader.FloatReader = [Attr, Default]( double UsdTimeCode )
+				Reader.FloatReader = [Attr, Default](double UsdTimeCode)
 				{
 					float Result = Default;
-					Attr.Get<float>( &Result, UsdTimeCode );
+					Attr.Get<float>(&Result, UsdTimeCode);
 					return Result;
 				};
 				return Reader;
 			}
 		}
 	}
-	else if ( const pxr::UsdLuxLightAPI LightAPI{ Prim } )
+	else if (const pxr::UsdLuxLightAPI LightAPI{Prim})
 	{
-		if ( PropertyPath == UnrealIdentifiers::LightColorPropertyName )
+		if (PropertyPath == UnrealIdentifiers::LightColorPropertyName)
 		{
-			if ( pxr::UsdAttribute Attr = LightAPI.GetColorAttr() )
+			if (pxr::UsdAttribute Attr = LightAPI.GetColorAttr())
 			{
 				pxr::GfVec3f UsdDefault;
-				Attr.Get<pxr::GfVec3f>( &UsdDefault );
-				FLinearColor Default = UsdToUnreal::ConvertColor( UsdDefault );
+				Attr.Get<pxr::GfVec3f>(&UsdDefault);
+				FLinearColor Default = UsdToUnreal::ConvertColor(UsdDefault);
 
-				Reader.ColorReader = [Attr, Default]( double UsdTimeCode )
+				Reader.ColorReader = [Attr, Default](double UsdTimeCode)
 				{
 					FLinearColor Result = Default;
 
 					pxr::GfVec3f Value;
-					if ( Attr.Get<pxr::GfVec3f>( &Value, UsdTimeCode ) )
+					if (Attr.Get<pxr::GfVec3f>(&Value, UsdTimeCode))
 					{
-						Result = UsdToUnreal::ConvertColor( Value );
+						Result = UsdToUnreal::ConvertColor(Value);
 					}
 
 					return Result;
@@ -1228,55 +1326,55 @@ UsdToUnreal::FPropertyTrackReader UsdToUnreal::CreatePropertyTrackReader( const
 				return Reader;
 			}
 		}
-		else if ( PropertyPath == UnrealIdentifiers::UseTemperaturePropertyName )
+		else if (PropertyPath == UnrealIdentifiers::UseTemperaturePropertyName)
 		{
-			if ( pxr::UsdAttribute Attr = LightAPI.GetEnableColorTemperatureAttr() )
+			if (pxr::UsdAttribute Attr = LightAPI.GetEnableColorTemperatureAttr())
 			{
 				bool Default;
-				Attr.Get<bool>( &Default );
+				Attr.Get<bool>(&Default);
 
-				Reader.BoolReader = [Attr, Default]( double UsdTimeCode )
+				Reader.BoolReader = [Attr, Default](double UsdTimeCode)
 				{
 					bool Result = Default;
-					Attr.Get<bool>( &Result, UsdTimeCode );
+					Attr.Get<bool>(&Result, UsdTimeCode);
 					return Result;
 				};
 				return Reader;
 			}
 		}
-		else if ( PropertyPath == UnrealIdentifiers::TemperaturePropertyName )
+		else if (PropertyPath == UnrealIdentifiers::TemperaturePropertyName)
 		{
-			if ( pxr::UsdAttribute Attr = LightAPI.GetColorTemperatureAttr() )
+			if (pxr::UsdAttribute Attr = LightAPI.GetColorTemperatureAttr())
 			{
 				float Default;
-				Attr.Get<float>( &Default );
+				Attr.Get<float>(&Default);
 
-				Reader.FloatReader = [Attr, Default]( double UsdTimeCode )
+				Reader.FloatReader = [Attr, Default](double UsdTimeCode)
 				{
 					float Result = Default;
-					Attr.Get<float>( &Result, UsdTimeCode );
+					Attr.Get<float>(&Result, UsdTimeCode);
 					return Result;
 				};
 				return Reader;
 			}
 		}
 
-		else if ( pxr::UsdLuxSphereLight SphereLight{ UsdPrim } )
+		else if (pxr::UsdLuxSphereLight SphereLight{UsdPrim})
 		{
-			if ( PropertyPath == UnrealIdentifiers::SourceRadiusPropertyName )
+			if (PropertyPath == UnrealIdentifiers::SourceRadiusPropertyName)
 			{
-				if ( pxr::UsdAttribute Attr = SphereLight.GetRadiusAttr() )
+				if (pxr::UsdAttribute Attr = SphereLight.GetRadiusAttr())
 				{
 					float Default = 0.0;
-					Attr.Get<float>( &Default );
-					Default = UsdToUnreal::ConvertDistance( StageInfo, Default );
+					Attr.Get<float>(&Default);
+					Default = UsdToUnreal::ConvertDistance(StageInfo, Default);
 
-					Reader.FloatReader = [Attr, Default, StageInfo]( double UsdTimeCode )
+					Reader.FloatReader = [Attr, Default, StageInfo](double UsdTimeCode)
 					{
 						float Result = Default;
-						if ( Attr.Get<float>( &Result, UsdTimeCode ) )
+						if (Attr.Get<float>(&Result, UsdTimeCode))
 						{
-							Result = UsdToUnreal::ConvertDistance( StageInfo, Result );
+							Result = UsdToUnreal::ConvertDistance(StageInfo, Result);
 						}
 
 						return Result;
@@ -1286,11 +1384,11 @@ UsdToUnreal::FPropertyTrackReader UsdToUnreal::CreatePropertyTrackReader( const
 			}
 
 			// Spot light
-			else if ( UsdPrim.HasAPI< pxr::UsdLuxShapingAPI >() )
+			else if (UsdPrim.HasAPI<pxr::UsdLuxShapingAPI>())
 			{
-				pxr::UsdLuxShapingAPI ShapingAPI{ UsdPrim };
+				pxr::UsdLuxShapingAPI ShapingAPI{UsdPrim};
 
-				if ( PropertyPath == UnrealIdentifiers::IntensityPropertyName )
+				if (PropertyPath == UnrealIdentifiers::IntensityPropertyName)
 				{
 					pxr::UsdAttribute IntensityAttr = SphereLight.GetIntensityAttr();
 					pxr::UsdAttribute ExposureAttr = SphereLight.GetExposureAttr();
@@ -1298,7 +1396,7 @@ UsdToUnreal::FPropertyTrackReader UsdToUnreal::CreatePropertyTrackReader( const
 					pxr::UsdAttribute ConeAngleAttr = ShapingAPI.GetShapingConeAngleAttr();
 					pxr::UsdAttribute ConeSoftnessAttr = ShapingAPI.GetShapingConeSoftnessAttr();
 
-					if ( IntensityAttr && ExposureAttr && RadiusAttr && ConeAngleAttr && ConeSoftnessAttr )
+					if (IntensityAttr && ExposureAttr && RadiusAttr && ConeAngleAttr && ConeSoftnessAttr)
 					{
 						// Default values directly from pxr/usd/usdLux/schema.usda
 						float DefaultUsdIntensity = 1.0f;
@@ -1307,11 +1405,11 @@ UsdToUnreal::FPropertyTrackReader UsdToUnreal::CreatePropertyTrackReader( const
 						float DefaultUsdConeAngle = 90.0f;
 						float DefaultUsdConeSoftness = 0.0f;
 
-						IntensityAttr.Get<float>( &DefaultUsdIntensity );
-						ExposureAttr.Get<float>( &DefaultUsdExposure );
-						RadiusAttr.Get<float>( &DefaultUsdRadius );
-						ConeAngleAttr.Get<float>( &DefaultUsdConeAngle );
-						ConeSoftnessAttr.Get<float>( &DefaultUsdConeSoftness );
+						IntensityAttr.Get<float>(&DefaultUsdIntensity);
+						ExposureAttr.Get<float>(&DefaultUsdExposure);
+						RadiusAttr.Get<float>(&DefaultUsdRadius);
+						ConeAngleAttr.Get<float>(&DefaultUsdConeAngle);
+						ConeSoftnessAttr.Get<float>(&DefaultUsdConeSoftness);
 
 						float Default = UsdToUnreal::ConvertLuxShapingAPIIntensityAttr(
 							DefaultUsdIntensity,
@@ -1322,7 +1420,8 @@ UsdToUnreal::FPropertyTrackReader UsdToUnreal::CreatePropertyTrackReader( const
 							StageInfo
 						);
 
-						Reader.FloatReader = [IntensityAttr, ExposureAttr, RadiusAttr, ConeAngleAttr, ConeSoftnessAttr, Default, StageInfo]( double UsdTimeCode )
+						Reader.FloatReader =
+							[IntensityAttr, ExposureAttr, RadiusAttr, ConeAngleAttr, ConeSoftnessAttr, Default, StageInfo](double UsdTimeCode)
 						{
 							float Result = Default;
 
@@ -1331,11 +1430,9 @@ UsdToUnreal::FPropertyTrackReader UsdToUnreal::CreatePropertyTrackReader( const
 							float UsdRadius = 0.5f;
 							float UsdConeAngle = 90.0f;
 							float UsdConeSoftness = 0.0f;
-							if ( IntensityAttr.Get<float>( &UsdIntensity, UsdTimeCode ) &&
-								 ExposureAttr.Get<float>( &UsdExposure, UsdTimeCode ) &&
-								 RadiusAttr.Get<float>( &UsdRadius, UsdTimeCode ) &&
-								 ConeAngleAttr.Get<float>( &UsdConeAngle, UsdTimeCode ) &&
-								 ConeSoftnessAttr.Get<float>( &UsdConeSoftness, UsdTimeCode ) )
+							if (IntensityAttr.Get<float>(&UsdIntensity, UsdTimeCode) && ExposureAttr.Get<float>(&UsdExposure, UsdTimeCode)
+								&& RadiusAttr.Get<float>(&UsdRadius, UsdTimeCode) && ConeAngleAttr.Get<float>(&UsdConeAngle, UsdTimeCode)
+								&& ConeSoftnessAttr.Get<float>(&UsdConeSoftness, UsdTimeCode))
 							{
 								Result = UsdToUnreal::ConvertLuxShapingAPIIntensityAttr(
 									UsdIntensity,
@@ -1352,48 +1449,48 @@ UsdToUnreal::FPropertyTrackReader UsdToUnreal::CreatePropertyTrackReader( const
 						return Reader;
 					}
 				}
-				else if ( PropertyPath == UnrealIdentifiers::OuterConeAnglePropertyName )
+				else if (PropertyPath == UnrealIdentifiers::OuterConeAnglePropertyName)
 				{
-					if ( pxr::UsdAttribute Attr = ShapingAPI.GetShapingConeAngleAttr() )
+					if (pxr::UsdAttribute Attr = ShapingAPI.GetShapingConeAngleAttr())
 					{
 						float Default;
-						Attr.Get<float>( &Default );
+						Attr.Get<float>(&Default);
 
-						Reader.FloatReader = [Attr, Default]( double UsdTimeCode )
+						Reader.FloatReader = [Attr, Default](double UsdTimeCode)
 						{
 							float Result = Default;
-							Attr.Get<float>( &Result, UsdTimeCode );
+							Attr.Get<float>(&Result, UsdTimeCode);
 							return Result;
 						};
 						return Reader;
 					}
 				}
-				else if ( PropertyPath == UnrealIdentifiers::InnerConeAnglePropertyName )
+				else if (PropertyPath == UnrealIdentifiers::InnerConeAnglePropertyName)
 				{
 					pxr::UsdAttribute ConeAngleAttr = ShapingAPI.GetShapingConeAngleAttr();
 					pxr::UsdAttribute ConeSoftnessAttr = ShapingAPI.GetShapingConeSoftnessAttr();
 
-					if ( ConeAngleAttr && ConeSoftnessAttr )
+					if (ConeAngleAttr && ConeSoftnessAttr)
 					{
 						// Default values directly from pxr/usd/usdLux/schema.usda
 						float DefaultUsdConeAngle = 90.0f;
 						float DefaultUsdConeSoftness = 0.0f;
 
-						ConeAngleAttr.Get<float>( &DefaultUsdConeAngle );
-						ConeSoftnessAttr.Get<float>( &DefaultUsdConeSoftness );
+						ConeAngleAttr.Get<float>(&DefaultUsdConeAngle);
+						ConeSoftnessAttr.Get<float>(&DefaultUsdConeSoftness);
 
 						float Default = 0.0f;
-						UsdToUnreal::ConvertConeAngleSoftnessAttr( DefaultUsdConeAngle, DefaultUsdConeSoftness, Default );
+						UsdToUnreal::ConvertConeAngleSoftnessAttr(DefaultUsdConeAngle, DefaultUsdConeSoftness, Default);
 
-						Reader.FloatReader = [ConeAngleAttr, ConeSoftnessAttr, Default]( double UsdTimeCode )
+						Reader.FloatReader = [ConeAngleAttr, ConeSoftnessAttr, Default](double UsdTimeCode)
 						{
 							float Result = Default;
 
 							float UsdConeAngle = 90.0f;
 							float UsdConeSoftness = 0.0f;
-							if ( ConeAngleAttr.Get<float>( &UsdConeAngle, UsdTimeCode ) && ConeSoftnessAttr.Get<float>( &UsdConeSoftness, UsdTimeCode ) )
+							if (ConeAngleAttr.Get<float>(&UsdConeAngle, UsdTimeCode) && ConeSoftnessAttr.Get<float>(&UsdConeSoftness, UsdTimeCode))
 							{
-								UsdToUnreal::ConvertConeAngleSoftnessAttr( UsdConeAngle, UsdConeSoftness, Result );
+								UsdToUnreal::ConvertConeAngleSoftnessAttr(UsdConeAngle, UsdConeSoftness, Result);
 							}
 
 							return Result;
@@ -1405,22 +1502,22 @@ UsdToUnreal::FPropertyTrackReader UsdToUnreal::CreatePropertyTrackReader( const
 			// Just a point light
 			else
 			{
-				if ( PropertyPath == UnrealIdentifiers::IntensityPropertyName )
+				if (PropertyPath == UnrealIdentifiers::IntensityPropertyName)
 				{
 					pxr::UsdAttribute IntensityAttr = SphereLight.GetIntensityAttr();
 					pxr::UsdAttribute ExposureAttr = SphereLight.GetExposureAttr();
 					pxr::UsdAttribute RadiusAttr = SphereLight.GetRadiusAttr();
 
-					if ( IntensityAttr && ExposureAttr && RadiusAttr )
+					if (IntensityAttr && ExposureAttr && RadiusAttr)
 					{
 						// Default values directly from pxr/usd/usdLux/schema.usda
 						float DefaultUsdIntensity = 1.0f;
 						float DefaultUsdExposure = 0.0f;
 						float DefaultUsdRadius = 0.5f;
 
-						IntensityAttr.Get<float>( &DefaultUsdIntensity );
-						ExposureAttr.Get<float>( &DefaultUsdExposure );
-						RadiusAttr.Get<float>( &DefaultUsdRadius );
+						IntensityAttr.Get<float>(&DefaultUsdIntensity);
+						ExposureAttr.Get<float>(&DefaultUsdExposure);
+						RadiusAttr.Get<float>(&DefaultUsdRadius);
 
 						float Default = UsdToUnreal::ConvertSphereLightIntensityAttr(
 							DefaultUsdIntensity,
@@ -1429,23 +1526,17 @@ UsdToUnreal::FPropertyTrackReader UsdToUnreal::CreatePropertyTrackReader( const
 							StageInfo
 						);
 
-						Reader.FloatReader = [IntensityAttr, ExposureAttr, RadiusAttr, Default, StageInfo]( double UsdTimeCode )
+						Reader.FloatReader = [IntensityAttr, ExposureAttr, RadiusAttr, Default, StageInfo](double UsdTimeCode)
 						{
 							float Result = Default;
 
 							float UsdIntensity = 1.0f;
 							float UsdExposure = 0.0f;
 							float UsdRadius = 0.5f;
-							if ( IntensityAttr.Get<float>( &UsdIntensity, UsdTimeCode ) &&
-								 ExposureAttr.Get<float>( &UsdExposure, UsdTimeCode ) &&
-								 RadiusAttr.Get<float>( &UsdRadius, UsdTimeCode ) )
+							if (IntensityAttr.Get<float>(&UsdIntensity, UsdTimeCode) && ExposureAttr.Get<float>(&UsdExposure, UsdTimeCode)
+								&& RadiusAttr.Get<float>(&UsdRadius, UsdTimeCode))
 							{
-								Result = UsdToUnreal::ConvertSphereLightIntensityAttr(
-									UsdIntensity,
-									UsdExposure,
-									UsdRadius,
-									StageInfo
-								);
+								Result = UsdToUnreal::ConvertSphereLightIntensityAttr(UsdIntensity, UsdExposure, UsdRadius, StageInfo);
 							}
 
 							return Result;
@@ -1455,22 +1546,22 @@ UsdToUnreal::FPropertyTrackReader UsdToUnreal::CreatePropertyTrackReader( const
 				}
 			}
 		}
-		else if ( pxr::UsdLuxRectLight RectLight{ UsdPrim } )
+		else if (pxr::UsdLuxRectLight RectLight{UsdPrim})
 		{
-			if ( PropertyPath == UnrealIdentifiers::SourceWidthPropertyName )
+			if (PropertyPath == UnrealIdentifiers::SourceWidthPropertyName)
 			{
-				if ( pxr::UsdAttribute Attr = RectLight.GetWidthAttr() )
+				if (pxr::UsdAttribute Attr = RectLight.GetWidthAttr())
 				{
 					float Default = 0.0;
-					Attr.Get<float>( &Default );
-					Default = UsdToUnreal::ConvertDistance( StageInfo, Default );
+					Attr.Get<float>(&Default);
+					Default = UsdToUnreal::ConvertDistance(StageInfo, Default);
 
-					Reader.FloatReader = [Attr, Default, StageInfo]( double UsdTimeCode )
+					Reader.FloatReader = [Attr, Default, StageInfo](double UsdTimeCode)
 					{
 						float Result = Default;
-						if ( Attr.Get<float>( &Result, UsdTimeCode ) )
+						if (Attr.Get<float>(&Result, UsdTimeCode))
 						{
-							Result = UsdToUnreal::ConvertDistance( StageInfo, Result );
+							Result = UsdToUnreal::ConvertDistance(StageInfo, Result);
 						}
 
 						return Result;
@@ -1478,20 +1569,20 @@ UsdToUnreal::FPropertyTrackReader UsdToUnreal::CreatePropertyTrackReader( const
 					return Reader;
 				}
 			}
-			else if ( PropertyPath == UnrealIdentifiers::SourceHeightPropertyName )
+			else if (PropertyPath == UnrealIdentifiers::SourceHeightPropertyName)
 			{
-				if ( pxr::UsdAttribute Attr = RectLight.GetHeightAttr() )
+				if (pxr::UsdAttribute Attr = RectLight.GetHeightAttr())
 				{
 					float Default = 0.0;
-					Attr.Get<float>( &Default );
-					Default = UsdToUnreal::ConvertDistance( StageInfo, Default );
+					Attr.Get<float>(&Default);
+					Default = UsdToUnreal::ConvertDistance(StageInfo, Default);
 
-					Reader.FloatReader = [Attr, Default, StageInfo]( double UsdTimeCode )
+					Reader.FloatReader = [Attr, Default, StageInfo](double UsdTimeCode)
 					{
 						float Result = Default;
-						if ( Attr.Get<float>( &Result, UsdTimeCode ) )
+						if (Attr.Get<float>(&Result, UsdTimeCode))
 						{
-							Result = UsdToUnreal::ConvertDistance( StageInfo, Result );
+							Result = UsdToUnreal::ConvertDistance(StageInfo, Result);
 						}
 
 						return Result;
@@ -1499,14 +1590,14 @@ UsdToUnreal::FPropertyTrackReader UsdToUnreal::CreatePropertyTrackReader( const
 					return Reader;
 				}
 			}
-			else if ( PropertyPath == UnrealIdentifiers::IntensityPropertyName )
+			else if (PropertyPath == UnrealIdentifiers::IntensityPropertyName)
 			{
 				pxr::UsdAttribute IntensityAttr = RectLight.GetIntensityAttr();
 				pxr::UsdAttribute ExposureAttr = RectLight.GetExposureAttr();
 				pxr::UsdAttribute WidthAttr = RectLight.GetWidthAttr();
 				pxr::UsdAttribute HeightAttr = RectLight.GetHeightAttr();
 
-				if ( IntensityAttr && ExposureAttr && WidthAttr && HeightAttr )
+				if (IntensityAttr && ExposureAttr && WidthAttr && HeightAttr)
 				{
 					// Default values directly from pxr/usd/usdLux/schema.usda
 					float DefaultUsdIntensity = 1.0f;
@@ -1514,10 +1605,10 @@ UsdToUnreal::FPropertyTrackReader UsdToUnreal::CreatePropertyTrackReader( const
 					float DefaultUsdWidth = 1.0f;
 					float DefaultUsdHeight = 1.0f;
 
-					IntensityAttr.Get<float>( &DefaultUsdIntensity );
-					ExposureAttr.Get<float>( &DefaultUsdExposure );
-					WidthAttr.Get<float>( &DefaultUsdWidth );
-					HeightAttr.Get<float>( &DefaultUsdHeight );
+					IntensityAttr.Get<float>(&DefaultUsdIntensity);
+					ExposureAttr.Get<float>(&DefaultUsdExposure);
+					WidthAttr.Get<float>(&DefaultUsdWidth);
+					HeightAttr.Get<float>(&DefaultUsdHeight);
 
 					float Default = UsdToUnreal::ConvertRectLightIntensityAttr(
 						DefaultUsdIntensity,
@@ -1527,7 +1618,7 @@ UsdToUnreal::FPropertyTrackReader UsdToUnreal::CreatePropertyTrackReader( const
 						StageInfo
 					);
 
-					Reader.FloatReader = [IntensityAttr, ExposureAttr, WidthAttr, HeightAttr, Default, StageInfo]( double UsdTimeCode )
+					Reader.FloatReader = [IntensityAttr, ExposureAttr, WidthAttr, HeightAttr, Default, StageInfo](double UsdTimeCode)
 					{
 						float Result = Default;
 
@@ -1535,18 +1626,10 @@ UsdToUnreal::FPropertyTrackReader UsdToUnreal::CreatePropertyTrackReader( const
 						float UsdExposure = 0.0f;
 						float UsdWidth = 1.0f;
 						float UsdHeight = 1.0f;
-						if ( IntensityAttr.Get<float>( &UsdIntensity, UsdTimeCode ) &&
-							 ExposureAttr.Get<float>( &UsdExposure, UsdTimeCode ) &&
-							 WidthAttr.Get<float>( &UsdWidth, UsdTimeCode ) &&
-							 HeightAttr.Get<float>( &UsdHeight, UsdTimeCode ) )
+						if (IntensityAttr.Get<float>(&UsdIntensity, UsdTimeCode) && ExposureAttr.Get<float>(&UsdExposure, UsdTimeCode)
+							&& WidthAttr.Get<float>(&UsdWidth, UsdTimeCode) && HeightAttr.Get<float>(&UsdHeight, UsdTimeCode))
 						{
-							Result = UsdToUnreal::ConvertRectLightIntensityAttr(
-								UsdIntensity,
-								UsdExposure,
-								UsdWidth,
-								UsdHeight,
-								StageInfo
-							);
+							Result = UsdToUnreal::ConvertRectLightIntensityAttr(UsdIntensity, UsdExposure, UsdWidth, UsdHeight, StageInfo);
 						}
 
 						return Result;
@@ -1555,24 +1638,24 @@ UsdToUnreal::FPropertyTrackReader UsdToUnreal::CreatePropertyTrackReader( const
 				}
 			}
 		}
-		else if ( pxr::UsdLuxDiskLight DiskLight{ UsdPrim } )
+		else if (pxr::UsdLuxDiskLight DiskLight{UsdPrim})
 		{
-			if ( PropertyPath == UnrealIdentifiers::SourceWidthPropertyName || PropertyPath == UnrealIdentifiers::SourceHeightPropertyName )
+			if (PropertyPath == UnrealIdentifiers::SourceWidthPropertyName || PropertyPath == UnrealIdentifiers::SourceHeightPropertyName)
 			{
-				if ( pxr::UsdAttribute Attr = DiskLight.GetRadiusAttr() )
+				if (pxr::UsdAttribute Attr = DiskLight.GetRadiusAttr())
 				{
 					// Our conversion is that Width == Height == 2 * Radius
 
 					float Default = 0.0;
-					Attr.Get<float>( &Default );
-					Default = 2.0f * UsdToUnreal::ConvertDistance( StageInfo, Default );
+					Attr.Get<float>(&Default);
+					Default = 2.0f * UsdToUnreal::ConvertDistance(StageInfo, Default);
 
-					Reader.FloatReader = [Attr, Default, StageInfo]( double UsdTimeCode )
+					Reader.FloatReader = [Attr, Default, StageInfo](double UsdTimeCode)
 					{
 						float Result = Default;
-						if ( Attr.Get<float>( &Result, UsdTimeCode ) )
+						if (Attr.Get<float>(&Result, UsdTimeCode))
 						{
-							Result = 2.0f * UsdToUnreal::ConvertDistance( StageInfo, Result );
+							Result = 2.0f * UsdToUnreal::ConvertDistance(StageInfo, Result);
 						}
 
 						return Result;
@@ -1580,47 +1663,36 @@ UsdToUnreal::FPropertyTrackReader UsdToUnreal::CreatePropertyTrackReader( const
 					return Reader;
 				}
 			}
-			else if ( PropertyPath == UnrealIdentifiers::IntensityPropertyName )
+			else if (PropertyPath == UnrealIdentifiers::IntensityPropertyName)
 			{
 				pxr::UsdAttribute IntensityAttr = DiskLight.GetIntensityAttr();
 				pxr::UsdAttribute ExposureAttr = DiskLight.GetExposureAttr();
 				pxr::UsdAttribute RadiusAttr = DiskLight.GetRadiusAttr();
 
-				if ( IntensityAttr && ExposureAttr && RadiusAttr )
+				if (IntensityAttr && ExposureAttr && RadiusAttr)
 				{
 					// Default values directly from pxr/usd/usdLux/schema.usda
 					float DefaultUsdIntensity = 1.0f;
 					float DefaultUsdExposure = 0.0f;
 					float DefaultUsdRadius = 0.5f;
 
-					IntensityAttr.Get<float>( &DefaultUsdIntensity );
-					ExposureAttr.Get<float>( &DefaultUsdExposure );
-					RadiusAttr.Get<float>( &DefaultUsdRadius );
+					IntensityAttr.Get<float>(&DefaultUsdIntensity);
+					ExposureAttr.Get<float>(&DefaultUsdExposure);
+					RadiusAttr.Get<float>(&DefaultUsdRadius);
 
-					float Default = UsdToUnreal::ConvertDiskLightIntensityAttr(
-						DefaultUsdIntensity,
-						DefaultUsdExposure,
-						DefaultUsdRadius,
-						StageInfo
-					);
+					float Default = UsdToUnreal::ConvertDiskLightIntensityAttr(DefaultUsdIntensity, DefaultUsdExposure, DefaultUsdRadius, StageInfo);
 
-					Reader.FloatReader = [IntensityAttr, ExposureAttr, RadiusAttr, Default, StageInfo]( double UsdTimeCode )
+					Reader.FloatReader = [IntensityAttr, ExposureAttr, RadiusAttr, Default, StageInfo](double UsdTimeCode)
 					{
 						float Result = Default;
 
 						float UsdIntensity = 1.0f;
 						float UsdExposure = 0.0f;
 						float UsdRadius = 0.5f;
-						if ( IntensityAttr.Get<float>( &UsdIntensity, UsdTimeCode ) &&
-							 ExposureAttr.Get<float>( &UsdExposure, UsdTimeCode ) &&
-							 RadiusAttr.Get<float>( &UsdRadius, UsdTimeCode ) )
+						if (IntensityAttr.Get<float>(&UsdIntensity, UsdTimeCode) && ExposureAttr.Get<float>(&UsdExposure, UsdTimeCode)
+							&& RadiusAttr.Get<float>(&UsdRadius, UsdTimeCode))
 						{
-							Result = UsdToUnreal::ConvertDiskLightIntensityAttr(
-								UsdIntensity,
-								UsdExposure,
-								UsdRadius,
-								StageInfo
-							);
+							Result = UsdToUnreal::ConvertDiskLightIntensityAttr(UsdIntensity, UsdExposure, UsdRadius, StageInfo);
 						}
 
 						return Result;
@@ -1629,49 +1701,49 @@ UsdToUnreal::FPropertyTrackReader UsdToUnreal::CreatePropertyTrackReader( const
 				}
 			}
 		}
-		else if ( pxr::UsdLuxDistantLight DistantLight{ UsdPrim } )
+		else if (pxr::UsdLuxDistantLight DistantLight{UsdPrim})
 		{
-			if ( PropertyPath == UnrealIdentifiers::LightSourceAnglePropertyName )
+			if (PropertyPath == UnrealIdentifiers::LightSourceAnglePropertyName)
 			{
-				if ( pxr::UsdAttribute Attr = DistantLight.GetAngleAttr() )
+				if (pxr::UsdAttribute Attr = DistantLight.GetAngleAttr())
 				{
 					float Default;
-					Attr.Get<float>( &Default );
+					Attr.Get<float>(&Default);
 
-					Reader.FloatReader = [Attr, Default]( double UsdTimeCode )
+					Reader.FloatReader = [Attr, Default](double UsdTimeCode)
 					{
 						float Result = Default;
-						Attr.Get<float>( &Result, UsdTimeCode );
+						Attr.Get<float>(&Result, UsdTimeCode);
 						return Result;
 					};
 					return Reader;
 				}
 			}
-			else if ( PropertyPath == UnrealIdentifiers::IntensityPropertyName )
+			else if (PropertyPath == UnrealIdentifiers::IntensityPropertyName)
 			{
 				pxr::UsdAttribute IntensityAttr = SphereLight.GetIntensityAttr();
 				pxr::UsdAttribute ExposureAttr = SphereLight.GetExposureAttr();
 
-				if ( IntensityAttr && ExposureAttr )
+				if (IntensityAttr && ExposureAttr)
 				{
 					// Default values directly from pxr/usd/usdLux/schema.usda
 					float DefaultUsdIntensity = 1.0f;
 					float DefaultUsdExposure = 0.0f;
 
-					IntensityAttr.Get<float>( &DefaultUsdIntensity );
-					ExposureAttr.Get<float>( &DefaultUsdExposure );
+					IntensityAttr.Get<float>(&DefaultUsdIntensity);
+					ExposureAttr.Get<float>(&DefaultUsdExposure);
 
-					float Default = UsdToUnreal::ConvertDistantLightIntensityAttr( DefaultUsdIntensity, DefaultUsdExposure );
+					float Default = UsdToUnreal::ConvertDistantLightIntensityAttr(DefaultUsdIntensity, DefaultUsdExposure);
 
-					Reader.FloatReader = [IntensityAttr, ExposureAttr, Default, StageInfo]( double UsdTimeCode )
+					Reader.FloatReader = [IntensityAttr, ExposureAttr, Default, StageInfo](double UsdTimeCode)
 					{
 						float Result = Default;
 
 						float UsdIntensity = 1.0f;
 						float UsdExposure = 0.0f;
-						if ( IntensityAttr.Get<float>( &UsdIntensity, UsdTimeCode ) && ExposureAttr.Get<float>( &UsdExposure, UsdTimeCode ) )
+						if (IntensityAttr.Get<float>(&UsdIntensity, UsdTimeCode) && ExposureAttr.Get<float>(&UsdExposure, UsdTimeCode))
 						{
-							Result = UsdToUnreal::ConvertDistantLightIntensityAttr( UsdIntensity, UsdExposure );
+							Result = UsdToUnreal::ConvertDistantLightIntensityAttr(UsdIntensity, UsdExposure);
 						}
 
 						return Result;
@@ -1685,12 +1757,7 @@ UsdToUnreal::FPropertyTrackReader UsdToUnreal::CreatePropertyTrackReader( const
 	return Reader;
 }
 
-bool UsdToUnreal::ConvertBounds(
-	const pxr::UsdPrim& Prim,
-	UUsdDrawModeComponent* BoundsComponent,
-	double EvalTime,
-	pxr::UsdGeomBBoxCache* BBoxCache
-)
+bool UsdToUnreal::ConvertBounds(const pxr::UsdPrim& Prim, UUsdDrawModeComponent* BoundsComponent, double EvalTime, pxr::UsdGeomBBoxCache* BBoxCache)
 {
 	// We're not going to check if Prim actually has the "bounds" draw mode or if it has "applyDrawMode" set
 	// to true, as that can be expensive and this can get called from UpdateComponents, which can get called
@@ -1745,7 +1812,7 @@ bool UsdToUnreal::ConvertBounds(
 		FVector UESpaceUSDMax = UsdToUnreal::ConvertVector(StageInfo, Box.GetMax());
 		FBox UEBox{
 			TArray<FVector>{UESpaceUSDMin, UESpaceUSDMax}
-		 };
+		};
 
 		BoundsComponent->SetBoundsMin(UEBox.Min);
 		BoundsComponent->SetBoundsMax(UEBox.Max);
@@ -1917,10 +1984,7 @@ namespace UE::USDPrimConversionImpl::Private
 			// There is no real point in keeping track of these as they are defined on every prim and should just match
 			// what the prim's definition is. It's probably a bad idea to author a value that differs from the prim
 			// definition too
-			const static std::unordered_set<pxr::TfToken, pxr::TfHash> FieldsToSkip = {
-				pxr::SdfFieldKeys->Specifier,
-				pxr::SdfFieldKeys->TypeName
-			};
+			const static std::unordered_set<pxr::TfToken, pxr::TfHash> FieldsToSkip = {pxr::SdfFieldKeys->Specifier, pxr::SdfFieldKeys->TypeName};
 			if (FieldsToSkip.count(FieldName) > 0)
 			{
 				continue;
@@ -1967,7 +2031,7 @@ namespace UE::USDPrimConversionImpl::Private
 			}
 		}
 	}
-};
+};	  // namespace UE::USDPrimConversionImpl::Private
 
 bool UsdToUnreal::ConvertMetadata(
 	const pxr::UsdPrim& Prim,
@@ -1982,24 +2046,14 @@ bool UsdToUnreal::ConvertMetadata(
 		return false;
 	}
 
-	UE::USDPrimConversionImpl::Private::CollectMetadataForPrim(
-		Prim,
-		CombinedMetadata,
-		BlockedPrefixFilters,
-		bInvertFilters
-	);
+	UE::USDPrimConversionImpl::Private::CollectMetadataForPrim(Prim, CombinedMetadata, BlockedPrefixFilters, bInvertFilters);
 
 	if (bCollectFromEntireSubtrees)
 	{
 		pxr::UsdPrimRange PrimRange{Prim, pxr::UsdTraverseInstanceProxies()};
 		for (pxr::UsdPrimRange::iterator It = ++PrimRange.begin(); It != PrimRange.end(); ++It)
 		{
-			UE::USDPrimConversionImpl::Private::CollectMetadataForPrim(
-				*It,
-				CombinedMetadata,
-				BlockedPrefixFilters,
-				bInvertFilters
-			);
+			UE::USDPrimConversionImpl::Private::CollectMetadataForPrim(*It, CombinedMetadata, BlockedPrefixFilters, bInvertFilters);
 		}
 	}
 
@@ -2030,12 +2084,12 @@ bool UsdToUnreal::ConvertMetadata(
 	return UsdToUnreal::ConvertMetadata(Prim, CombinedMetadata, BlockedPrefixFilters, bInvertFilters, bCollectFromEntireSubtrees);
 }
 
-bool UnrealToUsd::ConvertCameraComponent( const UCineCameraComponent& CameraComponent, pxr::UsdPrim& Prim, double UsdTimeCode )
+bool UnrealToUsd::ConvertCameraComponent(const UCineCameraComponent& CameraComponent, pxr::UsdPrim& Prim, double UsdTimeCode)
 {
 	FScopedUsdAllocs UsdAllocs;
 
-	pxr::UsdGeomCamera GeomCamera{ Prim };
-	if ( !GeomCamera )
+	pxr::UsdGeomCamera GeomCamera{Prim};
+	if (!GeomCamera)
 	{
 		return false;
 	}
@@ -2045,99 +2099,105 @@ bool UnrealToUsd::ConvertCameraComponent( const UCineCameraComponent& CameraComp
 		return false;
 	}
 
-	FUsdStageInfo StageInfo( Prim.GetStage() );
+	FUsdStageInfo StageInfo(Prim.GetStage());
 
-	if ( pxr::UsdAttribute Attr = GeomCamera.CreateFocalLengthAttr() )
+	if (pxr::UsdAttribute Attr = GeomCamera.CreateFocalLengthAttr())
 	{
-		Attr.Set<float>( UnrealToUsd::ConvertDistance( StageInfo, CameraComponent.CurrentFocalLength ), UsdTimeCode );
+		Attr.Set<float>(UnrealToUsd::ConvertDistance(StageInfo, CameraComponent.CurrentFocalLength), UsdTimeCode);
 		UsdUtils::NotifyIfOverriddenOpinion(Attr);
 	}
 
-	if ( pxr::UsdAttribute Attr = GeomCamera.CreateFocusDistanceAttr() )
+	if (pxr::UsdAttribute Attr = GeomCamera.CreateFocusDistanceAttr())
 	{
-		Attr.Set<float>( UnrealToUsd::ConvertDistance( StageInfo, CameraComponent.FocusSettings.ManualFocusDistance ), UsdTimeCode );
+		Attr.Set<float>(UnrealToUsd::ConvertDistance(StageInfo, CameraComponent.FocusSettings.ManualFocusDistance), UsdTimeCode);
 		UsdUtils::NotifyIfOverriddenOpinion(Attr);
 	}
 
-	if ( pxr::UsdAttribute Attr = GeomCamera.CreateFStopAttr() )
+	if (pxr::UsdAttribute Attr = GeomCamera.CreateFStopAttr())
 	{
-		Attr.Set<float>( CameraComponent.CurrentAperture, UsdTimeCode );
+		Attr.Set<float>(CameraComponent.CurrentAperture, UsdTimeCode);
 		UsdUtils::NotifyIfOverriddenOpinion(Attr);
 	}
 
-	if ( pxr::UsdAttribute Attr = GeomCamera.CreateHorizontalApertureAttr() )
+	if (pxr::UsdAttribute Attr = GeomCamera.CreateHorizontalApertureAttr())
 	{
-		Attr.Set<float>( UnrealToUsd::ConvertDistance( StageInfo, CameraComponent.Filmback.SensorWidth ), UsdTimeCode );
+		Attr.Set<float>(UnrealToUsd::ConvertDistance(StageInfo, CameraComponent.Filmback.SensorWidth), UsdTimeCode);
 		UsdUtils::NotifyIfOverriddenOpinion(Attr);
 	}
 
-	if ( pxr::UsdAttribute Attr = GeomCamera.CreateVerticalApertureAttr() )
+	if (pxr::UsdAttribute Attr = GeomCamera.CreateVerticalApertureAttr())
 	{
-		Attr.Set<float>( UnrealToUsd::ConvertDistance( StageInfo, CameraComponent.Filmback.SensorHeight ), UsdTimeCode );
+		Attr.Set<float>(UnrealToUsd::ConvertDistance(StageInfo, CameraComponent.Filmback.SensorHeight), UsdTimeCode);
 		UsdUtils::NotifyIfOverriddenOpinion(Attr);
 	}
 
 	return true;
 }
 
-bool UnrealToUsd::ConvertBoolTrack( const UMovieSceneBoolTrack& MovieSceneTrack, const FMovieSceneSequenceTransform& SequenceTransform, const TFunction<void( bool, double )>& WriterFunc, UE::FUsdPrim& Prim )
+bool UnrealToUsd::ConvertBoolTrack(
+	const UMovieSceneBoolTrack& MovieSceneTrack,
+	const FMovieSceneSequenceTransform& SequenceTransform,
+	const TFunction<void(bool, double)>& WriterFunc,
+	UE::FUsdPrim& Prim
+)
 {
-	if ( !WriterFunc || !Prim )
+	if (!WriterFunc || !Prim)
 	{
 		return false;
 	}
 
 	UMovieScene* MovieScene = MovieSceneTrack.GetTypedOuter<UMovieScene>();
-	if ( !MovieScene )
+	if (!MovieScene)
 	{
 		return false;
 	}
 
 	UE::FUsdStage Stage = Prim.GetStage();
 
-	const TRange< FFrameNumber > PlaybackRange = MovieScene->GetPlaybackRange();
+	const TRange<FFrameNumber> PlaybackRange = MovieScene->GetPlaybackRange();
 	const FFrameRate Resolution = MovieScene->GetTickResolution();
 	const FFrameRate DisplayRate = MovieScene->GetDisplayRate();
 
 	const double StageTimeCodesPerSecond = Stage.GetTimeCodesPerSecond();
-	const FFrameRate StageFrameRate( StageTimeCodesPerSecond, 1 );
+	const FFrameRate StageFrameRate(StageTimeCodesPerSecond, 1);
 
-	auto EvaluateChannel = [&Resolution, &DisplayRate]( const FMovieSceneBoolChannel& Channel, bool InDefaultValue ) -> TArray< TPair< FFrameNumber, bool > >
+	auto EvaluateChannel = [&Resolution,
+							&DisplayRate](const FMovieSceneBoolChannel& Channel, bool InDefaultValue) -> TArray<TPair<FFrameNumber, bool>>
 	{
-		TArray< TPair< FFrameNumber, bool > > Values;
+		TArray<TPair<FFrameNumber, bool>> Values;
 
-		bool DefaultValue = Channel.GetDefault().Get( InDefaultValue );
+		bool DefaultValue = Channel.GetDefault().Get(InDefaultValue);
 
 		TArrayView<const FFrameNumber> KeyTimes = Channel.GetTimes();
 		TArrayView<const bool> KeyValues = Channel.GetValues();
 
-		for ( int32 KeyIndex = 0; KeyIndex < KeyTimes.Num(); ++KeyIndex )
+		for (int32 KeyIndex = 0; KeyIndex < KeyTimes.Num(); ++KeyIndex)
 		{
-			const FFrameNumber KeyTime = KeyTimes[ KeyIndex ];
-			bool KeyValue = KeyValues[ KeyIndex ];
+			const FFrameNumber KeyTime = KeyTimes[KeyIndex];
+			bool KeyValue = KeyValues[KeyIndex];
 
-			FFrameTime SnappedKeyTime{ FFrameRate::Snap( KeyTime, Resolution, DisplayRate ).FloorToFrame() };
+			FFrameTime SnappedKeyTime{FFrameRate::Snap(KeyTime, Resolution, DisplayRate).FloorToFrame()};
 
 			// We never need to bake bool tracks
-			Values.Emplace( SnappedKeyTime.GetFrame(), KeyValue );
+			Values.Emplace(SnappedKeyTime.GetFrame(), KeyValue);
 		}
 
 		return Values;
 	};
 
-	for ( UMovieSceneSection* Section : MovieSceneTrack.GetAllSections() )
+	for (UMovieSceneSection* Section : MovieSceneTrack.GetAllSections())
 	{
-		if ( UMovieSceneBoolSection* BoolSection = Cast<UMovieSceneBoolSection>( Section ) )
+		if (UMovieSceneBoolSection* BoolSection = Cast<UMovieSceneBoolSection>(Section))
 		{
-			for ( const TPair< FFrameNumber, bool >& Pair : EvaluateChannel( BoolSection->GetChannel(), false ) )
+			for (const TPair<FFrameNumber, bool>& Pair : EvaluateChannel(BoolSection->GetChannel(), false))
 			{
-				FFrameTime TransformedBakedKeyTime{ Pair.Key };
+				FFrameTime TransformedBakedKeyTime{Pair.Key};
 				TransformedBakedKeyTime *= SequenceTransform.InverseNoLooping();
-				FFrameTime UsdFrameTime = FFrameRate::TransformTime( TransformedBakedKeyTime, Resolution, StageFrameRate );
+				FFrameTime UsdFrameTime = FFrameRate::TransformTime(TransformedBakedKeyTime, Resolution, StageFrameRate);
 
 				bool UEValue = Pair.Value;
 
-				WriterFunc( UEValue, UsdFrameTime.AsDecimal() );
+				WriterFunc(UEValue, UsdFrameTime.AsDecimal());
 			}
 		}
 	}
@@ -2145,15 +2205,20 @@ bool UnrealToUsd::ConvertBoolTrack( const UMovieSceneBoolTrack& MovieSceneTrack,
 	return true;
 }
 
-bool UnrealToUsd::ConvertFloatTrack( const UMovieSceneFloatTrack& MovieSceneTrack, const FMovieSceneSequenceTransform& SequenceTransform, const TFunction<void( float, double )>& WriterFunc, UE::FUsdPrim& Prim )
+bool UnrealToUsd::ConvertFloatTrack(
+	const UMovieSceneFloatTrack& MovieSceneTrack,
+	const FMovieSceneSequenceTransform& SequenceTransform,
+	const TFunction<void(float, double)>& WriterFunc,
+	UE::FUsdPrim& Prim
+)
 {
-	if ( !WriterFunc || !Prim )
+	if (!WriterFunc || !Prim)
 	{
 		return false;
 	}
 
 	UMovieScene* MovieScene = MovieSceneTrack.GetTypedOuter<UMovieScene>();
-	if ( !MovieScene )
+	if (!MovieScene)
 	{
 		return false;
 	}
@@ -2163,40 +2228,43 @@ bool UnrealToUsd::ConvertFloatTrack( const UMovieSceneFloatTrack& MovieSceneTrac
 	ERichCurveInterpMode StageInterpMode;
 	{
 		FScopedUsdAllocs Allocs;
-		pxr::UsdStageRefPtr UsdStage{ Stage };
-		StageInterpMode = ( UsdStage->GetInterpolationType() == pxr::UsdInterpolationTypeLinear ) ? ERichCurveInterpMode::RCIM_Linear : ERichCurveInterpMode::RCIM_Constant;
+		pxr::UsdStageRefPtr UsdStage{Stage};
+		StageInterpMode = (UsdStage->GetInterpolationType() == pxr::UsdInterpolationTypeLinear) ? ERichCurveInterpMode::RCIM_Linear
+																								: ERichCurveInterpMode::RCIM_Constant;
 	}
 
-	const TRange< FFrameNumber > PlaybackRange = MovieScene->GetPlaybackRange();
+	const TRange<FFrameNumber> PlaybackRange = MovieScene->GetPlaybackRange();
 	const FFrameRate Resolution = MovieScene->GetTickResolution();
 	const FFrameRate DisplayRate = MovieScene->GetDisplayRate();
 
 	const double StageTimeCodesPerSecond = Stage.GetTimeCodesPerSecond();
-	const FFrameRate StageFrameRate( StageTimeCodesPerSecond, 1 );
+	const FFrameRate StageFrameRate(StageTimeCodesPerSecond, 1);
 
-	auto EvaluateChannel = [&Resolution, &DisplayRate, StageInterpMode]( const FMovieSceneFloatChannel& Channel, float InDefaultValue ) -> TArray< TPair< FFrameNumber, float > >
+	auto EvaluateChannel = [&Resolution,
+							&DisplayRate,
+							StageInterpMode](const FMovieSceneFloatChannel& Channel, float InDefaultValue) -> TArray<TPair<FFrameNumber, float>>
 	{
-		TArray< TPair< FFrameNumber, float > > Values;
+		TArray<TPair<FFrameNumber, float>> Values;
 
-		const FFrameTime BakeInterval = FFrameRate::TransformTime( 1, DisplayRate, Resolution );
+		const FFrameTime BakeInterval = FFrameRate::TransformTime(1, DisplayRate, Resolution);
 
-		float DefaultValue = Channel.GetDefault().Get( InDefaultValue );
+		float DefaultValue = Channel.GetDefault().Get(InDefaultValue);
 
 		TMovieSceneChannelData<const FMovieSceneFloatValue> ChannelData = Channel.GetData();
 		TArrayView<const FFrameNumber> KeyTimes = ChannelData.GetTimes();
 		TArrayView<const FMovieSceneFloatValue> KeyValues = ChannelData.GetValues();
 
-		for ( int32 KeyIndex = 0; KeyIndex < KeyTimes.Num(); ++KeyIndex )
+		for (int32 KeyIndex = 0; KeyIndex < KeyTimes.Num(); ++KeyIndex)
 		{
-			const FFrameNumber KeyTime = KeyTimes[ KeyIndex ];
-			const FMovieSceneFloatValue& KeyValue = KeyValues[ KeyIndex ];
+			const FFrameNumber KeyTime = KeyTimes[KeyIndex];
+			const FMovieSceneFloatValue& KeyValue = KeyValues[KeyIndex];
 
 			// If the channel has the same interpolation type as the stage (or we're the last key),
 			// we don't need to bake anything: Just write out the keyframe as is
-			if ( KeyValue.InterpMode == StageInterpMode || KeyIndex == ( KeyTimes.Num() - 1 ) )
+			if (KeyValue.InterpMode == StageInterpMode || KeyIndex == (KeyTimes.Num() - 1))
 			{
-				FFrameTime SnappedKeyTime( FFrameRate::Snap( KeyTime, Resolution, DisplayRate ).FloorToFrame() );
-				Values.Emplace( SnappedKeyTime.GetFrame(), KeyValue.Value );
+				FFrameTime SnappedKeyTime(FFrameRate::Snap(KeyTime, Resolution, DisplayRate).FloorToFrame());
+				Values.Emplace(SnappedKeyTime.GetFrame(), KeyValue.Value);
 			}
 			// We need to bake: Start from this key up until the next key (non-inclusive). We always want to put a keyframe at
 			// KeyTime, but then snap the other ones to the stage framerate
@@ -2204,17 +2272,17 @@ bool UnrealToUsd::ConvertFloatTrack( const UMovieSceneFloatTrack& MovieSceneTrac
 			{
 				// Don't use the snapped key time for the end of the bake range, because if the snapping moves it
 				// later we may end up stepping back again when it's time to bake from that key onwards
-				const FFrameNumber NextKey = KeyTimes[ KeyIndex + 1 ];
-				const FFrameTime NextKeyTime{ NextKey };
+				const FFrameNumber NextKey = KeyTimes[KeyIndex + 1];
+				const FFrameTime NextKeyTime{NextKey};
 
-				for ( FFrameTime EvalTime = KeyTime; EvalTime < NextKeyTime; EvalTime += BakeInterval )
+				for (FFrameTime EvalTime = KeyTime; EvalTime < NextKeyTime; EvalTime += BakeInterval)
 				{
-					FFrameNumber BakedKeyTime = FFrameRate::Snap( EvalTime, Resolution, DisplayRate ).FloorToFrame();
+					FFrameNumber BakedKeyTime = FFrameRate::Snap(EvalTime, Resolution, DisplayRate).FloorToFrame();
 
 					float Value = DefaultValue;
-					Channel.Evaluate( BakedKeyTime, Value );
+					Channel.Evaluate(BakedKeyTime, Value);
 
-					Values.Emplace( BakedKeyTime, Value );
+					Values.Emplace(BakedKeyTime, Value);
 				}
 			}
 		}
@@ -2222,19 +2290,19 @@ bool UnrealToUsd::ConvertFloatTrack( const UMovieSceneFloatTrack& MovieSceneTrac
 		return Values;
 	};
 
-	for ( UMovieSceneSection* Section : MovieSceneTrack.GetAllSections() )
+	for (UMovieSceneSection* Section : MovieSceneTrack.GetAllSections())
 	{
-		if ( UMovieSceneFloatSection* FloatSection = Cast<UMovieSceneFloatSection>( Section ) )
+		if (UMovieSceneFloatSection* FloatSection = Cast<UMovieSceneFloatSection>(Section))
 		{
-			for ( const TPair< FFrameNumber, float >& Pair : EvaluateChannel( FloatSection->GetChannel(), 0.0f ) )
+			for (const TPair<FFrameNumber, float>& Pair : EvaluateChannel(FloatSection->GetChannel(), 0.0f))
 			{
-				FFrameTime TransformedBakedKeyTime{ Pair.Key };
+				FFrameTime TransformedBakedKeyTime{Pair.Key};
 				TransformedBakedKeyTime *= SequenceTransform.InverseNoLooping();
-				FFrameTime UsdFrameTime = FFrameRate::TransformTime( TransformedBakedKeyTime, Resolution, StageFrameRate );
+				FFrameTime UsdFrameTime = FFrameRate::TransformTime(TransformedBakedKeyTime, Resolution, StageFrameRate);
 
 				float UEValue = Pair.Value;
 
-				WriterFunc( UEValue, UsdFrameTime.AsDecimal() );
+				WriterFunc(UEValue, UsdFrameTime.AsDecimal());
 			}
 		}
 	}
@@ -2242,15 +2310,20 @@ bool UnrealToUsd::ConvertFloatTrack( const UMovieSceneFloatTrack& MovieSceneTrac
 	return true;
 }
 
-bool UnrealToUsd::ConvertColorTrack( const UMovieSceneColorTrack& MovieSceneTrack, const FMovieSceneSequenceTransform& SequenceTransform, const TFunction<void( const FLinearColor&, double )>& WriterFunc, UE::FUsdPrim& Prim )
+bool UnrealToUsd::ConvertColorTrack(
+	const UMovieSceneColorTrack& MovieSceneTrack,
+	const FMovieSceneSequenceTransform& SequenceTransform,
+	const TFunction<void(const FLinearColor&, double)>& WriterFunc,
+	UE::FUsdPrim& Prim
+)
 {
-	if ( !WriterFunc || !Prim )
+	if (!WriterFunc || !Prim)
 	{
 		return false;
 	}
 
 	UMovieScene* MovieScene = MovieSceneTrack.GetTypedOuter<UMovieScene>();
-	if ( !MovieScene )
+	if (!MovieScene)
 	{
 		return false;
 	}
@@ -2260,36 +2333,38 @@ bool UnrealToUsd::ConvertColorTrack( const UMovieSceneColorTrack& MovieSceneTrac
 	ERichCurveInterpMode StageInterpMode;
 	{
 		FScopedUsdAllocs Allocs;
-		pxr::UsdStageRefPtr UsdStage{ Stage };
-		StageInterpMode = ( UsdStage->GetInterpolationType() == pxr::UsdInterpolationTypeLinear ) ? ERichCurveInterpMode::RCIM_Linear : ERichCurveInterpMode::RCIM_Constant;
+		pxr::UsdStageRefPtr UsdStage{Stage};
+		StageInterpMode = (UsdStage->GetInterpolationType() == pxr::UsdInterpolationTypeLinear) ? ERichCurveInterpMode::RCIM_Linear
+																								: ERichCurveInterpMode::RCIM_Constant;
 	}
 
-	const TRange< FFrameNumber > PlaybackRange = MovieScene->GetPlaybackRange();
+	const TRange<FFrameNumber> PlaybackRange = MovieScene->GetPlaybackRange();
 	const FFrameRate Resolution = MovieScene->GetTickResolution();
 	const FFrameRate DisplayRate = MovieScene->GetDisplayRate();
 
 	const double StageTimeCodesPerSecond = Stage.GetTimeCodesPerSecond();
-	const FFrameRate StageFrameRate( StageTimeCodesPerSecond, 1 );
+	const FFrameRate StageFrameRate(StageTimeCodesPerSecond, 1);
 
-	auto AppendChannelBakeTimes = [&Resolution, &DisplayRate, StageInterpMode]( const FMovieSceneFloatChannel& Channel, TSet<FFrameNumber>& OutBakeTimes )
+	auto AppendChannelBakeTimes =
+		[&Resolution, &DisplayRate, StageInterpMode](const FMovieSceneFloatChannel& Channel, TSet<FFrameNumber>& OutBakeTimes)
 	{
-		const FFrameTime BakeInterval = FFrameRate::TransformTime( 1, DisplayRate, Resolution );
+		const FFrameTime BakeInterval = FFrameRate::TransformTime(1, DisplayRate, Resolution);
 
 		TMovieSceneChannelData<const FMovieSceneFloatValue> ChannelData = Channel.GetData();
 		TArrayView<const FFrameNumber> KeyTimes = ChannelData.GetTimes();
 		TArrayView<const FMovieSceneFloatValue> KeyValues = ChannelData.GetValues();
 
-		for ( int32 KeyIndex = 0; KeyIndex < KeyTimes.Num(); ++KeyIndex )
+		for (int32 KeyIndex = 0; KeyIndex < KeyTimes.Num(); ++KeyIndex)
 		{
-			const FFrameNumber KeyTime = KeyTimes[ KeyIndex ];
-			const FMovieSceneFloatValue& KeyValue = KeyValues[ KeyIndex ];
+			const FFrameNumber KeyTime = KeyTimes[KeyIndex];
+			const FMovieSceneFloatValue& KeyValue = KeyValues[KeyIndex];
 
 			// If the channel has the same interpolation type as the stage (or we're the last key),
 			// we don't need to bake anything: Just write out the keyframe as is
-			if ( KeyValue.InterpMode == StageInterpMode || KeyIndex == ( KeyTimes.Num() - 1 ) )
+			if (KeyValue.InterpMode == StageInterpMode || KeyIndex == (KeyTimes.Num() - 1))
 			{
-				FFrameNumber SnappedKeyTime = FFrameRate::Snap( KeyTime, Resolution, DisplayRate ).FloorToFrame();
-				OutBakeTimes.Emplace( SnappedKeyTime );
+				FFrameNumber SnappedKeyTime = FFrameRate::Snap(KeyTime, Resolution, DisplayRate).FloorToFrame();
+				OutBakeTimes.Emplace(SnappedKeyTime);
 			}
 			// We need to bake: Start from this key up until the next key (non-inclusive). We always want to put a keyframe at
 			// KeyTime, but then snap the other ones to the stage framerate
@@ -2297,58 +2372,58 @@ bool UnrealToUsd::ConvertColorTrack( const UMovieSceneColorTrack& MovieSceneTrac
 			{
 				// Don't use the snapped key time for the end of the bake range, because if the snapping moves it
 				// later we may end up stepping back again when it's time to bake from that key onwards
-				const FFrameNumber NextKey = KeyTimes[ KeyIndex + 1 ];
-				const FFrameTime NextKeyTime{ NextKey };
+				const FFrameNumber NextKey = KeyTimes[KeyIndex + 1];
+				const FFrameTime NextKeyTime{NextKey};
 
-				for ( FFrameTime EvalTime = KeyTime; EvalTime < NextKeyTime; EvalTime += BakeInterval )
+				for (FFrameTime EvalTime = KeyTime; EvalTime < NextKeyTime; EvalTime += BakeInterval)
 				{
-					FFrameNumber BakedKeyTime = FFrameRate::Snap( EvalTime, Resolution, DisplayRate ).FloorToFrame();
-					OutBakeTimes.Emplace( BakedKeyTime );
+					FFrameNumber BakedKeyTime = FFrameRate::Snap(EvalTime, Resolution, DisplayRate).FloorToFrame();
+					OutBakeTimes.Emplace(BakedKeyTime);
 				}
 			}
 		}
 	};
 
-	for ( UMovieSceneSection* Section : MovieSceneTrack.GetAllSections() )
+	for (UMovieSceneSection* Section : MovieSceneTrack.GetAllSections())
 	{
-		if ( UMovieSceneColorSection* ColorSection = Cast<UMovieSceneColorSection>( Section ) )
+		if (UMovieSceneColorSection* ColorSection = Cast<UMovieSceneColorSection>(Section))
 		{
-			const FMovieSceneFloatChannel& RedChannel   = ColorSection->GetRedChannel();
+			const FMovieSceneFloatChannel& RedChannel = ColorSection->GetRedChannel();
 			const FMovieSceneFloatChannel& GreenChannel = ColorSection->GetGreenChannel();
-			const FMovieSceneFloatChannel& BlueChannel  = ColorSection->GetBlueChannel();
+			const FMovieSceneFloatChannel& BlueChannel = ColorSection->GetBlueChannel();
 			const FMovieSceneFloatChannel& AlphaChannel = ColorSection->GetAlphaChannel();
 
 			// Get the baked FFrameNumbers for each channel (without evaluating the channels yet),
 			// because they may have independent keys
 			TSet<FFrameNumber> ChannelBakeTimes;
-			AppendChannelBakeTimes( RedChannel, ChannelBakeTimes );
-			AppendChannelBakeTimes( GreenChannel, ChannelBakeTimes );
-			AppendChannelBakeTimes( BlueChannel, ChannelBakeTimes );
-			AppendChannelBakeTimes( AlphaChannel, ChannelBakeTimes );
+			AppendChannelBakeTimes(RedChannel, ChannelBakeTimes);
+			AppendChannelBakeTimes(GreenChannel, ChannelBakeTimes);
+			AppendChannelBakeTimes(BlueChannel, ChannelBakeTimes);
+			AppendChannelBakeTimes(AlphaChannel, ChannelBakeTimes);
 
 			TArray<FFrameNumber> BakeTimeUnion = ChannelBakeTimes.Array();
 			BakeTimeUnion.Sort();
 
 			// Sample all channels at the union of bake times, construct the value and write it out
-			for ( const FFrameNumber UntransformedBakeTime : BakeTimeUnion )
+			for (const FFrameNumber UntransformedBakeTime : BakeTimeUnion)
 			{
 				float RedValue = 0.0f;
 				float GreenValue = 0.0f;
 				float BlueValue = 0.0f;
 				float AlphaValue = 1.0f;
 
-				RedChannel.Evaluate( UntransformedBakeTime, RedValue );
-				GreenChannel.Evaluate( UntransformedBakeTime, GreenValue );
-				BlueChannel.Evaluate( UntransformedBakeTime, BlueValue );
-				AlphaChannel.Evaluate( UntransformedBakeTime, AlphaValue );
+				RedChannel.Evaluate(UntransformedBakeTime, RedValue);
+				GreenChannel.Evaluate(UntransformedBakeTime, GreenValue);
+				BlueChannel.Evaluate(UntransformedBakeTime, BlueValue);
+				AlphaChannel.Evaluate(UntransformedBakeTime, AlphaValue);
 
-				FLinearColor Color{ RedValue, GreenValue, BlueValue, AlphaValue };
+				FLinearColor Color{RedValue, GreenValue, BlueValue, AlphaValue};
 
-				FFrameTime TransformedBakedKeyTime{ UntransformedBakeTime };
+				FFrameTime TransformedBakedKeyTime{UntransformedBakeTime};
 				TransformedBakedKeyTime *= SequenceTransform.InverseNoLooping();
-				FFrameTime UsdFrameTime = FFrameRate::TransformTime( TransformedBakedKeyTime, Resolution, StageFrameRate );
+				FFrameTime UsdFrameTime = FFrameRate::TransformTime(TransformedBakedKeyTime, Resolution, StageFrameRate);
 
-				WriterFunc( Color, UsdFrameTime.AsDecimal() );
+				WriterFunc(Color, UsdFrameTime.AsDecimal());
 			}
 		}
 	}
@@ -2388,9 +2463,8 @@ bool UnrealToUsd::ConvertBoundsVectorTracks(
 	{
 		FScopedUsdAllocs Allocs;
 		pxr::UsdStageRefPtr UsdStage{Stage};
-		StageInterpMode = (UsdStage->GetInterpolationType() == pxr::UsdInterpolationTypeLinear)
-							  ? ERichCurveInterpMode::RCIM_Linear
-							  : ERichCurveInterpMode::RCIM_Constant;
+		StageInterpMode = (UsdStage->GetInterpolationType() == pxr::UsdInterpolationTypeLinear) ? ERichCurveInterpMode::RCIM_Linear
+																								: ERichCurveInterpMode::RCIM_Constant;
 	}
 
 	const FFrameRate Resolution = MovieScene->GetTickResolution();
@@ -2516,15 +2590,20 @@ bool UnrealToUsd::ConvertBoundsVectorTracks(
 	return true;
 }
 
-bool UnrealToUsd::Convert3DTransformTrack( const UMovieScene3DTransformTrack& MovieSceneTrack, const FMovieSceneSequenceTransform& SequenceTransform, const TFunction<void( const FTransform&, double )>& WriterFunc, UE::FUsdPrim& Prim )
+bool UnrealToUsd::Convert3DTransformTrack(
+	const UMovieScene3DTransformTrack& MovieSceneTrack,
+	const FMovieSceneSequenceTransform& SequenceTransform,
+	const TFunction<void(const FTransform&, double)>& WriterFunc,
+	UE::FUsdPrim& Prim
+)
 {
-	if ( !WriterFunc || !Prim )
+	if (!WriterFunc || !Prim)
 	{
 		return false;
 	}
 
 	UMovieScene* MovieScene = MovieSceneTrack.GetTypedOuter<UMovieScene>();
-	if ( !MovieScene )
+	if (!MovieScene)
 	{
 		return false;
 	}
@@ -2534,43 +2613,44 @@ bool UnrealToUsd::Convert3DTransformTrack( const UMovieScene3DTransformTrack& Mo
 	ERichCurveInterpMode StageInterpMode;
 	{
 		FScopedUsdAllocs Allocs;
-		pxr::UsdStageRefPtr UsdStage{ Stage };
-		StageInterpMode = ( UsdStage->GetInterpolationType() == pxr::UsdInterpolationTypeLinear ) ? ERichCurveInterpMode::RCIM_Linear : ERichCurveInterpMode::RCIM_Constant;
+		pxr::UsdStageRefPtr UsdStage{Stage};
+		StageInterpMode = (UsdStage->GetInterpolationType() == pxr::UsdInterpolationTypeLinear) ? ERichCurveInterpMode::RCIM_Linear
+																								: ERichCurveInterpMode::RCIM_Constant;
 	}
 
-	const TRange< FFrameNumber > PlaybackRange = MovieScene->GetPlaybackRange();
+	const TRange<FFrameNumber> PlaybackRange = MovieScene->GetPlaybackRange();
 	const FFrameRate Resolution = MovieScene->GetTickResolution();
 	const FFrameRate DisplayRate = MovieScene->GetDisplayRate();
 
 	const double StageTimeCodesPerSecond = Stage.GetTimeCodesPerSecond();
-	const FFrameRate StageFrameRate( StageTimeCodesPerSecond, 1 );
+	const FFrameRate StageFrameRate(StageTimeCodesPerSecond, 1);
 
-	auto EvaluateChannelTimes = [&Resolution, &DisplayRate, StageInterpMode]( const FMovieSceneDoubleChannel* Channel ) -> TSet<FFrameNumber>
+	auto EvaluateChannelTimes = [&Resolution, &DisplayRate, StageInterpMode](const FMovieSceneDoubleChannel* Channel) -> TSet<FFrameNumber>
 	{
 		TSet<FFrameNumber> BakeTimes;
 
-		if ( !Channel )
+		if (!Channel)
 		{
 			return BakeTimes;
 		}
 
-		const FFrameTime BakeInterval = FFrameRate::TransformTime( 1, DisplayRate, Resolution );
+		const FFrameTime BakeInterval = FFrameRate::TransformTime(1, DisplayRate, Resolution);
 
 		TMovieSceneChannelData<const FMovieSceneDoubleValue> ChannelData = Channel->GetData();
 		TArrayView<const FFrameNumber> KeyTimes = ChannelData.GetTimes();
 		TArrayView<const FMovieSceneDoubleValue> KeyValues = ChannelData.GetValues();
 
-		for ( int32 KeyIndex = 0; KeyIndex < KeyTimes.Num(); ++KeyIndex )
+		for (int32 KeyIndex = 0; KeyIndex < KeyTimes.Num(); ++KeyIndex)
 		{
-			const FFrameNumber KeyTime = KeyTimes[ KeyIndex ];
-			const FMovieSceneDoubleValue& KeyValue = KeyValues[ KeyIndex ];
+			const FFrameNumber KeyTime = KeyTimes[KeyIndex];
+			const FMovieSceneDoubleValue& KeyValue = KeyValues[KeyIndex];
 
 			// If the channel has the same interpolation type as the stage (or we're the last key),
 			// we don't need to bake anything: Just write out the keyframe as is
-			if ( KeyValue.InterpMode == StageInterpMode || KeyIndex == ( KeyTimes.Num() - 1 ) )
+			if (KeyValue.InterpMode == StageInterpMode || KeyIndex == (KeyTimes.Num() - 1))
 			{
-				FFrameNumber SnappedKeyTime = FFrameRate::Snap( KeyTime, Resolution, DisplayRate ).FloorToFrame();
-				BakeTimes.Emplace( SnappedKeyTime );
+				FFrameNumber SnappedKeyTime = FFrameRate::Snap(KeyTime, Resolution, DisplayRate).FloorToFrame();
+				BakeTimes.Emplace(SnappedKeyTime);
 			}
 			// We need to bake: Start from this key up until the next key (non-inclusive). We always want to put a keyframe at
 			// KeyTime, but then snap the other ones to the stage framerate
@@ -2578,13 +2658,13 @@ bool UnrealToUsd::Convert3DTransformTrack( const UMovieScene3DTransformTrack& Mo
 			{
 				// Don't use the snapped key time for the end of the bake range, because if the snapping moves it
 				// later we may end up stepping back again when it's time to bake from that key onwards
-				const FFrameNumber NextKey = KeyTimes[ KeyIndex + 1 ];
-				const FFrameTime NextKeyTime{ NextKey };
+				const FFrameNumber NextKey = KeyTimes[KeyIndex + 1];
+				const FFrameTime NextKeyTime{NextKey};
 
-				for ( FFrameTime EvalTime = KeyTime; EvalTime < NextKeyTime; EvalTime += BakeInterval )
+				for (FFrameTime EvalTime = KeyTime; EvalTime < NextKeyTime; EvalTime += BakeInterval)
 				{
-					FFrameNumber BakedKeyTime = FFrameRate::Snap( EvalTime, Resolution, DisplayRate ).FloorToFrame();
-					BakeTimes.Emplace( BakedKeyTime );
+					FFrameNumber BakedKeyTime = FFrameRate::Snap(EvalTime, Resolution, DisplayRate).FloorToFrame();
+					BakeTimes.Emplace(BakedKeyTime);
 				}
 			}
 		}
@@ -2592,57 +2672,63 @@ bool UnrealToUsd::Convert3DTransformTrack( const UMovieScene3DTransformTrack& Mo
 		return BakeTimes;
 	};
 
-	for ( UMovieSceneSection* Section : MovieSceneTrack.GetAllSections() )
+	for (UMovieSceneSection* Section : MovieSceneTrack.GetAllSections())
 	{
-		if ( UMovieScene3DTransformSection* TransformSection = Cast<UMovieScene3DTransformSection>( Section ) )
+		if (UMovieScene3DTransformSection* TransformSection = Cast<UMovieScene3DTransformSection>(Section))
 		{
-			TArrayView< FMovieSceneDoubleChannel* > Channels = TransformSection->GetChannelProxy().GetChannels< FMovieSceneDoubleChannel >();
-			if ( Channels.Num() < 9 )
+			TArrayView<FMovieSceneDoubleChannel*> Channels = TransformSection->GetChannelProxy().GetChannels<FMovieSceneDoubleChannel>();
+			if (Channels.Num() < 9)
 			{
-				UE_LOG( LogUsd, Error, TEXT( "Unexpected number of double tracks (%d) in transform section '%s'" ), Channels.Num(), *TransformSection->GetPathName() );
+				UE_LOG(
+					LogUsd,
+					Error,
+					TEXT("Unexpected number of double tracks (%d) in transform section '%s'"),
+					Channels.Num(),
+					*TransformSection->GetPathName()
+				);
 				continue;
 			}
 
-			FMovieSceneDoubleChannel* LocationXChannel = Channels[ 0 ];
-			FMovieSceneDoubleChannel* LocationYChannel = Channels[ 1 ];
-			FMovieSceneDoubleChannel* LocationZChannel = Channels[ 2 ];
+			FMovieSceneDoubleChannel* LocationXChannel = Channels[0];
+			FMovieSceneDoubleChannel* LocationYChannel = Channels[1];
+			FMovieSceneDoubleChannel* LocationZChannel = Channels[2];
 
-			FMovieSceneDoubleChannel* RotationXChannel = Channels[ 3 ];
-			FMovieSceneDoubleChannel* RotationYChannel = Channels[ 4 ];
-			FMovieSceneDoubleChannel* RotationZChannel = Channels[ 5 ];
+			FMovieSceneDoubleChannel* RotationXChannel = Channels[3];
+			FMovieSceneDoubleChannel* RotationYChannel = Channels[4];
+			FMovieSceneDoubleChannel* RotationZChannel = Channels[5];
 
-			FMovieSceneDoubleChannel* ScaleXChannel = Channels[ 6 ];
-			FMovieSceneDoubleChannel* ScaleYChannel = Channels[ 7 ];
-			FMovieSceneDoubleChannel* ScaleZChannel = Channels[ 8 ];
+			FMovieSceneDoubleChannel* ScaleXChannel = Channels[6];
+			FMovieSceneDoubleChannel* ScaleYChannel = Channels[7];
+			FMovieSceneDoubleChannel* ScaleZChannel = Channels[8];
 
-			TSet< FFrameNumber > LocationValuesX = EvaluateChannelTimes( LocationXChannel );
-			TSet< FFrameNumber > LocationValuesY = EvaluateChannelTimes( LocationYChannel );
-			TSet< FFrameNumber > LocationValuesZ = EvaluateChannelTimes( LocationZChannel );
+			TSet<FFrameNumber> LocationValuesX = EvaluateChannelTimes(LocationXChannel);
+			TSet<FFrameNumber> LocationValuesY = EvaluateChannelTimes(LocationYChannel);
+			TSet<FFrameNumber> LocationValuesZ = EvaluateChannelTimes(LocationZChannel);
 
-			TSet< FFrameNumber > RotationValuesX = EvaluateChannelTimes( RotationXChannel );
-			TSet< FFrameNumber > RotationValuesY = EvaluateChannelTimes( RotationYChannel );
-			TSet< FFrameNumber > RotationValuesZ = EvaluateChannelTimes( RotationZChannel );
+			TSet<FFrameNumber> RotationValuesX = EvaluateChannelTimes(RotationXChannel);
+			TSet<FFrameNumber> RotationValuesY = EvaluateChannelTimes(RotationYChannel);
+			TSet<FFrameNumber> RotationValuesZ = EvaluateChannelTimes(RotationZChannel);
 
-			TSet< FFrameNumber > ScaleValuesX = EvaluateChannelTimes( ScaleXChannel );
-			TSet< FFrameNumber > ScaleValuesY = EvaluateChannelTimes( ScaleYChannel );
-			TSet< FFrameNumber > ScaleValuesZ = EvaluateChannelTimes( ScaleZChannel );
+			TSet<FFrameNumber> ScaleValuesX = EvaluateChannelTimes(ScaleXChannel);
+			TSet<FFrameNumber> ScaleValuesY = EvaluateChannelTimes(ScaleYChannel);
+			TSet<FFrameNumber> ScaleValuesZ = EvaluateChannelTimes(ScaleZChannel);
 
-			LocationValuesX.Append( LocationValuesY );
-			LocationValuesX.Append( LocationValuesZ );
+			LocationValuesX.Append(LocationValuesY);
+			LocationValuesX.Append(LocationValuesZ);
 
-			LocationValuesX.Append( RotationValuesX );
-			LocationValuesX.Append( RotationValuesY );
-			LocationValuesX.Append( RotationValuesZ );
+			LocationValuesX.Append(RotationValuesX);
+			LocationValuesX.Append(RotationValuesY);
+			LocationValuesX.Append(RotationValuesZ);
 
-			LocationValuesX.Append( ScaleValuesX );
-			LocationValuesX.Append( ScaleValuesY );
-			LocationValuesX.Append( ScaleValuesZ );
+			LocationValuesX.Append(ScaleValuesX);
+			LocationValuesX.Append(ScaleValuesY);
+			LocationValuesX.Append(ScaleValuesZ);
 
 			TArray<FFrameNumber> BakeTimeUnion = LocationValuesX.Array();
 			BakeTimeUnion.Sort();
 
 			// Sample all channels at the union of bake times, construct the value and write it out
-			for ( const FFrameNumber UntransformedBakeTime : BakeTimeUnion )
+			for (const FFrameNumber UntransformedBakeTime : BakeTimeUnion)
 			{
 				double LocX = 0.0f;
 				double LocY = 0.0f;
@@ -2656,58 +2742,58 @@ bool UnrealToUsd::Convert3DTransformTrack( const UMovieScene3DTransformTrack& Mo
 				double ScaleY = 1.0f;
 				double ScaleZ = 1.0f;
 
-				if ( LocationXChannel )
+				if (LocationXChannel)
 				{
-					LocationXChannel->Evaluate( UntransformedBakeTime, LocX );
+					LocationXChannel->Evaluate(UntransformedBakeTime, LocX);
 				}
-				if ( LocationYChannel )
+				if (LocationYChannel)
 				{
-					LocationYChannel->Evaluate( UntransformedBakeTime, LocY );
+					LocationYChannel->Evaluate(UntransformedBakeTime, LocY);
 				}
-				if ( LocationZChannel )
+				if (LocationZChannel)
 				{
-					LocationZChannel->Evaluate( UntransformedBakeTime, LocZ );
+					LocationZChannel->Evaluate(UntransformedBakeTime, LocZ);
 				}
 
-				if ( RotationXChannel )
+				if (RotationXChannel)
 				{
-					RotationXChannel->Evaluate( UntransformedBakeTime, RotX );
+					RotationXChannel->Evaluate(UntransformedBakeTime, RotX);
 				}
-				if ( RotationYChannel )
+				if (RotationYChannel)
 				{
-					RotationYChannel->Evaluate( UntransformedBakeTime, RotY );
+					RotationYChannel->Evaluate(UntransformedBakeTime, RotY);
 				}
-				if ( RotationZChannel )
+				if (RotationZChannel)
 				{
-					RotationZChannel->Evaluate( UntransformedBakeTime, RotZ );
+					RotationZChannel->Evaluate(UntransformedBakeTime, RotZ);
 				}
 
-				if ( ScaleXChannel )
+				if (ScaleXChannel)
 				{
-					ScaleXChannel->Evaluate( UntransformedBakeTime, ScaleX );
+					ScaleXChannel->Evaluate(UntransformedBakeTime, ScaleX);
 				}
-				if ( ScaleYChannel )
+				if (ScaleYChannel)
 				{
-					ScaleYChannel->Evaluate( UntransformedBakeTime, ScaleY );
+					ScaleYChannel->Evaluate(UntransformedBakeTime, ScaleY);
 				}
-				if ( ScaleZChannel )
+				if (ScaleZChannel)
 				{
-					ScaleZChannel->Evaluate( UntransformedBakeTime, ScaleZ );
+					ScaleZChannel->Evaluate(UntransformedBakeTime, ScaleZ);
 				}
 
 				// Casting this to float right now because depending on the build and the LWC status FVectors contain FLargeWorldCoordinatesReal,
 				// which can be floats and turn these into narrowing conversions, which require explicit casts.
 				// TODO: Replace these casts with the underlying FVector type later
-				FVector Location{ ( float ) LocX, ( float ) LocY, ( float ) LocZ };
-				FRotator Rotation{ ( float ) RotY, ( float ) RotZ, ( float ) RotX };
-				FVector Scale{ ( float ) ScaleX, ( float ) ScaleY, ( float ) ScaleZ };
-				FTransform Transform{ Rotation, Location, Scale };
+				FVector Location{(float)LocX, (float)LocY, (float)LocZ};
+				FRotator Rotation{(float)RotY, (float)RotZ, (float)RotX};
+				FVector Scale{(float)ScaleX, (float)ScaleY, (float)ScaleZ};
+				FTransform Transform{Rotation, Location, Scale};
 
-				FFrameTime TransformedBakedKeyTime{ UntransformedBakeTime };
+				FFrameTime TransformedBakedKeyTime{UntransformedBakeTime};
 				TransformedBakedKeyTime *= SequenceTransform.InverseNoLooping();
-				FFrameTime UsdFrameTime = FFrameRate::TransformTime( TransformedBakedKeyTime, Resolution, StageFrameRate );
+				FFrameTime UsdFrameTime = FFrameRate::TransformTime(TransformedBakedKeyTime, Resolution, StageFrameRate);
 
-				WriterFunc( Transform, UsdFrameTime.AsDecimal() );
+				WriterFunc(Transform, UsdFrameTime.AsDecimal());
 			}
 		}
 	}
@@ -2715,9 +2801,9 @@ bool UnrealToUsd::Convert3DTransformTrack( const UMovieScene3DTransformTrack& Mo
 	return true;
 }
 
-bool UnrealToUsd::ConvertSceneComponent( const pxr::UsdStageRefPtr& Stage, const USceneComponent* SceneComponent, pxr::UsdPrim& UsdPrim )
+bool UnrealToUsd::ConvertSceneComponent(const pxr::UsdStageRefPtr& Stage, const USceneComponent* SceneComponent, pxr::UsdPrim& UsdPrim)
 {
-	if ( !UsdPrim || !SceneComponent )
+	if (!UsdPrim || !SceneComponent)
 	{
 		return false;
 	}
@@ -2730,8 +2816,8 @@ bool UnrealToUsd::ConvertSceneComponent( const pxr::UsdStageRefPtr& Stage, const
 	FScopedUsdAllocs UsdAllocs;
 
 	// Transform
-	pxr::UsdGeomXformable XForm( UsdPrim );
-	if ( !XForm )
+	pxr::UsdGeomXformable XForm(UsdPrim);
+	if (!XForm)
 	{
 		return false;
 	}
@@ -2758,28 +2844,28 @@ bool UnrealToUsd::ConvertSceneComponent( const pxr::UsdStageRefPtr& Stage, const
 	// In UE cameras shoot towards local + X, with + Z up.Lights also emit towards local + X, with + Z up
 	// Note that this wouldn't have worked in case we collapsed light and camera components, but these always get their own
 	// actors, so we know that we don't have a single component that represents a large collapsed prim hierarchy
-	if ( UsdPrim.IsA<pxr::UsdGeomCamera>() || UsdPrim.HasAPI<pxr::UsdLuxLightAPI>() )
+	if (UsdPrim.IsA<pxr::UsdGeomCamera>() || UsdPrim.HasAPI<pxr::UsdLuxLightAPI>())
 	{
-		FTransform AdditionalRotation = FTransform( FRotator( 0.0f, 90.f, 0.0f ) );
+		FTransform AdditionalRotation = FTransform(FRotator(0.0f, 90.f, 0.0f));
 
-		if ( UsdUtils::GetUsdStageUpAxis( Stage ) == pxr::UsdGeomTokens->z )
+		if (UsdUtils::GetUsdStageUpAxis(Stage) == pxr::UsdGeomTokens->z)
 		{
-			AdditionalRotation *= FTransform( FRotator( 90.0f, 0.f, 0.0f ) );
+			AdditionalRotation *= FTransform(FRotator(90.0f, 0.f, 0.0f));
 		}
 
 		RelativeTransform = AdditionalRotation * RelativeTransform;
 	}
 
 	// Invert compensation applied to parent if it's a light or camera component
-	if ( pxr::UsdPrim ParentPrim = UsdPrim.GetParent() )
+	if (pxr::UsdPrim ParentPrim = UsdPrim.GetParent())
 	{
-		if ( ParentPrim.IsA<pxr::UsdGeomCamera>() || ParentPrim.HasAPI<pxr::UsdLuxLightAPI>() )
+		if (ParentPrim.IsA<pxr::UsdGeomCamera>() || ParentPrim.HasAPI<pxr::UsdLuxLightAPI>())
 		{
-			FTransform AdditionalRotation = FTransform( FRotator( 0.0f, 90.f, 0.0f ) );
+			FTransform AdditionalRotation = FTransform(FRotator(0.0f, 90.f, 0.0f));
 
-			if ( UsdUtils::GetUsdStageUpAxis( Stage ) == pxr::UsdGeomTokens->z )
+			if (UsdUtils::GetUsdStageUpAxis(Stage) == pxr::UsdGeomTokens->z)
 			{
-				AdditionalRotation *= FTransform( FRotator( 90.0f, 0.f, 0.0f ) );
+				AdditionalRotation *= FTransform(FRotator(90.0f, 0.f, 0.0f));
 			}
 
 			RelativeTransform = RelativeTransform * AdditionalRotation.Inverse();
@@ -2787,28 +2873,27 @@ bool UnrealToUsd::ConvertSceneComponent( const pxr::UsdStageRefPtr& Stage, const
 	}
 
 	// Transform
-	ConvertXformable( RelativeTransform, UsdPrim, UsdUtils::GetDefaultTimeCode() );
+	ConvertXformable(RelativeTransform, UsdPrim, UsdUtils::GetDefaultTimeCode());
 
 	// Per-prim visibility
-	if ( pxr::UsdAttribute VisibilityAttr = XForm.CreateVisibilityAttr() )
+	if (pxr::UsdAttribute VisibilityAttr = XForm.CreateVisibilityAttr())
 	{
 		pxr::TfToken Value = pxr::UsdGeomTokens->inherited;
 
-		if ( SceneComponent->ComponentTags.Contains( UnrealIdentifiers::Invisible ) )
+		if (SceneComponent->ComponentTags.Contains(UnrealIdentifiers::Invisible))
 		{
 			Value = pxr::UsdGeomTokens->invisible;
 		}
-		else if ( !SceneComponent->ComponentTags.Contains( UnrealIdentifiers::Inherited ) )
+		else if (!SceneComponent->ComponentTags.Contains(UnrealIdentifiers::Inherited))
 		{
 			// We don't have visible nor inherited tags: We're probably exporting a pure UE component, so write out component visibility instead.
 			// Ignore invisibility from brush components though because they are always forced to bHiddenInGame=true, with the property even
 			// being hidden on the details panel
-			Value = SceneComponent->bHiddenInGame && !SceneComponent->IsA<UBrushComponent>()
-						? pxr::UsdGeomTokens->invisible
-						: pxr::UsdGeomTokens->inherited;
+			Value = SceneComponent->bHiddenInGame && !SceneComponent->IsA<UBrushComponent>() ? pxr::UsdGeomTokens->invisible
+																							 : pxr::UsdGeomTokens->inherited;
 		}
 
-		VisibilityAttr.Set<pxr::TfToken>( Value );
+		VisibilityAttr.Set<pxr::TfToken>(Value);
 		UsdUtils::NotifyIfOverriddenOpinion(VisibilityAttr);
 	}
 
@@ -2857,8 +2942,8 @@ bool UnrealToUsd::ConvertHierarchicalInstancedStaticMeshComponent(
 
 	FScopedUsdAllocs Allocs;
 
-	UsdGeomPointInstancer PointInstancer{ UsdPrim };
-	if ( !PointInstancer || !HISMComponent )
+	UsdGeomPointInstancer PointInstancer{UsdPrim};
+	if (!PointInstancer || !HISMComponent)
 	{
 		return false;
 	}
@@ -2869,7 +2954,7 @@ bool UnrealToUsd::ConvertHierarchicalInstancedStaticMeshComponent(
 	}
 
 	UsdStageRefPtr Stage = UsdPrim.GetStage();
-	FUsdStageInfo StageInfo{ Stage };
+	FUsdStageInfo StageInfo{Stage};
 
 	VtArray<int> ProtoIndices;
 	VtArray<GfVec3f> Positions;
@@ -2877,16 +2962,16 @@ bool UnrealToUsd::ConvertHierarchicalInstancedStaticMeshComponent(
 	VtArray<GfVec3f> Scales;
 
 	const int32 NumInstances = HISMComponent->GetInstanceCount();
-	ProtoIndices.reserve( ProtoIndices.size() + NumInstances );
-	Positions.reserve( Positions.size() + NumInstances );
-	Orientations.reserve( Orientations.size() + NumInstances );
-	Scales.reserve( Scales.size() + NumInstances );
+	ProtoIndices.reserve(ProtoIndices.size() + NumInstances);
+	Positions.reserve(Positions.size() + NumInstances);
+	Orientations.reserve(Orientations.size() + NumInstances);
+	Scales.reserve(Scales.size() + NumInstances);
 
-	for( const FInstancedStaticMeshInstanceData& InstanceData : HISMComponent->PerInstanceSMData )
+	for (const FInstancedStaticMeshInstanceData& InstanceData : HISMComponent->PerInstanceSMData)
 	{
 		// Convert axes
-		FTransform UETransform{ InstanceData.Transform };
-		FTransform USDTransform = UsdUtils::ConvertAxes( StageInfo.UpAxis == EUsdUpAxis::ZAxis, UETransform );
+		FTransform UETransform{InstanceData.Transform};
+		FTransform USDTransform = UsdUtils::ConvertAxes(StageInfo.UpAxis == EUsdUpAxis::ZAxis, UETransform);
 
 		FVector Translation = USDTransform.GetTranslation();
 		FQuat Rotation = USDTransform.GetRotation();
@@ -2894,40 +2979,41 @@ bool UnrealToUsd::ConvertHierarchicalInstancedStaticMeshComponent(
 
 		// Compensate metersPerUnit
 		const float UEMetersPerUnit = 0.01f;
-		if ( !FMath::IsNearlyEqual( UEMetersPerUnit, StageInfo.MetersPerUnit ) )
+		if (!FMath::IsNearlyEqual(UEMetersPerUnit, StageInfo.MetersPerUnit))
 		{
-			Translation *= ( UEMetersPerUnit / StageInfo.MetersPerUnit );
+			Translation *= (UEMetersPerUnit / StageInfo.MetersPerUnit);
 		}
 
-		ProtoIndices.push_back( 0 ); // We will always export a single prototype per PointInstancer, since HISM components handle only 1 mesh at a time
-		Positions.push_back( GfVec3f( Translation.X, Translation.Y, Translation.Z ) );
-		Orientations.push_back( GfQuath( Rotation.W, Rotation.X, Rotation.Y, Rotation.Z ) );
-		Scales.push_back( GfVec3f( Scale.X, Scale.Y, Scale.Z ) );
+		ProtoIndices.push_back(0);	  // We will always export a single prototype per PointInstancer, since HISM components handle only 1 mesh at a
+									  // time
+		Positions.push_back(GfVec3f(Translation.X, Translation.Y, Translation.Z));
+		Orientations.push_back(GfQuath(Rotation.W, Rotation.X, Rotation.Y, Rotation.Z));
+		Scales.push_back(GfVec3f(Scale.X, Scale.Y, Scale.Z));
 	}
 
-	const pxr::UsdTimeCode UsdTimeCode( TimeCode );
+	const pxr::UsdTimeCode UsdTimeCode(TimeCode);
 
-	if ( UsdAttribute Attr = PointInstancer.CreateProtoIndicesAttr() )
+	if (UsdAttribute Attr = PointInstancer.CreateProtoIndicesAttr())
 	{
-		Attr.Set( ProtoIndices, UsdTimeCode );
+		Attr.Set(ProtoIndices, UsdTimeCode);
 		UsdUtils::NotifyIfOverriddenOpinion(Attr);
 	}
 
-	if ( UsdAttribute Attr = PointInstancer.CreatePositionsAttr() )
+	if (UsdAttribute Attr = PointInstancer.CreatePositionsAttr())
 	{
-		Attr.Set( Positions, UsdTimeCode );
+		Attr.Set(Positions, UsdTimeCode);
 		UsdUtils::NotifyIfOverriddenOpinion(Attr);
 	}
 
-	if ( UsdAttribute Attr = PointInstancer.CreateOrientationsAttr() )
+	if (UsdAttribute Attr = PointInstancer.CreateOrientationsAttr())
 	{
-		Attr.Set( Orientations, UsdTimeCode );
+		Attr.Set(Orientations, UsdTimeCode);
 		UsdUtils::NotifyIfOverriddenOpinion(Attr);
 	}
 
-	if ( UsdAttribute Attr = PointInstancer.CreateScalesAttr() )
+	if (UsdAttribute Attr = PointInstancer.CreateScalesAttr())
 	{
-		Attr.Set( Scales, UsdTimeCode );
+		Attr.Set(Scales, UsdTimeCode);
 		UsdUtils::NotifyIfOverriddenOpinion(Attr);
 	}
 
@@ -3236,7 +3322,8 @@ bool UnrealToUsd::ConvertMaterialOverrides(
 					}
 					// TODO: We really need a separate function for ConvertingMaterialOverrides (to an opened stage) and ExportingMaterialOverrides
 					// that we can use when the SkeletalMesh is not something we generated ourselves (with annotated MaterialSlotToPrimPaths).
-					// We could collect an analogue for MaterialSlotToPrimPaths during the mesh export process to accurately author these overrides too
+					// We could collect an analogue for MaterialSlotToPrimPaths during the mesh export process to accurately author these overrides
+					// too
 					else
 					{
 						pxr::SdfPath OverridePrimPath;
@@ -3246,7 +3333,8 @@ bool UnrealToUsd::ConvertMaterialOverrides(
 						// Mesh prims for each LOD named "LOD0", "LOD1", etc., switched via a "LOD" variant set
 						if (bHasLODs)
 						{
-							OverridePrimPath = SkelRoot.GetPath().AppendPath(UnrealToUsd::ConvertPath(*FString::Printf(TEXT("LOD%d"), LODIndex)).Get());
+							OverridePrimPath = SkelRoot.GetPath().AppendPath(UnrealToUsd::ConvertPath(*FString::Printf(TEXT("LOD%d"), LODIndex)).Get()
+							);
 						}
 						else
 						{
@@ -3295,14 +3383,14 @@ bool UnrealToUsd::ConvertMaterialOverrides(
 		ensure(false);
 		return false;
 	}
-#endif // WITH_EDITOR
+#endif	  // WITH_EDITOR
 
 	return true;
 }
 
-bool UnrealToUsd::ConvertXformable( const FTransform& RelativeTransform, pxr::UsdPrim& UsdPrim, double TimeCode )
+bool UnrealToUsd::ConvertXformable(const FTransform& RelativeTransform, pxr::UsdPrim& UsdPrim, double TimeCode)
 {
-	if ( !UsdPrim )
+	if (!UsdPrim)
 	{
 		return false;
 	}
@@ -3310,8 +3398,8 @@ bool UnrealToUsd::ConvertXformable( const FTransform& RelativeTransform, pxr::Us
 	FScopedUsdAllocs UsdAllocs;
 
 	// Transform
-	pxr::UsdGeomXformable XForm( UsdPrim );
-	if ( !XForm )
+	pxr::UsdGeomXformable XForm(UsdPrim);
+	if (!XForm)
 	{
 		return false;
 	}
@@ -3321,14 +3409,14 @@ bool UnrealToUsd::ConvertXformable( const FTransform& RelativeTransform, pxr::Us
 		return false;
 	}
 
-	FUsdStageInfo StageInfo( UsdPrim.GetStage() );
-	pxr::GfMatrix4d UsdTransform = UnrealToUsd::ConvertTransform( StageInfo, RelativeTransform );
+	FUsdStageInfo StageInfo(UsdPrim.GetStage());
+	pxr::GfMatrix4d UsdTransform = UnrealToUsd::ConvertTransform(StageInfo, RelativeTransform);
 
-	const pxr::UsdTimeCode UsdTimeCode( TimeCode );
+	const pxr::UsdTimeCode UsdTimeCode(TimeCode);
 
-	if ( pxr::UsdGeomXformOp MatrixXform = UE::USDPrimConversionImpl::Private::ForceMatrixXform( XForm ) )
+	if (pxr::UsdGeomXformOp MatrixXform = UE::USDPrimConversionImpl::Private::ForceMatrixXform(XForm))
 	{
-		MatrixXform.Set( UsdTransform, UsdTimeCode );
+		MatrixXform.Set(UsdTransform, UsdTimeCode);
 
 		UsdUtils::NotifyIfOverriddenOpinion(MatrixXform.GetAttr());
 		UsdUtils::NotifyIfOverriddenOpinion(XForm.GetXformOpOrderAttr());
@@ -3360,18 +3448,18 @@ namespace UE
 
 				const int32 NumInstances = UEInstances.Num();
 
-				ProtoIndices.reserve( ProtoIndices.size() + NumInstances );
-				Positions.reserve( Positions.size() + NumInstances );
-				Orientations.reserve( Orientations.size() + NumInstances );
-				Scales.reserve( Scales.size() + NumInstances );
+				ProtoIndices.reserve(ProtoIndices.size() + NumInstances);
+				Positions.reserve(Positions.size() + NumInstances);
+				Orientations.reserve(Orientations.size() + NumInstances);
+				Scales.reserve(Scales.size() + NumInstances);
 
-				for ( int32 InstanceIndex : UEInstances )
+				for (int32 InstanceIndex : UEInstances)
 				{
-					const FFoliageInstancePlacementInfo* Instance = &Info.Instances[ InstanceIndex ];
+					const FFoliageInstancePlacementInfo* Instance = &Info.Instances[InstanceIndex];
 
 					// Convert axes
-					FTransform UEWorldTransform{ Instance->Rotation, (FVector)Instance->Location, (FVector)Instance->DrawScale3D };
-					FTransform USDTransform = UsdUtils::ConvertAxes( StageInfo.UpAxis == EUsdUpAxis::ZAxis, UEWorldTransform * UEWorldToFoliageActor );
+					FTransform UEWorldTransform{Instance->Rotation, (FVector)Instance->Location, (FVector)Instance->DrawScale3D};
+					FTransform USDTransform = UsdUtils::ConvertAxes(StageInfo.UpAxis == EUsdUpAxis::ZAxis, UEWorldTransform * UEWorldToFoliageActor);
 
 					FVector Translation = USDTransform.GetTranslation();
 					FQuat Rotation = USDTransform.GetRotation();
@@ -3379,31 +3467,31 @@ namespace UE
 
 					// Compensate metersPerUnit
 					const float UEMetersPerUnit = 0.01f;
-					if ( !FMath::IsNearlyEqual( UEMetersPerUnit, StageInfo.MetersPerUnit ) )
+					if (!FMath::IsNearlyEqual(UEMetersPerUnit, StageInfo.MetersPerUnit))
 					{
-						Translation *= ( UEMetersPerUnit / StageInfo.MetersPerUnit );
+						Translation *= (UEMetersPerUnit / StageInfo.MetersPerUnit);
 					}
 
-					ProtoIndices.push_back( PrototypeIndex );
-					Positions.push_back( pxr::GfVec3f( Translation.X, Translation.Y, Translation.Z ) );
-					Orientations.push_back( pxr::GfQuath( Rotation.W, Rotation.X, Rotation.Y, Rotation.Z ) );
-					Scales.push_back( pxr::GfVec3f( Scale.X, Scale.Y, Scale.Z ) );
+					ProtoIndices.push_back(PrototypeIndex);
+					Positions.push_back(pxr::GfVec3f(Translation.X, Translation.Y, Translation.Z));
+					Orientations.push_back(pxr::GfQuath(Rotation.W, Rotation.X, Rotation.Y, Rotation.Z));
+					Scales.push_back(pxr::GfVec3f(Scale.X, Scale.Y, Scale.Z));
 				}
 			}
-		}
-	}
-}
-#endif // WITH_EDITOR
+		}	 // namespace Private
+	}		 // namespace USDPrimConversionImpl
+}	 // namespace UE
+#endif	  // WITH_EDITOR
 
-bool UnrealToUsd::ConvertInstancedFoliageActor( const AInstancedFoliageActor& Actor, pxr::UsdPrim& UsdPrim, double TimeCode, ULevel* InstancesLevel )
+bool UnrealToUsd::ConvertInstancedFoliageActor(const AInstancedFoliageActor& Actor, pxr::UsdPrim& UsdPrim, double TimeCode, ULevel* InstancesLevel)
 {
 #if WITH_EDITOR
 	using namespace pxr;
 
 	FScopedUsdAllocs Allocs;
 
-	UsdGeomPointInstancer PointInstancer{ UsdPrim };
-	if ( !PointInstancer )
+	UsdGeomPointInstancer PointInstancer{UsdPrim};
+	if (!PointInstancer)
 	{
 		return false;
 	}
@@ -3414,7 +3502,7 @@ bool UnrealToUsd::ConvertInstancedFoliageActor( const AInstancedFoliageActor& Ac
 	}
 
 	UsdStageRefPtr Stage = UsdPrim.GetStage();
-	FUsdStageInfo StageInfo{ Stage };
+	FUsdStageInfo StageInfo{Stage};
 
 	VtArray<int> ProtoIndices;
 	VtArray<GfVec3f> Positions;
@@ -3428,24 +3516,24 @@ bool UnrealToUsd::ConvertInstancedFoliageActor( const AInstancedFoliageActor& Ac
 	FTransform UEWorldToFoliageActor = Actor.GetTransform().Inverse();
 
 	int PrototypeIndex = 0;
-	for ( const TPair<UFoliageType*, TUniqueObj<FFoliageInfo>>& FoliagePair : Actor.GetFoliageInfos() )
+	for (const TPair<UFoliageType*, TUniqueObj<FFoliageInfo>>& FoliagePair : Actor.GetFoliageInfos())
 	{
 		const UFoliageType* FoliageType = FoliagePair.Key;
 		const FFoliageInfo& Info = FoliagePair.Value.Get();
 
 		// Traverse valid foliage instances: Those that are being tracked to belonging to a particular component
-		for ( const TPair<FFoliageInstanceBaseId, FFoliageInstanceBaseInfo>& FoliageInstancePair : Actor.InstanceBaseCache.InstanceBaseMap )
+		for (const TPair<FFoliageInstanceBaseId, FFoliageInstanceBaseInfo>& FoliageInstancePair : Actor.InstanceBaseCache.InstanceBaseMap)
 		{
 			const FFoliageInstanceBaseId& ComponentId = FoliageInstancePair.Key;
-			HandledComponents.Add( ComponentId );
+			HandledComponents.Add(ComponentId);
 
 			UActorComponent* Comp = FoliageInstancePair.Value.BasePtr.Get();
-			if ( !Comp || ( InstancesLevel && ( Comp->GetComponentLevel() != InstancesLevel ) ) )
+			if (!Comp || (InstancesLevel && (Comp->GetComponentLevel() != InstancesLevel)))
 			{
 				continue;
 			}
 
-			if ( const TSet<int32>* InstanceSet = Info.ComponentHash.Find( ComponentId ) )
+			if (const TSet<int32>* InstanceSet = Info.ComponentHash.Find(ComponentId))
 			{
 				UE::USDPrimConversionImpl::Private::ConvertFoliageInstances(
 					Info,
@@ -3465,12 +3553,12 @@ bool UnrealToUsd::ConvertInstancedFoliageActor( const AInstancedFoliageActor& Ac
 		// Only export these when we're not given a particular level to export, or if that level is the actor's level (essentially
 		// pretending the invalid instances belong to the actor's level). This mostly helps prevent it from exporting the invalid instances
 		// multiple times in case we're calling this function repeatedly for each individual sublevel
-		if ( !InstancesLevel || InstancesLevel == Actor.GetLevel() )
+		if (!InstancesLevel || InstancesLevel == Actor.GetLevel())
 		{
-			for ( const TPair<FFoliageInstanceBaseId, TSet<int32>>& Pair : Info.ComponentHash )
+			for (const TPair<FFoliageInstanceBaseId, TSet<int32>>& Pair : Info.ComponentHash)
 			{
 				const FFoliageInstanceBaseId& ComponentId = Pair.Key;
-				if ( HandledComponents.Contains( ComponentId ) )
+				if (HandledComponents.Contains(ComponentId))
 				{
 					continue;
 				}
@@ -3493,55 +3581,60 @@ bool UnrealToUsd::ConvertInstancedFoliageActor( const AInstancedFoliageActor& Ac
 		++PrototypeIndex;
 	}
 
-	const pxr::UsdTimeCode UsdTimeCode( TimeCode );
+	const pxr::UsdTimeCode UsdTimeCode(TimeCode);
 
-	if ( UsdAttribute Attr = PointInstancer.CreateProtoIndicesAttr() )
+	if (UsdAttribute Attr = PointInstancer.CreateProtoIndicesAttr())
 	{
-		Attr.Set( ProtoIndices, UsdTimeCode );
+		Attr.Set(ProtoIndices, UsdTimeCode);
 		UsdUtils::NotifyIfOverriddenOpinion(Attr);
 	}
 
-	if ( UsdAttribute Attr = PointInstancer.CreatePositionsAttr() )
+	if (UsdAttribute Attr = PointInstancer.CreatePositionsAttr())
 	{
-		Attr.Set( Positions, UsdTimeCode );
+		Attr.Set(Positions, UsdTimeCode);
 		UsdUtils::NotifyIfOverriddenOpinion(Attr);
 	}
 
-	if ( UsdAttribute Attr = PointInstancer.CreateOrientationsAttr() )
+	if (UsdAttribute Attr = PointInstancer.CreateOrientationsAttr())
 	{
-		Attr.Set( Orientations, UsdTimeCode );
+		Attr.Set(Orientations, UsdTimeCode);
 		UsdUtils::NotifyIfOverriddenOpinion(Attr);
 	}
 
-	if ( UsdAttribute Attr = PointInstancer.CreateScalesAttr() )
+	if (UsdAttribute Attr = PointInstancer.CreateScalesAttr())
 	{
-		Attr.Set( Scales, UsdTimeCode );
+		Attr.Set(Scales, UsdTimeCode);
 		UsdUtils::NotifyIfOverriddenOpinion(Attr);
 	}
 
 	return true;
 #else
 	return false;
-#endif // WITH_EDITOR
+#endif	  // WITH_EDITOR
 }
 
-bool UnrealToUsd::CreateComponentPropertyBaker( UE::FUsdPrim& Prim, const USceneComponent& Component, const FString& PropertyPath, FComponentBaker& OutBaker )
+bool UnrealToUsd::CreateComponentPropertyBaker(
+	UE::FUsdPrim& Prim,
+	const USceneComponent& Component,
+	const FString& PropertyPath,
+	FComponentBaker& OutBaker
+)
 {
 	EBakingType BakerType = EBakingType::None;
 	TFunction<void(double)> BakerFunction;
 
 	FScopedUsdAllocs Allocs;
 
-	pxr::UsdPrim UsdPrim{ Prim };
+	pxr::UsdPrim UsdPrim{Prim};
 	pxr::UsdStageRefPtr UsdStage = UsdPrim.GetStage();
-	FUsdStageInfo StageInfo{ UsdStage };
+	FUsdStageInfo StageInfo{UsdStage};
 
 	// SceneComponent
 	{
-		if ( PropertyPath == TEXT( "Transform" ) )
+		if (PropertyPath == TEXT("Transform"))
 		{
-			pxr::UsdGeomXformable Xformable( UsdPrim );
-			if ( !Xformable )
+			pxr::UsdGeomXformable Xformable(UsdPrim);
+			if (!Xformable)
 			{
 				return false;
 			}
@@ -3549,14 +3642,14 @@ bool UnrealToUsd::CreateComponentPropertyBaker( UE::FUsdPrim& Prim, const UScene
 			Xformable.CreateXformOpOrderAttr();
 
 			// Clear existing transform data and leave just one Transform op there
-			pxr::UsdGeomXformOp TransformOp = UE::USDPrimConversionImpl::Private::ForceMatrixXform( Xformable );
-			if ( !TransformOp )
+			pxr::UsdGeomXformOp TransformOp = UE::USDPrimConversionImpl::Private::ForceMatrixXform(Xformable);
+			if (!TransformOp)
 			{
 				return false;
 			}
 
 			pxr::UsdAttribute Attr = TransformOp.GetAttr();
-			if ( !Attr )
+			if (!Attr)
 			{
 				return false;
 			}
@@ -3565,27 +3658,26 @@ bool UnrealToUsd::CreateComponentPropertyBaker( UE::FUsdPrim& Prim, const UScene
 
 			// Compensate different orientation for light or camera components
 			FTransform AdditionalRotation = FTransform::Identity;
-			if ( UsdPrim.IsA< pxr::UsdGeomCamera >() || UsdPrim.HasAPI< pxr::UsdLuxLightAPI >() )
+			if (UsdPrim.IsA<pxr::UsdGeomCamera>() || UsdPrim.HasAPI<pxr::UsdLuxLightAPI>())
 			{
-				AdditionalRotation = FTransform( FRotator( 0.0f, 90.0f, 0.0f ) );
+				AdditionalRotation = FTransform(FRotator(0.0f, 90.0f, 0.0f));
 
-				if ( StageInfo.UpAxis == EUsdUpAxis::ZAxis )
+				if (StageInfo.UpAxis == EUsdUpAxis::ZAxis)
 				{
-					AdditionalRotation *= FTransform( FRotator( 90.0f, 0.0f, 0.0f ) );
+					AdditionalRotation *= FTransform(FRotator(90.0f, 0.0f, 0.0f));
 				}
 			}
 
 			// Invert compensation applied to parent if it's a light or camera component
-			if ( const USceneComponent* AttachParent = Component.GetAttachParent() )
+			if (const USceneComponent* AttachParent = Component.GetAttachParent())
 			{
-				if ( AttachParent->IsA( UCineCameraComponent::StaticClass() ) ||
-					AttachParent->IsA( ULightComponent::StaticClass() ) )
+				if (AttachParent->IsA(UCineCameraComponent::StaticClass()) || AttachParent->IsA(ULightComponent::StaticClass()))
 				{
-					FTransform InverseCompensation = FTransform( FRotator( 0.0f, 90.f, 0.0f ) );
+					FTransform InverseCompensation = FTransform(FRotator(0.0f, 90.f, 0.0f));
 
-					if ( UsdUtils::GetUsdStageUpAxis( UsdStage ) == pxr::UsdGeomTokens->z )
+					if (UsdUtils::GetUsdStageUpAxis(UsdStage) == pxr::UsdGeomTokens->z)
 					{
-						InverseCompensation *= FTransform( FRotator( 90.0f, 0.f, 0.0f ) );
+						InverseCompensation *= FTransform(FRotator(90.0f, 0.f, 0.0f));
 					}
 
 					AdditionalRotation = AdditionalRotation * InverseCompensation.Inverse();
@@ -3593,7 +3685,7 @@ bool UnrealToUsd::CreateComponentPropertyBaker( UE::FUsdPrim& Prim, const UScene
 			}
 
 			BakerType = EBakingType::Transform;
-			BakerFunction = [&Component, AdditionalRotation, StageInfo, Attr]( double UsdTimeCode )
+			BakerFunction = [&Component, AdditionalRotation, StageInfo, Attr](double UsdTimeCode)
 			{
 				FScopedUsdAllocs Allocs;
 
@@ -3617,18 +3709,18 @@ bool UnrealToUsd::CreateComponentPropertyBaker( UE::FUsdPrim& Prim, const UScene
 				}
 
 				FTransform FinalUETransform = AdditionalRotation * RelativeTransform;
-				pxr::GfMatrix4d UsdTransform = UnrealToUsd::ConvertTransform( StageInfo, FinalUETransform );
-				Attr.Set< pxr::GfMatrix4d>( UsdTransform, UsdTimeCode );
+				pxr::GfMatrix4d UsdTransform = UnrealToUsd::ConvertTransform(StageInfo, FinalUETransform);
+				Attr.Set<pxr::GfMatrix4d>(UsdTransform, UsdTimeCode);
 			};
 		}
 		// bHidden is for the actor, and bHiddenInGame is for a component
 		// A component is only visible when it's not hidden and its actor is not hidden
 		// A bHidden is just handled like a bHiddenInGame for the actor's root component
 		// Whenever we handle a bHiddenInGame, we always combine it with the actor's bHidden
-		else if ( PropertyPath == TEXT( "bHidden" ) || PropertyPath == TEXT( "bHiddenInGame" ) )
+		else if (PropertyPath == TEXT("bHidden") || PropertyPath == TEXT("bHiddenInGame"))
 		{
-			pxr::UsdGeomImageable Imageable( UsdPrim );
-			if ( !Imageable )
+			pxr::UsdGeomImageable Imageable(UsdPrim);
+			if (!Imageable)
 			{
 				return false;
 			}
@@ -3637,23 +3729,23 @@ bool UnrealToUsd::CreateComponentPropertyBaker( UE::FUsdPrim& Prim, const UScene
 			Attr.Clear();
 
 			BakerType = EBakingType::Visibility;
-			BakerFunction = [&Component, Imageable]( double UsdTimeCode )
+			BakerFunction = [&Component, Imageable](double UsdTimeCode)
 			{
-				if ( Component.bHiddenInGame || Component.GetOwner()->IsHidden() )
+				if (Component.bHiddenInGame || Component.GetOwner()->IsHidden())
 				{
-					Imageable.MakeInvisible( UsdTimeCode );
+					Imageable.MakeInvisible(UsdTimeCode);
 
-					if ( pxr::UsdAttribute Attr = Imageable.CreateVisibilityAttr() )
+					if (pxr::UsdAttribute Attr = Imageable.CreateVisibilityAttr())
 					{
-						if ( !Attr.HasAuthoredValue() )
+						if (!Attr.HasAuthoredValue())
 						{
-							Attr.Set<pxr::TfToken>( pxr::UsdGeomTokens->invisible, UsdTimeCode );
+							Attr.Set<pxr::TfToken>(pxr::UsdGeomTokens->invisible, UsdTimeCode);
 						}
 					}
 				}
 				else
 				{
-					Imageable.MakeVisible( UsdTimeCode );
+					Imageable.MakeVisible(UsdTimeCode);
 
 					// Imagine our visibility track has a single key that switches to hidden at frame 60.
 					// If our prim is visible by default, MakeVisible will author absolutely nothing, and we'll end up
@@ -3661,11 +3753,11 @@ bool UnrealToUsd::CreateComponentPropertyBaker( UE::FUsdPrim& Prim, const UScene
 					// will be invisible throughout *the entire duration of the animation* though, which is not what we want.
 					// This check will ensure that if we're visible we should have a value here and not rely on the
 					// fallback value of 'visible', as that doesn't behave how we want.
-					if ( pxr::UsdAttribute Attr = Imageable.CreateVisibilityAttr() )
+					if (pxr::UsdAttribute Attr = Imageable.CreateVisibilityAttr())
 					{
-						if ( !Attr.HasAuthoredValue() )
+						if (!Attr.HasAuthoredValue())
 						{
-							Attr.Set<pxr::TfToken>( pxr::UsdGeomTokens->inherited, UsdTimeCode );
+							Attr.Set<pxr::TfToken>(pxr::UsdGeomTokens->inherited, UsdTimeCode);
 						}
 					}
 				}
@@ -3673,107 +3765,85 @@ bool UnrealToUsd::CreateComponentPropertyBaker( UE::FUsdPrim& Prim, const UScene
 		}
 	}
 
-	if ( const UCineCameraComponent* CameraComponent = Cast<UCineCameraComponent>( &Component ) )
+	if (const UCineCameraComponent* CameraComponent = Cast<UCineCameraComponent>(&Component))
 	{
 		static TSet<FString> RelevantProperties = {
-			TEXT( "CurrentFocalLength" ),
-			TEXT( "FocusSettings.ManualFocusDistance" ),
-			TEXT( "CurrentAperture" ),
-			TEXT( "Filmback.SensorWidth" ),
-			TEXT( "Filmback.SensorHeight" )
+			TEXT("CurrentFocalLength"),
+			TEXT("FocusSettings.ManualFocusDistance"),
+			TEXT("CurrentAperture"),
+			TEXT("Filmback.SensorWidth"),
+			TEXT("Filmback.SensorHeight")
 		};
 
-		if ( RelevantProperties.Contains( PropertyPath ) )
+		if (RelevantProperties.Contains(PropertyPath))
 		{
 			BakerType = EBakingType::Camera;
-			BakerFunction = [UsdStage, CameraComponent, UsdPrim]( double UsdTimeCode ) mutable
+			BakerFunction = [UsdStage, CameraComponent, UsdPrim](double UsdTimeCode) mutable
 			{
-				UnrealToUsd::ConvertCameraComponent( *CameraComponent, UsdPrim, UsdTimeCode );
+				UnrealToUsd::ConvertCameraComponent(*CameraComponent, UsdPrim, UsdTimeCode);
 			};
 		}
 	}
-	else if ( const ULightComponentBase* LightComponentBase = Cast<ULightComponentBase>( &Component ) )
+	else if (const ULightComponentBase* LightComponentBase = Cast<ULightComponentBase>(&Component))
 	{
-		if ( const URectLightComponent* RectLightComponent = Cast<URectLightComponent>( LightComponentBase ) )
+		if (const URectLightComponent* RectLightComponent = Cast<URectLightComponent>(LightComponentBase))
 		{
-			static TSet<FString> RelevantProperties = {
-				TEXT( "SourceHeight" ),
-				TEXT( "SourceWidth" ),
-				TEXT( "Temperature" ),
-				TEXT( "bUseTemperature" ),
-				TEXT( "LightColor" ),
-				TEXT( "Intensity" )
-			};
+			static TSet<FString> RelevantProperties =
+				{TEXT("SourceHeight"), TEXT("SourceWidth"), TEXT("Temperature"), TEXT("bUseTemperature"), TEXT("LightColor"), TEXT("Intensity")};
 
-			if ( RelevantProperties.Contains( PropertyPath ) )
+			if (RelevantProperties.Contains(PropertyPath))
 			{
 				BakerType = EBakingType::Light;
-				BakerFunction = [UsdStage, RectLightComponent, UsdPrim]( double UsdTimeCode ) mutable
+				BakerFunction = [UsdStage, RectLightComponent, UsdPrim](double UsdTimeCode) mutable
 				{
-					UnrealToUsd::ConvertLightComponent( *RectLightComponent, UsdPrim, UsdTimeCode );
-					UnrealToUsd::ConvertRectLightComponent( *RectLightComponent, UsdPrim, UsdTimeCode );
+					UnrealToUsd::ConvertLightComponent(*RectLightComponent, UsdPrim, UsdTimeCode);
+					UnrealToUsd::ConvertRectLightComponent(*RectLightComponent, UsdPrim, UsdTimeCode);
 				};
 			}
 		}
-		else if ( const USpotLightComponent* SpotLightComponent = Cast<USpotLightComponent>( LightComponentBase ) )
+		else if (const USpotLightComponent* SpotLightComponent = Cast<USpotLightComponent>(LightComponentBase))
 		{
-			static TSet<FString> RelevantProperties = {
-				TEXT( "OuterConeAngle" ),
-				TEXT( "InnerConeAngle" ),
-				TEXT( "Temperature" ),
-				TEXT( "bUseTemperature" ),
-				TEXT( "LightColor" ),
-				TEXT( "Intensity" )
-			};
+			static TSet<FString> RelevantProperties =
+				{TEXT("OuterConeAngle"), TEXT("InnerConeAngle"), TEXT("Temperature"), TEXT("bUseTemperature"), TEXT("LightColor"), TEXT("Intensity")};
 
-			if ( RelevantProperties.Contains( PropertyPath ) )
+			if (RelevantProperties.Contains(PropertyPath))
 			{
 				BakerType = EBakingType::Light;
-				BakerFunction = [UsdStage, SpotLightComponent, UsdPrim]( double UsdTimeCode ) mutable
+				BakerFunction = [UsdStage, SpotLightComponent, UsdPrim](double UsdTimeCode) mutable
 				{
-					UnrealToUsd::ConvertLightComponent( *SpotLightComponent, UsdPrim, UsdTimeCode );
-					UnrealToUsd::ConvertPointLightComponent( *SpotLightComponent, UsdPrim, UsdTimeCode );
-					UnrealToUsd::ConvertSpotLightComponent( *SpotLightComponent, UsdPrim, UsdTimeCode );
+					UnrealToUsd::ConvertLightComponent(*SpotLightComponent, UsdPrim, UsdTimeCode);
+					UnrealToUsd::ConvertPointLightComponent(*SpotLightComponent, UsdPrim, UsdTimeCode);
+					UnrealToUsd::ConvertSpotLightComponent(*SpotLightComponent, UsdPrim, UsdTimeCode);
 				};
 			}
 		}
-		else if ( const UPointLightComponent* PointLightComponent = Cast<UPointLightComponent>( LightComponentBase ) )
+		else if (const UPointLightComponent* PointLightComponent = Cast<UPointLightComponent>(LightComponentBase))
 		{
-			static TSet<FString> RelevantProperties = {
-				TEXT( "SourceRadius" ),
-				TEXT( "Temperature" ),
-				TEXT( "bUseTemperature" ),
-				TEXT( "LightColor" ),
-				TEXT( "Intensity" )
-			};
+			static TSet<FString> RelevantProperties =
+				{TEXT("SourceRadius"), TEXT("Temperature"), TEXT("bUseTemperature"), TEXT("LightColor"), TEXT("Intensity")};
 
-			if ( RelevantProperties.Contains( PropertyPath ) )
+			if (RelevantProperties.Contains(PropertyPath))
 			{
 				BakerType = EBakingType::Light;
-				BakerFunction = [UsdStage, PointLightComponent, UsdPrim]( double UsdTimeCode ) mutable
+				BakerFunction = [UsdStage, PointLightComponent, UsdPrim](double UsdTimeCode) mutable
 				{
-					UnrealToUsd::ConvertLightComponent( *PointLightComponent, UsdPrim, UsdTimeCode );
-					UnrealToUsd::ConvertPointLightComponent( *PointLightComponent, UsdPrim, UsdTimeCode );
+					UnrealToUsd::ConvertLightComponent(*PointLightComponent, UsdPrim, UsdTimeCode);
+					UnrealToUsd::ConvertPointLightComponent(*PointLightComponent, UsdPrim, UsdTimeCode);
 				};
 			}
 		}
-		else if ( const UDirectionalLightComponent* DirectionalLightComponent = Cast<UDirectionalLightComponent>( LightComponentBase ) )
+		else if (const UDirectionalLightComponent* DirectionalLightComponent = Cast<UDirectionalLightComponent>(LightComponentBase))
 		{
-			static TSet<FString> RelevantProperties = {
-				TEXT( "LightSourceAngle" ),
-				TEXT( "Temperature" ),
-				TEXT( "bUseTemperature" ),
-				TEXT( "LightColor" ),
-				TEXT( "Intensity" )
-			};
+			static TSet<FString> RelevantProperties =
+				{TEXT("LightSourceAngle"), TEXT("Temperature"), TEXT("bUseTemperature"), TEXT("LightColor"), TEXT("Intensity")};
 
-			if ( RelevantProperties.Contains( PropertyPath ) )
+			if (RelevantProperties.Contains(PropertyPath))
 			{
 				BakerType = EBakingType::Light;
-				BakerFunction = [UsdStage, DirectionalLightComponent, UsdPrim]( double UsdTimeCode ) mutable
+				BakerFunction = [UsdStage, DirectionalLightComponent, UsdPrim](double UsdTimeCode) mutable
 				{
-					UnrealToUsd::ConvertLightComponent( *DirectionalLightComponent, UsdPrim, UsdTimeCode );
-					UnrealToUsd::ConvertDirectionalLightComponent( *DirectionalLightComponent, UsdPrim, UsdTimeCode );
+					UnrealToUsd::ConvertLightComponent(*DirectionalLightComponent, UsdPrim, UsdTimeCode);
+					UnrealToUsd::ConvertDirectionalLightComponent(*DirectionalLightComponent, UsdPrim, UsdTimeCode);
 				};
 			}
 		}
@@ -3796,7 +3866,7 @@ bool UnrealToUsd::CreateComponentPropertyBaker( UE::FUsdPrim& Prim, const UScene
 		}
 	}
 
-	if ( BakerType != EBakingType::None && BakerFunction)
+	if (BakerType != EBakingType::None && BakerFunction)
 	{
 		OutBaker.BakerType = BakerType;
 		OutBaker.BakerFunction = BakerFunction;
@@ -3807,11 +3877,16 @@ bool UnrealToUsd::CreateComponentPropertyBaker( UE::FUsdPrim& Prim, const UScene
 	return false;
 }
 
-bool UnrealToUsd::CreateSkeletalAnimationBaker( UE::FUsdPrim& SkeletonPrim, UE::FUsdPrim& SkelAnimation, USkeletalMeshComponent& Component, FComponentBaker& OutBaker )
+bool UnrealToUsd::CreateSkeletalAnimationBaker(
+	UE::FUsdPrim& SkeletonPrim,
+	UE::FUsdPrim& SkelAnimation,
+	USkeletalMeshComponent& Component,
+	FComponentBaker& OutBaker
+)
 {
 #if WITH_EDITOR
 	USkeletalMesh* SkeletalMesh = Component.GetSkeletalMeshAsset();
-	if ( !SkeletalMesh )
+	if (!SkeletalMesh)
 	{
 		return false;
 	}
@@ -3830,8 +3905,8 @@ bool UnrealToUsd::CreateSkeletalAnimationBaker( UE::FUsdPrim& SkeletonPrim, UE::
 		return false;
 	}
 
-	pxr::UsdSkelAnimation UsdSkelAnimation{ SkelAnimation };
-	if ( !UsdSkeleton || !SkelAnimation )
+	pxr::UsdSkelAnimation UsdSkelAnimation{SkelAnimation};
+	if (!UsdSkeleton || !SkelAnimation)
 	{
 		return false;
 	}
@@ -3840,25 +3915,25 @@ bool UnrealToUsd::CreateSkeletalAnimationBaker( UE::FUsdPrim& SkeletonPrim, UE::
 	pxr::UsdPrim SkelAnimPrim = UsdSkelAnimation.GetPrim();
 	UsdUtils::BindAnimationSource(SkeletonPrim, SkelAnimPrim);
 
-	FUsdStageInfo StageInfo{ SkeletonPrim.GetStage() };
+	FUsdStageInfo StageInfo{SkeletonPrim.GetStage()};
 
-	pxr::UsdAttribute JointsAttr			= UsdSkelAnimation.CreateJointsAttr();
-	pxr::UsdAttribute TranslationsAttr		= UsdSkelAnimation.CreateTranslationsAttr();
-	pxr::UsdAttribute RotationsAttr			= UsdSkelAnimation.CreateRotationsAttr();
-	pxr::UsdAttribute ScalesAttr			= UsdSkelAnimation.CreateScalesAttr();
+	pxr::UsdAttribute JointsAttr = UsdSkelAnimation.CreateJointsAttr();
+	pxr::UsdAttribute TranslationsAttr = UsdSkelAnimation.CreateTranslationsAttr();
+	pxr::UsdAttribute RotationsAttr = UsdSkelAnimation.CreateRotationsAttr();
+	pxr::UsdAttribute ScalesAttr = UsdSkelAnimation.CreateScalesAttr();
 
 	// Joints
 	const FReferenceSkeleton& RefSkeleton = SkeletalMesh->GetRefSkeleton();
 	const int32 NumBones = RefSkeleton.GetRefBoneInfo().Num();
-	UnrealToUsd::ConvertJointsAttribute( RefSkeleton, JointsAttr );
+	UnrealToUsd::ConvertJointsAttribute(RefSkeleton, JointsAttr);
 
 	// Build active morph targets array if it isn't setup already
 	TArray<TObjectPtr<UMorphTarget>>& MorphTargets = SkeletalMesh->GetMorphTargets();
-	if ( Component.ActiveMorphTargets.Num() != Component.MorphTargetWeights.Num() && MorphTargets.Num() != 0 )
+	if (Component.ActiveMorphTargets.Num() != Component.MorphTargetWeights.Num() && MorphTargets.Num() != 0)
 	{
-		for ( int32 MorphTargetIndex = 0; MorphTargetIndex < MorphTargets.Num(); ++MorphTargetIndex )
+		for (int32 MorphTargetIndex = 0; MorphTargetIndex < MorphTargets.Num(); ++MorphTargetIndex)
 		{
-			Component.ActiveMorphTargets.Add( MorphTargets[ MorphTargetIndex ], MorphTargetIndex );
+			Component.ActiveMorphTargets.Add(MorphTargets[MorphTargetIndex], MorphTargetIndex);
 		}
 	}
 
@@ -3869,141 +3944,158 @@ bool UnrealToUsd::CreateSkeletalAnimationBaker( UE::FUsdPrim& SkeletonPrim, UE::
 	pxr::UsdAttribute BlendShapeWeightsAttr;
 	pxr::UsdAttribute BlendShapesAttr;
 	const int32 NumMorphTargets = Component.MorphTargetWeights.Num();
-	if ( NumMorphTargets > 0 )
+	if (NumMorphTargets > 0)
 	{
 		BlendShapeWeightsAttr = UsdSkelAnimation.CreateBlendShapeWeightsAttr();
 		BlendShapesAttr = UsdSkelAnimation.CreateBlendShapesAttr();
 
 		TArray<FMorphTargetWeightMap::ElementType> SortedMorphTargets = Component.ActiveMorphTargets.Array();
-		SortedMorphTargets.Sort([](const FMorphTargetWeightMap::ElementType& Left, const FMorphTargetWeightMap::ElementType& Right)
-		{
-			return Left.Value < Right.Value;
-		});
+		SortedMorphTargets.Sort(
+			[](const FMorphTargetWeightMap::ElementType& Left, const FMorphTargetWeightMap::ElementType& Right)
+			{
+				return Left.Value < Right.Value;
+			}
+		);
 
-		pxr::VtArray< pxr::TfToken > BlendShapeNames;
+		pxr::VtArray<pxr::TfToken> BlendShapeNames;
 		BlendShapeNames.reserve(SortedMorphTargets.Num());
 
-		for ( const FMorphTargetWeightMap::ElementType& ActiveMorphTarget : SortedMorphTargets )
+		for (const FMorphTargetWeightMap::ElementType& ActiveMorphTarget : SortedMorphTargets)
 		{
 			FString BlendShapeName;
-			if ( const UMorphTarget* MorphTarget = ActiveMorphTarget.Key )
+			if (const UMorphTarget* MorphTarget = ActiveMorphTarget.Key)
 			{
 				BlendShapeName = MorphTarget->GetFName().ToString();
 			}
 
-			BlendShapeNames.push_back( UnrealToUsd::ConvertToken( *BlendShapeName ).Get() );
+			BlendShapeNames.push_back(UnrealToUsd::ConvertToken(*BlendShapeName).Get());
 		}
 
-		BlendShapesAttr.Set( BlendShapeNames );
+		BlendShapesAttr.Set(BlendShapeNames);
 	}
 
 	pxr::VtVec3fArray Translations;
 	pxr::VtQuatfArray Rotations;
 	pxr::VtVec3hArray Scales;
-	pxr::VtArray< float > BlendShapeWeights;
+	pxr::VtArray<float> BlendShapeWeights;
 
 	OutBaker.ComponentPath = Component.GetPathName();
 	OutBaker.BakerType = EBakingType::Skeletal;
-	OutBaker.BakerFunction =
-		[&Component, StageInfo, Translations, Rotations, Scales, BlendShapeWeights, TranslationsAttr, RotationsAttr, ScalesAttr, BlendShapeWeightsAttr, NumBones, NumMorphTargets]
-		( double UsdTimeCode ) mutable
-		{
-			FScopedUsdAllocs InnerAllocs;
-
-			Translations.resize( NumBones );
-			Rotations.resize( NumBones );
-			Scales.resize( NumBones );
-
-			// This whole incantation is required or else the component will really not update until the next frame.
-			// Note: This will also cause the update of morph target weights.
-			Component.TickAnimation( 0.f, false );
-			Component.UpdateLODStatus();
-			Component.RefreshBoneTransforms();
-			Component.RefreshFollowerComponents();
-			Component.UpdateComponentToWorld();
-			Component.FinalizeBoneTransform();
-			Component.MarkRenderTransformDirty();
-			Component.MarkRenderDynamicDataDirty();
-
-			// I'm not entirely sure why this is needed but FFbxExporter::ExportAnimTrack and FFbxExporter::ExportLevelSequenceBaked3DTransformTrack
-			// do this so for safety maybe we should as well?
-			if (AActor* Owner = Component.GetOwner())
-			{
-				Owner->Tick(0.0f);
-			}
+	OutBaker.BakerFunction = [&Component,
+							  StageInfo,
+							  Translations,
+							  Rotations,
+							  Scales,
+							  BlendShapeWeights,
+							  TranslationsAttr,
+							  RotationsAttr,
+							  ScalesAttr,
+							  BlendShapeWeightsAttr,
+							  NumBones,
+							  NumMorphTargets](double UsdTimeCode) mutable
+	{
+		FScopedUsdAllocs InnerAllocs;
 
-			const TArray< FTransform >& LocalBoneTransforms = Component.GetBoneSpaceTransforms();
+		Translations.resize(NumBones);
+		Rotations.resize(NumBones);
+		Scales.resize(NumBones);
 
-			// For whatever reason it seems that sometimes this is not ready for us, so let's force it to be recalculated
-			if (LocalBoneTransforms.Num() == 0)
-			{
-				const int32 LODIndex = 0;
-				Component.RecalcRequiredBones(LODIndex);
-			}
-			if (LocalBoneTransforms.Num() != NumBones)
-			{
-				UE_LOG(
-					LogUsd,
-					Warning,
-					TEXT("Failed to retrieve bone transforms when baking skeletal animation for component '%s' at timeCode '%f'. Expected %d transforms, received %d"),
-					*Component.GetPathName(),
-					UsdTimeCode,
-					NumBones,
-					LocalBoneTransforms.Num()
-				);
-				return;
-			}
+		// This whole incantation is required or else the component will really not update until the next frame.
+		// Note: This will also cause the update of morph target weights.
+		Component.TickAnimation(0.f, false);
+		Component.UpdateLODStatus();
+		Component.RefreshBoneTransforms();
+		Component.RefreshFollowerComponents();
+		Component.UpdateComponentToWorld();
+		Component.FinalizeBoneTransform();
+		Component.MarkRenderTransformDirty();
+		Component.MarkRenderDynamicDataDirty();
 
-			for ( int32 BoneIndex = 0; BoneIndex < NumBones; ++BoneIndex )
-			{
-				FTransform BoneTransform = LocalBoneTransforms[ BoneIndex ];
-				BoneTransform = UsdUtils::ConvertAxes( StageInfo.UpAxis == EUsdUpAxis::ZAxis, BoneTransform );
+		// I'm not entirely sure why this is needed but FFbxExporter::ExportAnimTrack and FFbxExporter::ExportLevelSequenceBaked3DTransformTrack
+		// do this so for safety maybe we should as well?
+		if (AActor* Owner = Component.GetOwner())
+		{
+			Owner->Tick(0.0f);
+		}
 
-				Translations[ BoneIndex ] = UnrealToUsd::ConvertVectorFloat( BoneTransform.GetTranslation() ) * ( 0.01f / StageInfo.MetersPerUnit );
-				Rotations[ BoneIndex ] = UnrealToUsd::ConvertQuatFloat( BoneTransform.GetRotation() ).GetNormalized();
-				Scales[ BoneIndex ] = UnrealToUsd::ConvertVectorHalf( BoneTransform.GetScale3D() );
-			}
+		const TArray<FTransform>& LocalBoneTransforms = Component.GetBoneSpaceTransforms();
 
-			if ( Translations.size() > 0 )
-			{
-				TranslationsAttr.Set( Translations, UsdTimeCode );
-				RotationsAttr.Set( Rotations, UsdTimeCode );
-				ScalesAttr.Set( Scales, UsdTimeCode );
-			}
+		// For whatever reason it seems that sometimes this is not ready for us, so let's force it to be recalculated
+		if (LocalBoneTransforms.Num() == 0)
+		{
+			const int32 LODIndex = 0;
+			Component.RecalcRequiredBones(LODIndex);
+		}
+		if (LocalBoneTransforms.Num() != NumBones)
+		{
+			UE_LOG(
+				LogUsd,
+				Warning,
+				TEXT("Failed to retrieve bone transforms when baking skeletal animation for component '%s' at timeCode '%f'. Expected %d transforms, "
+					 "received %d"),
+				*Component.GetPathName(),
+				UsdTimeCode,
+				NumBones,
+				LocalBoneTransforms.Num()
+			);
+			return;
+		}
 
-			if ( NumMorphTargets > 0 && BlendShapeWeightsAttr )
-			{
-				BlendShapeWeights.resize( NumMorphTargets );
+		for (int32 BoneIndex = 0; BoneIndex < NumBones; ++BoneIndex)
+		{
+			FTransform BoneTransform = LocalBoneTransforms[BoneIndex];
+			BoneTransform = UsdUtils::ConvertAxes(StageInfo.UpAxis == EUsdUpAxis::ZAxis, BoneTransform);
 
-				for ( int32 MorphTargetIndex = 0; MorphTargetIndex < NumMorphTargets; ++MorphTargetIndex )
-				{
-					BlendShapeWeights[ MorphTargetIndex ] = Component.MorphTargetWeights[ MorphTargetIndex ];
-				}
+			Translations[BoneIndex] = UnrealToUsd::ConvertVectorFloat(BoneTransform.GetTranslation()) * (0.01f / StageInfo.MetersPerUnit);
+			Rotations[BoneIndex] = UnrealToUsd::ConvertQuatFloat(BoneTransform.GetRotation()).GetNormalized();
+			Scales[BoneIndex] = UnrealToUsd::ConvertVectorHalf(BoneTransform.GetScale3D());
+		}
 
-				BlendShapeWeightsAttr.Set( BlendShapeWeights, UsdTimeCode );
+		if (Translations.size() > 0)
+		{
+			TranslationsAttr.Set(Translations, UsdTimeCode);
+			RotationsAttr.Set(Rotations, UsdTimeCode);
+			ScalesAttr.Set(Scales, UsdTimeCode);
+		}
+
+		if (NumMorphTargets > 0 && BlendShapeWeightsAttr)
+		{
+			BlendShapeWeights.resize(NumMorphTargets);
+
+			for (int32 MorphTargetIndex = 0; MorphTargetIndex < NumMorphTargets; ++MorphTargetIndex)
+			{
+				BlendShapeWeights[MorphTargetIndex] = Component.MorphTargetWeights[MorphTargetIndex];
 			}
-		};
+
+			BlendShapeWeightsAttr.Set(BlendShapeWeights, UsdTimeCode);
+		}
+	};
 
 	return true;
 #else
 	return false;
-#endif // WITH_EDITOR
+#endif	  // WITH_EDITOR
 }
 
-UnrealToUsd::FPropertyTrackWriter UnrealToUsd::CreatePropertyTrackWriter( const USceneComponent& Component, const UMovieScenePropertyTrack& Track, UE::FUsdPrim& Prim, TSet<FName>& OutPropertyPathsToRefresh )
+UnrealToUsd::FPropertyTrackWriter UnrealToUsd::CreatePropertyTrackWriter(
+	const USceneComponent& Component,
+	const UMovieScenePropertyTrack& Track,
+	UE::FUsdPrim& Prim,
+	TSet<FName>& OutPropertyPathsToRefresh
+)
 {
 	UnrealToUsd::FPropertyTrackWriter Result;
 
-	if ( !Prim )
+	if (!Prim)
 	{
 		return Result;
 	}
 
 	FScopedUsdAllocs Allocs;
 
-	pxr::UsdPrim UsdPrim{ Prim };
+	pxr::UsdPrim UsdPrim{Prim};
 	pxr::UsdStageRefPtr UsdStage = UsdPrim.GetStage();
-	FUsdStageInfo StageInfo{ UsdStage };
+	FUsdStageInfo StageInfo{UsdStage};
 
 	if (UsdUtils::NotifyIfInstanceProxy(Prim))
 	{
@@ -4024,50 +4116,49 @@ UnrealToUsd::FPropertyTrackWriter UnrealToUsd::CreatePropertyTrackWriter( const
 
 		// SceneComponent
 		{
-			if ( PropertyPath == UnrealIdentifiers::TransformPropertyName )
+			if (PropertyPath == UnrealIdentifiers::TransformPropertyName)
 			{
-				if ( pxr::UsdGeomXformable Xformable{ UsdPrim } )
+				if (pxr::UsdGeomXformable Xformable{UsdPrim})
 				{
 					Xformable.CreateXformOpOrderAttr();
 
-					if ( pxr::UsdGeomXformOp TransformOp = UE::USDPrimConversionImpl::Private::ForceMatrixXform( Xformable ) )
+					if (pxr::UsdGeomXformOp TransformOp = UE::USDPrimConversionImpl::Private::ForceMatrixXform(Xformable))
 					{
 						Attr = TransformOp.GetAttr();
 
 						// Compensate different orientation for light or camera components
 						FTransform AdditionalRotation = FTransform::Identity;
-						if ( UsdPrim.IsA< pxr::UsdGeomCamera >() || UsdPrim.HasAPI< pxr::UsdLuxLightAPI >() )
+						if (UsdPrim.IsA<pxr::UsdGeomCamera>() || UsdPrim.HasAPI<pxr::UsdLuxLightAPI>())
 						{
-							AdditionalRotation = FTransform( FRotator( 0.0f, 90.0f, 0.0f ) );
+							AdditionalRotation = FTransform(FRotator(0.0f, 90.0f, 0.0f));
 
-							if ( StageInfo.UpAxis == EUsdUpAxis::ZAxis )
+							if (StageInfo.UpAxis == EUsdUpAxis::ZAxis)
 							{
-								AdditionalRotation *= FTransform( FRotator( 90.0f, 0.0f, 0.0f ) );
+								AdditionalRotation *= FTransform(FRotator(90.0f, 0.0f, 0.0f));
 							}
 						}
 
 						// Invert compensation applied to parent if it's a light or camera component
-						if ( const USceneComponent* AttachParent = Component.GetAttachParent() )
+						if (const USceneComponent* AttachParent = Component.GetAttachParent())
 						{
-							if ( AttachParent->IsA( UCineCameraComponent::StaticClass() ) ||
-								AttachParent->IsA( ULightComponent::StaticClass() ) )
+							if (AttachParent->IsA(UCineCameraComponent::StaticClass()) || AttachParent->IsA(ULightComponent::StaticClass()))
 							{
-								FTransform InverseCompensation = FTransform( FRotator( 0.0f, 90.f, 0.0f ) );
+								FTransform InverseCompensation = FTransform(FRotator(0.0f, 90.f, 0.0f));
 
-								if ( UsdUtils::GetUsdStageUpAxis( UsdStage ) == pxr::UsdGeomTokens->z )
+								if (UsdUtils::GetUsdStageUpAxis(UsdStage) == pxr::UsdGeomTokens->z)
 								{
-									InverseCompensation *= FTransform( FRotator( 90.0f, 0.f, 0.0f ) );
+									InverseCompensation *= FTransform(FRotator(90.0f, 0.f, 0.0f));
 								}
 
 								AdditionalRotation = AdditionalRotation * InverseCompensation.Inverse();
 							}
 						}
 
-						Result.TransformWriter = [&Component, AdditionalRotation, StageInfo, Attr]( const FTransform& UEValue, double UsdTimeCode )
+						Result.TransformWriter = [&Component, AdditionalRotation, StageInfo, Attr](const FTransform& UEValue, double UsdTimeCode)
 						{
 							FTransform FinalUETransform = AdditionalRotation * UEValue;
-							pxr::GfMatrix4d UsdTransform = UnrealToUsd::ConvertTransform( StageInfo, FinalUETransform );
-							Attr.Set< pxr::GfMatrix4d >( UsdTransform, UsdTimeCode );
+							pxr::GfMatrix4d UsdTransform = UnrealToUsd::ConvertTransform(StageInfo, FinalUETransform);
+							Attr.Set<pxr::GfMatrix4d>(UsdTransform, UsdTimeCode);
 						};
 					}
 				}
@@ -4076,27 +4167,27 @@ UnrealToUsd::FPropertyTrackWriter UnrealToUsd::CreatePropertyTrackWriter( const
 			// A component is only visible when it's not hidden and its actor is not hidden
 			// A bHidden is just handled like a bHiddenInGame for the actor's root component
 			// Whenever we handle a bHiddenInGame, we always combine it with the actor's bHidden
-			else if ( PropertyPath == UnrealIdentifiers::HiddenPropertyName || PropertyPath == UnrealIdentifiers::HiddenInGamePropertyName )
+			else if (PropertyPath == UnrealIdentifiers::HiddenPropertyName || PropertyPath == UnrealIdentifiers::HiddenInGamePropertyName)
 			{
-				if ( pxr::UsdGeomImageable Imageable{ UsdPrim } )
+				if (pxr::UsdGeomImageable Imageable{UsdPrim})
 				{
 					Attr = Imageable.CreateVisibilityAttr();
-					if ( Attr )
+					if (Attr)
 					{
-						Result.BoolWriter = [Imageable, Attr]( bool UEValue, double UsdTimeCode )
+						Result.BoolWriter = [Imageable, Attr](bool UEValue, double UsdTimeCode)
 						{
-							if ( UEValue )
+							if (UEValue)
 							{
 								// We have to do both here as MakeVisible will ensure we also flip any parent prims,
 								// and setting the attribute will ensure we write a timeSample. Otherwise if MakeVisible
 								// finds that the prim should already be visible due to a stronger opinion, it won't write anything
-								Attr.Set<pxr::TfToken>( pxr::UsdGeomTokens->inherited, UsdTimeCode );
-								Imageable.MakeVisible( UsdTimeCode );
+								Attr.Set<pxr::TfToken>(pxr::UsdGeomTokens->inherited, UsdTimeCode);
+								Imageable.MakeVisible(UsdTimeCode);
 							}
 							else
 							{
-								Attr.Set<pxr::TfToken>( pxr::UsdGeomTokens->invisible, UsdTimeCode );
-								Imageable.MakeInvisible( UsdTimeCode );
+								Attr.Set<pxr::TfToken>(pxr::UsdGeomTokens->invisible, UsdTimeCode);
+								Imageable.MakeInvisible(UsdTimeCode);
 							}
 						};
 					}
@@ -4104,109 +4195,109 @@ UnrealToUsd::FPropertyTrackWriter UnrealToUsd::CreatePropertyTrackWriter( const
 			}
 		}
 
-		if ( pxr::UsdGeomCamera Camera{ UsdPrim } )
+		if (pxr::UsdGeomCamera Camera{UsdPrim})
 		{
 			bool bConvertDistance = true;
 
-			if ( PropertyPath == UnrealIdentifiers::CurrentFocalLengthPropertyName )
+			if (PropertyPath == UnrealIdentifiers::CurrentFocalLengthPropertyName)
 			{
 				Attr = Camera.CreateFocalLengthAttr();
 			}
-			else if ( PropertyPath == UnrealIdentifiers::ManualFocusDistancePropertyName )
+			else if (PropertyPath == UnrealIdentifiers::ManualFocusDistancePropertyName)
 			{
 				Attr = Camera.CreateFocusDistanceAttr();
 			}
-			else if ( PropertyPath == UnrealIdentifiers::CurrentAperturePropertyName )
+			else if (PropertyPath == UnrealIdentifiers::CurrentAperturePropertyName)
 			{
 				bConvertDistance = false;
 				Attr = Camera.CreateFStopAttr();
 			}
-			else if ( PropertyPath == UnrealIdentifiers::SensorWidthPropertyName )
+			else if (PropertyPath == UnrealIdentifiers::SensorWidthPropertyName)
 			{
 				Attr = Camera.CreateHorizontalApertureAttr();
 			}
-			else if ( PropertyPath == UnrealIdentifiers::SensorHeightPropertyName )
+			else if (PropertyPath == UnrealIdentifiers::SensorHeightPropertyName)
 			{
 				Attr = Camera.CreateVerticalApertureAttr();
 			}
 
-			if ( Attr )
+			if (Attr)
 			{
-				if ( bConvertDistance )
+				if (bConvertDistance)
 				{
-					Result.FloatWriter = [Attr, StageInfo]( float UEValue, double UsdTimeCode ) mutable
+					Result.FloatWriter = [Attr, StageInfo](float UEValue, double UsdTimeCode) mutable
 					{
-						Attr.Set( UnrealToUsd::ConvertDistance( StageInfo, UEValue ), UsdTimeCode );
+						Attr.Set(UnrealToUsd::ConvertDistance(StageInfo, UEValue), UsdTimeCode);
 					};
 				}
 				else
 				{
-					Result.FloatWriter = [Attr]( float UEValue, double UsdTimeCode ) mutable
+					Result.FloatWriter = [Attr](float UEValue, double UsdTimeCode) mutable
 					{
-						Attr.Set( UEValue, UsdTimeCode );
+						Attr.Set(UEValue, UsdTimeCode);
 					};
 				}
 			}
 		}
-		else if ( pxr::UsdLuxLightAPI LightAPI{ Prim } )
+		else if (pxr::UsdLuxLightAPI LightAPI{Prim})
 		{
-			if ( PropertyPath == UnrealIdentifiers::LightColorPropertyName )
+			if (PropertyPath == UnrealIdentifiers::LightColorPropertyName)
 			{
 				Attr = LightAPI.GetColorAttr();
-				if ( Attr )
+				if (Attr)
 				{
-					Result.ColorWriter = [Attr]( const FLinearColor& UEValue, double UsdTimeCode )
+					Result.ColorWriter = [Attr](const FLinearColor& UEValue, double UsdTimeCode)
 					{
-						pxr::GfVec4f Vec4 = UnrealToUsd::ConvertColor( UEValue );
-						Attr.Set( pxr::GfVec3f{ Vec4[ 0 ], Vec4[ 1 ], Vec4[ 2 ] }, UsdTimeCode );
+						pxr::GfVec4f Vec4 = UnrealToUsd::ConvertColor(UEValue);
+						Attr.Set(pxr::GfVec3f{Vec4[0], Vec4[1], Vec4[2]}, UsdTimeCode);
 					};
 				}
 			}
-			else if ( PropertyPath == UnrealIdentifiers::UseTemperaturePropertyName )
+			else if (PropertyPath == UnrealIdentifiers::UseTemperaturePropertyName)
 			{
 				Attr = LightAPI.GetEnableColorTemperatureAttr();
-				if ( Attr )
+				if (Attr)
 				{
-					Result.BoolWriter = [Attr]( bool UEValue, double UsdTimeCode )
+					Result.BoolWriter = [Attr](bool UEValue, double UsdTimeCode)
 					{
-						Attr.Set( UEValue, UsdTimeCode );
+						Attr.Set(UEValue, UsdTimeCode);
 					};
 				}
 			}
-			else if ( PropertyPath == UnrealIdentifiers::TemperaturePropertyName )
+			else if (PropertyPath == UnrealIdentifiers::TemperaturePropertyName)
 			{
 				Attr = LightAPI.GetColorTemperatureAttr();
-				if ( Attr )
+				if (Attr)
 				{
-					Result.FloatWriter = [Attr]( float UEValue, double UsdTimeCode )
+					Result.FloatWriter = [Attr](float UEValue, double UsdTimeCode)
 					{
-						Attr.Set( UEValue, UsdTimeCode );
+						Attr.Set(UEValue, UsdTimeCode);
 					};
 				}
 			}
 
-			else if ( pxr::UsdLuxSphereLight SphereLight{ UsdPrim } )
+			else if (pxr::UsdLuxSphereLight SphereLight{UsdPrim})
 			{
-				if ( PropertyPath == UnrealIdentifiers::SourceRadiusPropertyName )
+				if (PropertyPath == UnrealIdentifiers::SourceRadiusPropertyName)
 				{
-					OutPropertyPathsToRefresh.Add( UnrealIdentifiers::IntensityPropertyName );
+					OutPropertyPathsToRefresh.Add(UnrealIdentifiers::IntensityPropertyName);
 
 					Attr = SphereLight.GetRadiusAttr();
-					if ( Attr )
+					if (Attr)
 					{
-						Result.FloatWriter = [Attr, StageInfo]( float UEValue, double UsdTimeCode )
+						Result.FloatWriter = [Attr, StageInfo](float UEValue, double UsdTimeCode)
 						{
-							Attr.Set( UnrealToUsd::ConvertDistance( StageInfo, UEValue ), UsdTimeCode );
+							Attr.Set(UnrealToUsd::ConvertDistance(StageInfo, UEValue), UsdTimeCode);
 						};
 					}
 				}
 
 				// Spot light
-				else if ( UsdPrim.HasAPI< pxr::UsdLuxShapingAPI >() )
+				else if (UsdPrim.HasAPI<pxr::UsdLuxShapingAPI>())
 				{
-					pxr::UsdLuxShapingAPI ShapingAPI{ UsdPrim };
+					pxr::UsdLuxShapingAPI ShapingAPI{UsdPrim};
 
-					if ( PropertyPath == UnrealIdentifiers::IntensityPropertyName )
+					if (PropertyPath == UnrealIdentifiers::IntensityPropertyName)
 					{
 						Attr = SphereLight.GetIntensityAttr();
 						pxr::UsdAttribute RadiusAttr = SphereLight.GetRadiusAttr();
@@ -4215,65 +4306,73 @@ UnrealToUsd::FPropertyTrackWriter UnrealToUsd::CreatePropertyTrackWriter( const
 
 						// Always clear exposure because we'll put all of our "light intensity" on the intensity attr and assume exposure
 						// is zero, as we can't manipulate something like that exposure directly from UE anyway
-						if ( pxr::UsdAttribute ExposureAttr = SphereLight.GetExposureAttr() )
+						if (pxr::UsdAttribute ExposureAttr = SphereLight.GetExposureAttr())
 						{
 							ExposureAttr.Clear();
 						}
 
 						// For now we'll assume the light intensity units are constant and the user doesn't have any light intensity unit tracks...
 						ELightUnits Units = ELightUnits::Lumens;
-						if ( const ULocalLightComponent* LightComponent = Cast<const ULocalLightComponent>( &Component ) )
+						if (const ULocalLightComponent* LightComponent = Cast<const ULocalLightComponent>(&Component))
 						{
 							Units = LightComponent->IntensityUnits;
 						}
 
-						if ( Attr && RadiusAttr && ConeAngleAttr && ConeSoftnessAttr )
+						if (Attr && RadiusAttr && ConeAngleAttr && ConeSoftnessAttr)
 						{
-							Result.FloatWriter = [Attr, RadiusAttr, ConeAngleAttr, ConeSoftnessAttr, StageInfo, Units]( float UEValue, double UsdTimeCode )
+							Result.FloatWriter =
+								[Attr, RadiusAttr, ConeAngleAttr, ConeSoftnessAttr, StageInfo, Units](float UEValue, double UsdTimeCode)
 							{
-								const float UsdConeAngle = UsdUtils::GetUsdValue< float >( ConeAngleAttr, UsdTimeCode );
-								const float UsdConeSoftness = UsdUtils::GetUsdValue< float >( ConeSoftnessAttr, UsdTimeCode );
-								const float UsdRadius = UsdUtils::GetUsdValue< float >( RadiusAttr, UsdTimeCode );
+								const float UsdConeAngle = UsdUtils::GetUsdValue<float>(ConeAngleAttr, UsdTimeCode);
+								const float UsdConeSoftness = UsdUtils::GetUsdValue<float>(ConeSoftnessAttr, UsdTimeCode);
+								const float UsdRadius = UsdUtils::GetUsdValue<float>(RadiusAttr, UsdTimeCode);
 
 								float InnerConeAngle = 0.0f;
-								const float OuterConeAngle = UsdToUnreal::ConvertConeAngleSoftnessAttr( UsdConeAngle, UsdConeSoftness, InnerConeAngle );
-								const float SourceRadius = UsdToUnreal::ConvertDistance( StageInfo, UsdRadius );
+								const float OuterConeAngle = UsdToUnreal::ConvertConeAngleSoftnessAttr(UsdConeAngle, UsdConeSoftness, InnerConeAngle);
+								const float SourceRadius = UsdToUnreal::ConvertDistance(StageInfo, UsdRadius);
 
 								Attr.Set(
-									UnrealToUsd::ConvertSpotLightIntensityProperty( UEValue, OuterConeAngle, InnerConeAngle, SourceRadius, StageInfo, Units ),
+									UnrealToUsd::ConvertSpotLightIntensityProperty(
+										UEValue,
+										OuterConeAngle,
+										InnerConeAngle,
+										SourceRadius,
+										StageInfo,
+										Units
+									),
 									UsdTimeCode
 								);
 							};
 						}
 					}
-					else if ( PropertyPath == UnrealIdentifiers::OuterConeAnglePropertyName )
+					else if (PropertyPath == UnrealIdentifiers::OuterConeAnglePropertyName)
 					{
 						Attr = ShapingAPI.GetShapingConeAngleAttr();
-						if ( Attr )
+						if (Attr)
 						{
 							// InnerConeAngle is calculated based on ConeAngleAttr, so we need to refresh it
-							OutPropertyPathsToRefresh.Add( UnrealIdentifiers::InnerConeAnglePropertyName );
+							OutPropertyPathsToRefresh.Add(UnrealIdentifiers::InnerConeAnglePropertyName);
 
-							Result.FloatWriter = [Attr]( float UEValue, double UsdTimeCode )
+							Result.FloatWriter = [Attr](float UEValue, double UsdTimeCode)
 							{
-								Attr.Set( UEValue, UsdTimeCode );
+								Attr.Set(UEValue, UsdTimeCode);
 							};
 						}
 					}
-					else if ( PropertyPath == UnrealIdentifiers::InnerConeAnglePropertyName )
+					else if (PropertyPath == UnrealIdentifiers::InnerConeAnglePropertyName)
 					{
 						Attr = ShapingAPI.GetShapingConeSoftnessAttr();
 						pxr::UsdAttribute ConeAngleAttr = ShapingAPI.GetShapingConeAngleAttr();
 
-						if ( ConeAngleAttr && Attr )
+						if (ConeAngleAttr && Attr)
 						{
-							Result.FloatWriter = [Attr, ConeAngleAttr]( float UEValue, double UsdTimeCode )
+							Result.FloatWriter = [Attr, ConeAngleAttr](float UEValue, double UsdTimeCode)
 							{
-								const float UsdConeAngle = UsdUtils::GetUsdValue< float >( ConeAngleAttr, UsdTimeCode );
+								const float UsdConeAngle = UsdUtils::GetUsdValue<float>(ConeAngleAttr, UsdTimeCode);
 								const float OuterConeAngle = UsdConeAngle;
 
-								const float OutNewSoftness = UnrealToUsd::ConvertInnerConeAngleProperty( UEValue, OuterConeAngle );
-								Attr.Set( OutNewSoftness, UsdTimeCode );
+								const float OutNewSoftness = UnrealToUsd::ConvertInnerConeAngleProperty(UEValue, OuterConeAngle);
+								Attr.Set(OutNewSoftness, UsdTimeCode);
 							};
 						}
 					}
@@ -4281,65 +4380,68 @@ UnrealToUsd::FPropertyTrackWriter UnrealToUsd::CreatePropertyTrackWriter( const
 				// Just a point light
 				else
 				{
-					if ( PropertyPath == UnrealIdentifiers::IntensityPropertyName )
+					if (PropertyPath == UnrealIdentifiers::IntensityPropertyName)
 					{
 						Attr = SphereLight.GetIntensityAttr();
 						pxr::UsdAttribute RadiusAttr = SphereLight.GetRadiusAttr();
 
 						// Always clear exposure because we'll put all of our "light intensity" on the intensity attr and assume exposure
 						// is zero, as we can't manipulate something like that exposure directly from UE anyway
-						if ( pxr::UsdAttribute ExposureAttr = SphereLight.GetExposureAttr() )
+						if (pxr::UsdAttribute ExposureAttr = SphereLight.GetExposureAttr())
 						{
 							ExposureAttr.Clear();
 						}
 
 						// For now we'll assume the light intensity units are constant and the user doesn't have any light intensity unit tracks...
 						ELightUnits Units = ELightUnits::Lumens;
-						if ( const ULocalLightComponent* LightComponent = Cast<const ULocalLightComponent>( &Component ) )
+						if (const ULocalLightComponent* LightComponent = Cast<const ULocalLightComponent>(&Component))
 						{
 							Units = LightComponent->IntensityUnits;
 						}
 
-						if ( Attr && RadiusAttr )
+						if (Attr && RadiusAttr)
 						{
-							Result.FloatWriter = [Attr, RadiusAttr, StageInfo, Units]( float UEValue, double UsdTimeCode )
+							Result.FloatWriter = [Attr, RadiusAttr, StageInfo, Units](float UEValue, double UsdTimeCode)
 							{
-								const float SourceRadius = UsdToUnreal::ConvertDistance( StageInfo, UsdUtils::GetUsdValue< float >( RadiusAttr, UsdTimeCode ) );
-								Attr.Set( UnrealToUsd::ConvertPointLightIntensityProperty( UEValue, SourceRadius, StageInfo, Units ), UsdTimeCode );
+								const float SourceRadius = UsdToUnreal::ConvertDistance(
+									StageInfo,
+									UsdUtils::GetUsdValue<float>(RadiusAttr, UsdTimeCode)
+								);
+								Attr.Set(UnrealToUsd::ConvertPointLightIntensityProperty(UEValue, SourceRadius, StageInfo, Units), UsdTimeCode);
 							};
 						}
 					}
 				}
 			}
-			else if ( pxr::UsdLuxRectLight RectLight{ UsdPrim } )
+			else if (pxr::UsdLuxRectLight RectLight{UsdPrim})
 			{
-				if ( PropertyPath == UnrealIdentifiers::SourceWidthPropertyName )
+				if (PropertyPath == UnrealIdentifiers::SourceWidthPropertyName)
 				{
 					Attr = RectLight.GetWidthAttr();
-					if ( Attr )
+					if (Attr)
 					{
-						OutPropertyPathsToRefresh.Add( UnrealIdentifiers::IntensityPropertyName );
+						OutPropertyPathsToRefresh.Add(UnrealIdentifiers::IntensityPropertyName);
 
-						Result.FloatWriter = [Attr, StageInfo]( float UEValue, double UsdTimeCode )
+						Result.FloatWriter = [Attr, StageInfo](float UEValue, double UsdTimeCode)
 						{
-							Attr.Set( UnrealToUsd::ConvertDistance( StageInfo, UEValue ), UsdTimeCode );
+							Attr.Set(UnrealToUsd::ConvertDistance(StageInfo, UEValue), UsdTimeCode);
 						};
 					}
 				}
-				else if ( PropertyPath == UnrealIdentifiers::SourceHeightPropertyName )
+				else if (PropertyPath == UnrealIdentifiers::SourceHeightPropertyName)
 				{
 					Attr = RectLight.GetHeightAttr();
-					if ( Attr )
+					if (Attr)
 					{
-						OutPropertyPathsToRefresh.Add( UnrealIdentifiers::IntensityPropertyName );
+						OutPropertyPathsToRefresh.Add(UnrealIdentifiers::IntensityPropertyName);
 
-						Result.FloatWriter = [Attr, StageInfo]( float UEValue, double UsdTimeCode )
+						Result.FloatWriter = [Attr, StageInfo](float UEValue, double UsdTimeCode)
 						{
-							Attr.Set( UnrealToUsd::ConvertDistance( StageInfo, UEValue ), UsdTimeCode );
+							Attr.Set(UnrealToUsd::ConvertDistance(StageInfo, UEValue), UsdTimeCode);
 						};
 					}
 				}
-				else if ( PropertyPath == UnrealIdentifiers::IntensityPropertyName )
+				else if (PropertyPath == UnrealIdentifiers::IntensityPropertyName)
 				{
 					Attr = RectLight.GetIntensityAttr();
 					pxr::UsdAttribute WidthAttr = RectLight.GetWidthAttr();
@@ -4347,110 +4449,110 @@ UnrealToUsd::FPropertyTrackWriter UnrealToUsd::CreatePropertyTrackWriter( const
 
 					// Always clear exposure because we'll put all of our "light intensity" on the intensity attr and assume exposure
 					// is zero, as we can't manipulate something like that exposure directly from UE anyway
-					if ( pxr::UsdAttribute ExposureAttr = RectLight.GetExposureAttr() )
+					if (pxr::UsdAttribute ExposureAttr = RectLight.GetExposureAttr())
 					{
 						ExposureAttr.Clear();
 					}
 
 					// For now we'll assume the light intensity units are constant and the user doesn't have any light intensity unit tracks...
 					ELightUnits Units = ELightUnits::Lumens;
-					if ( const ULocalLightComponent* LightComponent = Cast<const ULocalLightComponent>( &Component ) )
+					if (const ULocalLightComponent* LightComponent = Cast<const ULocalLightComponent>(&Component))
 					{
 						Units = LightComponent->IntensityUnits;
 					}
 
-					if ( Attr && WidthAttr && HeightAttr )
+					if (Attr && WidthAttr && HeightAttr)
 					{
-						Result.FloatWriter = [Attr, WidthAttr, HeightAttr, StageInfo, Units]( float UEValue, double UsdTimeCode )
+						Result.FloatWriter = [Attr, WidthAttr, HeightAttr, StageInfo, Units](float UEValue, double UsdTimeCode)
 						{
-							const float Width = UsdToUnreal::ConvertDistance( StageInfo, UsdUtils::GetUsdValue< float >( WidthAttr, UsdTimeCode ) );
-							const float Height = UsdToUnreal::ConvertDistance( StageInfo, UsdUtils::GetUsdValue< float >( HeightAttr, UsdTimeCode ) );
+							const float Width = UsdToUnreal::ConvertDistance(StageInfo, UsdUtils::GetUsdValue<float>(WidthAttr, UsdTimeCode));
+							const float Height = UsdToUnreal::ConvertDistance(StageInfo, UsdUtils::GetUsdValue<float>(HeightAttr, UsdTimeCode));
 
-							Attr.Set( UnrealToUsd::ConvertRectLightIntensityProperty( UEValue, Width, Height, StageInfo, Units ), UsdTimeCode );
+							Attr.Set(UnrealToUsd::ConvertRectLightIntensityProperty(UEValue, Width, Height, StageInfo, Units), UsdTimeCode);
 						};
 					}
 				}
 			}
-			else if ( pxr::UsdLuxDiskLight DiskLight{ UsdPrim } )
+			else if (pxr::UsdLuxDiskLight DiskLight{UsdPrim})
 			{
-				if ( PropertyPath == UnrealIdentifiers::SourceWidthPropertyName || PropertyPath == UnrealIdentifiers::SourceHeightPropertyName )
+				if (PropertyPath == UnrealIdentifiers::SourceWidthPropertyName || PropertyPath == UnrealIdentifiers::SourceHeightPropertyName)
 				{
 					Attr = DiskLight.GetRadiusAttr();
-					if ( Attr )
+					if (Attr)
 					{
-						OutPropertyPathsToRefresh.Add( UnrealIdentifiers::IntensityPropertyName );
+						OutPropertyPathsToRefresh.Add(UnrealIdentifiers::IntensityPropertyName);
 
 						// Resync the other to match this one after we bake it, effectively always enforcing the UE rect light into a square shape
-						OutPropertyPathsToRefresh.Add( PropertyPath == UnrealIdentifiers::SourceWidthPropertyName
-							? UnrealIdentifiers::SourceHeightPropertyName
-							: UnrealIdentifiers::SourceWidthPropertyName
+						OutPropertyPathsToRefresh.Add(
+							PropertyPath == UnrealIdentifiers::SourceWidthPropertyName ? UnrealIdentifiers::SourceHeightPropertyName
+																					   : UnrealIdentifiers::SourceWidthPropertyName
 						);
 
-						Result.FloatWriter = [Attr, StageInfo]( float UEValue, double UsdTimeCode )
+						Result.FloatWriter = [Attr, StageInfo](float UEValue, double UsdTimeCode)
 						{
-							Attr.Set( UnrealToUsd::ConvertDistance( StageInfo, UEValue * 0.5f ), UsdTimeCode );
+							Attr.Set(UnrealToUsd::ConvertDistance(StageInfo, UEValue * 0.5f), UsdTimeCode);
 						};
 					}
 				}
-				else if ( PropertyPath == UnrealIdentifiers::IntensityPropertyName )
+				else if (PropertyPath == UnrealIdentifiers::IntensityPropertyName)
 				{
 					Attr = RectLight.GetIntensityAttr();
 					pxr::UsdAttribute RadiusAttr = DiskLight.GetRadiusAttr();
 
 					// Always clear exposure because we'll put all of our "light intensity" on the intensity attr and assume exposure
 					// is zero, as we can't manipulate something like that exposure directly from UE anyway
-					if ( pxr::UsdAttribute ExposureAttr = RectLight.GetExposureAttr() )
+					if (pxr::UsdAttribute ExposureAttr = RectLight.GetExposureAttr())
 					{
 						ExposureAttr.Clear();
 					}
 
 					// For now we'll assume the light intensity units are constant and the user doesn't have any light intensity unit tracks...
 					ELightUnits Units = ELightUnits::Lumens;
-					if ( const ULocalLightComponent* LightComponent = Cast<const ULocalLightComponent>( &Component ) )
+					if (const ULocalLightComponent* LightComponent = Cast<const ULocalLightComponent>(&Component))
 					{
 						Units = LightComponent->IntensityUnits;
 					}
 
-					if ( Attr && RadiusAttr )
+					if (Attr && RadiusAttr)
 					{
-						Result.FloatWriter = [Attr, RadiusAttr, StageInfo, Units]( float UEValue, double UsdTimeCode )
+						Result.FloatWriter = [Attr, RadiusAttr, StageInfo, Units](float UEValue, double UsdTimeCode)
 						{
-							const float Radius = UsdToUnreal::ConvertDistance( StageInfo, UsdUtils::GetUsdValue< float >( RadiusAttr, UsdTimeCode ) );
+							const float Radius = UsdToUnreal::ConvertDistance(StageInfo, UsdUtils::GetUsdValue<float>(RadiusAttr, UsdTimeCode));
 
-							Attr.Set( UnrealToUsd::ConvertRectLightIntensityProperty( UEValue, Radius, StageInfo, Units ), UsdTimeCode );
+							Attr.Set(UnrealToUsd::ConvertRectLightIntensityProperty(UEValue, Radius, StageInfo, Units), UsdTimeCode);
 						};
 					}
 				}
 			}
-			else if ( pxr::UsdLuxDistantLight DistantLight{ UsdPrim } )
+			else if (pxr::UsdLuxDistantLight DistantLight{UsdPrim})
 			{
-				if ( PropertyPath == UnrealIdentifiers::LightSourceAnglePropertyName )
+				if (PropertyPath == UnrealIdentifiers::LightSourceAnglePropertyName)
 				{
 					Attr = DistantLight.GetAngleAttr();
-					if ( Attr )
+					if (Attr)
 					{
-						Result.FloatWriter = [Attr]( float UEValue, double UsdTimeCode )
+						Result.FloatWriter = [Attr](float UEValue, double UsdTimeCode)
 						{
-							Attr.Set( UEValue, UsdTimeCode );
+							Attr.Set(UEValue, UsdTimeCode);
 						};
 					}
 				}
-				else if ( PropertyPath == UnrealIdentifiers::IntensityPropertyName )
+				else if (PropertyPath == UnrealIdentifiers::IntensityPropertyName)
 				{
 					Attr = DistantLight.GetIntensityAttr();
 
 					// Always clear exposure because we'll put all of our "light intensity" on the intensity attr and assume exposure
 					// is zero, as we can't manipulate something like that exposure directly from UE anyway
-					if ( pxr::UsdAttribute ExposureAttr = RectLight.GetExposureAttr() )
+					if (pxr::UsdAttribute ExposureAttr = RectLight.GetExposureAttr())
 					{
 						ExposureAttr.Clear();
 					}
 
-					if ( Attr )
+					if (Attr)
 					{
-						Result.FloatWriter = [Attr]( float UEValue, double UsdTimeCode )
+						Result.FloatWriter = [Attr](float UEValue, double UsdTimeCode)
 						{
-							Attr.Set( UnrealToUsd::ConvertLightIntensityProperty( UEValue ), UsdTimeCode );
+							Attr.Set(UnrealToUsd::ConvertLightIntensityProperty(UEValue), UsdTimeCode);
 						};
 					}
 				}
@@ -4492,11 +4594,13 @@ UnrealToUsd::FPropertyTrackWriter UnrealToUsd::CreatePropertyTrackWriter( const
 					pxr::GfVec3f UsdMin{
 						FMath::Min(UEBoundsMinUsdSpace[0], UEBoundsMaxUsdSpace[0]),
 						FMath::Min(UEBoundsMinUsdSpace[1], UEBoundsMaxUsdSpace[1]),
-						FMath::Min(UEBoundsMinUsdSpace[2], UEBoundsMaxUsdSpace[2])};
+						FMath::Min(UEBoundsMinUsdSpace[2], UEBoundsMaxUsdSpace[2])
+					};
 					pxr::GfVec3f UsdMax{
 						FMath::Max(UEBoundsMinUsdSpace[0], UEBoundsMaxUsdSpace[0]),
 						FMath::Max(UEBoundsMinUsdSpace[1], UEBoundsMaxUsdSpace[1]),
-						FMath::Max(UEBoundsMinUsdSpace[2], UEBoundsMaxUsdSpace[2])};
+						FMath::Max(UEBoundsMinUsdSpace[2], UEBoundsMaxUsdSpace[2])
+					};
 					pxr::VtArray<pxr::GfVec3f> Extents{UsdMin, UsdMax};
 
 					Attr.Set(Extents, UsdTimeCode);
@@ -4505,13 +4609,13 @@ UnrealToUsd::FPropertyTrackWriter UnrealToUsd::CreatePropertyTrackWriter( const
 		}
 	}
 
-	if ( Attr )
+	if (Attr)
 	{
 		std::vector<double> TimeSamples;
-		Attr.GetTimeSamples( &TimeSamples );
-		for ( double TimeSample : TimeSamples )
+		Attr.GetTimeSamples(&TimeSamples);
+		for (double TimeSample : TimeSamples)
 		{
-			Attr.ClearAtTime( TimeSample );
+			Attr.ClearAtTime(TimeSample);
 		}
 
 		// Note that we must do this only after the change block is destroyed!
@@ -4526,133 +4630,142 @@ UnrealToUsd::FPropertyTrackWriter UnrealToUsd::CreatePropertyTrackWriter( const
 	return Result;
 }
 
-bool UnrealToUsd::ConvertXformable( const UMovieScene3DTransformTrack& MovieSceneTrack, pxr::UsdPrim& UsdPrim, const FMovieSceneSequenceTransform& SequenceTransform )
+bool UnrealToUsd::ConvertXformable(
+	const UMovieScene3DTransformTrack& MovieSceneTrack,
+	pxr::UsdPrim& UsdPrim,
+	const FMovieSceneSequenceTransform& SequenceTransform
+)
 {
-	if ( !UsdPrim )
+	if (!UsdPrim)
 	{
 		return false;
 	}
 
-	const FUsdStageInfo StageInfo( UsdPrim.GetStage() );
+	const FUsdStageInfo StageInfo(UsdPrim.GetStage());
 
-	const UMovieScene* MovieScene = MovieSceneTrack.GetTypedOuter< UMovieScene >();
-	if ( !MovieScene )
+	const UMovieScene* MovieScene = MovieSceneTrack.GetTypedOuter<UMovieScene>();
+	if (!MovieScene)
 	{
 		return false;
 	}
 
 	FScopedUsdAllocs UsdAllocs;
 
-	pxr::UsdGeomXformable Xformable( UsdPrim );
-	if ( !Xformable )
+	pxr::UsdGeomXformable Xformable(UsdPrim);
+	if (!Xformable)
 	{
 		return false;
 	}
 
-	UMovieScene3DTransformSection* TransformSection = Cast< UMovieScene3DTransformSection >( const_cast< UMovieScene3DTransformTrack& >( MovieSceneTrack ).FindSection( 0 ) );
-	if ( !TransformSection )
+	UMovieScene3DTransformSection* TransformSection = Cast<UMovieScene3DTransformSection>(
+		const_cast<UMovieScene3DTransformTrack&>(MovieSceneTrack).FindSection(0)
+	);
+	if (!TransformSection)
 	{
 		return false;
 	}
 
-	const TRange< FFrameNumber > PlaybackRange = MovieScene->GetPlaybackRange();
+	const TRange<FFrameNumber> PlaybackRange = MovieScene->GetPlaybackRange();
 	const FFrameRate Resolution = MovieScene->GetTickResolution();
 	const FFrameRate DisplayRate = MovieScene->GetDisplayRate();
 
 	const double StageTimeCodesPerSecond = UsdPrim.GetStage()->GetTimeCodesPerSecond();
-	const FFrameRate StageFrameRate( StageTimeCodesPerSecond, 1 );
+	const FFrameRate StageFrameRate(StageTimeCodesPerSecond, 1);
 
-	auto EvaluateChannel = [ &PlaybackRange, &Resolution, &DisplayRate, &SequenceTransform ]( const FMovieSceneDoubleChannel* Channel, double DefaultValue ) -> TArray< TPair< FFrameNumber, float > >
+	auto EvaluateChannel = [&PlaybackRange,
+							&Resolution,
+							&DisplayRate,
+							&SequenceTransform](const FMovieSceneDoubleChannel* Channel, double DefaultValue) -> TArray<TPair<FFrameNumber, float>>
 	{
-		TArray< TPair< FFrameNumber, float > > Values;
+		TArray<TPair<FFrameNumber, float>> Values;
 
-		if ( PlaybackRange.HasLowerBound() && PlaybackRange.HasUpperBound() )
+		if (PlaybackRange.HasLowerBound() && PlaybackRange.HasUpperBound())
 		{
-			const FFrameTime Interval = FFrameRate::TransformTime( 1, DisplayRate, Resolution );
-			const FFrameNumber StartFrame = UE::MovieScene::DiscreteInclusiveLower( PlaybackRange );
-			const FFrameNumber EndFrame = UE::MovieScene::DiscreteExclusiveUpper( PlaybackRange );
+			const FFrameTime Interval = FFrameRate::TransformTime(1, DisplayRate, Resolution);
+			const FFrameNumber StartFrame = UE::MovieScene::DiscreteInclusiveLower(PlaybackRange);
+			const FFrameNumber EndFrame = UE::MovieScene::DiscreteExclusiveUpper(PlaybackRange);
 
-			for ( FFrameTime EvalTime = StartFrame; EvalTime < EndFrame; EvalTime += Interval )
+			for (FFrameTime EvalTime = StartFrame; EvalTime < EndFrame; EvalTime += Interval)
 			{
-				FFrameNumber KeyTime = FFrameRate::Snap( EvalTime, Resolution, DisplayRate ).FloorToFrame();
+				FFrameNumber KeyTime = FFrameRate::Snap(EvalTime, Resolution, DisplayRate).FloorToFrame();
 
 				double Result = DefaultValue;
-				if ( Channel )
+				if (Channel)
 				{
-					Result = Channel->GetDefault().Get( DefaultValue );
-					Channel->Evaluate( KeyTime, Result );
+					Result = Channel->GetDefault().Get(DefaultValue);
+					Channel->Evaluate(KeyTime, Result);
 				}
 
-				FFrameTime GlobalEvalTime( KeyTime );
+				FFrameTime GlobalEvalTime(KeyTime);
 				GlobalEvalTime *= SequenceTransform.InverseNoLooping();
 
-				Values.Emplace( GlobalEvalTime.GetFrame(), Result );
+				Values.Emplace(GlobalEvalTime.GetFrame(), Result);
 			}
 		}
 
 		return Values;
 	};
 
-	TArrayView< FMovieSceneDoubleChannel* > Channels = TransformSection->GetChannelProxy().GetChannels< FMovieSceneDoubleChannel >();
-	check( Channels.Num() >= 9 );
+	TArrayView<FMovieSceneDoubleChannel*> Channels = TransformSection->GetChannelProxy().GetChannels<FMovieSceneDoubleChannel>();
+	check(Channels.Num() >= 9);
 
-	auto GetChannel = [ &Channels ]( const int32 ChannelIndex ) -> const FMovieSceneDoubleChannel*
+	auto GetChannel = [&Channels](const int32 ChannelIndex) -> const FMovieSceneDoubleChannel*
 	{
-		if ( Channels.IsValidIndex( ChannelIndex ) )
+		if (Channels.IsValidIndex(ChannelIndex))
 		{
-			return Channels[ ChannelIndex ];
+			return Channels[ChannelIndex];
 		}
 		else
 		{
-			return  nullptr;
+			return nullptr;
 		}
 	};
 
 	// Translation
-	TArray< TPair< FFrameNumber, float > > LocationValuesX = EvaluateChannel( GetChannel(0), 0.0 );
-	TArray< TPair< FFrameNumber, float > > LocationValuesY = EvaluateChannel( GetChannel(1), 0.0 );
-	TArray< TPair< FFrameNumber, float > > LocationValuesZ = EvaluateChannel( GetChannel(2), 0.0 );
+	TArray<TPair<FFrameNumber, float>> LocationValuesX = EvaluateChannel(GetChannel(0), 0.0);
+	TArray<TPair<FFrameNumber, float>> LocationValuesY = EvaluateChannel(GetChannel(1), 0.0);
+	TArray<TPair<FFrameNumber, float>> LocationValuesZ = EvaluateChannel(GetChannel(2), 0.0);
 
 	// Rotation
-	TArray< TPair< FFrameNumber, float > > RotationValuesX = EvaluateChannel( GetChannel(3), 0.0 );
-	TArray< TPair< FFrameNumber, float > > RotationValuesY = EvaluateChannel( GetChannel(4), 0.0 );
-	TArray< TPair< FFrameNumber, float > > RotationValuesZ = EvaluateChannel( GetChannel(5), 0.0 );
+	TArray<TPair<FFrameNumber, float>> RotationValuesX = EvaluateChannel(GetChannel(3), 0.0);
+	TArray<TPair<FFrameNumber, float>> RotationValuesY = EvaluateChannel(GetChannel(4), 0.0);
+	TArray<TPair<FFrameNumber, float>> RotationValuesZ = EvaluateChannel(GetChannel(5), 0.0);
 
 	// Scale
-	TArray< TPair< FFrameNumber, float > > ScaleValuesX = EvaluateChannel( GetChannel(6), 1.0 );
-	TArray< TPair< FFrameNumber, float > > ScaleValuesY = EvaluateChannel( GetChannel(7), 1.0 );
-	TArray< TPair< FFrameNumber, float > > ScaleValuesZ = EvaluateChannel( GetChannel(8), 1.0 );
+	TArray<TPair<FFrameNumber, float>> ScaleValuesX = EvaluateChannel(GetChannel(6), 1.0);
+	TArray<TPair<FFrameNumber, float>> ScaleValuesY = EvaluateChannel(GetChannel(7), 1.0);
+	TArray<TPair<FFrameNumber, float>> ScaleValuesZ = EvaluateChannel(GetChannel(8), 1.0);
 
 	bool bIsDataOutOfSync = false;
 	{
 		int32 ValueIndex = 0;
 
-		FFrameTime UsdStartTime = FFrameRate::TransformTime( PlaybackRange.GetLowerBoundValue(), Resolution, StageFrameRate );
-		FFrameTime UsdEndTime = FFrameRate::TransformTime( PlaybackRange.GetUpperBoundValue(), Resolution, StageFrameRate );
+		FFrameTime UsdStartTime = FFrameRate::TransformTime(PlaybackRange.GetLowerBoundValue(), Resolution, StageFrameRate);
+		FFrameTime UsdEndTime = FFrameRate::TransformTime(PlaybackRange.GetUpperBoundValue(), Resolution, StageFrameRate);
 
-		std::vector< double > UsdTimeSamples;
-		if ( LocationValuesX.Num() > 0 || ( Xformable.GetTimeSamples( &UsdTimeSamples ) && UsdTimeSamples.size() > 0 ) )
+		std::vector<double> UsdTimeSamples;
+		if (LocationValuesX.Num() > 0 || (Xformable.GetTimeSamples(&UsdTimeSamples) && UsdTimeSamples.size() > 0))
 		{
-			bIsDataOutOfSync = ( UsdTimeSamples.size() != LocationValuesX.Num() );
+			bIsDataOutOfSync = (UsdTimeSamples.size() != LocationValuesX.Num());
 
-			if ( !bIsDataOutOfSync )
+			if (!bIsDataOutOfSync)
 			{
-				for ( const TPair< FFrameNumber, float >& Value : LocationValuesX )
+				for (const TPair<FFrameNumber, float>& Value : LocationValuesX)
 				{
-					FFrameTime UsdFrameTime = FFrameRate::TransformTime( Value.Key, Resolution, StageFrameRate );
+					FFrameTime UsdFrameTime = FFrameRate::TransformTime(Value.Key, Resolution, StageFrameRate);
 
-					FVector Location( LocationValuesX[ ValueIndex ].Value, LocationValuesY[ ValueIndex ].Value, LocationValuesZ[ ValueIndex ].Value );
-					FRotator Rotation( RotationValuesY[ ValueIndex ].Value, RotationValuesZ[ ValueIndex ].Value, RotationValuesX[ ValueIndex ].Value );
-					FVector Scale( ScaleValuesX[ ValueIndex ].Value, ScaleValuesY[ ValueIndex ].Value, ScaleValuesZ[ ValueIndex ].Value );
+					FVector Location(LocationValuesX[ValueIndex].Value, LocationValuesY[ValueIndex].Value, LocationValuesZ[ValueIndex].Value);
+					FRotator Rotation(RotationValuesY[ValueIndex].Value, RotationValuesZ[ValueIndex].Value, RotationValuesX[ValueIndex].Value);
+					FVector Scale(ScaleValuesX[ValueIndex].Value, ScaleValuesY[ValueIndex].Value, ScaleValuesZ[ValueIndex].Value);
 
-					FTransform Transform( Rotation, Location, Scale );
-					pxr::GfMatrix4d UsdTransform = UnrealToUsd::ConvertTransform( StageInfo, Transform );
+					FTransform Transform(Rotation, Location, Scale);
+					pxr::GfMatrix4d UsdTransform = UnrealToUsd::ConvertTransform(StageInfo, Transform);
 
 					pxr::GfMatrix4d UsdMatrix;
 					bool bResetXFormStack = false;
-					Xformable.GetLocalTransformation( &UsdMatrix, &bResetXFormStack, UsdFrameTime.AsDecimal() );
+					Xformable.GetLocalTransformation(&UsdMatrix, &bResetXFormStack, UsdFrameTime.AsDecimal());
 
-					if ( !pxr::GfIsClose( UsdMatrix, UsdTransform, THRESH_POINTS_ARE_NEAR ) )
+					if (!pxr::GfIsClose(UsdMatrix, UsdTransform, THRESH_POINTS_ARE_NEAR))
 					{
 						bIsDataOutOfSync = true;
 						break;
@@ -4664,53 +4777,53 @@ bool UnrealToUsd::ConvertXformable( const UMovieScene3DTransformTrack& MovieScen
 		}
 	}
 
-	if ( bIsDataOutOfSync )
+	if (bIsDataOutOfSync)
 	{
-		if ( pxr::UsdGeomXformOp TransformOp = UE::USDPrimConversionImpl::Private::ForceMatrixXform( Xformable ) )
+		if (pxr::UsdGeomXformOp TransformOp = UE::USDPrimConversionImpl::Private::ForceMatrixXform(Xformable))
 		{
-			TransformOp.GetAttr().Clear(); // Clear existing transform data
+			TransformOp.GetAttr().Clear();	  // Clear existing transform data
 		}
 
 		pxr::SdfChangeBlock ChangeBlock;
 
 		// Compensate different orientation for light or camera components
 		FTransform CameraCompensation = FTransform::Identity;
-		if ( UsdPrim.IsA< pxr::UsdGeomCamera >() || UsdPrim.HasAPI< pxr::UsdLuxLightAPI >() )
+		if (UsdPrim.IsA<pxr::UsdGeomCamera>() || UsdPrim.HasAPI<pxr::UsdLuxLightAPI>())
 		{
-			CameraCompensation = FTransform( FRotator( 0.0f, 90.0f, 0.0f ) );
+			CameraCompensation = FTransform(FRotator(0.0f, 90.0f, 0.0f));
 
-			if ( StageInfo.UpAxis == EUsdUpAxis::ZAxis )
+			if (StageInfo.UpAxis == EUsdUpAxis::ZAxis)
 			{
-				CameraCompensation *= FTransform( FRotator( 90.0f, 0.0f, 0.0f ) );
+				CameraCompensation *= FTransform(FRotator(90.0f, 0.0f, 0.0f));
 			}
 		}
 
 		// Invert compensation applied to parent if it's a light or camera component
 		FTransform InverseCameraCompensation = FTransform::Identity;
-		if ( pxr::UsdPrim ParentPrim = UsdPrim.GetParent() )
+		if (pxr::UsdPrim ParentPrim = UsdPrim.GetParent())
 		{
-			if ( ParentPrim.IsA<pxr::UsdGeomCamera>() || ParentPrim.HasAPI<pxr::UsdLuxLightAPI>() )
+			if (ParentPrim.IsA<pxr::UsdGeomCamera>() || ParentPrim.HasAPI<pxr::UsdLuxLightAPI>())
 			{
-				InverseCameraCompensation = FTransform( FRotator( 0.0f, 90.f, 0.0f ) );
+				InverseCameraCompensation = FTransform(FRotator(0.0f, 90.f, 0.0f));
 
-				if ( StageInfo.UpAxis == EUsdUpAxis::ZAxis )
+				if (StageInfo.UpAxis == EUsdUpAxis::ZAxis)
 				{
-					InverseCameraCompensation *= FTransform( FRotator( 90.0f, 0.f, 0.0f ) );
+					InverseCameraCompensation *= FTransform(FRotator(90.0f, 0.f, 0.0f));
 				}
 			}
 		}
 
 		int32 ValueIndex = 0;
-		for ( const TPair< FFrameNumber, float >& Value : LocationValuesX )
+		for (const TPair<FFrameNumber, float>& Value : LocationValuesX)
 		{
-			FFrameTime UsdFrameTime = FFrameRate::TransformTime( Value.Key, Resolution, StageFrameRate );
+			FFrameTime UsdFrameTime = FFrameRate::TransformTime(Value.Key, Resolution, StageFrameRate);
 
-			FVector Location( LocationValuesX[ ValueIndex ].Value, LocationValuesY[ ValueIndex ].Value, LocationValuesZ[ ValueIndex ].Value );
-			FRotator Rotation( RotationValuesY[ ValueIndex ].Value, RotationValuesZ[ ValueIndex ].Value, RotationValuesX[ ValueIndex ].Value );
-			FVector Scale( ScaleValuesX[ ValueIndex ].Value, ScaleValuesY[ ValueIndex ].Value, ScaleValuesZ[ ValueIndex ].Value );
+			FVector Location(LocationValuesX[ValueIndex].Value, LocationValuesY[ValueIndex].Value, LocationValuesZ[ValueIndex].Value);
+			FRotator Rotation(RotationValuesY[ValueIndex].Value, RotationValuesZ[ValueIndex].Value, RotationValuesX[ValueIndex].Value);
+			FVector Scale(ScaleValuesX[ValueIndex].Value, ScaleValuesY[ValueIndex].Value, ScaleValuesZ[ValueIndex].Value);
 
-			FTransform Transform( Rotation, Location, Scale );
-			ConvertXformable( CameraCompensation * Transform * InverseCameraCompensation.Inverse(), UsdPrim, UsdFrameTime.AsDecimal() );
+			FTransform Transform(Rotation, Location, Scale);
+			ConvertXformable(CameraCompensation * Transform * InverseCameraCompensation.Inverse(), UsdPrim, UsdFrameTime.AsDecimal());
 
 			++ValueIndex;
 		}
@@ -4719,15 +4832,15 @@ bool UnrealToUsd::ConvertXformable( const UMovieScene3DTransformTrack& MovieScen
 	return true;
 }
 
-TArray<UE::FUsdAttribute> UnrealToUsd::GetAttributesForProperty( const UE::FUsdPrim& Prim, const FName& PropertyPath )
+TArray<UE::FUsdAttribute> UnrealToUsd::GetAttributesForProperty(const UE::FUsdPrim& Prim, const FName& PropertyPath)
 {
 	using namespace UnrealIdentifiers;
 
 	FScopedUsdAllocs Allocs;
-	pxr::UsdPrim UsdPrim{ Prim };
+	pxr::UsdPrim UsdPrim{Prim};
 
 	// Common attributes
-	if ( PropertyPath == TransformPropertyName )
+	if (PropertyPath == TransformPropertyName)
 	{
 		TArray<UE::FUsdAttribute> Attrs;
 
@@ -4776,143 +4889,140 @@ TArray<UE::FUsdAttribute> UnrealToUsd::GetAttributesForProperty( const UE::FUsdP
 
 		return Attrs;
 	}
-	if ( PropertyPath == HiddenInGamePropertyName )
+	if (PropertyPath == HiddenInGamePropertyName)
 	{
-		return { UE::FUsdAttribute{ UsdPrim.GetAttribute( pxr::UsdGeomTokens->visibility ) } };
+		return {UE::FUsdAttribute{UsdPrim.GetAttribute(pxr::UsdGeomTokens->visibility)}};
 	}
 
 	// Camera attributes
-	else if ( PropertyPath == CurrentFocalLengthPropertyName )
+	else if (PropertyPath == CurrentFocalLengthPropertyName)
 	{
-		return { UE::FUsdAttribute{ UsdPrim.GetAttribute( pxr::UsdGeomTokens->focalLength ) } };
+		return {UE::FUsdAttribute{UsdPrim.GetAttribute(pxr::UsdGeomTokens->focalLength)}};
 	}
-	else if ( PropertyPath == ManualFocusDistancePropertyName )
+	else if (PropertyPath == ManualFocusDistancePropertyName)
 	{
-		return { UE::FUsdAttribute{ UsdPrim.GetAttribute( pxr::UsdGeomTokens->focusDistance ) } };
+		return {UE::FUsdAttribute{UsdPrim.GetAttribute(pxr::UsdGeomTokens->focusDistance)}};
 	}
-	else if ( PropertyPath == CurrentAperturePropertyName )
+	else if (PropertyPath == CurrentAperturePropertyName)
 	{
-		return { UE::FUsdAttribute{ UsdPrim.GetAttribute( pxr::UsdGeomTokens->fStop ) } };
+		return {UE::FUsdAttribute{UsdPrim.GetAttribute(pxr::UsdGeomTokens->fStop)}};
 	}
-	else if ( PropertyPath == SensorWidthPropertyName )
+	else if (PropertyPath == SensorWidthPropertyName)
 	{
-		return { UE::FUsdAttribute{ UsdPrim.GetAttribute( pxr::UsdGeomTokens->horizontalAperture ) } };
+		return {UE::FUsdAttribute{UsdPrim.GetAttribute(pxr::UsdGeomTokens->horizontalAperture)}};
 	}
-	else if ( PropertyPath == SensorHeightPropertyName )
+	else if (PropertyPath == SensorHeightPropertyName)
 	{
-		return { UE::FUsdAttribute{ UsdPrim.GetAttribute( pxr::UsdGeomTokens->verticalAperture ) } };
+		return {UE::FUsdAttribute{UsdPrim.GetAttribute(pxr::UsdGeomTokens->verticalAperture)}};
 	}
 
 	// Light attributes
-	else if ( PropertyPath == IntensityPropertyName )
+	else if (PropertyPath == IntensityPropertyName)
 	{
-		if ( UsdPrim.IsA<pxr::UsdLuxRectLight>() )
+		if (UsdPrim.IsA<pxr::UsdLuxRectLight>())
 		{
 			return {
-				UE::FUsdAttribute{ UsdPrim.GetAttribute( pxr::UsdLuxTokens->inputsIntensity ) },
-				UE::FUsdAttribute{ UsdPrim.GetAttribute( pxr::UsdLuxTokens->inputsExposure ) },
-				UE::FUsdAttribute{ UsdPrim.GetAttribute( pxr::UsdLuxTokens->inputsWidth ) },
-				UE::FUsdAttribute{ UsdPrim.GetAttribute( pxr::UsdLuxTokens->inputsHeight ) }
+				UE::FUsdAttribute{UsdPrim.GetAttribute(pxr::UsdLuxTokens->inputsIntensity)},
+				UE::FUsdAttribute{UsdPrim.GetAttribute(pxr::UsdLuxTokens->inputsExposure)},
+				UE::FUsdAttribute{UsdPrim.GetAttribute(pxr::UsdLuxTokens->inputsWidth)},
+				UE::FUsdAttribute{UsdPrim.GetAttribute(pxr::UsdLuxTokens->inputsHeight)}
 			};
 		}
-		else if ( UsdPrim.IsA<pxr::UsdLuxDiskLight>() )
+		else if (UsdPrim.IsA<pxr::UsdLuxDiskLight>())
 		{
 			return {
-				UE::FUsdAttribute{ UsdPrim.GetAttribute( pxr::UsdLuxTokens->inputsIntensity ) },
-				UE::FUsdAttribute{ UsdPrim.GetAttribute( pxr::UsdLuxTokens->inputsExposure ) },
-				UE::FUsdAttribute{ UsdPrim.GetAttribute( pxr::UsdLuxTokens->inputsRadius ) }
+				UE::FUsdAttribute{UsdPrim.GetAttribute(pxr::UsdLuxTokens->inputsIntensity)},
+				UE::FUsdAttribute{UsdPrim.GetAttribute(pxr::UsdLuxTokens->inputsExposure)},
+				UE::FUsdAttribute{UsdPrim.GetAttribute(pxr::UsdLuxTokens->inputsRadius)}
 			};
 		}
-		else if ( UsdPrim.IsA<pxr::UsdLuxDistantLight>() )
+		else if (UsdPrim.IsA<pxr::UsdLuxDistantLight>())
 		{
 			return {
-				UE::FUsdAttribute{ UsdPrim.GetAttribute( pxr::UsdLuxTokens->inputsIntensity ) },
-				UE::FUsdAttribute{ UsdPrim.GetAttribute( pxr::UsdLuxTokens->inputsExposure ) }
+				UE::FUsdAttribute{UsdPrim.GetAttribute(pxr::UsdLuxTokens->inputsIntensity)},
+				UE::FUsdAttribute{UsdPrim.GetAttribute(pxr::UsdLuxTokens->inputsExposure)}
 			};
 		}
-		else if ( UsdPrim.IsA<pxr::UsdLuxSphereLight>() )
+		else if (UsdPrim.IsA<pxr::UsdLuxSphereLight>())
 		{
-			if ( UsdPrim.HasAPI<pxr::UsdLuxShapingAPI>() )
+			if (UsdPrim.HasAPI<pxr::UsdLuxShapingAPI>())
 			{
 				return {
-					UE::FUsdAttribute{ UsdPrim.GetAttribute( pxr::UsdLuxTokens->inputsIntensity ) },
-					UE::FUsdAttribute{ UsdPrim.GetAttribute( pxr::UsdLuxTokens->inputsExposure ) },
-					UE::FUsdAttribute{ UsdPrim.GetAttribute( pxr::UsdLuxTokens->inputsRadius ) },
-					UE::FUsdAttribute{ UsdPrim.GetAttribute( pxr::UsdLuxTokens->inputsShapingConeAngle ) },
-					UE::FUsdAttribute{ UsdPrim.GetAttribute( pxr::UsdLuxTokens->inputsShapingConeSoftness ) }
+					UE::FUsdAttribute{UsdPrim.GetAttribute(pxr::UsdLuxTokens->inputsIntensity)},
+					UE::FUsdAttribute{UsdPrim.GetAttribute(pxr::UsdLuxTokens->inputsExposure)},
+					UE::FUsdAttribute{UsdPrim.GetAttribute(pxr::UsdLuxTokens->inputsRadius)},
+					UE::FUsdAttribute{UsdPrim.GetAttribute(pxr::UsdLuxTokens->inputsShapingConeAngle)},
+					UE::FUsdAttribute{UsdPrim.GetAttribute(pxr::UsdLuxTokens->inputsShapingConeSoftness)}
 				};
 			}
 			else
 			{
 				return {
-					UE::FUsdAttribute{ UsdPrim.GetAttribute( pxr::UsdLuxTokens->inputsIntensity ) },
-					UE::FUsdAttribute{ UsdPrim.GetAttribute( pxr::UsdLuxTokens->inputsExposure ) },
-					UE::FUsdAttribute{ UsdPrim.GetAttribute( pxr::UsdLuxTokens->inputsRadius ) }
+					UE::FUsdAttribute{UsdPrim.GetAttribute(pxr::UsdLuxTokens->inputsIntensity)},
+					UE::FUsdAttribute{UsdPrim.GetAttribute(pxr::UsdLuxTokens->inputsExposure)},
+					UE::FUsdAttribute{UsdPrim.GetAttribute(pxr::UsdLuxTokens->inputsRadius)}
 				};
 			}
 		}
 	}
-	else if ( PropertyPath == LightColorPropertyName )
+	else if (PropertyPath == LightColorPropertyName)
 	{
-		return { UE::FUsdAttribute{ UsdPrim.GetAttribute( pxr::UsdLuxTokens->inputsColor ) } };
+		return {UE::FUsdAttribute{UsdPrim.GetAttribute(pxr::UsdLuxTokens->inputsColor)}};
 	}
-	else if ( PropertyPath == UseTemperaturePropertyName )
+	else if (PropertyPath == UseTemperaturePropertyName)
 	{
-		return { UE::FUsdAttribute{ UsdPrim.GetAttribute( pxr::UsdLuxTokens->inputsEnableColorTemperature ) } };
+		return {UE::FUsdAttribute{UsdPrim.GetAttribute(pxr::UsdLuxTokens->inputsEnableColorTemperature)}};
 	}
-	else if ( PropertyPath == TemperaturePropertyName )
+	else if (PropertyPath == TemperaturePropertyName)
 	{
-		return { UE::FUsdAttribute{ UsdPrim.GetAttribute( pxr::UsdLuxTokens->inputsColorTemperature ) } };
+		return {UE::FUsdAttribute{UsdPrim.GetAttribute(pxr::UsdLuxTokens->inputsColorTemperature)}};
 	}
-	else if ( PropertyPath == SourceWidthPropertyName )
+	else if (PropertyPath == SourceWidthPropertyName)
 	{
-		if ( UsdPrim.IsA<pxr::UsdLuxDiskLight>() )
+		if (UsdPrim.IsA<pxr::UsdLuxDiskLight>())
 		{
-			return { UE::FUsdAttribute{ UsdPrim.GetAttribute( pxr::UsdLuxTokens->inputsRadius ) } };
+			return {UE::FUsdAttribute{UsdPrim.GetAttribute(pxr::UsdLuxTokens->inputsRadius)}};
 		}
 		else
 		{
-			return { UE::FUsdAttribute{ UsdPrim.GetAttribute( pxr::UsdLuxTokens->inputsWidth ) } };
+			return {UE::FUsdAttribute{UsdPrim.GetAttribute(pxr::UsdLuxTokens->inputsWidth)}};
 		}
 	}
-	else if ( PropertyPath == SourceHeightPropertyName )
+	else if (PropertyPath == SourceHeightPropertyName)
 	{
-		if ( UsdPrim.IsA<pxr::UsdLuxDiskLight>() )
+		if (UsdPrim.IsA<pxr::UsdLuxDiskLight>())
 		{
-			return { UE::FUsdAttribute{ UsdPrim.GetAttribute( pxr::UsdLuxTokens->inputsRadius ) } };
+			return {UE::FUsdAttribute{UsdPrim.GetAttribute(pxr::UsdLuxTokens->inputsRadius)}};
 		}
 		else
 		{
-			return { UE::FUsdAttribute{ UsdPrim.GetAttribute( pxr::UsdLuxTokens->inputsHeight ) } };
+			return {UE::FUsdAttribute{UsdPrim.GetAttribute(pxr::UsdLuxTokens->inputsHeight)}};
 		}
 	}
-	else if ( PropertyPath == SourceRadiusPropertyName )
+	else if (PropertyPath == SourceRadiusPropertyName)
 	{
-		return { UE::FUsdAttribute{ UsdPrim.GetAttribute( pxr::UsdLuxTokens->inputsRadius ) } };
+		return {UE::FUsdAttribute{UsdPrim.GetAttribute(pxr::UsdLuxTokens->inputsRadius)}};
 	}
-	else if ( PropertyPath == OuterConeAnglePropertyName )
+	else if (PropertyPath == OuterConeAnglePropertyName)
 	{
-		return { UE::FUsdAttribute{ UsdPrim.GetAttribute( pxr::UsdLuxTokens->inputsShapingConeAngle ) } };
+		return {UE::FUsdAttribute{UsdPrim.GetAttribute(pxr::UsdLuxTokens->inputsShapingConeAngle)}};
 	}
-	else if ( PropertyPath == InnerConeAnglePropertyName )
+	else if (PropertyPath == InnerConeAnglePropertyName)
 	{
 		return {
-			UE::FUsdAttribute{ UsdPrim.GetAttribute( pxr::UsdLuxTokens->inputsShapingConeAngle ) },
-			UE::FUsdAttribute{ UsdPrim.GetAttribute( pxr::UsdLuxTokens->inputsShapingConeSoftness ) }
+			UE::FUsdAttribute{UsdPrim.GetAttribute(pxr::UsdLuxTokens->inputsShapingConeAngle)},
+			UE::FUsdAttribute{UsdPrim.GetAttribute(pxr::UsdLuxTokens->inputsShapingConeSoftness)}
 		};
 	}
-	else if ( PropertyPath == LightSourceAnglePropertyName )
+	else if (PropertyPath == LightSourceAnglePropertyName)
 	{
-		return { UE::FUsdAttribute{ UsdPrim.GetAttribute( pxr::UsdLuxTokens->inputsAngle ) } };
+		return {UE::FUsdAttribute{UsdPrim.GetAttribute(pxr::UsdLuxTokens->inputsAngle)}};
 	}
 
 	// Bounds component properties
 	// For now we only support bounds animations, but there's nothing preventing us from supporting color/texture/draw mode animations here,
 	// especially since our component works with a "dynamic mesh" approach that already rebuilds its proxy on-demand anyway
-	else if (
-		PropertyPath == GET_MEMBER_NAME_CHECKED(UUsdDrawModeComponent, BoundsMin) ||
-		PropertyPath == GET_MEMBER_NAME_CHECKED(UUsdDrawModeComponent, BoundsMax)
-	)
+	else if (PropertyPath == GET_MEMBER_NAME_CHECKED(UUsdDrawModeComponent, BoundsMin) || PropertyPath == GET_MEMBER_NAME_CHECKED(UUsdDrawModeComponent, BoundsMax))
 	{
 		// If we have a model API, let's direct the caller to `extentsHint` instead, as that has priority over
 		// `extent` in case the model API is applied to a boundable and both are authored
@@ -4931,7 +5041,7 @@ TArray<UE::FUsdAttribute> UnrealToUsd::GetAttributesForProperty( const UE::FUsdP
 
 bool UnrealToUsd::ConvertBoundsComponent(const UUsdDrawModeComponent& BoundsComponent, pxr::UsdPrim& UsdPrim, bool bWriteExtents, double UsdTimeCode)
 {
-	if(!UsdPrim)
+	if (!UsdPrim)
 	{
 		return false;
 	}
@@ -4982,11 +5092,13 @@ bool UnrealToUsd::ConvertBoundsComponent(const UUsdDrawModeComponent& BoundsComp
 		pxr::GfVec3f UsdMin{
 			FMath::Min(UEBoundsMinUsdSpace[0], UEBoundsMaxUsdSpace[0]),
 			FMath::Min(UEBoundsMinUsdSpace[1], UEBoundsMaxUsdSpace[1]),
-			FMath::Min(UEBoundsMinUsdSpace[2], UEBoundsMaxUsdSpace[2])};
+			FMath::Min(UEBoundsMinUsdSpace[2], UEBoundsMaxUsdSpace[2])
+		};
 		pxr::GfVec3f UsdMax{
 			FMath::Max(UEBoundsMinUsdSpace[0], UEBoundsMaxUsdSpace[0]),
 			FMath::Max(UEBoundsMinUsdSpace[1], UEBoundsMaxUsdSpace[1]),
-			FMath::Max(UEBoundsMinUsdSpace[2], UEBoundsMaxUsdSpace[2])};
+			FMath::Max(UEBoundsMinUsdSpace[2], UEBoundsMaxUsdSpace[2])
+		};
 
 		pxr::VtArray<pxr::GfVec3f> Extents{UsdMin, UsdMax};
 
@@ -5011,11 +5123,11 @@ bool UnrealToUsd::ConvertBoundsComponent(const UUsdDrawModeComponent& BoundsComp
 	if (pxr::UsdAttribute Attr = GeomModelAPI.CreateModelDrawModeAttr())
 	{
 		Attr.Set(
-			BoundsComponent.DrawMode == EUsdDrawMode::Origin	   ? pxr::UsdGeomTokens->origin
-			: BoundsComponent.DrawMode == EUsdDrawMode::Bounds	   ? pxr::UsdGeomTokens->bounds
-			: BoundsComponent.DrawMode == EUsdDrawMode::Cards	   ? pxr::UsdGeomTokens->cards
-			: BoundsComponent.DrawMode == EUsdDrawMode::Inherited  ? pxr::UsdGeomTokens->inherited
-																   : pxr::UsdGeomTokens->default_,
+			BoundsComponent.DrawMode == EUsdDrawMode::Origin	  ? pxr::UsdGeomTokens->origin
+			: BoundsComponent.DrawMode == EUsdDrawMode::Bounds	  ? pxr::UsdGeomTokens->bounds
+			: BoundsComponent.DrawMode == EUsdDrawMode::Cards	  ? pxr::UsdGeomTokens->cards
+			: BoundsComponent.DrawMode == EUsdDrawMode::Inherited ? pxr::UsdGeomTokens->inherited
+																  : pxr::UsdGeomTokens->default_,
 			UsdTimeCode
 		);
 	}
@@ -5102,12 +5214,36 @@ bool UnrealToUsd::ConvertBoundsComponent(const UUsdDrawModeComponent& BoundsComp
 			}
 		}
 	};
-	ExportCardFace(EUsdModelCardFace::XPos, &pxr::UsdGeomModelAPI::GetModelCardTextureXPosAttr, &pxr::UsdGeomModelAPI::CreateModelCardTextureXPosAttr);
-	ExportCardFace(EUsdModelCardFace::YPos, &pxr::UsdGeomModelAPI::GetModelCardTextureYPosAttr, &pxr::UsdGeomModelAPI::CreateModelCardTextureYPosAttr);
-	ExportCardFace(EUsdModelCardFace::ZPos, &pxr::UsdGeomModelAPI::GetModelCardTextureZPosAttr, &pxr::UsdGeomModelAPI::CreateModelCardTextureZPosAttr);
-	ExportCardFace(EUsdModelCardFace::XNeg, &pxr::UsdGeomModelAPI::GetModelCardTextureXNegAttr, &pxr::UsdGeomModelAPI::CreateModelCardTextureXNegAttr);
-	ExportCardFace(EUsdModelCardFace::YNeg, &pxr::UsdGeomModelAPI::GetModelCardTextureYNegAttr, &pxr::UsdGeomModelAPI::CreateModelCardTextureYNegAttr);
-	ExportCardFace(EUsdModelCardFace::ZNeg, &pxr::UsdGeomModelAPI::GetModelCardTextureZNegAttr, &pxr::UsdGeomModelAPI::CreateModelCardTextureZNegAttr);
+	ExportCardFace(
+		EUsdModelCardFace::XPos,
+		&pxr::UsdGeomModelAPI::GetModelCardTextureXPosAttr,
+		&pxr::UsdGeomModelAPI::CreateModelCardTextureXPosAttr
+	);
+	ExportCardFace(
+		EUsdModelCardFace::YPos,
+		&pxr::UsdGeomModelAPI::GetModelCardTextureYPosAttr,
+		&pxr::UsdGeomModelAPI::CreateModelCardTextureYPosAttr
+	);
+	ExportCardFace(
+		EUsdModelCardFace::ZPos,
+		&pxr::UsdGeomModelAPI::GetModelCardTextureZPosAttr,
+		&pxr::UsdGeomModelAPI::CreateModelCardTextureZPosAttr
+	);
+	ExportCardFace(
+		EUsdModelCardFace::XNeg,
+		&pxr::UsdGeomModelAPI::GetModelCardTextureXNegAttr,
+		&pxr::UsdGeomModelAPI::CreateModelCardTextureXNegAttr
+	);
+	ExportCardFace(
+		EUsdModelCardFace::YNeg,
+		&pxr::UsdGeomModelAPI::GetModelCardTextureYNegAttr,
+		&pxr::UsdGeomModelAPI::CreateModelCardTextureYNegAttr
+	);
+	ExportCardFace(
+		EUsdModelCardFace::ZNeg,
+		&pxr::UsdGeomModelAPI::GetModelCardTextureZNegAttr,
+		&pxr::UsdGeomModelAPI::CreateModelCardTextureZNegAttr
+	);
 
 	return true;
 }
@@ -5124,7 +5260,7 @@ namespace UE::USDPrimConversionImpl::Private
 		}
 
 		// Make sure we don't start with a delimiter
-		while(PrimPath.RemoveFromStart(UnrealIdentifiers::UsdNamespaceDelimiter))
+		while (PrimPath.RemoveFromStart(UnrealIdentifiers::UsdNamespaceDelimiter))
 		{
 		}
 		return PrimPath;
@@ -5184,12 +5320,8 @@ namespace UE::USDPrimConversionImpl::Private
 				const static FString CustomDataPrefix = UsdToUnreal::ConvertToken(pxr::SdfFieldKeys->CustomData);
 
 				// e.g. "customData:fromSourcePrims:Root:MyXform:CollapsedMesh1:customData:myIntMetadataValue"
-				FullKeyPath =
-					CustomDataPrefix +
-					UnrealIdentifiers::UsdNamespaceDelimiter +
-					NamespacePrefix +
-					UnrealIdentifiers::UsdNamespaceDelimiter +
-					FullKeyPath;
+				FullKeyPath = CustomDataPrefix + UnrealIdentifiers::UsdNamespaceDelimiter + NamespacePrefix + UnrealIdentifiers::UsdNamespaceDelimiter
+							  + FullKeyPath;
 
 				// USD is fine with us authoring apiSchemas directly as top level metadata, but it can't understand the typename
 				// if it's in a nested dictionary. We're never going to be actively using that value as actual apiSchemas
@@ -5207,11 +5339,7 @@ namespace UE::USDPrimConversionImpl::Private
 			}
 
 			pxr::VtValue UnstringifiedValue;
-			bSuccess &= UsdUtils::Unstringify(
-				MetadataValue.StringifiedValue,
-				*TypeNameToUse,
-				UnstringifiedValue
-			);
+			bSuccess &= UsdUtils::Unstringify(MetadataValue.StringifiedValue, *TypeNameToUse, UnstringifiedValue);
 
 			if (!bSuccess)
 			{
@@ -5268,7 +5396,7 @@ namespace UE::USDPrimConversionImpl::Private
 		}
 		return bSuccess;
 	}
-}
+}	 // namespace UE::USDPrimConversionImpl::Private
 
 bool UnrealToUsd::ConvertMetadata(
 	const FUsdCombinedPrimMetadata& CombinedPrimMetadata,
@@ -5397,4 +5525,4 @@ bool UnrealToUsd::ConvertMetadata(
 	return true;
 }
 
-#endif // #if USE_USD_SDK
+#endif	  // #if USE_USD_SDK
