USD: Reformat USDUtilities module.

diff --git a/Engine/Plugins/Importers/USDImporter/Source/USDUtilities/Private/USDShadeConversion.cpp b/Engine/Plugins/Importers/USDImporter/Source/USDUtilities/Private/USDShadeConversion.cpp
--- a/Engine/Plugins/Importers/USDImporter/Source/USDUtilities/Private/USDShadeConversion.cpp
+++ b/Engine/Plugins/Importers/USDImporter/Source/USDUtilities/Private/USDShadeConversion.cpp
@@ -18,7 +18,6 @@
 #include "USDMemory.h"
 #include "USDTypesConversion.h"
 
-#include "UsdWrappers/SdfLayer.h"
 #include "UsdWrappers/UsdAttribute.h"
 #include "UsdWrappers/UsdPrim.h"
 
@@ -30,7 +29,6 @@
 #include "Materials/Material.h"
 #include "Materials/MaterialExpressionConstant.h"
 #include "Materials/MaterialExpressionConstant4Vector.h"
-#include "Materials/MaterialExpressionTextureCoordinate.h"
 #include "Materials/MaterialExpressionTextureSample.h"
 #include "Materials/MaterialExpressionVertexColor.h"
 #include "Materials/MaterialInstanceConstant.h"
@@ -40,29 +38,26 @@
 #include "Misc/FileHelper.h"
 #include "Modules/ModuleManager.h"
 #include "RenderUtils.h"
-#include "RHI.h"
 #include "TextureResource.h"
 
 #if WITH_EDITOR
-	#include "Factories/TextureFactory.h"
-	#include "IMaterialBakingModule.h"
-	#include "MaterialBakingStructures.h"
-	#include "MaterialEditingLibrary.h"
-	#include "MaterialOptions.h"
-	#include "MaterialUtilities.h"
-#endif // WITH_EDITOR
+#include "Factories/TextureFactory.h"
+#include "IMaterialBakingModule.h"
+#include "MaterialBakingStructures.h"
+#include "MaterialEditingLibrary.h"
+#include "MaterialOptions.h"
+#include "MaterialUtilities.h"
+#endif	  // WITH_EDITOR
 
 #include "USDIncludesStart.h"
-	#include "pxr/base/tf/token.h"
-	#include "pxr/usd/ar/ar.h"
-	#include "pxr/usd/ar/asset.h"
-	#include "pxr/usd/ar/resolvedPath.h"
-	#include "pxr/usd/ar/resolver.h"
-	#include "pxr/usd/ar/resolverScopedCache.h"
-	#include "pxr/usd/sdf/layer.h"
-	#include "pxr/usd/sdf/layerUtils.h"
-	#include "pxr/usd/usdShade/material.h"
-	#include "pxr/usd/usdShade/nodeDefAPI.h"
+#include "pxr/base/tf/token.h"
+#include "pxr/usd/ar/asset.h"
+#include "pxr/usd/ar/resolvedPath.h"
+#include "pxr/usd/ar/resolver.h"
+#include "pxr/usd/ar/resolverScopedCache.h"
+#include "pxr/usd/sdf/layerUtils.h"
+#include "pxr/usd/usdShade/material.h"
+#include "pxr/usd/usdShade/nodeDefAPI.h"
 #include "USDIncludesEnd.h"
 
 // Required for packaging
@@ -78,20 +73,19 @@ namespace UE
 		{
 #if WITH_EDITOR
 			using FlattenToMatEntry = TPairInitializer<const EFlattenMaterialProperties&, const EMaterialProperty&>;
-			const static TMap<EFlattenMaterialProperties, EMaterialProperty> FlattenToMaterialProperty
-			({
-				FlattenToMatEntry( EFlattenMaterialProperties::Diffuse, EMaterialProperty::MP_BaseColor ),
-				FlattenToMatEntry( EFlattenMaterialProperties::Metallic, EMaterialProperty::MP_Metallic ),
-				FlattenToMatEntry( EFlattenMaterialProperties::Specular, EMaterialProperty::MP_Specular ),
-				FlattenToMatEntry( EFlattenMaterialProperties::Roughness, EMaterialProperty::MP_Roughness ),
-				FlattenToMatEntry( EFlattenMaterialProperties::Anisotropy, EMaterialProperty::MP_Anisotropy ),
-				FlattenToMatEntry( EFlattenMaterialProperties::Normal, EMaterialProperty::MP_Normal ),
-				FlattenToMatEntry( EFlattenMaterialProperties::Tangent, EMaterialProperty::MP_Tangent ),
-				FlattenToMatEntry( EFlattenMaterialProperties::Opacity, EMaterialProperty::MP_Opacity ),
-				FlattenToMatEntry( EFlattenMaterialProperties::Emissive, EMaterialProperty::MP_EmissiveColor ),
-				FlattenToMatEntry( EFlattenMaterialProperties::SubSurface, EMaterialProperty::MP_SubsurfaceColor ),
-				FlattenToMatEntry( EFlattenMaterialProperties::OpacityMask, EMaterialProperty::MP_OpacityMask ),
-				FlattenToMatEntry( EFlattenMaterialProperties::AmbientOcclusion, EMaterialProperty::MP_AmbientOcclusion ),
+			const static TMap<EFlattenMaterialProperties, EMaterialProperty> FlattenToMaterialProperty({
+				FlattenToMatEntry(EFlattenMaterialProperties::Diffuse, EMaterialProperty::MP_BaseColor),
+				FlattenToMatEntry(EFlattenMaterialProperties::Metallic, EMaterialProperty::MP_Metallic),
+				FlattenToMatEntry(EFlattenMaterialProperties::Specular, EMaterialProperty::MP_Specular),
+				FlattenToMatEntry(EFlattenMaterialProperties::Roughness, EMaterialProperty::MP_Roughness),
+				FlattenToMatEntry(EFlattenMaterialProperties::Anisotropy, EMaterialProperty::MP_Anisotropy),
+				FlattenToMatEntry(EFlattenMaterialProperties::Normal, EMaterialProperty::MP_Normal),
+				FlattenToMatEntry(EFlattenMaterialProperties::Tangent, EMaterialProperty::MP_Tangent),
+				FlattenToMatEntry(EFlattenMaterialProperties::Opacity, EMaterialProperty::MP_Opacity),
+				FlattenToMatEntry(EFlattenMaterialProperties::Emissive, EMaterialProperty::MP_EmissiveColor),
+				FlattenToMatEntry(EFlattenMaterialProperties::SubSurface, EMaterialProperty::MP_SubsurfaceColor),
+				FlattenToMatEntry(EFlattenMaterialProperties::OpacityMask, EMaterialProperty::MP_OpacityMask),
+				FlattenToMatEntry(EFlattenMaterialProperties::AmbientOcclusion, EMaterialProperty::MP_AmbientOcclusion),
 			});
 
 			/** Simple wrapper around the FBakeOutput data that we can reuse for data coming in from a FFlattenMaterial */
@@ -101,41 +95,41 @@ namespace UE
 				TMap<EMaterialProperty, TArray<FFloat16Color>*> HDRPropertyData;
 				TMap<EMaterialProperty, FIntPoint> PropertySizes;
 
-				explicit FBakedMaterialView( FBakeOutput& BakeOutput )
-					: PropertySizes( BakeOutput.PropertySizes )
+				explicit FBakedMaterialView(FBakeOutput& BakeOutput)
+					: PropertySizes(BakeOutput.PropertySizes)
 				{
-					HDRPropertyData.Reserve( BakeOutput.HDRPropertyData.Num() );
-					for ( TPair<EMaterialProperty, TArray<FFloat16Color>>& BakedPair : BakeOutput.HDRPropertyData )
+					HDRPropertyData.Reserve(BakeOutput.HDRPropertyData.Num());
+					for (TPair<EMaterialProperty, TArray<FFloat16Color>>& BakedPair : BakeOutput.HDRPropertyData)
 					{
-						HDRPropertyData.Add( BakedPair.Key, &BakedPair.Value );
+						HDRPropertyData.Add(BakedPair.Key, &BakedPair.Value);
 					}
 
-					PropertyData.Reserve( BakeOutput.PropertyData.Num() );
-					for ( TPair<EMaterialProperty, TArray<FColor>>& BakedPair : BakeOutput.PropertyData )
+					PropertyData.Reserve(BakeOutput.PropertyData.Num());
+					for (TPair<EMaterialProperty, TArray<FColor>>& BakedPair : BakeOutput.PropertyData)
 					{
-						PropertyData.Add( BakedPair.Key, &BakedPair.Value );
+						PropertyData.Add(BakedPair.Key, &BakedPair.Value);
 					}
 				}
 
-				explicit FBakedMaterialView( FFlattenMaterial& FlattenMaterial )
+				explicit FBakedMaterialView(FFlattenMaterial& FlattenMaterial)
 				{
-					PropertyData.Reserve( static_cast< uint8 >( EFlattenMaterialProperties::NumFlattenMaterialProperties ) );
-					for ( const TPair< EFlattenMaterialProperties, EMaterialProperty>& FlattenToProperty : FlattenToMaterialProperty )
+					PropertyData.Reserve(static_cast<uint8>(EFlattenMaterialProperties::NumFlattenMaterialProperties));
+					for (const TPair<EFlattenMaterialProperties, EMaterialProperty>& FlattenToProperty : FlattenToMaterialProperty)
 					{
-						PropertyData.Add( FlattenToProperty.Value, &FlattenMaterial.GetPropertySamples( FlattenToProperty.Key ) );
+						PropertyData.Add(FlattenToProperty.Value, &FlattenMaterial.GetPropertySamples(FlattenToProperty.Key));
 					}
 
-					PropertySizes.Reserve( static_cast< uint8 >( EFlattenMaterialProperties::NumFlattenMaterialProperties ) );
-					for ( const TPair< EFlattenMaterialProperties, EMaterialProperty>& FlattenToProperty : FlattenToMaterialProperty )
+					PropertySizes.Reserve(static_cast<uint8>(EFlattenMaterialProperties::NumFlattenMaterialProperties));
+					for (const TPair<EFlattenMaterialProperties, EMaterialProperty>& FlattenToProperty : FlattenToMaterialProperty)
 					{
-						PropertySizes.Add( FlattenToProperty.Value, FlattenMaterial.GetPropertySize( FlattenToProperty.Key ) );
+						PropertySizes.Add(FlattenToProperty.Value, FlattenMaterial.GetPropertySize(FlattenToProperty.Key));
 					}
 				}
 			};
-#endif // WITH_EDITOR
+#endif	  // WITH_EDITOR
 
 			// Given an AssetPath, resolve it to an actual file path
-			FString ResolveTexturePath( const pxr::SdfLayerHandle& LayerHandle, const FString& AssetPath )
+			FString ResolveTexturePath(const pxr::SdfLayerHandle& LayerHandle, const FString& AssetPath)
 			{
 				// TODO: Most of this stuff is incompatible with custom resolvers, as these asset paths may be URLs
 				// or just GUIDs or anything like that, where relative vs absolute path make no sense. We will need
@@ -143,9 +137,10 @@ namespace UE
 
 				FScopedUsdAllocs UsdAllocs;
 
-				FString PathToResolve = LayerHandle
-					? UsdToUnreal::ConvertString(pxr::SdfComputeAssetPathRelativeToLayer(LayerHandle, UnrealToUsd::ConvertString(*AssetPath).Get()))
-					: AssetPath;
+				FString PathToResolve = LayerHandle ? UsdToUnreal::ConvertString(
+											pxr::SdfComputeAssetPathRelativeToLayer(LayerHandle, UnrealToUsd::ConvertString(*AssetPath).Get())
+										)
+													: AssetPath;
 
 				// We need to provide absolute paths to the resolver later: It has no idea what to do with a relative
 				// path relative to some random location
@@ -181,40 +176,39 @@ namespace UE
 
 				pxr::ArResolverScopedCache ResolverCache;
 				pxr::ArResolver& Resolver = pxr::ArGetResolver();
-				std::string AssetIdentifier = Resolver.Resolve(
-					UnrealToUsd::ConvertString(*PathToResolve).Get().c_str()
-				);
+				std::string AssetIdentifier = Resolver.Resolve(UnrealToUsd::ConvertString(*PathToResolve).Get().c_str());
 
 				// Don't normalize an empty path as the result will be "."
-				if ( AssetIdentifier.size() > 0 )
+				if (AssetIdentifier.size() > 0)
 				{
-					AssetIdentifier = Resolver.CreateIdentifier( AssetIdentifier );
+					AssetIdentifier = Resolver.CreateIdentifier(AssetIdentifier);
 				}
 
-				FString ResolvedAssetPath = UsdToUnreal::ConvertString( AssetIdentifier );
+				FString ResolvedAssetPath = UsdToUnreal::ConvertString(AssetIdentifier);
 				return ResolvedAssetPath;
 			}
 
-			// If ResolvedTexturePath is in a format like "C:/MyFiles/scene.usdz[0/texture.png]", this will place "png" in OutExtensionNoDot, and return true
+			// If ResolvedTexturePath is in a format like "C:/MyFiles/scene.usdz[0/texture.png]", this will place "png" in OutExtensionNoDot, and
+			// return true
 			bool IsInsideUsdzArchive(const FString& ResolvedTexturePath, FString& OutExtensionNoDot)
 			{
 				// We need at least an opening and closing bracket
-				if ( ResolvedTexturePath.Len() < 3 )
+				if (ResolvedTexturePath.Len() < 3)
 				{
 					return false;
 				}
 
-				int32 OpenBracketPos = ResolvedTexturePath.Find( TEXT( "[" ), ESearchCase::IgnoreCase, ESearchDir::FromEnd );
-				int32 CloseBracketPos = ResolvedTexturePath.Find( TEXT( "]" ), ESearchCase::IgnoreCase, ESearchDir::FromEnd );
-				if ( OpenBracketPos != INDEX_NONE && CloseBracketPos != INDEX_NONE && CloseBracketPos > OpenBracketPos - 1 )
+				int32 OpenBracketPos = ResolvedTexturePath.Find(TEXT("["), ESearchCase::IgnoreCase, ESearchDir::FromEnd);
+				int32 CloseBracketPos = ResolvedTexturePath.Find(TEXT("]"), ESearchCase::IgnoreCase, ESearchDir::FromEnd);
+				if (OpenBracketPos != INDEX_NONE && CloseBracketPos != INDEX_NONE && CloseBracketPos > OpenBracketPos - 1)
 				{
 					// Should be like "C:/MyFiles/scene.usdz"
-					FString UsdzFilePath = ResolvedTexturePath.Left( OpenBracketPos );
+					FString UsdzFilePath = ResolvedTexturePath.Left(OpenBracketPos);
 
 					// Should be like "0/texture.png"
-					FString TextureInUsdzPath = ResolvedTexturePath.Mid( OpenBracketPos + 1, CloseBracketPos - OpenBracketPos - 1 );
+					FString TextureInUsdzPath = ResolvedTexturePath.Mid(OpenBracketPos + 1, CloseBracketPos - OpenBracketPos - 1);
 
-					if ( FPaths::GetExtension( UsdzFilePath ).ToLower() == TEXT( "usdz" ) )
+					if (FPaths::GetExtension(UsdzFilePath).ToLower() == TEXT("usdz"))
 					{
 						// Should be something like "png"
 						OutExtensionNoDot = FPaths::GetExtension(TextureInUsdzPath);
@@ -225,17 +219,16 @@ namespace UE
 				return false;
 			}
 
-			TUsdStore<std::shared_ptr<const char>> ReadTextureBufferFromUsdzArchive( const FString& ResolvedTexturePath, uint64& OutBufferSize )
+			TUsdStore<std::shared_ptr<const char>> ReadTextureBufferFromUsdzArchive(const FString& ResolvedTexturePath, uint64& OutBufferSize)
 			{
 				pxr::ArResolver& Resolver = pxr::ArGetResolver();
-				std::shared_ptr<pxr::ArAsset> Asset = Resolver.OpenAsset(
-					pxr::ArResolvedPath( UnrealToUsd::ConvertString( *ResolvedTexturePath ).Get() ) );
+				std::shared_ptr<pxr::ArAsset> Asset = Resolver.OpenAsset(pxr::ArResolvedPath(UnrealToUsd::ConvertString(*ResolvedTexturePath).Get()));
 
 				TUsdStore<std::shared_ptr<const char>> Buffer;
 
-				if ( Asset )
+				if (Asset)
 				{
-					OutBufferSize = static_cast< uint64 >( Asset->GetSize() );
+					OutBufferSize = static_cast<uint64>(Asset->GetSize());
 					{
 						FScopedUsdAllocs Allocs;
 						Buffer = Asset->GetBuffer();
@@ -245,15 +238,22 @@ namespace UE
 				return Buffer;
 			}
 
-			/** If ResolvedTexturePath points at a texture inside an usdz file, this will use USD to pull the asset from the file, and TextureFactory to import it directly from the binary buffer */
-			UTexture* ReadTextureFromUsdzArchiveEditor(const FString& ResolvedTexturePath, const FString& TextureExtension, UTextureFactory* TextureFactory, UObject* Outer, EObjectFlags ObjectFlags )
+			/** If ResolvedTexturePath points at a texture inside an usdz file, this will use USD to pull the asset from the file, and TextureFactory
+			 * to import it directly from the binary buffer */
+			UTexture* ReadTextureFromUsdzArchiveEditor(
+				const FString& ResolvedTexturePath,
+				const FString& TextureExtension,
+				UTextureFactory* TextureFactory,
+				UObject* Outer,
+				EObjectFlags ObjectFlags
+			)
 			{
 #if WITH_EDITOR
 				uint64 BufferSize = 0;
 				TUsdStore<std::shared_ptr<const char>> Buffer = ReadTextureBufferFromUsdzArchive(ResolvedTexturePath, BufferSize);
-				const uint8* BufferStart = reinterpret_cast< const uint8* >( Buffer.Get().get() );
+				const uint8* BufferStart = reinterpret_cast<const uint8*>(Buffer.Get().get());
 
-				if ( BufferSize == 0 || BufferStart == nullptr )
+				if (BufferSize == 0 || BufferStart == nullptr)
 				{
 					return nullptr;
 				}
@@ -275,39 +275,41 @@ namespace UE
 					*TextureExtension,
 					BufferStart,
 					BufferStart + BufferSize,
-					nullptr));
-#endif // WITH_EDITOR
+					nullptr
+				));
+#endif	  // WITH_EDITOR
 				return nullptr;
 			}
 
-			UTexture* ReadTextureFromUsdzArchiveRuntime( const FString& ResolvedTexturePath )
+			UTexture* ReadTextureFromUsdzArchiveRuntime(const FString& ResolvedTexturePath)
 			{
 				uint64 BufferSize = 0;
-				TUsdStore<std::shared_ptr<const char>> Buffer = ReadTextureBufferFromUsdzArchive( ResolvedTexturePath, BufferSize );
-				const uint8* BufferStart = reinterpret_cast< const uint8* >( Buffer.Get().get() );
+				TUsdStore<std::shared_ptr<const char>> Buffer = ReadTextureBufferFromUsdzArchive(ResolvedTexturePath, BufferSize);
+				const uint8* BufferStart = reinterpret_cast<const uint8*>(Buffer.Get().get());
 
-				if ( BufferSize == 0 || BufferStart == nullptr )
+				if (BufferSize == 0 || BufferStart == nullptr)
 				{
 					return nullptr;
 				}
 
 				FScopedUnrealAllocs UEAllocs;
 
-				// Copied from FImageUtils::ImportBufferAsTexture2D( Buffer ) so that we can avoid a copy into the TArray<uint8> that it takes as parameter
-				IImageWrapperModule& ImageWrapperModule = FModuleManager::Get().LoadModuleChecked<IImageWrapperModule>( TEXT( "ImageWrapper" ) );
-				EImageFormat Format = ImageWrapperModule.DetectImageFormat( BufferStart, BufferSize );
+				// Copied from FImageUtils::ImportBufferAsTexture2D( Buffer ) so that we can avoid a copy into the TArray<uint8> that it takes as
+				// parameter
+				IImageWrapperModule& ImageWrapperModule = FModuleManager::Get().LoadModuleChecked<IImageWrapperModule>(TEXT("ImageWrapper"));
+				EImageFormat Format = ImageWrapperModule.DetectImageFormat(BufferStart, BufferSize);
 
 				UTexture2D* NewTexture = nullptr;
 				EPixelFormat PixelFormat = PF_Unknown;
-				if ( Format != EImageFormat::Invalid )
+				if (Format != EImageFormat::Invalid)
 				{
-					TSharedPtr<IImageWrapper> ImageWrapper = ImageWrapperModule.CreateImageWrapper( Format );
+					TSharedPtr<IImageWrapper> ImageWrapper = ImageWrapperModule.CreateImageWrapper(Format);
 
 					int32 BitDepth = 0;
 					int32 Width = 0;
 					int32 Height = 0;
 
-					if ( ImageWrapper->SetCompressed( ( void* ) BufferStart, BufferSize ) )
+					if (ImageWrapper->SetCompressed((void*)BufferStart, BufferSize))
 					{
 						PixelFormat = PF_Unknown;
 
@@ -318,35 +320,35 @@ namespace UE
 						Width = ImageWrapper->GetWidth();
 						Height = ImageWrapper->GetHeight();
 
-						if ( BitDepth == 16 )
+						if (BitDepth == 16)
 						{
 							PixelFormat = PF_FloatRGBA;
 							RGBFormat = ERGBFormat::BGRA;
 						}
-						else if ( BitDepth == 8 )
+						else if (BitDepth == 8)
 						{
 							PixelFormat = PF_B8G8R8A8;
 							RGBFormat = ERGBFormat::BGRA;
 						}
 						else
 						{
-							UE_LOG( LogUsd, Warning, TEXT( "Error creating texture. Bit depth is unsupported. (%d)" ), BitDepth );
+							UE_LOG(LogUsd, Warning, TEXT("Error creating texture. Bit depth is unsupported. (%d)"), BitDepth);
 							return nullptr;
 						}
 
 						TArray64<uint8> UncompressedData;
-						ImageWrapper->GetRaw( RGBFormat, BitDepth, UncompressedData );
+						ImageWrapper->GetRaw(RGBFormat, BitDepth, UncompressedData);
 
-						NewTexture = UTexture2D::CreateTransient( Width, Height, PixelFormat );
-						if ( NewTexture )
+						NewTexture = UTexture2D::CreateTransient(Width, Height, PixelFormat);
+						if (NewTexture)
 						{
 							NewTexture->bNotOfflineProcessed = true;
-							uint8* MipData = static_cast< uint8* >( NewTexture->GetPlatformData()->Mips[ 0 ].BulkData.Lock( LOCK_READ_WRITE ) );
+							uint8* MipData = static_cast<uint8*>(NewTexture->GetPlatformData()->Mips[0].BulkData.Lock(LOCK_READ_WRITE));
 
 							// Bulk data was already allocated for the correct size when we called CreateTransient above
-							FMemory::Memcpy( MipData, UncompressedData.GetData(), NewTexture->GetPlatformData()->Mips[ 0 ].BulkData.GetBulkDataSize() );
+							FMemory::Memcpy(MipData, UncompressedData.GetData(), NewTexture->GetPlatformData()->Mips[0].BulkData.GetBulkDataSize());
 
-							NewTexture->GetPlatformData()->Mips[ 0 ].BulkData.Unlock();
+							NewTexture->GetPlatformData()->Mips[0].BulkData.Unlock();
 
 							NewTexture->UpdateResource();
 						}
@@ -354,7 +356,7 @@ namespace UE
 				}
 				else
 				{
-					UE_LOG( LogUsd, Warning, TEXT( "Error creating texture. Couldn't determine the file format" ) );
+					UE_LOG(LogUsd, Warning, TEXT("Error creating texture. Couldn't determine the file format"));
 				}
 
 				return NewTexture;
@@ -362,26 +364,26 @@ namespace UE
 
 			// Will traverse the shade material graph backwards looking for a string/token value and return it.
 			// Returns the empty string if it didn't find anything.
-			FString RecursivelySearchForStringValue( pxr::UsdShadeInput Input )
+			FString RecursivelySearchForStringValue(pxr::UsdShadeInput Input)
 			{
-				if ( !Input )
+				if (!Input)
 				{
 					return {};
 				}
 
 				FScopedUsdAllocs Allocs;
 
-				if ( Input.HasConnectedSource() )
+				if (Input.HasConnectedSource())
 				{
 					pxr::UsdShadeConnectableAPI Souce;
 					pxr::TfToken SourceName;
 					pxr::UsdShadeAttributeType SourceType;
-					if ( pxr::UsdShadeConnectableAPI::GetConnectedSource( Input.GetAttr(), &Souce, &SourceName, &SourceType ) )
+					if (pxr::UsdShadeConnectableAPI::GetConnectedSource(Input.GetAttr(), &Souce, &SourceName, &SourceType))
 					{
-						for ( const pxr::UsdShadeInput& DeeperInput : Souce.GetInputs() )
+						for (const pxr::UsdShadeInput& DeeperInput : Souce.GetInputs())
 						{
-							FString Result = RecursivelySearchForStringValue( DeeperInput );
-							if ( !Result.IsEmpty() )
+							FString Result = RecursivelySearchForStringValue(DeeperInput);
+							if (!Result.IsEmpty())
 							{
 								return Result;
 							}
@@ -391,15 +393,15 @@ namespace UE
 				else
 				{
 					std::string StringValue;
-					if ( Input.Get< std::string >( &StringValue ) )
+					if (Input.Get<std::string>(&StringValue))
 					{
-						return UsdToUnreal::ConvertString( StringValue );
+						return UsdToUnreal::ConvertString(StringValue);
 					}
 
 					pxr::TfToken TokenValue;
-					if ( Input.Get< pxr::TfToken >( &TokenValue ) )
+					if (Input.Get<pxr::TfToken>(&TokenValue))
 					{
-						return UsdToUnreal::ConvertToken( TokenValue );
+						return UsdToUnreal::ConvertToken(TokenValue);
 					}
 				}
 
@@ -415,11 +417,13 @@ namespace UE
 				FVector2f UVTranslation;
 				int32 OutputIndex = 0;
 			};
+
 			struct FPrimvarReaderParameterValue
 			{
 				std::string PrimvarName;
 				FVector FallbackValue;
 			};
+
 			using FParameterValue = TVariant<float, FVector, FTextureParameterValue, FPrimvarReaderParameterValue, bool>;
 
 			/**
@@ -448,12 +452,7 @@ namespace UE
 				bool bFoundPrimvarReader = false;
 
 				// Traverse through potentially N UV transform nodes
-				while (pxr::UsdShadeConnectableAPI::GetConnectedSource(
-					StInput.GetAttr(),
-					&Connectable,
-					&ConnectableOutput,
-					&AttributeType
-				))
+				while (pxr::UsdShadeConnectableAPI::GetConnectedSource(StInput.GetAttr(), &Connectable, &ConnectableOutput, &AttributeType))
 				{
 					pxr::UsdPrim ConnectedPrim = Connectable.GetPrim();
 					pxr::UsdShadeNodeDefAPI ConnectedShadeNode{ConnectedPrim};
@@ -562,10 +561,9 @@ namespace UE
 			{
 				FScopedUsdAllocs UsdAllocs;
 
-				const bool bIsNormalInput = ( ShadeInput.GetTypeName() == pxr::SdfValueTypeNames->Normal3f
-					|| ShadeInput.GetTypeName() == pxr::SdfValueTypeNames->Normal3fArray
-					|| ShadeInput.GetBaseName() == UnrealIdentifiers::Normal
-				);
+				const bool bIsNormalInput =
+					(ShadeInput.GetTypeName() == pxr::SdfValueTypeNames->Normal3f || ShadeInput.GetTypeName() == pxr::SdfValueTypeNames->Normal3fArray
+					 || ShadeInput.GetBaseName() == UnrealIdentifiers::Normal);
 
 				pxr::UsdShadeConnectableAPI UsdUVTextureSource;
 				pxr::TfToken UsdUVTextureSourceName;
@@ -574,38 +572,48 @@ namespace UE
 				// Clear it, as we'll signal that it has a valid texture bound by setting it with an FTextureParameterValue
 				OutValue.Set<float>(0.0f);
 
-				if ( pxr::UsdShadeConnectableAPI::GetConnectedSource( ShadeInput.GetAttr(), &UsdUVTextureSource, &UsdUVTextureSourceName, &UsdUVTextureAttributeType ) )
+				if (pxr::UsdShadeConnectableAPI::GetConnectedSource(
+						ShadeInput.GetAttr(),
+						&UsdUVTextureSource,
+						&UsdUVTextureSourceName,
+						&UsdUVTextureAttributeType
+					))
 				{
 					pxr::UsdShadeInput FileInput;
 
 					// UsdUVTexture: Get its file input
-					if ( UsdUVTextureAttributeType == pxr::UsdShadeAttributeType::Output )
+					if (UsdUVTextureAttributeType == pxr::UsdShadeAttributeType::Output)
 					{
-						FileInput = UsdUVTextureSource.GetInput( UnrealIdentifiers::File );
+						FileInput = UsdUVTextureSource.GetInput(UnrealIdentifiers::File);
 					}
 					// Check if we are being directly passed an asset
 					else
 					{
-						FileInput = UsdUVTextureSource.GetInput( UsdUVTextureSourceName );
+						FileInput = UsdUVTextureSource.GetInput(UsdUVTextureSourceName);
 					}
 
 					// Recursively traverse "inputs:file" connections until we stop finding other connected prims
 					pxr::UsdShadeConnectableAPI TextureFileSource;
 					pxr::TfToken TextureFileSourceName;
 					pxr::UsdShadeAttributeType TextureFileAttributeType;
-					while ( FileInput )
-					{
-						if ( pxr::UsdShadeConnectableAPI::GetConnectedSource( FileInput.GetAttr(), &TextureFileSource, &TextureFileSourceName, &TextureFileAttributeType ) )
+					while (FileInput)
+					{
+						if (pxr::UsdShadeConnectableAPI::GetConnectedSource(
+								FileInput.GetAttr(),
+								&TextureFileSource,
+								&TextureFileSourceName,
+								&TextureFileAttributeType
+							))
 						{
 							// Another connection, get its file input
-							if ( TextureFileAttributeType == pxr::UsdShadeAttributeType::Output )
+							if (TextureFileAttributeType == pxr::UsdShadeAttributeType::Output)
 							{
-								FileInput = TextureFileSource.GetInput( UnrealIdentifiers::File );
+								FileInput = TextureFileSource.GetInput(UnrealIdentifiers::File);
 							}
 							// Check if we are being directly passed an asset
 							else
 							{
-								FileInput = TextureFileSource.GetInput( TextureFileSourceName );
+								FileInput = TextureFileSource.GetInput(TextureFileSourceName);
 							}
 						}
 						else
@@ -617,42 +625,40 @@ namespace UE
 					// Get desired texture wrapping modes
 					TextureAddress AddressX = TextureAddress::TA_Wrap;
 					TextureAddress AddressY = TextureAddress::TA_Wrap;
-					if ( pxr::UsdAttribute WrapSAttr = UsdUVTextureSource.GetInput( UnrealIdentifiers::WrapS ) )
+					if (pxr::UsdAttribute WrapSAttr = UsdUVTextureSource.GetInput(UnrealIdentifiers::WrapS))
 					{
 						pxr::TfToken WrapS;
-						if ( WrapSAttr.Get( &WrapS ) )
+						if (WrapSAttr.Get(&WrapS))
 						{
-							AddressX = WrapS == UnrealIdentifiers::Repeat
-								? TextureAddress::TA_Wrap
-								: WrapS == UnrealIdentifiers::Mirror
-									? TextureAddress::TA_Mirror
-									: TextureAddress::TA_Clamp;  // We also consider the "black" wrap mode as clamp
-																 // as that is the closest we can get
+							AddressX = WrapS == UnrealIdentifiers::Repeat	? TextureAddress::TA_Wrap
+									   : WrapS == UnrealIdentifiers::Mirror ? TextureAddress::TA_Mirror
+																			: TextureAddress::TA_Clamp;	   // We also consider the "black" wrap mode
+																										   // as clamp as that is the closest we can
+																										   // get
 						}
 					}
-					if ( pxr::UsdAttribute WrapTAttr = UsdUVTextureSource.GetInput( UnrealIdentifiers::WrapT ) )
+					if (pxr::UsdAttribute WrapTAttr = UsdUVTextureSource.GetInput(UnrealIdentifiers::WrapT))
 					{
 						pxr::TfToken WrapT;
-						if ( WrapTAttr.Get( &WrapT ) )
+						if (WrapTAttr.Get(&WrapT))
 						{
-							AddressY = WrapT == UnrealIdentifiers::Repeat
-								? TextureAddress::TA_Wrap
-								: WrapT == UnrealIdentifiers::Mirror
-									? TextureAddress::TA_Mirror
-									: TextureAddress::TA_Clamp;  // We also consider the "black" wrap mode as clamp
-																 // as that is the closest we can get
+							AddressY = WrapT == UnrealIdentifiers::Repeat	? TextureAddress::TA_Wrap
+									   : WrapT == UnrealIdentifiers::Mirror ? TextureAddress::TA_Mirror
+																			: TextureAddress::TA_Clamp;	   // We also consider the "black" wrap mode
+																										   // as clamp as that is the closest we can
+																										   // get
 						}
 					}
 
-					if ( FileInput && FileInput.GetTypeName() == pxr::SdfValueTypeNames->Asset ) // Check that FileInput is of type Asset
+					if (FileInput && FileInput.GetTypeName() == pxr::SdfValueTypeNames->Asset)	  // Check that FileInput is of type Asset
 					{
-						const FString TexturePath = UsdUtils::GetResolvedTexturePath( FileInput.GetAttr() );
+						const FString TexturePath = UsdUtils::GetResolvedTexturePath(FileInput.GetAttr());
 						// We will assume the texture is valid, and show a warning if we fail to parse this later.
 						// Note that we don't even check that the file exists: If we have a texture bound to opacity then we
 						// assume this material is meant to be translucent, even if the texture path is invalid (or points inside an USDZ archive)
-						if ( !TexturePath.IsEmpty() )
+						if (!TexturePath.IsEmpty())
 						{
-							OutValue.Set<FTextureParameterValue>( FTextureParameterValue{} );
+							OutValue.Set<FTextureParameterValue>(FTextureParameterValue{});
 						}
 						else
 						{
@@ -660,12 +666,12 @@ namespace UE
 						}
 
 						// We'll add these to the hash because the materials are built to try and reuse the same textures for multiple channels,
-						// and those may expect linear or sRGB values. Without this we may parse a texture as linear because we hit the opacity channel first,
-						// and then reuse it as linear for the base color channel even though it should have been sRGB.
-						// Plus we may have something weird like a normal map being plugged into the base color and the normal channel.
+						// and those may expect linear or sRGB values. Without this we may parse a texture as linear because we hit the opacity
+						// channel first, and then reuse it as linear for the base color channel even though it should have been sRGB. Plus we may
+						// have something weird like a normal map being plugged into the base color and the normal channel.
 						bool bSRGB = true;
 						TextureCompressionSettings CompressionSettings = TextureCompressionSettings::TC_Default;
-						if ( bIsNormalInput )
+						if (bIsNormalInput)
 						{
 							// Disable SRGB when parsing float textures, as they're likely specular/roughness maps
 							bSRGB = false;
@@ -673,40 +679,39 @@ namespace UE
 						}
 						else
 						{
-							if ( LODGroup == TEXTUREGROUP_WorldSpecular )
+							if (LODGroup == TEXTUREGROUP_WorldSpecular)
 							{
 								bSRGB = false;
 							}
 
 							// Override the sRGB flag with whatever is specified on the texture shader, if it has anything specific.
-							if ( pxr::UsdAttribute SourceColorSpaceAttr = UsdUVTextureSource.GetInput( UnrealIdentifiers::SourceColorSpaceToken ) )
+							if (pxr::UsdAttribute SourceColorSpaceAttr = UsdUVTextureSource.GetInput(UnrealIdentifiers::SourceColorSpaceToken))
 							{
 								pxr::TfToken SourceColorSpaceValue;
-								if ( SourceColorSpaceAttr.HasAuthoredValue() && SourceColorSpaceAttr.Get( &SourceColorSpaceValue ) )
+								if (SourceColorSpaceAttr.HasAuthoredValue() && SourceColorSpaceAttr.Get(&SourceColorSpaceValue))
 								{
-									bSRGB = SourceColorSpaceValue == UnrealIdentifiers::RawColorSpaceToken
-										? false
-										: SourceColorSpaceValue == UnrealIdentifiers::SRGBColorSpaceToken
-											? true
-											: bSRGB;
+									bSRGB = SourceColorSpaceValue == UnrealIdentifiers::RawColorSpaceToken	  ? false
+											: SourceColorSpaceValue == UnrealIdentifiers::SRGBColorSpaceToken ? true
+																											  : bSRGB;
 								}
 							}
 						}
 
 						const FString PrefixedTextureHash = UsdUtils::GetAssetHashPrefix(ShadeInput.GetPrim(), bReuseIdenticalAssets)
-													+ UsdUtils::GetTextureHash(TexturePath, bSRGB, CompressionSettings, AddressX, AddressY);
+															+ UsdUtils::GetTextureHash(TexturePath, bSRGB, CompressionSettings, AddressX, AddressY);
 
 						// We only actually want to retrieve the textures if we have a cache to put them in
-						if ( TexturesCache )
+						if (TexturesCache)
 						{
 							UTexture* Texture = Cast<UTexture>(TexturesCache->GetCachedAsset(PrefixedTextureHash));
-							if ( !Texture )
+							if (!Texture)
 							{
 								// Give the same prim path to the texture, so that it ends up imported right next to the material
 								FString MaterialPrimPath;
-								if ( Material )
+								if (Material)
 								{
-									if (UUsdAssetUserData* UserData = const_cast<UMaterialInterface*>(Material)->GetAssetUserData<UUsdAssetUserData>())
+									if (UUsdAssetUserData* UserData = const_cast<UMaterialInterface*>(Material)->GetAssetUserData<UUsdAssetUserData>(
+										))
 									{
 										if (!UserData->PrimPaths.IsEmpty())
 										{
@@ -715,24 +720,24 @@ namespace UE
 									}
 								}
 
-								Texture = UsdUtils::CreateTexture( FileInput.GetAttr(), MaterialPrimPath, LODGroup, TexturesCache );
+								Texture = UsdUtils::CreateTexture(FileInput.GetAttr(), MaterialPrimPath, LODGroup, TexturesCache);
 
-								if ( Texture )
+								if (Texture)
 								{
 									Texture->SRGB = bSRGB;
 									Texture->CompressionSettings = CompressionSettings;
-									if ( UTexture2D* Texture2D = Cast<UTexture2D>( Texture ) )
+									if (UTexture2D* Texture2D = Cast<UTexture2D>(Texture))
 									{
 										Texture2D->AddressX = AddressX;
 										Texture2D->AddressY = AddressY;
 									}
 									Texture->UpdateResource();
 
-									TexturesCache->CacheAsset( PrefixedTextureHash, Texture );
+									TexturesCache->CacheAsset(PrefixedTextureHash, Texture);
 								}
 							}
 
-							if ( Texture )
+							if (Texture)
 							{
 								FTextureParameterValue OutTextureValue;
 								OutTextureValue.Texture = Texture;
@@ -741,40 +746,41 @@ namespace UE
 								// The UsdUVTexture Shader itself prim has multiple standard outputs, but this is not
 								// full swizzle support (check
 								// [REDACTED_PROFILE_URL]/USD/blob/5c5ebddff35012461a2b0ba773c47f05730cbab4/pxr/usdImaging/plugin/usdShaders/shaders/shaderDefs.usda#L198)
-								if ( UsdUVTextureAttributeType == pxr::UsdShadeAttributeType::Output )
+								if (UsdUVTextureAttributeType == pxr::UsdShadeAttributeType::Output)
 								{
-									const FName OutputName( UsdToUnreal::ConvertToken( UsdUVTextureSourceName ) );
+									const FName OutputName(UsdToUnreal::ConvertToken(UsdUVTextureSourceName));
 
-									if ( OutputName == TEXT("rgb") )
+									if (OutputName == TEXT("rgb"))
 									{
 										OutTextureValue.OutputIndex = 0;
 									}
-									else if ( OutputName == TEXT("r") )
+									else if (OutputName == TEXT("r"))
 									{
 										OutTextureValue.OutputIndex = 1;
 									}
-									else if ( OutputName == TEXT("g") )
+									else if (OutputName == TEXT("g"))
 									{
 										OutTextureValue.OutputIndex = 2;
 									}
-									else if ( OutputName == TEXT("b") )
+									else if (OutputName == TEXT("b"))
 									{
 										OutTextureValue.OutputIndex = 3;
 									}
-									else if ( OutputName == TEXT("a") )
+									else if (OutputName == TEXT("a"))
 									{
 										OutTextureValue.OutputIndex = 4;
 									}
 								}
 
-								OutValue.Set<FTextureParameterValue>( OutTextureValue );
+								OutValue.Set<FTextureParameterValue>(OutTextureValue);
 							}
 							else
 							{
 								FUsdLogManager::LogMessage(
 									EMessageSeverity::Warning,
-									FText::Format( LOCTEXT( "FailedToParseTexture", "Failed to parse texture at path '{0}'" ),
-										FText::FromString( TexturePath )
+									FText::Format(
+										LOCTEXT("FailedToParseTexture", "Failed to parse texture at path '{0}'"),
+										FText::FromString(TexturePath)
 									)
 								);
 							}
@@ -782,10 +788,9 @@ namespace UE
 					}
 				}
 
-				// We may be calling this from IsMaterialTranslucent, when we have no TexturesCache: We wouldn't be able to import textures without the
-				// cache as we may repeatedly parse the same texture.
-				// Because of this, and how we clear OutValue to float before we begin, we can say that if we have a FTextureParameterValue at all, then there
-				// is a valid texture that *can* be parsed, and so we return true.
+				// We may be calling this from IsMaterialTranslucent, when we have no TexturesCache: We wouldn't be able to import textures without
+				// the cache as we may repeatedly parse the same texture. Because of this, and how we clear OutValue to float before we begin, we can
+				// say that if we have a FTextureParameterValue at all, then there is a valid texture that *can* be parsed, and so we return true.
 				return OutValue.IsType<FTextureParameterValue>();
 			}
 
@@ -801,8 +806,8 @@ namespace UE
 			{
 				FScopedUsdAllocs Allocs;
 
-				pxr::UsdShadeInput Input = Connectable.GetInput( InputName );
-				if ( !Input )
+				pxr::UsdShadeInput Input = Connectable.GetInput(InputName);
+				if (!Input)
 				{
 					return false;
 				}
@@ -811,17 +816,17 @@ namespace UE
 				pxr::UsdShadeConnectableAPI Source;
 				pxr::TfToken SourceName;
 				pxr::UsdShadeAttributeType AttributeType;
-				if ( pxr::UsdShadeConnectableAPI::GetConnectedSource( Input.GetAttr(), &Source, &SourceName, &AttributeType ) )
+				if (pxr::UsdShadeConnectableAPI::GetConnectedSource(Input.GetAttr(), &Source, &SourceName, &AttributeType))
 				{
 					if (!GetTextureParameterValue(Input, TEXTUREGROUP_WorldSpecular, OutValue, Material, TexturesCache, bReuseIdenticalAssets))
 					{
 						// Check if we have a fallback input that we can use instead, since we don't have a valid texture value
-						if ( const pxr::UsdShadeInput FallbackInput = Source.GetInput( UnrealIdentifiers::Fallback ) )
+						if (const pxr::UsdShadeInput FallbackInput = Source.GetInput(UnrealIdentifiers::Fallback))
 						{
 							float UsdFallbackFloat;
-							if ( FallbackInput.Get<float>( &UsdFallbackFloat ) )
+							if (FallbackInput.Get<float>(&UsdFallbackFloat))
 							{
-								OutValue.Set<float>( UsdFallbackFloat );
+								OutValue.Set<float>(UsdFallbackFloat);
 								return true;
 							}
 						}
@@ -835,41 +840,46 @@ namespace UE
 				else
 				{
 					float InputValue = DefaultValue;
-					Input.Get< float >( &InputValue );
+					Input.Get<float>(&InputValue);
 
-					OutValue.Set<float>( InputValue );
+					OutValue.Set<float>(InputValue);
 				}
 
 				return true;
 			}
 
-			bool GetPrimvarReaderParameterValue( const pxr::UsdShadeInput& Input, const pxr::TfToken& PrimvarReaderShaderId, const FLinearColor& DefaultValue, FParameterValue& OutValue )
+			bool GetPrimvarReaderParameterValue(
+				const pxr::UsdShadeInput& Input,
+				const pxr::TfToken& PrimvarReaderShaderId,
+				const FLinearColor& DefaultValue,
+				FParameterValue& OutValue
+			)
 			{
 				FScopedUsdAllocs Allocs;
 
-				if ( !Input )
+				if (!Input)
 				{
 					return false;
 				}
 
 				const bool bShaderOutputsOnly = true;
-				const pxr::UsdShadeAttributeVector ValProdAttrs = Input.GetValueProducingAttributes( bShaderOutputsOnly );
-				for ( const pxr::UsdAttribute& ValProdAttr : ValProdAttrs )
+				const pxr::UsdShadeAttributeVector ValProdAttrs = Input.GetValueProducingAttributes(bShaderOutputsOnly);
+				for (const pxr::UsdAttribute& ValProdAttr : ValProdAttrs)
 				{
-					const pxr::UsdShadeShader ValProdShader( ValProdAttr.GetPrim() );
-					if ( !ValProdShader )
+					const pxr::UsdShadeShader ValProdShader(ValProdAttr.GetPrim());
+					if (!ValProdShader)
 					{
 						continue;
 					}
 
 					pxr::TfToken ShaderId;
-					if ( !ValProdShader.GetShaderId( &ShaderId ) || ShaderId != PrimvarReaderShaderId )
+					if (!ValProdShader.GetShaderId(&ShaderId) || ShaderId != PrimvarReaderShaderId)
 					{
 						continue;
 					}
 
-					const pxr::UsdShadeInput VarnameInput = ValProdShader.GetInput( UnrealIdentifiers::Varname );
-					if ( !VarnameInput )
+					const pxr::UsdShadeInput VarnameInput = ValProdShader.GetInput(UnrealIdentifiers::Varname);
+					if (!VarnameInput)
 					{
 						continue;
 					}
@@ -878,37 +888,37 @@ namespace UE
 					// string-typed, but some assets might be set up token-typed, so we'll consider
 					// either type.
 					std::string PrimvarName;
-					if ( VarnameInput.GetTypeName() == pxr::SdfValueTypeNames->String )
+					if (VarnameInput.GetTypeName() == pxr::SdfValueTypeNames->String)
 					{
-						if ( !VarnameInput.Get( &PrimvarName ) )
+						if (!VarnameInput.Get(&PrimvarName))
 						{
 							continue;
 						}
 					}
-					else if ( VarnameInput.GetTypeName() == pxr::SdfValueTypeNames->Token )
+					else if (VarnameInput.GetTypeName() == pxr::SdfValueTypeNames->Token)
 					{
 						pxr::TfToken PrimvarNameToken;
-						if ( !VarnameInput.Get( &PrimvarNameToken ) )
+						if (!VarnameInput.Get(&PrimvarNameToken))
 						{
 							continue;
 						}
 						PrimvarName = PrimvarNameToken.GetString();
 					}
 
-					if ( PrimvarName.empty() )
+					if (PrimvarName.empty())
 					{
 						continue;
 					}
 
 					FLinearColor FallbackColor = DefaultValue;
-					const pxr::UsdShadeInput FallbackInput = ValProdShader.GetInput( UnrealIdentifiers::Fallback );
+					const pxr::UsdShadeInput FallbackInput = ValProdShader.GetInput(UnrealIdentifiers::Fallback);
 					pxr::GfVec3f UsdFallbackColor;
-					if ( FallbackInput && FallbackInput.Get<pxr::GfVec3f>( &UsdFallbackColor ) )
+					if (FallbackInput && FallbackInput.Get<pxr::GfVec3f>(&UsdFallbackColor))
 					{
-						FallbackColor = UsdToUnreal::ConvertColor( UsdFallbackColor );
+						FallbackColor = UsdToUnreal::ConvertColor(UsdFallbackColor);
 					}
 
-					OutValue.Set<FPrimvarReaderParameterValue>( FPrimvarReaderParameterValue{ PrimvarName, FVector(FallbackColor) } );
+					OutValue.Set<FPrimvarReaderParameterValue>(FPrimvarReaderParameterValue{PrimvarName, FVector(FallbackColor)});
 					return true;
 				}
 
@@ -928,8 +938,8 @@ namespace UE
 			{
 				FScopedUsdAllocs Allocs;
 
-				pxr::UsdShadeInput Input = Connectable.GetInput( InputName );
-				if ( !Input )
+				pxr::UsdShadeInput Input = Connectable.GetInput(InputName);
+				if (!Input)
 				{
 					return false;
 				}
@@ -938,7 +948,7 @@ namespace UE
 				pxr::UsdShadeConnectableAPI Source;
 				pxr::TfToken SourceName;
 				pxr::UsdShadeAttributeType AttributeType;
-				if ( pxr::UsdShadeConnectableAPI::GetConnectedSource( Input.GetAttr(), &Source, &SourceName, &AttributeType ) )
+				if (pxr::UsdShadeConnectableAPI::GetConnectedSource(Input.GetAttr(), &Source, &SourceName, &AttributeType))
 				{
 					if (!GetTextureParameterValue(
 							Input,
@@ -951,32 +961,39 @@ namespace UE
 					{
 						// Check whether this input receives its value through a connection to a
 						// primvar reader shader.
-						if ( GetPrimvarReaderParameterValue( Input, UnrealIdentifiers::UsdPrimvarReader_float3, DefaultValue, OutValue ) )
+						if (GetPrimvarReaderParameterValue(Input, UnrealIdentifiers::UsdPrimvarReader_float3, DefaultValue, OutValue))
 						{
 							return true;
 						}
 
 						// Check if we have a fallback input that we can use instead, since we don't have a valid texture value
-						if ( const pxr::UsdShadeInput FallbackInput = Source.GetInput( UnrealIdentifiers::Fallback ) )
+						if (const pxr::UsdShadeInput FallbackInput = Source.GetInput(UnrealIdentifiers::Fallback))
 						{
 							pxr::GfVec3f UsdFallbackVec3;
-							if ( FallbackInput.Get<pxr::GfVec3f>( &UsdFallbackVec3 ) )
+							if (FallbackInput.Get<pxr::GfVec3f>(&UsdFallbackVec3))
 							{
-								OutValue.Set<FVector>( UsdToUnreal::ConvertVector( UsdFallbackVec3 ) );
+								OutValue.Set<FVector>(UsdToUnreal::ConvertVector(UsdFallbackVec3));
 								return true;
 							}
 
 							pxr::GfVec4f UsdFallbackVec4;
-							if ( FallbackInput.Get<pxr::GfVec4f>( &UsdFallbackVec4 ) )
+							if (FallbackInput.Get<pxr::GfVec4f>(&UsdFallbackVec4))
 							{
-								if ( !FMath::IsNearlyEqual( UsdFallbackVec4[ 3 ], 1.0f ) )
+								if (!FMath::IsNearlyEqual(UsdFallbackVec4[3], 1.0f))
 								{
-									UE_LOG( LogUsd, Warning, TEXT( "Ignoring alpha value from fallback GfVec4f [%f, %f, %f, %f] used for Shader '%s'" ),
-										UsdFallbackVec4[ 0 ], UsdFallbackVec4[ 1 ], UsdFallbackVec4[ 2 ], UsdFallbackVec4[ 3 ], *UsdToUnreal::ConvertPath( Source.GetPrim().GetPrimPath() )
+									UE_LOG(
+										LogUsd,
+										Warning,
+										TEXT("Ignoring alpha value from fallback GfVec4f [%f, %f, %f, %f] used for Shader '%s'"),
+										UsdFallbackVec4[0],
+										UsdFallbackVec4[1],
+										UsdFallbackVec4[2],
+										UsdFallbackVec4[3],
+										*UsdToUnreal::ConvertPath(Source.GetPrim().GetPrimPath())
 									);
 								}
 
-								OutValue.Set<FVector>( FVector{ UsdFallbackVec4[ 0 ], UsdFallbackVec4[ 1 ], UsdFallbackVec4[ 2 ] } );
+								OutValue.Set<FVector>(FVector{UsdFallbackVec4[0], UsdFallbackVec4[1], UsdFallbackVec4[2]});
 								return true;
 							}
 						}
@@ -995,16 +1012,16 @@ namespace UE
 					}
 				}
 				// No other node connected, so we must have some value
-				else if ( InputName != UnrealIdentifiers::Normal )
+				else if (InputName != UnrealIdentifiers::Normal)
 				{
 					FLinearColor DiffuseColor = DefaultValue;
 					pxr::GfVec3f UsdDiffuseColor;
-					if ( Input.Get< pxr::GfVec3f >( &UsdDiffuseColor ) )
+					if (Input.Get<pxr::GfVec3f>(&UsdDiffuseColor))
 					{
-						DiffuseColor = UsdToUnreal::ConvertColor( UsdDiffuseColor );
+						DiffuseColor = UsdToUnreal::ConvertColor(UsdDiffuseColor);
 					}
 
-					OutValue.Set<FVector>( FVector(DiffuseColor) );
+					OutValue.Set<FVector>(FVector(DiffuseColor));
 				}
 
 				return true;
@@ -1021,8 +1038,8 @@ namespace UE
 			{
 				FScopedUsdAllocs Allocs;
 
-				pxr::UsdShadeInput Input = Connectable.GetInput( InputName );
-				if ( !Input )
+				pxr::UsdShadeInput Input = Connectable.GetInput(InputName);
+				if (!Input)
 				{
 					return false;
 				}
@@ -1031,7 +1048,7 @@ namespace UE
 				pxr::UsdShadeConnectableAPI Source;
 				pxr::TfToken SourceName;
 				pxr::UsdShadeAttributeType AttributeType;
-				if ( pxr::UsdShadeConnectableAPI::GetConnectedSource( Input.GetAttr(), &Source, &SourceName, &AttributeType ) )
+				if (pxr::UsdShadeConnectableAPI::GetConnectedSource(Input.GetAttr(), &Source, &SourceName, &AttributeType))
 				{
 					// Recurse because the attribute may just be pointing at some other attribute that has the data
 					// (e.g. when shader input is just "hoisted" and connected to the parent material input)
@@ -1041,9 +1058,9 @@ namespace UE
 				else
 				{
 					bool InputValue = DefaultValue;
-					Input.Get< bool >( &InputValue );
+					Input.Get<bool>(&InputValue);
 
-					OutValue.Set< bool >( InputValue );
+					OutValue.Set<bool>(InputValue);
 				}
 
 				return true;
@@ -1058,51 +1075,58 @@ namespace UE
 			{
 				FGetExpressionForValueVisitor(UMaterial& InMaterial)
 					: Material(InMaterial)
-				{}
+				{
+				}
 
-				UMaterialExpression* operator()( const float FloatValue ) const
+				UMaterialExpression* operator()(const float FloatValue) const
 				{
-					UMaterialExpressionConstant* Expression = Cast< UMaterialExpressionConstant >( UMaterialEditingLibrary::CreateMaterialExpression( &Material, UMaterialExpressionConstant::StaticClass() ) );
+					UMaterialExpressionConstant* Expression = Cast<UMaterialExpressionConstant>(
+						UMaterialEditingLibrary::CreateMaterialExpression(&Material, UMaterialExpressionConstant::StaticClass())
+					);
 					Expression->R = FloatValue;
 
 					return Expression;
 				}
 
-				UMaterialExpression* operator()( const FVector& VectorValue ) const
+				UMaterialExpression* operator()(const FVector& VectorValue) const
 				{
-					UMaterialExpressionConstant4Vector* Expression = Cast< UMaterialExpressionConstant4Vector >( UMaterialEditingLibrary::CreateMaterialExpression( &Material, UMaterialExpressionConstant4Vector::StaticClass() ) );
-					Expression->Constant = FLinearColor( VectorValue );
+					UMaterialExpressionConstant4Vector* Expression = Cast<UMaterialExpressionConstant4Vector>(
+						UMaterialEditingLibrary::CreateMaterialExpression(&Material, UMaterialExpressionConstant4Vector::StaticClass())
+					);
+					Expression->Constant = FLinearColor(VectorValue);
 
 					return Expression;
 				}
 
-				UMaterialExpression* operator()( const FTextureParameterValue& TextureValue ) const
+				UMaterialExpression* operator()(const FTextureParameterValue& TextureValue) const
 				{
-					UMaterialExpressionTextureSample* Expression = Cast< UMaterialExpressionTextureSample >( UMaterialEditingLibrary::CreateMaterialExpression( &Material, UMaterialExpressionTextureSample::StaticClass() ) );
+					UMaterialExpressionTextureSample* Expression = Cast<UMaterialExpressionTextureSample>(
+						UMaterialEditingLibrary::CreateMaterialExpression(&Material, UMaterialExpressionTextureSample::StaticClass())
+					);
 					Expression->Texture = TextureValue.Texture;
-					Expression->SamplerType = UMaterialExpressionTextureBase::GetSamplerTypeForTexture( TextureValue.Texture );
+					Expression->SamplerType = UMaterialExpressionTextureBase::GetSamplerTypeForTexture(TextureValue.Texture);
 
 					return Expression;
 				}
 
-				UMaterialExpression* operator()( const FPrimvarReaderParameterValue& PrimvarReaderValue ) const
+				UMaterialExpression* operator()(const FPrimvarReaderParameterValue& PrimvarReaderValue) const
 				{
 					UMaterialExpression* Expression = nullptr;
 
 					// We currently only import the "displayColor" primvar in USD as vertex colors
 					// on UE meshes, so that's the only primvar we can support here.
-					if ( PrimvarReaderValue.PrimvarName == "displayColor" )
+					if (PrimvarReaderValue.PrimvarName == "displayColor")
 					{
-						Expression = UMaterialEditingLibrary::CreateMaterialExpression( &Material, UMaterialExpressionVertexColor::StaticClass() );
+						Expression = UMaterialEditingLibrary::CreateMaterialExpression(&Material, UMaterialExpressionVertexColor::StaticClass());
 					}
 					else
 					{
-						Expression = UMaterialEditingLibrary::CreateMaterialExpression( &Material, UMaterialExpressionConstant4Vector::StaticClass() );
+						Expression = UMaterialEditingLibrary::CreateMaterialExpression(&Material, UMaterialExpressionConstant4Vector::StaticClass());
 
-						UMaterialExpressionConstant4Vector* FallbackExpression = Cast< UMaterialExpressionConstant4Vector >( Expression );
-						if ( FallbackExpression )
+						UMaterialExpressionConstant4Vector* FallbackExpression = Cast<UMaterialExpressionConstant4Vector>(Expression);
+						if (FallbackExpression)
 						{
-							FallbackExpression->Constant = FLinearColor( PrimvarReaderValue.FallbackValue );
+							FallbackExpression->Constant = FLinearColor(PrimvarReaderValue.FallbackValue);
 						}
 					}
 
@@ -1112,16 +1136,16 @@ namespace UE
 			private:
 				UMaterial& Material;
 			};
-#endif // WITH_EDITOR
+#endif	  // WITH_EDITOR
 
-			UMaterialExpression* GetExpressionForValue( UMaterial& Material, const FParameterValue& ParameterValue )
+			UMaterialExpression* GetExpressionForValue(UMaterial& Material, const FParameterValue& ParameterValue)
 			{
 #if WITH_EDITOR
-				FGetExpressionForValueVisitor Visitor{ Material };
-				return Visit( Visitor, ParameterValue );
+				FGetExpressionForValueVisitor Visitor{Material};
+				return Visit(Visitor, ParameterValue);
 #else
 				return nullptr;
-#endif // WITH_EDITOR
+#endif	  // WITH_EDITOR
 			}
 
 			/**
@@ -1130,34 +1154,35 @@ namespace UE
 			 */
 			struct FSetParameterValueVisitor
 			{
-				explicit FSetParameterValueVisitor( UMaterialInstance& InMaterial, const TCHAR* InParameterName )
+				explicit FSetParameterValueVisitor(UMaterialInstance& InMaterial, const TCHAR* InParameterName)
 					: Material(InMaterial)
 					, ParameterName(InParameterName)
-				{}
+				{
+				}
 
-				void operator()( const float FloatValue ) const
+				void operator()(const float FloatValue) const
 				{
-					UsdUtils::SetScalarParameterValue( Material, ParameterName, FloatValue );
+					UsdUtils::SetScalarParameterValue(Material, ParameterName, FloatValue);
 				}
 
-				void operator()( const FVector& VectorValue ) const
+				void operator()(const FVector& VectorValue) const
 				{
-					UsdUtils::SetVectorParameterValue( Material, ParameterName, FLinearColor(VectorValue) );
+					UsdUtils::SetVectorParameterValue(Material, ParameterName, FLinearColor(VectorValue));
 				}
 
-				void operator()( const FTextureParameterValue& TextureValue ) const
+				void operator()(const FTextureParameterValue& TextureValue) const
 				{
-					UsdUtils::SetTextureParameterValue( Material, ParameterName, TextureValue.Texture );
+					UsdUtils::SetTextureParameterValue(Material, ParameterName, TextureValue.Texture);
 				}
 
-				void operator()( const FPrimvarReaderParameterValue& PrimvarReaderValue ) const
+				void operator()(const FPrimvarReaderParameterValue& PrimvarReaderValue) const
 				{
-					UsdUtils::SetVectorParameterValue( Material, ParameterName, FLinearColor(PrimvarReaderValue.FallbackValue) );
+					UsdUtils::SetVectorParameterValue(Material, ParameterName, FLinearColor(PrimvarReaderValue.FallbackValue));
 				}
 
-				void operator()( const bool BoolValue ) const
+				void operator()(const bool BoolValue) const
 				{
-					UsdUtils::SetBoolParameterValue( Material, ParameterName, BoolValue );
+					UsdUtils::SetBoolParameterValue(Material, ParameterName, BoolValue);
 				}
 
 			protected:
@@ -1180,39 +1205,32 @@ namespace UE
 				{
 				}
 
-				void operator()( const float FloatValue ) const
+				void operator()(const float FloatValue) const
 				{
-					FSetParameterValueVisitor::operator()( FloatValue );
-					UsdUtils::SetScalarParameterValue( Material, *FString::Printf( TEXT( "Use%sTexture" ), ParameterName ), 0.0f );
+					FSetParameterValueVisitor::operator()(FloatValue);
+					UsdUtils::SetScalarParameterValue(Material, *FString::Printf(TEXT("Use%sTexture"), ParameterName), 0.0f);
 				}
 
-				void operator()( const FVector& VectorValue ) const
+				void operator()(const FVector& VectorValue) const
 				{
-					FSetParameterValueVisitor::operator()( VectorValue );
-					UsdUtils::SetScalarParameterValue( Material, *FString::Printf( TEXT( "Use%sTexture" ), ParameterName ), 0.0f );
+					FSetParameterValueVisitor::operator()(VectorValue);
+					UsdUtils::SetScalarParameterValue(Material, *FString::Printf(TEXT("Use%sTexture"), ParameterName), 0.0f);
 				}
 
-				void operator()( const FTextureParameterValue& TextureValue ) const
+				void operator()(const FTextureParameterValue& TextureValue) const
 				{
-					UsdUtils::SetTextureParameterValue( Material, *FString::Printf( TEXT( "%sTexture" ), ParameterName ), TextureValue.Texture );
-					UsdUtils::SetScalarParameterValue( Material, *FString::Printf( TEXT( "Use%sTexture" ), ParameterName ), 1.0f );
+					UsdUtils::SetTextureParameterValue(Material, *FString::Printf(TEXT("%sTexture"), ParameterName), TextureValue.Texture);
+					UsdUtils::SetScalarParameterValue(Material, *FString::Printf(TEXT("Use%sTexture"), ParameterName), 1.0f);
 
 					FLinearColor ScaleAndTranslation = FLinearColor{
 						TextureValue.UVScale.GetVector()[0],
 						TextureValue.UVScale.GetVector()[1],
 						TextureValue.UVTranslation[0],
-						TextureValue.UVTranslation[1]};
-					UsdUtils::SetVectorParameterValue(
-						Material,
-						*FString::Printf(TEXT("%sScaleTranslation"), ParameterName),
-						ScaleAndTranslation
-					);
+						TextureValue.UVTranslation[1]
+					};
+					UsdUtils::SetVectorParameterValue(Material, *FString::Printf(TEXT("%sScaleTranslation"), ParameterName), ScaleAndTranslation);
 
-					UsdUtils::SetScalarParameterValue(
-						Material,
-						*FString::Printf(TEXT("%sRotation"), ParameterName),
-						TextureValue.UVRotation
-					);
+					UsdUtils::SetScalarParameterValue(Material, *FString::Printf(TEXT("%sRotation"), ParameterName), TextureValue.UVRotation);
 
 					if (const int32* FoundIndex = PrimvarToUVIndex.Find(TextureValue.Primvar))
 					{
@@ -1229,45 +1247,43 @@ namespace UE
 							ParameterName,
 							*Material.GetPathName(),
 							*UsdUtils::StringifyMap(PrimvarToUVIndex),
-							TextureValue.Primvar.IsEmpty() ? TEXT(" Is your UsdUVTexture Shader missing the 'inputs:st' attribute? (It specifies which UV set to sample the texture with)") : TEXT("")
+							TextureValue.Primvar.IsEmpty() ? TEXT(" Is your UsdUVTexture Shader missing the 'inputs:st' attribute? (It specifies "
+																  "which UV set to sample the texture with)")
+														   : TEXT("")
 						);
 					}
 
 					FLinearColor ComponentMask;
 					switch (TextureValue.OutputIndex)
 					{
-					case 0: // RGB
-						ComponentMask = FLinearColor{1.f, 1.f, 1.f, 0.f};
-						break;
-					case 1: // R
-						ComponentMask = FLinearColor{1.f, 0.f, 0.f, 0.f};
-						break;
-					case 2: // G
-						ComponentMask = FLinearColor{0.f, 1.f, 0.f, 0.f};
-						break;
-					case 3: // B
-						ComponentMask = FLinearColor{0.f, 0.f, 1.f, 0.f};
-						break;
-					case 4: // A
-						ComponentMask = FLinearColor{0.f, 0.f, 0.f, 1.f};
-						break;
+						case 0:	   // RGB
+							ComponentMask = FLinearColor{1.f, 1.f, 1.f, 0.f};
+							break;
+						case 1:	   // R
+							ComponentMask = FLinearColor{1.f, 0.f, 0.f, 0.f};
+							break;
+						case 2:	   // G
+							ComponentMask = FLinearColor{0.f, 1.f, 0.f, 0.f};
+							break;
+						case 3:	   // B
+							ComponentMask = FLinearColor{0.f, 0.f, 1.f, 0.f};
+							break;
+						case 4:	   // A
+							ComponentMask = FLinearColor{0.f, 0.f, 0.f, 1.f};
+							break;
 					}
-					UsdUtils::SetVectorParameterValue(
-						Material,
-						*FString::Printf(TEXT("%sTextureComponent"), ParameterName),
-						ComponentMask
-					);
+					UsdUtils::SetVectorParameterValue(Material, *FString::Printf(TEXT("%sTextureComponent"), ParameterName), ComponentMask);
 				}
 
-				void operator()( const FPrimvarReaderParameterValue& PrimvarReaderValue ) const
+				void operator()(const FPrimvarReaderParameterValue& PrimvarReaderValue) const
 				{
-					FSetParameterValueVisitor::operator()( PrimvarReaderValue.FallbackValue );
+					FSetParameterValueVisitor::operator()(PrimvarReaderValue.FallbackValue);
 
 					// We currently only import the "displayColor" primvar in USD as vertex colors
 					// on UE meshes, so that's the only primvar we can support here.
-					if ( PrimvarReaderValue.PrimvarName == "displayColor" )
+					if (PrimvarReaderValue.PrimvarName == "displayColor")
 					{
-						UsdUtils::SetScalarParameterValue( Material, TEXT( "UseVertexColorForBaseColor" ), 1.0f );
+						UsdUtils::SetScalarParameterValue(Material, TEXT("UseVertexColorForBaseColor"), 1.0f);
 					}
 				}
 
@@ -1283,68 +1299,73 @@ namespace UE
 				const TMap<FString, int32>& PrimvarToUVIndex
 			)
 			{
-				if ( bForUsdPreviewSurface )
+				if (bForUsdPreviewSurface)
 				{
 					FSetPreviewSurfaceParameterValueVisitor Visitor{Material, ParameterName, PrimvarToUVIndex};
-					Visit( Visitor, ParameterValue );
+					Visit(Visitor, ParameterValue);
 				}
 				else
 				{
 					FSetParameterValueVisitor Visitor{Material, ParameterName};
-					Visit( Visitor, ParameterValue );
+					Visit(Visitor, ParameterValue);
 				}
 			}
 
-			UTexture* CreateTextureWithEditor( const pxr::UsdAttribute& TextureAssetPathAttr, const FString& PrimPath, TextureGroup LODGroup, UObject* Outer )
+			UTexture* CreateTextureWithEditor(
+				const pxr::UsdAttribute& TextureAssetPathAttr,
+				const FString& PrimPath,
+				TextureGroup LODGroup,
+				UObject* Outer
+			)
 			{
 				UTexture* Texture = nullptr;
 #if WITH_EDITOR
 				FScopedUsdAllocs UsdAllocs;
 
 				pxr::SdfAssetPath TextureAssetPath;
-				TextureAssetPathAttr.Get< pxr::SdfAssetPath >( &TextureAssetPath );
+				TextureAssetPathAttr.Get<pxr::SdfAssetPath>(&TextureAssetPath);
 
-				FString TexturePath = UsdToUnreal::ConvertString( TextureAssetPath.GetAssetPath() );
-				const bool bIsSupportedUdimTexture = TexturePath.Contains( TEXT( "<UDIM>" ) );
-				FPaths::NormalizeFilename( TexturePath );
+				FString TexturePath = UsdToUnreal::ConvertString(TextureAssetPath.GetAssetPath());
+				const bool bIsSupportedUdimTexture = TexturePath.Contains(TEXT("<UDIM>"));
+				FPaths::NormalizeFilename(TexturePath);
 
 				FScopedUnrealAllocs UnrealAllocs;
 
-				if ( !TexturePath.IsEmpty() )
+				if (!TexturePath.IsEmpty())
 				{
 					bool bOutCancelled = false;
-					UTextureFactory* TextureFactory = NewObject< UTextureFactory >();
+					UTextureFactory* TextureFactory = NewObject<UTextureFactory>();
 					TextureFactory->SuppressImportOverwriteDialog();
 					TextureFactory->bUseHashAsGuid = true;
 					TextureFactory->LODGroup = LODGroup;
 					TextureFactory->HDRImportShouldBeLongLatCubeMap = EAppReturnType::YesAll;
 
-					if ( bIsSupportedUdimTexture )
+					if (bIsSupportedUdimTexture)
 					{
-						FString BaseFileName = FPaths::GetBaseFilename( TexturePath );
+						FString BaseFileName = FPaths::GetBaseFilename(TexturePath);
 
 						FString BaseFileNameBeforeUdim;
 						FString BaseFileNameAfterUdim;
-						BaseFileName.Split( TEXT( "<UDIM>" ), &BaseFileNameBeforeUdim, &BaseFileNameAfterUdim );
+						BaseFileName.Split(TEXT("<UDIM>"), &BaseFileNameBeforeUdim, &BaseFileNameAfterUdim);
 
-						FString UdimRegexPattern = FString::Printf( TEXT( R"((%s)(\d{4})(%s))" ), *BaseFileNameBeforeUdim, *BaseFileNameAfterUdim );
-						TextureFactory->UdimRegexPattern = MoveTemp( UdimRegexPattern );
+						FString UdimRegexPattern = FString::Printf(TEXT(R"((%s)(\d{4})(%s))"), *BaseFileNameBeforeUdim, *BaseFileNameAfterUdim);
+						TextureFactory->UdimRegexPattern = MoveTemp(UdimRegexPattern);
 					}
 
-					const FString ResolvedTexturePath = UsdUtils::GetResolvedTexturePath( TextureAssetPathAttr );
-					if ( !ResolvedTexturePath.IsEmpty() )
+					const FString ResolvedTexturePath = UsdUtils::GetResolvedTexturePath(TextureAssetPathAttr);
+					if (!ResolvedTexturePath.IsEmpty())
 					{
 						EObjectFlags ObjectFlags = RF_Transactional | RF_Transient;
-						if ( !Outer )
+						if (!Outer)
 						{
 							Outer = GetTransientPackage();
 						}
 
 						FString TextureExtension;
-						if ( IsInsideUsdzArchive( ResolvedTexturePath, TextureExtension ) )
+						if (IsInsideUsdzArchive(ResolvedTexturePath, TextureExtension))
 						{
 							// Always prefer using the TextureFactory if we can, as it may provide compression, which the runtime version never will
-							Texture = ReadTextureFromUsdzArchiveEditor( ResolvedTexturePath, TextureExtension, TextureFactory, Outer, ObjectFlags );
+							Texture = ReadTextureFromUsdzArchiveEditor(ResolvedTexturePath, TextureExtension, TextureFactory, Outer, ObjectFlags);
 						}
 						// Not inside an USDZ archive, just a regular texture
 						else
@@ -1352,14 +1373,20 @@ namespace UE
 							FName TextureName = MakeUniqueObjectName(
 								Outer,
 								UTexture::StaticClass(),
-								*IUsdClassesModule::SanitizeObjectName(
-									FPaths::GetBaseFilename(ResolvedTexturePath)
-								)
+								*IUsdClassesModule::SanitizeObjectName(FPaths::GetBaseFilename(ResolvedTexturePath))
 							);
-							Texture = Cast< UTexture >( TextureFactory->ImportObject( UTexture::StaticClass(), Outer, TextureName, ObjectFlags, ResolvedTexturePath, TEXT(""), bOutCancelled));
+							Texture = Cast<UTexture>(TextureFactory->ImportObject(
+								UTexture::StaticClass(),
+								Outer,
+								TextureName,
+								ObjectFlags,
+								ResolvedTexturePath,
+								TEXT(""),
+								bOutCancelled
+							));
 						}
 
-						if ( Texture )
+						if (Texture)
 						{
 							UUsdAssetUserData* UserData = NewObject<UUsdAssetUserData>(Texture, TEXT("USDAssetUserData"));
 							UserData->PrimPaths = {PrimPath};
@@ -1373,40 +1400,40 @@ namespace UE
 						}
 					}
 				}
-#endif // WITH_EDITOR
+#endif	  // WITH_EDITOR
 				return Texture;
 			}
 
-			UTexture* CreateTextureAtRuntime( const pxr::UsdAttribute& TextureAssetPathAttr, const FString& PrimPath, TextureGroup LODGroup )
+			UTexture* CreateTextureAtRuntime(const pxr::UsdAttribute& TextureAssetPathAttr, const FString& PrimPath, TextureGroup LODGroup)
 			{
 				FScopedUsdAllocs UsdAllocs;
 
 				pxr::SdfAssetPath TextureAssetPath;
-				TextureAssetPathAttr.Get< pxr::SdfAssetPath >( &TextureAssetPath );
+				TextureAssetPathAttr.Get<pxr::SdfAssetPath>(&TextureAssetPath);
 
-				FString TexturePath = UsdToUnreal::ConvertString( TextureAssetPath.GetAssetPath() );
-				FPaths::NormalizeFilename( TexturePath );
+				FString TexturePath = UsdToUnreal::ConvertString(TextureAssetPath.GetAssetPath());
+				FPaths::NormalizeFilename(TexturePath);
 
 				FScopedUnrealAllocs UnrealAllocs;
 
 				UTexture* Texture = nullptr;
 
-				if ( !TexturePath.IsEmpty() )
+				if (!TexturePath.IsEmpty())
 				{
-					const FString ResolvedTexturePath = UsdUtils::GetResolvedTexturePath( TextureAssetPathAttr );
-					if ( !ResolvedTexturePath.IsEmpty() )
+					const FString ResolvedTexturePath = UsdUtils::GetResolvedTexturePath(TextureAssetPathAttr);
+					if (!ResolvedTexturePath.IsEmpty())
 					{
 						// Try checking if the texture is inside an USDZ archive first, or else TextureFactory throws an error
 						FString TextureExtension;
-						if ( IsInsideUsdzArchive( ResolvedTexturePath, TextureExtension ) )
+						if (IsInsideUsdzArchive(ResolvedTexturePath, TextureExtension))
 						{
-							Texture = ReadTextureFromUsdzArchiveRuntime( ResolvedTexturePath );
+							Texture = ReadTextureFromUsdzArchiveRuntime(ResolvedTexturePath);
 						}
 
 						// Not inside an USDZ archive, just a regular texture
-						if ( !Texture )
+						if (!Texture)
 						{
-							Texture = FImageUtils::ImportFileAsTexture2D( ResolvedTexturePath );
+							Texture = FImageUtils::ImportFileAsTexture2D(ResolvedTexturePath);
 						}
 					}
 				}
@@ -1415,8 +1442,8 @@ namespace UE
 			}
 
 #if WITH_EDITOR
-			// Note that we will bake things that aren't supported on the default USD surface shader schema. These could be useful in case the user has a custom
-			// renderer though, and they can pick which properties they want anyway
+			// Note that we will bake things that aren't supported on the default USD surface shader schema. These could be useful in case the user
+			// has a custom renderer though, and they can pick which properties they want anyway
 			bool BakeMaterial(
 				const UMaterialInterface& Material,
 				const TArray<FPropertyEntry>& InMaterialProperties,
@@ -1428,90 +1455,93 @@ namespace UE
 				const bool bAllQualityLevels = true;
 				const bool bAllFeatureLevels = true;
 				TArray<UTexture*> MaterialTextures;
-				Material.GetUsedTextures( MaterialTextures, EMaterialQualityLevel::Num, bAllQualityLevels, GMaxRHIFeatureLevel, bAllFeatureLevels );
+				Material.GetUsedTextures(MaterialTextures, EMaterialQualityLevel::Num, bAllQualityLevels, GMaxRHIFeatureLevel, bAllFeatureLevels);
 
 				// Precache all used textures, otherwise could get everything rendered with low-res textures.
-				for ( UTexture* Texture : MaterialTextures )
+				for (UTexture* Texture : MaterialTextures)
 				{
-					if ( UTexture2D* Texture2D = Cast<UTexture2D>( Texture ) )
+					if (UTexture2D* Texture2D = Cast<UTexture2D>(Texture))
 					{
-						Texture2D->SetForceMipLevelsToBeResident( 30.0f );
+						Texture2D->SetForceMipLevelsToBeResident(30.0f);
 						Texture2D->WaitForStreaming();
 					}
 				}
 
 				FMaterialData MatSet;
-				MatSet.Material = const_cast<UMaterialInterface*>(&Material); // We don't modify it and neither does the material baking module, it's just a bad signature
+				MatSet.Material = const_cast<UMaterialInterface*>(&Material);	 // We don't modify it and neither does the material baking module,
+																				 // it's just a bad signature
 				MatSet.bPerformShrinking = bInDecayTexturesToSinglePixel;
 
-				for ( const FPropertyEntry& Entry : InMaterialProperties )
+				for (const FPropertyEntry& Entry : InMaterialProperties)
 				{
 					// No point in asking it to bake if we're going to use the user-supplied value
-					if ( Entry.bUseConstantValue )
+					if (Entry.bUseConstantValue)
 					{
 						continue;
 					}
 
-					switch ( Entry.Property )
+					switch (Entry.Property)
 					{
-					case MP_Normal:
-						if ( !Material.GetMaterial()->HasNormalConnected() && !Material.GetMaterial()->bUseMaterialAttributes )
-						{
-							continue;
-						}
-						break;
-					case MP_Tangent:
-						if ( !Material.GetMaterial()->GetEditorOnlyData()->Tangent.IsConnected() && !Material.GetMaterial()->bUseMaterialAttributes )
-						{
-							continue;
-						}
-						break;
-					case MP_EmissiveColor:
-						if ( !Material.GetMaterial()->GetEditorOnlyData()->EmissiveColor.IsConnected() && !Material.GetMaterial()->bUseMaterialAttributes )
-						{
-							continue;
-						}
-						break;
-					case MP_OpacityMask:
-						if ( !Material.IsPropertyActive( MP_OpacityMask ) || !IsMaskedBlendMode(Material) )
-						{
-							continue;
-						}
-						break;
-					case MP_Opacity:
-						if ( !Material.IsPropertyActive( MP_Opacity ) || !IsTranslucentBlendMode(Material) )
-						{
-							continue;
-						}
-						break;
-					case MP_MAX:
-						continue;
-						break;
-					default:
-						if ( !Material.IsPropertyActive( Entry.Property ) )
-						{
+						case MP_Normal:
+							if (!Material.GetMaterial()->HasNormalConnected() && !Material.GetMaterial()->bUseMaterialAttributes)
+							{
+								continue;
+							}
+							break;
+						case MP_Tangent:
+							if (!Material.GetMaterial()->GetEditorOnlyData()->Tangent.IsConnected()
+								&& !Material.GetMaterial()->bUseMaterialAttributes)
+							{
+								continue;
+							}
+							break;
+						case MP_EmissiveColor:
+							if (!Material.GetMaterial()->GetEditorOnlyData()->EmissiveColor.IsConnected()
+								&& !Material.GetMaterial()->bUseMaterialAttributes)
+							{
+								continue;
+							}
+							break;
+						case MP_OpacityMask:
+							if (!Material.IsPropertyActive(MP_OpacityMask) || !IsMaskedBlendMode(Material))
+							{
+								continue;
+							}
+							break;
+						case MP_Opacity:
+							if (!Material.IsPropertyActive(MP_Opacity) || !IsTranslucentBlendMode(Material))
+							{
+								continue;
+							}
+							break;
+						case MP_MAX:
 							continue;
-						}
-						break;
+							break;
+						default:
+							if (!Material.IsPropertyActive(Entry.Property))
+							{
+								continue;
+							}
+							break;
 					}
 
-					MatSet.PropertySizes.Add( Entry.Property, Entry.bUseCustomSize ? Entry.CustomSize : InDefaultTextureSize );
+					MatSet.PropertySizes.Add(Entry.Property, Entry.bUseCustomSize ? Entry.CustomSize : InDefaultTextureSize);
 				}
 
 				FMeshData MeshSettings;
 				MeshSettings.MeshDescription = nullptr;
-				MeshSettings.TextureCoordinateBox = FBox2D( FVector2D( 0.0f, 0.0f ), FVector2D( 1.0f, 1.0f ) );
+				MeshSettings.TextureCoordinateBox = FBox2D(FVector2D(0.0f, 0.0f), FVector2D(1.0f, 1.0f));
 				MeshSettings.TextureCoordinateIndex = 0;
 
 				TArray<FBakeOutput> BakeOutputs;
-				IMaterialBakingModule& Module = FModuleManager::Get().LoadModuleChecked<IMaterialBakingModule>( "MaterialBaking" );
+				IMaterialBakingModule& Module = FModuleManager::Get().LoadModuleChecked<IMaterialBakingModule>("MaterialBaking");
 				const bool bLinearBake = true;
-				Module.SetLinearBake( bLinearBake );
+				Module.SetLinearBake(bLinearBake);
 				const bool bEmissiveHDR = true;
-				Module.SetEmissiveHDR( bEmissiveHDR );
-				Module.BakeMaterials( { &MatSet }, { &MeshSettings }, BakeOutputs );
+				Module.SetEmissiveHDR(bEmissiveHDR);
+				Module.BakeMaterials({&MatSet}, {&MeshSettings}, BakeOutputs);
 
-				if ( BakeOutputs.Num() < 1 )
+				if (BakeOutputs.Num() < 1)
 				{
 					return false;
 				}
@@ -1520,115 +1550,128 @@ namespace UE
 				return true;
 			}
 
-			// Writes textures for all baked channels in BakedSamples that are larger than 1x1, and returns the filenames of the emitted textures for each channel
-			TMap<EMaterialProperty, FString> WriteTextures( FBakedMaterialView& BakedSamples, const FString& TextureNamePrefix, const FDirectoryPath& TexturesFolder )
+			// Writes textures for all baked channels in BakedSamples that are larger than 1x1, and returns the filenames of the emitted textures for
+			// each channel
+			TMap<EMaterialProperty, FString> WriteTextures(
+				FBakedMaterialView& BakedSamples,
+				const FString& TextureNamePrefix,
+				const FDirectoryPath& TexturesFolder
+			)
 			{
-				IImageWrapperModule& ImageWrapperModule = FModuleManager::Get().LoadModuleChecked<IImageWrapperModule>( TEXT( "ImageWrapper" ) );
-				TSharedPtr<IImageWrapper> EXRImageWrapper = ImageWrapperModule.CreateImageWrapper( EImageFormat::EXR );
-				TSharedPtr<IImageWrapper> PNGImageWrapper = ImageWrapperModule.CreateImageWrapper( EImageFormat::PNG );
+				IImageWrapperModule& ImageWrapperModule = FModuleManager::Get().LoadModuleChecked<IImageWrapperModule>(TEXT("ImageWrapper"));
+				TSharedPtr<IImageWrapper> EXRImageWrapper = ImageWrapperModule.CreateImageWrapper(EImageFormat::EXR);
+				TSharedPtr<IImageWrapper> PNGImageWrapper = ImageWrapperModule.CreateImageWrapper(EImageFormat::PNG);
 
 				TMap<EMaterialProperty, FString> WrittenTexturesPerChannel;
 
 				// Write textures for HDR baked properties larger than 1x1
-				for ( TPair<EMaterialProperty, TArray<FFloat16Color>*>& BakedDataPair : BakedSamples.HDRPropertyData )
+				for (TPair<EMaterialProperty, TArray<FFloat16Color>*>& BakedDataPair : BakedSamples.HDRPropertyData)
 				{
 					const EMaterialProperty Property = BakedDataPair.Key;
 
 					TArray<FFloat16Color>* Samples = BakedDataPair.Value;
-					if ( !Samples || Samples->Num() < 2 )
+					if (!Samples || Samples->Num() < 2)
 					{
 						continue;
 					}
 
-					const FIntPoint& FinalSize = BakedSamples.PropertySizes.FindChecked( Property );
-					if ( FinalSize.GetMin() < 2 )
+					const FIntPoint& FinalSize = BakedSamples.PropertySizes.FindChecked(Property);
+					if (FinalSize.GetMin() < 2)
 					{
 						continue;
 					}
 
 					const UEnum* PropertyEnum = StaticEnum<EMaterialProperty>();
-					FName PropertyName = PropertyEnum->GetNameByValue( Property );
+					FName PropertyName = PropertyEnum->GetNameByValue(Property);
 					FString TrimmedPropertyName = PropertyName.ToString();
-					TrimmedPropertyName.RemoveFromStart( TEXT( "MP_" ) );
-
-					// Final FilePath will be something like "C:/TexturesFolder/Game_ContentFolder_Materials_Red_BaseColor.png", which automatically guarantees
-					// it won't overwrite another texture from the same export, but will overwrite old textures from previous exports
-					FString TextureFileName = IUsdClassesModule::SanitizeObjectName( FPaths::ChangeExtension( TextureNamePrefix, TEXT( "" ) ) );
-					TextureFileName.RemoveFromStart( TEXT( "_" ) );
-					FString TextureFilePath = FPaths::Combine( TexturesFolder.Path, FString::Printf( TEXT( "%s_%s.exr" ), *IUsdClassesModule::SanitizeObjectName( TextureFileName ), *TrimmedPropertyName ) );
+					TrimmedPropertyName.RemoveFromStart(TEXT("MP_"));
+
+					// Final FilePath will be something like "C:/TexturesFolder/Game_ContentFolder_Materials_Red_BaseColor.png", which automatically
+					// guarantees it won't overwrite another texture from the same export, but will overwrite old textures from previous exports
+					FString TextureFileName = IUsdClassesModule::SanitizeObjectName(FPaths::ChangeExtension(TextureNamePrefix, TEXT("")));
+					TextureFileName.RemoveFromStart(TEXT("_"));
+					FString TextureFilePath = FPaths::Combine(
+						TexturesFolder.Path,
+						FString::Printf(TEXT("%s_%s.exr"), *IUsdClassesModule::SanitizeObjectName(TextureFileName), *TrimmedPropertyName)
+					);
 
-					// For some reason the baked samples always have zero alpha and there is nothing we can do about it... It seems like the material baking module is made
-					// with the intent that the data ends up in UTexture2Ds, where they can be set to be compressed without alpha and have the value ignored.
-					// Since we need to write these to file, we must set them back up to full alpha. This is potentially useless as USD handles these at most as color3f, but
-					// it could be annoying for the user if they intend on using the textures for anything else
-					for ( FFloat16Color& Sample : *Samples )
+					// For some reason the baked samples always have zero alpha and there is nothing we can do about it... It seems like the material
+					// baking module is made with the intent that the data ends up in UTexture2Ds, where they can be set to be compressed without
+					// alpha and have the value ignored. Since we need to write these to file, we must set them back up to full alpha. This is
+					// potentially useless as USD handles these at most as color3f, but it could be annoying for the user if they intend on using the
+					// textures for anything else
+					for (FFloat16Color& Sample : *Samples)
 					{
 						Sample.A = 1.0f;
 					}
 
-					EXRImageWrapper->SetRaw( Samples->GetData(), Samples->GetAllocatedSize(), FinalSize.X, FinalSize.Y, ERGBFormat::RGBAF, 16 );
-					const TArray64<uint8> Data = EXRImageWrapper->GetCompressed( 100 );
+					EXRImageWrapper->SetRaw(Samples->GetData(), Samples->GetAllocatedSize(), FinalSize.X, FinalSize.Y, ERGBFormat::RGBAF, 16);
+					const TArray64<uint8> Data = EXRImageWrapper->GetCompressed(100);
 
-					bool bWroteFile = FFileHelper::SaveArrayToFile( Data, *TextureFilePath );
-					if ( bWroteFile )
+					bool bWroteFile = FFileHelper::SaveArrayToFile(Data, *TextureFilePath);
+					if (bWroteFile)
 					{
-						WrittenTexturesPerChannel.Add( Property, TextureFilePath );
+						WrittenTexturesPerChannel.Add(Property, TextureFilePath);
 					}
 					else
 					{
-						UE_LOG( LogUsd, Warning, TEXT( "Failed to write texture '%s', baked channel will be ignored." ), *TextureFilePath );
+						UE_LOG(LogUsd, Warning, TEXT("Failed to write texture '%s', baked channel will be ignored."), *TextureFilePath);
 					}
 				}
 
 				// Write textures for baked properties larger than 1x1
-				for ( TPair<EMaterialProperty, TArray<FColor>*>& BakedDataPair : BakedSamples.PropertyData )
+				for (TPair<EMaterialProperty, TArray<FColor>*>& BakedDataPair : BakedSamples.PropertyData)
 				{
 					const EMaterialProperty Property = BakedDataPair.Key;
 
 					// The material baking module still generates and sends an SDR version of any HDR channel it also bakes,
 					// so here let's skip this one in case we already generated an HDR texture for the channel
-					if ( WrittenTexturesPerChannel.Contains( Property ) )
+					if (WrittenTexturesPerChannel.Contains(Property))
 					{
 						continue;
 					}
 
 					TArray<FColor>* Samples = BakedDataPair.Value;
-					if ( !Samples || Samples->Num() < 2 )
+					if (!Samples || Samples->Num() < 2)
 					{
 						continue;
 					}
 
-					const FIntPoint& FinalSize = BakedSamples.PropertySizes.FindChecked( Property );
-					if ( FinalSize.GetMin() < 2 )
+					const FIntPoint& FinalSize = BakedSamples.PropertySizes.FindChecked(Property);
+					if (FinalSize.GetMin() < 2)
 					{
 						continue;
 					}
 
 					const UEnum* PropertyEnum = StaticEnum<EMaterialProperty>();
-					FName PropertyName = PropertyEnum->GetNameByValue( Property );
+					FName PropertyName = PropertyEnum->GetNameByValue(Property);
 					FString TrimmedPropertyName = PropertyName.ToString();
-					TrimmedPropertyName.RemoveFromStart( TEXT( "MP_" ) );
-
-					// Final FilePath will be something like "C:/TexturesFolder/Game_ContentFolder_Materials_Red_BaseColor.png", which automatically guarantees
-					// it won't overwrite another texture from the same export, but will overwrite old textures from previous exports
-					FString TextureFileName = IUsdClassesModule::SanitizeObjectName( FPaths::ChangeExtension( TextureNamePrefix, TEXT( "" ) ) );
-					TextureFileName.RemoveFromStart( TEXT( "_" ) );
-					FString TextureFilePath = FPaths::Combine( TexturesFolder.Path, FString::Printf( TEXT( "%s_%s.png" ), *IUsdClassesModule::SanitizeObjectName( TextureFileName ), *TrimmedPropertyName ) );
+					TrimmedPropertyName.RemoveFromStart(TEXT("MP_"));
+
+					// Final FilePath will be something like "C:/TexturesFolder/Game_ContentFolder_Materials_Red_BaseColor.png", which automatically
+					// guarantees it won't overwrite another texture from the same export, but will overwrite old textures from previous exports
+					FString TextureFileName = IUsdClassesModule::SanitizeObjectName(FPaths::ChangeExtension(TextureNamePrefix, TEXT("")));
+					TextureFileName.RemoveFromStart(TEXT("_"));
+					FString TextureFilePath = FPaths::Combine(
+						TexturesFolder.Path,
+						FString::Printf(TEXT("%s_%s.png"), *IUsdClassesModule::SanitizeObjectName(TextureFileName), *TrimmedPropertyName)
+					);
 
-					// For some reason the baked samples always have zero alpha and there is nothing we can do about it... It seems like the material baking module is made
-					// with the intent that the data ends up in UTexture2Ds, where they can be set to be compressed without alpha and have the value ignored.
-					// Since we need to write these to file, we must set them back up to full alpha. This is potentially useless as USD handles these at most as color3f, but
-					// it could be annoying for the user if they intend on using the textures for anything else
-					for ( FColor& Sample : *Samples )
+					// For some reason the baked samples always have zero alpha and there is nothing we can do about it... It seems like the material
+					// baking module is made with the intent that the data ends up in UTexture2Ds, where they can be set to be compressed without
+					// alpha and have the value ignored. Since we need to write these to file, we must set them back up to full alpha. This is
+					// potentially useless as USD handles these at most as color3f, but it could be annoying for the user if they intend on using the
+					// textures for anything else
+					for (FColor& Sample : *Samples)
 					{
 						Sample.A = 255;
 					}
 
-					PNGImageWrapper->SetRaw( Samples->GetData(), Samples->GetAllocatedSize(), FinalSize.X, FinalSize.Y, ERGBFormat::BGRA, 8 );
-					const TArray64<uint8> PNGData = PNGImageWrapper->GetCompressed( 100 );
+					PNGImageWrapper->SetRaw(Samples->GetData(), Samples->GetAllocatedSize(), FinalSize.X, FinalSize.Y, ERGBFormat::BGRA, 8);
+					const TArray64<uint8> PNGData = PNGImageWrapper->GetCompressed(100);
 
-					bool bWroteFile = FFileHelper::SaveArrayToFile( PNGData, *TextureFilePath );
-					if ( bWroteFile )
+					bool bWroteFile = FFileHelper::SaveArrayToFile(PNGData, *TextureFilePath);
+					if (bWroteFile)
 					{
 						WrittenTexturesPerChannel.Add(Property, TextureFilePath);
 					}
@@ -1641,27 +1684,35 @@ namespace UE
 				return WrittenTexturesPerChannel;
 			}
 
-			bool ConfigureShadePrim( const FBakedMaterialView& BakedData, const TMap<EMaterialProperty, FString>& WrittenTextures, const TMap<EMaterialProperty, float>& UserConstantValues, pxr::UsdShadeMaterial& OutUsdShadeMaterial )
+			bool ConfigureShadePrim(
+				const FBakedMaterialView& BakedData,
+				const TMap<EMaterialProperty, FString>& WrittenTextures,
+				const TMap<EMaterialProperty, float>& UserConstantValues,
+				pxr::UsdShadeMaterial& OutUsdShadeMaterial
+			)
 			{
 				FScopedUsdAllocs Allocs;
 
 				pxr::UsdPrim MaterialPrim = OutUsdShadeMaterial.GetPrim();
 				pxr::UsdStageRefPtr Stage = MaterialPrim.GetStage();
-				if ( !MaterialPrim || !Stage )
+				if (!MaterialPrim || !Stage)
 				{
 					return false;
 				}
 
-				FString UsdFilePath = UsdToUnreal::ConvertString( Stage->GetRootLayer()->GetRealPath() );
+				FString UsdFilePath = UsdToUnreal::ConvertString(Stage->GetRootLayer()->GetRealPath());
 
-				FUsdStageInfo StageInfo{ Stage };
+				FUsdStageInfo StageInfo{Stage};
 
 				pxr::SdfPath MaterialPath = MaterialPrim.GetPath();
 
 				// Create surface shader
-				pxr::UsdShadeShader ShadeShader = pxr::UsdShadeShader::Define( Stage, MaterialPath.AppendChild( UnrealToUsd::ConvertToken( TEXT( "SurfaceShader" ) ).Get() ) );
-				ShadeShader.SetShaderId( UnrealIdentifiers::UsdPreviewSurface );
-				pxr::UsdShadeOutput ShaderOutOutput = ShadeShader.CreateOutput( UnrealIdentifiers::Surface, pxr::SdfValueTypeNames->Token );
+				pxr::UsdShadeShader ShadeShader = pxr::UsdShadeShader::Define(
+					Stage,
+					MaterialPath.AppendChild(UnrealToUsd::ConvertToken(TEXT("SurfaceShader")).Get())
+				);
+				ShadeShader.SetShaderId(UnrealIdentifiers::UsdPreviewSurface);
+				pxr::UsdShadeOutput ShaderOutOutput = ShadeShader.CreateOutput(UnrealIdentifiers::Surface, pxr::SdfValueTypeNames->Token);
 
 				// Connect material to surface shader
 				pxr::UsdShadeOutput MaterialSurfaceOutput = OutUsdShadeMaterial.CreateSurfaceOutput();
@@ -1675,15 +1726,15 @@ namespace UE
 				// Collect all the properties we'll process, as some data is baked and some comes from values the user input as export options
 				TSet<EMaterialProperty> PropertiesToProcess;
 				{
-					PropertiesToProcess.Reserve( BakedData.PropertyData.Num() + UserConstantValues.Num() + BakedData.HDRPropertyData.Num() );
-					BakedData.PropertyData.GetKeys( PropertiesToProcess );
+					PropertiesToProcess.Reserve(BakedData.PropertyData.Num() + UserConstantValues.Num() + BakedData.HDRPropertyData.Num());
+					BakedData.PropertyData.GetKeys(PropertiesToProcess);
 
 					TSet<EMaterialProperty> UsedProperties;
-					UserConstantValues.GetKeys( UsedProperties );
-					PropertiesToProcess.Append( UsedProperties );
+					UserConstantValues.GetKeys(UsedProperties);
+					PropertiesToProcess.Append(UsedProperties);
 
-					BakedData.HDRPropertyData.GetKeys( UsedProperties );
-					PropertiesToProcess.Append( UsedProperties );
+					BakedData.HDRPropertyData.GetKeys(UsedProperties);
+					PropertiesToProcess.Append(UsedProperties);
 				}
 
 				// We always write BaseColor because the default in UE is full black, but in USD seems to be 0.18. If we left
@@ -1696,11 +1747,11 @@ namespace UE
 				const float Zero = 0.0f;
 
 				// Fill in outputs
-				for ( const EMaterialProperty& Property : PropertiesToProcess )
+				for (const EMaterialProperty& Property : PropertiesToProcess)
 				{
-					const FString* TextureFilePath = WrittenTextures.Find( Property );
-					const float* UserConstantValue = UserConstantValues.Find( Property );
-					const FIntPoint* SampleSize = BakedData.PropertySizes.Find( Property );
+					const FString* TextureFilePath = WrittenTextures.Find(Property);
+					const float* UserConstantValue = UserConstantValues.Find(Property);
+					const FIntPoint* SampleSize = BakedData.PropertySizes.Find(Property);
 
 					uint64 NumSamples = 0;
 
@@ -1711,66 +1762,58 @@ namespace UE
 					// will still bake the SDR version of the channel too. We keep both in our FBakedMaterialView
 					// because it only does the mechanism of decaying to a single value on the SDR data array
 					bool bParsedHDR = false;
-					if ( BakedData.HDRPropertyData.Contains( Property ) )
+					if (BakedData.HDRPropertyData.Contains(Property))
 					{
-						const TArray<FColor>* SDRSamples = BakedData.PropertyData.FindRef( Property );
+						const TArray<FColor>* SDRSamples = BakedData.PropertyData.FindRef(Property);
 						NumSamples = SDRSamples ? SDRSamples->Num() : 0;
 
 						const bool bDecayedToSingleSample = NumSamples == 1;
-						bPropertyValueIsConstant = ( UserConstantValue != nullptr || bDecayedToSingleSample );
+						bPropertyValueIsConstant = (UserConstantValue != nullptr || bDecayedToSingleSample);
 
-						const TArray<FFloat16Color>* Samples = BakedData.HDRPropertyData.FindRef( Property );
-						if ( bDecayedToSingleSample && Samples )
+						const TArray<FFloat16Color>* Samples = BakedData.HDRPropertyData.FindRef(Property);
+						if (bDecayedToSingleSample && Samples)
 						{
 							// If it decayed to single sample we know our SDR array only has one value,
 							// and so should our HDR array. Get the value from the HDR one to avoid an
 							// unnecessary quantization though
-							const FFloat16Color& Sample = ( *Samples )[ 0 ];
-							ConstantLinearValue = pxr::GfVec3f{ Sample.R, Sample.G, Sample.B };
+							const FFloat16Color& Sample = (*Samples)[0];
+							ConstantLinearValue = pxr::GfVec3f{Sample.R, Sample.G, Sample.B};
 						}
 
-						if ( Samples && Samples->Num() > 0 )
+						if (Samples && Samples->Num() > 0)
 						{
 							bParsedHDR = true;
 						}
 					}
 
-					if ( !bParsedHDR )
+					if (!bParsedHDR)
 					{
-						const TArray<FColor>* Samples = BakedData.PropertyData.FindRef( Property );
+						const TArray<FColor>* Samples = BakedData.PropertyData.FindRef(Property);
 						NumSamples = Samples ? Samples->Num() : 0;
 
-						bPropertyValueIsConstant = ( UserConstantValue != nullptr || NumSamples == 1 );
-						if ( NumSamples == 1 )
+						bPropertyValueIsConstant = (UserConstantValue != nullptr || NumSamples == 1);
+						if (NumSamples == 1)
 						{
-							switch ( Property )
+							switch (Property)
 							{
 								case EMaterialProperty::MP_BaseColor:
 								case EMaterialProperty::MP_SubsurfaceColor:
 								{
-									pxr::GfVec4f ConvertedColor = UnrealToUsd::ConvertColor( ( *Samples )[ 0 ] );
-									ConstantLinearValue = pxr::GfVec3f{
-										ConvertedColor[ 0 ],
-										ConvertedColor[ 1 ],
-										ConvertedColor[ 2 ]
-									};
+									pxr::GfVec4f ConvertedColor = UnrealToUsd::ConvertColor((*Samples)[0]);
+									ConstantLinearValue = pxr::GfVec3f{ConvertedColor[0], ConvertedColor[1], ConvertedColor[2]};
 									break;
 								}
 								case EMaterialProperty::MP_Normal:
 								case EMaterialProperty::MP_Tangent:
 								{
-									FVector ConvertedNormal{ ( *Samples )[ 0 ].ReinterpretAsLinear() };
-									ConstantLinearValue = UnrealToUsd::ConvertVectorFloat( StageInfo, ConvertedNormal );
+									FVector ConvertedNormal{(*Samples)[0].ReinterpretAsLinear()};
+									ConstantLinearValue = UnrealToUsd::ConvertVectorFloat(StageInfo, ConvertedNormal);
 									break;
 								}
 								default:
 								{
-									const FColor& Sample = ( *Samples )[ 0 ];
-									ConstantLinearValue = pxr::GfVec3f{
-										Sample.R / 255.0f,
-										Sample.G / 255.0f,
-										Sample.B / 255.0f
-									};
+									const FColor& Sample = (*Samples)[0];
+									ConstantLinearValue = pxr::GfVec3f{Sample.R / 255.0f, Sample.G / 255.0f, Sample.B / 255.0f};
 									break;
 								}
 							}
@@ -1783,24 +1826,37 @@ namespace UE
 						bPropertyValueIsConstant = true;
 					}
 
-					if ( ( NumSamples == 0 || !SampleSize ) && !UserConstantValue )
+					if ((NumSamples == 0 || !SampleSize) && !UserConstantValue)
 					{
-						UE_LOG( LogUsd, Log, TEXT( "Skipping material property %d as we have no valid data to use." ), Property );
+						UE_LOG(LogUsd, Log, TEXT("Skipping material property %d as we have no valid data to use."), Property);
 						continue;
 					}
 
-					if ( !UserConstantValue && NumSamples > 0 && SampleSize )
+					if (!UserConstantValue && NumSamples > 0 && SampleSize)
 					{
-						if ( NumSamples != SampleSize->X * SampleSize->Y )
+						if (NumSamples != SampleSize->X * SampleSize->Y)
 						{
-							UE_LOG( LogUsd, Warning, TEXT( "Skipping material property %d as it has an unexpected number of samples (%d instead of %d)." ), Property, NumSamples, SampleSize->X * SampleSize->Y );
+							UE_LOG(
+								LogUsd,
+								Warning,
+								TEXT("Skipping material property %d as it has an unexpected number of samples (%d instead of %d)."),
+								Property,
+								NumSamples,
+								SampleSize->X * SampleSize->Y
+							);
 							continue;
 						}
 					}
 
-					if ( !bPropertyValueIsConstant && ( TextureFilePath == nullptr || !FPaths::FileExists( *TextureFilePath ) ) )
+					if (!bPropertyValueIsConstant && (TextureFilePath == nullptr || !FPaths::FileExists(*TextureFilePath)))
 					{
-						UE_LOG( LogUsd, Warning, TEXT( "Skipping material property %d as target texture '%s' was not found." ), Property, TextureFilePath ? **TextureFilePath : TEXT( "" ) );
+						UE_LOG(
+							LogUsd,
+							Warning,
+							TEXT("Skipping material property %d as target texture '%s' was not found."),
+							Property,
+							TextureFilePath ? **TextureFilePath : TEXT("")
+						);
 						continue;
 					}
 
@@ -1809,182 +1865,194 @@ namespace UE
 					pxr::VtValue ConstantValue;
 					pxr::GfVec4f FallbackValue;
 					pxr::TfToken ColorSpaceToken = UnrealIdentifiers::RawColorSpaceToken;
-					switch ( Property )
+					switch (Property)
 					{
-					case MP_BaseColor:
-						InputToken = UnrealIdentifiers::DiffuseColor;
-						InputType = pxr::SdfValueTypeNames->Color3f;
-						if ( bPropertyValueIsConstant )
-						{
-							if ( UserConstantValue )
-							{
-								ConstantValue = pxr::GfVec3f( *UserConstantValue, *UserConstantValue, *UserConstantValue );
-							}
-							else
+						case MP_BaseColor:
+							InputToken = UnrealIdentifiers::DiffuseColor;
+							InputType = pxr::SdfValueTypeNames->Color3f;
+							if (bPropertyValueIsConstant)
 							{
-								ConstantValue = ConstantLinearValue;
-							}
-						}
-						FallbackValue = pxr::GfVec4f{ 0.0, 0.0, 0.0, 1.0f };
-						ColorSpaceToken = UnrealIdentifiers::SRGBColorSpaceToken;
-						break;
-					case MP_Specular:
-						InputToken = UnrealIdentifiers::Specular;
-						InputType = pxr::SdfValueTypeNames->Float;
-						ConstantValue = UserConstantValue ? *UserConstantValue : ConstantLinearValue[ 0 ];
-						FallbackValue = pxr::GfVec4f{ 0.5, 0.5, 0.5, 1.0f };
-						break;
-					case MP_Metallic:
-						InputToken = UnrealIdentifiers::Metallic;
-						InputType = pxr::SdfValueTypeNames->Float;
-						ConstantValue = UserConstantValue ? *UserConstantValue : ConstantLinearValue[ 0 ];
-						FallbackValue = pxr::GfVec4f{ 0.0, 0.0, 0.0, 1.0f };
-						break;
-					case MP_Roughness:
-						InputToken = UnrealIdentifiers::Roughness;
-						InputType = pxr::SdfValueTypeNames->Float;
-						ConstantValue = UserConstantValue ? *UserConstantValue : ConstantLinearValue[ 0 ];
-						FallbackValue = pxr::GfVec4f{ 0.5, 0.5, 0.5, 1.0f };
-						break;
-					case MP_Normal:
-						InputToken = UnrealIdentifiers::Normal;
-						InputType = pxr::SdfValueTypeNames->Normal3f;
-						if ( bPropertyValueIsConstant )
-						{
-							// This doesn't make much sense here but it's an available option, so here we go
-							if ( UserConstantValue )
-							{
-								ConstantValue = pxr::GfVec3f( *UserConstantValue, *UserConstantValue, *UserConstantValue );
-							}
-							else
-							{
-								ConstantValue = ConstantLinearValue;
-							}
-						}
-						FallbackValue = pxr::GfVec4f{ 0.0, 0.0, 1.0, 1.0f };
-						break;
-					case MP_Tangent:
-						InputToken = UnrealIdentifiers::Tangent;
-						InputType = pxr::SdfValueTypeNames->Normal3f;
-						if ( bPropertyValueIsConstant )
-						{
-							// This doesn't make much sense here but it's an available option, so here we go
-							if ( UserConstantValue )
-							{
-								ConstantValue = pxr::GfVec3f( *UserConstantValue, *UserConstantValue, *UserConstantValue );
-							}
-							else
-							{
-								ConstantValue = ConstantLinearValue;
+								if (UserConstantValue)
+								{
+									ConstantValue = pxr::GfVec3f(*UserConstantValue, *UserConstantValue, *UserConstantValue);
+								}
+								else
+								{
+									ConstantValue = ConstantLinearValue;
+								}
 							}
-						}
-						FallbackValue = pxr::GfVec4f{ 1.0, 0.0, 0.0, 1.0f };
-						break;
-					case MP_EmissiveColor:
-						InputToken = UnrealIdentifiers::EmissiveColor;
-						InputType = pxr::SdfValueTypeNames->Color3f;
-						if ( bPropertyValueIsConstant )
-						{
-							// This doesn't make much sense here but it's an available option, so here we go
-							if ( UserConstantValue )
+							FallbackValue = pxr::GfVec4f{0.0, 0.0, 0.0, 1.0f};
+							ColorSpaceToken = UnrealIdentifiers::SRGBColorSpaceToken;
+							break;
+						case MP_Specular:
+							InputToken = UnrealIdentifiers::Specular;
+							InputType = pxr::SdfValueTypeNames->Float;
+							ConstantValue = UserConstantValue ? *UserConstantValue : ConstantLinearValue[0];
+							FallbackValue = pxr::GfVec4f{0.5, 0.5, 0.5, 1.0f};
+							break;
+						case MP_Metallic:
+							InputToken = UnrealIdentifiers::Metallic;
+							InputType = pxr::SdfValueTypeNames->Float;
+							ConstantValue = UserConstantValue ? *UserConstantValue : ConstantLinearValue[0];
+							FallbackValue = pxr::GfVec4f{0.0, 0.0, 0.0, 1.0f};
+							break;
+						case MP_Roughness:
+							InputToken = UnrealIdentifiers::Roughness;
+							InputType = pxr::SdfValueTypeNames->Float;
+							ConstantValue = UserConstantValue ? *UserConstantValue : ConstantLinearValue[0];
+							FallbackValue = pxr::GfVec4f{0.5, 0.5, 0.5, 1.0f};
+							break;
+						case MP_Normal:
+							InputToken = UnrealIdentifiers::Normal;
+							InputType = pxr::SdfValueTypeNames->Normal3f;
+							if (bPropertyValueIsConstant)
 							{
-								ConstantValue = pxr::GfVec3f( *UserConstantValue, *UserConstantValue, *UserConstantValue );
+								// This doesn't make much sense here but it's an available option, so here we go
+								if (UserConstantValue)
+								{
+									ConstantValue = pxr::GfVec3f(*UserConstantValue, *UserConstantValue, *UserConstantValue);
+								}
+								else
+								{
+									ConstantValue = ConstantLinearValue;
+								}
 							}
-							else
+							FallbackValue = pxr::GfVec4f{0.0, 0.0, 1.0, 1.0f};
+							break;
+						case MP_Tangent:
+							InputToken = UnrealIdentifiers::Tangent;
+							InputType = pxr::SdfValueTypeNames->Normal3f;
+							if (bPropertyValueIsConstant)
 							{
-								ConstantValue = ConstantLinearValue;
+								// This doesn't make much sense here but it's an available option, so here we go
+								if (UserConstantValue)
+								{
+									ConstantValue = pxr::GfVec3f(*UserConstantValue, *UserConstantValue, *UserConstantValue);
+								}
+								else
+								{
+									ConstantValue = ConstantLinearValue;
+								}
 							}
-						}
-						FallbackValue = pxr::GfVec4f{ 0.0, 0.0, 0.0, 1.0f };
-						// Emissive is also written out with RawColorSpaceToken as we write them as EXR files now
-						break;
-					case MP_Opacity: // It's OK that we use the same for both as these are mutually exclusive blend modes
-					case MP_OpacityMask:
-						InputToken = UnrealIdentifiers::Opacity;
-						InputType = pxr::SdfValueTypeNames->Float;
-						ConstantValue = UserConstantValue ? *UserConstantValue : ConstantLinearValue[ 0 ];
-						FallbackValue = pxr::GfVec4f{ 1.0, 1.0, 1.0, 1.0f };
-						break;
-					case MP_Anisotropy:
-						InputToken = UnrealIdentifiers::Anisotropy;
-						InputType = pxr::SdfValueTypeNames->Float;
-						ConstantValue = UserConstantValue ? *UserConstantValue : ConstantLinearValue[ 0 ];
-						FallbackValue = pxr::GfVec4f{ 0.0, 0.0, 0.0, 1.0f };
-						break;
-					case MP_AmbientOcclusion:
-						InputToken = UnrealIdentifiers::Occlusion;
-						InputType = pxr::SdfValueTypeNames->Float;
-						ConstantValue = UserConstantValue ? *UserConstantValue : ConstantLinearValue[ 0 ];
-						FallbackValue = pxr::GfVec4f{ 1.0, 1.0, 1.0, 1.0f };
-						break;
-					case MP_SubsurfaceColor:
-						InputToken = UnrealIdentifiers::SubsurfaceColor;
-						InputType = pxr::SdfValueTypeNames->Color3f;
-						if ( bPropertyValueIsConstant )
-						{
-							if ( UserConstantValue )
+							FallbackValue = pxr::GfVec4f{1.0, 0.0, 0.0, 1.0f};
+							break;
+						case MP_EmissiveColor:
+							InputToken = UnrealIdentifiers::EmissiveColor;
+							InputType = pxr::SdfValueTypeNames->Color3f;
+							if (bPropertyValueIsConstant)
 							{
-								ConstantValue = pxr::GfVec3f( *UserConstantValue, *UserConstantValue, *UserConstantValue );
+								// This doesn't make much sense here but it's an available option, so here we go
+								if (UserConstantValue)
+								{
+									ConstantValue = pxr::GfVec3f(*UserConstantValue, *UserConstantValue, *UserConstantValue);
+								}
+								else
+								{
+									ConstantValue = ConstantLinearValue;
+								}
 							}
-							else
+							FallbackValue = pxr::GfVec4f{0.0, 0.0, 0.0, 1.0f};
+							// Emissive is also written out with RawColorSpaceToken as we write them as EXR files now
+							break;
+						case MP_Opacity:	// It's OK that we use the same for both as these are mutually exclusive blend modes
+						case MP_OpacityMask:
+							InputToken = UnrealIdentifiers::Opacity;
+							InputType = pxr::SdfValueTypeNames->Float;
+							ConstantValue = UserConstantValue ? *UserConstantValue : ConstantLinearValue[0];
+							FallbackValue = pxr::GfVec4f{1.0, 1.0, 1.0, 1.0f};
+							break;
+						case MP_Anisotropy:
+							InputToken = UnrealIdentifiers::Anisotropy;
+							InputType = pxr::SdfValueTypeNames->Float;
+							ConstantValue = UserConstantValue ? *UserConstantValue : ConstantLinearValue[0];
+							FallbackValue = pxr::GfVec4f{0.0, 0.0, 0.0, 1.0f};
+							break;
+						case MP_AmbientOcclusion:
+							InputToken = UnrealIdentifiers::Occlusion;
+							InputType = pxr::SdfValueTypeNames->Float;
+							ConstantValue = UserConstantValue ? *UserConstantValue : ConstantLinearValue[0];
+							FallbackValue = pxr::GfVec4f{1.0, 1.0, 1.0, 1.0f};
+							break;
+						case MP_SubsurfaceColor:
+							InputToken = UnrealIdentifiers::SubsurfaceColor;
+							InputType = pxr::SdfValueTypeNames->Color3f;
+							if (bPropertyValueIsConstant)
 							{
-								ConstantValue = ConstantLinearValue;
+								if (UserConstantValue)
+								{
+									ConstantValue = pxr::GfVec3f(*UserConstantValue, *UserConstantValue, *UserConstantValue);
+								}
+								else
+								{
+									ConstantValue = ConstantLinearValue;
+								}
 							}
-						}
-						FallbackValue = pxr::GfVec4f{ 1.0, 1.0, 1.0, 1.0f };
-						ColorSpaceToken = UnrealIdentifiers::SRGBColorSpaceToken;
-						break;
-					default:
-						continue;
-						break;
+							FallbackValue = pxr::GfVec4f{1.0, 1.0, 1.0, 1.0f};
+							ColorSpaceToken = UnrealIdentifiers::SRGBColorSpaceToken;
+							break;
+						default:
+							continue;
+							break;
 					}
 
-					pxr::UsdShadeInput ShadeInput = ShadeShader.CreateInput( InputToken, InputType );
-					if ( bPropertyValueIsConstant )
+					pxr::UsdShadeInput ShadeInput = ShadeShader.CreateInput(InputToken, InputType);
+					if (bPropertyValueIsConstant)
 					{
-						ShadeInput.Set( ConstantValue );
+						ShadeInput.Set(ConstantValue);
 					}
-					else // Its a texture
+					else	// Its a texture
 					{
 						// Create the primvar/uv set reader on-demand. We'll be using the same UV set for everything for now
-						if ( !PrimvarReaderShader )
+						if (!PrimvarReaderShader)
 						{
-							PrimvarReaderShader = pxr::UsdShadeShader::Define( Stage, MaterialPath.AppendChild( PrimvarReaderShaderName ) );
-							PrimvarReaderShader.SetShaderId( UnrealIdentifiers::UsdPrimvarReader_float2 );
+							PrimvarReaderShader = pxr::UsdShadeShader::Define(Stage, MaterialPath.AppendChild(PrimvarReaderShaderName));
+							PrimvarReaderShader.SetShaderId(UnrealIdentifiers::UsdPrimvarReader_float2);
 
 							// Create the 'st' input directly on the material, as that seems to be preferred
-							pxr::UsdShadeInput MaterialStInput = OutUsdShadeMaterial.CreateInput( PrimvarVariableName, pxr::SdfValueTypeNames->Token );
-							MaterialStInput.Set( UnrealIdentifiers::St );
+							pxr::UsdShadeInput MaterialStInput = OutUsdShadeMaterial.CreateInput(PrimvarVariableName, pxr::SdfValueTypeNames->Token);
+							MaterialStInput.Set(UnrealIdentifiers::St);
 
-							pxr::UsdShadeInput VarnameInput = PrimvarReaderShader.CreateInput( UnrealIdentifiers::Varname, pxr::SdfValueTypeNames->String );
-							VarnameInput.ConnectToSource( MaterialStInput );
+							pxr::UsdShadeInput VarnameInput = PrimvarReaderShader.CreateInput(
+								UnrealIdentifiers::Varname,
+								pxr::SdfValueTypeNames->String
+							);
+							VarnameInput.ConnectToSource(MaterialStInput);
 
-							PrimvarReaderShader.CreateOutput( UnrealIdentifiers::Result, pxr::SdfValueTypeNames->Token );
+							PrimvarReaderShader.CreateOutput(UnrealIdentifiers::Result, pxr::SdfValueTypeNames->Token);
 						}
 
-						FString TextureReaderName = UsdToUnreal::ConvertToken( InputToken );
-						TextureReaderName.RemoveFromEnd( TEXT( "Color" ) );
-						TextureReaderName += TEXT( "Texture" );
+						FString TextureReaderName = UsdToUnreal::ConvertToken(InputToken);
+						TextureReaderName.RemoveFromEnd(TEXT("Color"));
+						TextureReaderName += TEXT("Texture");
 
-						pxr::UsdShadeShader UsdUVTextureShader = pxr::UsdShadeShader::Define( Stage, MaterialPath.AppendChild( UnrealToUsd::ConvertToken( *TextureReaderName ).Get() ) );
-						UsdUVTextureShader.SetShaderId( UnrealIdentifiers::UsdUVTexture );
+						pxr::UsdShadeShader UsdUVTextureShader = pxr::UsdShadeShader::Define(
+							Stage,
+							MaterialPath.AppendChild(UnrealToUsd::ConvertToken(*TextureReaderName).Get())
+						);
+						UsdUVTextureShader.SetShaderId(UnrealIdentifiers::UsdUVTexture);
 
-						pxr::UsdShadeInput TextureFileInput = UsdUVTextureShader.CreateInput( UnrealIdentifiers::File, pxr::SdfValueTypeNames->Asset );
+						pxr::UsdShadeInput TextureFileInput = UsdUVTextureShader.CreateInput(UnrealIdentifiers::File, pxr::SdfValueTypeNames->Asset);
 						FString TextureRelativePath = *TextureFilePath;
-						if ( !UsdFilePath.IsEmpty() )
+						if (!UsdFilePath.IsEmpty())
 						{
-							FPaths::MakePathRelativeTo( TextureRelativePath, *UsdFilePath );
+							FPaths::MakePathRelativeTo(TextureRelativePath, *UsdFilePath);
 						}
-						TextureFileInput.Set( pxr::SdfAssetPath( UnrealToUsd::ConvertString( *TextureRelativePath ).Get() ) );
+						TextureFileInput.Set(pxr::SdfAssetPath(UnrealToUsd::ConvertString(*TextureRelativePath).Get()));
 
-						pxr::UsdShadeInput TextureStInput = UsdUVTextureShader.CreateInput( UnrealIdentifiers::St, pxr::SdfValueTypeNames->Float2 );
-						TextureStInput.ConnectToSource( PrimvarReaderShader.GetOutput( UnrealIdentifiers::Result ) );
+						pxr::UsdShadeInput TextureStInput = UsdUVTextureShader.CreateInput(UnrealIdentifiers::St, pxr::SdfValueTypeNames->Float2);
+						TextureStInput.ConnectToSource(PrimvarReaderShader.GetOutput(UnrealIdentifiers::Result));
 
-						pxr::UsdShadeInput TextureColorSpaceInput = UsdUVTextureShader.CreateInput( UnrealIdentifiers::SourceColorSpaceToken, pxr::SdfValueTypeNames->Token );
-						TextureColorSpaceInput.Set( ColorSpaceToken );
+						pxr::UsdShadeInput TextureColorSpaceInput = UsdUVTextureShader.CreateInput(
+							UnrealIdentifiers::SourceColorSpaceToken,
+							pxr::SdfValueTypeNames->Token
+						);
+						TextureColorSpaceInput.Set(ColorSpaceToken);
 
-						pxr::UsdShadeInput TextureFallbackInput = UsdUVTextureShader.CreateInput( UnrealIdentifiers::Fallback, pxr::SdfValueTypeNames->Float4 );
-						TextureFallbackInput.Set( FallbackValue );
+						pxr::UsdShadeInput TextureFallbackInput = UsdUVTextureShader.CreateInput(
+							UnrealIdentifiers::Fallback,
+							pxr::SdfValueTypeNames->Float4
+						);
+						TextureFallbackInput.Set(FallbackValue);
 
 						// In the general case it's impossible to set a "correct" wrapping value here because the
 						// material we just baked may be using 3 different textures with UV transforms an all different
@@ -1995,92 +2063,92 @@ namespace UE
 							UnrealIdentifiers::WrapS,
 							pxr::SdfValueTypeNames->Token
 						);
-						TextureFileWrapSInput.Set( UnrealIdentifiers::Repeat );
+						TextureFileWrapSInput.Set(UnrealIdentifiers::Repeat);
 						pxr::UsdShadeInput TextureFileWrapTInput = UsdUVTextureShader.CreateInput(
 							UnrealIdentifiers::WrapT,
 							pxr::SdfValueTypeNames->Token
 						);
-						TextureFileWrapTInput.Set( UnrealIdentifiers::Repeat );
+						TextureFileWrapTInput.Set(UnrealIdentifiers::Repeat);
 
 						pxr::UsdShadeOutput TextureOutput = UsdUVTextureShader.CreateOutput(
 							InputType == pxr::SdfValueTypeNames->Float ? UnrealIdentifiers::R : UnrealIdentifiers::RGB,
 							InputType
 						);
 
-						ShadeInput.ConnectToSource( TextureOutput );
+						ShadeInput.ConnectToSource(TextureOutput);
 					}
 				}
 
 				return true;
 			}
-#endif // WITH_EDITOR
+#endif	  // WITH_EDITOR
 
-			void HashShadeInput( const pxr::UsdShadeInput& ShadeInput, FSHA1& InOutHashState )
+			void HashShadeInput(const pxr::UsdShadeInput& ShadeInput, FSHA1& InOutHashState)
 			{
-				if ( !ShadeInput )
+				if (!ShadeInput)
 				{
 					return;
 				}
 
 				FScopedUsdAllocs UsdAllocs;
 
-				FString InputName = UsdToUnreal::ConvertToken( ShadeInput.GetBaseName() );
-				InOutHashState.UpdateWithString( *InputName, InputName.Len() );
+				FString InputName = UsdToUnreal::ConvertToken(ShadeInput.GetBaseName());
+				InOutHashState.UpdateWithString(*InputName, InputName.Len());
 
-				FString InputTypeName = UsdToUnreal::ConvertToken( ShadeInput.GetTypeName().GetAsToken() );
-				InOutHashState.UpdateWithString( *InputTypeName, InputTypeName.Len() );
+				FString InputTypeName = UsdToUnreal::ConvertToken(ShadeInput.GetTypeName().GetAsToken());
+				InOutHashState.UpdateWithString(*InputTypeName, InputTypeName.Len());
 
 				// Connected to something else, recurse
 				pxr::UsdShadeConnectableAPI Source;
 				pxr::TfToken SourceName;
 				pxr::UsdShadeAttributeType AttributeType;
-				if ( pxr::UsdShadeConnectableAPI::GetConnectedSource( ShadeInput.GetAttr(), &Source, &SourceName, &AttributeType ) )
+				if (pxr::UsdShadeConnectableAPI::GetConnectedSource(ShadeInput.GetAttr(), &Source, &SourceName, &AttributeType))
 				{
-					FString SourceOutputName = UsdToUnreal::ConvertToken( SourceName );
-					InOutHashState.UpdateWithString( *SourceOutputName, SourceOutputName.Len() );
+					FString SourceOutputName = UsdToUnreal::ConvertToken(SourceName);
+					InOutHashState.UpdateWithString(*SourceOutputName, SourceOutputName.Len());
 
 					// Skip in case our input is connected to an output on the same prim, or else we'll recurse forever
-					if ( Source.GetPrim() == ShadeInput.GetPrim() )
+					if (Source.GetPrim() == ShadeInput.GetPrim())
 					{
 						return;
 					}
 
-					for ( const pxr::UsdShadeInput& ChildInput : Source.GetInputs() )
+					for (const pxr::UsdShadeInput& ChildInput : Source.GetInputs())
 					{
-						HashShadeInput( ChildInput, InOutHashState );
+						HashShadeInput(ChildInput, InOutHashState);
 					}
 				}
 				// Not connected to anything, has a value (this could be a texture file path too)
 				else
 				{
 					pxr::VtValue ShadeInputValue;
-					ShadeInput.Get( &ShadeInputValue );
+					ShadeInput.Get(&ShadeInputValue);
 
-					// We have to manually resolve and hash these file paths or else resolvedpaths inside usdz archives will have upper case driver letters
-					// when we first open the stage, but will switch to lower case drive letters if we reload them. This is not something we're doing,
-					// as it happens with a pure USD python script. (this with USD 21.05 in June 2021)
-					if ( ShadeInputValue.IsHolding<pxr::SdfAssetPath>() )
+					// We have to manually resolve and hash these file paths or else resolvedpaths inside usdz archives will have upper case driver
+					// letters when we first open the stage, but will switch to lower case drive letters if we reload them. This is not something
+					// we're doing, as it happens with a pure USD python script. (this with USD 21.05 in June 2021)
+					if (ShadeInputValue.IsHolding<pxr::SdfAssetPath>())
 					{
-						FString ResolvedPath = UsdUtils::GetResolvedTexturePath( ShadeInput.GetAttr() );
-						InOutHashState.UpdateWithString( *ResolvedPath, ResolvedPath.Len() );
+						FString ResolvedPath = UsdUtils::GetResolvedTexturePath(ShadeInput.GetAttr());
+						InOutHashState.UpdateWithString(*ResolvedPath, ResolvedPath.Len());
 					}
 					else if (ShadeInputValue.IsHolding<pxr::TfToken>())
 					{
 						// Stringify instead of using GetHash because if ShadeInputValue contains a pxr::TfToken then
 						// it will actually just contain some non-deterministic integer IDs
 						FString Stringified = UsdToUnreal::ConvertString(pxr::TfStringify(ShadeInputValue));
-						InOutHashState.UpdateWithString( *Stringified, Stringified.Len() );
+						InOutHashState.UpdateWithString(*Stringified, Stringified.Len());
 					}
 					else
 					{
-						uint64 ValueHash = ( uint64 ) ShadeInputValue.GetHash();
-						InOutHashState.Update( reinterpret_cast< uint8* >( &ValueHash ), sizeof( uint64 ) );
+						uint64 ValueHash = (uint64)ShadeInputValue.GetHash();
+						InOutHashState.Update(reinterpret_cast<uint8*>(&ValueHash), sizeof(uint64));
 					}
 				}
 			}
-		}
-	}
-}
+		}	 // namespace Private
+	}		 // namespace UsdShadeConversion
+}	 // namespace UE
 namespace UsdShadeConversionImpl = UE::UsdShadeConversion::Private;
 
 bool UsdToUnreal::ConvertMaterial(
@@ -2242,7 +2310,8 @@ bool UsdToUnreal::ConvertMaterial(
 	TSet<FString> UsedPrimvars;
 	for (const TPair<FString, UsdShadeConversionImpl::FParameterValue>& Parameter : MaterialParameters)
 	{
-		if (const UsdShadeConversionImpl::FTextureParameterValue* TextureParameterValue = Parameter.Value.TryGet<UsdShadeConversionImpl::FTextureParameterValue>())
+		if (const UsdShadeConversionImpl::FTextureParameterValue*
+				TextureParameterValue = Parameter.Value.TryGet<UsdShadeConversionImpl::FTextureParameterValue>())
 		{
 			UsedPrimvars.Add(TextureParameterValue->Primvar);
 			ParameterToPrimvar.Add(Parameter.Key, TextureParameterValue->Primvar);
@@ -2250,7 +2319,7 @@ bool UsdToUnreal::ConvertMaterial(
 	}
 	UsedPrimvars.Remove(TEXT(""));
 	TArray<FString> SortedPrimvars = UsedPrimvars.Array();
-	SortedPrimvars.Sort();  // Try for some deterministic ordering (st0 should come before st1, etc.)
+	SortedPrimvars.Sort();	  // Try for some deterministic ordering (st0 should come before st1, etc.)
 	TMap<FString, int32> PrimvarToUVIndex;
 	PrimvarToUVIndex.Reserve(SortedPrimvars.Num());
 	int32 UVIndex = 0;
@@ -2263,13 +2332,7 @@ bool UsdToUnreal::ConvertMaterial(
 	const bool bForUsdPreviewSurface = true;
 	for (const TPair<FString, UsdShadeConversionImpl::FParameterValue>& Parameter : MaterialParameters)
 	{
-		UsdShadeConversionImpl::SetParameterValue(
-			Material,
-			*Parameter.Key,
-			Parameter.Value,
-			bForUsdPreviewSurface,
-			PrimvarToUVIndex
-		);
+		UsdShadeConversionImpl::SetParameterValue(Material, *Parameter.Key, Parameter.Value, bForUsdPreviewSurface, PrimvarToUVIndex);
 	}
 
 	// Handle world space normals
@@ -2443,7 +2506,15 @@ bool UsdToUnreal::ConvertMaterial(
 	}
 
 	// Ambient occlusion
-	if (UsdShadeConversionImpl::GetFloatParameterValue(Connectable, UnrealIdentifiers::Occlusion, 1.f, TempParameterValue, &Material, TexturesCache, bReuseIdenticalAssets))
+	if (UsdShadeConversionImpl::GetFloatParameterValue(
+			Connectable,
+			UnrealIdentifiers::Occlusion,
+			1.f,
+			TempParameterValue,
+			&Material,
+			TexturesCache,
+			bReuseIdenticalAssets
+		))
 	{
 		AmbientOcclusionParameter = TempParameterValue;
 	}
@@ -2454,17 +2525,18 @@ bool UsdToUnreal::ConvertMaterial(
 	TSet<FString> UsedPrimvars;
 	TFunction<void(const FString& Parameter, const TOptional<UsdShadeConversionImpl::FParameterValue>&)> FetchPrimvar =
 		[&UsedPrimvars, &ParameterToPrimvar](const FString& Parameter, const TOptional<UsdShadeConversionImpl::FParameterValue>& ParameterValue)
+	{
+		if (!ParameterValue.IsSet())
 		{
-			if (!ParameterValue.IsSet())
-			{
-				return;
-			}
-			if (const UsdShadeConversionImpl::FTextureParameterValue* TextureParameterValue = ParameterValue.GetValue().TryGet<UsdShadeConversionImpl::FTextureParameterValue>())
-			{
-				UsedPrimvars.Add(TextureParameterValue->Primvar);
-				ParameterToPrimvar.Add(Parameter, TextureParameterValue->Primvar);
-			}
-		};
+			return;
+		}
+		if (const UsdShadeConversionImpl::FTextureParameterValue*
+				TextureParameterValue = ParameterValue.GetValue().TryGet<UsdShadeConversionImpl::FTextureParameterValue>())
+		{
+			UsedPrimvars.Add(TextureParameterValue->Primvar);
+			ParameterToPrimvar.Add(Parameter, TextureParameterValue->Primvar);
+		}
+	};
 	FetchPrimvar(TEXT("BaseColor"), BaseColorParameter);
 	FetchPrimvar(TEXT("EmissiveColor"), EmissiveParameter);
 	FetchPrimvar(TEXT("Metallic"), MetallicParameter);
@@ -2486,94 +2558,76 @@ bool UsdToUnreal::ConvertMaterial(
 
 	// auto so we can use the MaterialInput generic parameter for FColorMaterialInput, FScalarMaterialInput and FVectorMaterialInput
 	auto ConnectMaterialInput =
-		[&PrimvarToUVIndex]
-		(UMaterial& Material, auto& MaterialInput, const UsdShadeConversionImpl::FParameterValue& InParameterValue) -> bool
+		[&PrimvarToUVIndex](UMaterial& Material, auto& MaterialInput, const UsdShadeConversionImpl::FParameterValue& InParameterValue) -> bool
+	{
+		UMaterialExpression* Expression = UsdShadeConversionImpl::GetExpressionForValue(Material, InParameterValue);
+		if (!Expression)
 		{
-			UMaterialExpression* Expression = UsdShadeConversionImpl::GetExpressionForValue(Material, InParameterValue);
-			if (!Expression)
-			{
-				return false;
-			}
+			return false;
+		}
 
-			MaterialInput.Expression = Expression;
+		MaterialInput.Expression = Expression;
 
-			if (const UsdShadeConversionImpl::FTextureParameterValue* TextureParameterValue = InParameterValue.TryGet<UsdShadeConversionImpl::FTextureParameterValue>())
-			{
-				MaterialInput.OutputIndex = TextureParameterValue->OutputIndex;
+		if (const UsdShadeConversionImpl::FTextureParameterValue*
+				TextureParameterValue = InParameterValue.TryGet<UsdShadeConversionImpl::FTextureParameterValue>())
+		{
+			MaterialInput.OutputIndex = TextureParameterValue->OutputIndex;
 
-				if (UMaterialExpressionTextureSample* TextureExpression = Cast<UMaterialExpressionTextureSample>(Expression))
+			if (UMaterialExpressionTextureSample* TextureExpression = Cast<UMaterialExpressionTextureSample>(Expression))
+			{
+				if (int32* FoundCoordinate = PrimvarToUVIndex.Find(TextureParameterValue->Primvar))
 				{
-					if (int32* FoundCoordinate = PrimvarToUVIndex.Find(TextureParameterValue->Primvar))
-					{
-						TextureExpression->ConstCoordinate = *FoundCoordinate;
-					}
-					else
-					{
-						UE_LOG(
-							LogTemp,
-							Warning,
-							TEXT("Failed to find primvar '%s' when setting material parameter. Available primvars and UV indices: %s.%s"),
-							*TextureParameterValue->Primvar,
-							*UsdUtils::StringifyMap(PrimvarToUVIndex),
-							TextureParameterValue->Primvar.IsEmpty() ? TEXT(" Is your UsdUVTexture Shader missing the 'inputs:st' attribute? (It specifies which UV set to sample the texture with)") : TEXT("")
-						);
-					}
+					TextureExpression->ConstCoordinate = *FoundCoordinate;
+				}
+				else
+				{
+					UE_LOG(
+						LogTemp,
+						Warning,
+						TEXT("Failed to find primvar '%s' when setting material parameter. Available primvars and UV indices: %s.%s"),
+						*TextureParameterValue->Primvar,
+						*UsdUtils::StringifyMap(PrimvarToUVIndex),
+						TextureParameterValue->Primvar.IsEmpty() ? TEXT(" Is your UsdUVTexture Shader missing the 'inputs:st' attribute? (It "
+																		"specifies which UV set to sample the texture with)")
+																 : TEXT("")
+					);
 				}
 			}
+		}
 
-			return true;
-		};
+		return true;
+	};
 
 	UMaterialEditorOnlyData* EditorOnly = Material.GetEditorOnlyData();
 
 	// Base color
 	if (UsdShadeConversionImpl::FParameterValue* BaseColorParameterValue = BaseColorParameter.GetPtrOrNull())
 	{
-		ConnectMaterialInput(
-			Material,
-			EditorOnly->BaseColor,
-			*BaseColorParameterValue
-		);
+		ConnectMaterialInput(Material, EditorOnly->BaseColor, *BaseColorParameterValue);
 	}
 
 	// Emissive color
 	if (UsdShadeConversionImpl::FParameterValue* EmissiveParameterValue = EmissiveParameter.GetPtrOrNull())
 	{
-		ConnectMaterialInput(
-			Material,
-			EditorOnly->EmissiveColor,
-			*EmissiveParameterValue
-		);
+		ConnectMaterialInput(Material, EditorOnly->EmissiveColor, *EmissiveParameterValue);
 	}
 
 	// Metallic
 	if (UsdShadeConversionImpl::FParameterValue* MetallicParameterValue = MetallicParameter.GetPtrOrNull())
 	{
-		ConnectMaterialInput(
-			Material,
-			EditorOnly->Metallic,
-			*MetallicParameterValue
-		);
+		ConnectMaterialInput(Material, EditorOnly->Metallic, *MetallicParameterValue);
 	}
 
 	// Roughness
 	if (UsdShadeConversionImpl::FParameterValue* RoughnessParameterValue = RoughnessParameter.GetPtrOrNull())
 	{
-		ConnectMaterialInput(
-			Material,
-			EditorOnly->Roughness,
-			*RoughnessParameterValue
-		);
+		ConnectMaterialInput(Material, EditorOnly->Roughness, *RoughnessParameterValue);
 	}
 
 	// Opacity
 	if (UsdShadeConversionImpl::FParameterValue* OpacityParameterValue = OpacityParameter.GetPtrOrNull())
 	{
-		const bool bConnected = ConnectMaterialInput(
-			Material,
-			EditorOnly->Opacity,
-			*OpacityParameterValue
-		);
+		const bool bConnected = ConnectMaterialInput(Material, EditorOnly->Opacity, *OpacityParameterValue);
 
 		if (bConnected)
 		{
@@ -2584,31 +2638,19 @@ bool UsdToUnreal::ConvertMaterial(
 	// Normal
 	if (UsdShadeConversionImpl::FParameterValue* NormalParameterValue = NormalParameter.GetPtrOrNull())
 	{
-		ConnectMaterialInput(
-			Material,
-			EditorOnly->Normal,
-			*NormalParameterValue
-		);
+		ConnectMaterialInput(Material, EditorOnly->Normal, *NormalParameterValue);
 	}
 
 	// Refraction
 	if (UsdShadeConversionImpl::FParameterValue* RefractionParameterValue = RefractionParameter.GetPtrOrNull())
 	{
-		ConnectMaterialInput(
-			Material,
-			EditorOnly->Refraction,
-			*RefractionParameterValue
-		);
+		ConnectMaterialInput(Material, EditorOnly->Refraction, *RefractionParameterValue);
 	}
 
 	// Ambient occlusion
 	if (UsdShadeConversionImpl::FParameterValue* OcclusionParameterValue = AmbientOcclusionParameter.GetPtrOrNull())
 	{
-		ConnectMaterialInput(
-			Material,
-			EditorOnly->AmbientOcclusion,
-			*OcclusionParameterValue
-		);
+		ConnectMaterialInput(Material, EditorOnly->AmbientOcclusion, *OcclusionParameterValue);
 	}
 
 	// Handle world space normals
@@ -2634,7 +2676,7 @@ bool UsdToUnreal::ConvertMaterial(
 	return true;
 #else
 	return false;
-#endif // WITH_EDITOR
+#endif	  // WITH_EDITOR
 }
 
 // Deprecated
@@ -2674,37 +2716,37 @@ bool UsdToUnreal::ConvertShadeInputsToParameters(
 	bool bHasMaterialInfo = false;
 
 	pxr::TfToken RenderContextToken = pxr::UsdShadeTokens->universalRenderContext;
-	if ( RenderContext )
+	if (RenderContext)
 	{
-		RenderContextToken = UnrealToUsd::ConvertToken( RenderContext ).Get();
+		RenderContextToken = UnrealToUsd::ConvertToken(RenderContext).Get();
 	}
 
-	pxr::UsdShadeShader SurfaceShader = UsdShadeMaterial.ComputeSurfaceSource( RenderContextToken );
-	if ( !SurfaceShader )
+	pxr::UsdShadeShader SurfaceShader = UsdShadeMaterial.ComputeSurfaceSource(RenderContextToken);
+	if (!SurfaceShader)
 	{
 		return false;
 	}
 
-	pxr::UsdShadeConnectableAPI Connectable{ SurfaceShader };
+	pxr::UsdShadeConnectableAPI Connectable{SurfaceShader};
 
-	for ( pxr::UsdShadeInput& ShadeInput : SurfaceShader.GetInputs() )
+	for (pxr::UsdShadeInput& ShadeInput : SurfaceShader.GetInputs())
 	{
-		FString InputName = UsdToUnreal::ConvertToken( ShadeInput.GetBaseName() );
+		FString InputName = UsdToUnreal::ConvertToken(ShadeInput.GetBaseName());
 
 		pxr::UsdShadeInput ConnectInput = ShadeInput;
-		if ( ShadeInput.HasConnectedSource() )
+		if (ShadeInput.HasConnectedSource())
 		{
 			pxr::UsdShadeConnectableAPI Source;
 			pxr::TfToken SourceName;
 			pxr::UsdShadeAttributeType SourceType;
-			pxr::UsdShadeConnectableAPI::GetConnectedSource( ShadeInput.GetAttr(), &Source, &SourceName, &SourceType );
+			pxr::UsdShadeConnectableAPI::GetConnectedSource(ShadeInput.GetAttr(), &Source, &SourceName, &SourceType);
 
 			ConnectInput = Source.GetInput(SourceName);
 		}
 
-		FString DisplayName = UsdToUnreal::ConvertString( ConnectInput.GetAttr().GetDisplayName() );
+		FString DisplayName = UsdToUnreal::ConvertString(ConnectInput.GetAttr().GetDisplayName());
 
-		if ( DisplayName.IsEmpty() )
+		if (DisplayName.IsEmpty())
 		{
 			DisplayName = InputName;
 		}
@@ -2712,14 +2754,21 @@ bool UsdToUnreal::ConvertShadeInputsToParameters(
 		UsdShadeConversionImpl::FParameterValue ParameterValue;
 
 		const bool bForUsdPreviewSurface = false;
-		const bool bIsNormalMap = DisplayName.Contains( TEXT("normal") );
+		const bool bIsNormalMap = DisplayName.Contains(TEXT("normal"));
 
 		// For now it seems we don't set "primvar parameters" anyway, so don't bother building this up
 		TMap<FString, int32> Unused;
 
-		if ( ShadeInput.GetTypeName() == pxr::SdfValueTypeNames->Bool )
+		if (ShadeInput.GetTypeName() == pxr::SdfValueTypeNames->Bool)
 		{
-			if ( UsdShadeConversionImpl::GetBoolParameterValue( Connectable, ShadeInput.GetBaseName(), false, ParameterValue, &MaterialInstance, TexturesCache) )
+			if (UsdShadeConversionImpl::GetBoolParameterValue(
+					Connectable,
+					ShadeInput.GetBaseName(),
+					false,
+					ParameterValue,
+					&MaterialInstance,
+					TexturesCache
+				))
 			{
 				UsdShadeConversionImpl::SetParameterValue(MaterialInstance, *DisplayName, ParameterValue, bForUsdPreviewSurface, Unused);
 				bHasMaterialInfo = true;
@@ -2763,22 +2812,40 @@ bool UsdToUnreal::ConvertShadeInputsToParameters(
 }
 
 PRAGMA_DISABLE_DEPRECATION_WARNINGS
-bool UsdToUnreal::ConvertMaterial(const pxr::UsdShadeMaterial& UsdShadeMaterial, UMaterialInstance& Material, UUsdAssetCache* TexturesCache, TMap<FString, int32>& PrimvarToUVIndex, const TCHAR* RenderContext)
+bool UsdToUnreal::ConvertMaterial(
+	const pxr::UsdShadeMaterial& UsdShadeMaterial,
+	UMaterialInstance& Material,
+	UUsdAssetCache* TexturesCache,
+	TMap<FString, int32>& PrimvarToUVIndex,
+	const TCHAR* RenderContext
+)
 {
 	return UsdToUnreal::ConvertMaterial(UsdShadeMaterial, Material);
 }
 
-bool UsdToUnreal::ConvertMaterial(const pxr::UsdShadeMaterial& UsdShadeMaterial, UMaterial& Material, UUsdAssetCache* TexturesCache, TMap<FString, int32>& PrimvarToUVIndex, const TCHAR* RenderContext)
+bool UsdToUnreal::ConvertMaterial(
+	const pxr::UsdShadeMaterial& UsdShadeMaterial,
+	UMaterial& Material,
+	UUsdAssetCache* TexturesCache,
+	TMap<FString, int32>& PrimvarToUVIndex,
+	const TCHAR* RenderContext
+)
 {
 	UUsdAssetCache2* NewCache = nullptr;
 	return UsdToUnreal::ConvertMaterial(UsdShadeMaterial, Material, NewCache, RenderContext);
 }
 
-bool UsdToUnreal::ConvertShadeInputsToParameters(const pxr::UsdShadeMaterial& UsdShadeMaterial, UMaterialInstance& MaterialInstance, UUsdAssetCache* TexturesCache, const TCHAR* RenderContext)
+bool UsdToUnreal::ConvertShadeInputsToParameters(
+	const pxr::UsdShadeMaterial& UsdShadeMaterial,
+	UMaterialInstance& MaterialInstance,
+	UUsdAssetCache* TexturesCache,
+	const TCHAR* RenderContext
+)
 {
 	UUsdAssetCache2* NewCache = nullptr;
 	return UsdToUnreal::ConvertShadeInputsToParameters(UsdShadeMaterial, MaterialInstance, NewCache, RenderContext);
 }
+
 PRAGMA_ENABLE_DEPRECATION_WARNINGS
 
 #if WITH_EDITOR
@@ -2791,93 +2858,99 @@ bool UnrealToUsd::ConvertMaterialToBakedSurface(
 	bool bInDecayTexturesToSinglePixel
 )
 {
-	pxr::UsdShadeMaterial OutUsdShadeMaterial{ OutUsdShadeMaterialPrim };
-	if ( !OutUsdShadeMaterial )
+	pxr::UsdShadeMaterial OutUsdShadeMaterial{OutUsdShadeMaterialPrim};
+	if (!OutUsdShadeMaterial)
 	{
 		return false;
 	}
 
 	FBakeOutput BakedData;
-	if ( !UsdShadeConversionImpl::BakeMaterial( InMaterial, InMaterialProperties, InDefaultTextureSize, BakedData, bInDecayTexturesToSinglePixel ) )
+	if (!UsdShadeConversionImpl::BakeMaterial(InMaterial, InMaterialProperties, InDefaultTextureSize, BakedData, bInDecayTexturesToSinglePixel))
 	{
 		return false;
 	}
 
-	UsdShadeConversionImpl::FBakedMaterialView View{ BakedData };
-	TMap<EMaterialProperty, FString> WrittenTextures = UsdShadeConversionImpl::WriteTextures( View, InMaterial.GetPathName(), InTexturesDir );
+	UsdShadeConversionImpl::FBakedMaterialView View{BakedData};
+	TMap<EMaterialProperty, FString> WrittenTextures = UsdShadeConversionImpl::WriteTextures(View, InMaterial.GetPathName(), InTexturesDir);
 
 	// Manually add user supplied constant values. Can't place these in InMaterial as they're floats, and baked data is just quantized FColors
 	TMap<EMaterialProperty, float> UserConstantValues;
-	for ( const FPropertyEntry& Entry : InMaterialProperties )
+	for (const FPropertyEntry& Entry : InMaterialProperties)
 	{
-		if ( Entry.bUseConstantValue )
+		if (Entry.bUseConstantValue)
 		{
-			UserConstantValues.Add( Entry.Property, Entry.ConstantValue );
+			UserConstantValues.Add(Entry.Property, Entry.ConstantValue);
 		}
 	}
 
-	return UsdShadeConversionImpl::ConfigureShadePrim( View, WrittenTextures, UserConstantValues, OutUsdShadeMaterial );
+	return UsdShadeConversionImpl::ConfigureShadePrim(View, WrittenTextures, UserConstantValues, OutUsdShadeMaterial);
 }
 
-bool UnrealToUsd::ConvertFlattenMaterial( const FString& InMaterialName, FFlattenMaterial& InMaterial, const TArray<FPropertyEntry>& InMaterialProperties, const FDirectoryPath& InTexturesDir, UE::FUsdPrim& OutUsdShadeMaterialPrim )
+bool UnrealToUsd::ConvertFlattenMaterial(
+	const FString& InMaterialName,
+	FFlattenMaterial& InMaterial,
+	const TArray<FPropertyEntry>& InMaterialProperties,
+	const FDirectoryPath& InTexturesDir,
+	UE::FUsdPrim& OutUsdShadeMaterialPrim
+)
 {
-	pxr::UsdShadeMaterial OutUsdShadeMaterial{ pxr::UsdPrim{ OutUsdShadeMaterialPrim } };
-	if ( !OutUsdShadeMaterial )
+	pxr::UsdShadeMaterial OutUsdShadeMaterial{pxr::UsdPrim{OutUsdShadeMaterialPrim}};
+	if (!OutUsdShadeMaterial)
 	{
 		return false;
 	}
 
-	UsdShadeConversionImpl::FBakedMaterialView View{ InMaterial };
-	TMap<EMaterialProperty, FString> WrittenTextures = UsdShadeConversionImpl::WriteTextures( View, InMaterialName, InTexturesDir );
+	UsdShadeConversionImpl::FBakedMaterialView View{InMaterial};
+	TMap<EMaterialProperty, FString> WrittenTextures = UsdShadeConversionImpl::WriteTextures(View, InMaterialName, InTexturesDir);
 
 	// Manually add user supplied constant values. Can't place these in InMaterial as they're floats, and baked data is just quantized FColors
 	TMap<EMaterialProperty, float> UserConstantValues;
-	for ( const FPropertyEntry& Entry : InMaterialProperties )
+	for (const FPropertyEntry& Entry : InMaterialProperties)
 	{
-		if ( Entry.bUseConstantValue )
+		if (Entry.bUseConstantValue)
 		{
-			UserConstantValues.Add( Entry.Property, Entry.ConstantValue );
+			UserConstantValues.Add(Entry.Property, Entry.ConstantValue);
 		}
 	}
 
-	return UsdShadeConversionImpl::ConfigureShadePrim( View, WrittenTextures, UserConstantValues, OutUsdShadeMaterial );
+	return UsdShadeConversionImpl::ConfigureShadePrim(View, WrittenTextures, UserConstantValues, OutUsdShadeMaterial);
 }
 
-#endif // WITH_EDITOR
+#endif	  // WITH_EDITOR
 
-FString UsdUtils::GetResolvedTexturePath( const pxr::UsdAttribute& TextureAssetPathAttr )
+FString UsdUtils::GetResolvedTexturePath(const pxr::UsdAttribute& TextureAssetPathAttr)
 {
 	FScopedUsdAllocs UsdAllocs;
 
 	pxr::SdfAssetPath TextureAssetPath;
-	TextureAssetPathAttr.Get< pxr::SdfAssetPath >( &TextureAssetPath );
+	TextureAssetPathAttr.Get<pxr::SdfAssetPath>(&TextureAssetPath);
 
 	std::string AssetIdentifier = TextureAssetPath.GetResolvedPath();
 	// Don't normalize an empty path as the result will be "."
-	if ( AssetIdentifier.size() > 0 )
+	if (AssetIdentifier.size() > 0)
 	{
 		pxr::ArResolver& Resolver = pxr::ArGetResolver();
-		AssetIdentifier = Resolver.CreateIdentifier( AssetIdentifier );
+		AssetIdentifier = Resolver.CreateIdentifier(AssetIdentifier);
 	}
 
-	FString ResolvedTexturePath = UsdToUnreal::ConvertString( AssetIdentifier );
-	FPaths::NormalizeFilename( ResolvedTexturePath );
+	FString ResolvedTexturePath = UsdToUnreal::ConvertString(AssetIdentifier);
+	FPaths::NormalizeFilename(ResolvedTexturePath);
 
-	if ( ResolvedTexturePath.IsEmpty() )
+	if (ResolvedTexturePath.IsEmpty())
 	{
-		FString TexturePath = UsdToUnreal::ConvertString( TextureAssetPath.GetAssetPath() );
-		FPaths::NormalizeFilename( TexturePath );
+		FString TexturePath = UsdToUnreal::ConvertString(TextureAssetPath.GetAssetPath());
+		FPaths::NormalizeFilename(TexturePath);
 
-		if ( !TexturePath.IsEmpty() )
+		if (!TexturePath.IsEmpty())
 		{
-			pxr::SdfLayerRefPtr TextureLayer = UsdUtils::FindLayerForAttribute( TextureAssetPathAttr, pxr::UsdTimeCode::EarliestTime().GetValue() );
-			ResolvedTexturePath = UsdShadeConversionImpl::ResolveTexturePath( TextureLayer, TexturePath );
+			pxr::SdfLayerRefPtr TextureLayer = UsdUtils::FindLayerForAttribute(TextureAssetPathAttr, pxr::UsdTimeCode::EarliestTime().GetValue());
+			ResolvedTexturePath = UsdShadeConversionImpl::ResolveTexturePath(TextureLayer, TexturePath);
 		}
 	}
 
-	if ( ResolvedTexturePath.IsEmpty() )
+	if (ResolvedTexturePath.IsEmpty())
 	{
-		UE_LOG( LogUsd, Warning, TEXT( "Failed to resolve texture path on attribute '%s'" ), *UsdToUnreal::ConvertPath( TextureAssetPathAttr.GetPath() ) );
+		UE_LOG(LogUsd, Warning, TEXT("Failed to resolve texture path on attribute '%s'"), *UsdToUnreal::ConvertPath(TextureAssetPathAttr.GetPath()));
 	}
 
 	return ResolvedTexturePath;
@@ -2943,24 +3016,24 @@ FString UsdUtils::GetTextureHash(
 	return LexToString(Hash);
 }
 
-UTexture* UsdUtils::CreateTexture( const pxr::UsdAttribute& TextureAssetPathAttr, const FString& PrimPath, TextureGroup LODGroup, UObject* Outer )
+UTexture* UsdUtils::CreateTexture(const pxr::UsdAttribute& TextureAssetPathAttr, const FString& PrimPath, TextureGroup LODGroup, UObject* Outer)
 {
-	TRACE_CPUPROFILER_EVENT_SCOPE( UsdUtils::CreateTexture );
+	TRACE_CPUPROFILER_EVENT_SCOPE(UsdUtils::CreateTexture);
 
 	// Standalone game does have WITH_EDITOR defined, but it can't use the texture factories, so we need to check this instead
-	if ( GIsEditor )
+	if (GIsEditor)
 	{
-		return UsdShadeConversionImpl::CreateTextureWithEditor( TextureAssetPathAttr, PrimPath, LODGroup, Outer );
+		return UsdShadeConversionImpl::CreateTextureWithEditor(TextureAssetPathAttr, PrimPath, LODGroup, Outer);
 	}
 	else
 	{
-		return UsdShadeConversionImpl::CreateTextureAtRuntime( TextureAssetPathAttr, PrimPath, LODGroup );
+		return UsdShadeConversionImpl::CreateTextureAtRuntime(TextureAssetPathAttr, PrimPath, LODGroup);
 	}
 }
 
-void UsdUtils::NotifyIfVirtualTexturesNeeded( UTexture* Texture )
+void UsdUtils::NotifyIfVirtualTexturesNeeded(UTexture* Texture)
 {
-	if ( !Texture || !Texture->VirtualTextureStreaming )
+	if (!Texture || !Texture->VirtualTextureStreaming)
 	{
 		return;
 	}
@@ -2974,27 +3047,31 @@ void UsdUtils::NotifyIfVirtualTexturesNeeded( UTexture* Texture )
 		}
 	}
 
-	if ( !UseVirtualTexturing( GMaxRHIShaderPlatform ) )
+	if (!UseVirtualTexturing(GMaxRHIShaderPlatform))
 	{
 		FUsdLogManager::LogMessage(
 			EMessageSeverity::Warning,
-			FText::Format( LOCTEXT( "DisabledVirtualTexturing", "Texture '{0}' (from prim '{1}') requires Virtual Textures, but the feature is disabled for this project" ),
-				FText::FromString( Texture->GetName() ),
-				FText::FromString( TexturePath ) )
+			FText::Format(
+				LOCTEXT(
+					"DisabledVirtualTexturing",
+					"Texture '{0}' (from prim '{1}') requires Virtual Textures, but the feature is disabled for this project"
+				),
+				FText::FromString(Texture->GetName()),
+				FText::FromString(TexturePath)
+			)
 		);
 	}
 }
 
 #if WITH_EDITOR
-EFlattenMaterialProperties UsdUtils::MaterialPropertyToFlattenProperty( EMaterialProperty MaterialProperty )
+EFlattenMaterialProperties UsdUtils::MaterialPropertyToFlattenProperty(EMaterialProperty MaterialProperty)
 {
-	const static TMap<EMaterialProperty, EFlattenMaterialProperties> InvertedMap =
-	[]()
+	const static TMap<EMaterialProperty, EFlattenMaterialProperties> InvertedMap = []()
 	{
 		TMap<EMaterialProperty, EFlattenMaterialProperties> Result;
-		Result.Reserve( UsdShadeConversionImpl::FlattenToMaterialProperty.Num() );
+		Result.Reserve(UsdShadeConversionImpl::FlattenToMaterialProperty.Num());
 
-		for ( const TPair< EFlattenMaterialProperties, EMaterialProperty >& EnumPair : UsdShadeConversionImpl::FlattenToMaterialProperty )
+		for (const TPair<EFlattenMaterialProperties, EMaterialProperty>& EnumPair : UsdShadeConversionImpl::FlattenToMaterialProperty)
 		{
 			Result.Add(EnumPair.Value, EnumPair.Key);
 		}
@@ -3002,7 +3079,7 @@ EFlattenMaterialProperties UsdUtils::MaterialPropertyToFlattenProperty( EMateria
 		return Result;
 	}();
 
-	if ( const EFlattenMaterialProperties* FoundConversion = InvertedMap.Find( MaterialProperty ) )
+	if (const EFlattenMaterialProperties* FoundConversion = InvertedMap.Find(MaterialProperty))
 	{
 		return *FoundConversion;
 	}
@@ -3010,9 +3087,9 @@ EFlattenMaterialProperties UsdUtils::MaterialPropertyToFlattenProperty( EMateria
 	return EFlattenMaterialProperties::NumFlattenMaterialProperties;
 }
 
-EMaterialProperty UsdUtils::FlattenPropertyToMaterialProperty( EFlattenMaterialProperties FlattenProperty )
+EMaterialProperty UsdUtils::FlattenPropertyToMaterialProperty(EFlattenMaterialProperties FlattenProperty)
 {
-	if ( const EMaterialProperty* FoundConversion = UsdShadeConversionImpl::FlattenToMaterialProperty.Find( FlattenProperty ) )
+	if (const EMaterialProperty* FoundConversion = UsdShadeConversionImpl::FlattenToMaterialProperty.Find(FlattenProperty))
 	{
 		return *FoundConversion;
 	}
@@ -3020,18 +3097,24 @@ EMaterialProperty UsdUtils::FlattenPropertyToMaterialProperty( EFlattenMaterialP
 	return EMaterialProperty::MP_MAX;
 }
 
-void UsdUtils::CollapseConstantChannelsToSinglePixel( FFlattenMaterial& InMaterial )
+void UsdUtils::CollapseConstantChannelsToSinglePixel(FFlattenMaterial& InMaterial)
 {
-	auto ConvertSamplesInPlace = []( TArray<FColor>& Samples ) -> bool
+	auto ConvertSamplesInPlace = [](TArray<FColor>& Samples) -> bool
 	{
-		if ( Samples.Num() < 2 )
+		if (Samples.Num() < 2)
 		{
 			return false;
 		}
 
 		FColor ConstantValue = Samples[0];
-		bool bResize = Algo::AllOf( Samples, [&ConstantValue]( const FColor& Sample ) { return Sample == ConstantValue; } );
-		if ( bResize )
+		bool bResize = Algo::AllOf(
+			Samples,
+			[&ConstantValue](const FColor& Sample)
+			{
+				return Sample == ConstantValue;
+			}
+		);
+		if (bResize)
 		{
 			Samples.SetNum(1);
 			return true;
@@ -3040,24 +3123,24 @@ void UsdUtils::CollapseConstantChannelsToSinglePixel( FFlattenMaterial& InMateri
 		return false;
 	};
 
-	for ( const TPair< EFlattenMaterialProperties, EMaterialProperty >& EnumPair : UsdShadeConversionImpl::FlattenToMaterialProperty )
+	for (const TPair<EFlattenMaterialProperties, EMaterialProperty>& EnumPair : UsdShadeConversionImpl::FlattenToMaterialProperty)
 	{
 		const EFlattenMaterialProperties Property = EnumPair.Key;
-		bool bResized = ConvertSamplesInPlace( InMaterial.GetPropertySamples( Property ) );
-		if ( bResized )
+		bool bResized = ConvertSamplesInPlace(InMaterial.GetPropertySamples(Property));
+		if (bResized)
 		{
-			InMaterial.SetPropertySize( Property, FIntPoint( 1, 1 ) );
+			InMaterial.SetPropertySize(Property, FIntPoint(1, 1));
 		}
 	}
 }
-#endif // WITH_EDITOR
+#endif	  // WITH_EDITOR
 
-bool UsdUtils::MarkMaterialPrimWithWorldSpaceNormals( const UE::FUsdPrim& MaterialPrim )
+bool UsdUtils::MarkMaterialPrimWithWorldSpaceNormals(const UE::FUsdPrim& MaterialPrim)
 {
 	FScopedUsdAllocs UsdAllocs;
 
 	pxr::UsdPrim UsdPrim{MaterialPrim};
-	if ( !UsdPrim )
+	if (!UsdPrim)
 	{
 		return false;
 	}
@@ -3068,20 +3151,23 @@ bool UsdUtils::MarkMaterialPrimWithWorldSpaceNormals( const UE::FUsdPrim& Materi
 	}
 
 	const bool bCustom = true;
-	pxr::UsdAttribute Attr = UsdPrim.CreateAttribute( UnrealIdentifiers::WorldSpaceNormals, pxr::SdfValueTypeNames->Bool, bCustom );
-	if ( !Attr )
+	pxr::UsdAttribute Attr = UsdPrim.CreateAttribute(UnrealIdentifiers::WorldSpaceNormals, pxr::SdfValueTypeNames->Bool, bCustom);
+	if (!Attr)
 	{
 		return false;
 	}
 
-	Attr.Set<bool>( true );
+	Attr.Set<bool>(true);
 	UsdUtils::NotifyIfOverriddenOpinion(Attr);
 	return true;
 }
 
 void UsdUtils::SetScalarParameterValue(UMaterialInstance& Material, const TCHAR* ParameterName, float ParameterValue)
 {
-	UE_LOG(LogUsd, Verbose, TEXT("Setting material '%s' scalar parameter '%s' with value '%f'"),
+	UE_LOG(
+		LogUsd,
+		Verbose,
+		TEXT("Setting material '%s' scalar parameter '%s' with value '%f'"),
 		*Material.GetPathName(),
 		ParameterName ? ParameterName : TEXT("nullptr"),
 		ParameterValue
@@ -3099,12 +3185,15 @@ void UsdUtils::SetScalarParameterValue(UMaterialInstance& Material, const TCHAR*
 	{
 		Constant->SetScalarParameterValueEditorOnly(Info, ParameterValue);
 	}
-#endif // WITH_EDITOR
+#endif	  // WITH_EDITOR
 }
 
 void UsdUtils::SetVectorParameterValue(UMaterialInstance& Material, const TCHAR* ParameterName, FLinearColor ParameterValue)
 {
-	UE_LOG(LogUsd, Verbose, TEXT("Setting material '%s' vector parameter '%s' with value '%s'"),
+	UE_LOG(
+		LogUsd,
+		Verbose,
+		TEXT("Setting material '%s' vector parameter '%s' with value '%s'"),
 		*Material.GetPathName(),
 		ParameterName ? ParameterName : TEXT("nullptr"),
 		*ParameterValue.ToString()
@@ -3122,12 +3211,15 @@ void UsdUtils::SetVectorParameterValue(UMaterialInstance& Material, const TCHAR*
 	{
 		Constant->SetVectorParameterValueEditorOnly(Info, ParameterValue);
 	}
-#endif // WITH_EDITOR
+#endif	  // WITH_EDITOR
 }
 
 void UsdUtils::SetTextureParameterValue(UMaterialInstance& Material, const TCHAR* ParameterName, UTexture* ParameterValue)
 {
-	UE_LOG(LogUsd, Verbose, TEXT("Setting material '%s' texture parameter '%s' with value '%s'"),
+	UE_LOG(
+		LogUsd,
+		Verbose,
+		TEXT("Setting material '%s' texture parameter '%s' with value '%s'"),
 		*Material.GetPathName(),
 		ParameterName ? ParameterName : TEXT("nullptr"),
 		*ParameterValue->GetPathName()
@@ -3145,12 +3237,15 @@ void UsdUtils::SetTextureParameterValue(UMaterialInstance& Material, const TCHAR
 	{
 		Constant->SetTextureParameterValueEditorOnly(Info, ParameterValue);
 	}
-#endif // WITH_EDITOR
+#endif	  // WITH_EDITOR
 }
 
 void UsdUtils::SetBoolParameterValue(UMaterialInstance& Material, const TCHAR* ParameterName, bool bParameterValue)
 {
-	UE_LOG(LogUsd, Verbose, TEXT("Setting material '%s' bool parameter '%s' with value '%d'"),
+	UE_LOG(
+		LogUsd,
+		Verbose,
+		TEXT("Setting material '%s' bool parameter '%s' with value '%d'"),
 		*Material.GetPathName(),
 		ParameterName ? ParameterName : TEXT("nullptr"),
 		bParameterValue
@@ -3190,7 +3285,7 @@ void UsdUtils::SetBoolParameterValue(UMaterialInstance& Material, const TCHAR* P
 			Constant->UpdateStaticPermutation(StaticParameters);
 		}
 	}
-#endif // WITH_EDITOR
+#endif	  // WITH_EDITOR
 
 	// Try it as a scalar parameter
 	if (!bFound)
@@ -3199,56 +3294,56 @@ void UsdUtils::SetBoolParameterValue(UMaterialInstance& Material, const TCHAR* P
 	}
 }
 
-TOptional<FString> UsdUtils::GetUnrealSurfaceOutput( const pxr::UsdPrim& MaterialPrim )
+TOptional<FString> UsdUtils::GetUnrealSurfaceOutput(const pxr::UsdPrim& MaterialPrim)
 {
-	if ( !MaterialPrim )
+	if (!MaterialPrim)
 	{
 		return {};
 	}
 
 	FScopedUsdAllocs UsdAllocs;
 
-	pxr::UsdShadeMaterial ShadeMaterial{ MaterialPrim };
-	if ( !ShadeMaterial )
+	pxr::UsdShadeMaterial ShadeMaterial{MaterialPrim};
+	if (!ShadeMaterial)
 	{
 		return {};
 	}
 
-	pxr::UsdShadeShader SurfaceShader = ShadeMaterial.ComputeSurfaceSource( UnrealIdentifiers::Unreal );
-	if ( !SurfaceShader )
+	pxr::UsdShadeShader SurfaceShader = ShadeMaterial.ComputeSurfaceSource(UnrealIdentifiers::Unreal);
+	if (!SurfaceShader)
 	{
 		return {};
 	}
 
 	pxr::SdfAssetPath AssetPath;
-	if ( SurfaceShader.GetSourceAsset( &AssetPath, UnrealIdentifiers::Unreal ) )
+	if (SurfaceShader.GetSourceAsset(&AssetPath, UnrealIdentifiers::Unreal))
 	{
-		return UsdToUnreal::ConvertString( AssetPath.GetAssetPath() );
+		return UsdToUnreal::ConvertString(AssetPath.GetAssetPath());
 	}
 
 	return {};
 }
 
-bool UsdUtils::SetUnrealSurfaceOutput( pxr::UsdPrim& MaterialPrim, const FString& UnrealMaterialPathName )
+bool UsdUtils::SetUnrealSurfaceOutput(pxr::UsdPrim& MaterialPrim, const FString& UnrealMaterialPathName)
 {
-	if ( !MaterialPrim )
+	if (!MaterialPrim)
 	{
 		return false;
 	}
 
 	FScopedUsdAllocs UsdAllocs;
 
-	pxr::UsdShadeMaterial Material{ MaterialPrim };
-	if ( !Material )
+	pxr::UsdShadeMaterial Material{MaterialPrim};
+	if (!Material)
 	{
 		return false;
 	}
 
 	pxr::UsdStageRefPtr Stage = MaterialPrim.GetStage();
-	pxr::SdfPath ShaderPath = MaterialPrim.GetPath().AppendChild( UnrealToUsd::ConvertToken( TEXT( "UnrealShader" ) ).Get() );
+	pxr::SdfPath ShaderPath = MaterialPrim.GetPath().AppendChild(UnrealToUsd::ConvertToken(TEXT("UnrealShader")).Get());
 
-	pxr::UsdShadeShader UnrealShader = pxr::UsdShadeShader::Define( Stage, ShaderPath );
-	if ( !UnrealShader )
+	pxr::UsdShadeShader UnrealShader = pxr::UsdShadeShader::Define(Stage, ShaderPath);
+	if (!UnrealShader)
 	{
 		return false;
 	}
@@ -3256,56 +3351,56 @@ bool UsdUtils::SetUnrealSurfaceOutput( pxr::UsdPrim& MaterialPrim, const FString
 	// Let SetSourceAsset call CreateImplementationSourceAttr internally as it will create the attribute with the correct metadata.
 	// For some reason, if we try doing this on Linux we get an attribute that seems to always output "id" when we're exporting material bindings.
 	ensure(UnrealShader.SetSourceAsset(
-		UnrealMaterialPathName.IsEmpty() ? pxr::SdfAssetPath{} : pxr::SdfAssetPath{ UnrealToUsd::ConvertString( *UnrealMaterialPathName ).Get() },
+		UnrealMaterialPathName.IsEmpty() ? pxr::SdfAssetPath{} : pxr::SdfAssetPath{UnrealToUsd::ConvertString(*UnrealMaterialPathName).Get()},
 		UnrealIdentifiers::Unreal
 	));
-	pxr::UsdShadeOutput ShaderOutput = UnrealShader.CreateOutput( UnrealToUsd::ConvertToken( TEXT( "out" ) ).Get(), pxr::SdfValueTypeNames->Token );
+	pxr::UsdShadeOutput ShaderOutput = UnrealShader.CreateOutput(UnrealToUsd::ConvertToken(TEXT("out")).Get(), pxr::SdfValueTypeNames->Token);
 
-	pxr::UsdShadeOutput MaterialOutput = Material.CreateSurfaceOutput( UnrealIdentifiers::Unreal );
-	pxr::UsdShadeConnectableAPI::ConnectToSource( MaterialOutput, ShaderOutput );
+	pxr::UsdShadeOutput MaterialOutput = Material.CreateSurfaceOutput(UnrealIdentifiers::Unreal);
+	pxr::UsdShadeConnectableAPI::ConnectToSource(MaterialOutput, ShaderOutput);
 
 	return true;
 }
 
-bool UsdUtils::RemoveUnrealSurfaceOutput( pxr::UsdPrim& MaterialPrim, const UE::FSdfLayer& LayerToAuthorIn )
+bool UsdUtils::RemoveUnrealSurfaceOutput(pxr::UsdPrim& MaterialPrim, const UE::FSdfLayer& LayerToAuthorIn)
 {
 	FScopedUsdAllocs UsdAllocs;
 
-	pxr::UsdShadeMaterial ShadeMaterial{ MaterialPrim };
-	pxr::UsdShadeConnectableAPI Connectable{ MaterialPrim };
-	if ( !ShadeMaterial || !Connectable )
+	pxr::UsdShadeMaterial ShadeMaterial{MaterialPrim};
+	pxr::UsdShadeConnectableAPI Connectable{MaterialPrim};
+	if (!ShadeMaterial || !Connectable)
 	{
 		return false;
 	}
 
-	if ( pxr::UsdShadeOutput MaterialOutput = ShadeMaterial.GetSurfaceOutput( UnrealIdentifiers::Unreal ) )
+	if (pxr::UsdShadeOutput MaterialOutput = ShadeMaterial.GetSurfaceOutput(UnrealIdentifiers::Unreal))
 	{
-		if ( pxr::UsdShadeShader SurfaceShader = ShadeMaterial.ComputeSurfaceSource( UnrealIdentifiers::Unreal ) )
+		if (pxr::UsdShadeShader SurfaceShader = ShadeMaterial.ComputeSurfaceSource(UnrealIdentifiers::Unreal))
 		{
 			// Fully remove the UnrealShader
-			UsdUtils::RemoveAllLocalPrimSpecs( UE::FUsdPrim{ SurfaceShader.GetPrim() }, LayerToAuthorIn );
+			UsdUtils::RemoveAllLocalPrimSpecs(UE::FUsdPrim{SurfaceShader.GetPrim()}, LayerToAuthorIn);
 		}
 
 		// Disconnect would author something like `token outputs:unreal:surface.connect = None`, which is not quite what we want:
 		// That would be an opinion to have it connected to nothing, but instead we just want to remove any opinion whatsoever,
 		// which is what ClearSource does. Note that these will still leave behind `token outputs:unreal:surface` lines,
 		// but those don't actually count as opinions apparently
-		pxr::UsdShadeConnectableAPI::ClearSource( MaterialOutput );
+		pxr::UsdShadeConnectableAPI::ClearSource(MaterialOutput);
 	}
 
 	return true;
 }
 
-bool UsdUtils::IsMaterialTranslucent( const pxr::UsdShadeMaterial& UsdShadeMaterial )
+bool UsdUtils::IsMaterialTranslucent(const pxr::UsdShadeMaterial& UsdShadeMaterial)
 {
 	FScopedUsdAllocs UsdAllocs;
 
 	pxr::UsdShadeShader SurfaceShader = UsdShadeMaterial.ComputeSurfaceSource();
-	if ( !SurfaceShader )
+	if (!SurfaceShader)
 	{
 		return false;
 	}
-	pxr::UsdShadeConnectableAPI Connectable{ SurfaceShader };
+	pxr::UsdShadeConnectableAPI Connectable{SurfaceShader};
 
 	const UMaterialInterface* Material = nullptr;
 	UUsdAssetCache2* TexturesCache = nullptr;
@@ -3321,44 +3416,45 @@ bool UsdUtils::IsMaterialTranslucent( const pxr::UsdShadeMaterial& UsdShadeMater
 		bReuseIdenticalAssets
 	);
 
-	// Don't check if the texture is nullptr here as we won't actually parse it yet. If the variant has this type we know it's meant to be bound to a texture
+	// Don't check if the texture is nullptr here as we won't actually parse it yet. If the variant has this type we know it's meant to be bound to a
+	// texture
 	const bool bHasBoundTexture = ParameterValue.IsType<UsdShadeConversionImpl::FTextureParameterValue>();
-	const bool bIsTranslucentFloat = ( ParameterValue.IsType<float>() && !FMath::IsNearlyEqual( ParameterValue.Get<float>(), 1.0f ) );
+	const bool bIsTranslucentFloat = (ParameterValue.IsType<float>() && !FMath::IsNearlyEqual(ParameterValue.Get<float>(), 1.0f));
 
-	return bHasOpacityConnection && ( bIsTranslucentFloat || bHasBoundTexture );
+	return bHasOpacityConnection && (bIsTranslucentFloat || bHasBoundTexture);
 }
 
-FSHAHash UsdUtils::HashShadeMaterial( const pxr::UsdShadeMaterial& UsdShadeMaterial, const pxr::TfToken& RenderContext )
+FSHAHash UsdUtils::HashShadeMaterial(const pxr::UsdShadeMaterial& UsdShadeMaterial, const pxr::TfToken& RenderContext)
 {
 	FScopedUsdAllocs UsdAllocs;
 
-	pxr::UsdShadeShader SurfaceShader = UsdShadeMaterial.ComputeSurfaceSource( { RenderContext } );
-	if ( !SurfaceShader )
+	pxr::UsdShadeShader SurfaceShader = UsdShadeMaterial.ComputeSurfaceSource({RenderContext});
+	if (!SurfaceShader)
 	{
 		return {};
 	}
 
 	FSHA1 HashState;
 
-	for ( const pxr::UsdShadeInput& ShadeInput : SurfaceShader.GetInputs() )
+	for (const pxr::UsdShadeInput& ShadeInput : SurfaceShader.GetInputs())
 	{
-		UsdShadeConversionImpl::HashShadeInput( ShadeInput, HashState );
+		UsdShadeConversionImpl::HashShadeInput(ShadeInput, HashState);
 	}
 
 	bool bValue = false;
-	if ( pxr::UsdAttribute Attr = UsdShadeMaterial.GetPrim().GetAttribute( UnrealIdentifiers::WorldSpaceNormals ) )
+	if (pxr::UsdAttribute Attr = UsdShadeMaterial.GetPrim().GetAttribute(UnrealIdentifiers::WorldSpaceNormals))
 	{
-		Attr.Get<bool>( &bValue );
+		Attr.Get<bool>(&bValue);
 	}
-	HashState.Update( reinterpret_cast< uint8* >( &bValue ), sizeof( bValue ) );
+	HashState.Update(reinterpret_cast<uint8*>(&bValue), sizeof(bValue));
 
 	FSHAHash OutHash;
 	HashState.Final();
-	HashState.GetHash( &OutHash.Hash[0] );
+	HashState.GetHash(&OutHash.Hash[0]);
 
 	return OutHash;
 }
 
 #undef LOCTEXT_NAMESPACE
 
-#endif // #if USE_USD_SDK
+#endif	  // #if USE_USD_SDK
