USD: Reformat USDUtilities module.

diff --git a/Engine/Plugins/Importers/USDImporter/Source/USDUtilities/Private/USDSkeletalDataConversion.cpp b/Engine/Plugins/Importers/USDImporter/Source/USDUtilities/Private/USDSkeletalDataConversion.cpp
--- a/Engine/Plugins/Importers/USDImporter/Source/USDUtilities/Private/USDSkeletalDataConversion.cpp
+++ b/Engine/Plugins/Importers/USDImporter/Source/USDUtilities/Private/USDSkeletalDataConversion.cpp
@@ -5,7 +5,6 @@
 #include "Engine/SkinnedAssetCommon.h"
 #include "UnrealUSDWrapper.h"
 #include "UObject/Package.h"
-#include "USDAssetImportData.h"
 #include "USDAttributeUtils.h"
 #include "USDClassesModule.h"
 #include "USDConversionUtils.h"
@@ -31,15 +30,12 @@
 #include "ControlRig.h"
 #include "Evaluation/MovieSceneSequenceTransform.h"
 #include "IMovieScenePlayer.h"
-#include "Interfaces/ITargetPlatform.h"
 #include "Interfaces/ITargetPlatformManagerModule.h"
-#include "LevelSequencePlayer.h"
 #include "Misc/CoreMisc.h"
 #include "Misc/MemStack.h"
 #include "Modules/ModuleManager.h"
 #include "MovieScene.h"
 #include "MovieSceneTimeHelpers.h"
-#include "MovieSceneTrack.h"
 #include "Rendering/SkeletalMeshLODImporterData.h"
 #include "Rendering/SkeletalMeshLODModel.h"
 #include "Rendering/SkeletalMeshLODRenderData.h"
@@ -53,33 +49,32 @@
 #include "Factories/FbxSkeletalMeshImportData.h"
 #include "ImportUtils/SkeletalMeshImportUtils.h"
 #include "MeshUtilities.h"
-#endif // WITH_EDITOR
+#endif	  // WITH_EDITOR
 
 #if USE_USD_SDK
 #include "USDIncludesStart.h"
-	#include "pxr/usd/sdf/types.h"
-	#include "pxr/usd/usd/editContext.h"
-	#include "pxr/usd/usd/variantSets.h"
-	#include "pxr/usd/usdGeom/mesh.h"
-	#include "pxr/usd/usdGeom/primvarsAPI.h"
-	#include "pxr/usd/usdGeom/subset.h"
-	#include "pxr/usd/usdGeom/tokens.h"
-	#include "pxr/usd/usdGeom/xformable.h"
-	#include "pxr/usd/usdGeom/xformCache.h"
-	#include "pxr/usd/usdShade/tokens.h"
-	#include "pxr/usd/usdSkel/animation.h"
-	#include "pxr/usd/usdSkel/animMapper.h"
-	#include "pxr/usd/usdSkel/binding.h"
-	#include "pxr/usd/usdSkel/bindingAPI.h"
-	#include "pxr/usd/usdSkel/blendShape.h"
-	#include "pxr/usd/usdSkel/cache.h"
-	#include "pxr/usd/usdSkel/root.h"
-	#include "pxr/usd/usdSkel/skeletonQuery.h"
-	#include "pxr/usd/usdSkel/skinningQuery.h"
-	#include "pxr/usd/usdSkel/topology.h"
-	#include "pxr/usd/usdSkel/utils.h"
+#include "pxr/usd/sdf/types.h"
+#include "pxr/usd/usd/editContext.h"
+#include "pxr/usd/usd/variantSets.h"
+#include "pxr/usd/usdGeom/mesh.h"
+#include "pxr/usd/usdGeom/primvarsAPI.h"
+#include "pxr/usd/usdGeom/subset.h"
+#include "pxr/usd/usdGeom/tokens.h"
+#include "pxr/usd/usdGeom/xformable.h"
+#include "pxr/usd/usdShade/tokens.h"
+#include "pxr/usd/usdSkel/animation.h"
+#include "pxr/usd/usdSkel/animMapper.h"
+#include "pxr/usd/usdSkel/binding.h"
+#include "pxr/usd/usdSkel/bindingAPI.h"
+#include "pxr/usd/usdSkel/blendShape.h"
+#include "pxr/usd/usdSkel/cache.h"
+#include "pxr/usd/usdSkel/root.h"
+#include "pxr/usd/usdSkel/skeletonQuery.h"
+#include "pxr/usd/usdSkel/skinningQuery.h"
+#include "pxr/usd/usdSkel/topology.h"
+#include "pxr/usd/usdSkel/utils.h"
 #include "USDIncludesEnd.h"
-#endif // USE_USD_SDK
+#endif	  // USE_USD_SDK
 
 #define LOCTEXT_NAMESPACE "UsdSkeletalDataConversion"
 
@@ -89,7 +84,9 @@ static bool bAddCurveMetadataToSkeleton = true;
 static FAutoConsoleVariableRef CVarAddCurveMetadataToSkeleton(
 	TEXT("USD.AddCurveMetadataToSkeleton"),
 	bAddCurveMetadataToSkeleton,
-	TEXT("When true will cause blend shape / morph target float curve data to be added to generated USkeleton assets. When false, this curve data will be added to the generated USkeletalMesh assets instead."));
+	TEXT("When true will cause blend shape / morph target float curve data to be added to generated USkeleton assets. When false, this curve data "
+		 "will be added to the generated USkeletalMesh assets instead.")
+);
 
 #if USE_USD_SDK && WITH_EDITOR
 namespace SkelDataConversionImpl
@@ -97,45 +94,45 @@ namespace SkelDataConversionImpl
 	// Adapted from LODUtilities.cpp
 	struct FMeshDataBundle
 	{
-		TArray< FVector3f > Vertices;
-		TArray< FVector3f> NormalsPerVertex;
-		TArray< uint32 > Indices;
-		TArray< FVector2f > UVs;
-		TArray< uint32 > SmoothingGroups;
+		TArray<FVector3f> Vertices;
+		TArray<FVector3f> NormalsPerVertex;
+		TArray<uint32> Indices;
+		TArray<FVector2f> UVs;
+		TArray<uint32> SmoothingGroups;
 		TArray<SkeletalMeshImportData::FTriangle> Faces;
-		TMap< uint32, TArray< uint32 > > VertexIndexToFaceIndices;
+		TMap<uint32, TArray<uint32>> VertexIndexToFaceIndices;
 	};
 
 	struct FMorphedMeshBundle
 	{
-		TArray< FVector3f > Vertices;
-		TArray< FVector3f> NormalsPerIndex;
-		TArray< uint32 > Indices;
-		TArray< FVector2f > UVs;
-		TArray< uint32 > SmoothingGroups;
-		TArray< uint32 > MorphedIndexToSourceIndex;
+		TArray<FVector3f> Vertices;
+		TArray<FVector3f> NormalsPerIndex;
+		TArray<uint32> Indices;
+		TArray<FVector2f> UVs;
+		TArray<uint32> SmoothingGroups;
+		TArray<uint32> MorphedIndexToSourceIndex;
 	};
 
 	/** Converts from wedge-based vertex format into a flat format we can give to MeshUtilities */
-	void ConvertImportDataToMeshData( const FSkeletalMeshImportData& ImportData, FMeshDataBundle& MeshDataBundle )
+	void ConvertImportDataToMeshData(const FSkeletalMeshImportData& ImportData, FMeshDataBundle& MeshDataBundle)
 	{
 		MeshDataBundle.VertexIndexToFaceIndices.Reserve(ImportData.Points.Num());
 
-		for ( const SkeletalMeshImportData::FTriangle& Face : ImportData.Faces )
+		for (const SkeletalMeshImportData::FTriangle& Face : ImportData.Faces)
 		{
 			SkeletalMeshImportData::FTriangle FaceTriangle;
 			FaceTriangle = Face;
-			for ( int32 Index = 0; Index < 3; ++Index )
+			for (int32 Index = 0; Index < 3; ++Index)
 			{
-				const SkeletalMeshImportData::FVertex& Wedge = ImportData.Wedges[ Face.WedgeIndex[ Index ] ];
-				FaceTriangle.WedgeIndex[ Index ] = Wedge.VertexIndex;
+				const SkeletalMeshImportData::FVertex& Wedge = ImportData.Wedges[Face.WedgeIndex[Index]];
+				FaceTriangle.WedgeIndex[Index] = Wedge.VertexIndex;
 				MeshDataBundle.Indices.Add(Wedge.VertexIndex);
-				MeshDataBundle.UVs.Add( Wedge.UVs[ 0 ] );
+				MeshDataBundle.UVs.Add(Wedge.UVs[0]);
 
-				MeshDataBundle.VertexIndexToFaceIndices.FindOrAdd( Wedge.VertexIndex ).Add( MeshDataBundle.Faces.Num() );
+				MeshDataBundle.VertexIndexToFaceIndices.FindOrAdd(Wedge.VertexIndex).Add(MeshDataBundle.Faces.Num());
 			}
-			MeshDataBundle.Faces.Add( FaceTriangle );
-			MeshDataBundle.SmoothingGroups.Add( Face.SmoothingGroups );
+			MeshDataBundle.Faces.Add(FaceTriangle);
+			MeshDataBundle.SmoothingGroups.Add(Face.SmoothingGroups);
 		}
 
 		MeshDataBundle.Vertices = ImportData.Points;
@@ -147,51 +144,51 @@ namespace SkelDataConversionImpl
 	 * the skeletal mesh shares vertices between faces, so if a vertex is morphed, not only does its normal need to be recomputed, but also
 	 * the normals of all vertices of triangles that the vertex is a part of.
 	 */
-	void MorphMeshData( const FMeshDataBundle& InMeshDataBundle, TArray<FMorphTargetDelta>& InOutDeltas, FMorphedMeshBundle& OutBundle)
+	void MorphMeshData(const FMeshDataBundle& InMeshDataBundle, TArray<FMorphTargetDelta>& InOutDeltas, FMorphedMeshBundle& OutBundle)
 	{
-		OutBundle.Vertices.Reserve( InOutDeltas.Num());
-		OutBundle.Indices.Reserve( InOutDeltas.Num());
-		OutBundle.UVs.Reserve( InOutDeltas.Num());
-		OutBundle.SmoothingGroups.Reserve( InOutDeltas.Num());
-		OutBundle.MorphedIndexToSourceIndex.Reserve( InOutDeltas.Num());
+		OutBundle.Vertices.Reserve(InOutDeltas.Num());
+		OutBundle.Indices.Reserve(InOutDeltas.Num());
+		OutBundle.UVs.Reserve(InOutDeltas.Num());
+		OutBundle.SmoothingGroups.Reserve(InOutDeltas.Num());
+		OutBundle.MorphedIndexToSourceIndex.Reserve(InOutDeltas.Num());
 
-		TSet< uint32 > AddedFaces;
-		TArray< FMorphTargetDelta > NewDeltas;
-		TMap< uint32, uint32 > SourceIndexToMorphedIndex;
+		TSet<uint32> AddedFaces;
+		TArray<FMorphTargetDelta> NewDeltas;
+		TMap<uint32, uint32> SourceIndexToMorphedIndex;
 
 		// Add the existing deltas to the vertices array first
 		// Don't add indices yet as we can't guarantee these come in triangle order (they're straight from USD)
-		for ( const FMorphTargetDelta& Delta : InOutDeltas )
+		for (const FMorphTargetDelta& Delta : InOutDeltas)
 		{
 			uint32 SourceIndex = Delta.SourceIdx;
-			uint32 MorphedIndex = OutBundle.Vertices.Add(InMeshDataBundle.Vertices[ SourceIndex ] + Delta.PositionDelta);
+			uint32 MorphedIndex = OutBundle.Vertices.Add(InMeshDataBundle.Vertices[SourceIndex] + Delta.PositionDelta);
 
-			OutBundle.MorphedIndexToSourceIndex.Add( SourceIndex );
-			SourceIndexToMorphedIndex.Add( SourceIndex, MorphedIndex );
+			OutBundle.MorphedIndexToSourceIndex.Add(SourceIndex);
+			SourceIndexToMorphedIndex.Add(SourceIndex, MorphedIndex);
 		}
 
 		// Add all indices, creating any missing deltas/vertices
-		for ( const FMorphTargetDelta& Delta : InOutDeltas )
+		for (const FMorphTargetDelta& Delta : InOutDeltas)
 		{
-			if ( const TArray< uint32 >* FoundFaceIndices = InMeshDataBundle.VertexIndexToFaceIndices.Find( Delta.SourceIdx ) )
+			if (const TArray<uint32>* FoundFaceIndices = InMeshDataBundle.VertexIndexToFaceIndices.Find(Delta.SourceIdx))
 			{
-				for ( uint32 FaceIndex : *FoundFaceIndices )
+				for (uint32 FaceIndex : *FoundFaceIndices)
 				{
-					if ( AddedFaces.Contains( FaceIndex ) )
+					if (AddedFaces.Contains(FaceIndex))
 					{
 						continue;
 					}
-					AddedFaces.Add( FaceIndex );
+					AddedFaces.Add(FaceIndex);
 
-					const SkeletalMeshImportData::FTriangle& Face = InMeshDataBundle.Faces[ FaceIndex ];
-					OutBundle.SmoothingGroups.Add( Face.SmoothingGroups );
+					const SkeletalMeshImportData::FTriangle& Face = InMeshDataBundle.Faces[FaceIndex];
+					OutBundle.SmoothingGroups.Add(Face.SmoothingGroups);
 
-					for ( uint32 Index = 0; Index < 3; ++Index )
+					for (uint32 Index = 0; Index < 3; ++Index)
 					{
-						uint32 SourceIndex = Face.WedgeIndex[ Index ];
+						uint32 SourceIndex = Face.WedgeIndex[Index];
 						uint32 MorphedIndex = INDEX_NONE;
 
-						if ( uint32* FoundMorphedIndex = SourceIndexToMorphedIndex.Find( SourceIndex ) )
+						if (uint32* FoundMorphedIndex = SourceIndexToMorphedIndex.Find(SourceIndex))
 						{
 							MorphedIndex = *FoundMorphedIndex;
 						}
@@ -203,14 +200,14 @@ namespace SkelDataConversionImpl
 							NewDelta.TangentZDelta = FVector3f::ZeroVector;
 							NewDelta.SourceIdx = SourceIndex;
 
-							MorphedIndex = OutBundle.Vertices.Add( InMeshDataBundle.Vertices[ SourceIndex ] );
+							MorphedIndex = OutBundle.Vertices.Add(InMeshDataBundle.Vertices[SourceIndex]);
 
-							OutBundle.MorphedIndexToSourceIndex.Add( SourceIndex );
-							SourceIndexToMorphedIndex.Add( SourceIndex, MorphedIndex );
+							OutBundle.MorphedIndexToSourceIndex.Add(SourceIndex);
+							SourceIndexToMorphedIndex.Add(SourceIndex, MorphedIndex);
 						}
 
-						OutBundle.Indices.Add( MorphedIndex );
-						OutBundle.UVs.Add( InMeshDataBundle.UVs[ SourceIndex ] );
+						OutBundle.Indices.Add(MorphedIndex);
+						OutBundle.UVs.Add(InMeshDataBundle.UVs[SourceIndex]);
 					}
 				}
 			}
@@ -224,9 +221,9 @@ namespace SkelDataConversionImpl
 	 * when the morph target is applied to the skeletal mesh.
 	 * Note: This may add deltas to the blend shape: See MorphMeshData
 	 */
-	bool ComputeTangentDeltas( const FMeshDataBundle& MeshDataBundle, UsdUtils::FUsdBlendShape& BlendShape )
+	bool ComputeTangentDeltas(const FMeshDataBundle& MeshDataBundle, UsdUtils::FUsdBlendShape& BlendShape)
 	{
-		if ( BlendShape.bHasAuthoredTangents )
+		if (BlendShape.bHasAuthoredTangents)
 		{
 			return false;
 		}
@@ -234,34 +231,41 @@ namespace SkelDataConversionImpl
 		FMorphedMeshBundle MorphedBundle;
 		MorphMeshData(MeshDataBundle, BlendShape.Vertices, MorphedBundle);
 
-		IMeshUtilities& MeshUtilities = FModuleManager::Get().LoadModuleChecked<IMeshUtilities>( "MeshUtilities" );
-		ETangentOptions::Type TangentOptions = ( ETangentOptions::Type ) ( ETangentOptions::BlendOverlappingNormals | ETangentOptions::UseMikkTSpace );
-		MeshUtilities.CalculateNormals( MorphedBundle.Vertices, MorphedBundle.Indices, MorphedBundle.UVs, MorphedBundle.SmoothingGroups, TangentOptions, MorphedBundle.NormalsPerIndex );
+		IMeshUtilities& MeshUtilities = FModuleManager::Get().LoadModuleChecked<IMeshUtilities>("MeshUtilities");
+		ETangentOptions::Type TangentOptions = (ETangentOptions::Type)(ETangentOptions::BlendOverlappingNormals | ETangentOptions::UseMikkTSpace);
+		MeshUtilities.CalculateNormals(
+			MorphedBundle.Vertices,
+			MorphedBundle.Indices,
+			MorphedBundle.UVs,
+			MorphedBundle.SmoothingGroups,
+			TangentOptions,
+			MorphedBundle.NormalsPerIndex
+		);
 
-		TMap< uint32, FMorphTargetDelta* > SourceIndexToMorphDelta;
-		for ( FMorphTargetDelta& Delta : BlendShape.Vertices )
+		TMap<uint32, FMorphTargetDelta*> SourceIndexToMorphDelta;
+		for (FMorphTargetDelta& Delta : BlendShape.Vertices)
 		{
-			SourceIndexToMorphDelta.Add( Delta.SourceIdx, &Delta );
+			SourceIndexToMorphDelta.Add(Delta.SourceIdx, &Delta);
 		}
 
-		uint32 NumMorphedIndices = static_cast<uint32> ( MorphedBundle.Indices.Num() );
-		for ( uint32 MorphedIndexIndex = 0; MorphedIndexIndex < NumMorphedIndices; ++MorphedIndexIndex )
+		uint32 NumMorphedIndices = static_cast<uint32>(MorphedBundle.Indices.Num());
+		for (uint32 MorphedIndexIndex = 0; MorphedIndexIndex < NumMorphedIndices; ++MorphedIndexIndex)
 		{
-			const uint32 MorphedIndex = MorphedBundle.Indices[ MorphedIndexIndex ];
-			const uint32 SourceIndex = MorphedBundle.MorphedIndexToSourceIndex[ MorphedIndex ];
+			const uint32 MorphedIndex = MorphedBundle.Indices[MorphedIndexIndex];
+			const uint32 SourceIndex = MorphedBundle.MorphedIndexToSourceIndex[MorphedIndex];
 
 			// Note that we store the source normals as one per vertex, but we don't need to do that conversion for the
 			// morphed normals, as we're iterating directly over the indices anyway
-			const FVector& SourceNormal = (FVector)MeshDataBundle.NormalsPerVertex[ SourceIndex ];
-			const FVector& MorphedNormal = (FVector)MorphedBundle.NormalsPerIndex[ MorphedIndexIndex ];
+			const FVector& SourceNormal = (FVector)MeshDataBundle.NormalsPerVertex[SourceIndex];
+			const FVector& MorphedNormal = (FVector)MorphedBundle.NormalsPerIndex[MorphedIndexIndex];
 
-			if ( FMorphTargetDelta** FoundDelta = SourceIndexToMorphDelta.Find( SourceIndex ) )
+			if (FMorphTargetDelta** FoundDelta = SourceIndexToMorphDelta.Find(SourceIndex))
 			{
-				( *FoundDelta )->TangentZDelta = FVector3f(MorphedNormal - SourceNormal);
+				(*FoundDelta)->TangentZDelta = FVector3f(MorphedNormal - SourceNormal);
 
 				// We will visit each delta multiple times because we're iterating indices and these are per-vertex,
 				// so this prevents us from recalculating the delta many times
-				SourceIndexToMorphDelta.Remove( SourceIndex );
+				SourceIndexToMorphDelta.Remove(SourceIndex);
 			}
 		}
 
@@ -269,72 +273,94 @@ namespace SkelDataConversionImpl
 	}
 
 	/** Converts the given offsets into UnrealEditor space and fills in an FUsdBlendShape object with all the data that will become a morph target */
-	bool CreateUsdBlendShape( const FString& Name, const pxr::VtArray< pxr::GfVec3f >& PointOffsets, const pxr::VtArray< pxr::GfVec3f >& NormalOffsets, const pxr::VtArray< int >& PointIndices, const FUsdStageInfo& StageInfo, uint32 PointIndexOffset, int32 LODIndex, UsdUtils::FUsdBlendShape& OutBlendShape )
+	bool CreateUsdBlendShape(
+		const FString& Name,
+		const pxr::VtArray<pxr::GfVec3f>& PointOffsets,
+		const pxr::VtArray<pxr::GfVec3f>& NormalOffsets,
+		const pxr::VtArray<int>& PointIndices,
+		const FUsdStageInfo& StageInfo,
+		uint32 PointIndexOffset,
+		int32 LODIndex,
+		UsdUtils::FUsdBlendShape& OutBlendShape
+	)
 	{
 		uint32 NumOffsets = PointOffsets.size();
 		uint32 NumIndices = PointIndices.size();
 		uint32 NumNormals = NormalOffsets.size();
 
-		if ( NumNormals > 0 && NumOffsets != NumNormals )
+		if (NumNormals > 0 && NumOffsets != NumNormals)
 		{
-			UE_LOG( LogUsd, Warning, TEXT( "BlendShape '%s' has mismatching numbers of offsets (%d) and normalOffsets (%d) and will be ignored" ), *Name, NumOffsets, NumNormals );
+			UE_LOG(
+				LogUsd,
+				Warning,
+				TEXT("BlendShape '%s' has mismatching numbers of offsets (%d) and normalOffsets (%d) and will be ignored"),
+				*Name,
+				NumOffsets,
+				NumNormals
+			);
 			return false;
 		}
 
-		if ( NumIndices > 0 && NumOffsets != NumIndices )
+		if (NumIndices > 0 && NumOffsets != NumIndices)
 		{
-			UE_LOG( LogUsd, Warning, TEXT( "BlendShape '%s' has mismatching numbers of offsets (%d) and point indices (%d) and will be ignored" ), *Name, NumOffsets, NumIndices );
+			UE_LOG(
+				LogUsd,
+				Warning,
+				TEXT("BlendShape '%s' has mismatching numbers of offsets (%d) and point indices (%d) and will be ignored"),
+				*Name,
+				NumOffsets,
+				NumIndices
+			);
 			return false;
 		}
 
-		if ( NumOffsets + NumNormals == 0 )
+		if (NumOffsets + NumNormals == 0)
 		{
-			UE_LOG( LogUsd, Warning, TEXT( "BlendShape '%s' zero offsets and normalOffsets and will be ignored" ), *Name);
+			UE_LOG(LogUsd, Warning, TEXT("BlendShape '%s' zero offsets and normalOffsets and will be ignored"), *Name);
 			return false;
 		}
 
-		if ( NumNormals > 0 )
+		if (NumNormals > 0)
 		{
 			OutBlendShape.bHasAuthoredTangents = true;
 		}
 
 		OutBlendShape.Name = Name;
-		OutBlendShape.LODIndicesThatUseThis.Add( LODIndex );
+		OutBlendShape.LODIndicesThatUseThis.Add(LODIndex);
 
 		// Prepare the indices of the corresponding base points/normals for every local point/normal we have
 		TArray<int32> BaseIndices;
-		BaseIndices.Reserve( NumOffsets );
-		if ( NumIndices == 0 )
+		BaseIndices.Reserve(NumOffsets);
+		if (NumIndices == 0)
 		{
 			// If we have no indices it means we have information for all of our local points/normals
-			for ( uint32 BaseIndex = PointIndexOffset; BaseIndex < PointIndexOffset + NumOffsets; ++BaseIndex )
+			for (uint32 BaseIndex = PointIndexOffset; BaseIndex < PointIndexOffset + NumOffsets; ++BaseIndex)
 			{
-				BaseIndices.Add( static_cast< int32 >( BaseIndex ) );
+				BaseIndices.Add(static_cast<int32>(BaseIndex));
 			}
 		}
 		else
 		{
 			// If we have indices it means our morph target only affects a subset of the base vertices
-			for ( uint32 LocalIndex = 0; LocalIndex < NumOffsets; ++LocalIndex )
+			for (uint32 LocalIndex = 0; LocalIndex < NumOffsets; ++LocalIndex)
 			{
-				int32 BaseIndex = PointIndices[ LocalIndex ] + static_cast< int32 >( PointIndexOffset );
+				int32 BaseIndex = PointIndices[LocalIndex] + static_cast<int32>(PointIndexOffset);
 
-				BaseIndices.Add( BaseIndex );
+				BaseIndices.Add(BaseIndex);
 			}
 		}
 
-		OutBlendShape.Vertices.SetNumUninitialized( NumOffsets );
-		for ( uint32 OffsetIndex = 0; OffsetIndex < NumOffsets; ++OffsetIndex )
+		OutBlendShape.Vertices.SetNumUninitialized(NumOffsets);
+		for (uint32 OffsetIndex = 0; OffsetIndex < NumOffsets; ++OffsetIndex)
 		{
-			const FVector UEOffset = UsdToUnreal::ConvertVector( StageInfo, PointOffsets[ OffsetIndex ] );
-			const FVector UENormal = OutBlendShape.bHasAuthoredTangents
-				? UsdToUnreal::ConvertVector( StageInfo, NormalOffsets[ OffsetIndex ] )
-				: FVector( 0, 0, 0 );
+			const FVector UEOffset = UsdToUnreal::ConvertVector(StageInfo, PointOffsets[OffsetIndex]);
+			const FVector UENormal = OutBlendShape.bHasAuthoredTangents ? UsdToUnreal::ConvertVector(StageInfo, NormalOffsets[OffsetIndex])
+																		: FVector(0, 0, 0);
 
-			FMorphTargetDelta& ModifiedVertex = OutBlendShape.Vertices[ OffsetIndex ];
-			ModifiedVertex.PositionDelta = ( FVector3f ) UEOffset;
-			ModifiedVertex.TangentZDelta = ( FVector3f ) UENormal;
-			ModifiedVertex.SourceIdx = BaseIndices[ OffsetIndex ];
+			FMorphTargetDelta& ModifiedVertex = OutBlendShape.Vertices[OffsetIndex];
+			ModifiedVertex.PositionDelta = (FVector3f)UEOffset;
+			ModifiedVertex.TangentZDelta = (FVector3f)UENormal;
+			ModifiedVertex.SourceIdx = BaseIndices[OffsetIndex];
 		}
 
 		return true;
@@ -346,63 +372,63 @@ namespace SkelDataConversionImpl
 	 * which our deltas refer to these vertices come directly from USD. Example: If a vertex affected by the blend shape is split, we need
 	 * to duplicate the delta to all the split versions.
 	 */
-	void UpdatesDeltasToMeshBuild(TArray<FMorphTargetDelta>& MorphTargetDeltas, const TMap<int32, TArray<int32>>& OrigIndexToBuiltIndices )
+	void UpdatesDeltasToMeshBuild(TArray<FMorphTargetDelta>& MorphTargetDeltas, const TMap<int32, TArray<int32>>& OrigIndexToBuiltIndices)
 	{
 		TSet<int32> DeltasToDelete;
 		TArray<FMorphTargetDelta> NewDeltas;
-		for ( int32 DeltaIndex = 0; DeltaIndex < MorphTargetDeltas.Num(); ++DeltaIndex )
+		for (int32 DeltaIndex = 0; DeltaIndex < MorphTargetDeltas.Num(); ++DeltaIndex)
 		{
-			FMorphTargetDelta& ModifiedVertex = MorphTargetDeltas[ DeltaIndex ];
+			FMorphTargetDelta& ModifiedVertex = MorphTargetDeltas[DeltaIndex];
 
-			if ( const TArray< int32 >* BuiltIndices = OrigIndexToBuiltIndices.Find( ModifiedVertex.SourceIdx ) )
+			if (const TArray<int32>* BuiltIndices = OrigIndexToBuiltIndices.Find(ModifiedVertex.SourceIdx))
 			{
 				// Our index just got remapped somewhere else: Update it
-				if ( BuiltIndices->Num() >= 1 )
+				if (BuiltIndices->Num() >= 1)
 				{
-					ModifiedVertex.SourceIdx = static_cast< uint32 >( ( *BuiltIndices )[ 0 ] );
+					ModifiedVertex.SourceIdx = static_cast<uint32>((*BuiltIndices)[0]);
 				}
 
 				// The vertex we were pointing at got split into multiple vertices: Add a matching delta for each
-				for ( int32 NewDeltaIndex = 1; NewDeltaIndex < BuiltIndices->Num(); ++NewDeltaIndex )
+				for (int32 NewDeltaIndex = 1; NewDeltaIndex < BuiltIndices->Num(); ++NewDeltaIndex)
 				{
-					FMorphTargetDelta& NewDelta = NewDeltas.Add_GetRef( ModifiedVertex );
-					NewDelta.SourceIdx = static_cast< uint32 >( ( *BuiltIndices )[ NewDeltaIndex ] );
+					FMorphTargetDelta& NewDelta = NewDeltas.Add_GetRef(ModifiedVertex);
+					NewDelta.SourceIdx = static_cast<uint32>((*BuiltIndices)[NewDeltaIndex]);
 				}
 			}
 			// The vertex we were pointing at got deleted: Remove the delta
 			else
 			{
-				DeltasToDelete.Add( DeltaIndex );
+				DeltasToDelete.Add(DeltaIndex);
 			}
 		}
-		if ( DeltasToDelete.Num() > 0 )
+		if (DeltasToDelete.Num() > 0)
 		{
-			for ( int32 DeltaIndex = MorphTargetDeltas.Num() - 1; DeltaIndex >= 0; --DeltaIndex )
+			for (int32 DeltaIndex = MorphTargetDeltas.Num() - 1; DeltaIndex >= 0; --DeltaIndex)
 			{
-				if ( DeltasToDelete.Contains( DeltaIndex ) )
+				if (DeltasToDelete.Contains(DeltaIndex))
 				{
-					MorphTargetDeltas.RemoveAt( DeltaIndex, 1, false );
+					MorphTargetDeltas.RemoveAt(DeltaIndex, 1, false);
 				}
 			}
 		}
-		MorphTargetDeltas.Append( MoveTemp( NewDeltas ) );
+		MorphTargetDeltas.Append(MoveTemp(NewDeltas));
 	}
 
 	/**
 	 * Will find or create a AACF_DefaultCurve float curve with CurveName, and set its data to a copy of SourceData.
 	 * Adapted from UnFbx::FFbxImporter::ImportCurveToAnimSequence
 	 */
-	void SetFloatCurveData( UAnimSequence* Sequence, FName CurveName, const FRichCurve& SourceData )
+	void SetFloatCurveData(UAnimSequence* Sequence, FName CurveName, const FRichCurve& SourceData)
 	{
-		TRACE_CPUPROFILER_EVENT_SCOPE( SkelDataConversionImpl::SetFloatCurveData );
+		TRACE_CPUPROFILER_EVENT_SCOPE(SkelDataConversionImpl::SetFloatCurveData);
 
-		if ( !Sequence )
+		if (!Sequence)
 		{
 			return;
 		}
 
 		USkeleton* Skeleton = Sequence->GetSkeleton();
-		if ( !Skeleton )
+		if (!Skeleton)
 		{
 			return;
 		}
@@ -438,15 +464,15 @@ namespace SkelDataConversionImpl
 
 		// Ignore curves that don't contribute to the animation
 		bool bHasNonZeroKey = false;
-		for ( const FRichCurveKey& Key : SourceData.Keys )
+		for (const FRichCurveKey& Key : SourceData.Keys)
 		{
-			if ( !FMath::IsNearlyEqual( Key.Value, 0.0f ) )
+			if (!FMath::IsNearlyEqual(Key.Value, 0.0f))
 			{
 				bHasNonZeroKey = true;
 				break;
 			}
 		}
-		if ( !bHasNonZeroKey )
+		if (!bHasNonZeroKey)
 		{
 			return;
 		}
@@ -457,33 +483,45 @@ namespace SkelDataConversionImpl
 
 		FAnimationCurveIdentifier CurveId(CurveName, ERawCurveTrackTypes::RCT_Float);
 		const FFloatCurve* Curve = DataModel->FindFloatCurve(CurveId);
-		if ( !Curve )
+		if (!Curve)
 		{
 			// If curve doesn't exist, add one
-			Controller.AddCurve( CurveId, AACF_DefaultCurve, bShouldTransact );
-			Curve = DataModel->FindFloatCurve( CurveId );
+			Controller.AddCurve(CurveId, AACF_DefaultCurve, bShouldTransact);
+			Curve = DataModel->FindFloatCurve(CurveId);
 		}
 		else
 		{
-			if ( !( Curve->FloatCurve == SourceData ) )
+			if (!(Curve->FloatCurve == SourceData))
 			{
 				FUsdLogManager::LogMessage(
-						EMessageSeverity::Warning,
-						FText::Format( LOCTEXT( "OverwritingMorphTargetCurves", "Overwriting animation curve for morph target '{0}' with different data! If the Skeletal Mesh has multiple LODs, make sure each LOD mesh that wants to animate a certain blend shape does so with the same blend shape curve." ),
-						FText::FromName( CurveName ) )
+					EMessageSeverity::Warning,
+					FText::Format(
+						LOCTEXT(
+							"OverwritingMorphTargetCurves",
+							"Overwriting animation curve for morph target '{0}' with different data! If the Skeletal Mesh has multiple LODs, make "
+							"sure each LOD mesh that wants to animate a certain blend shape does so with the same blend shape curve."
+						),
+						FText::FromName(CurveName)
+					)
 				);
 			}
 
-			Controller.SetCurveFlags( CurveId, Curve->GetCurveTypeFlags() | AACF_DefaultCurve, bShouldTransact );
+			Controller.SetCurveFlags(CurveId, Curve->GetCurveTypeFlags() | AACF_DefaultCurve, bShouldTransact);
 		}
 
-		if ( Curve )
+		if (Curve)
 		{
-			Controller.SetCurveKeys( CurveId, SourceData.GetConstRefOfKeys(), bShouldTransact );
+			Controller.SetCurveKeys(CurveId, SourceData.GetConstRefOfKeys(), bShouldTransact);
 		}
 		else
 		{
-			UE_LOG( LogUsd, Error, TEXT( "Failed to create float curve with name '%s' for UAnimSequence '%s'" ), *CurveName.ToString(), *Sequence->GetName() );
+			UE_LOG(
+				LogUsd,
+				Error,
+				TEXT("Failed to create float curve with name '%s' for UAnimSequence '%s'"),
+				*CurveName.ToString(),
+				*Sequence->GetName()
+			);
 		}
 	}
 
@@ -493,95 +531,113 @@ namespace SkelDataConversionImpl
 	 * targets that were generated from the blend shape and its inbetweens, if any.
 	 * Index 0 of the returned array always contains the remapped primary morph target weight, and the rest match the inbetween order
 	 */
-	TArray<FRichCurve> ResolveWeightsForBlendShapeCurve( const UsdUtils::FUsdBlendShape& PrimaryBlendShape, const FRichCurve& ChannelWeightCurve )
+	TArray<FRichCurve> ResolveWeightsForBlendShapeCurve(const UsdUtils::FUsdBlendShape& PrimaryBlendShape, const FRichCurve& ChannelWeightCurve)
 	{
-		TRACE_CPUPROFILER_EVENT_SCOPE( SkelDataConversionImpl::ResolveWeightsForBlendShapeCurve );
+		TRACE_CPUPROFILER_EVENT_SCOPE(SkelDataConversionImpl::ResolveWeightsForBlendShapeCurve);
 
 		int32 NumInbetweens = PrimaryBlendShape.Inbetweens.Num();
-		if ( NumInbetweens == 0 )
+		if (NumInbetweens == 0)
 		{
-			return { ChannelWeightCurve };
+			return {ChannelWeightCurve};
 		}
 
 		TArray<FRichCurve> Result;
-		Result.SetNum( NumInbetweens + 1 ); // One for each inbetween and an additional one for the morph target generated from the primary blend shape
+		Result.SetNum(NumInbetweens + 1);	 // One for each inbetween and an additional one for the morph target generated from the primary blend
+											 // shape
 
 		TArray<float> ResolvedInbetweenWeightsSample;
-		ResolvedInbetweenWeightsSample.SetNum( NumInbetweens );
+		ResolvedInbetweenWeightsSample.SetNum(NumInbetweens);
 
-		for ( const FRichCurveKey& SourceKey : ChannelWeightCurve.Keys )
+		for (const FRichCurveKey& SourceKey : ChannelWeightCurve.Keys)
 		{
 			const float SourceTime = SourceKey.Time;
 			const float SourceValue = SourceKey.Value;
 
 			float ResolvedPrimarySample;
-			UsdUtils::ResolveWeightsForBlendShape( PrimaryBlendShape, SourceValue, ResolvedPrimarySample, ResolvedInbetweenWeightsSample );
+			UsdUtils::ResolveWeightsForBlendShape(PrimaryBlendShape, SourceValue, ResolvedPrimarySample, ResolvedInbetweenWeightsSample);
 
-			FRichCurve& PrimaryCurve = Result[ 0 ];
-			FKeyHandle PrimaryHandle = PrimaryCurve.AddKey( SourceTime, ResolvedPrimarySample );
-			PrimaryCurve.SetKeyInterpMode( PrimaryHandle, SourceKey.InterpMode );
+			FRichCurve& PrimaryCurve = Result[0];
+			FKeyHandle PrimaryHandle = PrimaryCurve.AddKey(SourceTime, ResolvedPrimarySample);
+			PrimaryCurve.SetKeyInterpMode(PrimaryHandle, SourceKey.InterpMode);
 
-			for ( int32 InbetweenIndex = 0; InbetweenIndex < NumInbetweens; ++InbetweenIndex )
+			for (int32 InbetweenIndex = 0; InbetweenIndex < NumInbetweens; ++InbetweenIndex)
 			{
-				FRichCurve& InbetweenCurve = Result[ InbetweenIndex + 1 ];
-				FKeyHandle InbetweenHandle = InbetweenCurve.AddKey( SourceTime, ResolvedInbetweenWeightsSample[ InbetweenIndex ] );
-				InbetweenCurve.SetKeyInterpMode( InbetweenHandle, SourceKey.InterpMode );
+				FRichCurve& InbetweenCurve = Result[InbetweenIndex + 1];
+				FKeyHandle InbetweenHandle = InbetweenCurve.AddKey(SourceTime, ResolvedInbetweenWeightsSample[InbetweenIndex]);
+				InbetweenCurve.SetKeyInterpMode(InbetweenHandle, SourceKey.InterpMode);
 			}
 		}
 
 		return Result;
 	}
-}
+}	 // namespace SkelDataConversionImpl
 
 namespace UsdToUnrealImpl
 {
-	int32 GetPrimValueIndex( const EUsdInterpolationMethod& InterpMethod, const int32 VertexIndex, const int32 VertexInstanceIndex, const int32 PolygonIndex )
+	int32 GetPrimValueIndex(
+		const EUsdInterpolationMethod& InterpMethod,
+		const int32 VertexIndex,
+		const int32 VertexInstanceIndex,
+		const int32 PolygonIndex
+	)
 	{
-		switch ( InterpMethod )
+		switch (InterpMethod)
 		{
-		case EUsdInterpolationMethod::Vertex:
-			return VertexIndex;
-			break;
-		case EUsdInterpolationMethod::FaceVarying:
-			return VertexInstanceIndex;
-			break;
-		case EUsdInterpolationMethod::Uniform:
-			return PolygonIndex;
-			break;
-		case EUsdInterpolationMethod::Constant:
-		default:
-			return 0;
-			break;
+			case EUsdInterpolationMethod::Vertex:
+				return VertexIndex;
+				break;
+			case EUsdInterpolationMethod::FaceVarying:
+				return VertexInstanceIndex;
+				break;
+			case EUsdInterpolationMethod::Uniform:
+				return PolygonIndex;
+				break;
+			case EUsdInterpolationMethod::Constant:
+			default:
+				return 0;
+				break;
 		}
 	}
 
-	void ComputeSourceNormals( SkelDataConversionImpl::FMeshDataBundle& UnmorphedShape )
+	void ComputeSourceNormals(SkelDataConversionImpl::FMeshDataBundle& UnmorphedShape)
 	{
 		TRACE_CPUPROFILER_EVENT_SCOPE(USDSkeletalDataConversion::ComputeSourceNormals);
 
 		IMeshUtilities& MeshUtilities = FModuleManager::Get().LoadModuleChecked<IMeshUtilities>(TEXT("MeshUtilities"));
 
 		// Calculate base normals for the mesh so that we can compute tangent deltas if we need to
-		ETangentOptions::Type TangentOptions = ( ETangentOptions::Type ) ( ETangentOptions::BlendOverlappingNormals | ETangentOptions::UseMikkTSpace );
+		ETangentOptions::Type TangentOptions = (ETangentOptions::Type)(ETangentOptions::BlendOverlappingNormals | ETangentOptions::UseMikkTSpace);
 		TArray<FVector3f> NormalsPerIndex;
-		MeshUtilities.CalculateNormals( UnmorphedShape.Vertices, UnmorphedShape.Indices, UnmorphedShape.UVs, UnmorphedShape.SmoothingGroups, TangentOptions, NormalsPerIndex );	// LWC_TODO: Perf pessimization (ConvertArray)
+		MeshUtilities.CalculateNormals(
+			UnmorphedShape.Vertices,
+			UnmorphedShape.Indices,
+			UnmorphedShape.UVs,
+			UnmorphedShape.SmoothingGroups,
+			TangentOptions,
+			NormalsPerIndex
+		);	  // LWC_TODO: Perf pessimization (ConvertArray)
 
 		// Convert our normals to one normal per vertex, making it faster to unpack the normals we compute in ComputeTangentDeltas
-		// This is possible because we compute them with ETangentOptions::BlendOverlappingNormals, so they are identical for all instances of the vertex
-		UnmorphedShape.NormalsPerVertex.SetNumZeroed( UnmorphedShape.Vertices.Num() );
-		for ( int32 IndexIndex = 0; IndexIndex < UnmorphedShape.Indices.Num(); ++IndexIndex )
+		// This is possible because we compute them with ETangentOptions::BlendOverlappingNormals, so they are identical for all instances of the
+		// vertex
+		UnmorphedShape.NormalsPerVertex.SetNumZeroed(UnmorphedShape.Vertices.Num());
+		for (int32 IndexIndex = 0; IndexIndex < UnmorphedShape.Indices.Num(); ++IndexIndex)
 		{
-			uint32 VertexIndex = UnmorphedShape.Indices[ IndexIndex ];
-			UnmorphedShape.NormalsPerVertex[ VertexIndex ] = NormalsPerIndex[ IndexIndex ];
+			uint32 VertexIndex = UnmorphedShape.Indices[IndexIndex];
+			UnmorphedShape.NormalsPerVertex[VertexIndex] = NormalsPerIndex[IndexIndex];
 		}
 	}
 
-	void CreateMorphTargets( UsdUtils::FBlendShapeMap& BlendShapes, const TArray<FSkeletalMeshImportData>& LODIndexToSkeletalMeshImportData, USkeletalMesh* SkeletalMesh )
+	void CreateMorphTargets(
+		UsdUtils::FBlendShapeMap& BlendShapes,
+		const TArray<FSkeletalMeshImportData>& LODIndexToSkeletalMeshImportData,
+		USkeletalMesh* SkeletalMesh
+	)
 	{
 		TRACE_CPUPROFILER_EVENT_SCOPE(USDSkeletalDataConversion::CreateMorphTargets);
 
 		FSkeletalMeshModel* ImportedResource = SkeletalMesh->GetImportedModel();
-		if ( LODIndexToSkeletalMeshImportData.Num() != ImportedResource->LODModels.Num() )
+		if (LODIndexToSkeletalMeshImportData.Num() != ImportedResource->LODModels.Num())
 		{
 			return;
 		}
@@ -590,10 +646,10 @@ namespace UsdToUnrealImpl
 
 		// Temporarily pack the mesh data into a different format that MeshUtilities can use
 		TArray<SkelDataConversionImpl::FMeshDataBundle> TempMeshBundlesPerLOD;
-		TempMeshBundlesPerLOD.Reserve( NumLODs );
+		TempMeshBundlesPerLOD.Reserve(NumLODs);
 
 		TArray<TMap<int32, TArray<int32>>> OrigIndexToBuiltIndicesPerLOD;
-		OrigIndexToBuiltIndicesPerLOD.Reserve( NumLODs );
+		OrigIndexToBuiltIndicesPerLOD.Reserve(NumLODs);
 
 		TArray<bool> AreNormalsComputedForLODIndex;
 		AreNormalsComputedForLODIndex.SetNumZeroed(NumLODs);
@@ -623,6 +679,7 @@ namespace UsdToUnrealImpl
 			UsdUtils::FUsdBlendShape* BlendShape = nullptr;
 			UMorphTarget* MorphTarget = nullptr;
 		};
+
 		TArray<FMorphTargetJob> MorphTargetJobs;
 		MorphTargetJobs.Reserve(BlendShapes.Num());
 
@@ -712,35 +769,44 @@ namespace UsdToUnrealImpl
 			}
 		);
 
-		if ( bHasValidMorphTarget )
+		if (bHasValidMorphTarget)
 		{
 			SkeletalMesh->MarkPackageDirty();
 			SkeletalMesh->InitMorphTargetsAndRebuildRenderData();
 		}
 	}
-}
+}	 // namespace UsdToUnrealImpl
 
 namespace UnrealToUsdImpl
 {
-	void ConvertSkeletalMeshLOD( const USkeletalMesh* SkeletalMesh, const FSkeletalMeshLODModel& LODModel, pxr::UsdGeomMesh& UsdLODPrimGeomMesh, bool bHasVertexColors, pxr::VtArray< std::string >& MaterialAssignments, const TArray<int32>& LODMaterialMap, const pxr::UsdTimeCode TimeCode, pxr::UsdPrim MaterialPrim )
+	void ConvertSkeletalMeshLOD(
+		const USkeletalMesh* SkeletalMesh,
+		const FSkeletalMeshLODModel& LODModel,
+		pxr::UsdGeomMesh& UsdLODPrimGeomMesh,
+		bool bHasVertexColors,
+		pxr::VtArray<std::string>& MaterialAssignments,
+		const TArray<int32>& LODMaterialMap,
+		const pxr::UsdTimeCode TimeCode,
+		pxr::UsdPrim MaterialPrim
+	)
 	{
 		FScopedUsdAllocs UsdAllocs;
 
 		pxr::UsdPrim MeshPrim = UsdLODPrimGeomMesh.GetPrim();
 		pxr::UsdStageRefPtr Stage = MeshPrim.GetStage();
 
-		pxr::UsdSkelBindingAPI SkelBindingAPI = pxr::UsdSkelBindingAPI::Apply( MeshPrim );
-		if ( !SkelBindingAPI )
+		pxr::UsdSkelBindingAPI SkelBindingAPI = pxr::UsdSkelBindingAPI::Apply(MeshPrim);
+		if (!SkelBindingAPI)
 		{
 			return;
 		}
 
-		if ( !Stage )
+		if (!Stage)
 		{
 			return;
 		}
 
-		const FUsdStageInfo StageInfo( Stage );
+		const FUsdStageInfo StageInfo(Stage);
 
 		// FSkelMeshSection can be "disabled", at which point they don't show up in the engine. We'll skip those
 		// sections, and will use this array to help remap from a source vertex index to the vertex's corresponding
@@ -750,7 +816,7 @@ namespace UnrealToUsdImpl
 
 		// Vertices
 		{
-			if ( LODModel.NumVertices == 0 )
+			if (LODModel.NumVertices == 0)
 			{
 				return;
 			}
@@ -758,44 +824,44 @@ namespace UnrealToUsdImpl
 			// We manually collect vertices here instead of calling LODModel.GetVertices as we need to skip
 			// vertices from disabled sections, which that function won't do
 			TArray<FSoftSkinVertex> Vertices;
-			Vertices.Reserve( LODModel.NumVertices );
-			SourceToPackedVertexIndex.SetNum( LODModel.NumVertices );
-			for ( int32 SectionIndex = 0; SectionIndex < LODModel.Sections.Num(); ++SectionIndex )
+			Vertices.Reserve(LODModel.NumVertices);
+			SourceToPackedVertexIndex.SetNum(LODModel.NumVertices);
+			for (int32 SectionIndex = 0; SectionIndex < LODModel.Sections.Num(); ++SectionIndex)
 			{
-				const FSkelMeshSection& Section = LODModel.Sections[ SectionIndex ];
-				if ( Section.bDisabled )
+				const FSkelMeshSection& Section = LODModel.Sections[SectionIndex];
+				if (Section.bDisabled)
 				{
 					// Mark that these indices are skipped
-					for ( int32 Index = 0; Index < Section.NumVertices; ++Index )
+					for (int32 Index = 0; Index < Section.NumVertices; ++Index)
 					{
-						SourceToPackedVertexIndex[ Section.BaseVertexIndex + Index ] = INDEX_NONE;
+						SourceToPackedVertexIndex[Section.BaseVertexIndex + Index] = INDEX_NONE;
 					}
 					continue;
 				}
 
-				for ( int32 Index = 0; Index < Section.NumVertices; ++Index )
+				for (int32 Index = 0; Index < Section.NumVertices; ++Index)
 				{
-					SourceToPackedVertexIndex[ Section.BaseVertexIndex + Index ] = PackedVertexIndex++;
+					SourceToPackedVertexIndex[Section.BaseVertexIndex + Index] = PackedVertexIndex++;
 				}
 
-				Vertices.Append( Section.SoftVertices );
+				Vertices.Append(Section.SoftVertices);
 			}
 			const int32 VertexCount = Vertices.Num();
 
 			// Points
 			{
 				pxr::UsdAttribute Points = UsdLODPrimGeomMesh.CreatePointsAttr();
-				if ( Points )
+				if (Points)
 				{
-					pxr::VtArray< pxr::GfVec3f > PointsArray;
-					PointsArray.reserve( VertexCount );
+					pxr::VtArray<pxr::GfVec3f> PointsArray;
+					PointsArray.reserve(VertexCount);
 
-					for ( int32 VertexIndex = 0; VertexIndex < VertexCount; ++VertexIndex )
+					for (int32 VertexIndex = 0; VertexIndex < VertexCount; ++VertexIndex)
 					{
-						PointsArray.push_back( UnrealToUsd::ConvertVectorFloat( StageInfo, (FVector)Vertices[ VertexIndex ].Position ) );
+						PointsArray.push_back(UnrealToUsd::ConvertVectorFloat(StageInfo, (FVector)Vertices[VertexIndex].Position));
 					}
 
-					Points.Set( PointsArray, TimeCode );
+					Points.Set(PointsArray, TimeCode);
 				}
 			}
 
@@ -805,76 +871,80 @@ namespace UnrealToUsdImpl
 				// actually subdivide (like usdview) will just discard authored normals and fully recompute them
 				// on-demand in case they have a valid subdivision scheme (which is the default state).
 				// Reference: https://graphics.pixar.com/usd/release/api/class_usd_geom_mesh.html#UsdGeom_Mesh_Normals
-				if ( pxr::UsdAttribute SubdivisionAttr = UsdLODPrimGeomMesh.CreateSubdivisionSchemeAttr() )
+				if (pxr::UsdAttribute SubdivisionAttr = UsdLODPrimGeomMesh.CreateSubdivisionSchemeAttr())
 				{
-					ensure( SubdivisionAttr.Set( pxr::UsdGeomTokens->none ) );
+					ensure(SubdivisionAttr.Set(pxr::UsdGeomTokens->none));
 				}
 
 				pxr::UsdAttribute NormalsAttribute = UsdLODPrimGeomMesh.CreateNormalsAttr();
-				if ( NormalsAttribute )
+				if (NormalsAttribute)
 				{
-					pxr::VtArray< pxr::GfVec3f > Normals;
-					Normals.reserve( VertexCount );
+					pxr::VtArray<pxr::GfVec3f> Normals;
+					Normals.reserve(VertexCount);
 
-					for ( int32 VertexIndex = 0; VertexIndex < VertexCount; ++VertexIndex )
+					for (int32 VertexIndex = 0; VertexIndex < VertexCount; ++VertexIndex)
 					{
-						Normals.push_back( UnrealToUsd::ConvertVectorFloat( StageInfo, FVector4(Vertices[ VertexIndex ].TangentZ) ) );
+						Normals.push_back(UnrealToUsd::ConvertVectorFloat(StageInfo, FVector4(Vertices[VertexIndex].TangentZ)));
 					}
 
-					NormalsAttribute.Set( Normals, TimeCode );
+					NormalsAttribute.Set(Normals, TimeCode);
 				}
 			}
 
 			// UVs
 			{
-				for ( uint32 TexCoordSourceIndex = 0; TexCoordSourceIndex < LODModel.NumTexCoords; ++TexCoordSourceIndex )
+				for (uint32 TexCoordSourceIndex = 0; TexCoordSourceIndex < LODModel.NumTexCoords; ++TexCoordSourceIndex)
 				{
-					pxr::TfToken UsdUVSetName = UsdUtils::GetUVSetName( TexCoordSourceIndex ).Get();
+					pxr::TfToken UsdUVSetName = UsdUtils::GetUVSetName(TexCoordSourceIndex).Get();
 
-					pxr::UsdGeomPrimvar PrimvarST = pxr::UsdGeomPrimvarsAPI(MeshPrim).CreatePrimvar( UsdUVSetName, pxr::SdfValueTypeNames->TexCoord2fArray, pxr::UsdGeomTokens->vertex );
+					pxr::UsdGeomPrimvar PrimvarST = pxr::UsdGeomPrimvarsAPI(MeshPrim).CreatePrimvar(
+						UsdUVSetName,
+						pxr::SdfValueTypeNames->TexCoord2fArray,
+						pxr::UsdGeomTokens->vertex
+					);
 
-					if ( PrimvarST )
+					if (PrimvarST)
 					{
 						pxr::VtVec2fArray UVs;
 
-						for ( int32 VertexIndex = 0; VertexIndex < VertexCount; ++VertexIndex )
+						for (int32 VertexIndex = 0; VertexIndex < VertexCount; ++VertexIndex)
 						{
-							FVector2D TexCoord = FVector2D(Vertices[ VertexIndex ].UVs[ TexCoordSourceIndex ]);
-							TexCoord[ 1 ] = 1.f - TexCoord[ 1 ];
+							FVector2D TexCoord = FVector2D(Vertices[VertexIndex].UVs[TexCoordSourceIndex]);
+							TexCoord[1] = 1.f - TexCoord[1];
 
-							UVs.push_back( UnrealToUsd::ConvertVectorFloat( TexCoord ) );
+							UVs.push_back(UnrealToUsd::ConvertVectorFloat(TexCoord));
 						}
 
-						PrimvarST.Set( UVs, TimeCode );
+						PrimvarST.Set(UVs, TimeCode);
 					}
 				}
 			}
 
 			// Vertex colors
-			if ( bHasVertexColors )
+			if (bHasVertexColors)
 			{
-				pxr::UsdGeomPrimvar DisplayColorPrimvar = UsdLODPrimGeomMesh.CreateDisplayColorPrimvar( pxr::UsdGeomTokens->vertex );
-				pxr::UsdGeomPrimvar DisplayOpacityPrimvar = UsdLODPrimGeomMesh.CreateDisplayOpacityPrimvar( pxr::UsdGeomTokens->vertex );
+				pxr::UsdGeomPrimvar DisplayColorPrimvar = UsdLODPrimGeomMesh.CreateDisplayColorPrimvar(pxr::UsdGeomTokens->vertex);
+				pxr::UsdGeomPrimvar DisplayOpacityPrimvar = UsdLODPrimGeomMesh.CreateDisplayOpacityPrimvar(pxr::UsdGeomTokens->vertex);
 
-				if ( DisplayColorPrimvar && DisplayOpacityPrimvar )
+				if (DisplayColorPrimvar && DisplayOpacityPrimvar)
 				{
-					pxr::VtArray< pxr::GfVec3f > DisplayColors;
-					DisplayColors.reserve( VertexCount );
+					pxr::VtArray<pxr::GfVec3f> DisplayColors;
+					DisplayColors.reserve(VertexCount);
 
-					pxr::VtArray< float > DisplayOpacities;
-					DisplayOpacities.reserve( VertexCount );
+					pxr::VtArray<float> DisplayOpacities;
+					DisplayOpacities.reserve(VertexCount);
 
-					for ( int32 VertexIndex = 0; VertexIndex < VertexCount; ++VertexIndex )
+					for (int32 VertexIndex = 0; VertexIndex < VertexCount; ++VertexIndex)
 					{
-						const FColor& VertexColor = Vertices[ VertexIndex ].Color;
+						const FColor& VertexColor = Vertices[VertexIndex].Color;
 
-						pxr::GfVec4f Color = UnrealToUsd::ConvertColor( VertexColor );
-						DisplayColors.push_back( pxr::GfVec3f( Color[ 0 ], Color[ 1 ], Color[ 2 ] ) );
-						DisplayOpacities.push_back( Color[ 3 ] );
+						pxr::GfVec4f Color = UnrealToUsd::ConvertColor(VertexColor);
+						DisplayColors.push_back(pxr::GfVec3f(Color[0], Color[1], Color[2]));
+						DisplayOpacities.push_back(Color[3]);
 					}
 
-					DisplayColorPrimvar.Set( DisplayColors, TimeCode );
-					DisplayOpacityPrimvar.Set( DisplayOpacities, TimeCode );
+					DisplayColorPrimvar.Set(DisplayColors, TimeCode);
+					DisplayOpacityPrimvar.Set(DisplayOpacities, TimeCode);
 				}
 			}
 
@@ -883,38 +953,38 @@ namespace UnrealToUsdImpl
 				const int32 NumInfluencesPerVertex = LODModel.GetMaxBoneInfluences();
 
 				const bool bConstantPrimvar = false;
-				pxr::UsdGeomPrimvar JointIndicesPrimvar = SkelBindingAPI.CreateJointIndicesPrimvar( bConstantPrimvar, NumInfluencesPerVertex );
-				pxr::UsdGeomPrimvar JointWeightsPrimvar = SkelBindingAPI.CreateJointWeightsPrimvar( bConstantPrimvar, NumInfluencesPerVertex );
+				pxr::UsdGeomPrimvar JointIndicesPrimvar = SkelBindingAPI.CreateJointIndicesPrimvar(bConstantPrimvar, NumInfluencesPerVertex);
+				pxr::UsdGeomPrimvar JointWeightsPrimvar = SkelBindingAPI.CreateJointWeightsPrimvar(bConstantPrimvar, NumInfluencesPerVertex);
 
-				if ( JointIndicesPrimvar && JointWeightsPrimvar )
+				if (JointIndicesPrimvar && JointWeightsPrimvar)
 				{
-					pxr::VtArray< int > JointIndices;
-					JointIndices.reserve( VertexCount * NumInfluencesPerVertex );
+					pxr::VtArray<int> JointIndices;
+					JointIndices.reserve(VertexCount * NumInfluencesPerVertex);
 
-					pxr::VtArray< float > JointWeights;
-					JointWeights.reserve( VertexCount * NumInfluencesPerVertex );
+					pxr::VtArray<float> JointWeights;
+					JointWeights.reserve(VertexCount * NumInfluencesPerVertex);
 
-					for ( const FSkelMeshSection& Section : LODModel.Sections )
+					for (const FSkelMeshSection& Section : LODModel.Sections)
 					{
-						if ( Section.bDisabled )
+						if (Section.bDisabled)
 						{
 							continue;
 						}
 
-						for ( const FSoftSkinVertex& Vertex : Section.SoftVertices )
+						for (const FSoftSkinVertex& Vertex : Section.SoftVertices)
 						{
-							for ( int32 InfluenceIndex = 0; InfluenceIndex < NumInfluencesPerVertex; ++InfluenceIndex )
+							for (int32 InfluenceIndex = 0; InfluenceIndex < NumInfluencesPerVertex; ++InfluenceIndex)
 							{
-								int32 BoneIndex = Section.BoneMap[ Vertex.InfluenceBones[ InfluenceIndex ] ];
+								int32 BoneIndex = Section.BoneMap[Vertex.InfluenceBones[InfluenceIndex]];
 
-								JointIndices.push_back( BoneIndex );
-								JointWeights.push_back( Vertex.InfluenceWeights[ InfluenceIndex ] / UE::AnimationCore::MaxRawBoneWeightFloat );
+								JointIndices.push_back(BoneIndex);
+								JointWeights.push_back(Vertex.InfluenceWeights[InfluenceIndex] / UE::AnimationCore::MaxRawBoneWeightFloat);
 							}
 						}
 					}
 
-					JointIndicesPrimvar.Set( JointIndices, TimeCode );
-					JointWeightsPrimvar.Set( JointWeights, TimeCode );
+					JointIndicesPrimvar.Set(JointIndices, TimeCode);
+					JointWeightsPrimvar.Set(JointWeights, TimeCode);
 				}
 			}
 		}
@@ -925,9 +995,9 @@ namespace UnrealToUsdImpl
 
 			// Face Vertex Counts
 			{
-				for ( const FSkelMeshSection& Section : LODModel.Sections )
+				for (const FSkelMeshSection& Section : LODModel.Sections)
 				{
-					if ( Section.bDisabled )
+					if (Section.bDisabled)
 					{
 						continue;
 					}
@@ -936,17 +1006,17 @@ namespace UnrealToUsdImpl
 				}
 
 				pxr::UsdAttribute FaceCountsAttribute = UsdLODPrimGeomMesh.CreateFaceVertexCountsAttr();
-				if ( FaceCountsAttribute )
+				if (FaceCountsAttribute)
 				{
-					pxr::VtArray< int > FaceVertexCounts;
-					FaceVertexCounts.reserve( TotalNumTriangles );
+					pxr::VtArray<int> FaceVertexCounts;
+					FaceVertexCounts.reserve(TotalNumTriangles);
 
-					for ( int32 FaceIndex = 0; FaceIndex < TotalNumTriangles; ++FaceIndex )
+					for (int32 FaceIndex = 0; FaceIndex < TotalNumTriangles; ++FaceIndex)
 					{
-						FaceVertexCounts.push_back( 3 );
+						FaceVertexCounts.push_back(3);
 					}
 
-					FaceCountsAttribute.Set( FaceVertexCounts, TimeCode );
+					FaceCountsAttribute.Set(FaceVertexCounts, TimeCode);
 				}
 			}
 
@@ -954,36 +1024,35 @@ namespace UnrealToUsdImpl
 			{
 				pxr::UsdAttribute FaceVertexIndicesAttribute = UsdLODPrimGeomMesh.GetFaceVertexIndicesAttr();
 
-				if ( FaceVertexIndicesAttribute )
+				if (FaceVertexIndicesAttribute)
 				{
-					pxr::VtArray< int > FaceVertexIndices;
-					FaceVertexIndices.reserve( TotalNumTriangles * 3 );
+					pxr::VtArray<int> FaceVertexIndices;
+					FaceVertexIndices.reserve(TotalNumTriangles * 3);
 
-					for ( const FSkelMeshSection& Section : LODModel.Sections )
+					for (const FSkelMeshSection& Section : LODModel.Sections)
 					{
-						if ( Section.bDisabled )
+						if (Section.bDisabled)
 						{
 							continue;
 						}
 
 						int32 TriangleCount = Section.NumTriangles;
-						for ( uint32 TriangleIndex = 0; TriangleIndex < Section.NumTriangles; ++TriangleIndex )
+						for (uint32 TriangleIndex = 0; TriangleIndex < Section.NumTriangles; ++TriangleIndex)
 						{
-							for ( uint32 PointIndex = 0; PointIndex < 3; PointIndex++ )
+							for (uint32 PointIndex = 0; PointIndex < 3; PointIndex++)
 							{
-								const int32 SourceVertexIndex =
-									LODModel.IndexBuffer[ Section.BaseIndex + ( ( TriangleIndex * 3 ) + PointIndex ) ];
+								const int32 SourceVertexIndex = LODModel.IndexBuffer[Section.BaseIndex + ((TriangleIndex * 3) + PointIndex)];
 
-								int32 PackedVertexPositionIndex = SourceToPackedVertexIndex[ SourceVertexIndex ];
+								int32 PackedVertexPositionIndex = SourceToPackedVertexIndex[SourceVertexIndex];
 
-								check( PackedVertexPositionIndex >= 0 && PackedVertexPositionIndex != INDEX_NONE );
+								check(PackedVertexPositionIndex >= 0 && PackedVertexPositionIndex != INDEX_NONE);
 
-								FaceVertexIndices.push_back( PackedVertexPositionIndex );
+								FaceVertexIndices.push_back(PackedVertexPositionIndex);
 							}
 						}
 					}
 
-					FaceVertexIndicesAttribute.Set( FaceVertexIndices, TimeCode );
+					FaceVertexIndicesAttribute.Set(FaceVertexIndices, TimeCode);
 				}
 			}
 		}
@@ -992,115 +1061,121 @@ namespace UnrealToUsdImpl
 		{
 			bool bHasUEMaterialAssignements = false;
 
-			pxr::VtArray< std::string > UnrealMaterialsForLOD;
-			for ( const FSkelMeshSection& Section : LODModel.Sections )
+			pxr::VtArray<std::string> UnrealMaterialsForLOD;
+			for (const FSkelMeshSection& Section : LODModel.Sections)
 			{
 				int32 SkeletalMaterialIndex = INDEX_NONE;
-				if ( LODMaterialMap.IsValidIndex( Section.MaterialIndex ) )
+				if (LODMaterialMap.IsValidIndex(Section.MaterialIndex))
 				{
-					SkeletalMaterialIndex = LODMaterialMap[ Section.MaterialIndex ];
+					SkeletalMaterialIndex = LODMaterialMap[Section.MaterialIndex];
 				}
 				// Note that the LODMaterialMap can contain INDEX_NONE to signify no remapping
-				if ( SkeletalMaterialIndex == INDEX_NONE )
+				if (SkeletalMaterialIndex == INDEX_NONE)
 				{
 					SkeletalMaterialIndex = Section.MaterialIndex;
 				}
 
-				if ( SkeletalMaterialIndex >= 0 && SkeletalMaterialIndex < MaterialAssignments.size() )
+				if (SkeletalMaterialIndex >= 0 && SkeletalMaterialIndex < MaterialAssignments.size())
 				{
-					UnrealMaterialsForLOD.push_back( MaterialAssignments[ SkeletalMaterialIndex ] );
+					UnrealMaterialsForLOD.push_back(MaterialAssignments[SkeletalMaterialIndex]);
 					bHasUEMaterialAssignements = true;
 				}
 				else
 				{
 					// Keep unrealMaterials with the same number of elements as our MaterialIndices expect
-					UnrealMaterialsForLOD.push_back( "" );
+					UnrealMaterialsForLOD.push_back("");
 				}
 			}
 
 			// This LOD has a single material assignment, just add an unrealMaterials attribute to the mesh prim
-			if ( bHasUEMaterialAssignements && UnrealMaterialsForLOD.size() == 1 )
+			if (bHasUEMaterialAssignements && UnrealMaterialsForLOD.size() == 1)
 			{
-				if ( pxr::UsdAttribute UEMaterialsAttribute = MaterialPrim.CreateAttribute( UnrealIdentifiers::MaterialAssignment, pxr::SdfValueTypeNames->String ) )
+				if (pxr::UsdAttribute UEMaterialsAttribute = MaterialPrim.CreateAttribute(
+						UnrealIdentifiers::MaterialAssignment,
+						pxr::SdfValueTypeNames->String
+					))
 				{
-					UEMaterialsAttribute.Set( UnrealMaterialsForLOD[0] );
+					UEMaterialsAttribute.Set(UnrealMaterialsForLOD[0]);
 				}
 			}
 			// Multiple material assignments to the same LOD (and so the same mesh prim). Need to create a GeomSubset for each UE mesh section
-			else if ( UnrealMaterialsForLOD.size() > 1 )
+			else if (UnrealMaterialsForLOD.size() > 1)
 			{
 				// Need to fetch all triangles of a section, and add their indices to the GeomSubset
-				for ( int32 SectionIndex = 0; SectionIndex < LODModel.Sections.Num(); ++SectionIndex )
+				for (int32 SectionIndex = 0; SectionIndex < LODModel.Sections.Num(); ++SectionIndex)
 				{
-					const FSkelMeshSection& Section = LODModel.Sections[ SectionIndex ];
-					if ( Section.bDisabled )
+					const FSkelMeshSection& Section = LODModel.Sections[SectionIndex];
+					if (Section.bDisabled)
 					{
 						continue;
 					}
 
 					// Note that we will continue on even if we have no material assignment, so as to satisfy the "partition" family condition
-					std::string SectionMaterial = UnrealMaterialsForLOD[ SectionIndex ];
+					std::string SectionMaterial = UnrealMaterialsForLOD[SectionIndex];
 
 					pxr::UsdPrim GeomSubsetPrim = Stage->DefinePrim(
-						UsdLODPrimGeomMesh.GetPath().AppendPath( pxr::SdfPath( "Section" + std::to_string( SectionIndex ) ) ),
-						UnrealToUsd::ConvertToken( TEXT( "GeomSubset" ) ).Get()
+						UsdLODPrimGeomMesh.GetPath().AppendPath(pxr::SdfPath("Section" + std::to_string(SectionIndex))),
+						UnrealToUsd::ConvertToken(TEXT("GeomSubset")).Get()
 					);
 
 					pxr::UsdPrim MaterialGeomSubsetPrim = GeomSubsetPrim;
-					if ( MaterialPrim.GetStage() != MeshPrim.GetStage() )
+					if (MaterialPrim.GetStage() != MeshPrim.GetStage())
 					{
 						MaterialGeomSubsetPrim = MaterialPrim.GetStage()->OverridePrim(
-							MaterialPrim.GetPath().AppendPath( pxr::SdfPath( "Section" + std::to_string( SectionIndex ) ) )
+							MaterialPrim.GetPath().AppendPath(pxr::SdfPath("Section" + std::to_string(SectionIndex)))
 						);
 					}
 
-					pxr::UsdGeomSubset GeomSubsetSchema{ GeomSubsetPrim };
+					pxr::UsdGeomSubset GeomSubsetSchema{GeomSubsetPrim};
 
 					// Element type attribute
 					pxr::UsdAttribute ElementTypeAttr = GeomSubsetSchema.CreateElementTypeAttr();
-					ElementTypeAttr.Set( pxr::UsdGeomTokens->face, TimeCode );
+					ElementTypeAttr.Set(pxr::UsdGeomTokens->face, TimeCode);
 
 					// Indices attribute
 					{
 						const uint32 TriangleCount = Section.NumTriangles;
-						uint32 FirstTriangleIndex = Section.BaseIndex / 3; // BaseIndex is the first *vertex* instance index
+						uint32 FirstTriangleIndex = Section.BaseIndex / 3;	  // BaseIndex is the first *vertex* instance index
 						pxr::VtArray<int> IndicesAttrValue;
 
 						// We may have some disabled sections (that wouldn't have emitted triangles). If so, we need to
 						// adjust our FirstTriangleIndex.
 						// This could be optimized in case vertex instances show up in LODModel.IndexBuffer according
 						// to the section order, but so far we haven't found such guarantee, so just check them all.
-						for ( int32 OtherSectionIndex = 0; OtherSectionIndex < LODModel.Sections.Num(); ++OtherSectionIndex )
+						for (int32 OtherSectionIndex = 0; OtherSectionIndex < LODModel.Sections.Num(); ++OtherSectionIndex)
 						{
-							const FSkelMeshSection& OtherSection = LODModel.Sections[ OtherSectionIndex ];
-							if ( OtherSection.bDisabled && OtherSection.BaseIndex < Section.BaseIndex )
+							const FSkelMeshSection& OtherSection = LODModel.Sections[OtherSectionIndex];
+							if (OtherSection.bDisabled && OtherSection.BaseIndex < Section.BaseIndex)
 							{
 								FirstTriangleIndex -= OtherSection.NumTriangles;
 							}
 						}
 
-						for ( uint32 TriangleIndex = FirstTriangleIndex; TriangleIndex - FirstTriangleIndex < TriangleCount; ++TriangleIndex )
+						for (uint32 TriangleIndex = FirstTriangleIndex; TriangleIndex - FirstTriangleIndex < TriangleCount; ++TriangleIndex)
 						{
 							// Note that we add VertexInstances in sequence to the usda file for the faceVertexInstances attribute, which
 							// also constitutes our triangle order
-							IndicesAttrValue.push_back( static_cast< int >( TriangleIndex ));
+							IndicesAttrValue.push_back(static_cast<int>(TriangleIndex));
 						}
 
 						pxr::UsdAttribute IndicesAttr = GeomSubsetSchema.CreateIndicesAttr();
-						IndicesAttr.Set( IndicesAttrValue, TimeCode );
+						IndicesAttr.Set(IndicesAttrValue, TimeCode);
 					}
 
 					// Family name attribute
 					pxr::UsdAttribute FamilyNameAttr = GeomSubsetSchema.CreateFamilyNameAttr();
-					FamilyNameAttr.Set( pxr::UsdShadeTokens->materialBind, TimeCode );
+					FamilyNameAttr.Set(pxr::UsdShadeTokens->materialBind, TimeCode);
 
 					// Family type
-					pxr::UsdGeomSubset::SetFamilyType( UsdLODPrimGeomMesh, pxr::UsdShadeTokens->materialBind, pxr::UsdGeomTokens->partition );
+					pxr::UsdGeomSubset::SetFamilyType(UsdLODPrimGeomMesh, pxr::UsdShadeTokens->materialBind, pxr::UsdGeomTokens->partition);
 
 					// unrealMaterials attribute
-					if ( pxr::UsdAttribute UEMaterialsAttribute = MaterialGeomSubsetPrim.CreateAttribute( UnrealIdentifiers::MaterialAssignment, pxr::SdfValueTypeNames->String ) )
+					if (pxr::UsdAttribute UEMaterialsAttribute = MaterialGeomSubsetPrim.CreateAttribute(
+							UnrealIdentifiers::MaterialAssignment,
+							pxr::SdfValueTypeNames->String
+						))
 					{
-						UEMaterialsAttribute.Set( UnrealMaterialsForLOD[ SectionIndex ] );
+						UEMaterialsAttribute.Set(UnrealMaterialsForLOD[SectionIndex]);
 					}
 				}
 			}
@@ -1108,31 +1183,31 @@ namespace UnrealToUsdImpl
 	}
 
 	// Converts UE morph target deltas from DeltaArray into offsets, pointIndices and normalOffsets attributes of BlendShape
-	bool ConvertMorphTargetDeltas( const FMorphTargetDelta* DeltaArray, int32 NumDeltas, pxr::UsdSkelBlendShape& BlendShape, pxr::UsdTimeCode TimeCode )
+	bool ConvertMorphTargetDeltas(const FMorphTargetDelta* DeltaArray, int32 NumDeltas, pxr::UsdSkelBlendShape& BlendShape, pxr::UsdTimeCode TimeCode)
 	{
-		if ( !DeltaArray || NumDeltas == 0 || !BlendShape )
+		if (!DeltaArray || NumDeltas == 0 || !BlendShape)
 		{
 			return false;
 		}
 
-		FUsdStageInfo StageInfo{ BlendShape.GetPrim().GetStage() };
+		FUsdStageInfo StageInfo{BlendShape.GetPrim().GetStage()};
 		FScopedUsdAllocs Allocs;
 
-		pxr::VtArray< pxr::GfVec3f > Offsets;
-		pxr::VtArray< int > PointIndices;
-		pxr::VtArray< pxr::GfVec3f > Normals;
+		pxr::VtArray<pxr::GfVec3f> Offsets;
+		pxr::VtArray<int> PointIndices;
+		pxr::VtArray<pxr::GfVec3f> Normals;
 
 		Offsets.reserve(NumDeltas);
 		PointIndices.reserve(NumDeltas);
 		Normals.reserve(NumDeltas);
 
-		for ( int32 DeltaIndex = 0; DeltaIndex < NumDeltas; ++DeltaIndex )
+		for (int32 DeltaIndex = 0; DeltaIndex < NumDeltas; ++DeltaIndex)
 		{
-			const FMorphTargetDelta& Delta = DeltaArray[ DeltaIndex ];
+			const FMorphTargetDelta& Delta = DeltaArray[DeltaIndex];
 
-			Offsets.push_back( UnrealToUsd::ConvertVectorFloat( StageInfo, (FVector)Delta.PositionDelta ) );
-			PointIndices.push_back( Delta.SourceIdx );
-			Normals.push_back( UnrealToUsd::ConvertVectorFloat( StageInfo, (FVector)Delta.TangentZDelta ) );
+			Offsets.push_back(UnrealToUsd::ConvertVectorFloat(StageInfo, (FVector)Delta.PositionDelta));
+			PointIndices.push_back(Delta.SourceIdx);
+			Normals.push_back(UnrealToUsd::ConvertVectorFloat(StageInfo, (FVector)Delta.TangentZDelta));
 		}
 
 		BlendShape.CreateOffsetsAttr().Set(Offsets, TimeCode);
@@ -1144,34 +1219,34 @@ namespace UnrealToUsdImpl
 
 	// BoneNamesInOrder represents a hierarchy of bones. OutFullPaths will be the full path to each bone, in the same order
 	// e.g. 'Root/Arm/Foot'
-	void CreateFullBonePaths(const TArray<FMeshBoneInfo>& BoneNamesInOrder, TArray<FString>& OutFullPaths )
+	void CreateFullBonePaths(const TArray<FMeshBoneInfo>& BoneNamesInOrder, TArray<FString>& OutFullPaths)
 	{
 		int32 NumBones = BoneNamesInOrder.Num();
-		if ( NumBones < 1 )
+		if (NumBones < 1)
 		{
 			return;
 		}
 
-		OutFullPaths.SetNum( NumBones );
+		OutFullPaths.SetNum(NumBones);
 
 		// The first bone is the root, and has ParentIndex == -1, so do it separately here to void checking the indices for all bones
 		// Sanitize because ExportName can have spaces, which USD doesn't like
 		OutFullPaths[0] = IUsdClassesModule::SanitizeObjectName(BoneNamesInOrder[0].ExportName);
 
 		// Bones are always stored in an increasing order, so we can do all paths in a single pass
-		for ( int32 BoneIndex = 1; BoneIndex < NumBones; ++BoneIndex )
+		for (int32 BoneIndex = 1; BoneIndex < NumBones; ++BoneIndex)
 		{
-			const FMeshBoneInfo& BoneInfo = BoneNamesInOrder[ BoneIndex ];
+			const FMeshBoneInfo& BoneInfo = BoneNamesInOrder[BoneIndex];
 			FString SanitizedBoneName = IUsdClassesModule::SanitizeObjectName(BoneInfo.ExportName);
 
-			OutFullPaths[BoneIndex] = FString::Printf(TEXT("%s/%s"), *OutFullPaths[ BoneInfo.ParentIndex ], *SanitizedBoneName );
+			OutFullPaths[BoneIndex] = FString::Printf(TEXT("%s/%s"), *OutFullPaths[BoneInfo.ParentIndex], *SanitizedBoneName);
 		}
 	}
-}
+}	 // namespace UnrealToUsdImpl
 
 bool UsdToUnreal::ConvertSkeleton(const pxr::UsdSkelSkeletonQuery& SkeletonQuery, FSkeletalMeshImportData& SkelMeshImportData)
 {
-	TRACE_CPUPROFILER_EVENT_SCOPE( UsdToUnreal::ConvertSkeleton );
+	TRACE_CPUPROFILER_EVENT_SCOPE(UsdToUnreal::ConvertSkeleton);
 
 	using namespace pxr;
 
@@ -1201,20 +1276,24 @@ bool UsdToUnreal::ConvertSkeleton(const pxr::UsdSkelSkeletonQuery& SkeletonQuery
 	}
 
 	// Skeleton has no joints: Generate a dummy single "Root" bone skeleton
-	if ( JointNames.Num() == 0 )
+	if (JointNames.Num() == 0)
 	{
-		FString SkeletonPrimPath = UsdToUnreal::ConvertPath( SkeletonQuery.GetPrim().GetPath() );
+		FString SkeletonPrimPath = UsdToUnreal::ConvertPath(SkeletonQuery.GetPrim().GetPath());
 
 		FUsdLogManager::LogMessage(
 			EMessageSeverity::Warning,
-			FText::Format( LOCTEXT( "NoBonesInSkeleton", "Skeleton prim '{0}' has no joints! "
-				"A new skeleton with a single 'Root' bone will be generated as USkeletalMeshes require valid skeletons. "
-				"Note that this new skeleton may be written back to the USD stage when exporting the corresponding asset." ),
-				FText::FromString( SkeletonPrimPath )
+			FText::Format(
+				LOCTEXT(
+					"NoBonesInSkeleton",
+					"Skeleton prim '{0}' has no joints! "
+					"A new skeleton with a single 'Root' bone will be generated as USkeletalMeshes require valid skeletons. "
+					"Note that this new skeleton may be written back to the USD stage when exporting the corresponding asset."
+				),
+				FText::FromString(SkeletonPrimPath)
 			)
 		);
 
-		SkelMeshImportData.RefBonesBinary.AddZeroed( 1 );
+		SkelMeshImportData.RefBonesBinary.AddZeroed(1);
 
 		SkeletalMeshImportData::FBone& Bone = SkelMeshImportData.RefBonesBinary.Last();
 		Bone.Name = TEXT("Root");
@@ -1228,7 +1307,7 @@ bool UsdToUnreal::ConvertSkeleton(const pxr::UsdSkelSkeletonQuery& SkeletonQuery
 		return true;
 	}
 
-	if ( JointNames.Num() > MAX_BONES )
+	if (JointNames.Num() > MAX_BONES)
 	{
 		return false;
 	}
@@ -1240,8 +1319,8 @@ bool UsdToUnreal::ConvertSkeleton(const pxr::UsdSkelSkeletonQuery& SkeletonQuery
 	// assumes the pose described by them when non-animated. See the gigantic comment inside ConvertSkinnedMesh.
 	VtArray<GfMatrix4d> JointLocalRestTransforms;
 	const bool bAtRest = true;
-	bool bJointTransformsComputed = SkeletonQuery.ComputeJointLocalTransforms( &JointLocalRestTransforms, UsdTimeCode::EarliestTime(), bAtRest );
-	if ( bJointTransformsComputed )
+	bool bJointTransformsComputed = SkeletonQuery.ComputeJointLocalTransforms(&JointLocalRestTransforms, UsdTimeCode::EarliestTime(), bAtRest);
+	if (bJointTransformsComputed)
 	{
 		UsdStageWeakPtr Stage = SkeletonQuery.GetSkeleton().GetPrim().GetStage();
 		const FUsdStageInfo StageInfo(Stage);
@@ -1249,7 +1328,7 @@ bool UsdToUnreal::ConvertSkeleton(const pxr::UsdSkelSkeletonQuery& SkeletonQuery
 		for (uint32 Index = 0; Index < JointLocalRestTransforms.size(); ++Index)
 		{
 			const GfMatrix4d& UsdMatrix = JointLocalRestTransforms[Index];
-			FTransform BoneTransform = UsdToUnreal::ConvertMatrix( StageInfo, UsdMatrix );
+			FTransform BoneTransform = UsdToUnreal::ConvertMatrix(StageInfo, UsdMatrix);
 			BoneTransforms.Add(BoneTransform);
 		}
 	}
@@ -1289,10 +1368,10 @@ bool UsdToUnreal::ConvertSkeleton(const pxr::UsdSkelSkeletonQuery& SkeletonQuery
 	}
 
 	// Store the retrieved data as bones into the SkeletalMeshImportData
-	SkelMeshImportData.RefBonesBinary.AddZeroed( JointNames.Num() );
+	SkelMeshImportData.RefBonesBinary.AddZeroed(JointNames.Num());
 	for (int32 Index = 0; Index < JointNames.Num(); ++Index)
 	{
-		SkeletalMeshImportData::FBone& Bone = SkelMeshImportData.RefBonesBinary[ Index ];
+		SkeletalMeshImportData::FBone& Bone = SkelMeshImportData.RefBonesBinary[Index];
 
 		Bone.Name = JointNames[Index];
 		Bone.ParentIndex = ParentJointIndices[Index];
@@ -1348,7 +1427,7 @@ bool UsdToUnreal::ConvertSkinnedMesh(
 	const FUsdMeshConversionOptions& CommonOptions
 )
 {
-	TRACE_CPUPROFILER_EVENT_SCOPE( UsdToUnreal::ConvertSkinnedMesh );
+	TRACE_CPUPROFILER_EVENT_SCOPE(UsdToUnreal::ConvertSkinnedMesh);
 
 	using namespace pxr;
 
@@ -1361,12 +1440,12 @@ bool UsdToUnreal::ConvertSkinnedMesh(
 
 	// Ref. FFbxImporter::FillSkelMeshImporterFromFbx
 	UsdGeomMesh UsdMesh = UsdGeomMesh(SkinningPrim);
-	if ( !UsdMesh )
+	if (!UsdMesh)
 	{
 		return false;
 	}
 
-	const FUsdStageInfo StageInfo( SkinningPrim.GetStage() );
+	const FUsdStageInfo StageInfo(SkinningPrim.GetStage());
 
 	// Retrieve the mesh points (vertices) from USD and append it to the SkeletalMeshImportData Points
 	uint32 NumPoints = 0;
@@ -1417,23 +1496,23 @@ bool UsdToUnreal::ConvertSkinnedMesh(
 		{
 			VtArray<GfMatrix4d> JointLocalRestTransforms;
 			const bool bAtRest = true;
-			bSuccess &= SkeletonQuery.ComputeJointLocalTransforms( &JointLocalRestTransforms, UsdTimeCode::EarliestTime(), bAtRest );
-			bSuccess &= UsdSkelConcatJointTransforms( SkeletonQuery.GetTopology(), JointLocalRestTransforms, &WorldSpaceRestTransforms );
+			bSuccess &= SkeletonQuery.ComputeJointLocalTransforms(&JointLocalRestTransforms, UsdTimeCode::EarliestTime(), bAtRest);
+			bSuccess &= UsdSkelConcatJointTransforms(SkeletonQuery.GetTopology(), JointLocalRestTransforms, &WorldSpaceRestTransforms);
 		}
 
 		// Get world-space bindTransforms
 		VtArray<GfMatrix4d> WorldSpaceBindTransforms;
-		bSuccess &= SkeletonQuery.GetJointWorldBindTransforms( &WorldSpaceBindTransforms );
+		bSuccess &= SkeletonQuery.GetJointWorldBindTransforms(&WorldSpaceBindTransforms);
 
-		if ( !bSuccess )
+		if (!bSuccess)
 		{
 			return false;
 		}
 
-		MeshToSkeletonRestPose.resize( WorldSpaceRestTransforms.size() );
-		for ( uint32 Index = 0; Index < WorldSpaceRestTransforms.size(); ++Index )
+		MeshToSkeletonRestPose.resize(WorldSpaceRestTransforms.size());
+		for (uint32 Index = 0; Index < WorldSpaceRestTransforms.size(); ++Index)
 		{
-			MeshToSkeletonRestPose[ Index ] = WorldSpaceBindTransforms[ Index ].GetInverse() * WorldSpaceRestTransforms[ Index ];
+			MeshToSkeletonRestPose[Index] = WorldSpaceBindTransforms[Index].GetInverse() * WorldSpaceRestTransforms[Index];
 		}
 	}
 
@@ -1443,7 +1522,7 @@ bool UsdToUnreal::ConvertSkinnedMesh(
 		VtArray<GfVec3f> UsdPoints;
 		PointsAttr.Get(&UsdPoints, UsdTimeCode::Default());
 
-		SkinningQuery.ComputeSkinnedPoints( MeshToSkeletonRestPose, &UsdPoints, UsdTimeCode::EarliestTime() );
+		SkinningQuery.ComputeSkinnedPoints(MeshToSkeletonRestPose, &UsdPoints, UsdTimeCode::EarliestTime());
 
 		NumPoints = UsdPoints.size();
 		SkelMeshImportData.Points.AddUninitialized(NumPoints);
@@ -1451,7 +1530,7 @@ bool UsdToUnreal::ConvertSkinnedMesh(
 		for (uint32 PointIndex = 0; PointIndex < NumPoints; ++PointIndex)
 		{
 			const GfVec3f& Point = UsdPoints[PointIndex];
-			SkelMeshImportData.Points[ PointIndex + NumExistingPoints ] = ( FVector3f ) UsdToUnreal::ConvertVector( StageInfo, Point );
+			SkelMeshImportData.Points[PointIndex + NumExistingPoints] = (FVector3f)UsdToUnreal::ConvertVector(StageInfo, Point);
 		}
 	}
 
@@ -1485,10 +1564,10 @@ bool UsdToUnreal::ConvertSkinnedMesh(
 	UsdAttribute NormalsAttribute = UsdMesh.GetNormalsAttr();
 	if (NormalsAttribute)
 	{
-		if ( NormalsAttribute.Get( &Normals, UsdTimeCode::Default() ) && Normals.size() > 0 )
+		if (NormalsAttribute.Get(&Normals, UsdTimeCode::Default()) && Normals.size() > 0)
 		{
 			// Like for points, we need to ensure these normals are in the same coordinate space of the skeleton
-			SkinningQuery.ComputeSkinnedNormals( MeshToSkeletonRestPose, &Normals, UsdTimeCode::EarliestTime() );
+			SkinningQuery.ComputeSkinnedNormals(MeshToSkeletonRestPose, &Normals, UsdTimeCode::EarliestTime());
 		}
 	}
 
@@ -1496,7 +1575,7 @@ bool UsdToUnreal::ConvertSkinnedMesh(
 	uint32 NumExistingWedges = SkelMeshImportData.Wedges.Num();
 
 	uint32 NumFaces = FaceCounts.size();
-	SkelMeshImportData.Faces.Reserve( NumFaces * 2 );
+	SkelMeshImportData.Faces.Reserve(NumFaces * 2);
 
 	// Material assignments
 	const bool bProvideMaterialIndices = true;
@@ -1507,8 +1586,8 @@ bool UsdToUnreal::ConvertSkinnedMesh(
 		CommonOptions.RenderContext,
 		CommonOptions.MaterialPurpose
 	);
-	TArray< UsdUtils::FUsdPrimMaterialSlot >& LocalMaterialSlots = LocalInfo.Slots;
-	TArray< int32 >& FaceMaterialIndices = LocalInfo.MaterialIndices;
+	TArray<UsdUtils::FUsdPrimMaterialSlot>& LocalMaterialSlots = LocalInfo.Slots;
+	TArray<int32>& FaceMaterialIndices = LocalInfo.MaterialIndices;
 
 	// We want to combine identical slots for skeletal meshes, which is different to static meshes, where each section gets a slot
 	// Note: This is a different index remapping to the one that happens for LODs, using LODMaterialMap! Here we're combining meshes of the same LOD
@@ -1516,36 +1595,36 @@ bool UsdToUnreal::ConvertSkinnedMesh(
 
 	// Position 3 in this has the value 6 --> Local material slot #3 is actually the combined material slot #6
 	TArray<int32> LocalToCombinedMaterialIndex;
-	LocalToCombinedMaterialIndex.SetNumZeroed( LocalInfo.Slots.Num() );
+	LocalToCombinedMaterialIndex.SetNumZeroed(LocalInfo.Slots.Num());
 
 	for (int32 Index = 0; Index < MaterialAssignments.Slots.Num(); ++Index)
 	{
 		const UsdUtils::FUsdPrimMaterialSlot& Slot = MaterialAssignments.Slots[Index];
 
 		// Combine entries in this way so that we can append PrimPaths
-		TMap< UsdUtils::FUsdPrimMaterialSlot, int32 >::TKeyIterator KeyIt = SlotToCombinedMaterialIndex.CreateKeyIterator( Slot );
-		if ( KeyIt )
+		TMap<UsdUtils::FUsdPrimMaterialSlot, int32>::TKeyIterator KeyIt = SlotToCombinedMaterialIndex.CreateKeyIterator(Slot);
+		if (KeyIt)
 		{
-			KeyIt.Key().PrimPaths.Append( Slot.PrimPaths );
+			KeyIt.Key().PrimPaths.Append(Slot.PrimPaths);
 			KeyIt.Value() = Index;
 		}
 		else
 		{
-			SlotToCombinedMaterialIndex.Add( Slot, Index );
+			SlotToCombinedMaterialIndex.Add(Slot, Index);
 		}
 	}
 	for (int32 LocalIndex = 0; LocalIndex < LocalInfo.Slots.Num(); ++LocalIndex)
 	{
-		const UsdUtils::FUsdPrimMaterialSlot& LocalSlot = LocalInfo.Slots[ LocalIndex ];
+		const UsdUtils::FUsdPrimMaterialSlot& LocalSlot = LocalInfo.Slots[LocalIndex];
 
 		// Combine entries in this way so that we can append PrimPaths
-		TMap< UsdUtils::FUsdPrimMaterialSlot, int32 >::TKeyIterator KeyIt = SlotToCombinedMaterialIndex.CreateKeyIterator( LocalSlot );
-		if ( KeyIt )
+		TMap<UsdUtils::FUsdPrimMaterialSlot, int32>::TKeyIterator KeyIt = SlotToCombinedMaterialIndex.CreateKeyIterator(LocalSlot);
+		if (KeyIt)
 		{
-			KeyIt.Key().PrimPaths.Append( LocalSlot.PrimPaths );
+			KeyIt.Key().PrimPaths.Append(LocalSlot.PrimPaths);
 
 			const int32 ExistingCombinedIndex = KeyIt.Value();
-			LocalToCombinedMaterialIndex[ LocalIndex ] = ExistingCombinedIndex;
+			LocalToCombinedMaterialIndex[LocalIndex] = ExistingCombinedIndex;
 		}
 		else
 		{
@@ -1558,8 +1637,8 @@ bool UsdToUnreal::ConvertSkinnedMesh(
 	// our output
 	for (UsdUtils::FUsdPrimMaterialSlot& Slot : MaterialAssignments.Slots)
 	{
-		TMap< UsdUtils::FUsdPrimMaterialSlot, int32 >::TKeyIterator KeyIt = SlotToCombinedMaterialIndex.CreateKeyIterator( Slot );
-		ensure( KeyIt );
+		TMap<UsdUtils::FUsdPrimMaterialSlot, int32>::TKeyIterator KeyIt = SlotToCombinedMaterialIndex.CreateKeyIterator(Slot);
+		ensure(KeyIt);
 		Slot.PrimPaths = KeyIt.Key().PrimPaths;
 	}
 
@@ -1570,51 +1649,55 @@ bool UsdToUnreal::ConvertSkinnedMesh(
 	if (ColorPrimvar)
 	{
 		pxr::VtArray<pxr::GfVec3f> UsdColors;
-		if ( ColorPrimvar.ComputeFlattened( &UsdColors ) )
+		if (ColorPrimvar.ComputeFlattened(&UsdColors))
 		{
 			uint32 NumExpectedColors = 0;
 			uint32 NumColors = UsdColors.size();
 			pxr::TfToken USDInterpType = ColorPrimvar.GetInterpolation();
 
-			if ( USDInterpType == pxr::UsdGeomTokens->uniform )
+			if (USDInterpType == pxr::UsdGeomTokens->uniform)
 			{
 				NumExpectedColors = NumFaces;
 				DisplayColorInterp = EUsdInterpolationMethod::Uniform;
 			}
-			else if ( USDInterpType == pxr::UsdGeomTokens->vertex || USDInterpType == pxr::UsdGeomTokens->varying )
+			else if (USDInterpType == pxr::UsdGeomTokens->vertex || USDInterpType == pxr::UsdGeomTokens->varying)
 			{
 				NumExpectedColors = NumPoints;
 				DisplayColorInterp = EUsdInterpolationMethod::Vertex;
 			}
-			else if ( USDInterpType == pxr::UsdGeomTokens->faceVarying )
+			else if (USDInterpType == pxr::UsdGeomTokens->faceVarying)
 			{
 				NumExpectedColors = NumVertexInstances;
 				DisplayColorInterp = EUsdInterpolationMethod::FaceVarying;
 			}
-			else if ( USDInterpType == pxr::UsdGeomTokens->constant )
+			else if (USDInterpType == pxr::UsdGeomTokens->constant)
 			{
 				NumExpectedColors = 1;
 				DisplayColorInterp = EUsdInterpolationMethod::Constant;
 			}
 
-			if ( NumExpectedColors == NumColors )
+			if (NumExpectedColors == NumColors)
 			{
-				Colors.Reserve( NumColors );
-				for ( uint32 Index = 0; Index < NumColors; ++Index )
+				Colors.Reserve(NumColors);
+				for (uint32 Index = 0; Index < NumColors; ++Index)
 				{
 					const bool bSRGB = true;
-					Colors.Add( UsdToUnreal::ConvertColor( UsdColors[ Index ] ).ToFColor( bSRGB ) );
+					Colors.Add(UsdToUnreal::ConvertColor(UsdColors[Index]).ToFColor(bSRGB));
 				}
 
 				SkelMeshImportData.bHasVertexColors = true;
 			}
 			else
 			{
-				UE_LOG( LogUsd, Warning, TEXT( "Prim '%s' has invalid number of displayColor values for primvar interpolation type '%s'! (expected %d, found %d)" ),
-					*UsdToUnreal::ConvertPath( SkinningPrim.GetPath() ),
-					*UsdToUnreal::ConvertToken( USDInterpType ),
+				UE_LOG(
+					LogUsd,
+					Warning,
+					TEXT("Prim '%s' has invalid number of displayColor values for primvar interpolation type '%s'! (expected %d, found %d)"),
+					*UsdToUnreal::ConvertPath(SkinningPrim.GetPath()),
+					*UsdToUnreal::ConvertToken(USDInterpType),
 					NumExpectedColors,
-					NumColors );
+					NumColors
+				);
 			}
 		}
 	}
@@ -1623,65 +1706,69 @@ bool UsdToUnreal::ConvertSkinnedMesh(
 	UsdGeomPrimvar OpacityPrimvar = UsdMesh.GetDisplayOpacityPrimvar();
 	TArray<float> Opacities;
 	EUsdInterpolationMethod DisplayOpacityInterp = EUsdInterpolationMethod::Constant;
-	if ( OpacityPrimvar )
+	if (OpacityPrimvar)
 	{
-		pxr::VtArray< float > UsdOpacities;
-		if ( OpacityPrimvar.ComputeFlattened( &UsdOpacities ) )
+		pxr::VtArray<float> UsdOpacities;
+		if (OpacityPrimvar.ComputeFlattened(&UsdOpacities))
 		{
 			uint32 NumExpectedOpacities = 0;
 			const uint32 NumOpacities = UsdOpacities.size();
 			pxr::TfToken USDInterpType = OpacityPrimvar.GetInterpolation();
 
-			if ( USDInterpType == pxr::UsdGeomTokens->uniform )
+			if (USDInterpType == pxr::UsdGeomTokens->uniform)
 			{
 				NumExpectedOpacities = NumFaces;
 				DisplayOpacityInterp = EUsdInterpolationMethod::Uniform;
 			}
-			else if ( USDInterpType == pxr::UsdGeomTokens->vertex || USDInterpType == pxr::UsdGeomTokens->varying )
+			else if (USDInterpType == pxr::UsdGeomTokens->vertex || USDInterpType == pxr::UsdGeomTokens->varying)
 			{
 				NumExpectedOpacities = NumPoints;
 				DisplayOpacityInterp = EUsdInterpolationMethod::Vertex;
 			}
-			else if ( USDInterpType == pxr::UsdGeomTokens->faceVarying )
+			else if (USDInterpType == pxr::UsdGeomTokens->faceVarying)
 			{
 				NumExpectedOpacities = NumVertexInstances;
 				DisplayOpacityInterp = EUsdInterpolationMethod::FaceVarying;
 			}
-			else if ( USDInterpType == pxr::UsdGeomTokens->constant )
+			else if (USDInterpType == pxr::UsdGeomTokens->constant)
 			{
 				NumExpectedOpacities = 1;
 				DisplayOpacityInterp = EUsdInterpolationMethod::Constant;
 			}
 
-			if ( NumExpectedOpacities == NumOpacities )
+			if (NumExpectedOpacities == NumOpacities)
 			{
-				Opacities.Reserve( NumOpacities );
-				for ( uint32 Index = 0; Index < NumOpacities; ++Index )
+				Opacities.Reserve(NumOpacities);
+				for (uint32 Index = 0; Index < NumOpacities; ++Index)
 				{
-					Opacities.Add( UsdOpacities[ Index ] );
+					Opacities.Add(UsdOpacities[Index]);
 				}
 
-				SkelMeshImportData.bHasVertexColors = true; // We'll need to store these in the vertex colors
+				SkelMeshImportData.bHasVertexColors = true;	   // We'll need to store these in the vertex colors
 			}
 			else
 			{
-				UE_LOG( LogUsd, Warning, TEXT( "Prim '%s' has invalid number of displayOpacity values for primvar interpolation type '%s'! (expected %d, found %d)" ),
-					*UsdToUnreal::ConvertPath( SkinningPrim.GetPath() ),
-					*UsdToUnreal::ConvertToken( USDInterpType ),
+				UE_LOG(
+					LogUsd,
+					Warning,
+					TEXT("Prim '%s' has invalid number of displayOpacity values for primvar interpolation type '%s'! (expected %d, found %d)"),
+					*UsdToUnreal::ConvertPath(SkinningPrim.GetPath()),
+					*UsdToUnreal::ConvertToken(USDInterpType),
 					NumExpectedOpacities,
-					NumOpacities );
+					NumOpacities
+				);
 			}
 		}
 	}
 
 	// Make sure these have at least one valid entry, as we'll default to Constant and we may have either valid opacities or colors
-	if ( Colors.Num() < 1 )
+	if (Colors.Num() < 1)
 	{
-		Colors.Add( FColor::White );
+		Colors.Add(FColor::White);
 	}
-	if ( Opacities.Num() < 1 )
+	if (Opacities.Num() < 1)
 	{
-		Opacities.Add( 1.0f );
+		Opacities.Add(1.0f);
 	}
 
 	SkelMeshImportData.NumTexCoords = 0;
@@ -1690,24 +1777,22 @@ bool UsdToUnreal::ConvertSkinnedMesh(
 
 	struct FUVSet
 	{
-		TOptional< VtIntArray > UVIndices; // UVs might be indexed or they might be flat (one per vertex)
+		TOptional<VtIntArray> UVIndices;	// UVs might be indexed or they might be flat (one per vertex)
 		VtVec2fArray UVs;
 
 		EUsdInterpolationMethod InterpolationMethod = EUsdInterpolationMethod::FaceVarying;
 	};
 
-	TArray< FUVSet > UVSets;
+	TArray<FUVSet> UVSets;
 
 	// If we already have a primvar to UV index assignment, let's just use that.
 	// When collapsing, we'll do a pre-pass on all meshes to translate and determine this beforehand.
 	TArray<TUsdStore<pxr::UsdGeomPrimvar>> PrimvarsByUVIndex;
 	if (MaterialAssignments.PrimvarToUVIndex.Num() > 0)
 	{
-		TArray<TUsdStore<pxr::UsdGeomPrimvar>> AllMeshUVPrimvars =
-			UsdUtils::GetUVSetPrimvars(SkinningPrim, TNumericLimits<int32>::Max());
+		TArray<TUsdStore<pxr::UsdGeomPrimvar>> AllMeshUVPrimvars = UsdUtils::GetUVSetPrimvars(SkinningPrim, TNumericLimits<int32>::Max());
 
-		PrimvarsByUVIndex =
-			UsdUtils::AssemblePrimvarsIntoUVSets(AllMeshUVPrimvars, MaterialAssignments.PrimvarToUVIndex);
+		PrimvarsByUVIndex = UsdUtils::AssemblePrimvarsIntoUVSets(AllMeshUVPrimvars, MaterialAssignments.PrimvarToUVIndex);
 	}
 	// Let's use the best primvar assignment for this particular mesh instead
 	else
@@ -1718,59 +1803,59 @@ bool UsdToUnreal::ConvertSkinnedMesh(
 	}
 
 	int32 UVChannelIndex = 0;
-	while ( true )
+	while (true)
 	{
-		if ( !PrimvarsByUVIndex.IsValidIndex( UVChannelIndex ) )
+		if (!PrimvarsByUVIndex.IsValidIndex(UVChannelIndex))
 		{
 			break;
 		}
 
-		UsdGeomPrimvar& PrimvarST = PrimvarsByUVIndex[ UVChannelIndex ].Get();
-		if ( !PrimvarST )
+		UsdGeomPrimvar& PrimvarST = PrimvarsByUVIndex[UVChannelIndex].Get();
+		if (!PrimvarST)
 		{
 			break;
 		}
 
-		if ( PrimvarST )
+		if (PrimvarST)
 		{
 			FUVSet UVSet;
 
-			if ( PrimvarST.GetInterpolation() == UsdGeomTokens->vertex )
+			if (PrimvarST.GetInterpolation() == UsdGeomTokens->vertex)
 			{
 				UVSet.InterpolationMethod = EUsdInterpolationMethod::Vertex;
 			}
-			else if ( PrimvarST.GetInterpolation() == UsdGeomTokens->faceVarying )
+			else if (PrimvarST.GetInterpolation() == UsdGeomTokens->faceVarying)
 			{
 				UVSet.InterpolationMethod = EUsdInterpolationMethod::FaceVarying;
 			}
-			else if (  PrimvarST.GetInterpolation() == UsdGeomTokens->uniform )
+			else if (PrimvarST.GetInterpolation() == UsdGeomTokens->uniform)
 			{
 				UVSet.InterpolationMethod = EUsdInterpolationMethod::Uniform;
 			}
-			else if ( PrimvarST.GetInterpolation() == UsdGeomTokens->constant )
+			else if (PrimvarST.GetInterpolation() == UsdGeomTokens->constant)
 			{
 				UVSet.InterpolationMethod = EUsdInterpolationMethod::Constant;
 			}
 
-			if ( PrimvarST.IsIndexed() )
+			if (PrimvarST.IsIndexed())
 			{
 				UVSet.UVIndices.Emplace();
 
-				if ( PrimvarST.GetIndices( &UVSet.UVIndices.GetValue() ) && PrimvarST.Get( &UVSet.UVs ) )
+				if (PrimvarST.GetIndices(&UVSet.UVIndices.GetValue()) && PrimvarST.Get(&UVSet.UVs))
 				{
-					if ( UVSet.UVs.size() > 0 )
+					if (UVSet.UVs.size() > 0)
 					{
-						UVSets.Add( MoveTemp( UVSet ) );
+						UVSets.Add(MoveTemp(UVSet));
 					}
 				}
 			}
 			else
 			{
-				if ( PrimvarST.Get( &UVSet.UVs ) )
+				if (PrimvarST.Get(&UVSet.UVs))
 				{
-					if ( UVSet.UVs.size() > 0 )
+					if (UVSet.UVs.size() > 0)
 					{
-						UVSets.Add( MoveTemp( UVSet ) );
+						UVSets.Add(MoveTemp(UVSet));
 					}
 				}
 			}
@@ -1783,7 +1868,7 @@ bool UsdToUnreal::ConvertSkinnedMesh(
 		++UVChannelIndex;
 	}
 
-	SkelMeshImportData.Wedges.Reserve( ( NumExistingFaces + NumFaces ) * 6 );
+	SkelMeshImportData.Wedges.Reserve((NumExistingFaces + NumFaces) * 6);
 
 	uint32 NumProcessedFaceVertexIndices = 0;
 	for (uint32 PolygonIndex = NumExistingFaces, LocalIndex = 0; PolygonIndex < NumExistingFaces + NumFaces; ++PolygonIndex, ++LocalIndex)
@@ -1793,28 +1878,28 @@ bool UsdToUnreal::ConvertSkinnedMesh(
 
 		// Manage materials
 		int32 LocalMaterialIndex = 0;
-		if ( FaceMaterialIndices.IsValidIndex( LocalIndex ) )
+		if (FaceMaterialIndices.IsValidIndex(LocalIndex))
 		{
-			LocalMaterialIndex = FaceMaterialIndices[ LocalIndex ];
-			if ( !LocalMaterialSlots.IsValidIndex( LocalMaterialIndex ) )
+			LocalMaterialIndex = FaceMaterialIndices[LocalIndex];
+			if (!LocalMaterialSlots.IsValidIndex(LocalMaterialIndex))
 			{
 				LocalMaterialIndex = 0;
 			}
 		}
 
-		int32 RealMaterialIndex = LocalToCombinedMaterialIndex[ LocalMaterialIndex ];
-		SkelMeshImportData.MaxMaterialIndex = FMath::Max< uint32 >( SkelMeshImportData.MaxMaterialIndex, RealMaterialIndex );
+		int32 RealMaterialIndex = LocalToCombinedMaterialIndex[LocalMaterialIndex];
+		SkelMeshImportData.MaxMaterialIndex = FMath::Max<uint32>(SkelMeshImportData.MaxMaterialIndex, RealMaterialIndex);
 
 		// SkeletalMeshImportData uses triangle faces so quads will have to be split into triangles
-		const bool bIsQuad = ( NumOriginalFaceVertices == 4 );
+		const bool bIsQuad = (NumOriginalFaceVertices == 4);
 		const uint32 NumTriangles = bIsQuad ? 2 : 1;
 
-		for ( uint32 TriangleIndex = 0; TriangleIndex < NumTriangles; ++TriangleIndex )
+		for (uint32 TriangleIndex = 0; TriangleIndex < NumTriangles; ++TriangleIndex)
 		{
 			// This needs to be zeroed as we'll hash these faces later
 			int32 TriangleFaceIndex = SkelMeshImportData.Faces.AddZeroed();
 
-			SkeletalMeshImportData::FTriangle& Triangle = SkelMeshImportData.Faces[ TriangleFaceIndex ];
+			SkeletalMeshImportData::FTriangle& Triangle = SkelMeshImportData.Faces[TriangleFaceIndex];
 
 			// Set the face smoothing by default. It could be any number, but not zero
 			Triangle.SmoothingGroups = 255;
@@ -1825,30 +1910,40 @@ bool UsdToUnreal::ConvertSkinnedMesh(
 			Triangle.AuxMatIndex = 0;
 
 			// Fill the wedge data and complete the triangle setup with the wedge indices
-			for ( uint32 CornerIndex = 0; CornerIndex < NumFinalFaceVertices; ++CornerIndex )
+			for (uint32 CornerIndex = 0; CornerIndex < NumFinalFaceVertices; ++CornerIndex)
 			{
-				uint32 OriginalCornerIndex = ( ( TriangleIndex * ( NumOriginalFaceVertices - 2 ) ) + CornerIndex ) % NumOriginalFaceVertices;
+				uint32 OriginalCornerIndex = ((TriangleIndex * (NumOriginalFaceVertices - 2)) + CornerIndex) % NumOriginalFaceVertices;
 				uint32 OriginalVertexInstanceIndex = NumProcessedFaceVertexIndices + OriginalCornerIndex;
-				int32 OriginalVertexIndex = OriginalFaceIndices[ OriginalVertexInstanceIndex ];
+				int32 OriginalVertexIndex = OriginalFaceIndices[OriginalVertexInstanceIndex];
 
 				int32 FinalCornerIndex = bReverseOrder ? NumFinalFaceVertices - 1 - CornerIndex : CornerIndex;
 
 				// Its important to make sure the UVs aren't just uninitialized memory because BuildSkeletalMesh will read them
 				// when trying to merge vertices. Uninitialized memory would lead to inconsistent, non-deterministic meshes
 				const uint32 WedgeIndex = SkelMeshImportData.Wedges.AddZeroed();
-				SkeletalMeshImportData::FVertex& SkelMeshWedge = SkelMeshImportData.Wedges[ WedgeIndex ];
+				SkeletalMeshImportData::FVertex& SkelMeshWedge = SkelMeshImportData.Wedges[WedgeIndex];
 
-				if ( SkelMeshImportData.bHasVertexColors )
+				if (SkelMeshImportData.bHasVertexColors)
 				{
-					uint32 DisplayColorIndex = UsdToUnrealImpl::GetPrimValueIndex( DisplayColorInterp, OriginalVertexIndex, OriginalVertexInstanceIndex, LocalIndex );
-					uint32 DisplayOpacityIndex = UsdToUnrealImpl::GetPrimValueIndex( DisplayOpacityInterp, OriginalVertexIndex, OriginalVertexInstanceIndex, LocalIndex );
+					uint32 DisplayColorIndex = UsdToUnrealImpl::GetPrimValueIndex(
+						DisplayColorInterp,
+						OriginalVertexIndex,
+						OriginalVertexInstanceIndex,
+						LocalIndex
+					);
+					uint32 DisplayOpacityIndex = UsdToUnrealImpl::GetPrimValueIndex(
+						DisplayOpacityInterp,
+						OriginalVertexIndex,
+						OriginalVertexInstanceIndex,
+						LocalIndex
+					);
 
-					const FColor& DisplayColor = Colors[ DisplayColorIndex ];
+					const FColor& DisplayColor = Colors[DisplayColorIndex];
 
 					SkelMeshWedge.Color.R = DisplayColor.R;
 					SkelMeshWedge.Color.G = DisplayColor.G;
 					SkelMeshWedge.Color.B = DisplayColor.B;
-					SkelMeshWedge.Color.A = static_cast<uint8>(FMath::Clamp(Opacities[ DisplayOpacityIndex ], 0.0f, 1.0f) * 255.0f + 0.5f);
+					SkelMeshWedge.Color.A = static_cast<uint8>(FMath::Clamp(Opacities[DisplayOpacityIndex], 0.0f, 1.0f) * 255.0f + 0.5f);
 				}
 
 				SkelMeshWedge.MatIndex = Triangle.MatIndex;
@@ -1856,53 +1951,53 @@ bool UsdToUnreal::ConvertSkinnedMesh(
 				SkelMeshWedge.Reserved = 0;
 
 				int32 UVLayerIndex = 0;
-				for ( const FUVSet& UVSet : UVSets )
+				for (const FUVSet& UVSet : UVSets)
 				{
 					int32 ValueIndex = 0;
 
-					if ( UVSet.InterpolationMethod == EUsdInterpolationMethod::Vertex )
+					if (UVSet.InterpolationMethod == EUsdInterpolationMethod::Vertex)
 					{
 						ValueIndex = OriginalVertexIndex;
 					}
-					else if ( UVSet.InterpolationMethod == EUsdInterpolationMethod::FaceVarying )
+					else if (UVSet.InterpolationMethod == EUsdInterpolationMethod::FaceVarying)
 					{
 						ValueIndex = OriginalVertexInstanceIndex;
 					}
-					else if ( UVSet.InterpolationMethod == EUsdInterpolationMethod::Uniform )
+					else if (UVSet.InterpolationMethod == EUsdInterpolationMethod::Uniform)
 					{
 						ValueIndex = PolygonIndex;
 					}
-					else if ( UVSet.InterpolationMethod == EUsdInterpolationMethod::Constant )
+					else if (UVSet.InterpolationMethod == EUsdInterpolationMethod::Constant)
 					{
 						ValueIndex = 0;
 					}
 
-					GfVec2f UV( 0.f, 0.f );
+					GfVec2f UV(0.f, 0.f);
 
-					if ( UVSet.UVIndices.IsSet() )
+					if (UVSet.UVIndices.IsSet())
 					{
-						if ( ensure( UVSet.UVIndices.GetValue().size() > ValueIndex ) )
+						if (ensure(UVSet.UVIndices.GetValue().size() > ValueIndex))
 						{
-							UV = UVSet.UVs[ UVSet.UVIndices.GetValue()[ ValueIndex ] ];
+							UV = UVSet.UVs[UVSet.UVIndices.GetValue()[ValueIndex]];
 						}
 					}
-					else if ( ensure( UVSet.UVs.size() > ValueIndex ) )
+					else if (ensure(UVSet.UVs.size() > ValueIndex))
 					{
-						UV = UVSet.UVs[ ValueIndex ];
+						UV = UVSet.UVs[ValueIndex];
 					}
 
 					// Flip V for Unreal uv's which match directx
-					FVector2f FinalUVVector( UV[0], 1.f - UV[1] );
-					SkelMeshWedge.UVs[ UVLayerIndex ] = FinalUVVector;
+					FVector2f FinalUVVector(UV[0], 1.f - UV[1]);
+					SkelMeshWedge.UVs[UVLayerIndex] = FinalUVVector;
 
 					++UVLayerIndex;
 				}
 
-				Triangle.TangentX[ FinalCornerIndex ] = FVector3f::ZeroVector;
-				Triangle.TangentY[ FinalCornerIndex ] = FVector3f::ZeroVector;
-				Triangle.TangentZ[ FinalCornerIndex ] = FVector3f::ZeroVector;
+				Triangle.TangentX[FinalCornerIndex] = FVector3f::ZeroVector;
+				Triangle.TangentY[FinalCornerIndex] = FVector3f::ZeroVector;
+				Triangle.TangentZ[FinalCornerIndex] = FVector3f::ZeroVector;
 
-				Triangle.WedgeIndex[ FinalCornerIndex ] = WedgeIndex;
+				Triangle.WedgeIndex[FinalCornerIndex] = WedgeIndex;
 			}
 		}
 
@@ -1912,7 +2007,8 @@ bool UsdToUnreal::ConvertSkinnedMesh(
 	// Convert joint influences into the SkeletalMeshImportData
 
 	// ComputeJointInfluences returns the influences per bone that applies to all the points of the mesh
-	// ComputeVaryingJointInfluences returns the joint influences for each points, expanding the influences to all points if the mesh is rigidly deformed
+	// ComputeVaryingJointInfluences returns the joint influences for each points, expanding the influences to all points if the mesh is rigidly
+	// deformed
 	VtArray<int> JointIndices;
 	VtArray<float> JointWeights;
 	SkinningQuery.ComputeVaryingJointInfluences(NumPoints, &JointIndices, &JointWeights);
@@ -1927,18 +2023,18 @@ bool UsdToUnreal::ConvertSkinnedMesh(
 	// Recompute the joint influences if we need to
 	uint32 NumInfluencesPerComponent = SkinningQuery.GetNumInfluencesPerComponent();
 	const uint32 MaxAllowedInfluences = EXTRA_BONE_INFLUENCES;
-	const bool bUseUnlimitedBoneInfluences = FGPUBaseSkinVertexFactory::UseUnlimitedBoneInfluences( NumInfluencesPerComponent );
-	if ( NumInfluencesPerComponent > MaxAllowedInfluences && !bUseUnlimitedBoneInfluences )
+	const bool bUseUnlimitedBoneInfluences = FGPUBaseSkinVertexFactory::UseUnlimitedBoneInfluences(NumInfluencesPerComponent);
+	if (NumInfluencesPerComponent > MaxAllowedInfluences && !bUseUnlimitedBoneInfluences)
 	{
-		UsdSkelResizeInfluences( &JointIndices, NumInfluencesPerComponent, MaxAllowedInfluences );
-		UsdSkelResizeInfluences( &JointWeights, NumInfluencesPerComponent, MaxAllowedInfluences );
+		UsdSkelResizeInfluences(&JointIndices, NumInfluencesPerComponent, MaxAllowedInfluences);
+		UsdSkelResizeInfluences(&JointWeights, NumInfluencesPerComponent, MaxAllowedInfluences);
 		NumInfluencesPerComponent = MaxAllowedInfluences;
 	}
 
 	// We keep track of which influences we added because we combine many Mesh prim (each with potentially a different
 	// explicit joint order) into the same skeletal mesh asset
 	const int32 NumInfluencesBefore = SkelMeshImportData.Influences.Num();
-	if ( JointWeights.size() > ( NumPoints - 1 ) * ( NumInfluencesPerComponent - 1 ) )
+	if (JointWeights.size() > (NumPoints - 1) * (NumInfluencesPerComponent - 1))
 	{
 		uint32 JointIndex = 0;
 		SkelMeshImportData.Influences.Reserve(NumPoints);
@@ -1947,7 +2043,8 @@ bool UsdToUnreal::ConvertSkinnedMesh(
 			// The JointIndices/JointWeights contain the influences data for NumPoints * NumInfluencesPerComponent
 			for (uint32 InfluenceIndex = 0; InfluenceIndex < NumInfluencesPerComponent; ++InfluenceIndex, ++JointIndex)
 			{
-				// BoneWeight could be 0 if the actual number of influences were less than NumInfluencesPerComponent for a given point so just ignore it
+				// BoneWeight could be 0 if the actual number of influences were less than NumInfluencesPerComponent for a given point so just ignore
+				// it
 				float BoneWeight = JointWeights[JointIndex];
 				if (BoneWeight != 0.f)
 				{
@@ -1962,38 +2059,39 @@ bool UsdToUnreal::ConvertSkinnedMesh(
 	const int32 NumInfluencesAfter = SkelMeshImportData.Influences.Num();
 
 	// If we have a joint mapper this Mesh has an explicit joint ordering, so we need to map joint indices to the skeleton's bone indices
-	if ( pxr::UsdSkelAnimMapperRefPtr AnimMapper = SkinningQuery.GetJointMapper() )
+	if (pxr::UsdSkelAnimMapperRefPtr AnimMapper = SkinningQuery.GetJointMapper())
 	{
 		VtArray<int> SkeletonBoneIndices;
-		if ( pxr::UsdSkelSkeleton BoundSkeleton = SkelBindingAPI.GetInheritedSkeleton() )
+		if (pxr::UsdSkelSkeleton BoundSkeleton = SkelBindingAPI.GetInheritedSkeleton())
 		{
-			if ( pxr::UsdAttribute SkeletonJointsAttr = BoundSkeleton.GetJointsAttr() )
+			if (pxr::UsdAttribute SkeletonJointsAttr = BoundSkeleton.GetJointsAttr())
 			{
 				VtArray<TfToken> SkeletonJoints;
-				if ( SkeletonJointsAttr.Get( &SkeletonJoints ) )
+				if (SkeletonJointsAttr.Get(&SkeletonJoints))
 				{
 					// If the skeleton has N bones, this will just contain { 0, 1, 2, ..., N-1 }
-					int NumUsdSkeletonBones = static_cast< int >( SkeletonJoints.size() );
-					for ( int SkeletonBoneIndex = 0; SkeletonBoneIndex < NumUsdSkeletonBones; ++SkeletonBoneIndex )
+					int NumUsdSkeletonBones = static_cast<int>(SkeletonJoints.size());
+					for (int SkeletonBoneIndex = 0; SkeletonBoneIndex < NumUsdSkeletonBones; ++SkeletonBoneIndex)
 					{
-						SkeletonBoneIndices.push_back( SkeletonBoneIndex );
+						SkeletonBoneIndices.push_back(SkeletonBoneIndex);
 					}
 
 					// Use the AnimMapper to produce the indices of the Mesh's joints within the Skeleton's list of joints.
-					// Example: Imagine skeleton had { "Root", "Root/Hip", "Root/Hip/Shoulder", "Root/Hip/Shoulder/Arm", "Root/Hip/Shoulder/Arm/Elbow" }, and so
-					// BoneIndexRemapping was { 0, 1, 2, 3, 4 }. Consider a Mesh that specifies the explicit joints { "Root/Hip/Shoulder", "Root/Hip/Shoulder/Arm" },
-					// and so uses the indices 0 and 1 to refer to Shoulder and Arm. After the Remap call SkeletonBoneIndices will hold { 2, 3 }, as those are the
-					// indices of Shoulder and Arm within the skeleton's bones
+					// Example: Imagine skeleton had { "Root", "Root/Hip", "Root/Hip/Shoulder", "Root/Hip/Shoulder/Arm", "Root/Hip/Shoulder/Arm/Elbow"
+					// }, and so BoneIndexRemapping was { 0, 1, 2, 3, 4 }. Consider a Mesh that specifies the explicit joints { "Root/Hip/Shoulder",
+					// "Root/Hip/Shoulder/Arm" }, and so uses the indices 0 and 1 to refer to Shoulder and Arm. After the Remap call
+					// SkeletonBoneIndices will hold { 2, 3 }, as those are the indices of Shoulder and Arm within the skeleton's bones
 					VtArray<int> BoneIndexRemapping;
-					if ( AnimMapper->Remap( SkeletonBoneIndices, &BoneIndexRemapping ) )
+					if (AnimMapper->Remap(SkeletonBoneIndices, &BoneIndexRemapping))
 					{
-						for ( int32 AddedInfluenceIndex = NumInfluencesBefore; AddedInfluenceIndex < NumInfluencesAfter; ++AddedInfluenceIndex )
+						for (int32 AddedInfluenceIndex = NumInfluencesBefore; AddedInfluenceIndex < NumInfluencesAfter; ++AddedInfluenceIndex)
 						{
-							SkeletalMeshImportData::FRawBoneInfluence& Influence = SkelMeshImportData.Influences[ AddedInfluenceIndex ];
+							SkeletalMeshImportData::FRawBoneInfluence& Influence = SkelMeshImportData.Influences[AddedInfluenceIndex];
 
 							// We have to remove our "NumAdditionalBones" offset from the influence's bone index because that's a UE concept that
 							// the BoneIndexRemapping array doesn't really know about. After that, we have a bone index that matches the USD Skeleton
-							// joint order, then we can remap with BoneIndexRemapping and add our NumAdditionalBones back in so that it matches our USkeleton
+							// joint order, then we can remap with BoneIndexRemapping and add our NumAdditionalBones back in so that it matches our
+							// USkeleton
 							Influence.BoneIndex = NumAdditionalBones + BoneIndexRemapping[Influence.BoneIndex - NumAdditionalBones];
 						}
 					}
@@ -2002,7 +2100,6 @@ bool UsdToUnreal::ConvertSkinnedMesh(
 		}
 	}
 
-
 	return true;
 }
 
@@ -2010,8 +2107,8 @@ bool UsdToUnreal::ConvertSkinnedMesh(
 	const pxr::UsdSkelSkinningQuery& SkinningQuery,
 	const pxr::UsdSkelSkeletonQuery& SkeletonQuery,
 	FSkeletalMeshImportData& SkelMeshImportData,
-	TArray< UsdUtils::FUsdPrimMaterialSlot >& MaterialAssignments,
-	const TMap< FString, TMap< FString, int32 > >& MaterialToPrimvarsUVSetNames,
+	TArray<UsdUtils::FUsdPrimMaterialSlot>& MaterialAssignments,
+	const TMap<FString, TMap<FString, int32>>& MaterialToPrimvarsUVSetNames,
 	const pxr::TfToken& RenderContext,
 	const pxr::TfToken& MaterialPurpose
 )
@@ -2023,13 +2120,7 @@ bool UsdToUnreal::ConvertSkinnedMesh(
 	UsdUtils::FUsdPrimMaterialAssignmentInfo TempInfo;
 	TempInfo.Slots = MaterialAssignments;
 
-	const bool bResult = ConvertSkinnedMesh(
-		SkinningQuery,
-		SkeletonQuery,
-		SkelMeshImportData,
-		TempInfo,
-		Options
-	);
+	const bool bResult = ConvertSkinnedMesh(SkinningQuery, SkeletonQuery, SkelMeshImportData, TempInfo, Options);
 
 	MaterialAssignments = TempInfo.Slots;
 	return bResult;
@@ -2047,47 +2138,49 @@ bool UsdToUnreal::ConvertSkelAnim(
 	float* OutStartOffsetSeconds
 )
 {
-	TRACE_CPUPROFILER_EVENT_SCOPE( UsdToUnreal::ConvertSkelAnim );
+	TRACE_CPUPROFILER_EVENT_SCOPE(UsdToUnreal::ConvertSkelAnim);
 
 	FScopedUnrealAllocs UEAllocs;
 
-	if ( !InUsdSkeletonQuery || !OutSkeletalAnimationAsset )
+	if (!InUsdSkeletonQuery || !OutSkeletalAnimationAsset)
 	{
 		return false;
 	}
 
 	// If we have no skeleton we can't add animation data to the AnimSequence, so we may as well just return
 	USkeleton* Skeleton = OutSkeletalAnimationAsset->GetSkeleton();
-	if ( !Skeleton )
+	if (!Skeleton)
 	{
 		return false;
 	}
 
 	TUsdStore<pxr::UsdSkelAnimQuery> AnimQuery = InUsdSkeletonQuery.GetAnimQuery();
-	if ( !AnimQuery.Get() )
+	if (!AnimQuery.Get())
 	{
 		return false;
 	}
 
 	pxr::UsdPrim SkelAnimPrim = AnimQuery.Get().GetPrim();
-	UE::FSdfLayerOffset Offset = UsdUtils::GetPrimToStageOffset( UE::FUsdPrim{ SkelAnimPrim } );
+	UE::FSdfLayerOffset Offset = UsdUtils::GetPrimToStageOffset(UE::FUsdPrim{SkelAnimPrim});
 
-	UE::FSdfLayer SkelAnimPrimLayer = UsdUtils::FindLayerForPrim( SkelAnimPrim );
+	UE::FSdfLayer SkelAnimPrimLayer = UsdUtils::FindLayerForPrim(SkelAnimPrim);
 	double LayerTimeCodesPerSecond = SkelAnimPrimLayer.GetTimeCodesPerSecond();
 
-	TUsdStore<pxr::UsdStageWeakPtr> Stage( InUsdSkeletonQuery.GetPrim().GetStage() );
-	FUsdStageInfo StageInfo{ Stage.Get() };
+	TUsdStore<pxr::UsdStageWeakPtr> Stage(InUsdSkeletonQuery.GetPrim().GetStage());
+	FUsdStageInfo StageInfo{Stage.Get()};
 	double StageTimeCodesPerSecond = Stage.Get()->GetTimeCodesPerSecond();
-	if ( FMath::IsNearlyZero( StageTimeCodesPerSecond ) )
+	if (FMath::IsNearlyZero(StageTimeCodesPerSecond))
 	{
-		FUsdLogManager::LogMessage( EMessageSeverity::Warning,
-									LOCTEXT("TimeCodesPerSecondIsZero", "Cannot bake skeletal animations as the stage has timeCodesPerSecond set to zero!") );
+		FUsdLogManager::LogMessage(
+			EMessageSeverity::Warning,
+			LOCTEXT("TimeCodesPerSecondIsZero", "Cannot bake skeletal animations as the stage has timeCodesPerSecond set to zero!")
+		);
 		return false;
 	}
 
 	const FReferenceSkeleton& RefSkeleton = Skeleton->GetReferenceSkeleton();
 	const TArray<FMeshBoneInfo>& BoneInfo = RefSkeleton.GetRawRefBoneInfo();
-	int32 NumBonesInUE = BoneInfo.Num();  // This will already contain any new "true root bone" we may have created
+	int32 NumBonesInUE = BoneInfo.Num();	// This will already contain any new "true root bone" we may have created
 	int32 NumBonesInUsd = static_cast<int32>(InUsdSkeletonQuery.GetJointOrder().size());
 
 	// Keep track of whether we added an additional "true" root bone in the cases the bound skeleton has
@@ -2106,27 +2199,27 @@ bool UsdToUnreal::ConvertSkelAnim(
 	}
 
 	TUsdStore<std::vector<double>> UsdJointTransformTimeSamples;
-	AnimQuery.Get().GetJointTransformTimeSamples( &( UsdJointTransformTimeSamples.Get() ) );
+	AnimQuery.Get().GetJointTransformTimeSamples(&(UsdJointTransformTimeSamples.Get()));
 	int32 NumJointTransformSamples = UsdJointTransformTimeSamples.Get().size();
 	TOptional<double> FirstJointSampleTimeCode;
 	TOptional<double> LastJointSampleTimeCode;
-	if ( UsdJointTransformTimeSamples.Get().size() > 0 )
+	if (UsdJointTransformTimeSamples.Get().size() > 0)
 	{
 		const std::vector<double>& JointTransformTimeSamples = UsdJointTransformTimeSamples.Get();
-		FirstJointSampleTimeCode = JointTransformTimeSamples[ 0 ];
-		LastJointSampleTimeCode = JointTransformTimeSamples[ JointTransformTimeSamples.size() - 1 ];
+		FirstJointSampleTimeCode = JointTransformTimeSamples[0];
+		LastJointSampleTimeCode = JointTransformTimeSamples[JointTransformTimeSamples.size() - 1];
 	}
 
 	TUsdStore<std::vector<double>> UsdBlendShapeTimeSamples;
-	AnimQuery.Get().GetBlendShapeWeightTimeSamples( &( UsdBlendShapeTimeSamples.Get() ) );
+	AnimQuery.Get().GetBlendShapeWeightTimeSamples(&(UsdBlendShapeTimeSamples.Get()));
 	int32 NumBlendShapeSamples = UsdBlendShapeTimeSamples.Get().size();
 	TOptional<double> FirstBlendShapeSampleTimeCode;
 	TOptional<double> LastBlendShapeSampleTimeCode;
-	if ( UsdBlendShapeTimeSamples.Get().size() > 0 )
+	if (UsdBlendShapeTimeSamples.Get().size() > 0)
 	{
 		const std::vector<double>& BlendShapeTimeSamples = UsdBlendShapeTimeSamples.Get();
-		FirstBlendShapeSampleTimeCode = BlendShapeTimeSamples[ 0 ];
-		LastBlendShapeSampleTimeCode = BlendShapeTimeSamples[ BlendShapeTimeSamples.size() - 1 ];
+		FirstBlendShapeSampleTimeCode = BlendShapeTimeSamples[0];
+		LastBlendShapeSampleTimeCode = BlendShapeTimeSamples[BlendShapeTimeSamples.size() - 1];
 	}
 
 	TUsdStore<std::vector<double>> UsdRootMotionPrimTimeSamples;
@@ -2138,23 +2231,23 @@ bool UsdToUnreal::ConvertSkelAnim(
 
 		// Note how we don't care whether the root motion is animated or not and will use RootMotionXformable
 		// regardless, to have a similar effect in case its just a single non-animated transform
-		RootMotionXformable = pxr::UsdGeomXformable{ RootMotionPrim };
-		if ( RootMotionXformable.Get() )
+		RootMotionXformable = pxr::UsdGeomXformable{RootMotionPrim};
+		if (RootMotionXformable.Get())
 		{
-			std::vector< double > UsdTimeSamples;
-			if ( RootMotionXformable.Get().GetTimeSamples(&UsdTimeSamples) )
+			std::vector<double> UsdTimeSamples;
+			if (RootMotionXformable.Get().GetTimeSamples(&UsdTimeSamples))
 			{
-				if ( UsdTimeSamples.size() > 0 )
+				if (UsdTimeSamples.size() > 0)
 				{
-					FirstRootMotionTimeCode = UsdTimeSamples[ 0 ];
-					LastRootMotionTimeCode = UsdTimeSamples[ UsdTimeSamples.size() - 1 ];
+					FirstRootMotionTimeCode = UsdTimeSamples[0];
+					LastRootMotionTimeCode = UsdTimeSamples[UsdTimeSamples.size() - 1];
 				}
 			}
 		}
 	}
 
 	// Nothing to do: we don't actually have joints or blend shape time samples
-	if ( !FirstJointSampleTimeCode.IsSet() && !FirstBlendShapeSampleTimeCode.IsSet() && !FirstRootMotionTimeCode.IsSet() )
+	if (!FirstJointSampleTimeCode.IsSet() && !FirstBlendShapeSampleTimeCode.IsSet() && !FirstRootMotionTimeCode.IsSet())
 	{
 		return true;
 	}
@@ -2170,28 +2263,22 @@ bool UsdToUnreal::ConvertSkelAnim(
 	// offsets and scalings anyway
 
 	const double StageStartTimeCode = FMath::Min(
-		FirstJointSampleTimeCode.Get( TNumericLimits<double>::Max() ),
-		FMath::Min(
-			FirstBlendShapeSampleTimeCode.Get( TNumericLimits<double>::Max() ),
-			FirstRootMotionTimeCode.Get( TNumericLimits<double>::Max() )
-		)
+		FirstJointSampleTimeCode.Get(TNumericLimits<double>::Max()),
+		FMath::Min(FirstBlendShapeSampleTimeCode.Get(TNumericLimits<double>::Max()), FirstRootMotionTimeCode.Get(TNumericLimits<double>::Max()))
 	);
 	const double StageEndTimeCode = FMath::Max(
-		LastJointSampleTimeCode.Get( TNumericLimits<double>::Lowest() ),
-		FMath::Max(
-			LastBlendShapeSampleTimeCode.Get( TNumericLimits<double>::Lowest() ),
-			LastRootMotionTimeCode.Get( TNumericLimits<double>::Lowest() )
-		)
+		LastJointSampleTimeCode.Get(TNumericLimits<double>::Lowest()),
+		FMath::Max(LastBlendShapeSampleTimeCode.Get(TNumericLimits<double>::Lowest()), LastRootMotionTimeCode.Get(TNumericLimits<double>::Lowest()))
 	);
 	const double StageStartSeconds = StageStartTimeCode / StageTimeCodesPerSecond;
 	const double StageSequenceLengthTimeCodes = StageEndTimeCode - StageStartTimeCode;
 	const double LayerSequenceLengthTimeCodes = StageSequenceLengthTimeCodes / Offset.Scale;
-	const double LayerSequenceLengthSeconds = FMath::Max<double>( LayerSequenceLengthTimeCodes / LayerTimeCodesPerSecond, MINIMUM_ANIMATION_LENGTH );
-	const double LayerStartTimeCode = ( StageStartTimeCode - Offset.Offset ) / Offset.Scale;
+	const double LayerSequenceLengthSeconds = FMath::Max<double>(LayerSequenceLengthTimeCodes / LayerTimeCodesPerSecond, MINIMUM_ANIMATION_LENGTH);
+	const double LayerStartTimeCode = (StageStartTimeCode - Offset.Offset) / Offset.Scale;
 	const double LayerStartSeconds = LayerStartTimeCode / LayerTimeCodesPerSecond;
 
 	// Just bake each time code in the source layer as a frame
-	const int32 NumBakedFrames = FMath::RoundToInt( FMath::Max( LayerSequenceLengthSeconds * LayerTimeCodesPerSecond + 1.0, 1.0 ) );
+	const int32 NumBakedFrames = FMath::RoundToInt(FMath::Max(LayerSequenceLengthSeconds * LayerTimeCodesPerSecond + 1.0, 1.0));
 	const double StageBakeIntervalTimeCodes = 1.0 * Offset.Scale;
 
 	IAnimationDataController& Controller = OutSkeletalAnimationAsset->GetController();
@@ -2199,26 +2286,26 @@ bool UsdToUnreal::ConvertSkelAnim(
 	// If we should transact, we'll already have a transaction from somewhere else. We should suppress this because
 	// it will also create a transaction when importing into UE assets, and the level sequence assets can emit some warnings about it
 	const bool bShouldTransact = false;
-	Controller.OpenBracket( LOCTEXT( "ImportUSDAnimData_Bracket", "Importing USD Animation Data" ), bShouldTransact );
+	Controller.OpenBracket(LOCTEXT("ImportUSDAnimData_Bracket", "Importing USD Animation Data"), bShouldTransact);
 	Controller.InitializeModel();
-	Controller.ResetModel( bShouldTransact );
+	Controller.ResetModel(bShouldTransact);
 
 	// Bake the animation for each frame.
 	// An alternative route would be to convert the time samples into TransformCurves, add them to UAnimSequence::RawCurveData,
-	// and then call UAnimSequence::BakeTrackCurvesToRawAnimation. Doing it this way provides a few benefits though: The main one is that the way with which
-	// UAnimSequence bakes can lead to artifacts on problematic joints (e.g. 90 degree rotation joints children of -1 scale joints, etc.) as it compounds the
-	// transformation with the rest pose. Another benefit is that that doing it this way lets us offload the interpolation to USD, so that it can do it
-	// however it likes, and we can just sample the joints at the target framerate
-	if ( NumJointTransformSamples >= 2 )
+	// and then call UAnimSequence::BakeTrackCurvesToRawAnimation. Doing it this way provides a few benefits though: The main one is that the way with
+	// which UAnimSequence bakes can lead to artifacts on problematic joints (e.g. 90 degree rotation joints children of -1 scale joints, etc.) as it
+	// compounds the transformation with the rest pose. Another benefit is that that doing it this way lets us offload the interpolation to USD, so
+	// that it can do it however it likes, and we can just sample the joints at the target framerate
+	if (NumJointTransformSamples >= 2)
 	{
 		FScopedUsdAllocs Allocs;
 
 		TArray<FRawAnimSequenceTrack> JointTracks;
 		JointTracks.SetNum(NumBonesInUE);
 
-		for ( int32 BoneIndex = 0; BoneIndex < NumBonesInUE; ++BoneIndex )
+		for (int32 BoneIndex = 0; BoneIndex < NumBonesInUE; ++BoneIndex)
 		{
-			FRawAnimSequenceTrack& JointTrack = JointTracks[ BoneIndex ];
+			FRawAnimSequenceTrack& JointTrack = JointTracks[BoneIndex];
 			JointTrack.PosKeys.Reserve(NumBakedFrames);
 			JointTrack.RotKeys.Reserve(NumBakedFrames);
 			JointTrack.ScaleKeys.Reserve(NumBakedFrames);
@@ -2227,12 +2314,12 @@ bool UsdToUnreal::ConvertSkelAnim(
 		FTransform RootMotionTransform;
 
 		pxr::VtArray<pxr::GfMatrix4d> UsdJointTransforms;
-		for ( int32 FrameIndex = 0; FrameIndex < NumBakedFrames; ++FrameIndex )
+		for (int32 FrameIndex = 0; FrameIndex < NumBakedFrames; ++FrameIndex)
 		{
 			const double StageFrameTimeCodes = StageStartTimeCode + FrameIndex * StageBakeIntervalTimeCodes;
 
-			InUsdSkeletonQuery.ComputeJointLocalTransforms( &UsdJointTransforms, StageFrameTimeCodes );
-			for ( int32 BoneIndex = 0; BoneIndex < NumBonesInUE; ++BoneIndex )
+			InUsdSkeletonQuery.ComputeJointLocalTransforms(&UsdJointTransforms, StageFrameTimeCodes);
+			for (int32 BoneIndex = 0; BoneIndex < NumBonesInUE; ++BoneIndex)
 			{
 				// UsdJointTransforms will never have a transform value for our AdditionalBones that we manually added
 				// (inserted "true root" bone), so we have to have this annoying check here to redirect the bone indices
@@ -2246,7 +2333,7 @@ bool UsdToUnreal::ConvertSkelAnim(
 
 				// Concatenate the root bone transform with the transform track actually present on the skel root as a
 				// whole
-				if ( BoneIndex == 0 )
+				if (BoneIndex == 0)
 				{
 					// We don't care about resetXformStack here: We'll always use the root motion prim's transform as
 					// a local transformation anyway
@@ -2259,28 +2346,34 @@ bool UsdToUnreal::ConvertSkelAnim(
 						OutResetTransformStack
 					);
 
-					if ( bSuccess )
+					if (bSuccess)
 					{
 						UEJointTransform = UEJointTransform * RootMotionTransform;
 					}
 				}
 
-				FRawAnimSequenceTrack& JointTrack = JointTracks[ BoneIndex ];
-				JointTrack.PosKeys.Add( FVector3f(UEJointTransform.GetTranslation()) );
-				JointTrack.RotKeys.Add( FQuat4f(UEJointTransform.GetRotation()) );
-				JointTrack.ScaleKeys.Add( FVector3f(UEJointTransform.GetScale3D()) );
+				FRawAnimSequenceTrack& JointTrack = JointTracks[BoneIndex];
+				JointTrack.PosKeys.Add(FVector3f(UEJointTransform.GetTranslation()));
+				JointTrack.RotKeys.Add(FQuat4f(UEJointTransform.GetRotation()));
+				JointTrack.ScaleKeys.Add(FVector3f(UEJointTransform.GetScale3D()));
 			}
 		}
 
-		for ( int32 BoneIndex = 0; BoneIndex < NumBonesInUE; ++BoneIndex )
+		for (int32 BoneIndex = 0; BoneIndex < NumBonesInUE; ++BoneIndex)
 		{
-			Controller.AddBoneCurve( BoneInfo[ BoneIndex ].Name, bShouldTransact );
-			Controller.SetBoneTrackKeys( BoneInfo[ BoneIndex ].Name, JointTracks[ BoneIndex ].PosKeys, JointTracks[ BoneIndex ].RotKeys, JointTracks[ BoneIndex ].ScaleKeys, bShouldTransact );
+			Controller.AddBoneCurve(BoneInfo[BoneIndex].Name, bShouldTransact);
+			Controller.SetBoneTrackKeys(
+				BoneInfo[BoneIndex].Name,
+				JointTracks[BoneIndex].PosKeys,
+				JointTracks[BoneIndex].RotKeys,
+				JointTracks[BoneIndex].ScaleKeys,
+				bShouldTransact
+			);
 		}
 	}
 
 	// Add float tracks to animate morph target weights
-	if ( InBlendShapes && InSkinningTargets )
+	if (InBlendShapes && InSkinningTargets)
 	{
 		FScopedUsdAllocs Allocs;
 
@@ -2289,29 +2382,31 @@ bool UsdToUnreal::ConvertSkelAnim(
 		pxr::VtTokenArray SkelAnimChannelOrder = UsdAnimQuery.GetBlendShapeOrder();
 		int32 NumSkelAnimChannels = SkelAnimChannelOrder.size();
 
-		if ( NumSkelAnimChannels > 0 )
+		if (NumSkelAnimChannels > 0)
 		{
 			// Create a float curve for each blend shape channel. These will be copied for each blend shape that uses it
 			// Don't remove redundant keys because if there are blendshapes with inbetweens that use this channel,
 			// we want to make sure that we don't miss the frames where the curve would have reached the exact weight of a blend shape
-			ERichCurveInterpMode CurveInterpMode = Stage.Get()->GetInterpolationType() == pxr::UsdInterpolationTypeHeld ? ERichCurveInterpMode::RCIM_Constant : ERichCurveInterpMode::RCIM_Linear;
+			ERichCurveInterpMode CurveInterpMode = Stage.Get()->GetInterpolationType() == pxr::UsdInterpolationTypeHeld
+													   ? ERichCurveInterpMode::RCIM_Constant
+													   : ERichCurveInterpMode::RCIM_Linear;
 			TArray<FRichCurve> SkelAnimChannelCurves;
-			SkelAnimChannelCurves.SetNum( NumSkelAnimChannels );
-			pxr::VtArray< float > WeightsForFrame;
-			for ( int32 FrameIndex = 0; FrameIndex < NumBakedFrames; ++FrameIndex )
+			SkelAnimChannelCurves.SetNum(NumSkelAnimChannels);
+			pxr::VtArray<float> WeightsForFrame;
+			for (int32 FrameIndex = 0; FrameIndex < NumBakedFrames; ++FrameIndex)
 			{
 				const double StageFrameTimeCodes = StageStartTimeCode + FrameIndex * StageBakeIntervalTimeCodes;
-				const double LayerFrameTimeCodes = ( StageFrameTimeCodes - Offset.Offset ) / Offset.Scale;
+				const double LayerFrameTimeCodes = (StageFrameTimeCodes - Offset.Offset) / Offset.Scale;
 				const double LayerFrameSeconds = LayerFrameTimeCodes / LayerTimeCodesPerSecond - LayerStartSeconds;
 
-				UsdAnimQuery.ComputeBlendShapeWeights( &WeightsForFrame, pxr::UsdTimeCode( StageFrameTimeCodes ) );
+				UsdAnimQuery.ComputeBlendShapeWeights(&WeightsForFrame, pxr::UsdTimeCode(StageFrameTimeCodes));
 
-				for ( int32 SkelAnimChannelIndex = 0; SkelAnimChannelIndex < NumSkelAnimChannels; ++SkelAnimChannelIndex )
+				for (int32 SkelAnimChannelIndex = 0; SkelAnimChannelIndex < NumSkelAnimChannels; ++SkelAnimChannelIndex)
 				{
-					FRichCurve& Curve = SkelAnimChannelCurves[ SkelAnimChannelIndex ];
+					FRichCurve& Curve = SkelAnimChannelCurves[SkelAnimChannelIndex];
 
-					FKeyHandle NewKeyHandle = Curve.AddKey( LayerFrameSeconds, WeightsForFrame[ SkelAnimChannelIndex ] );
-					Curve.SetKeyInterpMode( NewKeyHandle, CurveInterpMode );
+					FKeyHandle NewKeyHandle = Curve.AddKey(LayerFrameSeconds, WeightsForFrame[SkelAnimChannelIndex]);
+					Curve.SetKeyInterpMode(NewKeyHandle, CurveInterpMode);
 				}
 			}
 
@@ -2320,99 +2415,117 @@ bool UsdToUnreal::ConvertSkelAnim(
 			// Since we may need to switch variants to parse LODs, we could invalidate references to SkinningQuery objects, so we need
 			// to keep track of these by path and construct one whenever we need them
 			TArray<pxr::SdfPath> PathsToSkinnedPrims;
-			for ( const pxr::UsdSkelSkinningQuery& SkinningQuery : *InSkinningTargets )
+			for (const pxr::UsdSkelSkinningQuery& SkinningQuery : *InSkinningTargets)
 			{
 				// In USD, the skinning target need not be a mesh, but for Unreal we are only interested in skinning meshes
-				if ( pxr::UsdGeomMesh SkinningMesh = pxr::UsdGeomMesh( SkinningQuery.GetPrim() ) )
+				if (pxr::UsdGeomMesh SkinningMesh = pxr::UsdGeomMesh(SkinningQuery.GetPrim()))
 				{
-					PathsToSkinnedPrims.Add( SkinningMesh.GetPrim().GetPath() );
+					PathsToSkinnedPrims.Add(SkinningMesh.GetPrim().GetPath());
 				}
 			}
 
-			TFunction<bool( const pxr::UsdGeomMesh&, int32 )> CreateCurvesForLOD =
-			[ &InUsdSkeletonQuery, InBlendShapes, NumSkelAnimChannels, &SkelAnimChannelOrder, &SkelAnimChannelCurves, OutSkeletalAnimationAsset ]
-			( const pxr::UsdGeomMesh& LODMesh, int32 LODIndex )
+			TFunction<bool(const pxr::UsdGeomMesh&, int32)> CreateCurvesForLOD =
+				[&InUsdSkeletonQuery, InBlendShapes, NumSkelAnimChannels, &SkelAnimChannelOrder, &SkelAnimChannelCurves, OutSkeletalAnimationAsset](
+					const pxr::UsdGeomMesh& LODMesh,
+					int32 LODIndex
+				)
 			{
-				pxr::UsdSkelSkinningQuery SkinningQuery = UsdUtils::CreateSkinningQuery( LODMesh, InUsdSkeletonQuery );
-				if ( !SkinningQuery )
+				pxr::UsdSkelSkinningQuery SkinningQuery = UsdUtils::CreateSkinningQuery(LODMesh, InUsdSkeletonQuery);
+				if (!SkinningQuery)
 				{
-					return true; // Continue trying other LODs
+					return true;	// Continue trying other LODs
 				}
 
 				pxr::VtTokenArray MeshChannelOrder;
-				if ( !SkinningQuery.GetBlendShapeOrder( &MeshChannelOrder ) )
+				if (!SkinningQuery.GetBlendShapeOrder(&MeshChannelOrder))
 				{
 					return true;
 				}
 
 				pxr::SdfPathVector BlendShapeTargets;
 				const pxr::UsdRelationship& BlendShapeTargetsRel = SkinningQuery.GetBlendShapeTargetsRel();
-				BlendShapeTargetsRel.GetTargets( &BlendShapeTargets );
+				BlendShapeTargetsRel.GetTargets(&BlendShapeTargets);
 
 				// USD will already show a warning if this happens, so let's just continue
-				int32 NumMeshChannels = static_cast< int32 >( MeshChannelOrder.size() );
-				if ( NumMeshChannels != static_cast< int32 >( BlendShapeTargets.size() ) )
+				int32 NumMeshChannels = static_cast<int32>(MeshChannelOrder.size());
+				if (NumMeshChannels != static_cast<int32>(BlendShapeTargets.size()))
 				{
 					return true;
 				}
 
 				pxr::SdfPath MeshPath = SkinningQuery.GetPrim().GetPath();
-				for ( int32 MeshChannelIndex = 0; MeshChannelIndex < NumMeshChannels; ++MeshChannelIndex )
+				for (int32 MeshChannelIndex = 0; MeshChannelIndex < NumMeshChannels; ++MeshChannelIndex)
 				{
-					FString PrimaryBlendShapePath = UsdToUnreal::ConvertPath( BlendShapeTargets[ MeshChannelIndex ].MakeAbsolutePath( MeshPath ) );
+					FString PrimaryBlendShapePath = UsdToUnreal::ConvertPath(BlendShapeTargets[MeshChannelIndex].MakeAbsolutePath(MeshPath));
 
-					if ( const UsdUtils::FUsdBlendShape* FoundPrimaryBlendShape = InBlendShapes->Find( PrimaryBlendShapePath ) )
+					if (const UsdUtils::FUsdBlendShape* FoundPrimaryBlendShape = InBlendShapes->Find(PrimaryBlendShapePath))
 					{
 						// Find a float curve for the primary blend shape
 						FRichCurve* PrimaryBlendShapeCurve = nullptr;
-						pxr::TfToken& MeshChannel = MeshChannelOrder[ MeshChannelIndex ];
-						for ( int32 SkelAnimChannelIndex = 0; SkelAnimChannelIndex < NumSkelAnimChannels; ++SkelAnimChannelIndex )
+						pxr::TfToken& MeshChannel = MeshChannelOrder[MeshChannelIndex];
+						for (int32 SkelAnimChannelIndex = 0; SkelAnimChannelIndex < NumSkelAnimChannels; ++SkelAnimChannelIndex)
 						{
-							const pxr::TfToken& SkelAnimChannel = SkelAnimChannelOrder[ SkelAnimChannelIndex ];
-							if ( SkelAnimChannel == MeshChannel )
+							const pxr::TfToken& SkelAnimChannel = SkelAnimChannelOrder[SkelAnimChannelIndex];
+							if (SkelAnimChannel == MeshChannel)
 							{
-								PrimaryBlendShapeCurve = &SkelAnimChannelCurves[ SkelAnimChannelIndex ];
+								PrimaryBlendShapeCurve = &SkelAnimChannelCurves[SkelAnimChannelIndex];
 								break;
 							}
 						}
 
-						if ( !PrimaryBlendShapeCurve )
+						if (!PrimaryBlendShapeCurve)
 						{
 							FUsdLogManager::LogMessage(
 								EMessageSeverity::Warning,
-								FText::Format( LOCTEXT( "NoChannelForPrimary", "Could not find a float channel to apply to primary blend shape '{0}'" ),
-								FText::FromString( PrimaryBlendShapePath ) )
+								FText::Format(
+									LOCTEXT("NoChannelForPrimary", "Could not find a float channel to apply to primary blend shape '{0}'"),
+									FText::FromString(PrimaryBlendShapePath)
+								)
 							);
 							continue;
 						}
 
 						// Primary blend shape has no inbetweens, so we can just use the skel anim channel curve directly
-						if ( FoundPrimaryBlendShape->Inbetweens.Num() == 0 )
+						if (FoundPrimaryBlendShape->Inbetweens.Num() == 0)
 						{
-							SkelDataConversionImpl::SetFloatCurveData( OutSkeletalAnimationAsset, *FoundPrimaryBlendShape->Name, *PrimaryBlendShapeCurve );
+							SkelDataConversionImpl::SetFloatCurveData(
+								OutSkeletalAnimationAsset,
+								*FoundPrimaryBlendShape->Name,
+								*PrimaryBlendShapeCurve
+							);
 						}
-						// Blend shape has inbetweens --> Need to map these to multiple float curves. This can be different for each mesh, so we need to do it for each
+						// Blend shape has inbetweens --> Need to map these to multiple float curves. This can be different for each mesh, so we need
+						// to do it for each
 						else
 						{
-							TArray<FRichCurve> RemappedBlendShapeCurves = SkelDataConversionImpl::ResolveWeightsForBlendShapeCurve( *FoundPrimaryBlendShape, *PrimaryBlendShapeCurve );
-							if ( RemappedBlendShapeCurves.Num() != FoundPrimaryBlendShape->Inbetweens.Num() + 1 )
+							TArray<FRichCurve> RemappedBlendShapeCurves = SkelDataConversionImpl::ResolveWeightsForBlendShapeCurve(
+								*FoundPrimaryBlendShape,
+								*PrimaryBlendShapeCurve
+							);
+							if (RemappedBlendShapeCurves.Num() != FoundPrimaryBlendShape->Inbetweens.Num() + 1)
 							{
 								FUsdLogManager::LogMessage(
 									EMessageSeverity::Warning,
-									FText::Format( LOCTEXT( "FailedToRemapInbetweens", "Failed to remap inbetween float curves for blend shape '{0}'" ),
-									FText::FromString( PrimaryBlendShapePath ) )
+									FText::Format(
+										LOCTEXT("FailedToRemapInbetweens", "Failed to remap inbetween float curves for blend shape '{0}'"),
+										FText::FromString(PrimaryBlendShapePath)
+									)
 								);
 								continue;
 							}
 
-							SkelDataConversionImpl::SetFloatCurveData( OutSkeletalAnimationAsset, *FoundPrimaryBlendShape->Name, RemappedBlendShapeCurves[ 0 ] );
+							SkelDataConversionImpl::SetFloatCurveData(
+								OutSkeletalAnimationAsset,
+								*FoundPrimaryBlendShape->Name,
+								RemappedBlendShapeCurves[0]
+							);
 
-							for ( int32 InbetweenIndex = 0; InbetweenIndex < FoundPrimaryBlendShape->Inbetweens.Num(); ++InbetweenIndex )
+							for (int32 InbetweenIndex = 0; InbetweenIndex < FoundPrimaryBlendShape->Inbetweens.Num(); ++InbetweenIndex)
 							{
-								const UsdUtils::FUsdBlendShapeInbetween& Inbetween = FoundPrimaryBlendShape->Inbetweens[ InbetweenIndex ];
-								const FRichCurve& InbetweenCurve = RemappedBlendShapeCurves[ InbetweenIndex + 1 ]; // Index 0 is the primary
+								const UsdUtils::FUsdBlendShapeInbetween& Inbetween = FoundPrimaryBlendShape->Inbetweens[InbetweenIndex];
+								const FRichCurve& InbetweenCurve = RemappedBlendShapeCurves[InbetweenIndex + 1];	// Index 0 is the primary
 
-								SkelDataConversionImpl::SetFloatCurveData( OutSkeletalAnimationAsset, *Inbetween.Name, InbetweenCurve );
+								SkelDataConversionImpl::SetFloatCurveData(OutSkeletalAnimationAsset, *Inbetween.Name, InbetweenCurve);
 							}
 						}
 					}
@@ -2421,45 +2534,46 @@ bool UsdToUnreal::ConvertSkelAnim(
 				return true;
 			};
 
-			for ( const pxr::SdfPath& SkinnedPrimPath : PathsToSkinnedPrims )
+			for (const pxr::SdfPath& SkinnedPrimPath : PathsToSkinnedPrims)
 			{
-				pxr::UsdPrim SkinnedPrim = Stage.Get()->GetPrimAtPath( SkinnedPrimPath );
-				if ( !SkinnedPrim )
+				pxr::UsdPrim SkinnedPrim = Stage.Get()->GetPrimAtPath(SkinnedPrimPath);
+				if (!SkinnedPrim)
 				{
 					continue;
 				}
 
-				pxr::UsdGeomMesh SkinnedMesh{ SkinnedPrim };
-				if ( !SkinnedMesh )
+				pxr::UsdGeomMesh SkinnedMesh{SkinnedPrim};
+				if (!SkinnedMesh)
 				{
 					continue;
 				}
 
 				pxr::UsdPrim ParentPrim = SkinnedMesh.GetPrim().GetParent();
-				FString ParentPrimPath = UsdToUnreal::ConvertPath( ParentPrim.GetPath() );
+				FString ParentPrimPath = UsdToUnreal::ConvertPath(ParentPrim.GetPath());
 
 				bool bInterpretedLODs = false;
-				if ( bInInterpretLODs && ParentPrim && !ProcessedLODParentPaths.Contains( ParentPrimPath ) )
+				if (bInInterpretLODs && ParentPrim && !ProcessedLODParentPaths.Contains(ParentPrimPath))
 				{
-					// At the moment we only consider a single mesh per variant, so if multiple meshes tell us to process the same parent prim, we skip.
-					// This check would also prevent us from getting in here in case we just have many meshes children of a same prim, outside
+					// At the moment we only consider a single mesh per variant, so if multiple meshes tell us to process the same parent prim, we
+					// skip. This check would also prevent us from getting in here in case we just have many meshes children of a same prim, outside
 					// of a variant. In this case they don't fit the "one mesh per variant" pattern anyway, and we want to fallback to ignoring LODs
-					ProcessedLODParentPaths.Add( ParentPrimPath );
+					ProcessedLODParentPaths.Add(ParentPrimPath);
 
 					// WARNING: After this is called, references to objects that were inside any of the LOD Meshes will be invalidated!
-					bInterpretedLODs = UsdUtils::IterateLODMeshes( ParentPrim, CreateCurvesForLOD );
+					bInterpretedLODs = UsdUtils::IterateLODMeshes(ParentPrim, CreateCurvesForLOD);
 				}
 
-				if ( !bInterpretedLODs )
+				if (!bInterpretedLODs)
 				{
 					// Refresh reference to this prim as it could have been inside a variant that was temporarily switched by IterateLODMeshes
-					CreateCurvesForLOD( SkinnedMesh, 0 );
+					CreateCurvesForLOD(SkinnedMesh, 0);
 				}
 			}
 		}
 	}
 
-	OutSkeletalAnimationAsset->Interpolation = Stage.Get()->GetInterpolationType() == pxr::UsdInterpolationTypeHeld ? EAnimInterpolationType::Step : EAnimInterpolationType::Linear;
+	OutSkeletalAnimationAsset->Interpolation = Stage.Get()->GetInterpolationType() == pxr::UsdInterpolationTypeHeld ? EAnimInterpolationType::Step
+																													: EAnimInterpolationType::Linear;
 	OutSkeletalAnimationAsset->ImportFileFramerate = LayerTimeCodesPerSecond;
 	OutSkeletalAnimationAsset->ImportResampleFramerate = LayerTimeCodesPerSecond;
 
@@ -2467,13 +2581,13 @@ bool UsdToUnreal::ConvertSkelAnim(
 	Controller.SetFrameRate(FrameRate, bShouldTransact);
 	const FFrameNumber FrameNumber = FrameRate.AsFrameNumber(LayerSequenceLengthSeconds);
 	Controller.SetNumberOfFrames(FrameNumber, bShouldTransact);
-	Controller.NotifyPopulated(); // This call is important to get the controller to not use the sampling frequency as framerate
-	Controller.CloseBracket( bShouldTransact );
+	Controller.NotifyPopulated();	 // This call is important to get the controller to not use the sampling frequency as framerate
+	Controller.CloseBracket(bShouldTransact);
 
 	OutSkeletalAnimationAsset->PostEditChange();
 	OutSkeletalAnimationAsset->MarkPackageDirty();
 
-	if ( OutStartOffsetSeconds )
+	if (OutStartOffsetSeconds)
 	{
 		// We don't want to store just StartSeconds here, because part of that may be because the layer itself
 		// has an offset/scale within the stage. In OutStartOffsetSeconds we need to store the start of the animation
@@ -2485,14 +2599,27 @@ bool UsdToUnreal::ConvertSkelAnim(
 	return true;
 }
 
-bool UsdToUnreal::ConvertBlendShape( const pxr::UsdSkelBlendShape& UsdBlendShape, const FUsdStageInfo& StageInfo, uint32 PointIndexOffset, TSet<FString>& UsedMorphTargetNames, UsdUtils::FBlendShapeMap& OutBlendShapes )
+bool UsdToUnreal::ConvertBlendShape(
+	const pxr::UsdSkelBlendShape& UsdBlendShape,
+	const FUsdStageInfo& StageInfo,
+	uint32 PointIndexOffset,
+	TSet<FString>& UsedMorphTargetNames,
+	UsdUtils::FBlendShapeMap& OutBlendShapes
+)
 {
 	return ConvertBlendShape(UsdBlendShape, StageInfo, 0, PointIndexOffset, UsedMorphTargetNames, OutBlendShapes);
 }
 
-bool UsdToUnreal::ConvertBlendShape( const pxr::UsdSkelBlendShape& UsdBlendShape, const FUsdStageInfo& StageInfo, int32 LODIndex, uint32 PointIndexOffset, TSet<FString>& UsedMorphTargetNames, UsdUtils::FBlendShapeMap& OutBlendShapes )
+bool UsdToUnreal::ConvertBlendShape(
+	const pxr::UsdSkelBlendShape& UsdBlendShape,
+	const FUsdStageInfo& StageInfo,
+	int32 LODIndex,
+	uint32 PointIndexOffset,
+	TSet<FString>& UsedMorphTargetNames,
+	UsdUtils::FBlendShapeMap& OutBlendShapes
+)
 {
-	TRACE_CPUPROFILER_EVENT_SCOPE( UsdToUnreal::ConvertBlendShape );
+	TRACE_CPUPROFILER_EVENT_SCOPE(UsdToUnreal::ConvertBlendShape);
 
 	if (!UsdBlendShape)
 	{
@@ -2502,69 +2629,70 @@ bool UsdToUnreal::ConvertBlendShape( const pxr::UsdSkelBlendShape& UsdBlendShape
 	FScopedUsdAllocs Allocs;
 
 	pxr::UsdAttribute OffsetsAttr = UsdBlendShape.GetOffsetsAttr();
-	pxr::VtArray< pxr::GfVec3f > Offsets;
-	OffsetsAttr.Get( &Offsets );
+	pxr::VtArray<pxr::GfVec3f> Offsets;
+	OffsetsAttr.Get(&Offsets);
 
 	pxr::UsdAttribute IndicesAttr = UsdBlendShape.GetPointIndicesAttr();
-	pxr::VtArray< int > PointIndices;
-	IndicesAttr.Get( &PointIndices );
+	pxr::VtArray<int> PointIndices;
+	IndicesAttr.Get(&PointIndices);
 
 	pxr::UsdAttribute NormalsAttr = UsdBlendShape.GetNormalOffsetsAttr();
-	pxr::VtArray< pxr::GfVec3f > Normals;
-	NormalsAttr.Get( &Normals );
+	pxr::VtArray<pxr::GfVec3f> Normals;
+	NormalsAttr.Get(&Normals);
 
 	// We need to guarantee blend shapes have unique names because these will be used as UMorphTarget names
 	// Note that we can't just use the prim path here and need an index to guarantee uniqueness,
 	// because although the path is usually unique, USD has case sensitive paths and the FNames of the
 	// UMorphTargets are case insensitive
 	FString PrimaryName = UsdUtils::GetUniqueName(
-		IUsdClassesModule::SanitizeObjectName( UsdToUnreal::ConvertString( UsdBlendShape.GetPrim().GetName() ) ),
-		UsedMorphTargetNames );
-	FString PrimaryPath = UsdToUnreal::ConvertPath( UsdBlendShape.GetPrim().GetPath() );
-	if ( UsdUtils::FUsdBlendShape* ExistingBlendShape = OutBlendShapes.Find( PrimaryPath ) )
+		IUsdClassesModule::SanitizeObjectName(UsdToUnreal::ConvertString(UsdBlendShape.GetPrim().GetName())),
+		UsedMorphTargetNames
+	);
+	FString PrimaryPath = UsdToUnreal::ConvertPath(UsdBlendShape.GetPrim().GetPath());
+	if (UsdUtils::FUsdBlendShape* ExistingBlendShape = OutBlendShapes.Find(PrimaryPath))
 	{
-		ExistingBlendShape->LODIndicesThatUseThis.Add( LODIndex );
+		ExistingBlendShape->LODIndicesThatUseThis.Add(LODIndex);
 		return true;
 	}
 
 	UsdUtils::FUsdBlendShape PrimaryBlendShape;
-	if ( !SkelDataConversionImpl::CreateUsdBlendShape( PrimaryName, Offsets, Normals, PointIndices, StageInfo, PointIndexOffset, LODIndex, PrimaryBlendShape ) )
+	if (!SkelDataConversionImpl::
+			CreateUsdBlendShape(PrimaryName, Offsets, Normals, PointIndices, StageInfo, PointIndexOffset, LODIndex, PrimaryBlendShape))
 	{
 		return false;
 	}
-	UsedMorphTargetNames.Add( PrimaryBlendShape.Name );
+	UsedMorphTargetNames.Add(PrimaryBlendShape.Name);
 
 	UsdUtils::FBlendShapeMap InbetweenBlendShapes;
-	for ( const pxr::UsdSkelInbetweenShape& Inbetween : UsdBlendShape.GetInbetweens() )
+	for (const pxr::UsdSkelInbetweenShape& Inbetween : UsdBlendShape.GetInbetweens())
 	{
-		if ( !Inbetween )
+		if (!Inbetween)
 		{
 			continue;
 		}
 
 		float Weight = 0.0f;
-		if ( !Inbetween.GetWeight( &Weight ) )
+		if (!Inbetween.GetWeight(&Weight))
 		{
 			continue;
 		}
 
-		FString OrigInbetweenName = UsdToUnreal::ConvertString( Inbetween.GetAttr().GetName() );
-		FString InbetweenPath = FString::Printf(TEXT("%s_%s"), *PrimaryPath, *OrigInbetweenName );
+		FString OrigInbetweenName = UsdToUnreal::ConvertString(Inbetween.GetAttr().GetName());
+		FString InbetweenPath = FString::Printf(TEXT("%s_%s"), *PrimaryPath, *OrigInbetweenName);
 		FString InbetweenName = UsdUtils::GetUniqueName(
-			IUsdClassesModule::SanitizeObjectName( FPaths::GetCleanFilename( InbetweenPath ) ),
-			UsedMorphTargetNames );
+			IUsdClassesModule::SanitizeObjectName(FPaths::GetCleanFilename(InbetweenPath)),
+			UsedMorphTargetNames
+		);
 
-		if ( Weight > 1.0f || Weight < 0.0f || FMath::IsNearlyZero(Weight) || FMath::IsNearlyEqual(Weight, 1.0f) )
+		if (Weight > 1.0f || Weight < 0.0f || FMath::IsNearlyZero(Weight) || FMath::IsNearlyEqual(Weight, 1.0f))
 		{
-			//UE_LOG(LogUsd, Warning, TEXT("Inbetween shape '%s' for blend shape '%s' has invalid weight '%f' and will be ignored. Valid weights are > 0.0f and < 1.0f!"),
-			//	*OrigInbetweenName,
-			//	*PrimaryPath,
-			//	Weight);
+			// UE_LOG(LogUsd, Warning, TEXT("Inbetween shape '%s' for blend shape '%s' has invalid weight '%f' and will be ignored. Valid weights are
+			// > 0.0f and < 1.0f!"), 	*OrigInbetweenName, 	*PrimaryPath, 	Weight);
 			continue;
 		}
 
-		pxr::VtArray< pxr::GfVec3f > InbetweenPointsOffsets;
-		pxr::VtArray< pxr::GfVec3f > InbetweenNormalOffsets;
+		pxr::VtArray<pxr::GfVec3f> InbetweenPointsOffsets;
+		pxr::VtArray<pxr::GfVec3f> InbetweenNormalOffsets;
 
 		Inbetween.GetOffsets(&InbetweenPointsOffsets);
 		Inbetween.GetNormalOffsets(&InbetweenNormalOffsets);
@@ -2572,12 +2700,21 @@ bool UsdToUnreal::ConvertBlendShape( const pxr::UsdSkelBlendShape& UsdBlendShape
 		// Create separate blend shape for the inbetween
 		// Now how the inbetween always shares the same point indices as the parent
 		UsdUtils::FUsdBlendShape InbetweenShape;
-		if ( !SkelDataConversionImpl::CreateUsdBlendShape( InbetweenName, InbetweenPointsOffsets, InbetweenNormalOffsets, PointIndices, StageInfo, PointIndexOffset, LODIndex, InbetweenShape ) )
+		if (!SkelDataConversionImpl::CreateUsdBlendShape(
+				InbetweenName,
+				InbetweenPointsOffsets,
+				InbetweenNormalOffsets,
+				PointIndices,
+				StageInfo,
+				PointIndexOffset,
+				LODIndex,
+				InbetweenShape
+			))
 		{
 			continue;
 		}
-		UsedMorphTargetNames.Add( InbetweenShape.Name );
-		InbetweenBlendShapes.Add( InbetweenPath, InbetweenShape );
+		UsedMorphTargetNames.Add(InbetweenShape.Name);
+		InbetweenBlendShapes.Add(InbetweenPath, InbetweenShape);
 
 		// Keep track of it in the PrimaryBlendShape so we can resolve weights later
 		UsdUtils::FUsdBlendShapeInbetween& ConvertedInbetween = PrimaryBlendShape.Inbetweens.Emplace_GetRef();
@@ -2586,13 +2723,15 @@ bool UsdToUnreal::ConvertBlendShape( const pxr::UsdSkelBlendShape& UsdBlendShape
 	}
 
 	// Sort according to weight so they're easier to resolve later
-	PrimaryBlendShape.Inbetweens.Sort([](const UsdUtils::FUsdBlendShapeInbetween& Lhs, const UsdUtils::FUsdBlendShapeInbetween& Rhs)
-	{
-		return Lhs.InbetweenWeight < Rhs.InbetweenWeight;
-	});
+	PrimaryBlendShape.Inbetweens.Sort(
+		[](const UsdUtils::FUsdBlendShapeInbetween& Lhs, const UsdUtils::FUsdBlendShapeInbetween& Rhs)
+		{
+			return Lhs.InbetweenWeight < Rhs.InbetweenWeight;
+		}
+	);
 
-	OutBlendShapes.Add( PrimaryPath, PrimaryBlendShape );
-	OutBlendShapes.Append( MoveTemp(InbetweenBlendShapes) );
+	OutBlendShapes.Add(PrimaryPath, PrimaryBlendShape);
+	OutBlendShapes.Append(MoveTemp(InbetweenBlendShapes));
 
 	return true;
 }
@@ -2606,9 +2745,9 @@ USkeletalMesh* UsdToUnreal::GetSkeletalMeshFromImportData(
 	const FName& SkeletonName
 )
 {
-	TRACE_CPUPROFILER_EVENT_SCOPE( UsdToUnreal::GetSkeletalMeshFromImportData );
+	TRACE_CPUPROFILER_EVENT_SCOPE(UsdToUnreal::GetSkeletalMeshFromImportData);
 
-	if ( LODIndexToSkeletalMeshImportData.Num() == 0 || InSkeletonBones.Num() == 0 )
+	if (LODIndexToSkeletalMeshImportData.Num() == 0 || InSkeletonBones.Num() == 0)
 	{
 		return nullptr;
 	}
@@ -2621,17 +2760,21 @@ USkeletalMesh* UsdToUnreal::GetSkeletalMeshFromImportData(
 		USkeletalMesh::StaticClass(),
 		*IUsdClassesModule::SanitizeObjectName(MeshName.ToString())
 	);
-	USkeletalMesh* SkeletalMesh = NewObject<USkeletalMesh>(GetTransientPackage(), UniqueMeshName, ObjectFlags | EObjectFlags::RF_Public | EObjectFlags::RF_Transient);
+	USkeletalMesh* SkeletalMesh = NewObject<USkeletalMesh>(
+		GetTransientPackage(),
+		UniqueMeshName,
+		ObjectFlags | EObjectFlags::RF_Public | EObjectFlags::RF_Transient
+	);
 
 	// Process reference skeleton from import data
 	int32 SkeletalDepth = 0;
 	FSkeletalMeshImportData DummyData;
 	DummyData.RefBonesBinary = InSkeletonBones;
-	if ( !SkeletalMeshImportUtils::ProcessImportMeshSkeleton( SkeletalMesh->GetSkeleton(), SkeletalMesh->GetRefSkeleton(), SkeletalDepth, DummyData ) )
+	if (!SkeletalMeshImportUtils::ProcessImportMeshSkeleton(SkeletalMesh->GetSkeleton(), SkeletalMesh->GetRefSkeleton(), SkeletalDepth, DummyData))
 	{
 		return nullptr;
 	}
-	if ( SkeletalMesh->GetRefSkeleton().GetRawBoneNum() == 0 )
+	if (SkeletalMesh->GetRefSkeleton().GetRawBoneNum() == 0)
 	{
 		SkeletalMesh->MarkAsGarbage();
 		return nullptr;
@@ -2640,35 +2783,36 @@ USkeletalMesh* UsdToUnreal::GetSkeletalMeshFromImportData(
 	// This prevents PostEditChange calls when it is alive, also ensuring it is called once when we return from this function.
 	// This is required because we must ensure the morphtargets are in the SkeletalMesh before the first call to PostEditChange(),
 	// or else they will be effectively discarded
-	FScopedSkeletalMeshPostEditChange ScopedPostEditChange( SkeletalMesh );
-	SkeletalMesh->PreEditChange( nullptr );
+	FScopedSkeletalMeshPostEditChange ScopedPostEditChange(SkeletalMesh);
+	SkeletalMesh->PreEditChange(nullptr);
 
 	// Create initial bounding box based on expanded version of reference pose for meshes without physics assets
 	const FSkeletalMeshImportData& LowestLOD = LODIndexToSkeletalMeshImportData[0];
-	FBox3f BoundingBox( LowestLOD.Points.GetData(), LowestLOD.Points.Num() );
+	FBox3f BoundingBox(LowestLOD.Points.GetData(), LowestLOD.Points.Num());
 	FBox3f Temp = BoundingBox;
-	FVector3f MidMesh = 0.5f*(Temp.Min + Temp.Max);
-	BoundingBox.Min = Temp.Min + 1.0f*(Temp.Min - MidMesh);
-	BoundingBox.Max = Temp.Max + 1.0f*(Temp.Max - MidMesh);
-	BoundingBox.Min[2] = Temp.Min[2] + 0.1f*(Temp.Min[2] - MidMesh[2]);
+	FVector3f MidMesh = 0.5f * (Temp.Min + Temp.Max);
+	BoundingBox.Min = Temp.Min + 1.0f * (Temp.Min - MidMesh);
+	BoundingBox.Max = Temp.Max + 1.0f * (Temp.Max - MidMesh);
+	BoundingBox.Min[2] = Temp.Min[2] + 0.1f * (Temp.Min[2] - MidMesh[2]);
 	const FVector3f BoundingBoxSize = BoundingBox.GetSize();
-	if ( LowestLOD.Points.Num() > 2 && BoundingBoxSize.X < THRESH_POINTS_ARE_SAME && BoundingBoxSize.Y < THRESH_POINTS_ARE_SAME && BoundingBoxSize.Z < THRESH_POINTS_ARE_SAME )
+	if (LowestLOD.Points.Num() > 2 && BoundingBoxSize.X < THRESH_POINTS_ARE_SAME && BoundingBoxSize.Y < THRESH_POINTS_ARE_SAME
+		&& BoundingBoxSize.Z < THRESH_POINTS_ARE_SAME)
 	{
 		return nullptr;
 	}
 
 #if WITH_EDITOR
-	IMeshUtilities& MeshUtilities = FModuleManager::Get().LoadModuleChecked<IMeshUtilities>( "MeshUtilities" );
-#endif // WITH_EDITOR
+	IMeshUtilities& MeshUtilities = FModuleManager::Get().LoadModuleChecked<IMeshUtilities>("MeshUtilities");
+#endif	  // WITH_EDITOR
 
 	FSkeletalMeshModel* ImportedResource = SkeletalMesh->GetImportedModel();
 	ImportedResource->LODModels.Empty();
 	SkeletalMesh->ResetLODInfo();
 	bool bHasVertexColors = false;
-	for ( int32 LODIndex = 0; LODIndex < LODIndexToSkeletalMeshImportData.Num(); ++LODIndex )
+	for (int32 LODIndex = 0; LODIndex < LODIndexToSkeletalMeshImportData.Num(); ++LODIndex)
 	{
 		FSkeletalMeshImportData& LODImportData = LODIndexToSkeletalMeshImportData[LODIndex];
-		ImportedResource->LODModels.Add( new FSkeletalMeshLODModel() );
+		ImportedResource->LODModels.Add(new FSkeletalMeshLODModel());
 		FSkeletalMeshLODModel& LODModel = ImportedResource->LODModels.Last();
 
 		// Process bones influence (normalization and optimization) (optional)
@@ -2685,10 +2829,10 @@ USkeletalMesh* UsdToUnreal::GetSkeletalMeshFromImportData(
 		LODModel.NumTexCoords = FMath::Max<uint32>(1, LODImportData.NumTexCoords);
 
 		// Data needed by BuildSkeletalMesh
-		LODImportData.PointToRawMap.AddUninitialized( LODImportData.Points.Num() );
-		for ( int32 PointIndex = 0; PointIndex < LODImportData.Points.Num(); ++PointIndex )
+		LODImportData.PointToRawMap.AddUninitialized(LODImportData.Points.Num());
+		for (int32 PointIndex = 0; PointIndex < LODImportData.Points.Num(); ++PointIndex)
 		{
-			LODImportData.PointToRawMap[ PointIndex ] = PointIndex;
+			LODImportData.PointToRawMap[PointIndex] = PointIndex;
 		}
 
 		TArray<FVector3f> LODPoints;
@@ -2696,37 +2840,49 @@ USkeletalMesh* UsdToUnreal::GetSkeletalMeshFromImportData(
 		TArray<SkeletalMeshImportData::FMeshFace> LODFaces;
 		TArray<SkeletalMeshImportData::FVertInfluence> LODInfluences;
 		TArray<int32> LODPointToRawMap;
-		LODImportData.CopyLODImportData( LODPoints, LODWedges, LODFaces, LODInfluences, LODPointToRawMap );
+		LODImportData.CopyLODImportData(LODPoints, LODWedges, LODFaces, LODInfluences, LODPointToRawMap);
 
 #if WITH_EDITOR
 		IMeshUtilities::MeshBuildOptions BuildOptions;
 		BuildOptions.TargetPlatform = GetTargetPlatformManagerRef().GetRunningTargetPlatform();
 		// #ueent_todo: Normals and tangents shouldn't need to be recomputed when they are retrieved from USD
-		//BuildOptions.bComputeNormals = !LODImportData.bHasNormals;
-		//BuildOptions.bComputeTangents = !LODImportData.bHasTangents;
+		// BuildOptions.bComputeNormals = !LODImportData.bHasNormals;
+		// BuildOptions.bComputeTangents = !LODImportData.bHasTangents;
 		BuildOptions.bUseMikkTSpace = true;
 
 		TArray<FText> WarningMessages;
 		TArray<FName> WarningNames;
 
-		bool bBuildSuccess = MeshUtilities.BuildSkeletalMesh(LODModel, SkeletalMesh->GetPathName(), SkeletalMesh->GetRefSkeleton(), LODInfluences, LODWedges, LODFaces, LODPoints, LODPointToRawMap, BuildOptions, &WarningMessages, &WarningNames );
+		bool bBuildSuccess = MeshUtilities.BuildSkeletalMesh(
+			LODModel,
+			SkeletalMesh->GetPathName(),
+			SkeletalMesh->GetRefSkeleton(),
+			LODInfluences,
+			LODWedges,
+			LODFaces,
+			LODPoints,
+			LODPointToRawMap,
+			BuildOptions,
+			&WarningMessages,
+			&WarningNames
+		);
 
-		for ( int32 WarningIndex = 0; WarningIndex < FMath::Max( WarningMessages.Num(), WarningNames.Num() ); ++WarningIndex )
+		for (int32 WarningIndex = 0; WarningIndex < FMath::Max(WarningMessages.Num(), WarningNames.Num()); ++WarningIndex)
 		{
-			const FText& Text = WarningMessages.IsValidIndex( WarningIndex ) ? WarningMessages[ WarningIndex ] : FText::GetEmpty();
-			const FName& Name = WarningNames.IsValidIndex( WarningIndex ) ? WarningNames[ WarningIndex ] : NAME_None;
+			const FText& Text = WarningMessages.IsValidIndex(WarningIndex) ? WarningMessages[WarningIndex] : FText::GetEmpty();
+			const FName& Name = WarningNames.IsValidIndex(WarningIndex) ? WarningNames[WarningIndex] : NAME_None;
 
-			if ( bBuildSuccess )
+			if (bBuildSuccess)
 			{
-				UE_LOG( LogUsd, Warning, TEXT( "Warning when trying to build skeletal mesh from USD: '%s': '%s'" ), *Name.ToString(), *Text.ToString() );
+				UE_LOG(LogUsd, Warning, TEXT("Warning when trying to build skeletal mesh from USD: '%s': '%s'"), *Name.ToString(), *Text.ToString());
 			}
 			else
 			{
-				UE_LOG( LogUsd, Error, TEXT( "Error when trying to build skeletal mesh from USD: '%s': '%s'" ), *Name.ToString(), *Text.ToString() );
+				UE_LOG(LogUsd, Error, TEXT("Error when trying to build skeletal mesh from USD: '%s': '%s'"), *Name.ToString(), *Text.ToString());
 			}
 		}
 
-		if ( !bBuildSuccess )
+		if (!bBuildSuccess)
 		{
 			SkeletalMesh->MarkAsGarbage();
 			return nullptr;
@@ -2734,11 +2890,11 @@ USkeletalMesh* UsdToUnreal::GetSkeletalMeshFromImportData(
 
 		// This is important because it will fill in the LODModel's RawSkeletalMeshBulkDataID,
 		// which is the part of the skeletal mesh's DDC key that is affected by the actual mesh data
-		SkeletalMesh->SaveLODImportedData( LODIndex, LODImportData );
-#endif // WITH_EDITOR
+		SkeletalMesh->SaveLODImportedData(LODIndex, LODImportData);
+#endif	  // WITH_EDITOR
 	}
 
-	SkeletalMesh->SetImportedBounds( FBoxSphereBounds( (FBox)BoundingBox ) );
+	SkeletalMesh->SetImportedBounds(FBoxSphereBounds((FBox)BoundingBox));
 	SkeletalMesh->SetHasVertexColors(bHasVertexColors);
 	SkeletalMesh->SetVertexColorGuid(SkeletalMesh->GetHasVertexColors() ? FGuid::NewGuid() : FGuid());
 	SkeletalMesh->CalculateInvRefMatrices();
@@ -2749,7 +2905,11 @@ USkeletalMesh* UsdToUnreal::GetSkeletalMeshFromImportData(
 		USkeleton::StaticClass(),
 		*IUsdClassesModule::SanitizeObjectName(SkeletonName.ToString())
 	);
-	USkeleton* Skeleton = NewObject<USkeleton>( GetTransientPackage(), UniqueSkeletonName, ObjectFlags | EObjectFlags::RF_Public | EObjectFlags::RF_Transient );
+	USkeleton* Skeleton = NewObject<USkeleton>(
+		GetTransientPackage(),
+		UniqueSkeletonName,
+		ObjectFlags | EObjectFlags::RF_Public | EObjectFlags::RF_Transient
+	);
 	Skeleton->MergeAllBonesToBoneTree(SkeletalMesh);
 	Skeleton->SetPreviewMesh(SkeletalMesh);
 	SkeletalMesh->SetSkeleton(Skeleton);
@@ -2759,30 +2919,35 @@ USkeletalMesh* UsdToUnreal::GetSkeletalMeshFromImportData(
 	return SkeletalMesh;
 }
 
-#endif // #if USE_USD_SDK && WITH_EDITOR
+#endif	  // #if USE_USD_SDK && WITH_EDITOR
 
-void UsdUtils::ResolveWeightsForBlendShape( const UsdUtils::FUsdBlendShape& InBlendShape, float InWeight, float& OutMainWeight, TArray<float>& OutInbetweenWeights )
+void UsdUtils::ResolveWeightsForBlendShape(
+	const UsdUtils::FUsdBlendShape& InBlendShape,
+	float InWeight,
+	float& OutMainWeight,
+	TArray<float>& OutInbetweenWeights
+)
 {
 	OutMainWeight = 0.0f;
 	int32 NumInbetweens = InBlendShape.Inbetweens.Num();
-	if ( NumInbetweens == 0 )
+	if (NumInbetweens == 0)
 	{
 		OutMainWeight = InWeight;
 		return;
 	}
 
-	OutInbetweenWeights.SetNumUninitialized( NumInbetweens );
-	for ( float& OutInbetweenWeight : OutInbetweenWeights )
+	OutInbetweenWeights.SetNumUninitialized(NumInbetweens);
+	for (float& OutInbetweenWeight : OutInbetweenWeights)
 	{
 		OutInbetweenWeight = 0.0f;
 	}
 
-	if ( FMath::IsNearlyEqual( InWeight, 0.0f ) )
+	if (FMath::IsNearlyEqual(InWeight, 0.0f))
 	{
 		OutMainWeight = 0.0f;
 		return;
 	}
-	else if ( FMath::IsNearlyEqual( InWeight, 1.0f ) )
+	else if (FMath::IsNearlyEqual(InWeight, 1.0f))
 	{
 		OutMainWeight = 1.0f;
 		return;
@@ -2790,61 +2955,65 @@ void UsdUtils::ResolveWeightsForBlendShape( const UsdUtils::FUsdBlendShape& InBl
 
 	// Note how we don't care if UpperIndex/LowerIndex are beyond the bounds of the array here,
 	// as that signals when we're above/below all inbetweens
-	int32 UpperIndex = Algo::UpperBoundBy( InBlendShape.Inbetweens, InWeight, []( const UsdUtils::FUsdBlendShapeInbetween& Inbetween )
-	{
-		return Inbetween.InbetweenWeight;
-	} );
+	int32 UpperIndex = Algo::UpperBoundBy(
+		InBlendShape.Inbetweens,
+		InWeight,
+		[](const UsdUtils::FUsdBlendShapeInbetween& Inbetween)
+		{
+			return Inbetween.InbetweenWeight;
+		}
+	);
 	int32 LowerIndex = UpperIndex - 1;
 
 	float UpperWeight = 1.0f;
-	if ( UpperIndex <= NumInbetweens - 1 )
+	if (UpperIndex <= NumInbetweens - 1)
 	{
-		UpperWeight = InBlendShape.Inbetweens[ UpperIndex ].InbetweenWeight;
+		UpperWeight = InBlendShape.Inbetweens[UpperIndex].InbetweenWeight;
 	}
 
 	float LowerWeight = 0.0f;
-	if ( LowerIndex >= 0 )
+	if (LowerIndex >= 0)
 	{
-		LowerWeight = InBlendShape.Inbetweens[ LowerIndex ].InbetweenWeight;
+		LowerWeight = InBlendShape.Inbetweens[LowerIndex].InbetweenWeight;
 	}
 
-	UpperWeight = ( InWeight - LowerWeight ) / ( UpperWeight - LowerWeight );
-	LowerWeight = ( 1.0f - UpperWeight );
+	UpperWeight = (InWeight - LowerWeight) / (UpperWeight - LowerWeight);
+	LowerWeight = (1.0f - UpperWeight);
 
 	// We're between upper inbetween and the 1.0 weight
-	if ( UpperIndex > NumInbetweens - 1 )
+	if (UpperIndex > NumInbetweens - 1)
 	{
 		OutMainWeight = UpperWeight;
-		OutInbetweenWeights[ NumInbetweens - 1 ] = LowerWeight;
+		OutInbetweenWeights[NumInbetweens - 1] = LowerWeight;
 	}
 	// We're between 0.0 and the first inbetween weight
-	else if ( LowerIndex < 0 )
+	else if (LowerIndex < 0)
 	{
 		OutMainWeight = 0;
-		OutInbetweenWeights[ 0 ] = UpperWeight;
+		OutInbetweenWeights[0] = UpperWeight;
 	}
 	// We're between two inbetweens
 	else
 	{
-		OutInbetweenWeights[ UpperIndex ] = UpperWeight;
-		OutInbetweenWeights[ LowerIndex ] = LowerWeight;
+		OutInbetweenWeights[UpperIndex] = UpperWeight;
+		OutInbetweenWeights[LowerIndex] = LowerWeight;
 	}
 }
 
 #if USE_USD_SDK
 
 // Adapted from UsdSkel_CacheImpl::ReadScope::_FindOrCreateSkinningQuery because we need to manually create these on UsdGeomMeshes we already have
-pxr::UsdSkelSkinningQuery UsdUtils::CreateSkinningQuery( const pxr::UsdGeomMesh& SkinnedMesh, const pxr::UsdSkelSkeletonQuery& SkeletonQuery )
+pxr::UsdSkelSkinningQuery UsdUtils::CreateSkinningQuery(const pxr::UsdGeomMesh& SkinnedMesh, const pxr::UsdSkelSkeletonQuery& SkeletonQuery)
 {
 	pxr::UsdPrim SkinnedPrim = SkinnedMesh.GetPrim();
-	if ( !SkinnedPrim )
+	if (!SkinnedPrim)
 	{
 		return {};
 	}
 
 	const pxr::UsdSkelAnimQuery& AnimQuery = SkeletonQuery.GetAnimQuery();
 
-	pxr::UsdSkelBindingAPI SkelBindingAPI{ SkinnedPrim };
+	pxr::UsdSkelBindingAPI SkelBindingAPI{SkinnedPrim};
 
 	return pxr::UsdSkelSkinningQuery(
 		SkinnedPrim,
@@ -2879,7 +3048,7 @@ void UsdUtils::BindAnimationSource(pxr::UsdPrim& Prim, const pxr::UsdPrim& Anima
 
 UE::FUsdPrim UsdUtils::FindFirstAnimationSource(const UE::FUsdPrim& SkelRootPrim)
 {
-	if ( !SkelRootPrim )
+	if (!SkelRootPrim)
 	{
 		return {};
 	}
@@ -2988,108 +3157,111 @@ bool UsdUtils::GetSkelQueries(
 	return InOutSkelCache->ComputeSkelBinding(InSkelRootPrim, InSkeletonPrim, &OutSkelBinding, pxr::UsdTraverseInstanceProxies());
 }
 
-#endif // USE_USD_SDK
+#endif	  // USE_USD_SDK
 
 #if USE_USD_SDK && WITH_EDITOR
 
-bool UnrealToUsd::ConvertSkeleton( const FReferenceSkeleton& ReferenceSkeleton, pxr::UsdSkelSkeleton& UsdSkeleton )
+bool UnrealToUsd::ConvertSkeleton(const FReferenceSkeleton& ReferenceSkeleton, pxr::UsdSkelSkeleton& UsdSkeleton)
 {
 	FScopedUsdAllocs Allocs;
 
 	pxr::UsdStageRefPtr Stage = UsdSkeleton.GetPrim().GetStage();
-	if ( !Stage )
+	if (!Stage)
 	{
 		return false;
 	}
 
-	FUsdStageInfo StageInfo{ Stage };
+	FUsdStageInfo StageInfo{Stage};
 
 	// Joints
 	{
 		pxr::UsdAttribute JointsAttr = UsdSkeleton.CreateJointsAttr();
-		UnrealToUsd::ConvertJointsAttribute( ReferenceSkeleton, JointsAttr );
+		UnrealToUsd::ConvertJointsAttribute(ReferenceSkeleton, JointsAttr);
 	}
 
 	pxr::VtArray<pxr::GfMatrix4d> LocalSpaceJointTransforms;
-	LocalSpaceJointTransforms.reserve( ReferenceSkeleton.GetRefBonePose().Num() );
-	for ( const FTransform& BonePose : ReferenceSkeleton.GetRefBonePose() )
+	LocalSpaceJointTransforms.reserve(ReferenceSkeleton.GetRefBonePose().Num());
+	for (const FTransform& BonePose : ReferenceSkeleton.GetRefBonePose())
 	{
-		LocalSpaceJointTransforms.push_back( UnrealToUsd::ConvertTransform( StageInfo, BonePose ) );
+		LocalSpaceJointTransforms.push_back(UnrealToUsd::ConvertTransform(StageInfo, BonePose));
 	}
 
 	TArray<FTransform> WorldSpaceUEJointTransforms;
-	FAnimationRuntime::FillUpComponentSpaceTransforms(
-		ReferenceSkeleton,
-		ReferenceSkeleton.GetRefBonePose(),
-		WorldSpaceUEJointTransforms
-	);
+	FAnimationRuntime::FillUpComponentSpaceTransforms(ReferenceSkeleton, ReferenceSkeleton.GetRefBonePose(), WorldSpaceUEJointTransforms);
 
 	pxr::VtArray<pxr::GfMatrix4d> WorldSpaceJointTransforms;
-	WorldSpaceJointTransforms.reserve( WorldSpaceUEJointTransforms.Num() );
-	for ( const FTransform& WorldSpaceUETransform : WorldSpaceUEJointTransforms )
+	WorldSpaceJointTransforms.reserve(WorldSpaceUEJointTransforms.Num());
+	for (const FTransform& WorldSpaceUETransform : WorldSpaceUEJointTransforms)
 	{
-		WorldSpaceJointTransforms.push_back( UnrealToUsd::ConvertTransform( StageInfo, WorldSpaceUETransform ) );
+		WorldSpaceJointTransforms.push_back(UnrealToUsd::ConvertTransform(StageInfo, WorldSpaceUETransform));
 	}
 
 	// Rest transforms
 	{
 		pxr::UsdAttribute RestTransformsAttr = UsdSkeleton.CreateRestTransformsAttr();
-		RestTransformsAttr.Set( LocalSpaceJointTransforms );
+		RestTransformsAttr.Set(LocalSpaceJointTransforms);
 	}
 
 	// Bind transforms
 	{
 		pxr::UsdAttribute BindTransformsAttr = UsdSkeleton.CreateBindTransformsAttr();
-		BindTransformsAttr.Set( WorldSpaceJointTransforms );
+		BindTransformsAttr.Set(WorldSpaceJointTransforms);
 	}
 
 	// Use Guide purpose on skeletons by default, unless it has some specific purpose set already
-	if ( pxr::UsdAttribute PurposeAttr = UsdSkeleton.GetPurposeAttr() )
+	if (pxr::UsdAttribute PurposeAttr = UsdSkeleton.GetPurposeAttr())
 	{
-		if ( !PurposeAttr.HasAuthoredValue() )
+		if (!PurposeAttr.HasAuthoredValue())
 		{
-			PurposeAttr.Set( pxr::UsdGeomTokens->guide );
+			PurposeAttr.Set(pxr::UsdGeomTokens->guide);
 		}
 	}
 
 	return true;
 }
 
-bool UnrealToUsd::ConvertJointsAttribute( const FReferenceSkeleton& ReferenceSkeleton, pxr::UsdAttribute& JointsAttribute )
+bool UnrealToUsd::ConvertJointsAttribute(const FReferenceSkeleton& ReferenceSkeleton, pxr::UsdAttribute& JointsAttribute)
 {
-	if ( !JointsAttribute )
+	if (!JointsAttribute)
 	{
 		return false;
 	}
 
 	TArray<FString> FullBonePaths;
-	UnrealToUsdImpl::CreateFullBonePaths( ReferenceSkeleton.GetRefBoneInfo(), FullBonePaths );
+	UnrealToUsdImpl::CreateFullBonePaths(ReferenceSkeleton.GetRefBoneInfo(), FullBonePaths);
 
 	pxr::VtArray<pxr::TfToken> Joints;
-	Joints.reserve( FullBonePaths.Num() );
-	for ( const FString& BonePath : FullBonePaths )
+	Joints.reserve(FullBonePaths.Num());
+	for (const FString& BonePath : FullBonePaths)
 	{
-		Joints.push_back( UnrealToUsd::ConvertToken( *BonePath ).Get() );
+		Joints.push_back(UnrealToUsd::ConvertToken(*BonePath).Get());
 	}
 
-	JointsAttribute.Set( Joints );
+	JointsAttribute.Set(Joints);
 	return true;
 }
 
-bool UnrealToUsd::ConvertSkeleton( const USkeleton* Skeleton, pxr::UsdSkelSkeleton& UsdSkeleton )
+bool UnrealToUsd::ConvertSkeleton(const USkeleton* Skeleton, pxr::UsdSkelSkeleton& UsdSkeleton)
 {
-	if ( !Skeleton )
+	if (!Skeleton)
 	{
 		return false;
 	}
 
-	return UnrealToUsd::ConvertSkeleton( Skeleton->GetReferenceSkeleton(), UsdSkeleton );
+	return UnrealToUsd::ConvertSkeleton(Skeleton->GetReferenceSkeleton(), UsdSkeleton);
 }
 
-bool UnrealToUsd::ConvertSkeletalMesh( const USkeletalMesh* SkeletalMesh, pxr::UsdPrim& SkelRootPrim, const pxr::UsdTimeCode TimeCode, UE::FUsdStage* StageForMaterialAssignments, int32 LowestMeshLOD, int32 HighestMeshLOD )
+bool UnrealToUsd::ConvertSkeletalMesh(
+	const USkeletalMesh* SkeletalMesh,
+	pxr::UsdPrim& SkelRootPrim,
+	const pxr::UsdTimeCode TimeCode,
+	UE::FUsdStage* StageForMaterialAssignments,
+	int32 LowestMeshLOD,
+	int32 HighestMeshLOD
+)
 {
-	pxr::UsdSkelRoot SkelRoot{ SkelRootPrim };
-	if ( !SkeletalMesh || !SkeletalMesh->GetSkeleton() || !SkelRoot )
+	pxr::UsdSkelRoot SkelRoot{SkelRootPrim};
+	if (!SkeletalMesh || !SkeletalMesh->GetSkeleton() || !SkelRoot)
 	{
 		return false;
 	}
@@ -3097,35 +3269,41 @@ bool UnrealToUsd::ConvertSkeletalMesh( const USkeletalMesh* SkeletalMesh, pxr::U
 	FScopedUsdAllocs UsdAllocs;
 
 	pxr::UsdStageRefPtr Stage = SkelRootPrim.GetStage();
-	if ( !Stage )
+	if (!Stage)
 	{
 		return false;
 	}
-	const FUsdStageInfo StageInfo( Stage );
+	const FUsdStageInfo StageInfo(Stage);
 
 	const FSkeletalMeshModel* SkelMeshResource = SkeletalMesh->GetImportedModel();
 	int32 NumLODs = SkelMeshResource->LODModels.Num();
-	if ( NumLODs < 1 )
+	if (NumLODs < 1)
 	{
 		return false;
 	}
 
 	// Make sure they're both >= 0 (the options dialog slider is clamped, but this may be called directly)
-	LowestMeshLOD = FMath::Clamp( LowestMeshLOD, 0, NumLODs - 1 );
-	HighestMeshLOD = FMath::Clamp( HighestMeshLOD, 0, NumLODs - 1 );
+	LowestMeshLOD = FMath::Clamp(LowestMeshLOD, 0, NumLODs - 1);
+	HighestMeshLOD = FMath::Clamp(HighestMeshLOD, 0, NumLODs - 1);
 
 	// Make sure Lowest <= Highest
-	int32 Temp = FMath::Min( LowestMeshLOD, HighestMeshLOD );
-	HighestMeshLOD = FMath::Max( LowestMeshLOD, HighestMeshLOD );
+	int32 Temp = FMath::Min(LowestMeshLOD, HighestMeshLOD);
+	HighestMeshLOD = FMath::Max(LowestMeshLOD, HighestMeshLOD);
 	LowestMeshLOD = Temp;
 
 	// Make sure it's at least 1 LOD level
-	NumLODs = FMath::Max( HighestMeshLOD - LowestMeshLOD + 1, 1 );
+	NumLODs = FMath::Max(HighestMeshLOD - LowestMeshLOD + 1, 1);
 
 	pxr::UsdVariantSets VariantSets = SkelRootPrim.GetVariantSets();
-	if ( NumLODs > 1 && VariantSets.HasVariantSet( UnrealIdentifiers::LOD ) )
-	{
-		UE_LOG( LogUsd, Error, TEXT( "Failed to export higher LODs for skeletal mesh '%s', as the target prim already has a variant set named '%s'!" ), *SkeletalMesh->GetName(), *UsdToUnreal::ConvertToken( UnrealIdentifiers::LOD ) );
+	if (NumLODs > 1 && VariantSets.HasVariantSet(UnrealIdentifiers::LOD))
+	{
+		UE_LOG(
+			LogUsd,
+			Error,
+			TEXT("Failed to export higher LODs for skeletal mesh '%s', as the target prim already has a variant set named '%s'!"),
+			*SkeletalMesh->GetName(),
+			*UsdToUnreal::ConvertToken(UnrealIdentifiers::LOD)
+		);
 		NumLODs = 1;
 	}
 
@@ -3136,177 +3314,190 @@ bool UnrealToUsd::ConvertSkeletalMesh( const USkeletalMesh* SkeletalMesh, pxr::U
 
 	// Collect all material assignments, referenced by the sections' material indices
 	bool bHasMaterialAssignments = false;
-	pxr::VtArray< std::string > MaterialAssignments;
-	for ( const FSkeletalMaterial& SkeletalMaterial : SkeletalMesh->GetMaterials())
+	pxr::VtArray<std::string> MaterialAssignments;
+	for (const FSkeletalMaterial& SkeletalMaterial : SkeletalMesh->GetMaterials())
 	{
 		FString AssignedMaterialPathName;
-		if ( UMaterialInterface* Material = SkeletalMaterial.MaterialInterface )
+		if (UMaterialInterface* Material = SkeletalMaterial.MaterialInterface)
 		{
-			if ( Material->GetOutermost() != GetTransientPackage() )
+			if (Material->GetOutermost() != GetTransientPackage())
 			{
 				AssignedMaterialPathName = Material->GetPathName();
 				bHasMaterialAssignments = true;
 			}
 		}
 
-		MaterialAssignments.push_back( UnrealToUsd::ConvertString( *AssignedMaterialPathName ).Get() );
+		MaterialAssignments.push_back(UnrealToUsd::ConvertString(*AssignedMaterialPathName).Get());
 	}
-	if ( !bHasMaterialAssignments )
+	if (!bHasMaterialAssignments)
 	{
 		// Prevent creation of the unrealMaterials attribute in case we don't have any assignments at all
 		MaterialAssignments.clear();
 	}
 
 	// Create and fill skeleton
-	pxr::UsdSkelBindingAPI SkelBindingAPI = pxr::UsdSkelBindingAPI::Apply( SkelRootPrim );
+	pxr::UsdSkelBindingAPI SkelBindingAPI = pxr::UsdSkelBindingAPI::Apply(SkelRootPrim);
 	{
 		pxr::UsdPrim SkeletonPrim = Stage->DefinePrim(
 			SkelRootPrim.GetPath().AppendChild(UnrealToUsd::ConvertToken(UnrealIdentifiers::ExportedSkeletonPrimName).Get()),
 			UnrealToUsd::ConvertToken(TEXT("Skeleton")).Get()
 		);
-		pxr::UsdSkelSkeleton SkelSkeleton{ SkeletonPrim };
+		pxr::UsdSkelSkeleton SkelSkeleton{SkeletonPrim};
 
 		pxr::UsdRelationship SkelRel = SkelBindingAPI.CreateSkeletonRel();
 		SkelRel.SetTargets({SkeletonPrim.GetPath()});
 
-		UnrealToUsd::ConvertSkeleton( SkeletalMesh->GetRefSkeleton(), SkelSkeleton );
+		UnrealToUsd::ConvertSkeleton(SkeletalMesh->GetRefSkeleton(), SkelSkeleton);
 	}
 
 	// Actual meshes
-	for ( int32 LODIndex = LowestMeshLOD; LODIndex <= HighestMeshLOD; ++LODIndex )
+	for (int32 LODIndex = LowestMeshLOD; LODIndex <= HighestMeshLOD; ++LODIndex)
 	{
-		const FSkeletalMeshLODModel& LODModel = SkelMeshResource->LODModels[ LODIndex ];
+		const FSkeletalMeshLODModel& LODModel = SkelMeshResource->LODModels[LODIndex];
 
-		if ( LODModel.NumVertices == 0 || LODModel.Sections.Num() == 0 )
+		if (LODModel.NumVertices == 0 || LODModel.Sections.Num() == 0)
 		{
 			continue;
 		}
 
 		// LOD0, LOD1, etc.
-		std::string VariantName = UnrealIdentifiers::LOD.GetString() + UnrealToUsd::ConvertString( *LexToString( LODIndex ) ).Get();
-		if ( LowestLODAdded.size() == 0 )
+		std::string VariantName = UnrealIdentifiers::LOD.GetString() + UnrealToUsd::ConvertString(*LexToString(LODIndex)).Get();
+		if (LowestLODAdded.size() == 0)
 		{
 			LowestLODAdded = VariantName;
 		}
 
 		// Enable the variant edit context, if we are creating variant LODs
-		TOptional< pxr::UsdEditContext > EditContext;
-		if ( bExportMultipleLODs )
+		TOptional<pxr::UsdEditContext> EditContext;
+		if (bExportMultipleLODs)
 		{
-			pxr::UsdVariantSet VariantSet = VariantSets.GetVariantSet( UnrealIdentifiers::LOD );
+			pxr::UsdVariantSet VariantSet = VariantSets.GetVariantSet(UnrealIdentifiers::LOD);
 
-			if ( !VariantSet.AddVariant( VariantName ) )
+			if (!VariantSet.AddVariant(VariantName))
 			{
 				continue;
 			}
 
-			VariantSet.SetVariantSelection( VariantName );
-			EditContext.Emplace( VariantSet.GetVariantEditContext() );
+			VariantSet.SetVariantSelection(VariantName);
+			EditContext.Emplace(VariantSet.GetVariantEditContext());
 		}
 
-		pxr::SdfPath MeshPrimPath = ParentPrimPath.AppendPath( pxr::SdfPath( bExportMultipleLODs ? VariantName : UnrealToUsd::ConvertString( *UsdUtils::SanitizeUsdIdentifier( *SkeletalMesh->GetName() ) ).Get() ) );
-		pxr::UsdPrim UsdLODPrim = Stage->DefinePrim( MeshPrimPath, UnrealToUsd::ConvertToken( TEXT( "Mesh" ) ).Get() );
-		pxr::UsdGeomMesh UsdLODPrimGeomMesh{ UsdLODPrim };
+		pxr::SdfPath MeshPrimPath = ParentPrimPath.AppendPath(pxr::SdfPath(
+			bExportMultipleLODs ? VariantName : UnrealToUsd::ConvertString(*UsdUtils::SanitizeUsdIdentifier(*SkeletalMesh->GetName())).Get()
+		));
+		pxr::UsdPrim UsdLODPrim = Stage->DefinePrim(MeshPrimPath, UnrealToUsd::ConvertToken(TEXT("Mesh")).Get());
+		pxr::UsdGeomMesh UsdLODPrimGeomMesh{UsdLODPrim};
 
 		pxr::UsdPrim MaterialPrim = UsdLODPrim;
-		if ( StageForMaterialAssignments )
+		if (StageForMaterialAssignments)
 		{
-			pxr::UsdStageRefPtr MaterialStage{ *StageForMaterialAssignments };
-			MaterialPrim = MaterialStage->OverridePrim( MeshPrimPath );
+			pxr::UsdStageRefPtr MaterialStage{*StageForMaterialAssignments};
+			MaterialPrim = MaterialStage->OverridePrim(MeshPrimPath);
 		}
 
 		TArray<int32> LODMaterialMap;
-		if ( const FSkeletalMeshLODInfo* LODInfo = SkeletalMesh->GetLODInfo( LODIndex ) )
+		if (const FSkeletalMeshLODInfo* LODInfo = SkeletalMesh->GetLODInfo(LODIndex))
 		{
 			LODMaterialMap = LODInfo->LODMaterialMap;
 		}
 
-		UnrealToUsdImpl::ConvertSkeletalMeshLOD( SkeletalMesh, LODModel, UsdLODPrimGeomMesh, SkeletalMesh->GetHasVertexColors(), MaterialAssignments, LODMaterialMap, TimeCode, MaterialPrim );
+		UnrealToUsdImpl::ConvertSkeletalMeshLOD(
+			SkeletalMesh,
+			LODModel,
+			UsdLODPrimGeomMesh,
+			SkeletalMesh->GetHasVertexColors(),
+			MaterialAssignments,
+			LODMaterialMap,
+			TimeCode,
+			MaterialPrim
+		);
 
 		// Relationships can't target prims inside variants, so if we have BlendShapes to export we have to disable the edit target
 		// so that the blend shapes end up outside the variants and the Meshes can have their blendShapeTargets relationships pointing at them
-		if ( bExportMultipleLODs && SkeletalMesh->GetMorphTargets().Num() > 0)
+		if (bExportMultipleLODs && SkeletalMesh->GetMorphTargets().Num() > 0)
 		{
 			EditContext.Reset();
 		}
 
-		pxr::VtArray< pxr::TfToken > AddedBlendShapes;
+		pxr::VtArray<pxr::TfToken> AddedBlendShapes;
 		pxr::SdfPathVector AddedBlendShapeTargets;
-		for ( UMorphTarget* MorphTarget : SkeletalMesh->GetMorphTargets())
+		for (UMorphTarget* MorphTarget : SkeletalMesh->GetMorphTargets())
 		{
-			if ( !MorphTarget || !MorphTarget->HasDataForLOD( LODIndex ) )
+			if (!MorphTarget || !MorphTarget->HasDataForLOD(LODIndex))
 			{
 				continue;
 			}
 
 			int32 NumDeltas = 0;
-			const FMorphTargetDelta* DeltaArray = MorphTarget->GetMorphTargetDelta( LODIndex, NumDeltas );
-			if ( !DeltaArray || NumDeltas == 0 )
+			const FMorphTargetDelta* DeltaArray = MorphTarget->GetMorphTargetDelta(LODIndex, NumDeltas);
+			if (!DeltaArray || NumDeltas == 0)
 			{
 				continue;
 			}
 
 			pxr::SdfPath ParentPath = bExportMultipleLODs ? SkelRootPrim.GetPath() : UsdLODPrim.GetPath();
 
-			pxr::SdfPath BlendShapePath = ParentPath.AppendPath( UnrealToUsd::ConvertPath( *UsdUtils::SanitizeUsdIdentifier( *MorphTarget->GetName() ) ).Get() );
-			pxr::UsdPrim BlendShapePrim = UsdLODPrim.GetStage()->DefinePrim( BlendShapePath, UnrealToUsd::ConvertToken( TEXT( "BlendShape" ) ).Get() );
-			pxr::UsdSkelBlendShape BlendShape{ BlendShapePrim };
+			pxr::SdfPath BlendShapePath = ParentPath.AppendPath(
+				UnrealToUsd::ConvertPath(*UsdUtils::SanitizeUsdIdentifier(*MorphTarget->GetName())).Get()
+			);
+			pxr::UsdPrim BlendShapePrim = UsdLODPrim.GetStage()->DefinePrim(BlendShapePath, UnrealToUsd::ConvertToken(TEXT("BlendShape")).Get());
+			pxr::UsdSkelBlendShape BlendShape{BlendShapePrim};
 
 			bool bCreatedBlendShape = UnrealToUsdImpl::ConvertMorphTargetDeltas(DeltaArray, NumDeltas, BlendShape, TimeCode);
-			if ( !bCreatedBlendShape )
+			if (!bCreatedBlendShape)
 			{
 				continue;
 			}
 
-			AddedBlendShapes.push_back( UnrealToUsd::ConvertToken( *UsdUtils::SanitizeUsdIdentifier( *MorphTarget->GetName() ) ).Get() );
-			AddedBlendShapeTargets.push_back( BlendShapePath );
+			AddedBlendShapes.push_back(UnrealToUsd::ConvertToken(*UsdUtils::SanitizeUsdIdentifier(*MorphTarget->GetName())).Get());
+			AddedBlendShapeTargets.push_back(BlendShapePath);
 		}
 
-		if ( AddedBlendShapeTargets.size() > 0 )
+		if (AddedBlendShapeTargets.size() > 0)
 		{
 			// Restore the edit target to the current LOD variant so that the relationship itself ends up inside the mesh, inside the variant
-			if ( bExportMultipleLODs )
+			if (bExportMultipleLODs)
 			{
-				EditContext.Emplace( VariantSets.GetVariantSet( UnrealIdentifiers::LOD ).GetVariantEditContext() );
+				EditContext.Emplace(VariantSets.GetVariantSet(UnrealIdentifiers::LOD).GetVariantEditContext());
 			}
 
-			pxr::UsdSkelBindingAPI LODMeshSkelBindingAPI = pxr::UsdSkelBindingAPI::Apply( UsdLODPrim );
-			LODMeshSkelBindingAPI.CreateBlendShapeTargetsRel().SetTargets( AddedBlendShapeTargets );
-			LODMeshSkelBindingAPI.CreateBlendShapesAttr().Set( AddedBlendShapes );
+			pxr::UsdSkelBindingAPI LODMeshSkelBindingAPI = pxr::UsdSkelBindingAPI::Apply(UsdLODPrim);
+			LODMeshSkelBindingAPI.CreateBlendShapeTargetsRel().SetTargets(AddedBlendShapeTargets);
+			LODMeshSkelBindingAPI.CreateBlendShapesAttr().Set(AddedBlendShapes);
 		}
 	}
 
-	if ( bExportMultipleLODs )
+	if (bExportMultipleLODs)
 	{
-		VariantSets.GetVariantSet( UnrealIdentifiers::LOD ).SetVariantSelection( LowestLODAdded );
+		VariantSets.GetVariantSet(UnrealIdentifiers::LOD).SetVariantSelection(LowestLODAdded);
 	}
 
 	return true;
 }
 
-bool UnrealToUsd::ConvertAnimSequence( UAnimSequence* AnimSequence, pxr::UsdPrim& SkelAnimPrim )
+bool UnrealToUsd::ConvertAnimSequence(UAnimSequence* AnimSequence, pxr::UsdPrim& SkelAnimPrim)
 {
-	if ( !SkelAnimPrim || !AnimSequence || !AnimSequence->GetSkeleton() )
+	if (!SkelAnimPrim || !AnimSequence || !AnimSequence->GetSkeleton())
 	{
 		return false;
 	}
 
-	pxr::UsdSkelAnimation UsdSkelAnim( SkelAnimPrim );
+	pxr::UsdSkelAnimation UsdSkelAnim(SkelAnimPrim);
 
-	if ( !UsdSkelAnim )
+	if (!UsdSkelAnim)
 	{
 		return false;
 	}
 
 	USkeleton* AnimSkeleton = AnimSequence->GetSkeleton();
-	USkeletalMesh* SkeletalMesh = AnimSkeleton->GetAssetPreviewMesh( AnimSequence );
+	USkeletalMesh* SkeletalMesh = AnimSkeleton->GetAssetPreviewMesh(AnimSequence);
 
-	if ( !SkeletalMesh )
+	if (!SkeletalMesh)
 	{
 		SkeletalMesh = AnimSkeleton->FindCompatibleMesh();
 	}
 
-	if ( !SkeletalMesh )
+	if (!SkeletalMesh)
 	{
 		return false;
 	}
@@ -3316,7 +3507,7 @@ bool UnrealToUsd::ConvertAnimSequence( UAnimSequence* AnimSequence, pxr::UsdPrim
 	const double TimeCodesPerSecond = SkelAnimPrim.GetStage()->GetTimeCodesPerSecond();
 	const int32 NumTimeCodes = AnimSequence->GetPlayLength() * TimeCodesPerSecond;
 
-	if ( NumBones <= 0 )
+	if (NumBones <= 0)
 	{
 		return false;
 	}
@@ -3324,44 +3515,46 @@ bool UnrealToUsd::ConvertAnimSequence( UAnimSequence* AnimSequence, pxr::UsdPrim
 	FScopedUsdAllocs UsdAllocs;
 	pxr::SdfChangeBlock ChangeBlock;
 
-	FUsdStageInfo StageInfo( SkelAnimPrim.GetStage() );
+	FUsdStageInfo StageInfo(SkelAnimPrim.GetStage());
 
 	// Blend shapes
 	{
-		pxr::VtArray< pxr::TfToken > BlendShapeNames;
-		pxr::VtArray< float > BlendShapeWeights;
+		pxr::VtArray<pxr::TfToken> BlendShapeNames;
+		pxr::VtArray<float> BlendShapeWeights;
 
 		// We need to make sure we have at least one mark on the memstack allocator because FBlendedCurve
 		// will allocate using one and will assert if there aren't any marks yet
-		FMemMark Mark( FMemStack::Get() );
+		FMemMark Mark(FMemStack::Get());
 
 		// Blend shape weights
-		for ( int32 TimeCode = 0; TimeCode < NumTimeCodes; ++TimeCode )
+		for (int32 TimeCode = 0; TimeCode < NumTimeCodes; ++TimeCode)
 		{
 			const float AnimTime = TimeCode / TimeCodesPerSecond;
 
 			FBlendedCurve BlendedCurve;
 			const bool bForceUseRawData = true;
-			AnimSequence->EvaluateCurveData( BlendedCurve, AnimTime, bForceUseRawData );
+			AnimSequence->EvaluateCurveData(BlendedCurve, AnimTime, bForceUseRawData);
 
 			BlendShapeNames.clear();
-			BlendShapeNames.reserve( BlendedCurve.Num() );
+			BlendShapeNames.reserve(BlendedCurve.Num());
 			BlendShapeWeights.clear();
-			BlendShapeWeights.reserve( BlendedCurve.Num() );
+			BlendShapeWeights.reserve(BlendedCurve.Num());
 
-			BlendedCurve.ForEachElement([&BlendShapeWeights, &BlendShapeNames](const UE::Anim::FCurveElement& InElement)
-			{
-				if (EnumHasAnyFlags(InElement.Flags, UE::Anim::ECurveElementFlags::MorphTarget))
+			BlendedCurve.ForEachElement(
+				[&BlendShapeWeights, &BlendShapeNames](const UE::Anim::FCurveElement& InElement)
 				{
-					BlendShapeNames.push_back( UnrealToUsd::ConvertToken( *InElement.Name.ToString() ).Get() );
-					BlendShapeWeights.push_back( InElement.Value );
+					if (EnumHasAnyFlags(InElement.Flags, UE::Anim::ECurveElementFlags::MorphTarget))
+					{
+						BlendShapeNames.push_back(UnrealToUsd::ConvertToken(*InElement.Name.ToString()).Get());
+						BlendShapeWeights.push_back(InElement.Value);
+					}
 				}
-			});
+			);
 
-			if(BlendShapeWeights.size() > 0 && BlendShapeNames.size() > 0)
+			if (BlendShapeWeights.size() > 0 && BlendShapeNames.size() > 0)
 			{
-				UsdSkelAnim.CreateBlendShapesAttr().Set( BlendShapeNames, pxr::UsdTimeCode( TimeCode ) );
-				UsdSkelAnim.CreateBlendShapeWeightsAttr().Set( BlendShapeWeights, pxr::UsdTimeCode( TimeCode ) );
+				UsdSkelAnim.CreateBlendShapesAttr().Set(BlendShapeNames, pxr::UsdTimeCode(TimeCode));
+				UsdSkelAnim.CreateBlendShapeWeightsAttr().Set(BlendShapeWeights, pxr::UsdTimeCode(TimeCode));
 			}
 		}
 	}
@@ -3369,7 +3562,7 @@ bool UnrealToUsd::ConvertAnimSequence( UAnimSequence* AnimSequence, pxr::UsdPrim
 	// Joints
 	{
 		pxr::UsdAttribute JointsAttr = UsdSkelAnim.CreateJointsAttr();
-		UnrealToUsd::ConvertJointsAttribute( RefSkeleton, JointsAttr );
+		UnrealToUsd::ConvertJointsAttribute(RefSkeleton, JointsAttr);
 	}
 
 	// Translations, Rotations & Scales
@@ -3378,46 +3571,46 @@ bool UnrealToUsd::ConvertAnimSequence( UAnimSequence* AnimSequence, pxr::UsdPrim
 		pxr::UsdAttribute RotationsAttr = UsdSkelAnim.CreateRotationsAttr();
 		pxr::UsdAttribute ScalesAttr = UsdSkelAnim.CreateScalesAttr();
 
-		UDebugSkelMeshComponent* DebugSkelMeshComponent = NewObject< UDebugSkelMeshComponent >();
-		DebugSkelMeshComponent->RegisterComponentWithWorld( IUsdClassesModule::GetCurrentWorld() );
+		UDebugSkelMeshComponent* DebugSkelMeshComponent = NewObject<UDebugSkelMeshComponent>();
+		DebugSkelMeshComponent->RegisterComponentWithWorld(IUsdClassesModule::GetCurrentWorld());
 		DebugSkelMeshComponent->EmptyOverrideMaterials();
-		DebugSkelMeshComponent->SetSkeletalMesh( SkeletalMesh );
+		DebugSkelMeshComponent->SetSkeletalMesh(SkeletalMesh);
 
 		const bool bEnable = true;
-		DebugSkelMeshComponent->EnablePreview( bEnable, AnimSequence );
+		DebugSkelMeshComponent->EnablePreview(bEnable, AnimSequence);
 
-		for ( int32 TimeCode = 0; TimeCode < NumTimeCodes; ++TimeCode )
+		for (int32 TimeCode = 0; TimeCode < NumTimeCodes; ++TimeCode)
 		{
 			const float AnimTime = TimeCode / TimeCodesPerSecond;
 
 			const bool bFireNotifies = false;
-			DebugSkelMeshComponent->SetPosition( AnimTime, bFireNotifies );
+			DebugSkelMeshComponent->SetPosition(AnimTime, bFireNotifies);
 			DebugSkelMeshComponent->RefreshBoneTransforms();
 
 			pxr::VtVec3fArray Translations;
-			Translations.reserve( NumBones );
+			Translations.reserve(NumBones);
 
 			pxr::VtQuatfArray Rotations;
-			Rotations.reserve( NumBones );
+			Rotations.reserve(NumBones);
 
 			pxr::VtVec3hArray Scales;
-			Scales.reserve( NumBones );
+			Scales.reserve(NumBones);
 
-			TArray< FTransform > LocalBoneTransforms = DebugSkelMeshComponent->GetBoneSpaceTransforms();
+			TArray<FTransform> LocalBoneTransforms = DebugSkelMeshComponent->GetBoneSpaceTransforms();
 
-			for ( int32 BoneIndex = 0; BoneIndex < NumBones; ++BoneIndex )
+			for (int32 BoneIndex = 0; BoneIndex < NumBones; ++BoneIndex)
 			{
-				FTransform BoneTransform = LocalBoneTransforms[ BoneIndex ];
-				BoneTransform = UsdUtils::ConvertAxes( StageInfo.UpAxis == EUsdUpAxis::ZAxis, BoneTransform );
+				FTransform BoneTransform = LocalBoneTransforms[BoneIndex];
+				BoneTransform = UsdUtils::ConvertAxes(StageInfo.UpAxis == EUsdUpAxis::ZAxis, BoneTransform);
 
-				Translations.push_back( UnrealToUsd::ConvertVectorFloat( BoneTransform.GetTranslation() ) );
-				Rotations.push_back( UnrealToUsd::ConvertQuatFloat( BoneTransform.GetRotation() ).GetNormalized() );
-				Scales.push_back( UnrealToUsd::ConvertVectorHalf( BoneTransform.GetScale3D() ) );
+				Translations.push_back(UnrealToUsd::ConvertVectorFloat(BoneTransform.GetTranslation()));
+				Rotations.push_back(UnrealToUsd::ConvertQuatFloat(BoneTransform.GetRotation()).GetNormalized());
+				Scales.push_back(UnrealToUsd::ConvertVectorHalf(BoneTransform.GetScale3D()));
 			}
 
-			TranslationsAttr.Set( Translations, pxr::UsdTimeCode( TimeCode ) );
-			RotationsAttr.Set( Rotations, pxr::UsdTimeCode( TimeCode ) );
-			ScalesAttr.Set( Scales, pxr::UsdTimeCode( TimeCode ) );
+			TranslationsAttr.Set(Translations, pxr::UsdTimeCode(TimeCode));
+			RotationsAttr.Set(Rotations, pxr::UsdTimeCode(TimeCode));
+			ScalesAttr.Set(Scales, pxr::UsdTimeCode(TimeCode));
 		}
 
 		// Actively delete it or else it will remain visible on the viewport
@@ -3426,9 +3619,9 @@ bool UnrealToUsd::ConvertAnimSequence( UAnimSequence* AnimSequence, pxr::UsdPrim
 
 	const int32 StageEndTimeCode = SkelAnimPrim.GetStage()->GetEndTimeCode();
 
-	if ( NumTimeCodes > StageEndTimeCode )
+	if (NumTimeCodes > StageEndTimeCode)
 	{
-		SkelAnimPrim.GetStage()->SetEndTimeCode( NumTimeCodes - 1 );
+		SkelAnimPrim.GetStage()->SetEndTimeCode(NumTimeCodes - 1);
 	}
 
 	return true;
@@ -3445,22 +3638,22 @@ bool UnrealToUsd::ConvertControlRigSection(
 	const UsdUtils::FBlendShapeMap* InBlendShapeMap
 )
 {
-	if ( !InSection || !InPlayer || !OutSkelAnimPrim )
+	if (!InSection || !InPlayer || !OutSkelAnimPrim)
 	{
 		return false;
 	}
 
 	UControlRig* ControlRig = InSection->GetControlRig();
-	if ( !ControlRig )
+	if (!ControlRig)
 	{
 		return false;
 	}
 
 	FScopedUsdAllocs UsdAllocs;
 
-	pxr::UsdSkelAnimation SkelAnim{ OutSkelAnimPrim };
+	pxr::UsdSkelAnimation SkelAnim{OutSkelAnimPrim};
 	pxr::UsdStageRefPtr UsdStage = OutSkelAnimPrim.GetStage();
-	if ( !UsdStage || !SkelAnim )
+	if (!UsdStage || !SkelAnim)
 	{
 		return false;
 	}
@@ -3470,18 +3663,18 @@ bool UnrealToUsd::ConvertControlRigSection(
 		return false;
 	}
 
-	FUsdStageInfo StageInfo( UsdStage );
+	FUsdStageInfo StageInfo(UsdStage);
 
 	double StartTime = FPlatformTime::Cycles64();
 
 	ControlRig->Initialize();
 	ControlRig->RequestInit();
-	ControlRig->Evaluate_AnyThread(); // Important as it runs the Construction event, which can change topology
+	ControlRig->Evaluate_AnyThread();	 // Important as it runs the Construction event, which can change topology
 
 	// Record how the topology looks while we setup our arrays and maps. If this changes during
 	// baking we'll just drop everything and return
 	URigHierarchy* InitialHierarchy = ControlRig->GetHierarchy();
-	if ( !InitialHierarchy )
+	if (!InitialHierarchy)
 	{
 		return false;
 	}
@@ -3491,28 +3684,27 @@ bool UnrealToUsd::ConvertControlRigSection(
 	// This works because the topology won't change in here, and bone names are unique across the entire skeleton
 	// Its possible we'll be putting INDEX_NONEs into RigIndexToRefSkeletonIndex, but that's alright.
 	TArray<int32> RigJointIndexToRefSkeletonIndex;
-	TFunction<void()> RegenerateRigJointIndexToRefSkeletonIndex =
-		[ControlRig, &RigJointIndexToRefSkeletonIndex, &InRefSkeleton]()
+	TFunction<void()> RegenerateRigJointIndexToRefSkeletonIndex = [ControlRig, &RigJointIndexToRefSkeletonIndex, &InRefSkeleton]()
+	{
+		URigHierarchy* Hierarchy = ControlRig->GetHierarchy();
+		if (!Hierarchy)
 		{
-			URigHierarchy* Hierarchy = ControlRig->GetHierarchy();
-			if (!Hierarchy)
-			{
-				return;
-			}
+			return;
+		}
 
-			RigJointIndexToRefSkeletonIndex.Reset();
-			for (FRigBoneElement* RigBone : Hierarchy->GetBones())
-			{
-				RigJointIndexToRefSkeletonIndex.Add(InRefSkeleton.FindBoneIndex(RigBone->GetFName()));
-			}
-		};
+		RigJointIndexToRefSkeletonIndex.Reset();
+		for (FRigBoneElement* RigBone : Hierarchy->GetBones())
+		{
+			RigJointIndexToRefSkeletonIndex.Add(InRefSkeleton.FindBoneIndex(RigBone->GetFName()));
+		}
+	};
 	RegenerateRigJointIndexToRefSkeletonIndex();
 
 	pxr::UsdAttribute JointsAttr = SkelAnim.CreateJointsAttr();
-	UnrealToUsd::ConvertJointsAttribute( InRefSkeleton, JointsAttr );
+	UnrealToUsd::ConvertJointsAttribute(InRefSkeleton, JointsAttr);
 
 	TArray<FTransform> GlobalUEJointTransformsForFrame;
-	GlobalUEJointTransformsForFrame.SetNum( InRefSkeleton.GetNum() );
+	GlobalUEJointTransformsForFrame.SetNum(InRefSkeleton.GetNum());
 
 	pxr::UsdAttribute TranslationsAttr = SkelAnim.CreateTranslationsAttr();
 	pxr::UsdAttribute RotationsAttr = SkelAnim.CreateRotationsAttr();
@@ -3526,15 +3718,15 @@ bool UnrealToUsd::ConvertControlRigSection(
 	pxr::VtVec3fArray Translations;
 	pxr::VtQuatfArray Rotations;
 	pxr::VtVec3hArray Scales;
-	Translations.resize( InRefSkeleton.GetNum() );
-	Rotations.resize( InRefSkeleton.GetNum() );
-	Scales.resize( InRefSkeleton.GetNum() );
+	Translations.resize(InRefSkeleton.GetNum());
+	Rotations.resize(InRefSkeleton.GetNum());
+	Scales.resize(InRefSkeleton.GetNum());
 
 	FFrameRate TickResolution = InMovieScene->GetTickResolution();
 	FFrameRate DisplayRate = InMovieScene->GetDisplayRate();
 
 	const double StageTimeCodesPerSecond = UsdStage->GetTimeCodesPerSecond();
-	const FFrameRate StageFrameRate( StageTimeCodesPerSecond, 1 );
+	const FFrameRate StageFrameRate(StageTimeCodesPerSecond, 1);
 
 	TRange<FFrameNumber> PlaybackRange = InMovieScene->GetPlaybackRange();
 	TRange<FFrameNumber> BakeTickRange = InSection->ComputeEffectiveRange();
@@ -3544,51 +3736,51 @@ bool UnrealToUsd::ConvertControlRigSection(
 	FFrameNumber EndInclTickFrame;
 	{
 		TOptional<TRangeBound<FFrameNumber>> LowerBoundToUse;
-		if ( BakeTickRange.HasLowerBound() )
+		if (BakeTickRange.HasLowerBound())
 		{
 			TRangeBound<FFrameNumber> SectionLowerBound = BakeTickRange.GetLowerBound();
-			if ( !SectionLowerBound.IsOpen() )
+			if (!SectionLowerBound.IsOpen())
 			{
 				LowerBoundToUse = SectionLowerBound;
 			}
 		}
-		if ( !LowerBoundToUse.IsSet() && PlaybackRange.HasLowerBound() )
+		if (!LowerBoundToUse.IsSet() && PlaybackRange.HasLowerBound())
 		{
 			TRangeBound<FFrameNumber> PlaybackLowerBound = PlaybackRange.GetLowerBound();
-			if ( !PlaybackLowerBound.IsOpen() )
+			if (!PlaybackLowerBound.IsOpen())
 			{
 				LowerBoundToUse = PlaybackLowerBound;
 			}
 		}
-		if ( !LowerBoundToUse.IsSet() )
+		if (!LowerBoundToUse.IsSet())
 		{
 			return false;
 		}
-		StartInclTickFrame = LowerBoundToUse.GetValue().GetValue() + ( LowerBoundToUse.GetValue().IsInclusive() ? 0 : 1 );
+		StartInclTickFrame = LowerBoundToUse.GetValue().GetValue() + (LowerBoundToUse.GetValue().IsInclusive() ? 0 : 1);
 	}
 	{
 		TOptional<TRangeBound<FFrameNumber>> UpperBoundToUse;
-		if ( BakeTickRange.HasUpperBound() )
+		if (BakeTickRange.HasUpperBound())
 		{
 			TRangeBound<FFrameNumber> SectionUpperBound = BakeTickRange.GetUpperBound();
-			if ( !SectionUpperBound.IsOpen() )
+			if (!SectionUpperBound.IsOpen())
 			{
 				UpperBoundToUse = SectionUpperBound;
 			}
 		}
-		if ( !UpperBoundToUse.IsSet() && PlaybackRange.HasUpperBound() )
+		if (!UpperBoundToUse.IsSet() && PlaybackRange.HasUpperBound())
 		{
 			TRangeBound<FFrameNumber> PlaybackUpperBound = PlaybackRange.GetUpperBound();
-			if ( !PlaybackUpperBound.IsOpen() )
+			if (!PlaybackUpperBound.IsOpen())
 			{
 				UpperBoundToUse = PlaybackUpperBound;
 			}
 		}
-		if ( !UpperBoundToUse.IsSet() )
+		if (!UpperBoundToUse.IsSet())
 		{
 			return false;
 		}
-		EndInclTickFrame = UpperBoundToUse.GetValue().GetValue() + ( UpperBoundToUse.GetValue().IsInclusive() ? 0 : -1 );
+		EndInclTickFrame = UpperBoundToUse.GetValue().GetValue() + (UpperBoundToUse.GetValue().IsInclusive() ? 0 : -1);
 	}
 
 	UsdUtils::NotifyIfOverriddenOpinion(BlendShapeWeightsAttr);
@@ -3618,17 +3810,17 @@ bool UnrealToUsd::ConvertControlRigSection(
 	// The best we can do at the moment is to make one channel for each curve on the SkelAnimation prim, but maintain
 	// each Mesh prim connected only to the primary blend shape channel, if it was originally. We'll show a warning
 	// explaining the situation though.
-	if ( InBlendShapeMap && InSkelRoot )
+	if (InBlendShapeMap && InSkelRoot)
 	{
 		TArray<FRigCurveElement*> CurveElements = InitialHierarchy->GetCurves();
 
-		CurveNames.reserve( CurveElements.Num() );
-		for ( int32 Index = 0; Index < CurveElements.Num(); ++Index )
+		CurveNames.reserve(CurveElements.Num());
+		for (int32 Index = 0; Index < CurveElements.Num(); ++Index)
 		{
 			FRigCurveElement* Element = CurveElements[Index];
 			const FString& CurveNameString = Element->GetName();
 
-			CurveNames.push_back( UnrealToUsd::ConvertToken( *CurveNameString ).Get() );
+			CurveNames.push_back(UnrealToUsd::ConvertToken(*CurveNameString).Get());
 		}
 
 		// Check if the blend shape channels on skel animation are the same names as morph target curves.
@@ -3637,20 +3829,20 @@ bool UnrealToUsd::ConvertControlRigSection(
 		// already have the processed, "one per morph target" channels
 		bool bNeedChannelUpdate = true;
 		pxr::VtArray<pxr::TfToken> SkelAnimBlendShapeChannels;
-		if ( BlendShapesAttr && BlendShapesAttr.Get( &SkelAnimBlendShapeChannels ) )
+		if (BlendShapesAttr && BlendShapesAttr.Get(&SkelAnimBlendShapeChannels))
 		{
-			if ( SkelAnimBlendShapeChannels.size() == CurveNames.size() )
+			if (SkelAnimBlendShapeChannels.size() == CurveNames.size())
 			{
 				std::unordered_set<pxr::TfToken, pxr::TfToken::HashFunctor> ExistingCurveNames;
-				for ( const pxr::TfToken& Channel : SkelAnimBlendShapeChannels )
+				for (const pxr::TfToken& Channel : SkelAnimBlendShapeChannels)
 				{
-					ExistingCurveNames.insert( Channel );
+					ExistingCurveNames.insert(Channel);
 				}
 
 				bool bFoundAllCurves = true;
-				for ( const pxr::TfToken& CurveName : CurveNames )
+				for (const pxr::TfToken& CurveName : CurveNames)
 				{
-					if ( ExistingCurveNames.count( CurveName ) == 0 )
+					if (ExistingCurveNames.count(CurveName) == 0)
 					{
 						bFoundAllCurves = false;
 						break;
@@ -3668,17 +3860,17 @@ bool UnrealToUsd::ConvertControlRigSection(
 		// control rigs; Have some morph target curves unintuitively "mirror each other" in UE, if at all possible;
 		// Try to keep the channels shared on USD's side, which would desync USD/UE and show a different result when
 		// reloading.
-		if ( bNeedChannelUpdate )
+		if (bNeedChannelUpdate)
 		{
 			// We'll change the blend shape channel names, so we need to update all meshes that were using them too.
 			// For now we'll assume that they're all inside the same skel root. We could upgrade this for the stage
 			// later too, if needed
 			// TODO: This could probably be updated to just find the actual skinned meshes, and have some better parameters
 			// like skinning/skeleton queries
-			for ( UE::FUsdPrim& MeshPrim : UsdUtils::GetAllPrimsOfType( UE::FUsdPrim{ InSkelRoot }, TEXT( "UsdGeomMesh" ) ) )
+			for (UE::FUsdPrim& MeshPrim : UsdUtils::GetAllPrimsOfType(UE::FUsdPrim{InSkelRoot}, TEXT("UsdGeomMesh")))
 			{
-				pxr::UsdSkelBindingAPI SkelBindingAPI{ MeshPrim };
-				if ( !SkelBindingAPI )
+				pxr::UsdSkelBindingAPI SkelBindingAPI{MeshPrim};
+				if (!SkelBindingAPI)
 				{
 					continue;
 				}
@@ -3686,40 +3878,48 @@ bool UnrealToUsd::ConvertControlRigSection(
 				pxr::UsdRelationship TargetsRel = SkelBindingAPI.GetBlendShapeTargetsRel();
 				pxr::UsdAttribute ChannelsAttr = SkelBindingAPI.GetBlendShapesAttr();
 
-				if ( TargetsRel && ChannelsAttr )
+				if (TargetsRel && ChannelsAttr)
 				{
 					pxr::SdfPathVector BlendShapeTargets;
-					if ( TargetsRel.GetTargets( &BlendShapeTargets ) )
+					if (TargetsRel.GetTargets(&BlendShapeTargets))
 					{
 						pxr::VtArray<pxr::TfToken> BlendShapeChannels;
-						ChannelsAttr.Get( &BlendShapeChannels );
+						ChannelsAttr.Get(&BlendShapeChannels);
 
-						BlendShapeChannels.resize( BlendShapeTargets.size() );
+						BlendShapeChannels.resize(BlendShapeTargets.size());
 
-						pxr::SdfPath MeshPath{ MeshPrim.GetPrimPath() };
+						pxr::SdfPath MeshPath{MeshPrim.GetPrimPath()};
 
 						bool bRenamedAChannel = false;
-						for ( int32 BlendShapeIndex = 0; BlendShapeIndex < BlendShapeTargets.size(); ++BlendShapeIndex )
+						for (int32 BlendShapeIndex = 0; BlendShapeIndex < BlendShapeTargets.size(); ++BlendShapeIndex)
 						{
-							const pxr::SdfPath& BlendShapePath = BlendShapeTargets[ BlendShapeIndex ];
-							FString PrimaryBlendShapePath = UsdToUnreal::ConvertPath( BlendShapePath.MakeAbsolutePath( MeshPath ) );
+							const pxr::SdfPath& BlendShapePath = BlendShapeTargets[BlendShapeIndex];
+							FString PrimaryBlendShapePath = UsdToUnreal::ConvertPath(BlendShapePath.MakeAbsolutePath(MeshPath));
 
 							// Mesh had <blendshape1> target on channel "C" -> We have a morph target called "blendshape1" already, and
 							// we'll create a new channel on SkelAnimation called "blendshape1" -> Let's replace channel "C" with channel
 							// "blendshape1"
-							if ( const UsdUtils::FUsdBlendShape* FoundBlendShape = InBlendShapeMap->Find( PrimaryBlendShapePath ) )
+							if (const UsdUtils::FUsdBlendShape* FoundBlendShape = InBlendShapeMap->Find(PrimaryBlendShapePath))
 							{
 								bRenamedAChannel = true;
-								BlendShapeChannels[ BlendShapeIndex ] = UnrealToUsd::ConvertToken( *FoundBlendShape->Name ).Get();
-								UE_LOG( LogUsd, Log, TEXT( "Updating Mesh '%s' to bind BlendShape target '%s' to SkelAnimation curve '%s'" ),
+								BlendShapeChannels[BlendShapeIndex] = UnrealToUsd::ConvertToken(*FoundBlendShape->Name).Get();
+								UE_LOG(
+									LogUsd,
+									Log,
+									TEXT("Updating Mesh '%s' to bind BlendShape target '%s' to SkelAnimation curve '%s'"),
 									*MeshPrim.GetPrimPath().GetString(),
 									*PrimaryBlendShapePath,
 									*FoundBlendShape->Name
 								);
 
-								if ( FoundBlendShape->Inbetweens.Num() > 0 )
+								if (FoundBlendShape->Inbetweens.Num() > 0)
 								{
-									UE_LOG( LogUsd, Warning, TEXT( "Baking Control Rig parameter sections for BlendShapes with inbetweens (like '%s') is not currently supported, so animation for mesh '%s' may look incorrect! Please flatten the inbetweens into separate BlendShapes beforehand (importing and exporting will do that)." ),
+									UE_LOG(
+										LogUsd,
+										Warning,
+										TEXT("Baking Control Rig parameter sections for BlendShapes with inbetweens (like '%s') is not currently "
+											 "supported, so animation for mesh '%s' may look incorrect! Please flatten the inbetweens into separate "
+											 "BlendShapes beforehand (importing and exporting will do that)."),
 										*FoundBlendShape->Name,
 										*MeshPrim.GetPrimPath().GetString()
 									);
@@ -3727,9 +3927,9 @@ bool UnrealToUsd::ConvertControlRigSection(
 							}
 						}
 
-						if ( bRenamedAChannel )
+						if (bRenamedAChannel)
 						{
-							ChannelsAttr.Set( BlendShapeChannels );
+							ChannelsAttr.Set(BlendShapeChannels);
 							UsdUtils::NotifyIfOverriddenOpinion(ChannelsAttr);
 						}
 					}
@@ -3740,40 +3940,43 @@ bool UnrealToUsd::ConvertControlRigSection(
 		// Now that we updated the channel names we need to make sure we clear the previous weights as they'll
 		// make no sense
 		BlendShapeWeightsAttr.Clear();
-		BlendShapesAttr.Set( CurveNames );
-		CurvesValuesForTime.resize( CurveNames.size() );
+		BlendShapesAttr.Set(CurveNames);
+		CurvesValuesForTime.resize(CurveNames.size());
 
 		UsdUtils::NotifyIfOverriddenOpinion(BlendShapesAttr);
 	}
 
-	FFrameTime TickIncr = FFrameRate::TransformTime( 1, DisplayRate, TickResolution );
-	for ( FFrameTime FrameTickTime = StartInclTickFrame; FrameTickTime <= EndInclTickFrame; FrameTickTime += TickIncr )
+	FFrameTime TickIncr = FFrameRate::TransformTime(1, DisplayRate, TickResolution);
+	for (FFrameTime FrameTickTime = StartInclTickFrame; FrameTickTime <= EndInclTickFrame; FrameTickTime += TickIncr)
 	{
 		FFrameTime TransformedFrameTickTime = FrameTickTime * InTransform;
 
-		double UsdTimeCode = FFrameRate::TransformTime( TransformedFrameTickTime, TickResolution, StageFrameRate ).AsDecimal();
+		double UsdTimeCode = FFrameRate::TransformTime(TransformedFrameTickTime, TickResolution, StageFrameRate).AsDecimal();
 
 		FMovieSceneContext Context = FMovieSceneContext(
-			FMovieSceneEvaluationRange(
-				TransformedFrameTickTime,
-				TickResolution
-			),
-			InPlayer->GetPlaybackStatus()
-		).SetHasJumped( true );
-
-		InPlayer->GetEvaluationTemplate().EvaluateSynchronousBlocking( Context );
+										 FMovieSceneEvaluationRange(TransformedFrameTickTime, TickResolution),
+										 InPlayer->GetPlaybackStatus()
+		)
+										 .SetHasJumped(true);
+
+		InPlayer->GetEvaluationTemplate().EvaluateSynchronousBlocking(Context);
 		ControlRig->Evaluate_AnyThread();
 
 		URigHierarchy* Hierarchy = ControlRig->GetHierarchy();
-		if ( !Hierarchy )
+		if (!Hierarchy)
 		{
-			UE_LOG( LogUsd, Error, TEXT( "Baking Control Rig tracks for rig '%s' failed" ), *ControlRig->GetPathName() );
+			UE_LOG(LogUsd, Error, TEXT("Baking Control Rig tracks for rig '%s' failed"), *ControlRig->GetPathName());
 			return false;
 		}
 
 		if (Hierarchy->GetTopologyVersion() != TopologyVersion)
 		{
-			UE_LOG(LogUsd, Log, TEXT("Regenerating ControlRig to reference skeleton mapping for rig '%s' as its topology changed"), *ControlRig->GetPathName());
+			UE_LOG(
+				LogUsd,
+				Log,
+				TEXT("Regenerating ControlRig to reference skeleton mapping for rig '%s' as its topology changed"),
+				*ControlRig->GetPathName()
+			);
 			RegenerateRigJointIndexToRefSkeletonIndex();
 			TopologyVersion = Hierarchy->GetTopologyVersion();
 		}
@@ -3782,61 +3985,64 @@ bool UnrealToUsd::ConvertControlRigSection(
 		// (c.f. FControlRigBindingHelper::BindToSequencerInstance, URigHierarchy::CopyHierarchy)
 		TArray<FRigBoneElement*> BoneElements = Hierarchy->GetBones();
 
-		if ( CurvesValuesForTime.size() > 0 )
+		if (CurvesValuesForTime.size() > 0)
 		{
 			TArray<FRigCurveElement*> CurveElements = Hierarchy->GetCurves();
-			for ( int32 ElementIndex = 0; ElementIndex < CurveElements.Num(); ++ElementIndex )
+			for (int32 ElementIndex = 0; ElementIndex < CurveElements.Num(); ++ElementIndex)
 			{
-				FRigCurveElement* Element = CurveElements[ ElementIndex ];
-				CurvesValuesForTime[ ElementIndex ] = Hierarchy->GetCurveValue( Element );
+				FRigCurveElement* Element = CurveElements[ElementIndex];
+				CurvesValuesForTime[ElementIndex] = Hierarchy->GetCurveValue(Element);
 			}
 
-			BlendShapeWeightsAttr.Set( CurvesValuesForTime, pxr::UsdTimeCode( UsdTimeCode ) );
+			BlendShapeWeightsAttr.Set(CurvesValuesForTime, pxr::UsdTimeCode(UsdTimeCode));
 		}
 
-		for ( int32 RigBoneIndex = 0; RigBoneIndex < BoneElements.Num(); ++RigBoneIndex )
+		for (int32 RigBoneIndex = 0; RigBoneIndex < BoneElements.Num(); ++RigBoneIndex)
 		{
-			FRigBoneElement* El = BoneElements[ RigBoneIndex ];
+			FRigBoneElement* El = BoneElements[RigBoneIndex];
 
 			// Our skeleton doesn't have this rig bone
-			int32 RefSkeletonBoneIndex = RigJointIndexToRefSkeletonIndex[ RigBoneIndex ];
-			if ( RefSkeletonBoneIndex == INDEX_NONE )
+			int32 RefSkeletonBoneIndex = RigJointIndexToRefSkeletonIndex[RigBoneIndex];
+			if (RefSkeletonBoneIndex == INDEX_NONE)
 			{
 				continue;
 			}
 
-			GlobalUEJointTransformsForFrame[ RefSkeletonBoneIndex ] = Hierarchy->GetTransform( El, ERigTransformType::CurrentGlobal );
+			GlobalUEJointTransformsForFrame[RefSkeletonBoneIndex] = Hierarchy->GetTransform(El, ERigTransformType::CurrentGlobal);
 
 			FTransform UsdTransform;
 
 			// We have to calculate the local transforms ourselves since the parent element could be a control
-			int32 RefSkeletonParentBoneIndex = InRefSkeleton.GetParentIndex( RefSkeletonBoneIndex );
-			if ( RefSkeletonParentBoneIndex == INDEX_NONE )
+			int32 RefSkeletonParentBoneIndex = InRefSkeleton.GetParentIndex(RefSkeletonBoneIndex);
+			if (RefSkeletonParentBoneIndex == INDEX_NONE)
 			{
-				UsdTransform = UsdUtils::ConvertTransformToUsdSpace( StageInfo, GlobalUEJointTransformsForFrame[ RefSkeletonBoneIndex ] );
+				UsdTransform = UsdUtils::ConvertTransformToUsdSpace(StageInfo, GlobalUEJointTransformsForFrame[RefSkeletonBoneIndex]);
 			}
 			else
 			{
-				const FTransform& ChildGlobal = GlobalUEJointTransformsForFrame[ RefSkeletonBoneIndex ];
-				const FTransform& ParentGlobal = GlobalUEJointTransformsForFrame[ RefSkeletonParentBoneIndex ];
-				UsdTransform = UsdUtils::ConvertTransformToUsdSpace( StageInfo, ChildGlobal.GetRelativeTransform( ParentGlobal ) );
+				const FTransform& ChildGlobal = GlobalUEJointTransformsForFrame[RefSkeletonBoneIndex];
+				const FTransform& ParentGlobal = GlobalUEJointTransformsForFrame[RefSkeletonParentBoneIndex];
+				UsdTransform = UsdUtils::ConvertTransformToUsdSpace(StageInfo, ChildGlobal.GetRelativeTransform(ParentGlobal));
 			}
 
-			Translations[ RefSkeletonBoneIndex ] = UnrealToUsd::ConvertVectorFloat( UsdTransform.GetTranslation() );
-			Rotations[ RefSkeletonBoneIndex ] = UnrealToUsd::ConvertQuatFloat( UsdTransform.GetRotation() ).GetNormalized();
-			Scales[ RefSkeletonBoneIndex ] = UnrealToUsd::ConvertVectorHalf( UsdTransform.GetScale3D() );
+			Translations[RefSkeletonBoneIndex] = UnrealToUsd::ConvertVectorFloat(UsdTransform.GetTranslation());
+			Rotations[RefSkeletonBoneIndex] = UnrealToUsd::ConvertQuatFloat(UsdTransform.GetRotation()).GetNormalized();
+			Scales[RefSkeletonBoneIndex] = UnrealToUsd::ConvertVectorHalf(UsdTransform.GetScale3D());
 		}
 
-		TranslationsAttr.Set( Translations, pxr::UsdTimeCode( UsdTimeCode ) );
-		RotationsAttr.Set( Rotations, pxr::UsdTimeCode( UsdTimeCode ) );
-		ScalesAttr.Set( Scales, pxr::UsdTimeCode( UsdTimeCode ) );
+		TranslationsAttr.Set(Translations, pxr::UsdTimeCode(UsdTimeCode));
+		RotationsAttr.Set(Rotations, pxr::UsdTimeCode(UsdTimeCode));
+		ScalesAttr.Set(Scales, pxr::UsdTimeCode(UsdTimeCode));
 	}
 
-	double ElapsedSeconds = FPlatformTime::ToSeconds64( FPlatformTime::Cycles64() - StartTime );
-	int ElapsedMin = int( ElapsedSeconds / 60.0 );
-	ElapsedSeconds -= 60.0 * ( double ) ElapsedMin;
-	UE_LOG( LogUsd, Log, TEXT( "Baked new animation for prim '%s' in [%d min %.3f s]" ),
-		*UsdToUnreal::ConvertPath( OutSkelAnimPrim.GetPrimPath() ),
+	double ElapsedSeconds = FPlatformTime::ToSeconds64(FPlatformTime::Cycles64() - StartTime);
+	int ElapsedMin = int(ElapsedSeconds / 60.0);
+	ElapsedSeconds -= 60.0 * (double)ElapsedMin;
+	UE_LOG(
+		LogUsd,
+		Log,
+		TEXT("Baked new animation for prim '%s' in [%d min %.3f s]"),
+		*UsdToUnreal::ConvertPath(OutSkelAnimPrim.GetPrimPath()),
 		ElapsedMin,
 		ElapsedSeconds
 	);
@@ -3846,4 +4052,4 @@ bool UnrealToUsd::ConvertControlRigSection(
 
 #undef LOCTEXT_NAMESPACE
 
-#endif // #if USE_USD_SDK && WITH_EDITOR
+#endif	  // #if USE_USD_SDK && WITH_EDITOR
