USD: Reformat USDUtilities module.

diff --git a/Engine/Plugins/Importers/USDImporter/Source/USDUtilities/Private/USDTypesConversion.cpp b/Engine/Plugins/Importers/USDImporter/Source/USDUtilities/Private/USDTypesConversion.cpp
--- a/Engine/Plugins/Importers/USDImporter/Source/USDUtilities/Private/USDTypesConversion.cpp
+++ b/Engine/Plugins/Importers/USDImporter/Source/USDUtilities/Private/USDTypesConversion.cpp
@@ -9,21 +9,24 @@
 #if USE_USD_SDK
 
 #include "USDIncludesStart.h"
-	#include "pxr/base/gf/quatf.h"
-	#include "pxr/base/gf/vec2f.h"
-	#include "pxr/base/gf/vec3f.h"
-	#include "pxr/base/gf/vec4f.h"
-
-	#include "pxr/usd/sdf/path.h"
-	#include "pxr/usd/usd/stage.h"
-	#include "pxr/usd/usdGeom/tokens.h"
+#include "pxr/base/gf/matrix2d.h"
+#include "pxr/base/gf/matrix3d.h"
+#include "pxr/base/gf/matrix4d.h"
+#include "pxr/base/gf/quatd.h"
+#include "pxr/base/gf/quatf.h"
+#include "pxr/base/gf/quath.h"
+#include "pxr/base/gf/vec2f.h"
+#include "pxr/base/gf/vec3f.h"
+#include "pxr/base/gf/vec4f.h"
+#include "pxr/usd/sdf/path.h"
+#include "pxr/usd/usdGeom/tokens.h"
 #include "USDIncludesEnd.h"
 
-FUsdStageInfo::FUsdStageInfo( const pxr::UsdStageRefPtr& Stage )
+FUsdStageInfo::FUsdStageInfo(const pxr::UsdStageRefPtr& Stage)
 {
-	pxr::TfToken UsdStageAxis = UsdUtils::GetUsdStageUpAxis( Stage );
+	pxr::TfToken UsdStageAxis = UsdUtils::GetUsdStageUpAxis(Stage);
 
-	if ( UsdStageAxis == pxr::UsdGeomTokens->y )
+	if (UsdStageAxis == pxr::UsdGeomTokens->y)
 	{
 		UpAxis = EUsdUpAxis::YAxis;
 	}
@@ -32,18 +35,18 @@ FUsdStageInfo::FUsdStageInfo( const pxr::UsdStageRefPtr& Stage )
 		UpAxis = EUsdUpAxis::ZAxis;
 	}
 
-	MetersPerUnit = UsdUtils::GetUsdStageMetersPerUnit( Stage );
+	MetersPerUnit = UsdUtils::GetUsdStageMetersPerUnit(Stage);
 }
 
 namespace UsdUtils
 {
-	FTransform ConvertAxes( const bool bZUp, const FTransform Transform )
+	FTransform ConvertAxes(const bool bZUp, const FTransform Transform)
 	{
 		FVector Translation = Transform.GetTranslation();
 		FQuat Rotation = Transform.GetRotation();
 		FVector Scale = Transform.GetScale3D();
 
-		if ( bZUp )
+		if (bZUp)
 		{
 			Translation.Y = -Translation.Y;
 			Rotation.X = -Rotation.X;
@@ -51,71 +54,71 @@ namespace UsdUtils
 		}
 		else
 		{
-			Swap( Translation.Y, Translation.Z );
+			Swap(Translation.Y, Translation.Z);
 
 			Rotation = Rotation.Inverse();
-			Swap( Rotation.Y, Rotation.Z );
+			Swap(Rotation.Y, Rotation.Z);
 
-			Swap( Scale.Y, Scale.Z );
+			Swap(Scale.Y, Scale.Z);
 		}
 
-		return FTransform( Rotation, Translation, Scale );
+		return FTransform(Rotation, Translation, Scale);
 	}
 
-	FTransform ConvertTransformToUsdSpace( const FUsdStageInfo& StageInfo, const FTransform& TransformInUESpace )
+	FTransform ConvertTransformToUsdSpace(const FUsdStageInfo& StageInfo, const FTransform& TransformInUESpace)
 	{
-		FTransform TransformInUsdSpace = UsdUtils::ConvertAxes( StageInfo.UpAxis == EUsdUpAxis::ZAxis, TransformInUESpace );
+		FTransform TransformInUsdSpace = UsdUtils::ConvertAxes(StageInfo.UpAxis == EUsdUpAxis::ZAxis, TransformInUESpace);
 
 		const float UEMetersPerUnit = 0.01f;
-		if ( !FMath::IsNearlyEqual( StageInfo.MetersPerUnit, UEMetersPerUnit ) )
+		if (!FMath::IsNearlyEqual(StageInfo.MetersPerUnit, UEMetersPerUnit))
 		{
-			TransformInUsdSpace.ScaleTranslation( UEMetersPerUnit / StageInfo.MetersPerUnit );
+			TransformInUsdSpace.ScaleTranslation(UEMetersPerUnit / StageInfo.MetersPerUnit);
 		}
 
 		return TransformInUsdSpace;
 	}
-}
+}	 // namespace UsdUtils
 
 namespace UsdToUnreal
 {
-	FString ConvertString( const std::string& InString )
+	FString ConvertString(const std::string& InString)
 	{
-		return FString( ANSI_TO_TCHAR( InString.c_str() ) );
+		return FString(ANSI_TO_TCHAR(InString.c_str()));
 	}
 
-	FString ConvertString( const char* InString )
+	FString ConvertString(const char* InString)
 	{
-		return FString( ANSI_TO_TCHAR( InString ) );
+		return FString(ANSI_TO_TCHAR(InString));
 	}
 
-	FString ConvertPath( const pxr::SdfPath& Path )
+	FString ConvertPath(const pxr::SdfPath& Path)
 	{
-		return ConvertString( Path.GetString().c_str() );
+		return ConvertString(Path.GetString().c_str());
 	}
 
-	FName ConvertName( const char* InString )
+	FName ConvertName(const char* InString)
 	{
-		return FName( InString );
+		return FName(InString);
 	}
 
-	FName ConvertName( const std::string& InString )
+	FName ConvertName(const std::string& InString)
 	{
-		return FName( InString.c_str() );
+		return FName(InString.c_str());
 	}
 
-	FString ConvertToken( const pxr::TfToken& Token )
+	FString ConvertToken(const pxr::TfToken& Token)
 	{
-		return UsdToUnreal::ConvertString( Token.GetString() );
+		return UsdToUnreal::ConvertString(Token.GetString());
 	}
 
-	FLinearColor ConvertColor( const pxr::GfVec3f& InValue )
+	FLinearColor ConvertColor(const pxr::GfVec3f& InValue)
 	{
-		return FLinearColor( InValue[0], InValue[1], InValue[2] );
+		return FLinearColor(InValue[0], InValue[1], InValue[2]);
 	}
 
-	FLinearColor ConvertColor( const pxr::GfVec4f& InValue )
+	FLinearColor ConvertColor(const pxr::GfVec4f& InValue)
 	{
-		return FLinearColor( InValue[0], InValue[1], InValue[2], InValue[3] );
+		return FLinearColor(InValue[0], InValue[1], InValue[2], InValue[3]);
 	}
 
 	FVector2D ConvertVector(const pxr::GfVec2h& InValue)
@@ -235,7 +238,7 @@ namespace UsdToUnreal
 		return Result;
 	}
 
-	FMatrix ConvertMatrix( const pxr::GfMatrix4d& Matrix )
+	FMatrix ConvertMatrix(const pxr::GfMatrix4d& Matrix)
 	{
 		FMatrix UnrealMatrix(
 			FPlane(Matrix[0][0], Matrix[0][1], Matrix[0][2], Matrix[0][3]),
@@ -247,17 +250,17 @@ namespace UsdToUnreal
 		return UnrealMatrix;
 	}
 
-	FTransform ConvertMatrix( const FUsdStageInfo& StageInfo, const pxr::GfMatrix4d& InMatrix )
+	FTransform ConvertMatrix(const FUsdStageInfo& StageInfo, const pxr::GfMatrix4d& InMatrix)
 	{
-		FMatrix Matrix = ConvertMatrix( InMatrix );
-		FTransform Transform( Matrix );
+		FMatrix Matrix = ConvertMatrix(InMatrix);
+		FTransform Transform(Matrix);
 
-		Transform = UsdUtils::ConvertAxes( StageInfo.UpAxis == EUsdUpAxis::ZAxis, Transform );
+		Transform = UsdUtils::ConvertAxes(StageInfo.UpAxis == EUsdUpAxis::ZAxis, Transform);
 
 		const float UEMetersPerUnit = 0.01f;
-		if ( !FMath::IsNearlyEqual( StageInfo.MetersPerUnit, UEMetersPerUnit ) )
+		if (!FMath::IsNearlyEqual(StageInfo.MetersPerUnit, UEMetersPerUnit))
 		{
-			Transform.ScaleTranslation( StageInfo.MetersPerUnit / UEMetersPerUnit );
+			Transform.ScaleTranslation(StageInfo.MetersPerUnit / UEMetersPerUnit);
 		}
 
 		return Transform;
@@ -281,48 +284,48 @@ namespace UsdToUnreal
 		return FQuat{Imaginary[0], Imaginary[1], Imaginary[2], InValue.GetReal()};
 	}
 
-	float ConvertDistance( const FUsdStageInfo& StageInfo, float Value )
+	float ConvertDistance(const FUsdStageInfo& StageInfo, float Value)
 	{
 		const float UEMetersPerUnit = 0.01f;
-		if ( !FMath::IsNearlyEqual( StageInfo.MetersPerUnit, UEMetersPerUnit ) )
+		if (!FMath::IsNearlyEqual(StageInfo.MetersPerUnit, UEMetersPerUnit))
 		{
 			Value *= StageInfo.MetersPerUnit / UEMetersPerUnit;
 		}
 
 		return Value;
 	}
-}
+}	 // namespace UsdToUnreal
 
 namespace UnrealToUsd
 {
-	TUsdStore< std::string > ConvertString( const TCHAR* InString )
+	TUsdStore<std::string> ConvertString(const TCHAR* InString)
 	{
-		return MakeUsdStore< std::string >( TCHAR_TO_ANSI( InString ) );
+		return MakeUsdStore<std::string>(TCHAR_TO_ANSI(InString));
 	}
 
-	TUsdStore< pxr::SdfPath > ConvertPath( const TCHAR* InString )
+	TUsdStore<pxr::SdfPath> ConvertPath(const TCHAR* InString)
 	{
-		return MakeUsdStore< pxr::SdfPath >( TCHAR_TO_ANSI( InString ) );
+		return MakeUsdStore<pxr::SdfPath>(TCHAR_TO_ANSI(InString));
 	}
 
-	TUsdStore< std::string > ConvertName( const FName& InName )
+	TUsdStore<std::string> ConvertName(const FName& InName)
 	{
-		return MakeUsdStore< std::string >( TCHAR_TO_ANSI( *InName.ToString() ) );
+		return MakeUsdStore<std::string>(TCHAR_TO_ANSI(*InName.ToString()));
 	}
 
-	TUsdStore< pxr::TfToken > ConvertToken( const TCHAR* InString )
+	TUsdStore<pxr::TfToken> ConvertToken(const TCHAR* InString)
 	{
-		return MakeUsdStore< pxr::TfToken >( TCHAR_TO_ANSI( InString ) );
+		return MakeUsdStore<pxr::TfToken>(TCHAR_TO_ANSI(InString));
 	}
 
-	pxr::GfVec4f ConvertColor( const FLinearColor& InValue )
+	pxr::GfVec4f ConvertColor(const FLinearColor& InValue)
 	{
-		return pxr::GfVec4f( InValue.R, InValue.G, InValue.B, InValue.A );
+		return pxr::GfVec4f(InValue.R, InValue.G, InValue.B, InValue.A);
 	}
 
-	pxr::GfVec4f ConvertColor( const FColor& InValue )
+	pxr::GfVec4f ConvertColor(const FColor& InValue)
 	{
-		return ConvertColor( FLinearColor( InValue ) );
+		return ConvertColor(FLinearColor(InValue));
 	}
 
 	// Deprecated
@@ -461,10 +464,22 @@ namespace UnrealToUsd
 	pxr::GfMatrix4d ConvertMatrix(const FMatrix& Matrix)
 	{
 		pxr::GfMatrix4d UsdMatrix(
-			Matrix.M[0][0], Matrix.M[0][1], Matrix.M[0][2], Matrix.M[0][3],
-			Matrix.M[1][0], Matrix.M[1][1], Matrix.M[1][2], Matrix.M[1][3],
-			Matrix.M[2][0], Matrix.M[2][1], Matrix.M[2][2], Matrix.M[2][3],
-			Matrix.M[3][0], Matrix.M[3][1], Matrix.M[3][2], Matrix.M[3][3]
+			Matrix.M[0][0],
+			Matrix.M[0][1],
+			Matrix.M[0][2],
+			Matrix.M[0][3],
+			Matrix.M[1][0],
+			Matrix.M[1][1],
+			Matrix.M[1][2],
+			Matrix.M[1][3],
+			Matrix.M[2][0],
+			Matrix.M[2][1],
+			Matrix.M[2][2],
+			Matrix.M[2][3],
+			Matrix.M[3][0],
+			Matrix.M[3][1],
+			Matrix.M[3][2],
+			Matrix.M[3][3]
 		);
 
 		return UsdMatrix;
@@ -482,7 +497,8 @@ namespace UnrealToUsd
 			static_cast<pxr::GfHalf>(InValue.W),
 			static_cast<pxr::GfHalf>(InValue.X),
 			static_cast<pxr::GfHalf>(InValue.Y),
-			static_cast<pxr::GfHalf>(InValue.Z)};
+			static_cast<pxr::GfHalf>(InValue.Z)
+		};
 	}
 
 	pxr::GfQuatf ConvertQuatFloat(const FQuat& InValue)
@@ -491,7 +507,8 @@ namespace UnrealToUsd
 			static_cast<float>(InValue.W),
 			static_cast<float>(InValue.X),
 			static_cast<float>(InValue.Y),
-			static_cast<float>(InValue.Z)};
+			static_cast<float>(InValue.Z)
+		};
 	}
 
 	pxr::GfQuatd ConvertQuatDouble(const FQuat& InValue)
@@ -499,23 +516,23 @@ namespace UnrealToUsd
 		return pxr::GfQuatd{InValue.W, InValue.X, InValue.Y, InValue.Z};
 	}
 
-	pxr::GfMatrix4d ConvertTransform( const FUsdStageInfo& StageInfo, const FTransform& Transform )
+	pxr::GfMatrix4d ConvertTransform(const FUsdStageInfo& StageInfo, const FTransform& Transform)
 	{
-		FTransform TransformInUsdSpace = UsdUtils::ConvertTransformToUsdSpace( StageInfo, Transform );
+		FTransform TransformInUsdSpace = UsdUtils::ConvertTransformToUsdSpace(StageInfo, Transform);
 
-		return ConvertMatrix( TransformInUsdSpace.ToMatrixWithScale() );
+		return ConvertMatrix(TransformInUsdSpace.ToMatrixWithScale());
 	}
 
-	float ConvertDistance( const FUsdStageInfo& StageInfo, float Value )
+	float ConvertDistance(const FUsdStageInfo& StageInfo, float Value)
 	{
 		const float UEMetersPerUnit = 0.01f;
-		if ( !FMath::IsNearlyEqual( StageInfo.MetersPerUnit, UEMetersPerUnit ) )
+		if (!FMath::IsNearlyEqual(StageInfo.MetersPerUnit, UEMetersPerUnit))
 		{
 			Value *= UEMetersPerUnit / StageInfo.MetersPerUnit;
 		}
 
 		return Value;
 	}
-}
+}	 // namespace UnrealToUsd
 
-#endif // #if USE_USD_SDK
+#endif	  // #if USE_USD_SDK
diff --git a/Engine/Plugins/Importers/USDImporter/Source/USDUtilities/Private/USDValueConversion.cpp b/Engine/Plugins/Importers/USDImporter/Source/USDUtilities/Private/USDValueConversion.cpp
--- a/Engine/Plugins/Importers/USDImporter/Source/USDUtilities/Private/USDValueConversion.cpp
+++ b/Engine/Plugins/Importers/USDImporter/Source/USDUtilities/Private/USDValueConversion.cpp
@@ -3,27 +3,20 @@
 #include "USDValueConversion.h"
 
 #include "USDConversionUtils.h"
-#include "USDErrorUtils.h"
 #include "USDLog.h"
 #include "USDTypesConversion.h"
 
-#include "UsdWrappers/SdfLayer.h"
 #include "UsdWrappers/UsdStage.h"
 #include "UsdWrappers/VtValue.h"
 
-#include "Containers/StringConv.h"
-
 #if USE_USD_SDK
 #include "USDIncludesStart.h"
-	#include "pxr/base/tf/stringUtils.h"
-	#include "pxr/base/vt/value.h"
-	#include "pxr/usd/sdf/path.h"
-	#include "pxr/usd/sdf/types.h"
-	#include "pxr/usd/usd/attribute.h"
-	#include "pxr/usd/usd/stage.h"
-	#include "pxr/usd/usdGeom/tokens.h"
+#include "pxr/base/tf/stringUtils.h"
+#include "pxr/base/vt/value.h"
+#include "pxr/usd/sdf/schema.h"
+#include "pxr/usd/sdf/types.h"
 #include "USDIncludesEnd.h"
-#endif // USE_USD_SDK
+#endif	  // USE_USD_SDK
 
 namespace UE::USDValueConversion::Private
 {
@@ -47,101 +40,113 @@ namespace UsdToUnrealImpl
 	 * []( const USDType& Val ) -> UsdUtils::FConvertedVtValueEntry {}
 	 */
 	template<typename UEType, typename USDType, typename Func>
-	void ConvertInner( const pxr::VtValue& InValue, UsdUtils::FConvertedVtValue& OutValue, Func Function )
+	void ConvertInner(const pxr::VtValue& InValue, UsdUtils::FConvertedVtValue& OutValue, Func Function)
 	{
-		if ( InValue.IsArrayValued() )
+		if (InValue.IsArrayValued())
 		{
-			OutValue.Entries.Reset( InValue.GetArraySize() );
+			OutValue.Entries.Reset(InValue.GetArraySize());
 
-			for ( const USDType& Val : InValue.UncheckedGet<pxr::VtArray<USDType>>() )
+			for (const USDType& Val : InValue.UncheckedGet<pxr::VtArray<USDType>>())
 			{
-				OutValue.Entries.Add( Function( Val ) );
+				OutValue.Entries.Add(Function(Val));
 			}
 		}
 		else
 		{
 			const USDType& Val = InValue.UncheckedGet<USDType>();
 
-			OutValue.Entries = { Function( Val ) };
+			OutValue.Entries = {Function(Val)};
 		}
 	}
 
 	// Bool, Uchar, Int, Uint, Int64, Uint64, Half (can cast to float), Float, Double
 	template<typename UEType, typename USDType>
-	void ConvertSimpleValue( const pxr::VtValue& InValue, UsdUtils::FConvertedVtValue& OutValue )
+	void ConvertSimpleValue(const pxr::VtValue& InValue, UsdUtils::FConvertedVtValue& OutValue)
 	{
-		ConvertInner<UEType, USDType>( InValue, OutValue,
-			[]( const USDType& Val ) -> UsdUtils::FConvertedVtValueEntry
+		ConvertInner<UEType, USDType>(
+			InValue,
+			OutValue,
+			[](const USDType& Val) -> UsdUtils::FConvertedVtValueEntry
 			{
-				return { UsdUtils::FConvertedVtValueComponent( TInPlaceType<UEType>(), static_cast< UEType >( Val ) ) };
+				return {UsdUtils::FConvertedVtValueComponent(TInPlaceType<UEType>(), static_cast<UEType>(Val))};
 			}
 		);
 	}
 
 	// TimeCode
 	template<>
-	void ConvertSimpleValue<double, pxr::SdfTimeCode>( const pxr::VtValue& InValue, UsdUtils::FConvertedVtValue& OutValue )
+	void ConvertSimpleValue<double, pxr::SdfTimeCode>(const pxr::VtValue& InValue, UsdUtils::FConvertedVtValue& OutValue)
 	{
-		ConvertInner<double, pxr::SdfTimeCode>( InValue, OutValue,
-			[]( const pxr::SdfTimeCode& Val ) -> UsdUtils::FConvertedVtValueEntry
+		ConvertInner<double, pxr::SdfTimeCode>(
+			InValue,
+			OutValue,
+			[](const pxr::SdfTimeCode& Val) -> UsdUtils::FConvertedVtValueEntry
 			{
-				return { UsdUtils::FConvertedVtValueComponent( TInPlaceType<double>(), Val.GetValue() ) };
+				return {UsdUtils::FConvertedVtValueComponent(TInPlaceType<double>(), Val.GetValue())};
 			}
 		);
 	}
 
 	// String
 	template<>
-	void ConvertSimpleValue<FString, std::string>( const pxr::VtValue& InValue, UsdUtils::FConvertedVtValue& OutValue )
+	void ConvertSimpleValue<FString, std::string>(const pxr::VtValue& InValue, UsdUtils::FConvertedVtValue& OutValue)
 	{
-		ConvertInner<FString, std::string>( InValue, OutValue,
-			[]( const std::string& Val ) -> UsdUtils::FConvertedVtValueEntry
+		ConvertInner<FString, std::string>(
+			InValue,
+			OutValue,
+			[](const std::string& Val) -> UsdUtils::FConvertedVtValueEntry
 			{
-				return { UsdUtils::FConvertedVtValueComponent( TInPlaceType<FString>(), UsdToUnreal::ConvertString( Val ) ) };
+				return {UsdUtils::FConvertedVtValueComponent(TInPlaceType<FString>(), UsdToUnreal::ConvertString(Val))};
 			}
 		);
 	}
 
 	// Token
 	template<>
-	void ConvertSimpleValue<FString, pxr::TfToken>( const pxr::VtValue& InValue, UsdUtils::FConvertedVtValue& OutValue )
+	void ConvertSimpleValue<FString, pxr::TfToken>(const pxr::VtValue& InValue, UsdUtils::FConvertedVtValue& OutValue)
 	{
-		ConvertInner<FString, pxr::TfToken>( InValue, OutValue,
-			[]( const pxr::TfToken& Val ) -> UsdUtils::FConvertedVtValueEntry
+		ConvertInner<FString, pxr::TfToken>(
+			InValue,
+			OutValue,
+			[](const pxr::TfToken& Val) -> UsdUtils::FConvertedVtValueEntry
 			{
-				return { UsdUtils::FConvertedVtValueComponent( TInPlaceType<FString>(), UsdToUnreal::ConvertToken( Val ) ) };
+				return {UsdUtils::FConvertedVtValueComponent(TInPlaceType<FString>(), UsdToUnreal::ConvertToken(Val))};
 			}
 		);
 	}
 
 	// Asset
 	template<>
-	void ConvertSimpleValue<FString, pxr::SdfAssetPath>( const pxr::VtValue& InValue, UsdUtils::FConvertedVtValue& OutValue )
+	void ConvertSimpleValue<FString, pxr::SdfAssetPath>(const pxr::VtValue& InValue, UsdUtils::FConvertedVtValue& OutValue)
 	{
-		ConvertInner<FString, pxr::SdfAssetPath>( InValue, OutValue,
-			[]( const pxr::SdfAssetPath& Val ) -> UsdUtils::FConvertedVtValueEntry
+		ConvertInner<FString, pxr::SdfAssetPath>(
+			InValue,
+			OutValue,
+			[](const pxr::SdfAssetPath& Val) -> UsdUtils::FConvertedVtValueEntry
 			{
-				return { UsdUtils::FConvertedVtValueComponent( TInPlaceType<FString>(), UsdToUnreal::ConvertString( Val.GetAssetPath() ) ) };
+				return {UsdUtils::FConvertedVtValueComponent(TInPlaceType<FString>(), UsdToUnreal::ConvertString(Val.GetAssetPath()))};
 			}
 		);
 	}
 
 	// Matrix2d, Matrix3d, Matrix4d	(always double)
 	template<typename USDType>
-	void ConvertMatrixValue( const pxr::VtValue& InValue, UsdUtils::FConvertedVtValue& OutValue )
+	void ConvertMatrixValue(const pxr::VtValue& InValue, UsdUtils::FConvertedVtValue& OutValue)
 	{
-		ConvertInner<double, USDType>( InValue, OutValue,
-			[]( const USDType& Val ) -> UsdUtils::FConvertedVtValueEntry
+		ConvertInner<double, USDType>(
+			InValue,
+			OutValue,
+			[](const USDType& Val) -> UsdUtils::FConvertedVtValueEntry
 			{
 				const int32 NumElements = USDType::numRows * USDType::numColumns;
 
 				UsdUtils::FConvertedVtValueEntry Entry;
-				Entry.Reserve( NumElements );
+				Entry.Reserve(NumElements);
 
 				const double* MatrixArray = Val.GetArray();
-				for ( int32 Index = 0; Index < NumElements; ++Index )
+				for (int32 Index = 0; Index < NumElements; ++Index)
 				{
-					Entry.Emplace( TInPlaceType<double>(), MatrixArray[ Index ] );
+					Entry.Emplace(TInPlaceType<double>(), MatrixArray[Index]);
 				}
 
 				return Entry;
@@ -151,21 +156,22 @@ namespace UsdToUnrealImpl
 
 	// Quath, Quatf, Quatd
 	template<typename UEType, typename USDType>
-	void ConvertQuatValue( const pxr::VtValue& InValue, UsdUtils::FConvertedVtValue& OutValue )
+	void ConvertQuatValue(const pxr::VtValue& InValue, UsdUtils::FConvertedVtValue& OutValue)
 	{
-		ConvertInner<UEType, USDType>( InValue, OutValue,
-			[]( const USDType& Val ) -> UsdUtils::FConvertedVtValueEntry
+		ConvertInner<UEType, USDType>(
+			InValue,
+			OutValue,
+			[](const USDType& Val) -> UsdUtils::FConvertedVtValueEntry
 			{
 				// Auto here because this is the vec of the corresponding type (e.g. Quath -> Vec3h)
 				const auto& Img = Val.GetImaginary();
 				double Real = Val.GetReal();
 
-				return
-				{
-					UsdUtils::FConvertedVtValueComponent( TInPlaceType<UEType>(), Img[ 0 ] ),
-					UsdUtils::FConvertedVtValueComponent( TInPlaceType<UEType>(), Img[ 1 ] ),
-					UsdUtils::FConvertedVtValueComponent( TInPlaceType<UEType>(), Img[ 2 ] ),
-					UsdUtils::FConvertedVtValueComponent( TInPlaceType<UEType>(), Real ),
+				return {
+					UsdUtils::FConvertedVtValueComponent(TInPlaceType<UEType>(), Img[0]),
+					UsdUtils::FConvertedVtValueComponent(TInPlaceType<UEType>(), Img[1]),
+					UsdUtils::FConvertedVtValueComponent(TInPlaceType<UEType>(), Img[2]),
+					UsdUtils::FConvertedVtValueComponent(TInPlaceType<UEType>(), Real),
 				};
 			}
 		);
@@ -175,29 +181,31 @@ namespace UsdToUnrealImpl
 	// Double3, Float3, Half3, Int3
 	// Double4, Float4, Half4, Int4
 	template<typename UEType, typename USDType>
-	void ConvertVecValue( const pxr::VtValue& InValue, UsdUtils::FConvertedVtValue& OutValue )
+	void ConvertVecValue(const pxr::VtValue& InValue, UsdUtils::FConvertedVtValue& OutValue)
 	{
-		ConvertInner<UEType, USDType>( InValue, OutValue,
-			[]( const USDType& Val ) -> UsdUtils::FConvertedVtValueEntry
+		ConvertInner<UEType, USDType>(
+			InValue,
+			OutValue,
+			[](const USDType& Val) -> UsdUtils::FConvertedVtValueEntry
 			{
 				UsdUtils::FConvertedVtValueEntry Entry;
-				Entry.Reset( USDType::dimension );
+				Entry.Reset(USDType::dimension);
 
-				for ( int32 Index = 0; Index < USDType::dimension; ++Index )
+				for (int32 Index = 0; Index < USDType::dimension; ++Index)
 				{
-					Entry.Emplace( TInPlaceType<UEType>(), Val[ Index ] );
+					Entry.Emplace(TInPlaceType<UEType>(), Val[Index]);
 				}
 
 				return Entry;
 			}
 		);
 	}
-}
-#endif // USE_USD_SDK
+}	 // namespace UsdToUnrealImpl
+#endif	  // USE_USD_SDK
 
 namespace UsdToUnreal
 {
-	bool ConvertValue( const UE::FVtValue& InValue, UsdUtils::FConvertedVtValue& OutValue )
+	bool ConvertValue(const UE::FVtValue& InValue, UsdUtils::FConvertedVtValue& OutValue)
 	{
 #if USE_USD_SDK
 		using namespace pxr;
@@ -210,7 +218,7 @@ namespace UsdToUnreal
 		OutValue.bIsEmpty = true;
 
 		// We consider a success returning an empty value if our input value was also empty
-		if ( InValue.IsEmpty() )
+		if (InValue.IsEmpty())
 		{
 			return true;
 		}
@@ -224,183 +232,183 @@ namespace UsdToUnreal
 #pragma push_macro("CHECK_TYPE")
 #define CHECK_TYPE(T) UnderlyingType.IsA<T>() || UnderlyingType.IsA<VtArray<T>>()
 
-		if ( CHECK_TYPE(bool) )
+		if (CHECK_TYPE(bool))
 		{
 			OutValue.SourceType = EUsdBasicDataTypes::Bool;
-			ConvertSimpleValue<bool, bool>( UsdValue, OutValue );
+			ConvertSimpleValue<bool, bool>(UsdValue, OutValue);
 		}
-		else if ( CHECK_TYPE(uint8_t) )
+		else if (CHECK_TYPE(uint8_t))
 		{
 			OutValue.SourceType = EUsdBasicDataTypes::Uchar;
-			ConvertSimpleValue<uint8, uint8_t>( UsdValue, OutValue );
+			ConvertSimpleValue<uint8, uint8_t>(UsdValue, OutValue);
 		}
-		else if ( CHECK_TYPE(int32_t) )
+		else if (CHECK_TYPE(int32_t))
 		{
 			OutValue.SourceType = EUsdBasicDataTypes::Int;
-			ConvertSimpleValue<int32, int32_t>( UsdValue, OutValue );
+			ConvertSimpleValue<int32, int32_t>(UsdValue, OutValue);
 		}
-		else if ( CHECK_TYPE(uint32_t) )
+		else if (CHECK_TYPE(uint32_t))
 		{
 			OutValue.SourceType = EUsdBasicDataTypes::Uint;
-			ConvertSimpleValue<uint32, uint32_t>( UsdValue, OutValue );
+			ConvertSimpleValue<uint32, uint32_t>(UsdValue, OutValue);
 		}
-		else if ( CHECK_TYPE(int64_t) )
+		else if (CHECK_TYPE(int64_t))
 		{
 			OutValue.SourceType = EUsdBasicDataTypes::Int64;
-			ConvertSimpleValue<int64, int64_t>( UsdValue, OutValue );
+			ConvertSimpleValue<int64, int64_t>(UsdValue, OutValue);
 		}
-		else if ( CHECK_TYPE(uint64_t) )
+		else if (CHECK_TYPE(uint64_t))
 		{
 			OutValue.SourceType = EUsdBasicDataTypes::Uint64;
-			ConvertSimpleValue<uint64, uint64_t>( UsdValue, OutValue );
+			ConvertSimpleValue<uint64, uint64_t>(UsdValue, OutValue);
 		}
-		else if ( CHECK_TYPE(GfHalf) )
+		else if (CHECK_TYPE(GfHalf))
 		{
 			OutValue.SourceType = EUsdBasicDataTypes::Half;
-			ConvertSimpleValue<float, GfHalf>( UsdValue, OutValue );
+			ConvertSimpleValue<float, GfHalf>(UsdValue, OutValue);
 		}
-		else if ( CHECK_TYPE(float) )
+		else if (CHECK_TYPE(float))
 		{
 			OutValue.SourceType = EUsdBasicDataTypes::Float;
-			ConvertSimpleValue<float, float>( UsdValue, OutValue );
+			ConvertSimpleValue<float, float>(UsdValue, OutValue);
 		}
-		else if ( CHECK_TYPE(double) )
+		else if (CHECK_TYPE(double))
 		{
 			OutValue.SourceType = EUsdBasicDataTypes::Double;
-			ConvertSimpleValue<double, double>( UsdValue, OutValue );
+			ConvertSimpleValue<double, double>(UsdValue, OutValue);
 		}
-		else if ( CHECK_TYPE(SdfTimeCode) )
+		else if (CHECK_TYPE(SdfTimeCode))
 		{
 			OutValue.SourceType = EUsdBasicDataTypes::Timecode;
-			ConvertSimpleValue<double, SdfTimeCode>( UsdValue, OutValue );
+			ConvertSimpleValue<double, SdfTimeCode>(UsdValue, OutValue);
 		}
-		else if ( CHECK_TYPE(std::string) )
+		else if (CHECK_TYPE(std::string))
 		{
 			OutValue.SourceType = EUsdBasicDataTypes::String;
-			ConvertSimpleValue<FString, std::string>( UsdValue, OutValue );
+			ConvertSimpleValue<FString, std::string>(UsdValue, OutValue);
 		}
-		else if ( CHECK_TYPE(TfToken) )
+		else if (CHECK_TYPE(TfToken))
 		{
 			OutValue.SourceType = EUsdBasicDataTypes::Token;
-			ConvertSimpleValue<FString, TfToken>( UsdValue, OutValue );
+			ConvertSimpleValue<FString, TfToken>(UsdValue, OutValue);
 		}
-		else if ( CHECK_TYPE(SdfAssetPath) )
+		else if (CHECK_TYPE(SdfAssetPath))
 		{
 			OutValue.SourceType = EUsdBasicDataTypes::Asset;
-			ConvertSimpleValue<FString, SdfAssetPath>( UsdValue, OutValue );
+			ConvertSimpleValue<FString, SdfAssetPath>(UsdValue, OutValue);
 		}
-		else if ( CHECK_TYPE(GfMatrix2d) )
+		else if (CHECK_TYPE(GfMatrix2d))
 		{
 			OutValue.SourceType = EUsdBasicDataTypes::Matrix2d;
-			ConvertMatrixValue<GfMatrix2d>( UsdValue, OutValue );
+			ConvertMatrixValue<GfMatrix2d>(UsdValue, OutValue);
 		}
-		else if ( CHECK_TYPE(GfMatrix3d) )
+		else if (CHECK_TYPE(GfMatrix3d))
 		{
 			OutValue.SourceType = EUsdBasicDataTypes::Matrix3d;
-			ConvertMatrixValue<GfMatrix3d>( UsdValue, OutValue );
+			ConvertMatrixValue<GfMatrix3d>(UsdValue, OutValue);
 		}
-		else if ( CHECK_TYPE(GfMatrix4d) )
+		else if (CHECK_TYPE(GfMatrix4d))
 		{
 			OutValue.SourceType = EUsdBasicDataTypes::Matrix4d;
-			ConvertMatrixValue<GfMatrix4d>( UsdValue, OutValue );
+			ConvertMatrixValue<GfMatrix4d>(UsdValue, OutValue);
 		}
-		else if ( CHECK_TYPE(GfQuatd) )
+		else if (CHECK_TYPE(GfQuatd))
 		{
 			OutValue.SourceType = EUsdBasicDataTypes::Quatd;
 			ConvertQuatValue<double, GfQuatd>(UsdValue, OutValue);
 		}
-		else if ( CHECK_TYPE(GfQuatf) )
+		else if (CHECK_TYPE(GfQuatf))
 		{
 			OutValue.SourceType = EUsdBasicDataTypes::Quatf;
-			ConvertQuatValue<float, GfQuatf>( UsdValue, OutValue );
+			ConvertQuatValue<float, GfQuatf>(UsdValue, OutValue);
 		}
-		else if ( CHECK_TYPE(GfQuath) )
+		else if (CHECK_TYPE(GfQuath))
 		{
 			OutValue.SourceType = EUsdBasicDataTypes::Quath;
-			ConvertQuatValue<float, GfQuath>( UsdValue, OutValue );
+			ConvertQuatValue<float, GfQuath>(UsdValue, OutValue);
 		}
-		else if ( CHECK_TYPE(GfVec2d) )
+		else if (CHECK_TYPE(GfVec2d))
 		{
 			OutValue.SourceType = EUsdBasicDataTypes::Double2;
-			ConvertVecValue<double, GfVec2d>( UsdValue, OutValue );
+			ConvertVecValue<double, GfVec2d>(UsdValue, OutValue);
 		}
-		else if ( CHECK_TYPE(GfVec2f) )
+		else if (CHECK_TYPE(GfVec2f))
 		{
 			OutValue.SourceType = EUsdBasicDataTypes::Float2;
-			ConvertVecValue<float, GfVec2f>( UsdValue, OutValue );
+			ConvertVecValue<float, GfVec2f>(UsdValue, OutValue);
 		}
-		else if ( CHECK_TYPE(GfVec2h) )
+		else if (CHECK_TYPE(GfVec2h))
 		{
 			OutValue.SourceType = EUsdBasicDataTypes::Half2;
-			ConvertVecValue<float, GfVec2h>( UsdValue, OutValue );
+			ConvertVecValue<float, GfVec2h>(UsdValue, OutValue);
 		}
-		else if ( CHECK_TYPE(GfVec2i) )
+		else if (CHECK_TYPE(GfVec2i))
 		{
 			OutValue.SourceType = EUsdBasicDataTypes::Int2;
-			ConvertVecValue<int32, GfVec2i>( UsdValue, OutValue );
+			ConvertVecValue<int32, GfVec2i>(UsdValue, OutValue);
 		}
-		else if ( CHECK_TYPE(GfVec3d) )
+		else if (CHECK_TYPE(GfVec3d))
 		{
 			OutValue.SourceType = EUsdBasicDataTypes::Double3;
-			ConvertVecValue<double, GfVec3d>( UsdValue, OutValue );
+			ConvertVecValue<double, GfVec3d>(UsdValue, OutValue);
 		}
-		else if ( CHECK_TYPE(GfVec3f) )
+		else if (CHECK_TYPE(GfVec3f))
 		{
 			OutValue.SourceType = EUsdBasicDataTypes::Float3;
-			ConvertVecValue<float, GfVec3f>( UsdValue, OutValue );
+			ConvertVecValue<float, GfVec3f>(UsdValue, OutValue);
 		}
-		else if ( CHECK_TYPE(GfVec3h) )
+		else if (CHECK_TYPE(GfVec3h))
 		{
 			OutValue.SourceType = EUsdBasicDataTypes::Half3;
-			ConvertVecValue<float, GfVec3h>( UsdValue, OutValue );
+			ConvertVecValue<float, GfVec3h>(UsdValue, OutValue);
 		}
-		else if ( CHECK_TYPE(GfVec3i) )
+		else if (CHECK_TYPE(GfVec3i))
 		{
 			OutValue.SourceType = EUsdBasicDataTypes::Int3;
-			ConvertVecValue<int32, GfVec3i>( UsdValue, OutValue );
+			ConvertVecValue<int32, GfVec3i>(UsdValue, OutValue);
 		}
-		else if ( CHECK_TYPE(GfVec4d) )
+		else if (CHECK_TYPE(GfVec4d))
 		{
 			OutValue.SourceType = EUsdBasicDataTypes::Double4;
-			ConvertVecValue<double, GfVec4d>( UsdValue, OutValue );
+			ConvertVecValue<double, GfVec4d>(UsdValue, OutValue);
 		}
-		else if ( CHECK_TYPE(GfVec4f) )
+		else if (CHECK_TYPE(GfVec4f))
 		{
 			OutValue.SourceType = EUsdBasicDataTypes::Float4;
-			ConvertVecValue<float, GfVec4f>( UsdValue, OutValue );
+			ConvertVecValue<float, GfVec4f>(UsdValue, OutValue);
 		}
-		else if ( CHECK_TYPE(GfVec4h) )
+		else if (CHECK_TYPE(GfVec4h))
 		{
 			OutValue.SourceType = EUsdBasicDataTypes::Half4;
-			ConvertVecValue<float, GfVec4h>( UsdValue, OutValue );
+			ConvertVecValue<float, GfVec4h>(UsdValue, OutValue);
 		}
-		else if ( CHECK_TYPE(GfVec4i) )
+		else if (CHECK_TYPE(GfVec4i))
 		{
 			OutValue.SourceType = EUsdBasicDataTypes::Int4;
-			ConvertVecValue<int32, GfVec4i>( UsdValue, OutValue );
+			ConvertVecValue<int32, GfVec4i>(UsdValue, OutValue);
 		}
 		// These types should only appear within metadata, and don't support arrays.
 		// There are more of them (e.g. pxr/usd/usd/crateDataTypes.h), but these are the most common.
 		// Also check pxr/usd/sdf/types.cpp for where these are defined. These are simple enums
-		else if ( UnderlyingType.IsA<SdfPermission>() )
+		else if (UnderlyingType.IsA<SdfPermission>())
 		{
 			OutValue.SourceType = EUsdBasicDataTypes::Int;
-			ConvertSimpleValue<int32, SdfPermission>( UsdValue, OutValue );
+			ConvertSimpleValue<int32, SdfPermission>(UsdValue, OutValue);
 		}
-		else if ( UnderlyingType.IsA<SdfSpecifier>() )
+		else if (UnderlyingType.IsA<SdfSpecifier>())
 		{
 			OutValue.SourceType = EUsdBasicDataTypes::Int;
-			ConvertSimpleValue<int32, SdfSpecifier>( UsdValue, OutValue );
+			ConvertSimpleValue<int32, SdfSpecifier>(UsdValue, OutValue);
 		}
-		else if ( UnderlyingType.IsA<SdfVariability>() )
+		else if (UnderlyingType.IsA<SdfVariability>())
 		{
 			OutValue.SourceType = EUsdBasicDataTypes::Int;
-			ConvertSimpleValue<int32, SdfVariability>( UsdValue, OutValue );
+			ConvertSimpleValue<int32, SdfVariability>(UsdValue, OutValue);
 		}
-		else if ( UnderlyingType.IsA<SdfSpecType>() )
+		else if (UnderlyingType.IsA<SdfSpecType>())
 		{
 			OutValue.SourceType = EUsdBasicDataTypes::Int;
-			ConvertSimpleValue<int32, SdfSpecType>( UsdValue, OutValue );
+			ConvertSimpleValue<int32, SdfSpecType>(UsdValue, OutValue);
 		}
 		else
 		{
@@ -413,9 +421,9 @@ namespace UsdToUnreal
 		return true;
 #else
 		return false;
-#endif // USE_USD_SDK
+#endif	  // USE_USD_SDK
 	}
-}
+}	 // namespace UsdToUnreal
 
 #if USE_USD_SDK
 namespace UnrealToUsdImpl
@@ -426,42 +434,44 @@ namespace UnrealToUsdImpl
 	 * []( const UsdUtils::FConvertedVtValueEntry& Entry ) -> USDElementType {}
 	 */
 	template<typename UEElementType, typename USDElementType, typename Func>
-	void ConvertInner( const UsdUtils::FConvertedVtValue& InValue, pxr::VtValue& OutValue, Func Function )
+	void ConvertInner(const UsdUtils::FConvertedVtValue& InValue, pxr::VtValue& OutValue, Func Function)
 	{
-		if ( InValue.bIsArrayValued )
+		if (InValue.bIsArrayValued)
 		{
 			pxr::VtArray<USDElementType> Array;
-			Array.reserve( InValue.Entries.Num() );
+			Array.reserve(InValue.Entries.Num());
 
-			for ( const UsdUtils::FConvertedVtValueEntry& Entry : InValue.Entries )
+			for (const UsdUtils::FConvertedVtValueEntry& Entry : InValue.Entries)
 			{
-				if ( Entry.Num() > 0 )
+				if (Entry.Num() > 0)
 				{
-					Array.push_back( Function( Entry ) );
+					Array.push_back(Function(Entry));
 				}
 			}
 
 			OutValue = Array;
 		}
-		else if ( InValue.Entries.Num() > 0 )
+		else if (InValue.Entries.Num() > 0)
 		{
 			const UsdUtils::FConvertedVtValueEntry& Entry = InValue.Entries[0];
-			if ( Entry.Num() > 0 )
+			if (Entry.Num() > 0)
 			{
-				OutValue = Function( Entry );
+				OutValue = Function(Entry);
 			}
 		}
 	}
 
 	template<typename UEElementType, typename USDElementType>
-	void ConvertSimpleValue( const UsdUtils::FConvertedVtValue& InValue, pxr::VtValue& OutValue )
+	void ConvertSimpleValue(const UsdUtils::FConvertedVtValue& InValue, pxr::VtValue& OutValue)
 	{
-		ConvertInner< UEElementType, USDElementType >( InValue, OutValue,
-			[]( const UsdUtils::FConvertedVtValueEntry& Entry ) -> USDElementType
+		ConvertInner<UEElementType, USDElementType>(
+			InValue,
+			OutValue,
+			[](const UsdUtils::FConvertedVtValueEntry& Entry) -> USDElementType
 			{
-				if ( const UEElementType* Value = Entry[ 0 ].TryGet<UEElementType>() )
+				if (const UEElementType* Value = Entry[0].TryGet<UEElementType>())
 				{
-					return static_cast< USDElementType >( *Value );
+					return static_cast<USDElementType>(*Value);
 				}
 
 				return USDElementType();
@@ -470,14 +480,16 @@ namespace UnrealToUsdImpl
 	}
 
 	template<>
-	void ConvertSimpleValue<FString, std::string>( const UsdUtils::FConvertedVtValue& InValue, pxr::VtValue& OutValue )
+	void ConvertSimpleValue<FString, std::string>(const UsdUtils::FConvertedVtValue& InValue, pxr::VtValue& OutValue)
 	{
-		ConvertInner< FString, std::string >( InValue, OutValue,
-			[]( const UsdUtils::FConvertedVtValueEntry& Entry ) -> std::string
+		ConvertInner<FString, std::string>(
+			InValue,
+			OutValue,
+			[](const UsdUtils::FConvertedVtValueEntry& Entry) -> std::string
 			{
-				if ( const FString* Value = Entry[ 0 ].TryGet<FString>() )
+				if (const FString* Value = Entry[0].TryGet<FString>())
 				{
-					return UnrealToUsd::ConvertString( **Value ).Get();
+					return UnrealToUsd::ConvertString(**Value).Get();
 				}
 
 				return std::string();
@@ -486,14 +498,16 @@ namespace UnrealToUsdImpl
 	}
 
 	template<>
-	void ConvertSimpleValue<FString, pxr::TfToken>( const UsdUtils::FConvertedVtValue& InValue, pxr::VtValue& OutValue )
+	void ConvertSimpleValue<FString, pxr::TfToken>(const UsdUtils::FConvertedVtValue& InValue, pxr::VtValue& OutValue)
 	{
-		ConvertInner< FString, pxr::TfToken >( InValue, OutValue,
-			[]( const UsdUtils::FConvertedVtValueEntry& Entry ) -> pxr::TfToken
+		ConvertInner<FString, pxr::TfToken>(
+			InValue,
+			OutValue,
+			[](const UsdUtils::FConvertedVtValueEntry& Entry) -> pxr::TfToken
 			{
-				if ( const FString* Value = Entry[ 0 ].TryGet<FString>() )
+				if (const FString* Value = Entry[0].TryGet<FString>())
 				{
-					return UnrealToUsd::ConvertToken( **Value ).Get();
+					return UnrealToUsd::ConvertToken(**Value).Get();
 				}
 
 				return pxr::TfToken();
@@ -502,14 +516,16 @@ namespace UnrealToUsdImpl
 	}
 
 	template<>
-	void ConvertSimpleValue<FString, pxr::SdfAssetPath>( const UsdUtils::FConvertedVtValue& InValue, pxr::VtValue& OutValue )
+	void ConvertSimpleValue<FString, pxr::SdfAssetPath>(const UsdUtils::FConvertedVtValue& InValue, pxr::VtValue& OutValue)
 	{
-		ConvertInner< FString, pxr::SdfAssetPath >( InValue, OutValue,
-			[]( const UsdUtils::FConvertedVtValueEntry& Entry ) -> pxr::SdfAssetPath
+		ConvertInner<FString, pxr::SdfAssetPath>(
+			InValue,
+			OutValue,
+			[](const UsdUtils::FConvertedVtValueEntry& Entry) -> pxr::SdfAssetPath
 			{
-				if ( const FString* Value = Entry[ 0 ].TryGet<FString>() )
+				if (const FString* Value = Entry[0].TryGet<FString>())
 				{
-					return pxr::SdfAssetPath( UnrealToUsd::ConvertString( **Value ).Get() );
+					return pxr::SdfAssetPath(UnrealToUsd::ConvertString(**Value).Get());
 				}
 
 				return pxr::SdfAssetPath();
@@ -519,50 +535,53 @@ namespace UnrealToUsdImpl
 
 	/** We need the USDElementType parameter to do the final float to pxr::GfHalf conversions */
 	template<typename UEElementType, typename USDArrayType, typename USDElementType = UEElementType>
-	void ConvertCompoundValue( const UsdUtils::FConvertedVtValue& InValue, pxr::VtValue& OutValue )
+	void ConvertCompoundValue(const UsdUtils::FConvertedVtValue& InValue, pxr::VtValue& OutValue)
 	{
-		ConvertInner< UEElementType, USDArrayType >( InValue, OutValue,
-			[]( const UsdUtils::FConvertedVtValueEntry& Entry ) -> USDArrayType
+		ConvertInner<UEElementType, USDArrayType>(
+			InValue,
+			OutValue,
+			[](const UsdUtils::FConvertedVtValueEntry& Entry) -> USDArrayType
 			{
-				USDArrayType USDVal( UEElementType( 0 ) );
+				USDArrayType USDVal(UEElementType(0));
 				USDElementType* DataPtr = USDVal.data();
 
-				for ( int32 Index = 0; Index < Entry.Num(); ++Index )
+				for (int32 Index = 0; Index < Entry.Num(); ++Index)
 				{
-					if ( const UEElementType* IndexValue = Entry[ Index ].TryGet<UEElementType>() )
+					if (const UEElementType* IndexValue = Entry[Index].TryGet<UEElementType>())
 					{
-						DataPtr[ Index ] = static_cast< USDElementType >( *IndexValue );
+						DataPtr[Index] = static_cast<USDElementType>(*IndexValue);
 					}
 				}
 				return USDVal;
 			}
 		);
-
 	}
 
 	/** USD quaternions don't have the access operator defined, and the elements need to be reordered */
 	template<typename UEElementType, typename USDQuatType>
-	void ConvertQuatValue( const UsdUtils::FConvertedVtValue& InValue, pxr::VtValue& OutValue )
+	void ConvertQuatValue(const UsdUtils::FConvertedVtValue& InValue, pxr::VtValue& OutValue)
 	{
-		ConvertInner< UEElementType, USDQuatType >( InValue, OutValue,
-			[]( const UsdUtils::FConvertedVtValueEntry& Entry ) -> USDQuatType
+		ConvertInner<UEElementType, USDQuatType>(
+			InValue,
+			OutValue,
+			[](const UsdUtils::FConvertedVtValueEntry& Entry) -> USDQuatType
 			{
-				if ( Entry.Num() == 4 )
+				if (Entry.Num() == 4)
 				{
-					UEElementType QuatValues[ 4 ] = { 0, 0, 0, 1 };
-					for ( int32 Index = 0; Index < 4; ++Index )
+					UEElementType QuatValues[4] = {0, 0, 0, 1};
+					for (int32 Index = 0; Index < 4; ++Index)
 					{
-						if ( const UEElementType* QuatValue = Entry[ Index ].TryGet<UEElementType>() )
+						if (const UEElementType* QuatValue = Entry[Index].TryGet<UEElementType>())
 						{
-							QuatValues[ Index ] = *QuatValue;
+							QuatValues[Index] = *QuatValue;
 						}
 					}
 
 					return USDQuatType(
-						QuatValues[ 3 ],  // Real part comes first for USD
-						QuatValues[ 0 ],
-						QuatValues[ 1 ],
-						QuatValues[ 2 ]
+						QuatValues[3],	  // Real part comes first for USD
+						QuatValues[0],
+						QuatValues[1],
+						QuatValues[2]
 					);
 				}
 
@@ -570,12 +589,12 @@ namespace UnrealToUsdImpl
 			}
 		);
 	}
-}
-#endif // USE_USD_SDK
+}	 // namespace UnrealToUsdImpl
+#endif	  // USE_USD_SDK
 
 namespace UnrealToUsd
 {
-	bool ConvertValue( const UsdUtils::FConvertedVtValue& InValue, UE::FVtValue& OutValue )
+	bool ConvertValue(const UsdUtils::FConvertedVtValue& InValue, UE::FVtValue& OutValue)
 	{
 #if USE_USD_SDK
 		using namespace pxr;
@@ -589,115 +608,115 @@ namespace UnrealToUsd
 
 		VtValue& UsdValue = OutValue.GetUsdValue();
 
-		switch ( InValue.SourceType )
-		{
-		case EUsdBasicDataTypes::Bool:
-			ConvertSimpleValue<bool, bool>( InValue, UsdValue );
-			break;
-		case EUsdBasicDataTypes::Uchar:
-			ConvertSimpleValue<uint8, uint8_t>( InValue, UsdValue );
-			break;
-		case EUsdBasicDataTypes::Int:
-			ConvertSimpleValue<int32, int32_t>( InValue, UsdValue );
-			break;
-		case EUsdBasicDataTypes::Uint:
-			ConvertSimpleValue<uint32, uint32_t>( InValue, UsdValue );
-			break;
-		case EUsdBasicDataTypes::Int64:
-			ConvertSimpleValue<int64, int64_t>( InValue, UsdValue );
-			break;
-		case EUsdBasicDataTypes::Uint64:
-			ConvertSimpleValue<uint64, uint64_t>( InValue, UsdValue );
-			break;
-		case EUsdBasicDataTypes::Half:
-			ConvertSimpleValue<float, GfHalf>( InValue, UsdValue );
-			break;
-		case EUsdBasicDataTypes::Float:
-			ConvertSimpleValue<float, float>( InValue, UsdValue );
-			break;
-		case EUsdBasicDataTypes::Double:
-			ConvertSimpleValue<double, double>( InValue, UsdValue );
-			break;
-		case EUsdBasicDataTypes::Timecode:
-			ConvertSimpleValue<double, SdfTimeCode>( InValue, UsdValue );
-			break;
-		case EUsdBasicDataTypes::String:
-			ConvertSimpleValue<FString, std::string>( InValue, UsdValue );
-			break;
-		case EUsdBasicDataTypes::Token:
-			ConvertSimpleValue<FString, TfToken>( InValue, UsdValue );
-			break;
-		case EUsdBasicDataTypes::Asset:
-			ConvertSimpleValue<FString, SdfAssetPath>( InValue, UsdValue );
-			break;
-		case EUsdBasicDataTypes::Matrix2d:
-			ConvertCompoundValue<double, GfMatrix2d>( InValue, UsdValue );
-			break;
-		case EUsdBasicDataTypes::Matrix3d:
-			ConvertCompoundValue<double, GfMatrix3d>( InValue, UsdValue );
-			break;
-		case EUsdBasicDataTypes::Matrix4d:
-			ConvertCompoundValue<double, GfMatrix4d>( InValue, UsdValue );
-			break;
-		case EUsdBasicDataTypes::Quatd:
-			ConvertQuatValue<double, GfQuatd>( InValue, UsdValue );
-			break;
-		case EUsdBasicDataTypes::Quatf:
-			ConvertQuatValue<float, GfQuatf>( InValue, UsdValue );
-			break;
-		case EUsdBasicDataTypes::Quath:
-			ConvertQuatValue<float, GfQuath>( InValue, UsdValue );
-			break;
-		case EUsdBasicDataTypes::Double2:
-			ConvertCompoundValue<double, GfVec2d>( InValue, UsdValue );
-			break;
-		case EUsdBasicDataTypes::Float2:
-			ConvertCompoundValue<float, GfVec2f>( InValue, UsdValue );
-			break;
-		case EUsdBasicDataTypes::Half2:
-			ConvertCompoundValue<float, GfVec2h, GfHalf>( InValue, UsdValue );
-			break;
-		case EUsdBasicDataTypes::Int2:
-			ConvertCompoundValue<int32, GfVec2i>( InValue, UsdValue );
-			break;
-		case EUsdBasicDataTypes::Double3:
-			ConvertCompoundValue<double, GfVec3d>( InValue, UsdValue );
-			break;
-		case EUsdBasicDataTypes::Float3:
-			ConvertCompoundValue<float, GfVec3f>( InValue, UsdValue );
-			break;
-		case EUsdBasicDataTypes::Half3:
-			ConvertCompoundValue<float, GfVec3h, GfHalf>( InValue, UsdValue );
-			break;
-		case EUsdBasicDataTypes::Int3:
-			ConvertCompoundValue<int32, GfVec3i>( InValue, UsdValue );
-			break;
-		case EUsdBasicDataTypes::Double4:
-			ConvertCompoundValue<double, GfVec4d>( InValue, UsdValue );
-			break;
-		case EUsdBasicDataTypes::Float4:
-			ConvertCompoundValue<float, GfVec4f>( InValue, UsdValue );
-			break;
-		case EUsdBasicDataTypes::Half4:
-			ConvertCompoundValue<float, GfVec4h, GfHalf>( InValue, UsdValue );
-			break;
-		case EUsdBasicDataTypes::Int4:
-			ConvertCompoundValue<int32, GfVec4i>( InValue, UsdValue );
-			break;
-		case EUsdBasicDataTypes::None:
-		default:
-			break;
+		switch (InValue.SourceType)
+		{
+			case EUsdBasicDataTypes::Bool:
+				ConvertSimpleValue<bool, bool>(InValue, UsdValue);
+				break;
+			case EUsdBasicDataTypes::Uchar:
+				ConvertSimpleValue<uint8, uint8_t>(InValue, UsdValue);
+				break;
+			case EUsdBasicDataTypes::Int:
+				ConvertSimpleValue<int32, int32_t>(InValue, UsdValue);
+				break;
+			case EUsdBasicDataTypes::Uint:
+				ConvertSimpleValue<uint32, uint32_t>(InValue, UsdValue);
+				break;
+			case EUsdBasicDataTypes::Int64:
+				ConvertSimpleValue<int64, int64_t>(InValue, UsdValue);
+				break;
+			case EUsdBasicDataTypes::Uint64:
+				ConvertSimpleValue<uint64, uint64_t>(InValue, UsdValue);
+				break;
+			case EUsdBasicDataTypes::Half:
+				ConvertSimpleValue<float, GfHalf>(InValue, UsdValue);
+				break;
+			case EUsdBasicDataTypes::Float:
+				ConvertSimpleValue<float, float>(InValue, UsdValue);
+				break;
+			case EUsdBasicDataTypes::Double:
+				ConvertSimpleValue<double, double>(InValue, UsdValue);
+				break;
+			case EUsdBasicDataTypes::Timecode:
+				ConvertSimpleValue<double, SdfTimeCode>(InValue, UsdValue);
+				break;
+			case EUsdBasicDataTypes::String:
+				ConvertSimpleValue<FString, std::string>(InValue, UsdValue);
+				break;
+			case EUsdBasicDataTypes::Token:
+				ConvertSimpleValue<FString, TfToken>(InValue, UsdValue);
+				break;
+			case EUsdBasicDataTypes::Asset:
+				ConvertSimpleValue<FString, SdfAssetPath>(InValue, UsdValue);
+				break;
+			case EUsdBasicDataTypes::Matrix2d:
+				ConvertCompoundValue<double, GfMatrix2d>(InValue, UsdValue);
+				break;
+			case EUsdBasicDataTypes::Matrix3d:
+				ConvertCompoundValue<double, GfMatrix3d>(InValue, UsdValue);
+				break;
+			case EUsdBasicDataTypes::Matrix4d:
+				ConvertCompoundValue<double, GfMatrix4d>(InValue, UsdValue);
+				break;
+			case EUsdBasicDataTypes::Quatd:
+				ConvertQuatValue<double, GfQuatd>(InValue, UsdValue);
+				break;
+			case EUsdBasicDataTypes::Quatf:
+				ConvertQuatValue<float, GfQuatf>(InValue, UsdValue);
+				break;
+			case EUsdBasicDataTypes::Quath:
+				ConvertQuatValue<float, GfQuath>(InValue, UsdValue);
+				break;
+			case EUsdBasicDataTypes::Double2:
+				ConvertCompoundValue<double, GfVec2d>(InValue, UsdValue);
+				break;
+			case EUsdBasicDataTypes::Float2:
+				ConvertCompoundValue<float, GfVec2f>(InValue, UsdValue);
+				break;
+			case EUsdBasicDataTypes::Half2:
+				ConvertCompoundValue<float, GfVec2h, GfHalf>(InValue, UsdValue);
+				break;
+			case EUsdBasicDataTypes::Int2:
+				ConvertCompoundValue<int32, GfVec2i>(InValue, UsdValue);
+				break;
+			case EUsdBasicDataTypes::Double3:
+				ConvertCompoundValue<double, GfVec3d>(InValue, UsdValue);
+				break;
+			case EUsdBasicDataTypes::Float3:
+				ConvertCompoundValue<float, GfVec3f>(InValue, UsdValue);
+				break;
+			case EUsdBasicDataTypes::Half3:
+				ConvertCompoundValue<float, GfVec3h, GfHalf>(InValue, UsdValue);
+				break;
+			case EUsdBasicDataTypes::Int3:
+				ConvertCompoundValue<int32, GfVec3i>(InValue, UsdValue);
+				break;
+			case EUsdBasicDataTypes::Double4:
+				ConvertCompoundValue<double, GfVec4d>(InValue, UsdValue);
+				break;
+			case EUsdBasicDataTypes::Float4:
+				ConvertCompoundValue<float, GfVec4f>(InValue, UsdValue);
+				break;
+			case EUsdBasicDataTypes::Half4:
+				ConvertCompoundValue<float, GfVec4h, GfHalf>(InValue, UsdValue);
+				break;
+			case EUsdBasicDataTypes::Int4:
+				ConvertCompoundValue<int32, GfVec4i>(InValue, UsdValue);
+				break;
+			case EUsdBasicDataTypes::None:
+			default:
+				break;
 		}
 
 		// We consider a success returning an empty value if our input value was also empty
 		return InValue.Entries.Num() == 0 || !OutValue.IsEmpty();
 #else
 		return false;
-#endif // USE_USD_SDK
+#endif	  // USE_USD_SDK
 	}
-}
+}	 // namespace UnrealToUsd
 
-FArchive& operator<<( FArchive& Ar, UsdUtils::FConvertedVtValue& Struct )
+FArchive& operator<<(FArchive& Ar, UsdUtils::FConvertedVtValue& Struct)
 {
 	Ar << Struct.Entries;
 	Ar << Struct.SourceType;
@@ -706,145 +725,145 @@ FArchive& operator<<( FArchive& Ar, UsdUtils::FConvertedVtValue& Struct )
 	return Ar;
 }
 
-FArchive& operator<<( FArchive& Ar, UsdUtils::FConvertedVtValueComponent& Component )
+FArchive& operator<<(FArchive& Ar, UsdUtils::FConvertedVtValueComponent& Component)
 {
-	if ( Ar.IsSaving() )
+	if (Ar.IsSaving())
 	{
 		uint64 TypeIndex = Component.GetIndex();
 		Ar << TypeIndex;
 
-		switch ( TypeIndex )
+		switch (TypeIndex)
 		{
-			case UsdUtils::FConvertedVtValueComponent::IndexOfType<bool>() :
-				if ( bool* Val = Component.TryGet<bool>() )
+			case UsdUtils::FConvertedVtValueComponent::IndexOfType<bool>():
+				if (bool* Val = Component.TryGet<bool>())
 				{
 					Ar << *Val;
 				}
-			break;
-			case UsdUtils::FConvertedVtValueComponent::IndexOfType<uint8>() :
-				if ( uint8* Val = Component.TryGet<uint8>() )
+				break;
+			case UsdUtils::FConvertedVtValueComponent::IndexOfType<uint8>():
+				if (uint8* Val = Component.TryGet<uint8>())
 				{
 					Ar << *Val;
 				}
-			break;
-			case UsdUtils::FConvertedVtValueComponent::IndexOfType<int32>() :
-				if ( int32* Val = Component.TryGet<int32>() )
+				break;
+			case UsdUtils::FConvertedVtValueComponent::IndexOfType<int32>():
+				if (int32* Val = Component.TryGet<int32>())
 				{
 					Ar << *Val;
 				}
-			break;
-			case UsdUtils::FConvertedVtValueComponent::IndexOfType<uint32>() :
-				if ( uint32* Val = Component.TryGet<uint32>() )
+				break;
+			case UsdUtils::FConvertedVtValueComponent::IndexOfType<uint32>():
+				if (uint32* Val = Component.TryGet<uint32>())
 				{
 					Ar << *Val;
 				}
-			break;
-			case UsdUtils::FConvertedVtValueComponent::IndexOfType<int64>() :
-				if ( int64* Val = Component.TryGet<int64>() )
+				break;
+			case UsdUtils::FConvertedVtValueComponent::IndexOfType<int64>():
+				if (int64* Val = Component.TryGet<int64>())
 				{
 					Ar << *Val;
 				}
-			break;
-			case UsdUtils::FConvertedVtValueComponent::IndexOfType<uint64>() :
-				if ( uint64* Val = Component.TryGet<uint64>() )
+				break;
+			case UsdUtils::FConvertedVtValueComponent::IndexOfType<uint64>():
+				if (uint64* Val = Component.TryGet<uint64>())
 				{
 					Ar << *Val;
 				}
-			break;
-			case UsdUtils::FConvertedVtValueComponent::IndexOfType<float>() :
-				if ( float* Val = Component.TryGet<float>() )
+				break;
+			case UsdUtils::FConvertedVtValueComponent::IndexOfType<float>():
+				if (float* Val = Component.TryGet<float>())
 				{
 					Ar << *Val;
 				}
-			break;
-			case UsdUtils::FConvertedVtValueComponent::IndexOfType<double>() :
-				if ( double* Val = Component.TryGet<double>() )
+				break;
+			case UsdUtils::FConvertedVtValueComponent::IndexOfType<double>():
+				if (double* Val = Component.TryGet<double>())
 				{
 					Ar << *Val;
 				}
-			break;
-			case UsdUtils::FConvertedVtValueComponent::IndexOfType<FString>() :
-				if ( FString* Val = Component.TryGet<FString>() )
+				break;
+			case UsdUtils::FConvertedVtValueComponent::IndexOfType<FString>():
+				if (FString* Val = Component.TryGet<FString>())
 				{
 					Ar << *Val;
 				}
-			break;
+				break;
 			default:
 				break;
 		}
 	}
-	else // IsLoading
+	else	// IsLoading
 	{
 		uint64 TypeIndex;
 		Ar << TypeIndex;
 
-		switch ( TypeIndex )
+		switch (TypeIndex)
 		{
-			case UsdUtils::FConvertedVtValueComponent::IndexOfType<bool>() :
+			case UsdUtils::FConvertedVtValueComponent::IndexOfType<bool>():
 			{
 				bool Val;
 				Ar << Val;
-				Component.Set<bool>( Val );
+				Component.Set<bool>(Val);
 				break;
 			}
-			case UsdUtils::FConvertedVtValueComponent::IndexOfType<uint8>() :
+			case UsdUtils::FConvertedVtValueComponent::IndexOfType<uint8>():
 			{
 				uint8 Val;
 				Ar << Val;
-				Component.Set<uint8>( Val );
+				Component.Set<uint8>(Val);
 				break;
 			}
-			case UsdUtils::FConvertedVtValueComponent::IndexOfType<int32>() :
+			case UsdUtils::FConvertedVtValueComponent::IndexOfType<int32>():
 			{
 				int32 Val;
 				Ar << Val;
-				Component.Set<int32>( Val );
+				Component.Set<int32>(Val);
 				break;
 			}
-			case UsdUtils::FConvertedVtValueComponent::IndexOfType<uint32>() :
+			case UsdUtils::FConvertedVtValueComponent::IndexOfType<uint32>():
 			{
 				uint32 Val;
 				Ar << Val;
-				Component.Set<uint32>( Val );
+				Component.Set<uint32>(Val);
 				break;
 			}
-			case UsdUtils::FConvertedVtValueComponent::IndexOfType<int64>() :
+			case UsdUtils::FConvertedVtValueComponent::IndexOfType<int64>():
 			{
 				int64 Val;
 				Ar << Val;
-				Component.Set<int64>( Val );
+				Component.Set<int64>(Val);
 				break;
 			}
-			case UsdUtils::FConvertedVtValueComponent::IndexOfType<uint64>() :
+			case UsdUtils::FConvertedVtValueComponent::IndexOfType<uint64>():
 			{
 				uint64 Val;
 				Ar << Val;
-				Component.Set<uint64>( Val );
+				Component.Set<uint64>(Val);
 				break;
 			}
-			case UsdUtils::FConvertedVtValueComponent::IndexOfType<float>() :
+			case UsdUtils::FConvertedVtValueComponent::IndexOfType<float>():
 			{
 				float Val;
 				Ar << Val;
-				Component.Set<float>( Val );
+				Component.Set<float>(Val);
 				break;
 			}
-			case UsdUtils::FConvertedVtValueComponent::IndexOfType<double>() :
+			case UsdUtils::FConvertedVtValueComponent::IndexOfType<double>():
 			{
 				double Val;
 				Ar << Val;
-				Component.Set<double>( Val );
+				Component.Set<double>(Val);
 				break;
 			}
-			case UsdUtils::FConvertedVtValueComponent::IndexOfType<FString>() :
+			case UsdUtils::FConvertedVtValueComponent::IndexOfType<FString>():
 			{
 				FString Val;
 				Ar << Val;
-				Component.Set<FString>( Val );
+				Component.Set<FString>(Val);
 				break;
 			}
 			default:
-				Component.Set<bool>( false );
+				Component.Set<bool>(false);
 				break;
 		}
 	}
@@ -899,7 +918,6 @@ namespace UsdUtils
 			{
 				Result += FString::Printf(TEXT("'%s', "), *UsdToUnreal::ConvertString(String));
 			}
-
 		}
 		Result.RemoveFromEnd(TEXT(", "));
 		Result += TEXT("]");
@@ -1027,13 +1045,13 @@ namespace UsdUtils
 		using TypeToFuncMap = std::unordered_map<pxr::TfType, TFunction<FString(const pxr::VtValue&)>, pxr::TfHash>;
 
 		static const TypeToFuncMap Stringifiers = {
-			{pxr::TfType::Find<std::string>(), StringifyUsdString},
-			{pxr::TfType::Find<pxr::VtArray<std::string>>(), StringifyUsdStringArray},
-			{pxr::TfType::Find<pxr::TfToken>(), StringifyUsdToken},
-			{pxr::TfType::Find<pxr::VtArray<pxr::TfToken>>(), StringifyUsdTokenArray},
-			{pxr::TfType::Find<pxr::SdfAssetPath>(), StringifyUsdAsset},
-			{pxr::TfType::Find<pxr::VtArray<pxr::SdfAssetPath>>(), StringifyUsdAssetArray},
-			{pxr::TfType::Find<pxr::SdfListOp<pxr::TfToken>>(), StringifyTokenList},
+			{					pxr::TfType::Find<std::string>(),	   StringifyUsdString},
+			{		 pxr::TfType::Find<pxr::VtArray<std::string>>(), StringifyUsdStringArray},
+			{				   pxr::TfType::Find<pxr::TfToken>(),		StringifyUsdToken},
+			{	 pxr::TfType::Find<pxr::VtArray<pxr::TfToken>>(),  StringifyUsdTokenArray},
+			{			  pxr::TfType::Find<pxr::SdfAssetPath>(),		 StringifyUsdAsset},
+			{pxr::TfType::Find<pxr::VtArray<pxr::SdfAssetPath>>(),  StringifyUsdAssetArray},
+			{	 pxr::TfType::Find<pxr::SdfListOp<pxr::TfToken>>(),		StringifyTokenList},
 		};
 
 		TypeToFuncMap::const_iterator FoundStringifier = Stringifiers.find(Value.GetType());
@@ -2162,7 +2180,8 @@ namespace UE::USDValueConversion::Private
 				FlattenedValues[FlattenedIndex + 0],
 				FlattenedValues[FlattenedIndex + 1],
 				FlattenedValues[FlattenedIndex + 2],
-				FlattenedValues[FlattenedIndex + 3]};
+				FlattenedValues[FlattenedIndex + 3]
+			};
 		}
 
 		return true;
@@ -2241,7 +2260,7 @@ namespace UsdUtils
 		using namespace UE::USDValueConversion::Private;
 
 		static const TMap<FString, TFunction<bool(const FString&, pxr::VtValue&)>> Unstringifiers = {
-			// Basic datatypes
+	// Basic datatypes
 			{TEXT("bool"), WrapInVtValue<bool, UnstringifyInner<bool>>},
 			{TEXT("uchar"), WrapInVtValue<uint8_t, UnstringifyInner<uint8_t>>},
 			{TEXT("int"), WrapInVtValue<int32_t, UnstringifyInner<int32_t>>},
@@ -2274,7 +2293,7 @@ namespace UsdUtils
 			{TEXT("half4"), WrapInVtValue<pxr::GfVec4h, UnstringifyVec<pxr::GfVec4h>>},
 			{TEXT("int4"), WrapInVtValue<pxr::GfVec4i, UnstringifyVec<pxr::GfVec4i>>},
 
-			// Array versions of the ones above
+ // Array versions of the ones above
 			{TEXT("bool[]"), WrapInVtValue<pxr::VtArray<bool>, UnstringifyBasicArray<bool>>},
 			{TEXT("uchar[]"), WrapInVtValue<pxr::VtArray<uint8_t>, UnstringifyBasicArray<uint8_t>>},
 			{TEXT("int[]"), WrapInVtValue<pxr::VtArray<int32_t>, UnstringifyBasicArray<int32_t>>},
@@ -2307,8 +2326,9 @@ namespace UsdUtils
 			{TEXT("half4[]"), WrapInVtValue<pxr::VtArray<pxr::GfVec4h>, UnstringifyNumberArray<pxr::GfVec4h, 4>>},
 			{TEXT("int4[]"), WrapInVtValue<pxr::VtArray<pxr::GfVec4i>, UnstringifyNumberArray<pxr::GfVec4i, 4>>},
 
-			// // Exotic types found in some scenarios
-			{TEXT("SdfListOp<TfToken>"), WrapInVtValue<pxr::SdfListOp<pxr::TfToken>, UnstringifyListOpTokens>}, // This is used for apiSchemas metadata
+ // // Exotic types found in some scenarios
+			{TEXT("SdfListOp<TfToken>"), WrapInVtValue<pxr::SdfListOp<pxr::TfToken>, UnstringifyListOpTokens>}, // This is used for apiSchemas
+  // metadata
 		};
 
 		if (const TFunction<bool(const FString&, pxr::VtValue&)>* FoundUnstringifier = Unstringifiers.Find(TypeName))
diff --git a/Engine/Plugins/Importers/USDImporter/Source/USDUtilities/Public/USDAttributeUtils.h b/Engine/Plugins/Importers/USDImporter/Source/USDUtilities/Public/USDAttributeUtils.h
--- a/Engine/Plugins/Importers/USDImporter/Source/USDUtilities/Public/USDAttributeUtils.h
+++ b/Engine/Plugins/Importers/USDImporter/Source/USDUtilities/Public/USDAttributeUtils.h
@@ -12,7 +12,7 @@ PXR_NAMESPACE_OPEN_SCOPE
 	class UsdPrim;
 	class UsdProperty;
 PXR_NAMESPACE_CLOSE_SCOPE
-#endif // USE_USD_SDK
+#endif	  // USE_USD_SDK
 
 namespace UE
 {
@@ -22,14 +22,16 @@ namespace UE
 
 namespace UsdUtils
 {
-	/** Adds a "Muted" CustomData entry to the attribute at the stage's UE state sublayer, which will prevent it from being animated when loaded into UE */
-	USDUTILITIES_API bool MuteAttribute( UE::FUsdAttribute& Attribute, const UE::FUsdStage& Stage );
+	/** Adds a "Muted" CustomData entry to the attribute at the stage's UE state sublayer, which will prevent it from being animated when loaded into
+	 * UE */
+	USDUTILITIES_API bool MuteAttribute(UE::FUsdAttribute& Attribute, const UE::FUsdStage& Stage);
 
 	/** Removes the "Muted" CustomData entry from the attribute at the stage's UE state sublayer, letting it be animated when loaded into UE */
-	USDUTILITIES_API bool UnmuteAttribute( UE::FUsdAttribute& Attribute, const UE::FUsdStage& Stage );
+	USDUTILITIES_API bool UnmuteAttribute(UE::FUsdAttribute& Attribute, const UE::FUsdStage& Stage);
 
-	/** Returns whether this attribute has the "Muted" CustomData on the stage's UE state sublayer, meaning it shouldn't be animated when loaded into UE */
-	USDUTILITIES_API bool IsAttributeMuted( const UE::FUsdAttribute& Attribute, const UE::FUsdStage& Stage );
+	/** Returns whether this attribute has the "Muted" CustomData on the stage's UE state sublayer, meaning it shouldn't be animated when loaded into
+	 * UE */
+	USDUTILITIES_API bool IsAttributeMuted(const UE::FUsdAttribute& Attribute, const UE::FUsdStage& Stage);
 
 #if USE_USD_SDK
 	/**
@@ -46,7 +48,7 @@ namespace UsdUtils
 	 * Use this after setting any attribute: This function does nothing in case the Stage's current edit target has
 	 * the strongest opinion for the attribute already
 	 */
-	USDUTILITIES_API void NotifyIfOverriddenOpinion( const pxr::UsdProperty& Property );
+	USDUTILITIES_API void NotifyIfOverriddenOpinion(const pxr::UsdProperty& Property);
 
 	/**
 	 * Generally we are not allowed to author opinions inside instanceables. Our actual results when doing that will range from
@@ -54,6 +56,5 @@ namespace UsdUtils
 	 * This will show a warning and return true if Prim is an instance, letting us avoid those scenarios
 	 */
 	USDUTILITIES_API bool NotifyIfInstanceProxy(const pxr::UsdPrim& Prim);
-#endif // USE_USD_SDK
+#endif	  // USE_USD_SDK
 }
-
diff --git a/Engine/Plugins/Importers/USDImporter/Source/USDUtilities/Public/USDConversionUtils.h b/Engine/Plugins/Importers/USDImporter/Source/USDUtilities/Public/USDConversionUtils.h
--- a/Engine/Plugins/Importers/USDImporter/Source/USDUtilities/Public/USDConversionUtils.h
+++ b/Engine/Plugins/Importers/USDImporter/Source/USDUtilities/Public/USDConversionUtils.h
@@ -15,8 +15,8 @@
 
 #if USE_USD_SDK
 #include "USDIncludesStart.h"
-	#include "pxr/pxr.h"
-	#include "pxr/usd/usdShade/tokens.h"
+#include "pxr/pxr.h"
+#include "pxr/usd/usdShade/tokens.h"
 #include "USDIncludesEnd.h"
 
 PXR_NAMESPACE_OPEN_SCOPE
@@ -35,15 +35,17 @@ PXR_NAMESPACE_OPEN_SCOPE
 	class VtValue;
 
 	class UsdStage;
-	template< typename T > class TfRefPtr;
-	using UsdStageRefPtr = TfRefPtr< UsdStage >;
+	template<typename T>
+	class TfRefPtr;
+	using UsdStageRefPtr = TfRefPtr<UsdStage>;
 
 	class SdfPrimSpec;
-	template <class T> class SdfHandle;
-	using SdfPrimSpecHandle = SdfHandle< SdfPrimSpec >;
+	template<class T>
+	class SdfHandle;
+	using SdfPrimSpecHandle = SdfHandle<SdfPrimSpec>;
 PXR_NAMESPACE_CLOSE_SCOPE
 
-#endif // #if USE_USD_SDK
+#endif	  // #if USE_USD_SDK
 
 class UAssetImportData;
 class USceneComponent;
@@ -53,10 +55,12 @@ enum class EUsdDrawMode : int32;
 enum class EUsdDuplicateType : uint8;
 enum class EUsdUpAxis : uint8;
 struct FUsdUnrealAssetInfo;
+
 namespace UE
 {
 	class FUsdPrim;
 }
+
 namespace UsdUtils
 {
 	struct FUsdPrimMaterialAssignmentInfo;
@@ -65,7 +69,7 @@ namespace UsdUtils
 namespace UsdUtils
 {
 	template<typename T>
-	T* FindOrCreateObject( UObject* InParent, const FString& InName, EObjectFlags Flags )
+	T* FindOrCreateObject(UObject* InParent, const FString& InName, EObjectFlags Flags)
 	{
 		T* Object = FindObject<T>(InParent, *InName);
 
@@ -78,60 +82,63 @@ namespace UsdUtils
 	}
 
 	/** Case sensitive hashing function for TMap */
-	template <typename ValueType>
-	struct FCaseSensitiveStringMapFuncs : BaseKeyFuncs<ValueType, FString, /*bInAllowDuplicateKeys*/false>
+	template<typename ValueType>
+	struct FCaseSensitiveStringMapFuncs : BaseKeyFuncs<ValueType, FString, /*bInAllowDuplicateKeys*/ false>
 	{
-		static FORCEINLINE const FString& GetSetKey( const TPair<FString, ValueType>& Element )
+		static FORCEINLINE const FString& GetSetKey(const TPair<FString, ValueType>& Element)
 		{
 			return Element.Key;
 		}
-		static FORCEINLINE bool Matches( const FString& A, const FString& B )
+
+		static FORCEINLINE bool Matches(const FString& A, const FString& B)
 		{
-			return A.Equals( B, ESearchCase::CaseSensitive );
+			return A.Equals(B, ESearchCase::CaseSensitive);
 		}
-		static FORCEINLINE uint32 GetKeyHash( const FString& Key )
+
+		static FORCEINLINE uint32 GetKeyHash(const FString& Key)
 		{
-			return FCrc::StrCrc32<TCHAR>( *Key );
+			return FCrc::StrCrc32<TCHAR>(*Key);
 		}
 	};
 
 #if USE_USD_SDK
-	template< typename ValueType >
-	USDUTILITIES_API ValueType GetUsdValue( const pxr::UsdAttribute& Attribute, pxr::UsdTimeCode TimeCode );
+	template<typename ValueType>
+	USDUTILITIES_API ValueType GetUsdValue(const pxr::UsdAttribute& Attribute, pxr::UsdTimeCode TimeCode);
 
-	USDUTILITIES_API pxr::TfToken GetUsdStageUpAxis( const pxr::UsdStageRefPtr& Stage );
-	USDUTILITIES_API EUsdUpAxis GetUsdStageUpAxisAsEnum( const pxr::UsdStageRefPtr& Stage );
+	USDUTILITIES_API pxr::TfToken GetUsdStageUpAxis(const pxr::UsdStageRefPtr& Stage);
+	USDUTILITIES_API EUsdUpAxis GetUsdStageUpAxisAsEnum(const pxr::UsdStageRefPtr& Stage);
 
-	USDUTILITIES_API void SetUsdStageUpAxis( const pxr::UsdStageRefPtr& Stage, pxr::TfToken Axis );
-	USDUTILITIES_API void SetUsdStageUpAxis( const pxr::UsdStageRefPtr& Stage, EUsdUpAxis Axis );
+	USDUTILITIES_API void SetUsdStageUpAxis(const pxr::UsdStageRefPtr& Stage, pxr::TfToken Axis);
+	USDUTILITIES_API void SetUsdStageUpAxis(const pxr::UsdStageRefPtr& Stage, EUsdUpAxis Axis);
 
-	USDUTILITIES_API float GetUsdStageMetersPerUnit( const pxr::UsdStageRefPtr& Stage );
-	USDUTILITIES_API void SetUsdStageMetersPerUnit( const pxr::UsdStageRefPtr& Stage, float MetersPerUnit );
+	USDUTILITIES_API float GetUsdStageMetersPerUnit(const pxr::UsdStageRefPtr& Stage);
+	USDUTILITIES_API void SetUsdStageMetersPerUnit(const pxr::UsdStageRefPtr& Stage, float MetersPerUnit);
 
-	USDUTILITIES_API int32 GetUsdStageNumFrames( const pxr::UsdStageRefPtr& Stage );
+	USDUTILITIES_API int32 GetUsdStageNumFrames(const pxr::UsdStageRefPtr& Stage);
 
-	USDUTILITIES_API bool HasCompositionArcs( const pxr::UsdPrim& Prim );
-	USDUTILITIES_API bool HasCompositionArcs( const pxr::SdfPrimSpecHandle& PrimSpec );
+	USDUTILITIES_API bool HasCompositionArcs(const pxr::UsdPrim& Prim);
+	USDUTILITIES_API bool HasCompositionArcs(const pxr::SdfPrimSpecHandle& PrimSpec);
 
-	USDUTILITIES_API UClass* GetActorTypeForPrim( const pxr::UsdPrim& Prim );
-	USDUTILITIES_API UClass* GetComponentTypeForPrim( const pxr::UsdPrim& Prim );
+	USDUTILITIES_API UClass* GetActorTypeForPrim(const pxr::UsdPrim& Prim);
+	USDUTILITIES_API UClass* GetComponentTypeForPrim(const pxr::UsdPrim& Prim);
 
 	/** Returns the USD schema name that should be used when exporting Component (e.g. "Xform", "Mesh", "Camera", etc.) */
-	USDUTILITIES_API FString GetSchemaNameForComponent( const USceneComponent& Component );
+	USDUTILITIES_API FString GetSchemaNameForComponent(const USceneComponent& Component);
 
 	/**
 	 * Returns a prim path to use when exporting a given ActorOrComponent, assuming we're exporting according to the component attachment hierarchy.
-	 * e.g. We have a top-level actor with label "Actor" with an attached child actor "Child" that has a component "Comp1" attached to its root, and "Comp2"
-	 * attached to "Comp1". If Comp2 is provided to this function, it will return "/Root/Actor/Child/Comp1/Comp2".
+	 * e.g. We have a top-level actor with label "Actor" with an attached child actor "Child" that has a component "Comp1" attached to its root, and
+	 * "Comp2" attached to "Comp1". If Comp2 is provided to this function, it will return "/Root/Actor/Child/Comp1/Comp2".
 	 *
 	 * The actor label is always used in place of the root component name.
 	 *
 	 * ParentPrimPath is optional, and without it the function will recurse upwards and build the full path all the way to the root. Actor folders
 	 * are handled just like if they were actors/components
 	 */
-	USDUTILITIES_API FString GetPrimPathForObject( const UObject* ActorOrComponent, const FString& ParentPrimPath = TEXT(""), bool bUseActorFolders = false );
+	USDUTILITIES_API FString
+	GetPrimPathForObject(const UObject* ActorOrComponent, const FString& ParentPrimPath = TEXT(""), bool bUseActorFolders = false);
 
-	USDUTILITIES_API TUsdStore< pxr::TfToken > GetUVSetName( int32 UVChannelIndex );
+	USDUTILITIES_API TUsdStore<pxr::TfToken> GetUVSetName(int32 UVChannelIndex);
 
 	/**
 	 * Heuristic to try and guess what UV index we should assign this primvar to.
@@ -141,7 +148,7 @@ namespace UsdUtils
 	 * @return UV index that should be used for this primvar
 	 */
 	UE_DEPRECATED(5.3, "This heuristic is no longer used. Check GetUVSetPrimvars for the new behavior.")
-	USDUTILITIES_API int32 GetPrimvarUVIndex( FString PrimvarName );
+	USDUTILITIES_API int32 GetPrimvarUVIndex(FString PrimvarName);
 
 	/**
 	 * Gets the primvars that should be used as UV sets, per index, for this UsdPrim.
@@ -165,21 +172,31 @@ namespace UsdUtils
 	 * (e.g. first item of array is primvar for UV set 0, second for UV set 1, etc).
 	 * This overload will only return primvars with 'texcoord2f' role.	 *
 	 * @param UsdMesh - Mesh that contains primvars that can be used as texture coordinates.
-	 * @param MaterialToPrimvarsUVSetNames - Maps from a material prim path, to pairs indicating which primvar names are used as 'st' coordinates, and which UVIndex the imported material will sample from (e.g. ["st0", 0], ["myUvSet2", 2], etc). These are supposed to be the materials used by the mesh, and we do this because it helps identify which primvars are valid/used as texture coordinates, as the user may have these named as 'myUvSet2' and still expect it to work
-	 * @param UsdMeshMaterialAssignmentInfo - Result of calling GetPrimMaterialAssignments on UsdMesh's Prim. This can be provided or will be retrieved on-demand using RenderContext
+	 * @param MaterialToPrimvarsUVSetNames - Maps from a material prim path, to pairs indicating which primvar names are used as 'st' coordinates, and
+	 * which UVIndex the imported material will sample from (e.g. ["st0", 0], ["myUvSet2", 2], etc). These are supposed to be the materials used by
+	 * the mesh, and we do this because it helps identify which primvars are valid/used as texture coordinates, as the user may have these named as
+	 * 'myUvSet2' and still expect it to work
+	 * @param UsdMeshMaterialAssignmentInfo - Result of calling GetPrimMaterialAssignments on UsdMesh's Prim. This can be provided or will be
+	 * retrieved on-demand using RenderContext
 	 * @param RenderContext - Render context to use when traversing through material shaders looking for used primvars
 	 * @param MaterialPurpose - Specific material purpose to use when parsing the UsdMesh's material bindings
 	 * @return Array where each index gives the primvar that should be used for that UV index
 	 */
-	UE_DEPRECATED(5.3, "Use the signature that just receives an UsdMesh, as the returned primvars no longer depend on material assignment information.")
-	USDUTILITIES_API TArray< TUsdStore< pxr::UsdGeomPrimvar > > GetUVSetPrimvars(
+	UE_DEPRECATED(
+		5.3,
+		"Use the signature that just receives an UsdMesh, as the returned primvars no longer depend on material assignment information."
+	)
+	USDUTILITIES_API TArray<TUsdStore<pxr::UsdGeomPrimvar>> GetUVSetPrimvars(
 		const pxr::UsdGeomMesh& UsdMesh,
 		const TMap<FString, TMap<FString, int32>>& MaterialToPrimvarsUVSetNames,
 		const pxr::TfToken& RenderContext = pxr::UsdShadeTokens->universalRenderContext,
 		const pxr::TfToken& MaterialPurpose = pxr::UsdShadeTokens->allPurpose
 	);
-	UE_DEPRECATED(5.3, "Use the signature that just receives an UsdMesh, as the returned primvars no longer depend on material assignment information.")
-	USDUTILITIES_API TArray< TUsdStore< pxr::UsdGeomPrimvar > > GetUVSetPrimvars(
+	UE_DEPRECATED(
+		5.3,
+		"Use the signature that just receives an UsdMesh, as the returned primvars no longer depend on material assignment information."
+	)
+	USDUTILITIES_API TArray<TUsdStore<pxr::UsdGeomPrimvar>> GetUVSetPrimvars(
 		const pxr::UsdGeomMesh& UsdMesh,
 		const TMap<FString, TMap<FString, int32>>& MaterialToPrimvarsUVSetNames,
 		const UsdUtils::FUsdPrimMaterialAssignmentInfo& UsdMeshMaterialAssignmentInfo
@@ -212,19 +229,14 @@ namespace UsdUtils
 	 * first)
 	 * @return Mapping from primvar names to the UV index
 	 */
-	USDUTILITIES_API TMap<FString, int32> AssemblePrimvarsIntoPrimvarToUVIndexMap(
-		const TArray<TUsdStore<pxr::UsdGeomPrimvar>>& AllMeshUVPrimvars
-	);
+	USDUTILITIES_API TMap<FString, int32> AssemblePrimvarsIntoPrimvarToUVIndexMap(const TArray<TUsdStore<pxr::UsdGeomPrimvar>>& AllMeshUVPrimvars);
 
 	/**
 	 * Given an array of primvars, and an array of which of those are preferred (AllPrimvars should also include
 	 * PreferredPrimvars), this will sort lexicographically and return a prioritized PrimvarToUVIndex assignment map of
 	 * up to USD_PREVIEW_SURFACE_MAX_UV_SETS UV indices
 	 */
-	USDUTILITIES_API TMap<FString, int32> CombinePrimvarsIntoUVSets(
-		const TSet<FString>& AllPrimvars,
-		const TSet<FString>& PreferredPrimvars
-	);
+	USDUTILITIES_API TMap<FString, int32> CombinePrimvarsIntoUVSets(const TSet<FString>& AllPrimvars, const TSet<FString>& PreferredPrimvars);
 
 	template<typename T, typename U>
 	FString StringifyMap(const TMap<T, U>& Map)
@@ -239,8 +251,8 @@ namespace UsdUtils
 		return Result;
 	}
 
-	USDUTILITIES_API bool IsAnimated( const pxr::UsdPrim& Prim );
-	USDUTILITIES_API bool HasAnimatedVisibility( const pxr::UsdPrim& Prim );
+	USDUTILITIES_API bool IsAnimated(const pxr::UsdPrim& Prim);
+	USDUTILITIES_API bool HasAnimatedVisibility(const pxr::UsdPrim& Prim);
 
 	/**
 	 * Returns true if Prim should have animated bounds (either from having authored animated `extent` or `extentsHint`
@@ -269,7 +281,7 @@ namespace UsdUtils
 	);
 
 	/** Returns whether Prim belongs to any of the default kinds, or a kind derived from them. The result can be a union of different kinds. */
-	USDUTILITIES_API EUsdDefaultKind GetDefaultKind( const pxr::UsdPrim& Prim );
+	USDUTILITIES_API EUsdDefaultKind GetDefaultKind(const pxr::UsdPrim& Prim);
 
 	/**
 	 * Sets the prim kind for Prim to the provided NewKind.
@@ -278,7 +290,7 @@ namespace UsdUtils
 	 * @param Prim - Prim to receive the new Kind
 	 * @return Whether we managed to set the kind or not.
 	 */
-	USDUTILITIES_API bool SetDefaultKind( pxr::UsdPrim& Prim, EUsdDefaultKind NewKind );
+	USDUTILITIES_API bool SetDefaultKind(pxr::UsdPrim& Prim, EUsdDefaultKind NewKind);
 
 	/**
 	 * Returns whether the prim has the UsdGeomModelAPI schema and should be drawn with one of the alternative draw modes, such as cards or bounds.
@@ -290,14 +302,28 @@ namespace UsdUtils
 	 * Returns all prims of type SchemaType (or a descendant type) in the subtree of prims rooted at StartPrim.
 	 * Stops going down the subtrees when it hits a schema type to exclude.
 	 */
-	USDUTILITIES_API TArray< TUsdStore< pxr::UsdPrim > > GetAllPrimsOfType( const pxr::UsdPrim& StartPrim, const pxr::TfType& SchemaType, const TArray< TUsdStore< pxr::TfType > >& ExcludeSchemaTypes = {} );
-	USDUTILITIES_API TArray< TUsdStore< pxr::UsdPrim > > GetAllPrimsOfType( const pxr::UsdPrim& StartPrim, const pxr::TfType& SchemaType, TFunction< bool( const pxr::UsdPrim& ) > PruneChildren, const TArray< TUsdStore< pxr::TfType > >& ExcludeSchemaTypes = {} );
+	USDUTILITIES_API TArray<TUsdStore<pxr::UsdPrim>> GetAllPrimsOfType(
+		const pxr::UsdPrim& StartPrim,
+		const pxr::TfType& SchemaType,
+		const TArray<TUsdStore<pxr::TfType>>& ExcludeSchemaTypes = {}
+	);
+	USDUTILITIES_API TArray<TUsdStore<pxr::UsdPrim>> GetAllPrimsOfType(
+		const pxr::UsdPrim& StartPrim,
+		const pxr::TfType& SchemaType,
+		TFunction<bool(const pxr::UsdPrim&)> PruneChildren,
+		const TArray<TUsdStore<pxr::TfType>>& ExcludeSchemaTypes = {}
+	);
 
-	USDUTILITIES_API FString GetAssetPathFromPrimPath( const FString& RootContentPath, const pxr::UsdPrim& Prim );
-#endif // #if USE_USD_SDK
+	USDUTILITIES_API FString GetAssetPathFromPrimPath(const FString& RootContentPath, const pxr::UsdPrim& Prim);
+#endif	  // #if USE_USD_SDK
 
-	USDUTILITIES_API TArray< UE::FUsdPrim > GetAllPrimsOfType( const UE::FUsdPrim& StartPrim, const TCHAR* SchemaName );
-	USDUTILITIES_API TArray< UE::FUsdPrim > GetAllPrimsOfType( const UE::FUsdPrim& StartPrim, const TCHAR* SchemaName, TFunction< bool( const UE::FUsdPrim& ) > PruneChildren, const TArray<const TCHAR*>& ExcludeSchemaNames = {} );
+	USDUTILITIES_API TArray<UE::FUsdPrim> GetAllPrimsOfType(const UE::FUsdPrim& StartPrim, const TCHAR* SchemaName);
+	USDUTILITIES_API TArray<UE::FUsdPrim> GetAllPrimsOfType(
+		const UE::FUsdPrim& StartPrim,
+		const TCHAR* SchemaName,
+		TFunction<bool(const UE::FUsdPrim&)> PruneChildren,
+		const TArray<const TCHAR*>& ExcludeSchemaNames = {}
+	);
 
 	/** Returns the time code for non-timesampled values. Usually a quiet NaN. */
 	USDUTILITIES_API double GetDefaultTimeCode();
@@ -318,6 +344,7 @@ namespace UsdUtils
 	 * Returns the object's UsdAssetUserData of a particular subclass if it has one
 	 */
 	USDUTILITIES_API UUsdAssetUserData* GetAssetUserData(const UObject* Object, TSubclassOf<UUsdAssetUserData> Class = {});
+
 	template<typename T>
 	inline T* GetAssetUserData(UObject* Object)
 	{
@@ -328,6 +355,7 @@ namespace UsdUtils
 	 * Makes sure Object has an instance of UUsdAssetUserData of the provided subclass (defaulting to just UUsdAssetUserData itself) and returns it
 	 */
 	USDUTILITIES_API UUsdAssetUserData* GetOrCreateAssetUserData(UObject* Object, TSubclassOf<UUsdAssetUserData> Class = {});
+
 	template<typename T>
 	inline T* GetOrCreateAssetUserData(UObject* Object)
 	{
@@ -347,16 +375,28 @@ namespace UsdUtils
 	 * Whenever bReuseIdenticalAssets is false, the prefix will be a SHA1 hash of the prim path and its stage identifier.
 	 */
 	USDUTILITIES_API FString GetAssetHashPrefix(const pxr::UsdPrim& PrimForAsset, bool bReuseIdenticalAssets);
-#endif // #if USE_USD_SDK
+#endif	  // #if USE_USD_SDK
 
 	/** Adds a reference on Prim to the layer at AbsoluteFilePath */
-	USDUTILITIES_API void AddReference( UE::FUsdPrim& Prim, const TCHAR* AbsoluteFilePath, const UE::FSdfPath& TargetPrimPath = {}, double TimeCodeOffset = 0.0, double TimeCodeScale = 1.0 );
+	USDUTILITIES_API void AddReference(
+		UE::FUsdPrim& Prim,
+		const TCHAR* AbsoluteFilePath,
+		const UE::FSdfPath& TargetPrimPath = {},
+		double TimeCodeOffset = 0.0,
+		double TimeCodeScale = 1.0
+	);
 
 	/** Gets the strongest direct reference on Prim with the given FileExtension. Returns true if there was one. */
 	USDUTILITIES_API bool GetReferenceFilePath(const UE::FUsdPrim& Prim, const FString& FileExtension, FString& OutReferenceFilePath);
 
 	/** Adds a payload on Prim pointing at the default prim of the layer at AbsoluteFilePath */
-	USDUTILITIES_API void AddPayload( UE::FUsdPrim& Prim, const TCHAR* AbsoluteFilePath, const UE::FSdfPath& TargetPrimPath = {}, double TimeCodeOffset = 0.0, double TimeCodeScale = 1.0 );
+	USDUTILITIES_API void AddPayload(
+		UE::FUsdPrim& Prim,
+		const TCHAR* AbsoluteFilePath,
+		const UE::FSdfPath& TargetPrimPath = {},
+		double TimeCodeOffset = 0.0,
+		double TimeCodeScale = 1.0
+	);
 
 	/**
 	 * Renames a single prim to a new name
@@ -368,13 +408,13 @@ namespace UsdUtils
 	 * @param NewPrimName - New name for the prim e.g. "MyNewName"
 	 * @return True if we managed to rename
 	 */
-	USDUTILITIES_API bool RenamePrim( UE::FUsdPrim& Prim, const TCHAR* NewPrimName );
+	USDUTILITIES_API bool RenamePrim(UE::FUsdPrim& Prim, const TCHAR* NewPrimName);
 
 	/**
 	 * Removes any numbered suffix, followed by any number of underscores (e.g. Asset_2, Asset__232_31 or Asset94 all become 'Asset'), making
 	 * sure the string is kept at least one character long. Returns true if it removed anything.
 	 */
-	USDUTILITIES_API bool RemoveNumberedSuffix( FString& Prefix );
+	USDUTILITIES_API bool RemoveNumberedSuffix(FString& Prefix);
 
 	/**
 	 * Appends numbered suffixes to Name until the result is not contained in UsedNames, and returns it.
@@ -383,15 +423,15 @@ namespace UsdUtils
 	 * @param UsedNames - Strings that cannot be used for the result
 	 * @return Modified Name so that it doesn't match anything in UsedNames (e.g. "MyName" again, or "MyName_0" or "MyName_423")
 	 */
-	USDUTILITIES_API FString GetUniqueName( FString Name, const TSet<FString>& UsedNames );
+	USDUTILITIES_API FString GetUniqueName(FString Name, const TSet<FString>& UsedNames);
 
 #if USE_USD_SDK
 	/**
 	 * Returns a sanitized, valid version of 'InName' that can be used as the name for a child prim of 'ParentPrim'.
 	 * Returns the empty string in case of an error.
 	 */
-	USDUTILITIES_API FString GetValidChildName( FString InName, const pxr::UsdPrim& ParentPrim );
-#endif // #if USE_USD_SDK
+	USDUTILITIES_API FString GetValidChildName(FString InName, const pxr::UsdPrim& ParentPrim);
+#endif	  // #if USE_USD_SDK
 
 	/**
 	 * Returns a modified version of InIdentifier that can be used as a USD prim or property name.
@@ -399,24 +439,28 @@ namespace UsdUtils
 	 * Additionally, the first character cannot be a number.
 	 * Note that this obviously doesn't check for a potential name collision.
 	 */
-	USDUTILITIES_API FString SanitizeUsdIdentifier( const TCHAR* InIdentifier );
+	USDUTILITIES_API FString SanitizeUsdIdentifier(const TCHAR* InIdentifier);
 
 	/** Will call UsdGeomImageable::MakeVisible/MakeInvisible if Prim is a UsdGeomImageable */
-	USDUTILITIES_API void MakeVisible( UE::FUsdPrim& Prim, double TimeCode = UsdUtils::GetDefaultTimeCode() );
-	USDUTILITIES_API void MakeInvisible( UE::FUsdPrim& Prim, double TimeCode = UsdUtils::GetDefaultTimeCode() );
+	USDUTILITIES_API void MakeVisible(UE::FUsdPrim& Prim, double TimeCode = UsdUtils::GetDefaultTimeCode());
+	USDUTILITIES_API void MakeInvisible(UE::FUsdPrim& Prim, double TimeCode = UsdUtils::GetDefaultTimeCode());
 
 	/** Returns if the ComputedVisibility for Prim says it should be visible */
-	USDUTILITIES_API bool IsVisible( const UE::FUsdPrim& Prim, double TimeCode = UsdUtils::GetDefaultTimeCode() );
+	USDUTILITIES_API bool IsVisible(const UE::FUsdPrim& Prim, double TimeCode = UsdUtils::GetDefaultTimeCode());
 
 	/** Returns whether Prim has visibility set to 'inherited' */
-	USDUTILITIES_API bool HasInheritedVisibility( const UE::FUsdPrim& Prim, double TimeCode = UsdUtils::GetDefaultTimeCode() );
+	USDUTILITIES_API bool HasInheritedVisibility(const UE::FUsdPrim& Prim, double TimeCode = UsdUtils::GetDefaultTimeCode());
 
 	/**
 	 * Travels up from Prim and returns true if we hit any invisible prim before we hit the stage pseudoroot.
 	 * Does not check `Prim` itself (or `RootPrim`, so that you can just pass the PseudoRoot in most cases).
 	 * Completely ignores whether non-UsdGeomImageable prims are visible or not and keeps travelling up.
 	 */
-	USDUTILITIES_API bool HasInvisibleParent( const UE::FUsdPrim& Prim, const UE::FUsdPrim& RootPrim, double TimeCode = UsdUtils::GetDefaultTimeCode() );
+	USDUTILITIES_API bool HasInvisibleParent(
+		const UE::FUsdPrim& Prim,
+		const UE::FUsdPrim& RootPrim,
+		double TimeCode = UsdUtils::GetDefaultTimeCode()
+	);
 
 	/**
 	 * Returns the prims in the subtree of Prim (potentially including Prim itself) that are visible due to the
@@ -437,7 +481,7 @@ namespace UsdUtils
 	 *
 	 * Returns an empty path in case the layer doesn't have a spec for this prim.
 	 */
-	USDUTILITIES_API UE::FSdfPath GetPrimSpecPathForLayer( const UE::FUsdPrim& Prim, const UE::FSdfLayer& Layer );
+	USDUTILITIES_API UE::FSdfPath GetPrimSpecPathForLayer(const UE::FUsdPrim& Prim, const UE::FSdfLayer& Layer);
 
 	/**
 	 * Removes all the prim specs for Prim on the given Layer, if the layer belongs to the stage's local layer stack.
@@ -452,7 +496,7 @@ namespace UsdUtils
 	 * @param Layer - Layer to remove prim specs from. This can be left with the invalid layer (default) in order to remove all
 	 *				  specs from the entire stage's local layer stack.
 	 */
-	USDUTILITIES_API void RemoveAllLocalPrimSpecs( const UE::FUsdPrim& Prim, const UE::FSdfLayer& Layer = UE::FSdfLayer{} );
+	USDUTILITIES_API void RemoveAllLocalPrimSpecs(const UE::FUsdPrim& Prim, const UE::FSdfLayer& Layer = UE::FSdfLayer{});
 
 	/**
 	 * Copies flattened versions of the input prims onto the clipboard stage and removes all the prim specs for Prims from their stages.
@@ -461,7 +505,7 @@ namespace UsdUtils
 	 * @param Prims - Prims to cut
 	 * @return True if we managed to cut
 	 */
-	USDUTILITIES_API bool CutPrims( const TArray<UE::FUsdPrim>& Prims );
+	USDUTILITIES_API bool CutPrims(const TArray<UE::FUsdPrim>& Prims);
 
 	/**
 	 * Copies flattened versions of the input prims onto the clipboard stage.
@@ -470,7 +514,7 @@ namespace UsdUtils
 	 * @param Prims - Prims to copy
 	 * @return True if we managed to copy
 	 */
-	USDUTILITIES_API bool CopyPrims( const TArray<UE::FUsdPrim>& Prims );
+	USDUTILITIES_API bool CopyPrims(const TArray<UE::FUsdPrim>& Prims);
 
 	/**
 	 * Pastes the prims from the clipboard stage as children of ParentPrim.
@@ -484,7 +528,7 @@ namespace UsdUtils
 	 * @param ParentPrim - Prim that will become parent to the pasted prims
 	 * @return Paths to the pasted prim specs, after they were added as children of ParentPrim
 	 */
-	USDUTILITIES_API TArray<UE::FSdfPath> PastePrims( const UE::FUsdPrim& ParentPrim );
+	USDUTILITIES_API TArray<UE::FSdfPath> PastePrims(const UE::FUsdPrim& ParentPrim);
 
 	/** Returns true if we have prims that we can paste within our clipboard stage */
 	USDUTILITIES_API bool CanPastePrims();
@@ -514,17 +558,16 @@ namespace UsdUtils
 	);
 
 	/** Adds to Prim the assetInfo metadata the values described in Info */
-	USDUTILITIES_API void SetPrimAssetInfo( UE::FUsdPrim& Prim, const FUsdUnrealAssetInfo& Info );
+	USDUTILITIES_API void SetPrimAssetInfo(UE::FUsdPrim& Prim, const FUsdUnrealAssetInfo& Info);
 
 	/** Retrieves from Prim the assetInfo metadata values that we use as export metadata, when exporting Unreal assets */
-	USDUTILITIES_API FUsdUnrealAssetInfo GetPrimAssetInfo( const UE::FUsdPrim& Prim );
+	USDUTILITIES_API FUsdUnrealAssetInfo GetPrimAssetInfo(const UE::FUsdPrim& Prim);
 
 #if USE_USD_SDK
 	/* Removes all metadata from Prim, except the typeName and specifier entries */
 	USDUTILITIES_API bool ClearNonEssentialPrimMetadata(const pxr::UsdPrim& Prim);
-#endif // USE_USD_SDK
+#endif	  // USE_USD_SDK
 
 	/** Collects how many times each schema shows up on the provided stage and send it as an analytics event */
 	USDUTILITIES_API void CollectSchemaAnalytics(const UE::FUsdStage& Stage, const FString& EventName);
-}
-
+}	 // namespace UsdUtils
diff --git a/Engine/Plugins/Importers/USDImporter/Source/USDUtilities/Public/USDErrorUtils.h b/Engine/Plugins/Importers/USDImporter/Source/USDUtilities/Public/USDErrorUtils.h
--- a/Engine/Plugins/Importers/USDImporter/Source/USDUtilities/Public/USDErrorUtils.h
+++ b/Engine/Plugins/Importers/USDImporter/Source/USDUtilities/Public/USDErrorUtils.h
@@ -7,23 +7,23 @@
 
 namespace UsdUtils
 {
-    /**
-     * Pushes an USD error monitoring object into the stack and catches any emitted errors
-     */
+	/**
+	 * Pushes an USD error monitoring object into the stack and catches any emitted errors
+	 */
 	USDUTILITIES_API void StartMonitoringErrors();
 
-    /**
-     * Returns all errors that were captured since StartMonitoringErrors(), clears and pops an
-     * error monitoring object from the stack
-     */
+	/**
+	 * Returns all errors that were captured since StartMonitoringErrors(), clears and pops an
+	 * error monitoring object from the stack
+	 */
 	USDUTILITIES_API TArray<FString> GetErrorsAndStopMonitoring();
 
-    /**
-     * Displays the error messages for each captured error since StartMonitoringErrors(),
+	/**
+	 * Displays the error messages for each captured error since StartMonitoringErrors(),
 	 * clears and pops an error monitoring object from the stack.
 	 * If ToastMessage is empty, a default message will be displayed.
-     * Returns true if there were any errors.
-     */
+	 * Returns true if there were any errors.
+	 */
 	USDUTILITIES_API bool ShowErrorsAndStopMonitoring(const FText& ToastMessage = FText());
 }
 
@@ -40,10 +40,10 @@ namespace UE
 		public:
 			FUsdMessageLog() = default;
 
-			FUsdMessageLog( const FUsdMessageLog& ) = default;
-			FUsdMessageLog& operator=( const FUsdMessageLog& ) = default;
-			FUsdMessageLog( FUsdMessageLog&& ) = default;
-			FUsdMessageLog& operator=( FUsdMessageLog&& ) = default;
+			FUsdMessageLog(const FUsdMessageLog&) = default;
+			FUsdMessageLog& operator=(const FUsdMessageLog&) = default;
+			FUsdMessageLog(FUsdMessageLog&&) = default;
+			FUsdMessageLog& operator=(FUsdMessageLog&&) = default;
 
 			~FUsdMessageLog();
 
@@ -51,28 +51,27 @@ namespace UE
 			/** Display pending messages, show the log window */
 			void Dump();
 
-		
 			/** Append a message */
-			void Push( const TSharedRef< FTokenizedMessage >& Message );
+			void Push(const TSharedRef<FTokenizedMessage>& Message);
 
 		private:
-			TArray< TSharedRef< FTokenizedMessage > > TokenizedMessages;
+			TArray<TSharedRef<FTokenizedMessage>> TokenizedMessages;
 		};
 	}
-}
+}	 // namespace UE
 
 class USDUTILITIES_API FUsdLogManager
 {
 public:
 	/** Sends the message to Message Log if it exists, if not, sends it to the Output Log */
-	static void LogMessage( EMessageSeverity::Type Severity, const FText& Message );
-	static void LogMessage( const TSharedRef< FTokenizedMessage >& Message );
+	static void LogMessage(EMessageSeverity::Type Severity, const FText& Message);
+	static void LogMessage(const TSharedRef<FTokenizedMessage>& Message);
 
 	static void EnableMessageLog();
 	static void DisableMessageLog();
 
 private:
-	static TOptional< UE::Internal::FUsdMessageLog > MessageLog;
+	static TOptional<UE::Internal::FUsdMessageLog> MessageLog;
 	static int32 MessageLogRefCount;
 
 	static FCriticalSection MessageLogLock;
@@ -87,8 +86,8 @@ public:
 	FScopedUsdMessageLog();
 	~FScopedUsdMessageLog();
 
-	FScopedUsdMessageLog( const FScopedUsdMessageLog& ) = delete;
-	FScopedUsdMessageLog& operator=( const FScopedUsdMessageLog& ) = delete;
-	FScopedUsdMessageLog( FScopedUsdMessageLog&& ) = delete;
-	FScopedUsdMessageLog& operator=( FScopedUsdMessageLog&& ) = delete;
+	FScopedUsdMessageLog(const FScopedUsdMessageLog&) = delete;
+	FScopedUsdMessageLog& operator=(const FScopedUsdMessageLog&) = delete;
+	FScopedUsdMessageLog(FScopedUsdMessageLog&&) = delete;
+	FScopedUsdMessageLog& operator=(FScopedUsdMessageLog&&) = delete;
 };
