USD: Reformat USDUtilities module.

diff --git a/Engine/Plugins/Importers/USDImporter/Source/USDUtilities/Public/USDGeomMeshConversion.h b/Engine/Plugins/Importers/USDImporter/Source/USDUtilities/Public/USDGeomMeshConversion.h
--- a/Engine/Plugins/Importers/USDImporter/Source/USDUtilities/Public/USDGeomMeshConversion.h
+++ b/Engine/Plugins/Importers/USDImporter/Source/USDUtilities/Public/USDGeomMeshConversion.h
@@ -5,14 +5,14 @@
 #include "CoreMinimal.h"
 
 #include "UnrealUSDWrapper.h"
-#include "USDClassesModule.h"  // Include this for now so that we can redirect FDisplayColorMaterial to the USDClasses definition
+#include "USDClassesModule.h"	 // Include this for now so that we can redirect FDisplayColorMaterial to the USDClasses definition
 #include "UsdWrappers/ForwardDeclarations.h"
 
 #if USE_USD_SDK
 
 #include "USDIncludesStart.h"
-	#include "pxr/usd/usd/timeCode.h"
-	#include "pxr/usd/usdShade/tokens.h"
+#include "pxr/usd/usd/timeCode.h"
+#include "pxr/usd/usdShade/tokens.h"
 #include "USDIncludesEnd.h"
 
 PXR_NAMESPACE_OPEN_SCOPE
@@ -33,10 +33,12 @@ struct FMeshDescription;
 struct FStaticMeshLODResources;
 struct FUsdCollapsingCache;
 struct FUsdStageInfo;
+
 namespace UsdUtils
 {
 	struct FUsdPrimMaterialAssignmentInfo;
 }
+
 namespace UE
 {
 	class FUsdPrim;
@@ -81,7 +83,7 @@ namespace UsdToUnreal
 		// this mesh, and which UVIndex materials will sample from( e.g.[ "st0", 0 ], [ "myUvSet2", 2 ], etc ).
 		// This is used to pick which primvars will become UV sets.
 		UE_DEPRECATED(5.3, "No longer used, now this information is stored directly on generated Material AssetImportData/AssetUserData")
-		const TMap< FString, TMap<FString, int32> >* MaterialToPrimvarToUVIndex;
+		const TMap<FString, TMap<FString, int32>>* MaterialToPrimvarToUVIndex;
 
 		// Whether to try reusing material slots (both local and the ones already in MaterialAssignments) when
 		// converting the material assignments from the mesh
@@ -187,12 +189,8 @@ namespace UsdToUnreal
 	 * @param OutTransform - The converted transform on the UE coordinate system
 	 * @return Whether the conversion was successful or not.
 	 */
-	USDUTILITIES_API bool ConvertGeomPrimitiveTransform(
-		const pxr::UsdPrim& InPrim,
-		const pxr::UsdTimeCode& InTimeCode,
-		FTransform& OutTransform
-	);
-}
+	USDUTILITIES_API bool ConvertGeomPrimitiveTransform(const pxr::UsdPrim& InPrim, const pxr::UsdTimeCode& InTimeCode, FTransform& OutTransform);
+}	 // namespace UsdToUnreal
 
 namespace UnrealToUsd
 {
@@ -203,12 +201,20 @@ namespace UnrealToUsd
 	 * @param StaticMesh - StaticMesh to convert
 	 * @param UsdPrim - Prim to receive the mesh data or LOD variant set
 	 * @param TimeCode - TimeCode to author the attributes with
-	 * @param StageForMaterialAssignments - Stage to use when authoring material assignments (we use this when we want to export the mesh to a payload layer, but the material assignments to an asset layer)
+	 * @param StageForMaterialAssignments - Stage to use when authoring material assignments (we use this when we want to export the mesh to a payload
+	 * layer, but the material assignments to an asset layer)
 	 * @param LowestMeshLOD - Lowest LOD of the UStaticMesh to export (start of the LOD range)
 	 * @param HighestMeshLOD - Lowest LOD of the UStaticMesh to export (end of the LOD range)
 	 * @return Whether the conversion was successful or not.
 	 */
-	USDUTILITIES_API bool ConvertStaticMesh( const UStaticMesh* StaticMesh, pxr::UsdPrim& UsdPrim, pxr::UsdTimeCode TimeCode = pxr::UsdTimeCode::Default(), UE::FUsdStage* StageForMaterialAssignments = nullptr, int32 LowestMeshLOD = 0, int32 HighestMeshLOD = INT32_MAX );
+	USDUTILITIES_API bool ConvertStaticMesh(
+		const UStaticMesh* StaticMesh,
+		pxr::UsdPrim& UsdPrim,
+		pxr::UsdTimeCode TimeCode = pxr::UsdTimeCode::Default(),
+		UE::FUsdStage* StageForMaterialAssignments = nullptr,
+		int32 LowestMeshLOD = 0,
+		int32 HighestMeshLOD = INT32_MAX
+	);
 
 	/**
 	 * Converts an array of mesh descriptions into mesh data, and places that data within the UsdGeomMesh UsdPrim.
@@ -216,18 +222,28 @@ namespace UnrealToUsd
 	 * If more than one MeshDescription are provided, a 'LOD' variant set will be created for UsdPrim, and LOD0, LOD1, etc. variants will be
 	 * created for each provided LOD index. Within each variant, a single Mesh prim also named LOD0, LOD1, etc. will contain the mesh data.
 	 */
-	USDUTILITIES_API bool ConvertMeshDescriptions( const TArray<FMeshDescription>& LODIndexToMeshDescription, pxr::UsdPrim& UsdPrim, const FMatrix& AdditionalTransform, const pxr::UsdTimeCode TimeCode = pxr::UsdTimeCode::Default() );
+	USDUTILITIES_API bool ConvertMeshDescriptions(
+		const TArray<FMeshDescription>& LODIndexToMeshDescription,
+		pxr::UsdPrim& UsdPrim,
+		const FMatrix& AdditionalTransform,
+		const pxr::UsdTimeCode TimeCode = pxr::UsdTimeCode::Default()
+	);
 
 	/**
-	* Extracts animated mesh data from GeometryCache and places the results in UsdPrim.
-	* @param GeometryCache - GeometryCache to convert
-	* @param UsdPrim - Prim to receive the mesh data or LOD variant set
-	* @param StageForMaterialAssignments - Stage to use when authoring material assignments (we use this when we want to export the mesh to a payload layer, but the material assignments to an asset layer)
-	* @return Whether the conversion was successful or not.
-	*/
-	USDUTILITIES_API bool ConvertGeometryCache(const UGeometryCache* GeometryCache, pxr::UsdPrim& UsdPrim, UE::FUsdStage* StageForMaterialAssignments = nullptr);
+	 * Extracts animated mesh data from GeometryCache and places the results in UsdPrim.
+	 * @param GeometryCache - GeometryCache to convert
+	 * @param UsdPrim - Prim to receive the mesh data or LOD variant set
+	 * @param StageForMaterialAssignments - Stage to use when authoring material assignments (we use this when we want to export the mesh to a payload
+	 * layer, but the material assignments to an asset layer)
+	 * @return Whether the conversion was successful or not.
+	 */
+	USDUTILITIES_API bool ConvertGeometryCache(
+		const UGeometryCache* GeometryCache,
+		pxr::UsdPrim& UsdPrim,
+		UE::FUsdStage* StageForMaterialAssignments = nullptr
+	);
 }
-#endif // USE_USD_SDK
+#endif	  // USE_USD_SDK
 
 namespace UsdUtils
 {
@@ -236,10 +252,10 @@ namespace UsdUtils
 	/** Describes what type of material assignment a FUsdPrimMaterialSlot represents */
 	enum class EPrimAssignmentType : uint8
 	{
-		None,			// There is no assignment for this material slot (or no material override)
-		DisplayColor,	// MaterialSource is a serialized FDisplayColorMaterial (e.g. '!DisplayColor_0_0')
-		MaterialPrim,	// MaterialSource is the USD path to a Material prim on the stage (e.g. '/Root/Materials/Red')
-		UnrealMaterial	// MaterialSource is the package path to an UE material (e.g. '/Game/Materials/Red.Red')
+		None,			  // There is no assignment for this material slot (or no material override)
+		DisplayColor,	  // MaterialSource is a serialized FDisplayColorMaterial (e.g. '!DisplayColor_0_0')
+		MaterialPrim,	  // MaterialSource is the USD path to a Material prim on the stage (e.g. '/Root/Materials/Red')
+		UnrealMaterial	  // MaterialSource is the package path to an UE material (e.g. '/Game/Materials/Red.Red')
 	};
 
 	/**
@@ -259,25 +275,18 @@ namespace UsdUtils
 		EPrimAssignmentType AssignmentType = EPrimAssignmentType::None;
 		bool bMeshIsDoubleSided = false;
 
-		friend bool operator==( const FUsdPrimMaterialSlot& Lhs, const FUsdPrimMaterialSlot& Rhs )
+		friend bool operator==(const FUsdPrimMaterialSlot& Lhs, const FUsdPrimMaterialSlot& Rhs)
 		{
-			return Lhs.AssignmentType == Rhs.AssignmentType &&
-				   Lhs.MaterialSource.Equals( Rhs.MaterialSource, ESearchCase::CaseSensitive ) &&
-				   Lhs.bMeshIsDoubleSided == Rhs.bMeshIsDoubleSided;
+			return Lhs.AssignmentType == Rhs.AssignmentType && Lhs.MaterialSource.Equals(Rhs.MaterialSource, ESearchCase::CaseSensitive)
+				   && Lhs.bMeshIsDoubleSided == Rhs.bMeshIsDoubleSided;
 		}
 
-		friend uint32 GetTypeHash( const FUsdPrimMaterialSlot& Slot )
+		friend uint32 GetTypeHash(const FUsdPrimMaterialSlot& Slot)
 		{
-			return HashCombine(
-				Slot.bMeshIsDoubleSided,
-				HashCombine(
-					GetTypeHash( Slot.MaterialSource ),
-					static_cast<uint32>(Slot.AssignmentType)
-				)
-			);
+			return HashCombine(Slot.bMeshIsDoubleSided, HashCombine(GetTypeHash(Slot.MaterialSource), static_cast<uint32>(Slot.AssignmentType)));
 		}
 
-		friend FArchive& operator<<( FArchive& Ar, FUsdPrimMaterialSlot& Slot )
+		friend FArchive& operator<<(FArchive& Ar, FUsdPrimMaterialSlot& Slot)
 		{
 			Ar << Slot.MaterialSource;
 			Ar << Slot.AssignmentType;
@@ -310,20 +319,27 @@ namespace UsdUtils
 	/** Creates a dynamic material instance using the right reference material depending on the given description */
 	PRAGMA_DISABLE_DEPRECATION_WARNINGS
 	UE_DEPRECATED(5.3, "This has moved to the IUsdClassesModule namespace, from USDClassesModule.h")
-	USDUTILITIES_API UMaterialInstanceDynamic* CreateDisplayColorMaterialInstanceDynamic( const UsdUtils::FDisplayColorMaterial& DisplayColorDescription );
+	USDUTILITIES_API UMaterialInstanceDynamic* CreateDisplayColorMaterialInstanceDynamic(
+		const UsdUtils::FDisplayColorMaterial& DisplayColorDescription
+	);
 	UE_DEPRECATED(5.3, "This has moved to the IUsdClassesModule namespace, from USDClassesModule.h")
-	USDUTILITIES_API UMaterialInstanceConstant* CreateDisplayColorMaterialInstanceConstant( const UsdUtils::FDisplayColorMaterial& DisplayColorDescription );
+	USDUTILITIES_API UMaterialInstanceConstant* CreateDisplayColorMaterialInstanceConstant(
+		const UsdUtils::FDisplayColorMaterial& DisplayColorDescription
+	);
 	PRAGMA_ENABLE_DEPRECATION_WARNINGS
 
 	/**
-	 * Extracts all material assignment data from UsdPrim, including material binding, multiple assignment with GeomSubsets, and the unrealMaterial custom USD attribute.
-	 * Guaranteed to return at least one material slot. If UsdPrim is a UsdGeomMesh, it is also guaranteed to have valid material indices (one for every face).
+	 * Extracts all material assignment data from UsdPrim, including material binding, multiple assignment with GeomSubsets, and the unrealMaterial
+	 * custom USD attribute. Guaranteed to return at least one material slot. If UsdPrim is a UsdGeomMesh, it is also guaranteed to have valid
+	 * material indices (one for every face).
 	 * @param UsdPrim - Prim to extract material assignments from
 	 * @param TimeCode - Instant where the material data is sampled
-	 * @param bProvideMaterialIndices - Whether to fill out the material index information for the assignment info (which can be expensive). If this is false, MaterialIndices on the result Struct will have zero values
+	 * @param bProvideMaterialIndices - Whether to fill out the material index information for the assignment info (which can be expensive). If this
+	 * is false, MaterialIndices on the result Struct will have zero values
 	 * @param RenderContext - Which render context to get the materials for. Defaults to universal.
 	 * @param MaterialPurpose - Which material purpose to use when retrieving material bindings
-	 * @return Struct containing an array of material assignments, and a corresponding array of material indices for all polygons of the prim, matching the ordering of the material assignments.
+	 * @return Struct containing an array of material assignments, and a corresponding array of material indices for all polygons of the prim,
+	 * matching the ordering of the material assignments.
 	 */
 	USDUTILITIES_API FUsdPrimMaterialAssignmentInfo GetPrimMaterialAssignments(
 		const pxr::UsdPrim& UsdPrim,
@@ -334,23 +350,29 @@ namespace UsdUtils
 	);
 
 	/**
-	 * Returns an array of prim paths to prims within MaterialPrim's stage that have a UsdShadeMaterialBindingAPI, and whose computed material binding points to MaterialPrim.
-	 * The user prim's schema is not even checked, so this list will naturally include regular UsdGeomMesh prims, UsdGeomSubset prims, and any other with the UsdShadeMaterialBindingAPI.
+	 * Returns an array of prim paths to prims within MaterialPrim's stage that have a UsdShadeMaterialBindingAPI, and whose computed material binding
+	 * points to MaterialPrim. The user prim's schema is not even checked, so this list will naturally include regular UsdGeomMesh prims,
+	 * UsdGeomSubset prims, and any other with the UsdShadeMaterialBindingAPI.
 	 */
 	UE_DEPRECATED(5.3, "No longer used as the FUsdInfoCache now retrieves and retains this information")
-	USDUTILITIES_API TArray<FString> GetMaterialUsers( const UE::FUsdPrim& MaterialPrim, FName MaterialPurpose = *UnrealIdentifiers::MaterialAllPurpose );
+	USDUTILITIES_API TArray<FString> GetMaterialUsers(
+		const UE::FUsdPrim& MaterialPrim,
+		FName MaterialPurpose = *UnrealIdentifiers::MaterialAllPurpose
+	);
 
-	/** Returns whether this prim can be interpreted as describing a static mesh with multiple LODs (i.e. if the prim holds the LOD variant set itself) */
-	USDUTILITIES_API bool DoesPrimContainMeshLODs( const pxr::UsdPrim& Prim );
+	/** Returns whether this prim can be interpreted as describing a static mesh with multiple LODs (i.e. if the prim holds the LOD variant set
+	 * itself) */
+	USDUTILITIES_API bool DoesPrimContainMeshLODs(const pxr::UsdPrim& Prim);
 
 	/** Returns whether this UsdMesh can be interpreted as a LOD of a mesh with multiple LODs (i.e. if the Mesh prim is inside a LOD variant) */
-	USDUTILITIES_API bool IsGeomMeshALOD( const pxr::UsdPrim& UsdMeshPrim );
+	USDUTILITIES_API bool IsGeomMeshALOD(const pxr::UsdPrim& UsdMeshPrim);
 
 	/** Returns whether this UsdPrim is a custom collision mesh */
 	USDUTILITIES_API bool IsCollisionMesh(const pxr::UsdPrim& UsdPrim);
 
-	/** Returns how many LOD variants the Prim has. Note that this will return 0 if called on one of the LOD meshes themselves, it's meant to be called on its parent */
-	USDUTILITIES_API int32 GetNumberOfLODVariants( const pxr::UsdPrim& Prim );
+	/** Returns how many LOD variants the Prim has. Note that this will return 0 if called on one of the LOD meshes themselves, it's meant to be
+	 * called on its parent */
+	USDUTILITIES_API int32 GetNumberOfLODVariants(const pxr::UsdPrim& Prim);
 
 	/**
 	 * If a prim has a variant set named "LOD", with variants named "LOD0", "LOD1", etc., and each has a single Mesh prim, this function
@@ -359,7 +381,7 @@ namespace UsdUtils
 	 * WARNING: There is no guarantee about LOD index ordering! Func may receive LOD2, followed by LOD0, then LOD1, etc.
 	 * WARNING: This will temporarily mutate the stage, and can invalidate references to children of ParentPrim!
 	 */
-	USDUTILITIES_API bool IterateLODMeshes( const pxr::UsdPrim& ParentPrim, TFunction<bool(const pxr::UsdGeomMesh& LODMesh, int32 LODIndex)> Func);
+	USDUTILITIES_API bool IterateLODMeshes(const pxr::UsdPrim& ParentPrim, TFunction<bool(const pxr::UsdGeomMesh& LODMesh, int32 LODIndex)> Func);
 
 	/**
 	 * Traverses `Stage` and authors material binding attributes for all `unrealMaterials` that were baked into USD material assets.
@@ -379,7 +401,7 @@ namespace UsdUtils
 		bool bIsAssetLayer,
 		bool bUsePayload
 	);
-	UE_DEPRECATED( 5.2, "The bRemoveUnrealMaterials parameter is now deprecated as removing the UE material assignments is no longer needed" )
+	UE_DEPRECATED(5.2, "The bRemoveUnrealMaterials parameter is now deprecated as removing the UE material assignments is no longer needed")
 	USDUTILITIES_API void ReplaceUnrealMaterialsWithBaked(
 		const UE::FUsdStage& Stage,
 		const UE::FSdfLayer& LayerToAuthorIn,
@@ -400,7 +422,13 @@ namespace UsdUtils
 	 * Places in OutInstanceTransforms the UE-space instance transforms for a given point instancer prototype index.
 	 * Returns whether the transforms were successfully retrieved or not.
 	 */
-	USDUTILITIES_API bool GetPointInstancerTransforms( const FUsdStageInfo& StageInfo, const pxr::UsdGeomPointInstancer& PointInstancer, const int32 ProtoIndex, pxr::UsdTimeCode EvalTime, TArray<FTransform>& OutInstanceTransforms );
+	USDUTILITIES_API bool GetPointInstancerTransforms(
+		const FUsdStageInfo& StageInfo,
+		const pxr::UsdGeomPointInstancer& PointInstancer,
+		const int32 ProtoIndex,
+		pxr::UsdTimeCode EvalTime,
+		TArray<FTransform>& OutInstanceTransforms
+	);
 
 	/** Returns true if Prim is a GeomMesh with animated attributes */
 	USDUTILITIES_API bool IsAnimatedMesh(const pxr::UsdPrim& Prim);
@@ -420,6 +448,5 @@ namespace UsdUtils
 
 	USDUTILITIES_API void AuthorIdentityTransformGprimAttributes(const pxr::UsdPrim& Gprim, bool bDefaultValues, bool bTimeSampleValues);
 
-#endif // #if USE_USD_SDK
-}
-
+#endif	  // #if USE_USD_SDK
+}	 // namespace UsdUtils
diff --git a/Engine/Plugins/Importers/USDImporter/Source/USDUtilities/Public/USDLayerUtils.h b/Engine/Plugins/Importers/USDImporter/Source/USDUtilities/Public/USDLayerUtils.h
--- a/Engine/Plugins/Importers/USDImporter/Source/USDUtilities/Public/USDLayerUtils.h
+++ b/Engine/Plugins/Importers/USDImporter/Source/USDUtilities/Public/USDLayerUtils.h
@@ -8,7 +8,7 @@
 #include "UsdWrappers/ForwardDeclarations.h"
 
 #include "USDIncludesStart.h"
-	#include "pxr/pxr.h"
+#include "pxr/pxr.h"
 #include "USDIncludesEnd.h"
 
 class SWidget;
@@ -19,10 +19,11 @@ PXR_NAMESPACE_OPEN_SCOPE
 	class UsdPrim;
 	class UsdStage;
 
-	template< typename T > class TfRefPtr;
+	template<typename T>
+	class TfRefPtr;
 
-	using SdfLayerRefPtr = TfRefPtr< SdfLayer >;
-	using UsdStageRefPtr = TfRefPtr< UsdStage >;
+	using SdfLayerRefPtr = TfRefPtr<SdfLayer>;
+	using UsdStageRefPtr = TfRefPtr<UsdStage>;
 PXR_NAMESPACE_CLOSE_SCOPE
 
 namespace UE
@@ -54,16 +55,13 @@ namespace UsdUtils
 		ErrorSubLayerIsParentLayer,
 		ErrorCycleDetected,
 	};
-	USDUTILITIES_API FText ToText( ECanInsertSublayerResult Result );
+	USDUTILITIES_API FText ToText(ECanInsertSublayerResult Result);
 
 	/**
 	 * Checks whether we can add SubLayerIdentifier as a sublayer to ParentLayer.
 	 * The index is actually not relevant.
 	 */
-	USDUTILITIES_API ECanInsertSublayerResult CanInsertSubLayer(
-		const pxr::SdfLayerRefPtr& ParentLayer,
-		const TCHAR* SubLayerIdentifier
-	);
+	USDUTILITIES_API ECanInsertSublayerResult CanInsertSubLayer(const pxr::SdfLayerRefPtr& ParentLayer, const TCHAR* SubLayerIdentifier);
 
 	/**
 	 * Inserts the SubLayerFile path into ParentLayer as a sublayer
@@ -85,64 +83,76 @@ namespace UsdUtils
 
 #if WITH_EDITOR
 	/** Opens a file dialog to open or save a USD file. The returned file path will always be absolute */
-	USDUTILITIES_API TOptional< FString > BrowseUsdFile( EBrowseFileMode Mode );
-#endif // #if WITH_EDITOR
+	USDUTILITIES_API TOptional<FString> BrowseUsdFile(EBrowseFileMode Mode);
+#endif	  // #if WITH_EDITOR
 
 	/**
 	 * Converts the file path from being absolute or relative to engine binary, into being relative to the current project's directory.
-	 * It will only do this if the file is actually within the project's directory (or within its folder tree). Otherwise it will return an absolute path
+	 * It will only do this if the file is actually within the project's directory (or within its folder tree). Otherwise it will return an absolute
+	 * path
 	 */
-	USDUTILITIES_API FString MakePathRelativeToProjectDir( const FString& Path );
+	USDUTILITIES_API FString MakePathRelativeToProjectDir(const FString& Path);
 
 	/** Creates a new layer with a default prim */
-	USDUTILITIES_API TUsdStore< pxr::SdfLayerRefPtr > CreateNewLayer( TUsdStore< pxr::UsdStageRefPtr > UsdStage, const TUsdStore< pxr::SdfLayerRefPtr >& ParentLayer, const TCHAR* LayerFilePath );
+	USDUTILITIES_API TUsdStore<pxr::SdfLayerRefPtr> CreateNewLayer(
+		TUsdStore<pxr::UsdStageRefPtr> UsdStage,
+		const TUsdStore<pxr::SdfLayerRefPtr>& ParentLayer,
+		const TCHAR* LayerFilePath
+	);
 
 	/** Finds which layer introduced the prim in the stage local layer stack */
-	USDUTILITIES_API UE::FSdfLayer FindLayerForPrim( const pxr::UsdPrim& Prim );
+	USDUTILITIES_API UE::FSdfLayer FindLayerForPrim(const pxr::UsdPrim& Prim);
 
 	/** Finds the strongest layer contributing to an attribute */
-	USDUTILITIES_API UE::FSdfLayer FindLayerForAttribute( const pxr::UsdAttribute& Attribute, double TimeCode );
+	USDUTILITIES_API UE::FSdfLayer FindLayerForAttribute(const pxr::UsdAttribute& Attribute, double TimeCode);
 
 	/** Finds the strongest layer that contributes to a set of attributes. Assumes these attributes are on the same stage */
-	USDUTILITIES_API UE::FSdfLayer FindLayerForAttributes( const TArray<UE::FUsdAttribute>& Attributes, double TimeCode, bool bIncludeSessionLayers=true );
+	USDUTILITIES_API UE::FSdfLayer FindLayerForAttributes(
+		const TArray<UE::FUsdAttribute>& Attributes,
+		double TimeCode,
+		bool bIncludeSessionLayers = true
+	);
 
 	/** Finds the layer for a sublayer path of a given root layer */
-	USDUTILITIES_API UE::FSdfLayer FindLayerForSubLayerPath( const UE::FSdfLayer& RootLayer, const FStringView& SubLayerPath );
+	USDUTILITIES_API UE::FSdfLayer FindLayerForSubLayerPath(const UE::FSdfLayer& RootLayer, const FStringView& SubLayerPath);
 
 	/** Sets the layer offset for the strongest reference or payload in this prim composition arcs */
-	USDUTILITIES_API bool SetRefOrPayloadLayerOffset( pxr::UsdPrim& Prim, const UE::FSdfLayerOffset& LayerOffset );
+	USDUTILITIES_API bool SetRefOrPayloadLayerOffset(pxr::UsdPrim& Prim, const UE::FSdfLayerOffset& LayerOffset);
 
 	/** Finds the layer offset that converts the Attribute local times to stage times */
-	USDUTILITIES_API UE::FSdfLayerOffset GetLayerToStageOffset( const pxr::UsdAttribute& Attribute );
+	USDUTILITIES_API UE::FSdfLayerOffset GetLayerToStageOffset(const pxr::UsdAttribute& Attribute);
 
 	/** Finds the full offset for a prim with respect to the composed stage */
-	USDUTILITIES_API UE::FSdfLayerOffset GetPrimToStageOffset( const UE::FUsdPrim& Prim );
+	USDUTILITIES_API UE::FSdfLayerOffset GetPrimToStageOffset(const UE::FUsdPrim& Prim);
 
 	/** Makes sure that the layer start and end timecodes include StartTimeCode and EndTimeCode */
-	USDUTILITIES_API void AddTimeCodeRangeToLayer( const pxr::SdfLayerRefPtr& Layer, double StartTimeCode, double EndTimeCode );
+	USDUTILITIES_API void AddTimeCodeRangeToLayer(const pxr::SdfLayerRefPtr& Layer, double StartTimeCode, double EndTimeCode);
 
 	/** Makes Path relative to the file path of Layer. Conversion happens in-place. */
-	USDUTILITIES_API void MakePathRelativeToLayer( const UE::FSdfLayer& Layer, FString& Path );
+	USDUTILITIES_API void MakePathRelativeToLayer(const UE::FSdfLayer& Layer, FString& Path);
 
-	/** Loads and returns the session sublayer that is used for storing persistent UE state, which can be saved to disk (e.g. metadata for whether an attribute is muted or not) */
-	USDUTILITIES_API UE::FSdfLayer GetUEPersistentStateSublayer( const UE::FUsdStage& Stage, bool bCreateIfNeeded = true );
+	/** Loads and returns the session sublayer that is used for storing persistent UE state, which can be saved to disk (e.g. metadata for whether an
+	 * attribute is muted or not) */
+	USDUTILITIES_API UE::FSdfLayer GetUEPersistentStateSublayer(const UE::FUsdStage& Stage, bool bCreateIfNeeded = true);
 
-	/** Loads and returns the anonymous session sublayer that is used for storing transient UE session state, and won't be saved to disk (e.g. the opinion that actually mutes the attribute) */
-	USDUTILITIES_API UE::FSdfLayer GetUESessionStateSublayer( const UE::FUsdStage& Stage, bool bCreateIfNeeded = true );
+	/** Loads and returns the anonymous session sublayer that is used for storing transient UE session state, and won't be saved to disk (e.g. the
+	 * opinion that actually mutes the attribute) */
+	USDUTILITIES_API UE::FSdfLayer GetUESessionStateSublayer(const UE::FUsdStage& Stage, bool bCreateIfNeeded = true);
 
-	/** Uses FindOrOpen to return the layer with the given identifier if possible. If the identifier is for an anonymous layer, it will search via display name instead */
-	USDUTILITIES_API UE::FSdfLayer FindLayerForIdentifier( const TCHAR* Identifier, const UE::FUsdStage& Stage );
+	/** Uses FindOrOpen to return the layer with the given identifier if possible. If the identifier is for an anonymous layer, it will search via
+	 * display name instead */
+	USDUTILITIES_API UE::FSdfLayer FindLayerForIdentifier(const TCHAR* Identifier, const UE::FUsdStage& Stage);
 
 	/** Returns true if Layer is a session layer within Stage's layer stack */
-	USDUTILITIES_API bool IsSessionLayerWithinStage( const pxr::SdfLayerRefPtr& Layer, const pxr::UsdStageRefPtr& Stage );
+	USDUTILITIES_API bool IsSessionLayerWithinStage(const pxr::SdfLayerRefPtr& Layer, const pxr::UsdStageRefPtr& Stage);
 
 	/**
 	 * Finds all fields that have asset paths as values (e.g. texture references) in LayerToConvert and (if they're relative paths)
 	 * updates them to be absolute with respect to AnchorLayer. LayerToConvert and AnchorLayer can be the same layer, but they don't have to
 	 */
-	USDUTILITIES_API void ConvertAssetRelativePathsToAbsolute( UE::FSdfLayer& LayerToConvert, const UE::FSdfLayer& AnchorLayer );
+	USDUTILITIES_API void ConvertAssetRelativePathsToAbsolute(UE::FSdfLayer& LayerToConvert, const UE::FSdfLayer& AnchorLayer);
 
-	USDUTILITIES_API int32 GetSdfLayerNumFrames( const pxr::SdfLayerRefPtr& Layer );
-}
+	USDUTILITIES_API int32 GetSdfLayerNumFrames(const pxr::SdfLayerRefPtr& Layer);
+}	 // namespace UsdUtils
 
-#endif // #if USE_USD_SDK
+#endif	  // #if USE_USD_SDK
diff --git a/Engine/Plugins/Importers/USDImporter/Source/USDUtilities/Public/USDPrimConversion.h b/Engine/Plugins/Importers/USDImporter/Source/USDUtilities/Public/USDPrimConversion.h
--- a/Engine/Plugins/Importers/USDImporter/Source/USDUtilities/Public/USDPrimConversion.h
+++ b/Engine/Plugins/Importers/USDImporter/Source/USDUtilities/Public/USDPrimConversion.h
@@ -9,7 +9,7 @@
 #include "USDConversionUtils.h"
 
 #include "USDIncludesStart.h"
-	#include "pxr/pxr.h"
+#include "pxr/pxr.h"
 #include "USDIncludesEnd.h"
 
 PXR_NAMESPACE_OPEN_SCOPE
@@ -22,9 +22,10 @@ PXR_NAMESPACE_OPEN_SCOPE
 	class UsdTyped;
 
 	class UsdStage;
-	template< typename T > class TfRefPtr;
+	template<typename T>
+	class TfRefPtr;
 
-	using UsdStageRefPtr = TfRefPtr< UsdStage >;
+	using UsdStageRefPtr = TfRefPtr<UsdStage>;
 PXR_NAMESPACE_CLOSE_SCOPE
 
 class AInstancedFoliageActor;
@@ -51,6 +52,7 @@ struct FMovieSceneSequenceTransform;
 struct FUsdCombinedPrimMetadata;
 struct FUsdPrimMetadata;
 struct FUsdStageInfo;
+
 namespace UE
 {
 	class FUsdAttribute;
@@ -68,7 +70,13 @@ namespace UsdToUnreal
 	 * @param bUsePrimTransform - Whether to convert transform data or not
 	 * @return Whether the conversion was successful or not.
 	 */
-	USDUTILITIES_API bool ConvertXformable( const pxr::UsdStageRefPtr& Stage, const pxr::UsdTyped& Schema, USceneComponent& SceneComponent, double EvalTime, bool bUsePrimTransform = true );
+	USDUTILITIES_API bool ConvertXformable(
+		const pxr::UsdStageRefPtr& Stage,
+		const pxr::UsdTyped& Schema,
+		USceneComponent& SceneComponent,
+		double EvalTime,
+		bool bUsePrimTransform = true
+	);
 
 	/**
 	 * Converts a pxr::UsdGeomXformable's transform attribute values into a UE-space FTransform
@@ -79,7 +87,13 @@ namespace UsdToUnreal
 	 * @param bOutResetTransformStack - Output parameter that, when available, will be set to true if this prim uses the "!resetXformStack!" op
 	 * @return Whether the conversion was successful or not.
 	 */
-	USDUTILITIES_API bool ConvertXformable( const pxr::UsdStageRefPtr& Stage, const pxr::UsdTyped& Schema, FTransform& OutTransform, double EvalTime, bool* bOutResetTransformStack = nullptr );
+	USDUTILITIES_API bool ConvertXformable(
+		const pxr::UsdStageRefPtr& Stage,
+		const pxr::UsdTyped& Schema,
+		FTransform& OutTransform,
+		double EvalTime,
+		bool* bOutResetTransformStack = nullptr
+	);
 
 	/**
 	 * Propagates the transform from Root to Leaf prims inclusively into a UE-space FTransform if they are connected by Xformables
@@ -89,20 +103,60 @@ namespace UsdToUnreal
 	 * @param EvalTime - Time at which to sample the prim's attributes
 	 * @param OutTransform - The composed transform; identity if Leaf is not completely connected to Root by Xformables
 	 */
-	USDUTILITIES_API void PropagateTransform(const pxr::UsdStageRefPtr& Stage, const pxr::UsdPrim& Root, const pxr::UsdPrim& Leaf, double EvalTime, FTransform& OutTransform);
+	USDUTILITIES_API void PropagateTransform(
+		const pxr::UsdStageRefPtr& Stage,
+		const pxr::UsdPrim& Root,
+		const pxr::UsdPrim& Leaf,
+		double EvalTime,
+		FTransform& OutTransform
+	);
 
-	USDUTILITIES_API bool ConvertGeomCamera( const UE::FUsdPrim& Prim, UCineCameraComponent& CameraComponent, double UsdTimeCode = UsdUtils::GetDefaultTimeCode() );
+	USDUTILITIES_API bool ConvertGeomCamera(
+		const UE::FUsdPrim& Prim,
+		UCineCameraComponent& CameraComponent,
+		double UsdTimeCode = UsdUtils::GetDefaultTimeCode()
+	);
 
 	/**
 	 * Calls the ReaderFunc on each time sample of UsdTimeSamples in order to bake values into MovieSceneTrack.
 	 * This is mostly used when reading attributes from USD into tracks for the automatically generated ULevelSequence provided with AUsdStageActors.
 	 * The bounds version must use a BBoxCache: If one is not provided it will be created on-demand for this call alone.
 	 */
-	USDUTILITIES_API bool ConvertBoolTimeSamples( const UE::FUsdStage& UsdStage, const TArray<double>& UsdTimeSamples, const TFunction<bool( double )>& ReaderFunc, UMovieSceneBoolTrack& MovieSceneTrack, const FMovieSceneSequenceTransform& SequenceTransform );
-	USDUTILITIES_API bool ConvertBoolTimeSamples( const UE::FUsdStage& UsdStage, const TArray<double>& UsdTimeSamples, const TFunction<bool(double)>& ReaderFunc, UMovieSceneVisibilityTrack& MovieSceneTrack, const FMovieSceneSequenceTransform& SequenceTransform );
-	USDUTILITIES_API bool ConvertFloatTimeSamples( const UE::FUsdStage& UsdStage, const TArray<double>& UsdTimeSamples, const TFunction<float( double )>& ReaderFunc, UMovieSceneFloatTrack& MovieSceneTrack, const FMovieSceneSequenceTransform& SequenceTransform );
-	USDUTILITIES_API bool ConvertColorTimeSamples( const UE::FUsdStage& UsdStage, const TArray<double>& UsdTimeSamples, const TFunction<FLinearColor( double )>& ReaderFunc, UMovieSceneColorTrack& MovieSceneTrack, const FMovieSceneSequenceTransform& SequenceTransform );
-	USDUTILITIES_API bool ConvertTransformTimeSamples( const UE::FUsdStage& UsdStage, const TArray<double>& UsdTimeSamples, const TFunction<FTransform( double )>& ReaderFunc, UMovieScene3DTransformTrack& MovieSceneTrack, const FMovieSceneSequenceTransform& SequenceTransform );
+	USDUTILITIES_API bool ConvertBoolTimeSamples(
+		const UE::FUsdStage& UsdStage,
+		const TArray<double>& UsdTimeSamples,
+		const TFunction<bool(double)>& ReaderFunc,
+		UMovieSceneBoolTrack& MovieSceneTrack,
+		const FMovieSceneSequenceTransform& SequenceTransform
+	);
+	USDUTILITIES_API bool ConvertBoolTimeSamples(
+		const UE::FUsdStage& UsdStage,
+		const TArray<double>& UsdTimeSamples,
+		const TFunction<bool(double)>& ReaderFunc,
+		UMovieSceneVisibilityTrack& MovieSceneTrack,
+		const FMovieSceneSequenceTransform& SequenceTransform
+	);
+	USDUTILITIES_API bool ConvertFloatTimeSamples(
+		const UE::FUsdStage& UsdStage,
+		const TArray<double>& UsdTimeSamples,
+		const TFunction<float(double)>& ReaderFunc,
+		UMovieSceneFloatTrack& MovieSceneTrack,
+		const FMovieSceneSequenceTransform& SequenceTransform
+	);
+	USDUTILITIES_API bool ConvertColorTimeSamples(
+		const UE::FUsdStage& UsdStage,
+		const TArray<double>& UsdTimeSamples,
+		const TFunction<FLinearColor(double)>& ReaderFunc,
+		UMovieSceneColorTrack& MovieSceneTrack,
+		const FMovieSceneSequenceTransform& SequenceTransform
+	);
+	USDUTILITIES_API bool ConvertTransformTimeSamples(
+		const UE::FUsdStage& UsdStage,
+		const TArray<double>& UsdTimeSamples,
+		const TFunction<FTransform(double)>& ReaderFunc,
+		UMovieScene3DTransformTrack& MovieSceneTrack,
+		const FMovieSceneSequenceTransform& SequenceTransform
+	);
 	USDUTILITIES_API bool ConvertBoundsTimeSamples(
 		const UE::FUsdPrim& InPrim,
 		const TArray<double>& InUsdTimeSamples,
@@ -115,15 +169,15 @@ namespace UsdToUnreal
 	/**
 	 * Struct with lambda functions that can be used to sample an FUsdPrim's attributes at a provided UsdTimeCode and return a converted result.
 	 * The target prims, stage and the actual attributes are captured into the lambda when creating it with CreatePropertyTrackReader.
-	 * This is used so that the reading process can be decoupled from the baking process (ConvertXTimeSamples functions) and reused for different attributes
-	 * c.f. CreatePropertyTrackReader
+	 * This is used so that the reading process can be decoupled from the baking process (ConvertXTimeSamples functions) and reused for different
+	 * attributes c.f. CreatePropertyTrackReader
 	 */
 	struct USDUTILITIES_API FPropertyTrackReader
 	{
-		TFunction<float( double )> FloatReader;
-		TFunction<bool( double )> BoolReader;
-		TFunction<FLinearColor( double )> ColorReader;
-		TFunction<FTransform( double )> TransformReader;
+		TFunction<float(double)> FloatReader;
+		TFunction<bool(double)> BoolReader;
+		TFunction<FLinearColor(double)> ColorReader;
+		TFunction<FTransform(double)> TransformReader;
 	};
 
 	/**
@@ -132,7 +186,8 @@ namespace UsdToUnreal
 	 * the intensity, exposure, width, height of Prim at each UsdTimeCode and returns the corresponding float value for Intensity.
 	 * When we generate a TransformReader, if bIgnorePrimLocalTransform is true it will cause it to ignore that prim's local transform for the reader.
 	 */
-	USDUTILITIES_API FPropertyTrackReader CreatePropertyTrackReader( const UE::FUsdPrim& Prim, const FName& PropertyPath, bool bIgnorePrimLocalTransform = false );
+	USDUTILITIES_API FPropertyTrackReader
+	CreatePropertyTrackReader(const UE::FUsdPrim& Prim, const FName& PropertyPath, bool bIgnorePrimLocalTransform = false);
 
 	/**
 	 * Convert a prim with UsdGeomModelAPI and an alternative drawMode (bounds, cards or origin) into property values of an
@@ -169,19 +224,39 @@ namespace UsdToUnreal
 		bool bInvertFilters = false,
 		bool bCollectFromEntireSubtrees = false
 	);
-}
+}	 // namespace UsdToUnreal
 
 namespace UnrealToUsd
 {
-	USDUTILITIES_API bool ConvertCameraComponent( const UCineCameraComponent& CameraComponent, pxr::UsdPrim& Prim, double UsdTimeCode = UsdUtils::GetDefaultTimeCode() );
+	USDUTILITIES_API bool ConvertCameraComponent(
+		const UCineCameraComponent& CameraComponent,
+		pxr::UsdPrim& Prim,
+		double UsdTimeCode = UsdUtils::GetDefaultTimeCode()
+	);
 
 	/**
-	 * Functions that call WriterFunc on each UsdTimeSample that corresponds to the keyframes of MovieSceneTrack, writing out time samples for attributes of Prim.
-	 * Mostly used to write out to USD the modified tracks from the automatically generated ULevelSequence owned by AUsdStageActors
+	 * Functions that call WriterFunc on each UsdTimeSample that corresponds to the keyframes of MovieSceneTrack, writing out time samples for
+	 * attributes of Prim. Mostly used to write out to USD the modified tracks from the automatically generated ULevelSequence owned by
+	 * AUsdStageActors
 	 */
-	USDUTILITIES_API bool ConvertFloatTrack( const UMovieSceneFloatTrack& MovieSceneTrack, const FMovieSceneSequenceTransform& SequenceTransform, const TFunction<void( float, double )>& WriterFunc, UE::FUsdPrim& Prim );
-	USDUTILITIES_API bool ConvertBoolTrack( const UMovieSceneBoolTrack& MovieSceneTrack, const FMovieSceneSequenceTransform& SequenceTransform, const TFunction<void( bool, double )>& WriterFunc, UE::FUsdPrim& Prim );
-	USDUTILITIES_API bool ConvertColorTrack( const UMovieSceneColorTrack& MovieSceneTrack, const FMovieSceneSequenceTransform& SequenceTransform, const TFunction<void( const FLinearColor&, double )>& WriterFunc, UE::FUsdPrim& Prim );
+	USDUTILITIES_API bool ConvertFloatTrack(
+		const UMovieSceneFloatTrack& MovieSceneTrack,
+		const FMovieSceneSequenceTransform& SequenceTransform,
+		const TFunction<void(float, double)>& WriterFunc,
+		UE::FUsdPrim& Prim
+	);
+	USDUTILITIES_API bool ConvertBoolTrack(
+		const UMovieSceneBoolTrack& MovieSceneTrack,
+		const FMovieSceneSequenceTransform& SequenceTransform,
+		const TFunction<void(bool, double)>& WriterFunc,
+		UE::FUsdPrim& Prim
+	);
+	USDUTILITIES_API bool ConvertColorTrack(
+		const UMovieSceneColorTrack& MovieSceneTrack,
+		const FMovieSceneSequenceTransform& SequenceTransform,
+		const TFunction<void(const FLinearColor&, double)>& WriterFunc,
+		UE::FUsdPrim& Prim
+	);
 	USDUTILITIES_API bool ConvertBoundsVectorTracks(
 		const UMovieSceneDoubleVectorTrack* MinTrack,
 		const UMovieSceneDoubleVectorTrack* MaxTrack,
@@ -189,15 +264,16 @@ namespace UnrealToUsd
 		const TFunction<void(const FVector&, const FVector&, double)>& WriterFunc,
 		UE::FUsdPrim& Prim
 	);
-	USDUTILITIES_API bool Convert3DTransformTrack( const UMovieScene3DTransformTrack& MovieSceneTrack, const FMovieSceneSequenceTransform& SequenceTransform, const TFunction<void( const FTransform&, double )>& WriterFunc, UE::FUsdPrim& Prim );
+	USDUTILITIES_API bool Convert3DTransformTrack(
+		const UMovieScene3DTransformTrack& MovieSceneTrack,
+		const FMovieSceneSequenceTransform& SequenceTransform,
+		const TFunction<void(const FTransform&, double)>& WriterFunc,
+		UE::FUsdPrim& Prim
+	);
 
-	USDUTILITIES_API bool ConvertSceneComponent( const pxr::UsdStageRefPtr& Stage, const USceneComponent* SceneComponent, pxr::UsdPrim& UsdPrim );
+	USDUTILITIES_API bool ConvertSceneComponent(const pxr::UsdStageRefPtr& Stage, const USceneComponent* SceneComponent, pxr::UsdPrim& UsdPrim);
 
-	USDUTILITIES_API bool ConvertMeshComponent(
-		const pxr::UsdStageRefPtr& Stage,
-		const UMeshComponent* MeshComponent,
-		pxr::UsdPrim& UsdPrim
-	);
+	USDUTILITIES_API bool ConvertMeshComponent(const pxr::UsdStageRefPtr& Stage, const UMeshComponent* MeshComponent, pxr::UsdPrim& UsdPrim);
 
 	USDUTILITIES_API bool ConvertHierarchicalInstancedStaticMeshComponent(
 		const UHierarchicalInstancedStaticMeshComponent* HISMComponent,
@@ -223,10 +299,15 @@ namespace UnrealToUsd
 	 * Converts a UMovieScene3DTransformTrack to a UsdGeomXformable
 	 * @param MovieSceneTrack      The track to read the time sampled transform from
 	 * @param Prim                 The Xformable to write to
-	 * @param SequenceTransform    The time transform to apply to the track keys to get them from Usd Stage time to track time (in other words: from main sequence to subsequence)
+	 * @param SequenceTransform    The time transform to apply to the track keys to get them from Usd Stage time to track time (in other words: from
+	 * main sequence to subsequence)
 	 */
-	USDUTILITIES_API bool ConvertXformable( const UMovieScene3DTransformTrack& MovieSceneTrack, pxr::UsdPrim& UsdPrim, const FMovieSceneSequenceTransform& SequenceTransform );
-	USDUTILITIES_API bool ConvertXformable( const FTransform& RelativeTransform, pxr::UsdPrim& UsdPrim, double TimeCode );
+	USDUTILITIES_API bool ConvertXformable(
+		const UMovieScene3DTransformTrack& MovieSceneTrack,
+		pxr::UsdPrim& UsdPrim,
+		const FMovieSceneSequenceTransform& SequenceTransform
+	);
+	USDUTILITIES_API bool ConvertXformable(const FTransform& RelativeTransform, pxr::UsdPrim& UsdPrim, double TimeCode);
 
 	/**
 	 * Converts a AInstancedFoliageActor to a prim containing a pxr::UsdGeomPointInstancer schema. Each foliage type should correspond to a prototype.
@@ -237,7 +318,12 @@ namespace UnrealToUsd
 	 * @param TimeCode			   TimeCode to write the attribute values at. Use UsdUtils::GetDefaultTimeCode() for the Default value.
 	 * @param InstancesLevel	   If this is not nullptr, only foliage instances placed on components that belong to this level will be exported
 	 */
-	USDUTILITIES_API bool ConvertInstancedFoliageActor( const AInstancedFoliageActor& Actor, pxr::UsdPrim& UsdPrim, double TimeCode, ULevel* InstancesLevel = nullptr );
+	USDUTILITIES_API bool ConvertInstancedFoliageActor(
+		const AInstancedFoliageActor& Actor,
+		pxr::UsdPrim& UsdPrim,
+		double TimeCode,
+		ULevel* InstancesLevel = nullptr
+	);
 
 	// We will defer to the common UnrealToUsd::ConvertXComponent functions when baking level sequences, which will already bake
 	// all of the properties of a component out at the same time.
@@ -254,15 +340,15 @@ namespace UnrealToUsd
 		Skeletal = 16,
 		Bounds = 32,
 	};
-	ENUM_CLASS_FLAGS( EBakingType );
+	ENUM_CLASS_FLAGS(EBakingType);
 	const static inline int32 NumBakingTypes = 6;
 
 	// Contains a lambda function responsible for baking a USceneComponent on the level into a prim on an FUsdStage
 	struct USDUTILITIES_API FComponentBaker
 	{
 		EBakingType BakerType;
-		FString ComponentPath;  // Used for sorting, to ensure we evaluate parents and skeletal stuff before children and attached stuff
-		TFunction<void( double UsdTimeCode )> BakerFunction;
+		FString ComponentPath;	  // Used for sorting, to ensure we evaluate parents and skeletal stuff before children and attached stuff
+		TFunction<void(double UsdTimeCode)> BakerFunction;
 	};
 
 	/**
@@ -275,43 +361,60 @@ namespace UnrealToUsd
 	 * EBakingType::Camera baker (that bakes all of the mappable UE cine camera properties at once) and return true.
 	 * If we provide PropertyPath == "FocusSettings.FocusSmoothingInterpSpeed" it will return false as that is not mapped to USD.
 	 */
-	USDUTILITIES_API bool CreateComponentPropertyBaker( UE::FUsdPrim& Prim, const USceneComponent& Component, const FString& PropertyPath, FComponentBaker& OutBaker );
+	USDUTILITIES_API bool CreateComponentPropertyBaker(
+		UE::FUsdPrim& Prim,
+		const USceneComponent& Component,
+		const FString& PropertyPath,
+		FComponentBaker& OutBaker
+	);
 
 	/**
 	 * Creates a property baker responsible for baking the joint/blend shape state of Component as a SkelAnimation,
 	 * returning true if a EBakingType::Skeletal baker was created
 	 */
-	USDUTILITIES_API bool CreateSkeletalAnimationBaker( UE::FUsdPrim& SkeletonPrim, UE::FUsdPrim& SkelAnimation, USkeletalMeshComponent& Component, FComponentBaker& OutBaker );
+	USDUTILITIES_API bool CreateSkeletalAnimationBaker(
+		UE::FUsdPrim& SkeletonPrim,
+		UE::FUsdPrim& SkelAnimation,
+		USkeletalMeshComponent& Component,
+		FComponentBaker& OutBaker
+	);
 
 	/**
-	 * Struct with lambda functions that can be used to convert an UE property's value into the corresponding USD attribute values at the received UsdTimeCode.
-	 * The target prims, stage and the actual attributes are captured into the lambda when creating it with CreatePropertyTrackWriter.
+	 * Struct with lambda functions that can be used to convert an UE property's value into the corresponding USD attribute values at the received
+	 * UsdTimeCode. The target prims, stage and the actual attributes are captured into the lambda when creating it with CreatePropertyTrackWriter.
 	 * This is used so that the writing process can be decoupled from the baking process (ConvertXTrack functions) and reused for different attributes
 	 * c.f. CreatePropertyTrackWriter
 	 */
 	struct USDUTILITIES_API FPropertyTrackWriter
 	{
-		TFunction<void( float, double )> FloatWriter;
-		TFunction<void( bool, double )> BoolWriter;
-		TFunction<void( const FLinearColor&, double )> ColorWriter;
-		TFunction<void( const FVector&, const FVector&, double )> TwoVectorWriter;  // Just used for the bounds tracks, where we use a vector for min and another for max
-		TFunction<void( const FTransform&, double )> TransformWriter;
+		TFunction<void(float, double)> FloatWriter;
+		TFunction<void(bool, double)> BoolWriter;
+		TFunction<void(const FLinearColor&, double)> ColorWriter;
+		TFunction<void(const FVector&, const FVector&, double)> TwoVectorWriter;	// Just used for the bounds tracks, where we use a vector for min
+																					// and another for max
+		TFunction<void(const FTransform&, double)> TransformWriter;
 	};
 
 	/**
-	 * Creates an FPropertyTrackWriter that can be used to receive, convert and output to USD the values of UE properties contained in Track, baking frame-by-frame when needed.
-	 * e.g. calling this with a RectLight prim and UnrealIdentifiers::SourceWidthPropertyName will return a writer with a FloatWriter member that writes
-	 * not only to the width, but also the intensity attributes of Prim at each UsdTimeCode.
-	 * Also outputs OutPropertyPathsToRefresh: In the case above containing "Intensity", which lets us know to refresh and re-read from USD the Intensity track, if we have one
+	 * Creates an FPropertyTrackWriter that can be used to receive, convert and output to USD the values of UE properties contained in Track, baking
+	 * frame-by-frame when needed. e.g. calling this with a RectLight prim and UnrealIdentifiers::SourceWidthPropertyName will return a writer with a
+	 * FloatWriter member that writes not only to the width, but also the intensity attributes of Prim at each UsdTimeCode. Also outputs
+	 * OutPropertyPathsToRefresh: In the case above containing "Intensity", which lets us know to refresh and re-read from USD the Intensity track, if
+	 * we have one
 	 */
-	USDUTILITIES_API FPropertyTrackWriter CreatePropertyTrackWriter( const USceneComponent& Component, const UMovieScenePropertyTrack& Track, UE::FUsdPrim& Prim, TSet<FName>& OutPropertyPathsToRefresh );
+	USDUTILITIES_API FPropertyTrackWriter CreatePropertyTrackWriter(
+		const USceneComponent& Component,
+		const UMovieScenePropertyTrack& Track,
+		UE::FUsdPrim& Prim,
+		TSet<FName>& OutPropertyPathsToRefresh
+	);
 
 	/**
 	 * Returns the attributes that in USD correspond to a property in UE.
 	 * The first one is the "main" attribute, when appropriate. e.g. for a RectLight prim and UnrealIdentifiers::IntensityPropertyName,
 	 * we'll receive intensity, exposure, width and height, in that order.
 	 */
-	USDUTILITIES_API TArray<UE::FUsdAttribute> GetAttributesForProperty( const UE::FUsdPrim& Prim, const FName& PropertyPath );
+	USDUTILITIES_API TArray<UE::FUsdAttribute> GetAttributesForProperty(const UE::FUsdPrim& Prim, const FName& PropertyPath);
 
 	/**
 	 * Converts the properties values from UUsdDrawModeComponent into attribute values of a prim with UsdGeomModelAPI.
@@ -363,6 +466,6 @@ namespace UnrealToUsd
 		const TArray<FString>& BlockedPrefixFilters = {},
 		bool bInvertFilter = false
 	);
-}
+}	 // namespace UnrealToUsd
 
-#endif // #if USE_USD_SDK
+#endif	  // #if USE_USD_SDK
diff --git a/Engine/Plugins/Importers/USDImporter/Source/USDUtilities/Public/USDShadeConversion.h b/Engine/Plugins/Importers/USDImporter/Source/USDUtilities/Public/USDShadeConversion.h
--- a/Engine/Plugins/Importers/USDImporter/Source/USDUtilities/Public/USDShadeConversion.h
+++ b/Engine/Plugins/Importers/USDImporter/Source/USDUtilities/Public/USDShadeConversion.h
@@ -13,9 +13,9 @@
 #include "UsdWrappers/SdfLayer.h"
 
 #include "USDIncludesStart.h"
-	#include "pxr/base/tf/token.h"
-	#include "pxr/usd/usd/timeCode.h"
-	#include "pxr/usd/usdShade/tokens.h"
+#include "pxr/base/tf/token.h"
+#include "pxr/usd/usd/timeCode.h"
+#include "pxr/usd/usdShade/tokens.h"
 #include "USDIncludesEnd.h"
 
 PXR_NAMESPACE_OPEN_SCOPE
@@ -37,6 +37,7 @@ enum TextureAddress : int;
 enum TextureCompressionSettings : int;
 struct FFlattenMaterial;
 struct FPropertyEntry;
+
 namespace UE
 {
 	class FUsdPrim;
@@ -45,12 +46,13 @@ namespace UE
 namespace UsdToUnreal
 {
 	/**
-	 * Extracts material data from UsdShadeMaterial and places the results in Material. Note that since this is used for UMaterialInstanceDynamics at runtime as well,
-	 * it will not set base property overrides (e.g. BlendMode) or the parent material, and will just assume that the caller handles that.
-	 * Note that in order to receive the primvar to UV index mapping calculated within this function, the provided Material should have an UUsdMaterialAssetImportData
-	 * object as its AssetImportData.
+	 * Extracts material data from UsdShadeMaterial and places the results in Material. Note that since this is used for UMaterialInstanceDynamics at
+	 * runtime as well, it will not set base property overrides (e.g. BlendMode) or the parent material, and will just assume that the caller handles
+	 * that. Note that in order to receive the primvar to UV index mapping calculated within this function, the provided Material should have an
+	 * UUsdMaterialAssetImportData object as its AssetImportData.
 	 * @param UsdShadeMaterial - Shade material with the data to convert
-	 * @param Material - Output parameter that will be filled with the converted data. Only the versions that receive a dynamic material instance will work at runtime
+	 * @param Material - Output parameter that will be filled with the converted data. Only the versions that receive a dynamic material instance will
+	 * work at runtime
 	 * @param TexturesCache - Cache to prevent importing a texture more than once
 	 * @param RenderContext - Which render context output to read from the UsdShadeMaterial
 	 * @param ReuseIdenticalAssets - Whether to reuse identical textures found in the TextureCache or to create dedicated textures for each material
@@ -72,7 +74,8 @@ namespace UsdToUnreal
 	);
 
 	/**
-	 * Attemps to assign the values of the surface shader inputs to the MaterialInstance parameters by matching the inputs display names to the parameters names.
+	 * Attemps to assign the values of the surface shader inputs to the MaterialInstance parameters by matching the inputs display names to the
+	 * parameters names.
 	 * @param UsdShadeMaterial - Shade material with the data to convert
 	 * @param MaterialInstance - Material instance on which we will set the parameter values
 	 * @param TexturesCache - Cache to prevent importing a texture more than once
@@ -89,18 +92,47 @@ namespace UsdToUnreal
 		bool bReuseIdenticalAssets = true
 	);
 
-PRAGMA_DISABLE_DEPRECATION_WARNINGS
+	PRAGMA_DISABLE_DEPRECATION_WARNINGS
 	UE_DEPRECATED(5.3, "Please use the overload that doesn't use PrimvarToUVIndex anymore")
-	USDUTILITIES_API bool ConvertMaterial( const pxr::UsdShadeMaterial& UsdShadeMaterial, UMaterial& Material, UUsdAssetCache2* TexturesCache, TMap< FString, int32 >& PrimvarToUVIndex, const TCHAR* RenderContext = nullptr );
+	USDUTILITIES_API bool ConvertMaterial(
+		const pxr::UsdShadeMaterial& UsdShadeMaterial,
+		UMaterial& Material,
+		UUsdAssetCache2* TexturesCache,
+		TMap<FString, int32>& PrimvarToUVIndex,
+		const TCHAR* RenderContext = nullptr
+	);
 	UE_DEPRECATED(5.3, "Use the other overload that receives an UUsdAssetCache2 object instead")
-	USDUTILITIES_API bool ConvertMaterial( const pxr::UsdShadeMaterial& UsdShadeMaterial, UMaterialInstance& Material, UUsdAssetCache2* TexturesCache, TMap< FString, int32 >& PrimvarToUVIndex, const TCHAR* RenderContext = nullptr );
+	USDUTILITIES_API bool ConvertMaterial(
+		const pxr::UsdShadeMaterial& UsdShadeMaterial,
+		UMaterialInstance& Material,
+		UUsdAssetCache2* TexturesCache,
+		TMap<FString, int32>& PrimvarToUVIndex,
+		const TCHAR* RenderContext = nullptr
+	);
 	UE_DEPRECATED(5.2, "Use the other overload that receives an UUsdAssetCache2 object instead")
-	USDUTILITIES_API bool ConvertMaterial( const pxr::UsdShadeMaterial& UsdShadeMaterial, UMaterialInstance& Material, UUsdAssetCache* TexturesCache, TMap< FString, int32 >& PrimvarToUVIndex, const TCHAR* RenderContext = nullptr );
+	USDUTILITIES_API bool ConvertMaterial(
+		const pxr::UsdShadeMaterial& UsdShadeMaterial,
+		UMaterialInstance& Material,
+		UUsdAssetCache* TexturesCache,
+		TMap<FString, int32>& PrimvarToUVIndex,
+		const TCHAR* RenderContext = nullptr
+	);
 	UE_DEPRECATED(5.2, "Use the other overload that receives an UUsdAssetCache2 object instead")
-	USDUTILITIES_API bool ConvertMaterial( const pxr::UsdShadeMaterial& UsdShadeMaterial, UMaterial& Material, UUsdAssetCache* TexturesCache, TMap< FString, int32 >& PrimvarToUVIndex, const TCHAR* RenderContext = nullptr );
+	USDUTILITIES_API bool ConvertMaterial(
+		const pxr::UsdShadeMaterial& UsdShadeMaterial,
+		UMaterial& Material,
+		UUsdAssetCache* TexturesCache,
+		TMap<FString, int32>& PrimvarToUVIndex,
+		const TCHAR* RenderContext = nullptr
+	);
 	UE_DEPRECATED(5.2, "Use the other overload that receives an UUsdAssetCache2 object instead")
-	USDUTILITIES_API bool ConvertShadeInputsToParameters( const pxr::UsdShadeMaterial& UsdShadeMaterial, UMaterialInstance& MaterialInstance, UUsdAssetCache* TexturesCache, const TCHAR* RenderContext = nullptr );
-PRAGMA_ENABLE_DEPRECATION_WARNINGS
+	USDUTILITIES_API bool ConvertShadeInputsToParameters(
+		const pxr::UsdShadeMaterial& UsdShadeMaterial,
+		UMaterialInstance& MaterialInstance,
+		UUsdAssetCache* TexturesCache,
+		const TCHAR* RenderContext = nullptr
+	);
+	PRAGMA_ENABLE_DEPRECATION_WARNINGS
 }
 
 #if WITH_EDITOR
@@ -128,7 +160,8 @@ namespace UnrealToUsd
 
 	/**
 	 * Converts a flattened material's data into textures placed at InTexturesDir, and configures OutUsdShadeMaterial to use the baked textures.
-	 * Note that to avoid a potentially useless copy, InMaterial's samples will be modified in place to have 255 alpha before being exported to textures.
+	 * Note that to avoid a potentially useless copy, InMaterial's samples will be modified in place to have 255 alpha before being exported to
+	 * textures.
 	 *
 	 * @param MaterialName - Name of the material, used as prefix on the exported texture filenames
 	 * @param InMaterial - Source material to bake
@@ -137,14 +170,21 @@ namespace UnrealToUsd
 	 * @param OutUsdShadeMaterialPrim - UsdPrim with the UsdShadeMaterial schema that will be configured to use the baked textures and constants
 	 * @return Whether the conversion was successful or not.
 	 */
-	USDUTILITIES_API bool ConvertFlattenMaterial( const FString& InMaterialName, FFlattenMaterial& InMaterial, const TArray<FPropertyEntry>& InMaterialProperties, const FDirectoryPath& InTexturesDir, UE::FUsdPrim& OutUsdShadeMaterialPrim );
+	USDUTILITIES_API bool ConvertFlattenMaterial(
+		const FString& InMaterialName,
+		FFlattenMaterial& InMaterial,
+		const TArray<FPropertyEntry>& InMaterialProperties,
+		const FDirectoryPath& InTexturesDir,
+		UE::FUsdPrim& OutUsdShadeMaterialPrim
+	);
 }
-#endif // WITH_EDITOR
+#endif	  // WITH_EDITOR
 
 namespace UsdUtils
 {
-	/** Returns a path to an UE asset (e.g. "/Game/Assets/Red.Red") if MaterialPrim has an 'unreal' render context surface output that points at one */
-	USDUTILITIES_API TOptional<FString> GetUnrealSurfaceOutput( const pxr::UsdPrim& MaterialPrim );
+	/** Returns a path to an UE asset (e.g. "/Game/Assets/Red.Red") if MaterialPrim has an 'unreal' render context surface output that points at one
+	 */
+	USDUTILITIES_API TOptional<FString> GetUnrealSurfaceOutput(const pxr::UsdPrim& MaterialPrim);
 
 	/**
 	 * Sets which UE material asset the 'unreal' render context surface output of MaterialPrim is pointing at (creating the surface output
@@ -154,7 +194,7 @@ namespace UsdUtils
 	 * @param UnrealMaterialPathName - Path to an UE UMaterialInterface asset (e.g. "/Game/Assets/Red.Red")
 	 * @return Whether we successfully set the surface output or not
 	 */
-	USDUTILITIES_API bool SetUnrealSurfaceOutput( pxr::UsdPrim& MaterialPrim, const FString& UnrealMaterialPathName );
+	USDUTILITIES_API bool SetUnrealSurfaceOutput(pxr::UsdPrim& MaterialPrim, const FString& UnrealMaterialPathName);
 
 	/**
 	 * Clears any opinions for the 'unreal' render context surface output of MaterialPrim within LayerToAuthorIn.
@@ -164,24 +204,25 @@ namespace UsdUtils
 	 * @param LayerToAuthorIn - Layer to clear the opinions in, or an invalid layer (e.g. UE::FSdfLayer{}, which is the default)
 	 * @return Whether we successfully cleared the opinions or not
 	 */
-	UE_DEPRECATED( 5.2, "No longer used as UE material assignments are only visible in the 'unreal' render context anyway" )
-	USDUTILITIES_API bool RemoveUnrealSurfaceOutput( pxr::UsdPrim& MaterialPrim, const UE::FSdfLayer& LayerToAuthorIn = UE::FSdfLayer{} );
+	UE_DEPRECATED(5.2, "No longer used as UE material assignments are only visible in the 'unreal' render context anyway")
+	USDUTILITIES_API bool RemoveUnrealSurfaceOutput(pxr::UsdPrim& MaterialPrim, const UE::FSdfLayer& LayerToAuthorIn = UE::FSdfLayer{});
 
 	/**
 	 * Returns whether the material needs to be rendered with the Translucent rendering mode.
 	 * This function exists because we need this information *before* we pick the right parent for a material instance and properly convert it.
 	 */
-	USDUTILITIES_API bool IsMaterialTranslucent( const pxr::UsdShadeMaterial& UsdShadeMaterial );
+	USDUTILITIES_API bool IsMaterialTranslucent(const pxr::UsdShadeMaterial& UsdShadeMaterial);
 
-	USDUTILITIES_API FSHAHash HashShadeMaterial( const pxr::UsdShadeMaterial& UsdShadeMaterial, const pxr::TfToken& RenderContext = pxr::UsdShadeTokens->universalRenderContext );
+	USDUTILITIES_API FSHAHash
+	HashShadeMaterial(const pxr::UsdShadeMaterial& UsdShadeMaterial, const pxr::TfToken& RenderContext = pxr::UsdShadeTokens->universalRenderContext);
 
 	/** Returns the resolved path from a pxr::SdfAssetPath attribute. For UDIMs path, returns the path to the 1001 tile. */
-	USDUTILITIES_API FString GetResolvedTexturePath( const pxr::UsdAttribute& TextureAssetPathAttr );
+	USDUTILITIES_API FString GetResolvedTexturePath(const pxr::UsdAttribute& TextureAssetPathAttr);
 
 	/**
 	 * Computes and returns the hash string for the texture at the given path.
 	 * Handles regular texture asset paths as well as asset paths identifying textures inside Usdz archives.
-     * Returns an empty string if the texture could not be hashed.
+	 * Returns an empty string if the texture could not be hashed.
 	 */
 	USDUTILITIES_API FString GetTextureHash(
 		const FString& ResolvedTexturePath,
@@ -191,31 +232,37 @@ namespace UsdUtils
 		TextureAddress AddressY
 	);
 
-	/** Creates a texture from a pxr::SdfAssetPath attribute. PrimPath is optional, and should point to the source shadematerial prim path. It will be placed in its UUsdAssetUserData */
-	USDUTILITIES_API UTexture* CreateTexture( const pxr::UsdAttribute& TextureAssetPathAttr, const FString& PrimPath = FString(), TextureGroup LODGroup = TEXTUREGROUP_World, UObject* Outer = GetTransientPackage() );
+	/** Creates a texture from a pxr::SdfAssetPath attribute. PrimPath is optional, and should point to the source shadematerial prim path. It will be
+	 * placed in its UUsdAssetUserData */
+	USDUTILITIES_API UTexture* CreateTexture(
+		const pxr::UsdAttribute& TextureAssetPathAttr,
+		const FString& PrimPath = FString(),
+		TextureGroup LODGroup = TEXTUREGROUP_World,
+		UObject* Outer = GetTransientPackage()
+	);
 
 	/** Checks if this texture needs virtual textures and emits a warning if it is disabled for the project */
-	USDUTILITIES_API void NotifyIfVirtualTexturesNeeded( UTexture* Texture );
+	USDUTILITIES_API void NotifyIfVirtualTexturesNeeded(UTexture* Texture);
 
 #if WITH_EDITOR
 	/** Convert between the two different types used to represent material channels to bake */
-	USDUTILITIES_API EFlattenMaterialProperties MaterialPropertyToFlattenProperty( EMaterialProperty MaterialProperty );
+	USDUTILITIES_API EFlattenMaterialProperties MaterialPropertyToFlattenProperty(EMaterialProperty MaterialProperty);
 
 	/** Convert between the two different types used to represent material channels to bake */
-	USDUTILITIES_API EMaterialProperty FlattenPropertyToMaterialProperty( EFlattenMaterialProperties FlattenProperty );
-#endif // WITH_EDITOR
+	USDUTILITIES_API EMaterialProperty FlattenPropertyToMaterialProperty(EFlattenMaterialProperties FlattenProperty);
+#endif	  // WITH_EDITOR
 
 	/** Converts channels that have the same value for every pixel into a channel that only has a single pixel with that value */
-	USDUTILITIES_API void CollapseConstantChannelsToSinglePixel( FFlattenMaterial& InMaterial );
+	USDUTILITIES_API void CollapseConstantChannelsToSinglePixel(FFlattenMaterial& InMaterial);
 
 	/** Temporary function until UnrealWrappers can create attributes, just adds a custom bool attribute 'worldSpaceNormals' as true */
-	USDUTILITIES_API bool MarkMaterialPrimWithWorldSpaceNormals( const UE::FUsdPrim& MaterialPrim );
+	USDUTILITIES_API bool MarkMaterialPrimWithWorldSpaceNormals(const UE::FUsdPrim& MaterialPrim);
 
 	/** Sets material instance parameters whether Material is a MaterialInstanceConstant or a MaterialInstanceDynamic */
 	USDUTILITIES_API void SetScalarParameterValue(UMaterialInstance& Material, const TCHAR* ParameterName, float ParameterValue);
 	USDUTILITIES_API void SetVectorParameterValue(UMaterialInstance& Material, const TCHAR* ParameterName, FLinearColor ParameterValue);
 	USDUTILITIES_API void SetTextureParameterValue(UMaterialInstance& Material, const TCHAR* ParameterName, UTexture* ParameterValue);
 	USDUTILITIES_API void SetBoolParameterValue(UMaterialInstance& Material, const TCHAR* ParameterName, bool bParameterValue);
-}
+}	 // namespace UsdUtils
 
-#endif // #if USE_USD_SDK
+#endif	  // #if USE_USD_SDK
diff --git a/Engine/Plugins/Importers/USDImporter/Source/USDUtilities/Public/USDSkeletalDataConversion.h b/Engine/Plugins/Importers/USDImporter/Source/USDUtilities/Public/USDSkeletalDataConversion.h
--- a/Engine/Plugins/Importers/USDImporter/Source/USDUtilities/Public/USDSkeletalDataConversion.h
+++ b/Engine/Plugins/Importers/USDImporter/Source/USDUtilities/Public/USDSkeletalDataConversion.h
@@ -13,9 +13,9 @@
 
 #if USE_USD_SDK
 #include "USDIncludesStart.h"
-	#include "pxr/pxr.h"
-	#include "pxr/usd/usd/timeCode.h"
-	#include "pxr/usd/usdShade/tokens.h"
+#include "pxr/pxr.h"
+#include "pxr/usd/usd/timeCode.h"
+#include "pxr/usd/usdShade/tokens.h"
 #include "USDIncludesEnd.h"
 
 PXR_NAMESPACE_OPEN_SCOPE
@@ -28,7 +28,8 @@ PXR_NAMESPACE_OPEN_SCOPE
 	class UsdSkelSkeleton;
 	class UsdSkelSkeletonQuery;
 	class UsdSkelSkinningQuery;
-	template<typename T> class VtArray;
+	template<typename T>
+	class VtArray;
 PXR_NAMESPACE_CLOSE_SCOPE
 
 class FSkeletalMeshImportData;
@@ -39,18 +40,20 @@ class UMovieSceneControlRigParameterSection;
 class USkeletalMesh;
 struct FUsdStageInfo;
 struct FMovieSceneSequenceTransform;
+
 namespace SkeletalMeshImportData
 {
 	struct FBone;
 	struct FMaterial;
 }
+
 namespace UsdUtils
 {
 	struct FUsdPrimMaterialSlot;
 	struct FUsdPrimMaterialAssignmentInfo;
 }
 
-#endif // #if USE_USD_SDK
+#endif	  // #if USE_USD_SDK
 
 namespace UsdUtils
 {
@@ -61,7 +64,7 @@ namespace UsdUtils
 
 		float InbetweenWeight;
 
-		friend FArchive& operator<<( FArchive& Ar, FUsdBlendShapeInbetween& BlendShape )
+		friend FArchive& operator<<(FArchive& Ar, FUsdBlendShapeInbetween& BlendShape)
 		{
 			Ar << BlendShape.Name;
 			Ar << BlendShape.InbetweenWeight;
@@ -81,9 +84,12 @@ namespace UsdUtils
 		TArray<FUsdBlendShapeInbetween> Inbetweens;
 		bool bHasAuthoredTangents = false;
 
-		bool IsValid() const { return Vertices.Num() > 0; }
+		bool IsValid() const
+		{
+			return Vertices.Num() > 0;
+		}
 
-		friend FArchive& operator<<( FArchive& Ar, FUsdBlendShape& BlendShape )
+		friend FArchive& operator<<(FArchive& Ar, FUsdBlendShape& BlendShape)
 		{
 			Ar << BlendShape.Name;
 			Ar << BlendShape.Inbetweens;
@@ -96,27 +102,35 @@ namespace UsdUtils
 	 * We need this to be case sensitive because USD paths are, and even the sample HumanFemale Skel scene has
 	 * paths that differ only by case (e.g. 'JawUD' and 'JAWUD' blend shapes)
 	 */
-	using FBlendShapeMap = TMap<FString, FUsdBlendShape, FDefaultSetAllocator, FCaseSensitiveStringMapFuncs< FUsdBlendShape > >;
+	using FBlendShapeMap = TMap<FString, FUsdBlendShape, FDefaultSetAllocator, FCaseSensitiveStringMapFuncs<FUsdBlendShape>>;
 
 	/**
-	 * We decompose inbetween blend shapes on import into separate morph targets, alongside the primary blend shape, which also becomes a morph target.
-	 * This function returns the morph target weights for the primary morph target, as well as all the morph targets for inbetween shapes of 'InBlendShape', given an
-	 * initial input weight value for the USD blend shape. Calculations are done following the equations at
+	 * We decompose inbetween blend shapes on import into separate morph targets, alongside the primary blend shape, which also becomes a morph
+	 * target. This function returns the morph target weights for the primary morph target, as well as all the morph targets for inbetween shapes of
+	 * 'InBlendShape', given an initial input weight value for the USD blend shape. Calculations are done following the equations at
 	 * https://graphics.pixar.com/usd/docs/api/_usd_skel__schemas.html#UsdSkel_BlendShape_Inbetweens
 	 *
 	 * Note: This assumes that the morph targets in InBlendShape are sorted by weight.
 	 */
-	USDUTILITIES_API void ResolveWeightsForBlendShape( const FUsdBlendShape& InBlendShape, float InWeight, float& OutPrimaryWeight, TArray<float>& OutInbetweenWeights );
+	USDUTILITIES_API void ResolveWeightsForBlendShape(
+		const FUsdBlendShape& InBlendShape,
+		float InWeight,
+		float& OutPrimaryWeight,
+		TArray<float>& OutInbetweenWeights
+	);
 
 #if USE_USD_SDK
 	/** Allows creation of a skinning query from the underlying skinned mesh and skeleton. Adapted from the USD SDK implementation */
-	USDUTILITIES_API pxr::UsdSkelSkinningQuery CreateSkinningQuery( const pxr::UsdGeomMesh& SkinnedMesh, const pxr::UsdSkelSkeletonQuery& SkeletonQuery );
+	USDUTILITIES_API pxr::UsdSkelSkinningQuery CreateSkinningQuery(
+		const pxr::UsdGeomMesh& SkinnedMesh,
+		const pxr::UsdSkelSkeletonQuery& SkeletonQuery
+	);
 
 	/**
 	 * Sets prim AnimationSource as the animation source for prim Prim.
 	 * Applies the SkelBindingAPI to Prim. See pxr::SkelBindingAPI::GetAnimationSourceRel.
 	 */
-	USDUTILITIES_API void BindAnimationSource( pxr::UsdPrim& Prim, const pxr::UsdPrim& AnimationSource );
+	USDUTILITIES_API void BindAnimationSource(pxr::UsdPrim& Prim, const pxr::UsdPrim& AnimationSource);
 
 	/**
 	 * Returns the SkelAnimation prim that is resolved for the first skeletal binding of SkelRootPrim, if it is a SkelRoot.
@@ -147,8 +161,8 @@ namespace UsdUtils
 		pxr::UsdSkelSkeletonQuery& OutSkeletonQuery,
 		pxr::UsdSkelCache* InOutSkelCache = nullptr
 	);
-#endif // USE_USD_SDK
-}
+#endif	  // USE_USD_SDK
+}	 // namespace UsdUtils
 
 #if USE_USD_SDK && WITH_EDITOR
 namespace UsdToUnreal
@@ -159,21 +173,37 @@ namespace UsdToUnreal
 	 * @param SkelMeshImportData - Output parameter that will be filled with the converted data
 	 * @return Whether the conversion was successful or not.
 	 */
-	USDUTILITIES_API bool ConvertSkeleton( const pxr::UsdSkelSkeletonQuery& UsdSkeletonQuery, FSkeletalMeshImportData& SkelMeshImportData );
+	USDUTILITIES_API bool ConvertSkeleton(const pxr::UsdSkelSkeletonQuery& UsdSkeletonQuery, FSkeletalMeshImportData& SkelMeshImportData);
 
 	/**
 	 * Converts an USD blend shape into zero, one or more FUsdBlendShapes, and places them in OutBlendShapes
 	 * @param UsdBlendShape - Source USD object with the blend shape data
 	 * @param StageInfo - Details about the stage, required to do the coordinate conversion from the USD blend shape to morph target data
 	 * @param LODIndex - LODIndex of the SkeletalMesh that will use the imported FUsdBlendShape
-	 * @param AdditionalTransform - Additional affine transform to apply to the blend shape deltas and tangents (Note: No longer used, even in the deprecated overload of this function)
-	 * @param PointIndexOffset - Index into the corresponding SkelMeshImportData.Points where the points corresponding to the UsdBlendShape's prim start
-	 * @param UsedMorphTargetNames - Names that the newly created FUsdBlendShapes cannot have (this function will also add the names of the newly created FUsdBlendShapes to it)
+	 * @param AdditionalTransform - Additional affine transform to apply to the blend shape deltas and tangents (Note: No longer used, even in the
+	 * deprecated overload of this function)
+	 * @param PointIndexOffset - Index into the corresponding SkelMeshImportData.Points where the points corresponding to the UsdBlendShape's prim
+	 * start
+	 * @param UsedMorphTargetNames - Names that the newly created FUsdBlendShapes cannot have (this function will also add the names of the newly
+	 * created FUsdBlendShapes to it)
 	 * @param OutBlendShapes - Where the newly created blend shapes will be added to
 	 * @return Whether the conversion was successful or not.
 	 */
-	USDUTILITIES_API bool ConvertBlendShape( const pxr::UsdSkelBlendShape& UsdBlendShape, const FUsdStageInfo& StageInfo, uint32 PointIndexOffset, TSet<FString>& UsedMorphTargetNames, UsdUtils::FBlendShapeMap& OutBlendShapes );
-	USDUTILITIES_API bool ConvertBlendShape( const pxr::UsdSkelBlendShape& UsdBlendShape, const FUsdStageInfo& StageInfo, int32 LODIndex, uint32 PointIndexOffset, TSet<FString>& UsedMorphTargetNames, UsdUtils::FBlendShapeMap& OutBlendShapes );
+	USDUTILITIES_API bool ConvertBlendShape(
+		const pxr::UsdSkelBlendShape& UsdBlendShape,
+		const FUsdStageInfo& StageInfo,
+		uint32 PointIndexOffset,
+		TSet<FString>& UsedMorphTargetNames,
+		UsdUtils::FBlendShapeMap& OutBlendShapes
+	);
+	USDUTILITIES_API bool ConvertBlendShape(
+		const pxr::UsdSkelBlendShape& UsdBlendShape,
+		const FUsdStageInfo& StageInfo,
+		int32 LODIndex,
+		uint32 PointIndexOffset,
+		TSet<FString>& UsedMorphTargetNames,
+		UsdUtils::FBlendShapeMap& OutBlendShapes
+	);
 
 	/**
 	 * Extracts skeletal mesh data fro UsdSkinningQuery, and places the results in SkelMeshImportData.
@@ -184,10 +214,13 @@ namespace UsdToUnreal
 	 *
 	 * @param UsdSkinningQuery - SkinningQuery with the data to convert
 	 * @param SkeletonQuery - Query object for the skeleton this mesh is bound to
-	 * @param AdditionalTransform - Additional transform to apply to the vertices of the mesh (Note: No longer used, even in the deprecated overload of this function)
+	 * @param AdditionalTransform - Additional transform to apply to the vertices of the mesh (Note: No longer used, even in the deprecated overload
+	 * of this function)
 	 * @param SkelMeshImportData - Output parameter that will be filled with the converted data
 	 * @param MaterialAssignments - Output parameter that will be filled with the material assignment data extracted from UsdSkinningQuery
-	 * @param MaterialToPrimvarsUVSetNames - Maps from a material prim path, to pairs indicating which primvar names are used as 'st' coordinates for this mesh, and which UVIndex materials will sample from (e.g. ["st0", 0], ["myUvSet2", 2], etc). This is used to pick which primvars will become UV sets.
+	 * @param MaterialToPrimvarsUVSetNames - Maps from a material prim path, to pairs indicating which primvar names are used as 'st' coordinates for
+	 * this mesh, and which UVIndex materials will sample from (e.g. ["st0", 0], ["myUvSet2", 2], etc). This is used to pick which primvars will
+	 * become UV sets.
 	 * @param RenderContext - Render context to use when parsing the skinned mesh's materials (e.g. '' for universal, or 'mdl', or 'unreal', etc.)
 	 * @param MaterialPurpose - Material purpose to use when parsing the skinned Mesh prim's material bindings
 	 * @return Whether the conversion was successful or not.
@@ -199,13 +232,13 @@ namespace UsdToUnreal
 		UsdUtils::FUsdPrimMaterialAssignmentInfo& MaterialAssignments,
 		const FUsdMeshConversionOptions& CommonOptions = FUsdMeshConversionOptions::DefaultOptions
 	);
-	UE_DEPRECATED( 5.3, "Please use the overload with an FUsdPrimMaterialAssignmentInfo parameter." )
+	UE_DEPRECATED(5.3, "Please use the overload with an FUsdPrimMaterialAssignmentInfo parameter.")
 	USDUTILITIES_API bool ConvertSkinnedMesh(
 		const pxr::UsdSkelSkinningQuery& UsdSkinningQuery,
 		const pxr::UsdSkelSkeletonQuery& SkeletonQuery,
 		FSkeletalMeshImportData& SkelMeshImportData,
-		TArray< UsdUtils::FUsdPrimMaterialSlot >& MaterialAssignments,
-		const TMap< FString, TMap< FString, int32 > >& MaterialToPrimvarsUVSetNames,
+		TArray<UsdUtils::FUsdPrimMaterialSlot>& MaterialAssignments,
+		const TMap<FString, TMap<FString, int32>>& MaterialToPrimvarsUVSetNames,
 		const pxr::TfToken& RenderContext = pxr::UsdShadeTokens->universalRenderContext,
 		const pxr::TfToken& MaterialPurpose = pxr::UsdShadeTokens->allPurpose
 	);
@@ -214,15 +247,18 @@ namespace UsdToUnreal
 	 * Will extract animation data from the animation source of InUsdSkeletonQuery's skeleton, and populate OutSkeletalAnimationAsset with the data.
 	 * Warning: UAnimSequence must be previously set with a USkeleton generated from the skeletal data of the same UsdSkelSkeletonQuery.
 	 * @param InUsdSkeletonQuery - SkinningQuery with the data to convert
-	 * @param InSkinningTargets - Skinned meshes that use the skeleton of InUsdSkeletonQuery. Required to fetch the blend shape ordering of each mesh. Optional (can be nullptr to ignore)
-	 * @param InBlendShapes - Converted blend shape data that will be used to interpret blend shape weights as morph target weight float curves. Optional (can be nullptr to ignore)
+	 * @param InSkinningTargets - Skinned meshes that use the skeleton of InUsdSkeletonQuery. Required to fetch the blend shape ordering of each mesh.
+	 *Optional (can be nullptr to ignore)
+	 * @param InBlendShapes - Converted blend shape data that will be used to interpret blend shape weights as morph target weight float curves.
+	 *Optional (can be nullptr to ignore)
 	 * @param InInterpretLODs - Whether we try parsing animation data from all LODs of skinning meshes that are inside LOD variant sets
 	 * @param RootMotionPrim - Optional prim whose transform animation will be concatenaded into the root bone's regular joint animation
 	 * @param OutSkeletalAnimationAsset - Output parameter that will be filled with the converted data
-	 * @param OutStartOffsetSeconds - Optional output parameter that will be filled with the offset in seconds of when this UAnimSequence asset should be played since the start of its layer
-	 *								  to match the intended composed animation. The baked UAnimSequence will only contain the range between the first and last joint and/or blend shape
-	 *                                timeSamples, and so the offset is needed to properly position the animation on the stage's timeline. This is in seconds because the main use case is to
-	 *                                use this offset when animating USkeletalMeshComponents, and those drive their UAnimSequences with seconds.
+	 * @param OutStartOffsetSeconds - Optional output parameter that will be filled with the offset in seconds of when this UAnimSequence asset should
+	 *be played since the start of its layer to match the intended composed animation. The baked UAnimSequence will only contain the range between the
+	 *first and last joint and/or blend shape timeSamples, and so the offset is needed to properly position the animation on the stage's timeline.
+	 *This is in seconds because the main use case is to use this offset when animating USkeletalMeshComponents, and those drive their UAnimSequences
+	 *with seconds.
 	 * @return Whether the conversion was successful or not.
 	 */
 	USDUTILITIES_API bool ConvertSkelAnim(
@@ -232,7 +268,7 @@ namespace UsdToUnreal
 		bool bInInterpretLODs,
 		const pxr::UsdPrim& RootMotionPrim,
 		UAnimSequence* OutSkeletalAnimationAsset,
-		float* OutStartOffsetSeconds=nullptr
+		float* OutStartOffsetSeconds = nullptr
 	);
 
 	/**
@@ -245,7 +281,14 @@ namespace UsdToUnreal
 	 * @param SkeletonName - Name to use for the new USkeleton asset
 	 * @return Newly created USkeletalMesh, or nullptr in case of failure
 	 */
-	USDUTILITIES_API USkeletalMesh* GetSkeletalMeshFromImportData( TArray<FSkeletalMeshImportData>& LODIndexToSkeletalMeshImportData, const TArray<SkeletalMeshImportData::FBone>& InSkeletonBones, UsdUtils::FBlendShapeMap& InBlendShapesByPath, EObjectFlags ObjectFlags, const FName& MeshName = NAME_None, const FName& SkeletonName = NAME_None );
+	USDUTILITIES_API USkeletalMesh* GetSkeletalMeshFromImportData(
+		TArray<FSkeletalMeshImportData>& LODIndexToSkeletalMeshImportData,
+		const TArray<SkeletalMeshImportData::FBone>& InSkeletonBones,
+		UsdUtils::FBlendShapeMap& InBlendShapesByPath,
+		EObjectFlags ObjectFlags,
+		const FName& MeshName = NAME_None,
+		const FName& SkeletonName = NAME_None
+	);
 }
 
 namespace UnrealToUsd
@@ -259,23 +302,33 @@ namespace UnrealToUsd
 	 * @param UsdSkeleton - Previously created prim with the UsdSkelSkeleton schema that will be filled with converted data
 	 * @return Whether the conversion was successful or not.
 	 */
-	USDUTILITIES_API bool ConvertSkeleton( const USkeleton* Skeleton, pxr::UsdSkelSkeleton& UsdSkeleton );
-	USDUTILITIES_API bool ConvertSkeleton( const FReferenceSkeleton& ReferenceSkeleton, pxr::UsdSkelSkeleton& UsdSkeleton );
+	USDUTILITIES_API bool ConvertSkeleton(const USkeleton* Skeleton, pxr::UsdSkelSkeleton& UsdSkeleton);
+	USDUTILITIES_API bool ConvertSkeleton(const FReferenceSkeleton& ReferenceSkeleton, pxr::UsdSkelSkeleton& UsdSkeleton);
 
 	// Fill out a UsdSkelAnimation's Joints attribute with data from ReferenceSkeleton, taking care to concatenate bone paths
-	USDUTILITIES_API bool ConvertJointsAttribute( const FReferenceSkeleton& ReferenceSkeleton, pxr::UsdAttribute& JointsAttribute );
+	USDUTILITIES_API bool ConvertJointsAttribute(const FReferenceSkeleton& ReferenceSkeleton, pxr::UsdAttribute& JointsAttribute);
 
 	/**
 	 * Converts SkeletalMesh, its skeleton and morph target data into the corresponding USD objects and populates SkelRoot with them, at time TimeCode
-	 * @param SkeletalMesh - Mesh with the source data. If it contains multiple LODs it will lead to the creation of LOD variant sets and variants within SkelRoot
-	 * @param SkelRoot - Root prim of the output source data. Child UsdSkelSkeleton, UsdGeomMesh, and UsdSkelBlendShape will be created as children of it, containing the converted data
+	 * @param SkeletalMesh - Mesh with the source data. If it contains multiple LODs it will lead to the creation of LOD variant sets and variants
+	 * within SkelRoot
+	 * @param SkelRoot - Root prim of the output source data. Child UsdSkelSkeleton, UsdGeomMesh, and UsdSkelBlendShape will be created as children of
+	 * it, containing the converted data
 	 * @param TimeCode - TimeCode with which the converted data will be placed in the USD stage
-	 * @param StageForMaterialAssignments - Stage to use when authoring material assignments (we use this when we want to export the mesh to a payload layer, but the material assignments to an asset layer)
+	 * @param StageForMaterialAssignments - Stage to use when authoring material assignments (we use this when we want to export the mesh to a payload
+	 * layer, but the material assignments to an asset layer)
 	 * @param LowestMeshLOD - Lowest LOD of the UStaticMesh to export (start of the LOD range)
 	 * @param HighestMeshLOD - Lowest LOD of the UStaticMesh to export (end of the LOD range)
 	 * @return Whether the conversion was successful or not.
 	 */
-	USDUTILITIES_API bool ConvertSkeletalMesh( const USkeletalMesh* SkeletalMesh, pxr::UsdPrim& SkelRootPrim, const pxr::UsdTimeCode TimeCode = pxr::UsdTimeCode::Default(), UE::FUsdStage* StageForMaterialAssignments = nullptr, int32 LowestMeshLOD = 0, int32 HighestMeshLOD = INT32_MAX );
+	USDUTILITIES_API bool ConvertSkeletalMesh(
+		const USkeletalMesh* SkeletalMesh,
+		pxr::UsdPrim& SkelRootPrim,
+		const pxr::UsdTimeCode TimeCode = pxr::UsdTimeCode::Default(),
+		UE::FUsdStage* StageForMaterialAssignments = nullptr,
+		int32 LowestMeshLOD = 0,
+		int32 HighestMeshLOD = INT32_MAX
+	);
 
 	/**
 	 * Converts an AnimSequence to a UsdSkelAnimation. Includes bone transforms and blend shape weights.
@@ -284,7 +337,7 @@ namespace UnrealToUsd
 	 * @param SkelAnimPrim - Expected to be of type UsdkSkelAnimation
 	 * @return Whether the conversion was successful or not.
 	 */
-	USDUTILITIES_API bool ConvertAnimSequence( UAnimSequence* AnimSequence, pxr::UsdPrim& SkelAnimPrim );
+	USDUTILITIES_API bool ConvertAnimSequence(UAnimSequence* AnimSequence, pxr::UsdPrim& SkelAnimPrim);
 
 	/**
 	 * Plays the provided Section in the background, driving its ControlRig and baking to USD the animated bones and
@@ -316,4 +369,4 @@ namespace UnrealToUsd
 	);
 }
 
-#endif // #if USE_USD_SDK && WITH_EDITOR
+#endif	  // #if USE_USD_SDK && WITH_EDITOR
diff --git a/Engine/Plugins/Importers/USDImporter/Source/USDUtilities/Public/USDTypesConversion.h b/Engine/Plugins/Importers/USDImporter/Source/USDUtilities/Public/USDTypesConversion.h
--- a/Engine/Plugins/Importers/USDImporter/Source/USDUtilities/Public/USDTypesConversion.h
+++ b/Engine/Plugins/Importers/USDImporter/Source/USDUtilities/Public/USDTypesConversion.h
@@ -10,7 +10,7 @@
 #include <string>
 
 #include "USDIncludesStart.h"
-	#include "pxr/pxr.h"
+#include "pxr/pxr.h"
 #include "USDIncludesEnd.h"
 
 PXR_NAMESPACE_OPEN_SCOPE
@@ -36,9 +36,10 @@ PXR_NAMESPACE_OPEN_SCOPE
 	class TfToken;
 
 	class UsdStage;
-	template< typename T > class TfRefPtr;
+	template<typename T>
+	class TfRefPtr;
 
-	using UsdStageRefPtr = TfRefPtr< UsdStage >;
+	using UsdStageRefPtr = TfRefPtr<UsdStage>;
 PXR_NAMESPACE_CLOSE_SCOPE
 
 struct USDUTILITIES_API FUsdStageInfo
@@ -46,24 +47,24 @@ struct USDUTILITIES_API FUsdStageInfo
 	EUsdUpAxis UpAxis = EUsdUpAxis::ZAxis;
 	float MetersPerUnit = 0.01f;
 
-	explicit FUsdStageInfo( const pxr::UsdStageRefPtr& Stage );
+	explicit FUsdStageInfo(const pxr::UsdStageRefPtr& Stage);
 };
 
 namespace UsdToUnreal
 {
-	USDUTILITIES_API FString ConvertString( const std::string& InString );
-	USDUTILITIES_API FString ConvertString( const char* InString );
+	USDUTILITIES_API FString ConvertString(const std::string& InString);
+	USDUTILITIES_API FString ConvertString(const char* InString);
 
-	USDUTILITIES_API FString ConvertPath( const pxr::SdfPath& Path );
+	USDUTILITIES_API FString ConvertPath(const pxr::SdfPath& Path);
 
-	USDUTILITIES_API FName ConvertName( const char* InString );
-	USDUTILITIES_API FName ConvertName( const std::string& InString );
+	USDUTILITIES_API FName ConvertName(const char* InString);
+	USDUTILITIES_API FName ConvertName(const std::string& InString);
 
-	USDUTILITIES_API FString ConvertToken( const pxr::TfToken& Token );
+	USDUTILITIES_API FString ConvertToken(const pxr::TfToken& Token);
 
 	/** Assumes the input color is in linear space */
-	USDUTILITIES_API FLinearColor ConvertColor( const pxr::GfVec3f& InValue );
-	USDUTILITIES_API FLinearColor ConvertColor( const pxr::GfVec4f& InValue );
+	USDUTILITIES_API FLinearColor ConvertColor(const pxr::GfVec3f& InValue);
+	USDUTILITIES_API FLinearColor ConvertColor(const pxr::GfVec4f& InValue);
 
 	USDUTILITIES_API FVector2D ConvertVector(const pxr::GfVec2h& InValue);
 	USDUTILITIES_API FVector2D ConvertVector(const pxr::GfVec2f& InValue);
@@ -85,33 +86,33 @@ namespace UsdToUnreal
 
 	USDUTILITIES_API FMatrix2D ConvertMatrix(const pxr::GfMatrix2d& Matrix);
 	USDUTILITIES_API FMatrix3D ConvertMatrix(const pxr::GfMatrix3d& Matrix);
-	USDUTILITIES_API FMatrix ConvertMatrix( const pxr::GfMatrix4d& Matrix );
+	USDUTILITIES_API FMatrix ConvertMatrix(const pxr::GfMatrix4d& Matrix);
 
-	USDUTILITIES_API FTransform ConvertMatrix( const FUsdStageInfo& StageInfo, const pxr::GfMatrix4d& InMatrix );
+	USDUTILITIES_API FTransform ConvertMatrix(const FUsdStageInfo& StageInfo, const pxr::GfMatrix4d& InMatrix);
 
 	USDUTILITIES_API FQuat ConvertQuat(const pxr::GfQuath& InValue);
 	USDUTILITIES_API FQuat ConvertQuat(const pxr::GfQuatf& InValue);
 	USDUTILITIES_API FQuat ConvertQuat(const pxr::GfQuatd& InValue);
 
 	/** Returns a distance in "UE units" (i.e. cm) */
-	USDUTILITIES_API float ConvertDistance( const FUsdStageInfo& StageInfo, float InValue );
-}
+	USDUTILITIES_API float ConvertDistance(const FUsdStageInfo& StageInfo, float InValue);
+}	 // namespace UsdToUnreal
 
 namespace UnrealToUsd
 {
-	USDUTILITIES_API TUsdStore< std::string > ConvertString( const TCHAR* InString );
+	USDUTILITIES_API TUsdStore<std::string> ConvertString(const TCHAR* InString);
 
-	USDUTILITIES_API TUsdStore< pxr::SdfPath > ConvertPath( const TCHAR* InString );
+	USDUTILITIES_API TUsdStore<pxr::SdfPath> ConvertPath(const TCHAR* InString);
 
-	USDUTILITIES_API TUsdStore< std::string > ConvertName( const FName& InName );
+	USDUTILITIES_API TUsdStore<std::string> ConvertName(const FName& InName);
 
-	USDUTILITIES_API TUsdStore< pxr::TfToken > ConvertToken( const TCHAR* InString );
+	USDUTILITIES_API TUsdStore<pxr::TfToken> ConvertToken(const TCHAR* InString);
 
 	/** Assumes the input color is in linear space. Returns a color in linear space */
-	USDUTILITIES_API pxr::GfVec4f ConvertColor( const FLinearColor& InValue );
+	USDUTILITIES_API pxr::GfVec4f ConvertColor(const FLinearColor& InValue);
 
 	/** Assumes the input color is in sRGB space. Returns a color in linear space */
-	USDUTILITIES_API pxr::GfVec4f ConvertColor( const FColor& InValue );
+	USDUTILITIES_API pxr::GfVec4f ConvertColor(const FColor& InValue);
 
 	UE_DEPRECATED(5.4, "Please use either ConvertVectorHalf, ConvertVectorFloat or ConvertVectorDouble.")
 	USDUTILITIES_API pxr::GfVec2f ConvertVector(const FVector2D& InValue);
@@ -149,16 +150,16 @@ namespace UnrealToUsd
 	USDUTILITIES_API pxr::GfQuatf ConvertQuatFloat(const FQuat& InValue);
 	USDUTILITIES_API pxr::GfQuatd ConvertQuatDouble(const FQuat& InValue);
 
-	USDUTILITIES_API pxr::GfMatrix4d ConvertTransform( const FUsdStageInfo& StageInfo, const FTransform& Transform );
+	USDUTILITIES_API pxr::GfMatrix4d ConvertTransform(const FUsdStageInfo& StageInfo, const FTransform& Transform);
 
 	/** Returns a distance in USD units (depends on metersPerUnit) */
-	USDUTILITIES_API float ConvertDistance( const FUsdStageInfo& StageInfo, float InValue );
-}
+	USDUTILITIES_API float ConvertDistance(const FUsdStageInfo& StageInfo, float InValue);
+}	 // namespace UnrealToUsd
 
 namespace UsdUtils
 {
-	USDUTILITIES_API FTransform ConvertTransformToUsdSpace( const FUsdStageInfo& StageInfo, const FTransform& TransformInUESpace );
+	USDUTILITIES_API FTransform ConvertTransformToUsdSpace(const FUsdStageInfo& StageInfo, const FTransform& TransformInUESpace);
 
-	FTransform ConvertAxes( const bool bZUp, const FTransform Transform );
+	FTransform ConvertAxes(const bool bZUp, const FTransform Transform);
 }
-#endif // #if USE_USD_SDK
+#endif	  // #if USE_USD_SDK
diff --git a/Engine/Plugins/Importers/USDImporter/Source/USDUtilities/Public/USDValueConversion.h b/Engine/Plugins/Importers/USDImporter/Source/USDUtilities/Public/USDValueConversion.h
--- a/Engine/Plugins/Importers/USDImporter/Source/USDUtilities/Public/USDValueConversion.h
+++ b/Engine/Plugins/Importers/USDImporter/Source/USDUtilities/Public/USDValueConversion.h
@@ -13,7 +13,7 @@
 #if USE_USD_SDK
 
 #include "USDIncludesStart.h"
-	#include "pxr/pxr.h"
+#include "pxr/pxr.h"
 #include "USDIncludesEnd.h"
 
 PXR_NAMESPACE_OPEN_SCOPE
@@ -23,10 +23,11 @@ PXR_NAMESPACE_OPEN_SCOPE
 	class VtValue;
 PXR_NAMESPACE_CLOSE_SCOPE
 
-#endif // #if USE_USD_SDK
+#endif	  // #if USE_USD_SDK
 
 struct FMatrix2D;
 struct FMatrix3D;
+
 namespace UE
 {
 	class FVtValue;
@@ -82,33 +83,34 @@ namespace UsdUtils
 
 	// WARNING: Do not change the order of the types in this TVariant declaration, or it will break backwards compatibility with data that
 	// was serialized before the change, due to how `operator<<( FArchive& Ar, FConvertedVtValueComponent& Component )` is implemented
-	using FConvertedVtValueComponent = TVariant< bool, uint8, int32, uint32, int64, uint64, float, double, FString >;
+	using FConvertedVtValueComponent = TVariant<bool, uint8, int32, uint32, int64, uint64, float, double, FString>;
 
 	/** Represents a single non-array value held by a pxr::VtValue object, like a 'float', a 'TfToken' or a 'GfMatrix3d' */
-	using FConvertedVtValueEntry = TArray< FConvertedVtValueComponent >;
+	using FConvertedVtValueEntry = TArray<FConvertedVtValueComponent>;
 
-	/** Corresponds to a value held by a pxr::VtValue converted into UE types. Should handle anything a VtValue can hold, including arrays of values */
+	/** Corresponds to a value held by a pxr::VtValue converted into UE types. Should handle anything a VtValue can hold, including arrays of values
+	 */
 	struct USDUTILITIES_API FConvertedVtValue
 	{
 		TArray<FConvertedVtValueEntry> Entries;
 
 		EUsdBasicDataTypes SourceType = EUsdBasicDataTypes::None;
 		bool bIsArrayValued = false;
-		bool bIsEmpty = false;  // Helps differentiating between empty arrays and having actually no value
+		bool bIsEmpty = false;	  // Helps differentiating between empty arrays and having actually no value
 	};
-}
+}	 // namespace UsdUtils
 
-USDUTILITIES_API FArchive& operator<<( FArchive& Ar, UsdUtils::FConvertedVtValue& Struct );
-USDUTILITIES_API FArchive& operator<<( FArchive& Ar, UsdUtils::FConvertedVtValueComponent& Component );
+USDUTILITIES_API FArchive& operator<<(FArchive& Ar, UsdUtils::FConvertedVtValue& Struct);
+USDUTILITIES_API FArchive& operator<<(FArchive& Ar, UsdUtils::FConvertedVtValueComponent& Component);
 
 namespace UsdToUnreal
 {
-	USDUTILITIES_API bool ConvertValue( const UE::FVtValue& InValue, UsdUtils::FConvertedVtValue& OutValue );
+	USDUTILITIES_API bool ConvertValue(const UE::FVtValue& InValue, UsdUtils::FConvertedVtValue& OutValue);
 }
 
 namespace UnrealToUsd
 {
-	USDUTILITIES_API bool ConvertValue( const UsdUtils::FConvertedVtValue& InValue, UE::FVtValue& OutValue );
+	USDUTILITIES_API bool ConvertValue(const UsdUtils::FConvertedVtValue& InValue, UE::FVtValue& OutValue);
 }
 
 namespace UsdUtils
@@ -204,7 +206,7 @@ namespace UsdUtils
 	 * Example: String is "(1, 0.5, 0.3)" and TypeName "float3" means Output will contain a pxr::GfVec3f(1, 0.5, 0.3) value
 	 */
 	USDUTILITIES_API bool Unstringify(const FString& String, const FString& TypeName, pxr::VtValue& Output);
-#endif // USE_USD_SDK
+#endif	  // USE_USD_SDK
 
 	// c.f. comment above StringifyAsBool
 	USDUTILITIES_API bool UnstringifyAsBool(const FString& String, bool& OutValue);
@@ -277,7 +279,7 @@ namespace UsdUtils
 	 * Template implementation must be hidden on cpp as we can't expose USD implementation on the header files.
 	 */
 	template<typename T>
-	USDUTILITIES_API TOptional<T> GetUnderlyingValue( const UE::FVtValue& InValue );
+	USDUTILITIES_API TOptional<T> GetUnderlyingValue(const UE::FVtValue& InValue);
 
 	/**
 	 * Uses USD to quickly set the underlying value of the wrapped pxr::VtValue for fundamental data types.
@@ -287,7 +289,7 @@ namespace UsdUtils
 	 * Template implementation must be hidden on cpp as we can't expose USD implementation on the header files.
 	 */
 	template<typename T>
-	USDUTILITIES_API bool SetUnderlyingValue( UE::FVtValue& InValue, const T& UnderlyingValue );
+	USDUTILITIES_API bool SetUnderlyingValue(UE::FVtValue& InValue, const T& UnderlyingValue);
 
 	/**
 	 * Returns the name of the SdfValueTypeName object for the type of Value
@@ -296,5 +298,5 @@ namespace UsdUtils
 	 * Note that this may not be the correct SdfValueTypeName for the *attribute*, it's just the implied type of the value (e.g. for an
 	 * attribute with SdfValueTypeName 'normal3d', it's VtValues would have underlying type 'GfVec3d', and so this function would return 'double3')
 	 */
-	USDUTILITIES_API FString GetImpliedTypeName( const UE::FVtValue& Value );
-}
+	USDUTILITIES_API FString GetImpliedTypeName(const UE::FVtValue& Value);
+}	 // namespace UsdUtils
