USD: Allow USDClasses module (and its classes) to be compiled on all platforms if needed, as it doesn't use the USD SDK directly anyway.

diff --git a/Engine/Plugins/Importers/USDImporter/USDImporter.uplugin b/Engine/Plugins/Importers/USDImporter/USDImporter.uplugin
--- a/Engine/Plugins/Importers/USDImporter/USDImporter.uplugin
+++ b/Engine/Plugins/Importers/USDImporter/USDImporter.uplugin
@@ -14,49 +14,70 @@
 	"CanContainContent": true,
 	"IsBetaVersion": true,
 	"Installed": false,
-	"SupportedTargetPlatforms": [
-		"Win64", "Mac", "Linux"
-	],
 	"Modules": [
 		{
 			"Name": "UnrealUSDWrapper",
 			"Type": "Runtime",
-			"LoadingPhase": "PreDefault"
+			"LoadingPhase": "PreDefault",
+			"PlatformAllowList": [
+				"Win64", "Mac", "Linux"
+			]
 		},
 		{
 			"Name": "USDUtilities",
 			"Type": "Runtime",
-			"LoadingPhase": "Default"
+			"LoadingPhase": "Default",
+			"PlatformAllowList": [
+				"Win64", "Mac", "Linux"
+			]
 		},
 		{
 			"Name": "USDSchemas",
 			"Type": "Runtime",
-			"LoadingPhase": "Default"
+			"LoadingPhase": "Default",
+			"PlatformAllowList": [
+				"Win64", "Mac", "Linux"
+			]
 		},
 		{
 			"Name": "USDStage",
 			"Type": "Runtime",
-			"LoadingPhase": "Default"
+			"LoadingPhase": "Default",
+			"PlatformAllowList": [
+				"Win64", "Mac", "Linux"
+			]
 		},
 		{
 			"Name": "USDStageEditor",
 			"Type": "Editor",
-			"LoadingPhase": "Default"
+			"LoadingPhase": "Default",
+			"PlatformAllowList": [
+				"Win64", "Mac", "Linux"
+			]
 		},
 		{
 			"Name": "USDStageImporter",
 			"Type": "Editor",
-			"LoadingPhase": "Default"
+			"LoadingPhase": "Default",
+			"PlatformAllowList": [
+				"Win64", "Mac", "Linux"
+			]
 		},
 		{
 			"Name": "USDStageEditorViewModels",
 			"Type": "Editor",
-			"LoadingPhase": "Default"
+			"LoadingPhase": "Default",
+			"PlatformAllowList": [
+				"Win64", "Mac", "Linux"
+			]
 		},
 		{
 			"Name": "USDExporter",
 			"Type": "Editor",
-			"LoadingPhase": "Default"
+			"LoadingPhase": "Default",
+			"PlatformAllowList": [
+				"Win64", "Mac", "Linux"
+			]
 		},
 		{
 			"Name": "USDClasses",
@@ -66,17 +87,26 @@
         {
 			"Name": "USDClassesEditor",
 			"Type": "Editor",
-			"LoadingPhase": "Default"
+			"LoadingPhase": "Default",
+			"PlatformAllowList": [
+				"Win64", "Mac", "Linux"
+			]
 		},
 		{
 			"Name": "GeometryCacheUSD",
 			"Type": "Editor",
-			"LoadingPhase": "Default"
+			"LoadingPhase": "Default",
+			"PlatformAllowList": [
+				"Win64", "Mac", "Linux"
+			]
 		},
 		{
 			"Name": "USDTests",
 			"Type": "Editor",
-			"LoadingPhase": "Default"
+			"LoadingPhase": "Default",
+			"PlatformAllowList": [
+				"Win64", "Mac", "Linux"
+			]
 		}
 	],
 	"Plugins": [
