USD: Reformat GeometryCacheUSD module.

diff --git a/Engine/Plugins/Importers/USDImporter/Source/GeometryCacheUSD/Private/GeometryCacheTrackUSD.cpp b/Engine/Plugins/Importers/USDImporter/Source/GeometryCacheUSD/Private/GeometryCacheTrackUSD.cpp
--- a/Engine/Plugins/Importers/USDImporter/Source/GeometryCacheUSD/Private/GeometryCacheTrackUSD.cpp
+++ b/Engine/Plugins/Importers/USDImporter/Source/GeometryCacheUSD/Private/GeometryCacheTrackUSD.cpp
@@ -15,13 +15,16 @@ static bool GDisableGeoCacheTracks = false;
 static FAutoConsoleVariableRef CVarDisableGeoCacheTracks(
 	TEXT("USD.DisableGeoCacheTracks"),
 	GDisableGeoCacheTracks,
-	TEXT("Set to true to disable geometry cache tracks in Sequencer and drive them by the Time property instead. The stage must be reloaded after changing this value."));
+	TEXT("Set to true to disable geometry cache tracks in Sequencer and drive them by the Time property instead. The stage must be reloaded after "
+		 "changing this value.")
+);
 
 UGeometryCacheTrackUsd::UGeometryCacheTrackUsd()
 	: FramesPerSecond(24.0)
 	, StartFrameIndex(0)
 	, EndFrameIndex(0)
-{}
+{
+}
 
 void UGeometryCacheTrackUsd::BeginDestroy()
 {
@@ -46,7 +49,12 @@ void UGeometryCacheTrackUsd::GetResourceSizeEx(FResourceSizeEx& CumulativeResour
 	CumulativeResourceSize.AddDedicatedSystemMemoryBytes(SampleInfos.GetAllocatedSize());
 }
 
-const bool UGeometryCacheTrackUsd::UpdateMeshData(const float Time, const bool bLooping, int32& InOutMeshSampleIndex, FGeometryCacheMeshData*& OutMeshData)
+const bool UGeometryCacheTrackUsd::UpdateMeshData(
+	const float Time,
+	const bool bLooping,
+	int32& InOutMeshSampleIndex,
+	FGeometryCacheMeshData*& OutMeshData
+)
 {
 	const int32 SampleIndex = FindSampleIndexFromTime(Time, bLooping);
 
@@ -64,7 +72,13 @@ const bool UGeometryCacheTrackUsd::UpdateMeshData(const float Time, const bool b
 	return false;
 }
 
-const bool UGeometryCacheTrackUsd::UpdateBoundsData(const float Time, const bool bLooping, const bool bIsPlayingBackward, int32& InOutBoundsSampleIndex, FBox& OutBounds)
+const bool UGeometryCacheTrackUsd::UpdateBoundsData(
+	const float Time,
+	const bool bLooping,
+	const bool bIsPlayingBackward,
+	int32& InOutBoundsSampleIndex,
+	FBox& OutBounds
+)
 {
 	const int32 SampleIndex = FindSampleIndexFromTime(Time, bLooping);
 
@@ -197,7 +211,13 @@ bool UGeometryCacheTrackUsd::LoadUsdStage()
 	}
 	else if (!StageRootLayerPath.IsEmpty())
 	{
-		UE_LOG(LogUsd, Warning, TEXT("UGeometryCacheTrackUsd is reopening the stage '%s' to stream in frames for the geometry cache generated for prim '%s'"), *StageRootLayerPath, *PrimPath);
+		UE_LOG(
+			LogUsd,
+			Warning,
+			TEXT("UGeometryCacheTrackUsd is reopening the stage '%s' to stream in frames for the geometry cache generated for prim '%s'"),
+			*StageRootLayerPath,
+			*PrimPath
+		);
 
 		// Reopen the stage. If our weak pointer is no longer valid then nothing cared about keeping that
 		// stage alive anyway, so it's likely not a problem if we start reading frames from the reopened stage
@@ -256,7 +276,7 @@ void UGeometryCacheTrackUsd::Initialize(
 	const UE::FUsdStage& InStage,
 	const FString& InPrimPath,
 	const FName& InRenderContext,
-	const TMap< FString, TMap< FString, int32 > >& InMaterialToPrimvarToUVIndex,
+	const TMap<FString, TMap<FString, int32>>& InMaterialToPrimvarToUVIndex,
 	int32 InStartFrameIndex,
 	int32 InEndFrameIndex,
 	FReadUsdMeshFunction InReadFunc
diff --git a/Engine/Plugins/Importers/USDImporter/Source/GeometryCacheUSD/Public/GeometryCacheTrackUSD.h b/Engine/Plugins/Importers/USDImporter/Source/GeometryCacheUSD/Public/GeometryCacheTrackUSD.h
--- a/Engine/Plugins/Importers/USDImporter/Source/GeometryCacheUSD/Public/GeometryCacheTrackUSD.h
+++ b/Engine/Plugins/Importers/USDImporter/Source/GeometryCacheUSD/Public/GeometryCacheTrackUSD.h
@@ -7,8 +7,8 @@
 #include "UsdWrappers/UsdStage.h"
 
 #include "CoreMinimal.h"
-#include "GeometryCacheTrack.h"
 #include "GeometryCacheMeshData.h"
+#include "GeometryCacheTrack.h"
 #include "GeometryCacheTrackUSDTypes.h"
 #include "GeometryCacheUSDStream.h"
 
@@ -16,6 +16,7 @@
 
 /** GeometryCacheTrack for querying USD */
 UCLASS(collapsecategories, hidecategories = Object, BlueprintType, config = Engine)
+
 class GEOMETRYCACHEUSD_API UGeometryCacheTrackUsd : public UGeometryCacheTrack
 {
 	GENERATED_BODY()
@@ -36,7 +37,7 @@ public:
 		const UE::FUsdStage& InStage,
 		const FString& InPrimPath,
 		const FName& InRenderContext,
-		const TMap< FString, TMap< FString, int32 > >& InMaterialToPrimvarToUVIndex,
+		const TMap<FString, TMap<FString, int32>>& InMaterialToPrimvarToUVIndex,
 		int32 InStartFrameIndex,
 		int32 InEndFrameIndex,
 		FReadUsdMeshFunction InReadFunc
@@ -48,8 +49,15 @@ public:
 	//~ End UObject Interface.
 
 	//~ Begin UGeometryCacheTrack Interface.
-	virtual const bool UpdateMeshData(const float Time, const bool bLooping, int32& InOutMeshSampleIndex, FGeometryCacheMeshData*& OutMeshData) override;
-	virtual const bool UpdateBoundsData(const float Time, const bool bLooping, const bool bIsPlayingBackward, int32& InOutBoundsSampleIndex, FBox& OutBounds) override;
+	virtual const bool UpdateMeshData(const float Time, const bool bLooping, int32& InOutMeshSampleIndex, FGeometryCacheMeshData*& OutMeshData)
+		override;
+	virtual const bool UpdateBoundsData(
+		const float Time,
+		const bool bLooping,
+		const bool bIsPlayingBackward,
+		int32& InOutBoundsSampleIndex,
+		FBox& OutBounds
+	) override;
 	virtual const FGeometryCacheTrackSampleInfo& GetSampleInfo(float Time, const bool bLooping) override;
 	virtual bool GetMeshDataAtTime(float Time, FGeometryCacheMeshData& OutMeshData) override;
 	virtual bool GetMeshDataAtSampleIndex(int32 SampleIndex, FGeometryCacheMeshData& OutMeshData) override;
@@ -62,7 +70,8 @@ public:
 
 	bool GetMeshData(int32 SampleIndex, FGeometryCacheMeshData& OutMeshData);
 
-	// Upgrades our CurrentStageWeak into CurrentStagePinned, or re-opens the stage if its stale. Returns whether the stage was successfully opened or not
+	// Upgrades our CurrentStageWeak into CurrentStagePinned, or re-opens the stage if its stale. Returns whether the stage was successfully opened or
+	// not
 	bool LoadUsdStage();
 
 	// Discards our CurrentStagePinned to release the stage
diff --git a/Engine/Plugins/Importers/USDImporter/Source/GeometryCacheUSD/Public/GeometryCacheUSDComponent.h b/Engine/Plugins/Importers/USDImporter/Source/GeometryCacheUSD/Public/GeometryCacheUSDComponent.h
--- a/Engine/Plugins/Importers/USDImporter/Source/GeometryCacheUSD/Public/GeometryCacheUSDComponent.h
+++ b/Engine/Plugins/Importers/USDImporter/Source/GeometryCacheUSD/Public/GeometryCacheUSDComponent.h
@@ -7,8 +7,10 @@
 
 #include "GeometryCacheUSDComponent.generated.h"
 
-/** GeometryCacheUSDComponent, encapsulates a transient GeometryCache asset instance that fetches its data from a USD file and implements functionality for rendering and playback */
+/** GeometryCacheUSDComponent, encapsulates a transient GeometryCache asset instance that fetches its data from a USD file and implements
+ * functionality for rendering and playback */
 UCLASS(HideDropDown, ClassGroup = (Rendering), meta = (DisplayName = "USD Geometry Cache"), Experimental, ClassGroup = Experimental)
+
 class GEOMETRYCACHEUSD_API UGeometryCacheUsdComponent : public UGeometryCacheComponent
 {
 	GENERATED_BODY()
