USD: Reformat the UnrealUSDWrapper module.

#rnx

diff --git a/Engine/Plugins/Importers/USDImporter/Source/UnrealUSDWrapper/Private/USDListener.cpp b/Engine/Plugins/Importers/USDImporter/Source/UnrealUSDWrapper/Private/USDListener.cpp
--- a/Engine/Plugins/Importers/USDImporter/Source/UnrealUSDWrapper/Private/USDListener.cpp
+++ b/Engine/Plugins/Importers/USDImporter/Source/UnrealUSDWrapper/Private/USDListener.cpp
@@ -14,7 +14,6 @@
 #if USE_USD_SDK
 
 #include "USDIncludesStart.h"
-
 #include "pxr/base/tf/weakBase.h"
 #include "pxr/usd/sdf/changeList.h"
 #include "pxr/usd/sdf/notice.h"
@@ -23,10 +22,9 @@
 #include "pxr/usd/usd/prim.h"
 #include "pxr/usd/usd/stage.h"
 #include "pxr/usd/usdGeom/tokens.h"
-
 #include "USDIncludesEnd.h"
 
-#endif // USE_USD_SDK
+#endif	  // USE_USD_SDK
 
 namespace UsdToUnreal
 {
@@ -41,44 +39,44 @@ namespace UsdToUnreal
 
 		// For most changes we'll only get one of these, but sometimes multiple changes are fired in sequence
 		// (e.g. if you change framesPerSecond, it will send a notice for it but also for the matching, updated timeCodesPerSecond)
-		for ( const std::pair<TfToken, std::pair<VtValue, VtValue>>& AttributeChange : Entry.infoChanged ) // Note: infoChanged here is just a naming conflict, there's no "resyncChanged"
+		for (const std::pair<TfToken, std::pair<VtValue, VtValue>>& AttributeChange : Entry.infoChanged)	// Note: infoChanged here is just a naming
+																											// conflict, there's no "resyncChanged"
 		{
 			UsdUtils::FAttributeChange& ConvertedAttributeChange = OutConverted.AttributeChanges.Emplace_GetRef();
 
-			const FString FieldToken = ANSI_TO_TCHAR( AttributeChange.first.GetString().c_str() );
+			const FString FieldToken = ANSI_TO_TCHAR(AttributeChange.first.GetString().c_str());
 
-			// For regular properties (most common case) FullFieldPath will already carry the property and FieldToken will be "default", "timeSamples", "variability", etc.
-			if ( ChangedObject.IsPropertyPath() )
+			// For regular properties (most common case) FullFieldPath will already carry the property and FieldToken will be "default",
+			// "timeSamples", "variability", etc.
+			if (ChangedObject.IsPropertyPath())
 			{
-				ConvertedAttributeChange.PropertyName = ANSI_TO_TCHAR( ChangedObject.GetName().c_str() );
+				ConvertedAttributeChange.PropertyName = ANSI_TO_TCHAR(ChangedObject.GetName().c_str());
 				ConvertedAttributeChange.Field = FieldToken;
 			}
-			// For stage properties it seems like USD likes to send just "/" as the FullFieldPath, and the actual property name in FieldToken (e.g. "metersPerUnit", "timeCodesPerSecond", etc.)
-			// We send "default" here for consistency
-			else if ( ChangedObject.IsAbsoluteRootPath() || FieldToken == FString( TEXT( "kind" ) ) )
+			// For stage properties it seems like USD likes to send just "/" as the FullFieldPath, and the actual property name in FieldToken (e.g.
+			// "metersPerUnit", "timeCodesPerSecond", etc.) We send "default" here for consistency
+			else if (ChangedObject.IsAbsoluteRootPath() || FieldToken == FString(TEXT("kind")))
 			{
 				ConvertedAttributeChange.PropertyName = FieldToken;
-				ConvertedAttributeChange.Field = TEXT( "default" );
+				ConvertedAttributeChange.Field = TEXT("default");
 			}
 
-			ConvertedAttributeChange.OldValue = UE::FVtValue{ AttributeChange.second.first };
-			ConvertedAttributeChange.NewValue = UE::FVtValue{ AttributeChange.second.second };
+			ConvertedAttributeChange.OldValue = UE::FVtValue{AttributeChange.second.first};
+			ConvertedAttributeChange.NewValue = UE::FVtValue{AttributeChange.second.second};
 		}
 
-		// Some notices (like creating/removing a property) don't have any actual infoChanged entries, so we need to create a fake one in here in order to convey the PropertyName, if applicable
-		if ( Entry.infoChanged.size() == 0 &&
-			( Entry.flags.didAddProperty ||
-				Entry.flags.didAddPropertyWithOnlyRequiredFields ||
-				Entry.flags.didRemoveProperty ||
-				Entry.flags.didRemovePropertyWithOnlyRequiredFields ||
-				Entry.flags.didChangeAttributeTimeSamples ) )
+		// Some notices (like creating/removing a property) don't have any actual infoChanged entries, so we need to create a fake one in here in
+		// order to convey the PropertyName, if applicable
+		if (Entry.infoChanged.size() == 0
+			&& (Entry.flags.didAddProperty || Entry.flags.didAddPropertyWithOnlyRequiredFields || Entry.flags.didRemoveProperty
+				|| Entry.flags.didRemovePropertyWithOnlyRequiredFields || Entry.flags.didChangeAttributeTimeSamples))
 		{
 			UsdUtils::FAttributeChange& ConvertedAttributeChange = OutConverted.AttributeChanges.Emplace_GetRef();
-			ConvertedAttributeChange.Field = Entry.flags.didChangeAttributeTimeSamples ? TEXT( "timeSamples" ) : TEXT( "default" );
+			ConvertedAttributeChange.Field = Entry.flags.didChangeAttributeTimeSamples ? TEXT("timeSamples") : TEXT("default");
 
-			if ( ChangedObject.IsPropertyPath() )
+			if (ChangedObject.IsPropertyPath())
 			{
-				ConvertedAttributeChange.PropertyName = ANSI_TO_TCHAR( ChangedObject.GetName().c_str() );
+				ConvertedAttributeChange.PropertyName = ANSI_TO_TCHAR(ChangedObject.GetName().c_str());
 			}
 		}
 
@@ -109,24 +107,33 @@ namespace UsdToUnreal
 		ConvertedFlags.bDidRemovePropertyWithOnlyRequiredFields = Entry.flags.didRemovePropertyWithOnlyRequiredFields;
 		ConvertedFlags.bDidRemoveProperty = Entry.flags.didRemoveProperty;
 
-		static_assert( static_cast< int >( UsdUtils::ESubLayerChangeType::SubLayerAdded ) == static_cast< int >( pxr::SdfChangeList::SubLayerChangeType::SubLayerAdded ), "Enum values changed!" );
-		static_assert( static_cast< int >( UsdUtils::ESubLayerChangeType::SubLayerOffset ) == static_cast< int >( pxr::SdfChangeList::SubLayerChangeType::SubLayerOffset ), "Enum values changed!" );
-		static_assert( static_cast< int >( UsdUtils::ESubLayerChangeType::SubLayerRemoved ) == static_cast< int >( pxr::SdfChangeList::SubLayerChangeType::SubLayerRemoved ), "Enum values changed!" );
-		for ( const std::pair<std::string, pxr::SdfChangeList::SubLayerChangeType>& SubLayerChange : Entry.subLayerChanges )
+		static_assert(
+			static_cast<int>(UsdUtils::ESubLayerChangeType::SubLayerAdded) == static_cast<int>(pxr::SdfChangeList::SubLayerChangeType::SubLayerAdded),
+			"Enum values changed!"
+		);
+		static_assert(
+			static_cast<int>(UsdUtils::ESubLayerChangeType::SubLayerOffset)
+				== static_cast<int>(pxr::SdfChangeList::SubLayerChangeType::SubLayerOffset),
+			"Enum values changed!"
+		);
+		static_assert(
+			static_cast<int>(UsdUtils::ESubLayerChangeType::SubLayerRemoved)
+				== static_cast<int>(pxr::SdfChangeList::SubLayerChangeType::SubLayerRemoved),
+			"Enum values changed!"
+		);
+		for (const std::pair<std::string, pxr::SdfChangeList::SubLayerChangeType>& SubLayerChange : Entry.subLayerChanges)
 		{
-			OutConverted.SubLayerChanges.Add(
-				TPair<FString, UsdUtils::ESubLayerChangeType>(
-					ANSI_TO_TCHAR( SubLayerChange.first.c_str() ),
-					static_cast< UsdUtils::ESubLayerChangeType >( SubLayerChange.second )
-				)
-			);
+			OutConverted.SubLayerChanges.Add(TPair<FString, UsdUtils::ESubLayerChangeType>(
+				ANSI_TO_TCHAR(SubLayerChange.first.c_str()),
+				static_cast<UsdUtils::ESubLayerChangeType>(SubLayerChange.second)
+			));
 		}
 
-		OutConverted.OldPath = ANSI_TO_TCHAR( Entry.oldPath.GetString().c_str() );
-		OutConverted.OldIdentifier = ANSI_TO_TCHAR( Entry.oldIdentifier.c_str() );
+		OutConverted.OldPath = ANSI_TO_TCHAR(Entry.oldPath.GetString().c_str());
+		OutConverted.OldIdentifier = ANSI_TO_TCHAR(Entry.oldIdentifier.c_str());
 	}
 
-	bool ConvertPathRange( const pxr::UsdNotice::ObjectsChanged::PathRange& PathRange, UsdUtils::FObjectChangesByPath& OutChanges )
+	bool ConvertPathRange(const pxr::UsdNotice::ObjectsChanged::PathRange& PathRange, UsdUtils::FObjectChangesByPath& OutChanges)
 	{
 		using namespace pxr;
 
@@ -164,22 +171,26 @@ namespace UsdToUnreal
 		return true;
 	}
 
-	bool ConvertObjectsChangedNotice( const pxr::UsdNotice::ObjectsChanged& InNotice, UsdUtils::FObjectChangesByPath& OutInfoChanges, UsdUtils::FObjectChangesByPath& OutResyncChanges)
+	bool ConvertObjectsChangedNotice(
+		const pxr::UsdNotice::ObjectsChanged& InNotice,
+		UsdUtils::FObjectChangesByPath& OutInfoChanges,
+		UsdUtils::FObjectChangesByPath& OutResyncChanges
+	)
 	{
-		ConvertPathRange( InNotice.GetChangedInfoOnlyPaths(), OutInfoChanges );
-		ConvertPathRange( InNotice.GetResyncedPaths(), OutResyncChanges );
+		ConvertPathRange(InNotice.GetChangedInfoOnlyPaths(), OutInfoChanges);
+		ConvertPathRange(InNotice.GetResyncedPaths(), OutResyncChanges);
 
-		for ( TPair< FString, TArray<UsdUtils::FSdfChangeListEntry> >& InfoPair : OutInfoChanges )
+		for (TPair<FString, TArray<UsdUtils::FSdfChangeListEntry>>& InfoPair : OutInfoChanges)
 		{
 			const FString& PrimPath = InfoPair.Key;
-			TArray< UsdUtils::FSdfChangeListEntry >* AnalogueResyncChanges = nullptr;
+			TArray<UsdUtils::FSdfChangeListEntry>* AnalogueResyncChanges = nullptr;
 
-			for ( TArray<UsdUtils::FSdfChangeListEntry>::TIterator ChangeIt = InfoPair.Value.CreateIterator(); ChangeIt; ++ChangeIt )
+			for (TArray<UsdUtils::FSdfChangeListEntry>::TIterator ChangeIt = InfoPair.Value.CreateIterator(); ChangeIt; ++ChangeIt)
 			{
 				bool bUpgrade = false;
 
 				// Upgrade info changes with content reloads into resync changes
-				if ( ChangeIt->Flags.bDidReloadContent )
+				if (ChangeIt->Flags.bDidReloadContent)
 				{
 					bUpgrade = true;
 				}
@@ -187,11 +198,11 @@ namespace UsdToUnreal
 				// Upgrade visibility changes to resyncs because in case of mesh collapsing having one of the collapsed meshes go visible/invisible
 				// should cause the regeneration of the collapsed asset. This is a bit expensive, but the asset cache will be used so its not as if
 				// the mesh will be completely regenerated however
-				if ( !bUpgrade )
+				if (!bUpgrade)
 				{
-					for ( UsdUtils::FAttributeChange& Change : ChangeIt->AttributeChanges )
+					for (UsdUtils::FAttributeChange& Change : ChangeIt->AttributeChanges)
 					{
-						if ( Change.PropertyName == ANSI_TO_TCHAR( pxr::UsdGeomTokens->visibility.GetString().c_str() ) )
+						if (Change.PropertyName == ANSI_TO_TCHAR(pxr::UsdGeomTokens->visibility.GetString().c_str()))
 						{
 							bUpgrade = true;
 							break;
@@ -199,14 +210,14 @@ namespace UsdToUnreal
 					}
 				}
 
-				if ( bUpgrade )
+				if (bUpgrade)
 				{
-					if ( !AnalogueResyncChanges )
+					if (!AnalogueResyncChanges)
 					{
-						AnalogueResyncChanges = &OutResyncChanges.FindOrAdd( PrimPath );
+						AnalogueResyncChanges = &OutResyncChanges.FindOrAdd(PrimPath);
 					}
 
-					AnalogueResyncChanges->Add( *ChangeIt );
+					AnalogueResyncChanges->Add(*ChangeIt);
 					ChangeIt.RemoveCurrent();
 				}
 			}
@@ -214,13 +225,13 @@ namespace UsdToUnreal
 
 		return true;
 	}
-#endif // USE_USD_SDK
-}
+#endif	  // USE_USD_SDK
+}	 // namespace UsdToUnreal
 
 class FUsdListenerImpl
 #if USE_USD_SDK
 	: public pxr::TfWeakBase
-#endif // #if USE_USD_SDK
+#endif	  // #if USE_USD_SDK
 {
 public:
 	PRAGMA_DISABLE_DEPRECATION_WARNINGS
@@ -240,9 +251,9 @@ public:
 	FThreadSafeCounter IsBlocked;
 
 #if USE_USD_SDK
-	FUsdListenerImpl( const pxr::UsdStageRefPtr& Stage );
+	FUsdListenerImpl(const pxr::UsdStageRefPtr& Stage);
 
-	void Register( const pxr::UsdStageRefPtr& Stage );
+	void Register(const pxr::UsdStageRefPtr& Stage);
 
 protected:
 	void HandleObjectsChangedNotice(const pxr::UsdNotice::ObjectsChanged& Notice, const pxr::UsdStageWeakPtr& Sender);
@@ -255,27 +266,27 @@ private:
 	pxr::TfNotice::Key RegisteredStageEditTargetChangedKey;
 	pxr::TfNotice::Key RegisteredSdfLayersChangedKey;
 	pxr::TfNotice::Key RegisteredSdfLayerDirtinessChangedKey;
-#endif // #if USE_USD_SDK
+#endif	  // #if USE_USD_SDK
 };
 
 FUsdListener::FUsdListener()
-	: Impl( MakeUnique< FUsdListenerImpl >() )
+	: Impl(MakeUnique<FUsdListenerImpl>())
 {
 }
 
-FUsdListener::FUsdListener( const UE::FUsdStage& Stage )
+FUsdListener::FUsdListener(const UE::FUsdStage& Stage)
 	: FUsdListener()
 {
-	Register( Stage );
+	Register(Stage);
 }
 
 FUsdListener::~FUsdListener() = default;
 
-void FUsdListener::Register( const UE::FUsdStage& Stage )
+void FUsdListener::Register(const UE::FUsdStage& Stage)
 {
 #if USE_USD_SDK
-	Impl->Register( Stage );
-#endif // #if USE_USD_SDK
+	Impl->Register(Stage);
+#endif	  // #if USE_USD_SDK
 }
 
 void FUsdListener::Block()
@@ -303,6 +314,7 @@ FUsdListener::FOnLayersChanged& FUsdListener::GetOnLayersChanged()
 {
 	return Impl->OnLayersChanged;
 }
+
 PRAGMA_ENABLE_DEPRECATION_WARNINGS
 
 FUsdListener::FOnSdfLayersChanged& FUsdListener::GetOnSdfLayersChanged()
@@ -321,81 +333,93 @@ FUsdListener::FOnObjectsChanged& FUsdListener::GetOnObjectsChanged()
 }
 
 #if USE_USD_SDK
-void FUsdListenerImpl::Register( const pxr::UsdStageRefPtr& Stage )
+void FUsdListenerImpl::Register(const pxr::UsdStageRefPtr& Stage)
 {
 	FScopedUsdAllocs UsdAllocs;
 
-	if ( RegisteredObjectsChangedKey.IsValid() )
+	if (RegisteredObjectsChangedKey.IsValid())
 	{
-		pxr::TfNotice::Revoke( RegisteredObjectsChangedKey );
+		pxr::TfNotice::Revoke(RegisteredObjectsChangedKey);
 	}
-	RegisteredObjectsChangedKey = pxr::TfNotice::Register( pxr::TfWeakPtr< FUsdListenerImpl >( this ), &FUsdListenerImpl::HandleObjectsChangedNotice, Stage );
+	RegisteredObjectsChangedKey = pxr::TfNotice::Register(
+		pxr::TfWeakPtr<FUsdListenerImpl>(this),
+		&FUsdListenerImpl::HandleObjectsChangedNotice,
+		Stage
+	);
 
-	if ( RegisteredStageEditTargetChangedKey.IsValid() )
+	if (RegisteredStageEditTargetChangedKey.IsValid())
 	{
-		pxr::TfNotice::Revoke( RegisteredStageEditTargetChangedKey );
+		pxr::TfNotice::Revoke(RegisteredStageEditTargetChangedKey);
 	}
-	RegisteredStageEditTargetChangedKey = pxr::TfNotice::Register( pxr::TfWeakPtr< FUsdListenerImpl >( this ), &FUsdListenerImpl::HandleStageEditTargetChangedNotice, Stage );
+	RegisteredStageEditTargetChangedKey = pxr::TfNotice::Register(
+		pxr::TfWeakPtr<FUsdListenerImpl>(this),
+		&FUsdListenerImpl::HandleStageEditTargetChangedNotice,
+		Stage
+	);
 
 	if (RegisteredSdfLayersChangedKey.IsValid())
 	{
 		pxr::TfNotice::Revoke(RegisteredSdfLayersChangedKey);
 	}
-	RegisteredSdfLayersChangedKey = pxr::TfNotice::Register( pxr::TfWeakPtr< FUsdListenerImpl >( this ), &FUsdListenerImpl::HandleLayersChangedNotice );
+	RegisteredSdfLayersChangedKey = pxr::TfNotice::Register(pxr::TfWeakPtr<FUsdListenerImpl>(this), &FUsdListenerImpl::HandleLayersChangedNotice);
 
 	if (RegisteredSdfLayerDirtinessChangedKey.IsValid())
 	{
 		pxr::TfNotice::Revoke(RegisteredSdfLayerDirtinessChangedKey);
 	}
-	RegisteredSdfLayerDirtinessChangedKey = pxr::TfNotice::Register( pxr::TfWeakPtr< FUsdListenerImpl >( this ), &FUsdListenerImpl::HandleLayerDirtinessChangedNotice );
+	RegisteredSdfLayerDirtinessChangedKey = pxr::TfNotice::Register(
+		pxr::TfWeakPtr<FUsdListenerImpl>(this),
+		&FUsdListenerImpl::HandleLayerDirtinessChangedNotice
+	);
 }
-#endif // #if USE_USD_SDK
+#endif	  // #if USE_USD_SDK
 
 PRAGMA_DISABLE_DEPRECATION_WARNINGS
 FUsdListenerImpl::~FUsdListenerImpl()
 {
 #if USE_USD_SDK
 	FScopedUsdAllocs UsdAllocs;
-	pxr::TfNotice::Revoke( RegisteredObjectsChangedKey );
-	pxr::TfNotice::Revoke( RegisteredStageEditTargetChangedKey );
-	pxr::TfNotice::Revoke( RegisteredSdfLayersChangedKey );
-	pxr::TfNotice::Revoke( RegisteredSdfLayerDirtinessChangedKey );
-#endif // #if USE_USD_SDK
+	pxr::TfNotice::Revoke(RegisteredObjectsChangedKey);
+	pxr::TfNotice::Revoke(RegisteredStageEditTargetChangedKey);
+	pxr::TfNotice::Revoke(RegisteredSdfLayersChangedKey);
+	pxr::TfNotice::Revoke(RegisteredSdfLayerDirtinessChangedKey);
+#endif	  // #if USE_USD_SDK
 }
+
 PRAGMA_ENABLE_DEPRECATION_WARNINGS
 
 #if USE_USD_SDK
-void FUsdListenerImpl::HandleObjectsChangedNotice( const pxr::UsdNotice::ObjectsChanged& Notice, const pxr::UsdStageWeakPtr& Sender )
+void FUsdListenerImpl::HandleObjectsChangedNotice(const pxr::UsdNotice::ObjectsChanged& Notice, const pxr::UsdStageWeakPtr& Sender)
 {
-	if ( !OnObjectsChanged.IsBound() )
+	if (!OnObjectsChanged.IsBound())
 	{
 		return;
 	}
 
-	if ( IsBlocked.GetValue() > 0 )
+	if (IsBlocked.GetValue() > 0)
 	{
 		return;
 	}
 
 	UsdUtils::FObjectChangesByPath InfoChanges;
 	UsdUtils::FObjectChangesByPath ResyncChanges;
-	UsdToUnreal::ConvertObjectsChangedNotice( Notice, InfoChanges, ResyncChanges );
-	if ( InfoChanges.Num() > 0 || ResyncChanges.Num() > 0 )
+	UsdToUnreal::ConvertObjectsChangedNotice(Notice, InfoChanges, ResyncChanges);
+	if (InfoChanges.Num() > 0 || ResyncChanges.Num() > 0)
 	{
 		FScopedUnrealAllocs UnrealAllocs;
-		OnObjectsChanged.Broadcast( InfoChanges, ResyncChanges );
+		OnObjectsChanged.Broadcast(InfoChanges, ResyncChanges);
 	}
 }
 
-void FUsdListenerImpl::HandleStageEditTargetChangedNotice( const pxr::UsdNotice::StageEditTargetChanged& Notice, const pxr::UsdStageWeakPtr& Sender )
+void FUsdListenerImpl::HandleStageEditTargetChangedNotice(const pxr::UsdNotice::StageEditTargetChanged& Notice, const pxr::UsdStageWeakPtr& Sender)
 {
 	FScopedUnrealAllocs UnrealAllocs;
 	OnStageEditTargetChanged.Broadcast();
 }
 
-void FUsdListenerImpl::HandleLayersChangedNotice( const pxr::SdfNotice::LayersDidChange& Notice )
+void FUsdListenerImpl::HandleLayersChangedNotice(const pxr::SdfNotice::LayersDidChange& Notice)
 {
-PRAGMA_DISABLE_DEPRECATION_WARNINGS
+	PRAGMA_DISABLE_DEPRECATION_WARNINGS
 	if ((!OnLayersChanged.IsBound() && !OnSdfLayersChanged.IsBound()) || IsBlocked.GetValue() > 0)
 	{
 		return;
@@ -434,9 +458,9 @@ PRAGMA_DISABLE_DEPRECATION_WARNINGS
 	}
 
 	FScopedUnrealAllocs UnrealAllocs;
-	OnLayersChanged.Broadcast( LayersNames );
+	OnLayersChanged.Broadcast(LayersNames);
 	OnSdfLayersChanged.Broadcast(ConvertedLayerToChangeList);
-PRAGMA_ENABLE_DEPRECATION_WARNINGS
+	PRAGMA_ENABLE_DEPRECATION_WARNINGS
 }
 
 void FUsdListenerImpl::HandleLayerDirtinessChangedNotice(const pxr::SdfNotice::LayerDirtinessChanged& Notice)
@@ -450,10 +474,10 @@ void FUsdListenerImpl::HandleLayerDirtinessChangedNotice(const pxr::SdfNotice::L
 
 	OnSdfLayerDirtinessChanged.Broadcast();
 }
-#endif // #if USE_USD_SDK
+#endif	  // #if USE_USD_SDK
 
-FScopedBlockNotices::FScopedBlockNotices( FUsdListener& InListener )
-	: Listener( InListener )
+FScopedBlockNotices::FScopedBlockNotices(FUsdListener& InListener)
+	: Listener(InListener)
 {
 	Listener.Block();
 }
diff --git a/Engine/Plugins/Importers/USDImporter/Source/UnrealUSDWrapper/Private/UnrealUSDWrapper.cpp b/Engine/Plugins/Importers/USDImporter/Source/UnrealUSDWrapper/Private/UnrealUSDWrapper.cpp
--- a/Engine/Plugins/Importers/USDImporter/Source/UnrealUSDWrapper/Private/UnrealUSDWrapper.cpp
+++ b/Engine/Plugins/Importers/USDImporter/Source/UnrealUSDWrapper/Private/UnrealUSDWrapper.cpp
@@ -7,9 +7,9 @@
 #include "USDMemory.h"
 #include "USDProjectSettings.h"
 
+#include "UsdWrappers/SdfLayer.h"
 #include "UsdWrappers/UsdAttribute.h"
 #include "UsdWrappers/UsdStage.h"
-#include "UsdWrappers/SdfLayer.h"
 
 #include "CineCameraComponent.h"
 #include "Components/DirectionalLightComponent.h"
@@ -68,9 +68,8 @@
 
 #include "USDIncludesEnd.h"
 
-
-using std::vector;
 using std::string;
+using std::vector;
 
 using namespace pxr;
 
@@ -126,38 +125,38 @@ namespace UnrealIdentifiers
 	const TfToken In = TfToken("in");
 	const TfToken Result = TfToken("result");
 	const TfToken File = TfToken("file");
-	const TfToken WrapT = TfToken( "wrapT" );
-	const TfToken WrapS = TfToken( "wrapS" );
-	const TfToken Repeat = TfToken( "repeat" );
-	const TfToken Mirror = TfToken( "mirror" );
-	const TfToken Clamp = TfToken( "clamp" );
+	const TfToken WrapT = TfToken("wrapT");
+	const TfToken WrapS = TfToken("wrapS");
+	const TfToken Repeat = TfToken("repeat");
+	const TfToken Mirror = TfToken("mirror");
+	const TfToken Clamp = TfToken("clamp");
 	const TfToken Fallback = TfToken("fallback");
 	const TfToken R = TfToken("r");
 	const TfToken RGB = TfToken("rgb");
-	const TfToken RawColorSpaceToken = TfToken{ "raw" };
-	const TfToken SRGBColorSpaceToken = TfToken{ "sRGB" };
-	const TfToken SourceColorSpaceToken = TfToken{ "sourceColorSpace" };
+	const TfToken RawColorSpaceToken = TfToken{"raw"};
+	const TfToken SRGBColorSpaceToken = TfToken{"sRGB"};
+	const TfToken SourceColorSpaceToken = TfToken{"sourceColorSpace"};
 
-	const TfToken UsdPreviewSurface = TfToken( "UsdPreviewSurface" );
-	const TfToken UsdTransform2d = TfToken( "UsdTransform2d" );
-	const TfToken UsdPrimvarReader_float2 = TfToken( "UsdPrimvarReader_float2" );
-	const TfToken UsdPrimvarReader_float3 = TfToken( "UsdPrimvarReader_float3" );
-	const TfToken UsdUVTexture = TfToken( "UsdUVTexture" );
+	const TfToken UsdPreviewSurface = TfToken("UsdPreviewSurface");
+	const TfToken UsdTransform2d = TfToken("UsdTransform2d");
+	const TfToken UsdPrimvarReader_float2 = TfToken("UsdPrimvarReader_float2");
+	const TfToken UsdPrimvarReader_float3 = TfToken("UsdPrimvarReader_float3");
+	const TfToken UsdUVTexture = TfToken("UsdUVTexture");
 
-	const TfToken WorldSpaceNormals = TfToken( "worldSpaceNormals" );
+	const TfToken WorldSpaceNormals = TfToken("worldSpaceNormals");
 	const TfToken PrimvarsNormals = TfToken("primvars:normals");
 	const TfToken PrimvarsPoints = TfToken("primvars:points");
 
-	const TfToken GroomAPI = TfToken( "GroomAPI" );
-	const TfToken GroomBindingAPI = TfToken( "GroomBindingAPI" );
-	const TfToken UnrealGroomToBind = TfToken( "unreal:groomBinding:groom" );
-	const TfToken UnrealGroomReferenceMesh = TfToken( "unreal:groomBinding:referenceMesh" );
+	const TfToken GroomAPI = TfToken("GroomAPI");
+	const TfToken GroomBindingAPI = TfToken("GroomBindingAPI");
+	const TfToken UnrealGroomToBind = TfToken("unreal:groomBinding:groom");
+	const TfToken UnrealGroomReferenceMesh = TfToken("unreal:groomBinding:referenceMesh");
 
-	const TfToken UnrealContentPath = TfToken( "unreal:contentPath" );
-	const TfToken UnrealAssetType = TfToken( "unreal:assetType" );
-	const TfToken UnrealExportTime = TfToken( "unreal:exportTime" );
-	const TfToken UnrealEngineVersion = TfToken( "unreal:engineVersion" );
-}
+	const TfToken UnrealContentPath = TfToken("unreal:contentPath");
+	const TfToken UnrealAssetType = TfToken("unreal:assetType");
+	const TfToken UnrealExportTime = TfToken("unreal:exportTime");
+	const TfToken UnrealEngineVersion = TfToken("unreal:engineVersion");
+}	 // namespace UnrealIdentifiers
 
 bool IUsdPrim::IsValidPrimName(const FString& Name, FText& OutReason)
 {
@@ -168,8 +167,8 @@ bool IUsdPrim::IsValidPrimName(const FString& Name, FText& OutReason)
 	}
 
 	const FString InvalidCharacters = TEXT("\\W");
-	FRegexPattern RegexPattern( InvalidCharacters );
-	FRegexMatcher RegexMatcher( RegexPattern, Name );
+	FRegexPattern RegexPattern(InvalidCharacters);
+	FRegexMatcher RegexMatcher(RegexPattern, Name);
 	if (RegexMatcher.FindNext())
 	{
 		OutReason = LOCTEXT("InvalidCharacter", "Can only use letters, numbers and underscore!");
@@ -185,7 +184,7 @@ bool IUsdPrim::IsValidPrimName(const FString& Name, FText& OutReason)
 	return true;
 }
 
-EUsdPurpose IUsdPrim::GetPurpose( const UsdPrim& Prim, bool bComputed )
+EUsdPurpose IUsdPrim::GetPurpose(const UsdPrim& Prim, bool bComputed)
 {
 	UsdGeomImageable Geom(Prim);
 	if (Geom)
@@ -245,7 +244,7 @@ int IUsdPrim::GetNumLODs(const UsdPrim& Prim)
 	if (Prim.HasVariantSets())
 	{
 		UsdVariantSet LODVariantSet = Prim.GetVariantSet(UnrealIdentifiers::LOD);
-		if(LODVariantSet.IsValid())
+		if (LODVariantSet.IsValid())
 		{
 			vector<string> VariantNames = LODVariantSet.GetVariantNames();
 			NumLODs = VariantNames.size();
@@ -261,10 +260,9 @@ bool IUsdPrim::IsKindChildOf(const UsdPrim& Prim, const std::string& InBaseKind)
 
 	KindRegistry& Registry = KindRegistry::GetInstance();
 
-	TfToken PrimKind( GetKind(Prim) );
+	TfToken PrimKind(GetKind(Prim));
 
 	return Registry.IsA(PrimKind, BaseKind);
-
 }
 
 TfToken IUsdPrim::GetKind(const pxr::UsdPrim& Prim)
@@ -279,7 +277,7 @@ TfToken IUsdPrim::GetKind(const pxr::UsdPrim& Prim)
 	else
 	{
 		// Prim is not a model, read kind directly from metadata
-		Prim.GetMetadata( SdfFieldKeys->Kind, &KindType );
+		Prim.GetMetadata(SdfFieldKeys->Kind, &KindType);
 	}
 
 	return KindType;
@@ -292,7 +290,7 @@ bool IUsdPrim::SetKind(const pxr::UsdPrim& Prim, const pxr::TfToken& Kind)
 	{
 		if (!Model.SetKind(Kind))
 		{
-			return Prim.SetMetadata( SdfFieldKeys->Kind, Kind );
+			return Prim.SetMetadata(SdfFieldKeys->Kind, Kind);
 		}
 
 		return true;
@@ -301,9 +299,9 @@ bool IUsdPrim::SetKind(const pxr::UsdPrim& Prim, const pxr::TfToken& Kind)
 	return false;
 }
 
-bool IUsdPrim::ClearKind( const pxr::UsdPrim& Prim )
+bool IUsdPrim::ClearKind(const pxr::UsdPrim& Prim)
 {
-	return Prim.ClearMetadata( SdfFieldKeys->Kind );
+	return Prim.ClearMetadata(SdfFieldKeys->Kind);
 }
 
 pxr::GfMatrix4d IUsdPrim::GetLocalTransform(const pxr::UsdPrim& Prim)
@@ -311,7 +309,7 @@ pxr::GfMatrix4d IUsdPrim::GetLocalTransform(const pxr::UsdPrim& Prim)
 	pxr::GfMatrix4d USDMatrix(1);
 
 	pxr::UsdGeomXformable XForm(Prim);
-	if(XForm)
+	if (XForm)
 	{
 		// Set transform
 		bool bResetXFormStack = false;
@@ -321,16 +319,15 @@ pxr::GfMatrix4d IUsdPrim::GetLocalTransform(const pxr::UsdPrim& Prim)
 	return USDMatrix;
 }
 
-pxr::GfMatrix4d IUsdPrim::GetLocalToWorldTransform(const pxr::UsdPrim& Prim )
+pxr::GfMatrix4d IUsdPrim::GetLocalToWorldTransform(const pxr::UsdPrim& Prim)
 {
-	return GetLocalToWorldTransform( Prim, pxr::UsdTimeCode::Default().GetValue() );
+	return GetLocalToWorldTransform(Prim, pxr::UsdTimeCode::Default().GetValue());
 }
 
 pxr::GfMatrix4d IUsdPrim::GetLocalToWorldTransform(const pxr::UsdPrim& Prim, double Time)
 {
 	pxr::SdfPath AbsoluteRootPath = pxr::SdfPath::AbsoluteRootPath();
 	return GetLocalToWorldTransform(Prim, Time, AbsoluteRootPath);
-
 }
 
 pxr::GfMatrix4d IUsdPrim::GetLocalToWorldTransform(const pxr::UsdPrim& Prim, double Time, const pxr::SdfPath& AbsoluteRootPath)
@@ -372,7 +369,7 @@ bool IUsdPrim::SetActiveLODIndex(const pxr::UsdPrim& Prim, int LODIndex)
 			vector<string> VariantNames = LODVariantSet.GetVariantNames();
 
 			bool bResult = false;
-			if(LODIndex < VariantNames.size())
+			if (LODIndex < VariantNames.size())
 			{
 				bResult = LODVariantSet.SetVariantSelection(VariantNames[LODIndex]);
 			}
@@ -384,7 +381,7 @@ bool IUsdPrim::SetActiveLODIndex(const pxr::UsdPrim& Prim, int LODIndex)
 
 EUsdGeomOrientation IUsdPrim::GetGeometryOrientation(const pxr::UsdGeomMesh& Mesh)
 {
-	return GetGeometryOrientation( Mesh, pxr::UsdTimeCode::Default().GetValue() );
+	return GetGeometryOrientation(Mesh, pxr::UsdTimeCode::Default().GetValue());
 }
 
 EUsdGeomOrientation IUsdPrim::GetGeometryOrientation(const pxr::UsdGeomMesh& Mesh, double Time)
@@ -394,7 +391,7 @@ EUsdGeomOrientation IUsdPrim::GetGeometryOrientation(const pxr::UsdGeomMesh& Mes
 	if (Mesh)
 	{
 		UsdAttribute Orientation = Mesh.GetOrientationAttr();
-		if(Orientation)
+		if (Orientation)
 		{
 			static TfToken RightHanded("rightHanded");
 			static TfToken LeftHanded("leftHanded");
@@ -408,7 +405,7 @@ EUsdGeomOrientation IUsdPrim::GetGeometryOrientation(const pxr::UsdGeomMesh& Mes
 
 	return GeomOrientation;
 }
-#endif // USE_USD_SDK
+#endif	  // USE_USD_SDK
 
 const TCHAR* UnrealIdentifiers::LayerSavedComment = TEXT("unreal:layerSaved");
 const TCHAR* UnrealIdentifiers::TwoSidedMaterialSuffix = TEXT("_TwoSided");
@@ -418,32 +415,32 @@ const TCHAR* UnrealIdentifiers::Inherited = TEXT("inherited");
 const TCHAR* UnrealIdentifiers::IdentifierPrefix = TEXT("@identifier:");
 const TCHAR* UnrealIdentifiers::ExportedSkeletonPrimName = TEXT("Skel");
 
-FName UnrealIdentifiers::TransformPropertyName = TEXT( "Transform" ); // Fake FName because the transform is stored decomposed on the component
-FName UnrealIdentifiers::HiddenInGamePropertyName = GET_MEMBER_NAME_CHECKED( USceneComponent, bHiddenInGame );
+FName UnrealIdentifiers::TransformPropertyName = TEXT("Transform");	   // Fake FName because the transform is stored decomposed on the component
+FName UnrealIdentifiers::HiddenInGamePropertyName = GET_MEMBER_NAME_CHECKED(USceneComponent, bHiddenInGame);
 FName UnrealIdentifiers::HiddenPropertyName = AActor::GetHiddenPropertyName();
 
-FName UnrealIdentifiers::CurrentFocalLengthPropertyName = GET_MEMBER_NAME_CHECKED( UCineCameraComponent, CurrentFocalLength );
-FName UnrealIdentifiers::ManualFocusDistancePropertyName = GET_MEMBER_NAME_CHECKED( UCineCameraComponent, FocusSettings.ManualFocusDistance );
-FName UnrealIdentifiers::CurrentAperturePropertyName = GET_MEMBER_NAME_CHECKED( UCineCameraComponent, CurrentAperture );
-FName UnrealIdentifiers::SensorWidthPropertyName = GET_MEMBER_NAME_CHECKED( UCineCameraComponent, Filmback.SensorWidth );
-FName UnrealIdentifiers::SensorHeightPropertyName = GET_MEMBER_NAME_CHECKED( UCineCameraComponent, Filmback.SensorHeight );
-
-FName UnrealIdentifiers::IntensityPropertyName = GET_MEMBER_NAME_CHECKED( ULightComponentBase, Intensity );
-FName UnrealIdentifiers::LightColorPropertyName = GET_MEMBER_NAME_CHECKED( ULightComponentBase, LightColor );
-FName UnrealIdentifiers::UseTemperaturePropertyName = GET_MEMBER_NAME_CHECKED( ULightComponent, bUseTemperature );
-FName UnrealIdentifiers::TemperaturePropertyName = GET_MEMBER_NAME_CHECKED( ULightComponent, Temperature );
-FName UnrealIdentifiers::SourceWidthPropertyName = GET_MEMBER_NAME_CHECKED( URectLightComponent, SourceWidth );
-FName UnrealIdentifiers::SourceHeightPropertyName = GET_MEMBER_NAME_CHECKED( URectLightComponent, SourceHeight );
-FName UnrealIdentifiers::SourceRadiusPropertyName = GET_MEMBER_NAME_CHECKED( UPointLightComponent, SourceRadius );
-FName UnrealIdentifiers::OuterConeAnglePropertyName = GET_MEMBER_NAME_CHECKED( USpotLightComponent, OuterConeAngle );
-FName UnrealIdentifiers::InnerConeAnglePropertyName = GET_MEMBER_NAME_CHECKED( USpotLightComponent, InnerConeAngle );
-FName UnrealIdentifiers::LightSourceAnglePropertyName = GET_MEMBER_NAME_CHECKED( UDirectionalLightComponent, LightSourceAngle );
-
-FString UnrealIdentifiers::MaterialAllPurposeText = TEXT( "allPurpose" );
+FName UnrealIdentifiers::CurrentFocalLengthPropertyName = GET_MEMBER_NAME_CHECKED(UCineCameraComponent, CurrentFocalLength);
+FName UnrealIdentifiers::ManualFocusDistancePropertyName = GET_MEMBER_NAME_CHECKED(UCineCameraComponent, FocusSettings.ManualFocusDistance);
+FName UnrealIdentifiers::CurrentAperturePropertyName = GET_MEMBER_NAME_CHECKED(UCineCameraComponent, CurrentAperture);
+FName UnrealIdentifiers::SensorWidthPropertyName = GET_MEMBER_NAME_CHECKED(UCineCameraComponent, Filmback.SensorWidth);
+FName UnrealIdentifiers::SensorHeightPropertyName = GET_MEMBER_NAME_CHECKED(UCineCameraComponent, Filmback.SensorHeight);
+
+FName UnrealIdentifiers::IntensityPropertyName = GET_MEMBER_NAME_CHECKED(ULightComponentBase, Intensity);
+FName UnrealIdentifiers::LightColorPropertyName = GET_MEMBER_NAME_CHECKED(ULightComponentBase, LightColor);
+FName UnrealIdentifiers::UseTemperaturePropertyName = GET_MEMBER_NAME_CHECKED(ULightComponent, bUseTemperature);
+FName UnrealIdentifiers::TemperaturePropertyName = GET_MEMBER_NAME_CHECKED(ULightComponent, Temperature);
+FName UnrealIdentifiers::SourceWidthPropertyName = GET_MEMBER_NAME_CHECKED(URectLightComponent, SourceWidth);
+FName UnrealIdentifiers::SourceHeightPropertyName = GET_MEMBER_NAME_CHECKED(URectLightComponent, SourceHeight);
+FName UnrealIdentifiers::SourceRadiusPropertyName = GET_MEMBER_NAME_CHECKED(UPointLightComponent, SourceRadius);
+FName UnrealIdentifiers::OuterConeAnglePropertyName = GET_MEMBER_NAME_CHECKED(USpotLightComponent, OuterConeAngle);
+FName UnrealIdentifiers::InnerConeAnglePropertyName = GET_MEMBER_NAME_CHECKED(USpotLightComponent, InnerConeAngle);
+FName UnrealIdentifiers::LightSourceAnglePropertyName = GET_MEMBER_NAME_CHECKED(UDirectionalLightComponent, LightSourceAngle);
+
+FString UnrealIdentifiers::MaterialAllPurposeText = TEXT("allPurpose");
 #if USE_USD_SDK
-FString UnrealIdentifiers::MaterialAllPurpose = ANSI_TO_TCHAR( pxr::UsdShadeTokens->allPurpose.GetString().c_str() );
-FString UnrealIdentifiers::MaterialPreviewPurpose = ANSI_TO_TCHAR( pxr::UsdShadeTokens->preview.GetString().c_str() );
-FString UnrealIdentifiers::MaterialFullPurpose = ANSI_TO_TCHAR( pxr::UsdShadeTokens->full.GetString().c_str() );
+FString UnrealIdentifiers::MaterialAllPurpose = ANSI_TO_TCHAR(pxr::UsdShadeTokens->allPurpose.GetString().c_str());
+FString UnrealIdentifiers::MaterialPreviewPurpose = ANSI_TO_TCHAR(pxr::UsdShadeTokens->preview.GetString().c_str());
+FString UnrealIdentifiers::MaterialFullPurpose = ANSI_TO_TCHAR(pxr::UsdShadeTokens->full.GetString().c_str());
 FString UnrealIdentifiers::PrimvarsDisplayColor = ANSI_TO_TCHAR(pxr::UsdGeomTokens->primvarsDisplayColor.GetString().c_str());
 FString UnrealIdentifiers::PrimvarsDisplayOpacity = ANSI_TO_TCHAR(pxr::UsdGeomTokens->primvarsDisplayOpacity.GetString().c_str());
 FString UnrealIdentifiers::DoubleSided = ANSI_TO_TCHAR(pxr::UsdGeomTokens->doubleSided.GetString().c_str());
@@ -451,16 +448,16 @@ FString UnrealIdentifiers::ModelDrawMode = ANSI_TO_TCHAR(pxr::UsdGeomTokens->mod
 FString UnrealIdentifiers::ModelApplyDrawMode = ANSI_TO_TCHAR(pxr::UsdGeomTokens->modelApplyDrawMode.GetString().c_str());
 FString UnrealIdentifiers::UsdNamespaceDelimiter = ANSI_TO_TCHAR(pxr::SdfPathTokens->namespaceDelimiter.GetString().c_str());
 #else
-FString UnrealIdentifiers::MaterialAllPurpose = TEXT( "" );
-FString UnrealIdentifiers::MaterialPreviewPurpose = TEXT( "preview" );
-FString UnrealIdentifiers::MaterialFullPurpose = TEXT( "full" );
+FString UnrealIdentifiers::MaterialAllPurpose = TEXT("");
+FString UnrealIdentifiers::MaterialPreviewPurpose = TEXT("preview");
+FString UnrealIdentifiers::MaterialFullPurpose = TEXT("full");
 FString UnrealIdentifiers::PrimvarsDisplayColor = TEXT("primvars:displayColor");
 FString UnrealIdentifiers::PrimvarsDisplayOpacity = TEXT("primvars:displayOpacity");
 FString UnrealIdentifiers::DoubleSided = TEXT("doubleSided");
 FString UnrealIdentifiers::ModelDrawMode = TEXT("model:drawMode");
 FString UnrealIdentifiers::ModelApplyDrawMode = TEXT("model:applyDrawMode");
 FString UnrealIdentifiers::UsdNamespaceDelimiter = TEXT(":");
-#endif // USE_USD_SDK
+#endif	  // USE_USD_SDK
 
 FUsdDelegates::FUsdImportDelegate FUsdDelegates::OnPreUsdImport;
 FUsdDelegates::FUsdImportDelegate FUsdDelegates::OnPostUsdImport;
@@ -470,8 +467,13 @@ namespace UsdWrapperUtils
 	void CheckIfForceDisabled()
 	{
 #if USD_FORCE_DISABLED
-		UE_LOG( LogUsd, Error, TEXT( "The USD SDK is disabled because the executable is not forcing the ansi C allocator (you need to set 'FORCE_ANSI_ALLOCATOR=1' as a global definition on your project *.Target.cs file). Read the comments at the end of UnrealUSDWrapper.Build.cs for more details." ) );
-#endif // USD_FORCE_DISABLED
+		UE_LOG(
+			LogUsd,
+			Error,
+			TEXT("The USD SDK is disabled because the executable is not forcing the ansi C allocator (you need to set 'FORCE_ANSI_ALLOCATOR=1' as a "
+				 "global definition on your project *.Target.cs file). Read the comments at the end of UnrealUSDWrapper.Build.cs for more details.")
+		);
+#endif	  // USD_FORCE_DISABLED
 	}
 }
 
@@ -479,7 +481,8 @@ namespace UsdWrapperUtils
 class FUsdDiagnosticDelegate : public pxr::TfDiagnosticMgr::Delegate
 {
 public:
-	virtual ~FUsdDiagnosticDelegate() override {};
+	virtual ~FUsdDiagnosticDelegate() override{};
+
 	virtual void IssueError(const pxr::TfError& Error) override
 	{
 		FScopedUsdAllocs Allocs;
@@ -488,12 +491,14 @@ public:
 		Msg += ": ";
 		Msg += Error.GetCommentary();
 
-		UE_LOG(LogUsd, Error, TEXT("%s"), ANSI_TO_TCHAR( Msg.c_str() ));
+		UE_LOG(LogUsd, Error, TEXT("%s"), ANSI_TO_TCHAR(Msg.c_str()));
 	}
+
 	virtual void IssueFatalError(const pxr::TfCallContext& Context, const std::string& Msg) override
 	{
-		UE_LOG(LogUsd, Error, TEXT("%s"), ANSI_TO_TCHAR( Msg.c_str() ));
+		UE_LOG(LogUsd, Error, TEXT("%s"), ANSI_TO_TCHAR(Msg.c_str()));
 	}
+
 	virtual void IssueStatus(const pxr::TfStatus& Status) override
 	{
 		FScopedUsdAllocs Allocs;
@@ -502,8 +507,9 @@ public:
 		Msg += ": ";
 		Msg += Status.GetCommentary();
 
-		UE_LOG(LogUsd, Log, TEXT("%s"), ANSI_TO_TCHAR( Msg.c_str() ));
+		UE_LOG(LogUsd, Log, TEXT("%s"), ANSI_TO_TCHAR(Msg.c_str()));
 	}
+
 	virtual void IssueWarning(const pxr::TfWarning& Warning) override
 	{
 		FScopedUsdAllocs Allocs;
@@ -512,12 +518,14 @@ public:
 		Msg += ": ";
 		Msg += Warning.GetCommentary();
 
-		UE_LOG(LogUsd, Warning, TEXT("%s"), ANSI_TO_TCHAR( Msg.c_str() ));
+		UE_LOG(LogUsd, Warning, TEXT("%s"), ANSI_TO_TCHAR(Msg.c_str()));
 	}
 };
 #else
-class FUsdDiagnosticDelegate { };
-#endif // USE_USD_SDK
+class FUsdDiagnosticDelegate
+{
+};
+#endif	  // USE_USD_SDK
 
 TUniquePtr<FUsdDiagnosticDelegate> UnrealUSDWrapper::Delegate = nullptr;
 
@@ -527,7 +535,7 @@ double UnrealUSDWrapper::GetDefaultTimeCode()
 {
 	return UsdTimeCode::Default().GetValue();
 }
-#endif // USE_USD_SDK
+#endif	  // USE_USD_SDK
 
 TArray<FString> UnrealUSDWrapper::RegisterPlugins(const FString& PathToPlugInfo)
 {
@@ -549,14 +557,13 @@ TArray<FString> UnrealUSDWrapper::RegisterPlugins(const TArray<FString>& PathsTo
 		UsdPathsToPlugInfo.emplace_back(TCHAR_TO_UTF8(*PathToPlugInfo));
 	}
 
-	const pxr::PlugPluginPtrVector UsdPlugins =
-		pxr::PlugRegistry::GetInstance().RegisterPlugins(UsdPathsToPlugInfo);
+	const pxr::PlugPluginPtrVector UsdPlugins = pxr::PlugRegistry::GetInstance().RegisterPlugins(UsdPathsToPlugInfo);
 
 	for (const pxr::PlugPluginPtr& UsdPlugin : UsdPlugins)
 	{
 		PluginNames.Emplace(UTF8_TO_TCHAR(UsdPlugin->GetName().c_str()));
 	}
-#endif // USE_USD_SDK
+#endif	  // USE_USD_SDK
 
 	return PluginNames;
 }
@@ -569,18 +576,18 @@ TArray<FString> UnrealUSDWrapper::GetAllSupportedFileFormats()
 	FScopedUsdAllocs Allocs;
 
 	std::set<std::string> Extensions = pxr::SdfFileFormat::FindAllFileFormatExtensions();
-	for ( const std::string& Ext : Extensions )
+	for (const std::string& Ext : Extensions)
 	{
 		// Ignore formats that don't target "usd"
 		pxr::SdfFileFormatConstPtr Format = pxr::SdfFileFormat::FindByExtension(Ext, pxr::UsdUsdFileFormatTokens->Target);
-		if ( Format == nullptr )
+		if (Format == nullptr)
 		{
 			continue;
 		}
 
-		Result.Emplace( ANSI_TO_TCHAR( Ext.c_str() ) );
+		Result.Emplace(ANSI_TO_TCHAR(Ext.c_str()));
 	}
-#endif // #if USE_USD_SDK
+#endif	  // #if USE_USD_SDK
 
 	return Result;
 }
@@ -601,25 +608,25 @@ TArray<FString> UnrealUSDWrapper::GetNativeFileFormats()
 
 	std::set<std::string> FileExtensions;
 
-	for ( const pxr::TfToken& FormatId : NativeFormatIds )
+	for (const pxr::TfToken& FormatId : NativeFormatIds)
 	{
-		const pxr::SdfFileFormatConstPtr Format = pxr::SdfFileFormat::FindById( FormatId );
-		if ( !Format )
+		const pxr::SdfFileFormatConstPtr Format = pxr::SdfFileFormat::FindById(FormatId);
+		if (!Format)
 		{
 			continue;
 		}
 
-		for ( const std::string& FileExtension : Format->GetFileExtensions() )
+		for (const std::string& FileExtension : Format->GetFileExtensions())
 		{
-			FileExtensions.insert( FileExtension );
+			FileExtensions.insert(FileExtension);
 		}
 	}
 
-	for ( const std::string& FileExtension : FileExtensions )
+	for (const std::string& FileExtension : FileExtensions)
 	{
-		Result.Emplace( ANSI_TO_TCHAR( FileExtension.c_str() ) );
+		Result.Emplace(ANSI_TO_TCHAR(FileExtension.c_str()));
 	}
-#endif // #if USE_USD_SDK
+#endif	  // #if USE_USD_SDK
 
 	return Result;
 }
@@ -635,7 +642,7 @@ namespace UE::UnrealUSDWrapper::Private
 		bool bForceReloadLayersFromDisk
 	)
 	{
-		if ( !RootIdentifier || FCString::Strlen( RootIdentifier ) == 0 )
+		if (!RootIdentifier || FCString::Strlen(RootIdentifier) == 0)
 		{
 			return UE::FUsdStage();
 		}
@@ -647,100 +654,100 @@ namespace UE::UnrealUSDWrapper::Private
 		pxr::UsdStageRefPtr Stage;
 
 		TOptional<pxr::UsdStageCacheContext> StageCacheContext;
-		if ( bUseStageCache )
+		if (bUseStageCache)
 		{
-			StageCacheContext.Emplace( pxr::UsdUtilsStageCache::Get() );
+			StageCacheContext.Emplace(pxr::UsdUtilsStageCache::Get());
 		}
 
 		pxr::UsdStagePopulationMask Mask;
-		if( PopulationMask )
+		if (PopulationMask)
 		{
 			// The USD OpenMasked functions don't actually consult or populate the stage cache
-			ensure( bUseStageCache == false );
+			ensure(bUseStageCache == false);
 
-			for ( const FString& AllowedPrimPath : *PopulationMask )
+			for (const FString& AllowedPrimPath : *PopulationMask)
 			{
-				Mask.Add( pxr::SdfPath{ TCHAR_TO_ANSI( *AllowedPrimPath ) } );
+				Mask.Add(pxr::SdfPath{TCHAR_TO_ANSI(*AllowedPrimPath)});
 			}
 		}
 
-		static_assert( ( int ) pxr::UsdStage::InitialLoadSet::LoadAll == ( int ) EUsdInitialLoadSet::LoadAll );
-		static_assert( ( int ) pxr::UsdStage::InitialLoadSet::LoadNone == ( int ) EUsdInitialLoadSet::LoadNone );
-		pxr::UsdStage::InitialLoadSet LoadSet = static_cast< pxr::UsdStage::InitialLoadSet >( InitialLoadSet );
+		static_assert((int)pxr::UsdStage::InitialLoadSet::LoadAll == (int)EUsdInitialLoadSet::LoadAll);
+		static_assert((int)pxr::UsdStage::InitialLoadSet::LoadNone == (int)EUsdInitialLoadSet::LoadNone);
+		pxr::UsdStage::InitialLoadSet LoadSet = static_cast<pxr::UsdStage::InitialLoadSet>(InitialLoadSet);
 
-		FString IdentifierStr = FString( RootIdentifier );
-		if ( FPaths::FileExists( IdentifierStr ) )
+		FString IdentifierStr = FString(RootIdentifier);
+		if (FPaths::FileExists(IdentifierStr))
 		{
-			if( PopulationMask )
+			if (PopulationMask)
 			{
-				Stage = pxr::UsdStage::OpenMasked( TCHAR_TO_ANSI( *IdentifierStr ), Mask, LoadSet );
+				Stage = pxr::UsdStage::OpenMasked(TCHAR_TO_ANSI(*IdentifierStr), Mask, LoadSet);
 			}
 			else
 			{
-				Stage = pxr::UsdStage::Open( TCHAR_TO_ANSI( *IdentifierStr ), LoadSet );
+				Stage = pxr::UsdStage::Open(TCHAR_TO_ANSI(*IdentifierStr), LoadSet);
 			}
 		}
 		else
 		{
-			FString SessionIdentifierStr = FString{ SessionIdentifier };
+			FString SessionIdentifierStr = FString{SessionIdentifier};
 
-			IdentifierStr.RemoveFromStart( UnrealIdentifiers::IdentifierPrefix );
-			SessionIdentifierStr.RemoveFromStart( UnrealIdentifiers::IdentifierPrefix );
+			IdentifierStr.RemoveFromStart(UnrealIdentifiers::IdentifierPrefix);
+			SessionIdentifierStr.RemoveFromStart(UnrealIdentifiers::IdentifierPrefix);
 
-			pxr::SdfLayerRefPtr RootLayer = pxr::SdfLayer::Find( TCHAR_TO_ANSI( *IdentifierStr ) );
-			pxr::SdfLayerRefPtr SessionLayer = pxr::SdfLayer::Find( TCHAR_TO_ANSI( *SessionIdentifierStr ) );
-			if ( RootLayer )
+			pxr::SdfLayerRefPtr RootLayer = pxr::SdfLayer::Find(TCHAR_TO_ANSI(*IdentifierStr));
+			pxr::SdfLayerRefPtr SessionLayer = pxr::SdfLayer::Find(TCHAR_TO_ANSI(*SessionIdentifierStr));
+			if (RootLayer)
 			{
-				if ( PopulationMask )
+				if (PopulationMask)
 				{
 					// We use this additional check so we don't have to replicate USD's "_CreateAnonymousSessionLayer"
 					// Basically we can't pass an invalid session layer pointer here the stage will actually end up
 					// with no session layer at all
-					if ( SessionLayer )
+					if (SessionLayer)
 					{
-						Stage = pxr::UsdStage::OpenMasked( RootLayer, SessionLayer, Mask, LoadSet );
+						Stage = pxr::UsdStage::OpenMasked(RootLayer, SessionLayer, Mask, LoadSet);
 					}
 					else
 					{
-						Stage = pxr::UsdStage::OpenMasked( RootLayer, Mask, LoadSet );
+						Stage = pxr::UsdStage::OpenMasked(RootLayer, Mask, LoadSet);
 					}
 				}
 				else
 				{
-					if( SessionLayer )
+					if (SessionLayer)
 					{
-						Stage = pxr::UsdStage::Open( RootLayer, SessionLayer, LoadSet );
+						Stage = pxr::UsdStage::Open(RootLayer, SessionLayer, LoadSet);
 					}
 					else
 					{
-						Stage = pxr::UsdStage::Open( RootLayer, LoadSet );
+						Stage = pxr::UsdStage::Open(RootLayer, LoadSet);
 					}
 				}
 			}
 		}
 
-		if ( bForceReloadLayersFromDisk && Stage )
+		if (bForceReloadLayersFromDisk && Stage)
 		{
 			// Layers are cached in the layer registry independently of the stage cache. If the layer is already in
 			// the registry by the time we try to open a stage, even if we're not using a stage cache at all the
 			// layer will be reused and the file will *not* be re-read, so here we manually reload them.
 			pxr::SdfLayerHandleVector StageLayers = Stage->GetLayerStack();
-			for ( pxr::SdfLayerHandle StageLayer : StageLayers )
+			for (pxr::SdfLayerHandle StageLayer : StageLayers)
 			{
-				if ( LoadedLayers.count( StageLayer ) > 0 )
+				if (LoadedLayers.count(StageLayer) > 0)
 				{
 					StageLayer->Reload();
 				}
 			}
 		}
 
-		return UE::FUsdStage( Stage );
+		return UE::FUsdStage(Stage);
 #else
 		UsdWrapperUtils::CheckIfForceDisabled();
 		return UE::FUsdStage();
-#endif // #if USE_USD_SDK
+#endif	  // #if USE_USD_SDK
 	}
-}
+}	 // namespace UE::UnrealUSDWrapper::Private
 
 UE::FUsdStage UnrealUSDWrapper::OpenStage(
 	const TCHAR* Identifier,
@@ -761,29 +768,29 @@ UE::FUsdStage UnrealUSDWrapper::OpenStage(
 	);
 }
 
- UE::FUsdStage UnrealUSDWrapper::OpenStage(
-	 UE::FSdfLayer RootLayer,
-	 UE::FSdfLayer SessionLayer,
-	 EUsdInitialLoadSet InitialLoadSet,
-	 bool bUseStageCache,
-	 bool bForceReloadLayersFromDisk
- )
+UE::FUsdStage UnrealUSDWrapper::OpenStage(
+	UE::FSdfLayer RootLayer,
+	UE::FSdfLayer SessionLayer,
+	EUsdInitialLoadSet InitialLoadSet,
+	bool bUseStageCache,
+	bool bForceReloadLayersFromDisk
+)
 {
-	 if( !RootLayer )
-	 {
-		 return {};
-	 }
-
-	 const TArray<FString>* PopulationMask = nullptr;
-	 const TCHAR* SessionIdentifier = nullptr;
-	 return UE::UnrealUSDWrapper::Private::OpenStageImpl(
-		 *RootLayer.GetIdentifier(),
-		 SessionLayer ? *SessionLayer.GetIdentifier() : nullptr,
-		 InitialLoadSet,
-		 bUseStageCache,
-		 nullptr,
-		 bForceReloadLayersFromDisk
-	 );
+	if (!RootLayer)
+	{
+		return {};
+	}
+
+	const TArray<FString>* PopulationMask = nullptr;
+	const TCHAR* SessionIdentifier = nullptr;
+	return UE::UnrealUSDWrapper::Private::OpenStageImpl(
+		*RootLayer.GetIdentifier(),
+		SessionLayer ? *SessionLayer.GetIdentifier() : nullptr,
+		InitialLoadSet,
+		bUseStageCache,
+		nullptr,
+		bForceReloadLayersFromDisk
+	);
 }
 
 UE::FUsdStage UnrealUSDWrapper::OpenMaskedStage(
@@ -813,7 +820,7 @@ UE::FUsdStage UnrealUSDWrapper::OpenMaskedStage(
 	bool bForceReloadLayersFromDisk
 )
 {
-	if ( !RootLayer )
+	if (!RootLayer)
 	{
 		return {};
 	}
@@ -830,22 +837,22 @@ UE::FUsdStage UnrealUSDWrapper::OpenMaskedStage(
 	);
 }
 
-UE::FUsdStage UnrealUSDWrapper::NewStage( const TCHAR* FilePath )
+UE::FUsdStage UnrealUSDWrapper::NewStage(const TCHAR* FilePath)
 {
 #if USE_USD_SDK
 	FScopedUsdAllocs UsdAllocs;
 
-	UE::FUsdStage UsdStage( pxr::UsdStage::CreateNew( TCHAR_TO_ANSI( FilePath ) ) );
-	if ( UsdStage )
+	UE::FUsdStage UsdStage(pxr::UsdStage::CreateNew(TCHAR_TO_ANSI(FilePath)));
+	if (UsdStage)
 	{
-		pxr::UsdGeomSetStageUpAxis( UsdStage, pxr::UsdGeomTokens->z );
+		pxr::UsdGeomSetStageUpAxis(UsdStage, pxr::UsdGeomTokens->z);
 	}
 
 	return UsdStage;
 #else
 	UsdWrapperUtils::CheckIfForceDisabled();
 	return UE::FUsdStage();
-#endif // #if USE_USD_SDK
+#endif	  // #if USE_USD_SDK
 }
 
 UE::FUsdStage UnrealUSDWrapper::NewStage()
@@ -853,17 +860,17 @@ UE::FUsdStage UnrealUSDWrapper::NewStage()
 #if USE_USD_SDK
 	FScopedUsdAllocs UsdAllocs;
 
-	UE::FUsdStage UsdStage( pxr::UsdStage::CreateInMemory() );
-	if ( UsdStage )
+	UE::FUsdStage UsdStage(pxr::UsdStage::CreateInMemory());
+	if (UsdStage)
 	{
-		pxr::UsdGeomSetStageUpAxis( UsdStage, pxr::UsdGeomTokens->z );
+		pxr::UsdGeomSetStageUpAxis(UsdStage, pxr::UsdGeomTokens->z);
 	}
 
 	return UsdStage;
 #else
 	UsdWrapperUtils::CheckIfForceDisabled();
 	return UE::FUsdStage();
-#endif // #if USE_USD_SDK
+#endif	  // #if USE_USD_SDK
 }
 
 UE::FUsdStage UnrealUSDWrapper::GetClipboardStage()
@@ -873,57 +880,57 @@ UE::FUsdStage UnrealUSDWrapper::GetClipboardStage()
 #if USE_USD_SDK
 	FScopedUsdAllocs UsdAllocs;
 
-	if ( !ClipboardStage )
+	if (!ClipboardStage)
 	{
 		// Use this specific name so that we can check for it on our automated tests
-		pxr::UsdStageCacheContext StageCacheContext{ pxr::UsdUtilsStageCache::Get() };
-		ClipboardStage = UE::FUsdStage{ pxr::UsdStage::CreateInMemory( TCHAR_TO_ANSI( TEXT( "unreal_clipboard_layer" ) ) ) };
+		pxr::UsdStageCacheContext StageCacheContext{pxr::UsdUtilsStageCache::Get()};
+		ClipboardStage = UE::FUsdStage{pxr::UsdStage::CreateInMemory(TCHAR_TO_ANSI(TEXT("unreal_clipboard_layer")))};
 
-		pxr::UsdGeomSetStageUpAxis( ClipboardStage, pxr::UsdGeomTokens->z );
+		pxr::UsdGeomSetStageUpAxis(ClipboardStage, pxr::UsdGeomTokens->z);
 	}
-#endif // #if USE_USD_SDK
+#endif	  // #if USE_USD_SDK
 
 	return ClipboardStage;
 }
 
-TArray< UE::FUsdStage > UnrealUSDWrapper::GetAllStagesFromCache()
+TArray<UE::FUsdStage> UnrealUSDWrapper::GetAllStagesFromCache()
 {
-	TArray< UE::FUsdStage > StagesInCache;
+	TArray<UE::FUsdStage> StagesInCache;
 
 #if USE_USD_SDK
 	FScopedUsdAllocs UsdAllocs;
 
-	for ( const pxr::UsdStageRefPtr& StageInCache : pxr::UsdUtilsStageCache::Get().GetAllStages() )
+	for (const pxr::UsdStageRefPtr& StageInCache : pxr::UsdUtilsStageCache::Get().GetAllStages())
 	{
-		StagesInCache.Emplace( StageInCache );
+		StagesInCache.Emplace(StageInCache);
 	}
-#endif // #if USE_USD_SDK
+#endif	  // #if USE_USD_SDK
 
 	return StagesInCache;
 }
 
-void UnrealUSDWrapper::EraseStageFromCache( const UE::FUsdStage& Stage )
+void UnrealUSDWrapper::EraseStageFromCache(const UE::FUsdStage& Stage)
 {
 #if USE_USD_SDK
-	pxr::UsdUtilsStageCache::Get().Erase( Stage );
-#endif // #if USE_USD_SDK
+	pxr::UsdUtilsStageCache::Get().Erase(Stage);
+#endif	  // #if USE_USD_SDK
 }
 
-void UnrealUSDWrapper::SetDefaultResolverDefaultSearchPath( const TArray<FDirectoryPath>& SearchPath )
+void UnrealUSDWrapper::SetDefaultResolverDefaultSearchPath(const TArray<FDirectoryPath>& SearchPath)
 {
 #if USE_USD_SDK
 	FScopedUsdAllocs UsdAllocs;
 
-	std::vector< std::string > DefaultResolverSearchPath;
-	DefaultResolverSearchPath.reserve( SearchPath.Num() );
+	std::vector<std::string> DefaultResolverSearchPath;
+	DefaultResolverSearchPath.reserve(SearchPath.Num());
 
-	for ( const FDirectoryPath& Directory : SearchPath )
+	for (const FDirectoryPath& Directory : SearchPath)
 	{
-		DefaultResolverSearchPath.push_back( TCHAR_TO_UTF8( *Directory.Path ) );
+		DefaultResolverSearchPath.push_back(TCHAR_TO_UTF8(*Directory.Path));
 	}
 
-	pxr::ArDefaultResolver::SetDefaultSearchPath( DefaultResolverSearchPath );
-#endif // #if USE_USD_SDK
+	pxr::ArDefaultResolver::SetDefaultSearchPath(DefaultResolverSearchPath);
+#endif	  // #if USE_USD_SDK
 }
 
 void UnrealUSDWrapper::SetupDiagnosticDelegate()
@@ -938,7 +945,7 @@ void UnrealUSDWrapper::SetupDiagnosticDelegate()
 
 	pxr::TfDiagnosticMgr& DiagMgr = pxr::TfDiagnosticMgr::GetInstance();
 	DiagMgr.AddDelegate(Delegate.Get());
-#endif // USE_USD_SDK
+#endif	  // USE_USD_SDK
 }
 
 void UnrealUSDWrapper::ClearDiagnosticDelegate()
@@ -953,7 +960,7 @@ void UnrealUSDWrapper::ClearDiagnosticDelegate()
 	DiagMgr.RemoveDelegate(Delegate.Get());
 
 	Delegate = nullptr;
-#endif // USE_USD_SDK
+#endif	  // USE_USD_SDK
 }
 
 class FUnrealUSDWrapperModule : public IUnrealUSDWrapperModule
@@ -965,35 +972,45 @@ public:
 		LLM_SCOPE_BYTAG(Usd);
 
 		// Path to USD base plugins
-		FString UsdPluginsPath = FPaths::Combine( TEXT( ".." ), TEXT( "ThirdParty" ), TEXT( "USD" ), TEXT( "UsdResources" ) );
-		UsdPluginsPath = FPaths::ConvertRelativePathToFull( UsdPluginsPath );
+		FString UsdPluginsPath = FPaths::Combine(TEXT(".."), TEXT("ThirdParty"), TEXT("USD"), TEXT("UsdResources"));
+		UsdPluginsPath = FPaths::ConvertRelativePathToFull(UsdPluginsPath);
 #if PLATFORM_WINDOWS
-		UsdPluginsPath /= FPaths::Combine( TEXT( "Win64" ), TEXT( "plugins" ) );
+		UsdPluginsPath /= FPaths::Combine(TEXT("Win64"), TEXT("plugins"));
 #elif PLATFORM_LINUX
-		UsdPluginsPath /= FPaths::Combine( TEXT( "Linux" ), TEXT( "plugins" ) );
+		UsdPluginsPath /= FPaths::Combine(TEXT("Linux"), TEXT("plugins"));
 #elif PLATFORM_MAC
-		UsdPluginsPath /= FPaths::Combine( TEXT( "Mac" ), TEXT( "plugins" ) );
-#endif // PLATFORM_WINDOWS
+		UsdPluginsPath /= FPaths::Combine(TEXT("Mac"), TEXT("plugins"));
+#endif	  // PLATFORM_WINDOWS
 
 #ifdef USE_LIBRARIES_FROM_PLUGIN_FOLDER
 		// e.g. "../../../Engine/Plugins/Importers/USDImporter/Source/ThirdParty"
-		FString TargetDllFolder = FPaths::Combine( IPluginManager::Get().FindPlugin( TEXT( "USDImporter" ) )->GetBaseDir(), TEXT( "Source" ), TEXT( "ThirdParty" ) );
+		FString TargetDllFolder = FPaths::Combine(
+			IPluginManager::Get().FindPlugin(TEXT("USDImporter"))->GetBaseDir(),
+			TEXT("Source"),
+			TEXT("ThirdParty")
+		);
 
 #if PLATFORM_WINDOWS
-		TargetDllFolder /= FPaths::Combine( TEXT( "USD" ), TEXT( "bin" ) );
+		TargetDllFolder /= FPaths::Combine(TEXT("USD"), TEXT("bin"));
 #elif PLATFORM_LINUX
-		TargetDllFolder /= FPaths::Combine( TEXT( "Linux" ), TEXT( "bin" ), TEXT( "x86_64-unknown-linux-gnu" ) );
+		TargetDllFolder /= FPaths::Combine(TEXT("Linux"), TEXT("bin"), TEXT("x86_64-unknown-linux-gnu"));
 #elif PLATFORM_MAC
-		TargetDllFolder /= FPaths::Combine( TEXT( "Mac" ), TEXT( "bin" ) );
-#endif // PLATFORM_WINDOWS
+		TargetDllFolder /= FPaths::Combine(TEXT("Mac"), TEXT("bin"));
+#endif	  // PLATFORM_WINDOWS
 
 #else
 		FString TargetDllFolder = FPlatformProcess::BaseDir();
-#endif // USE_LIBRARIES_FROM_PLUGIN_FOLDER
+#endif	  // USE_LIBRARIES_FROM_PLUGIN_FOLDER
 
 		// Path to the MaterialX standard data libraries.
-		FString MaterialXStdDataLibsPath = FPaths::Combine( FPaths::EngineDir(), TEXT( "Binaries" ), TEXT("ThirdParty"), TEXT("MaterialX"), TEXT("libraries"));
-		MaterialXStdDataLibsPath = FPaths::ConvertRelativePathToFull( MaterialXStdDataLibsPath );
+		FString MaterialXStdDataLibsPath = FPaths::Combine(
+			FPaths::EngineDir(),
+			TEXT("Binaries"),
+			TEXT("ThirdParty"),
+			TEXT("MaterialX"),
+			TEXT("libraries")
+		);
+		MaterialXStdDataLibsPath = FPaths::ConvertRelativePathToFull(MaterialXStdDataLibsPath);
 		if (FPaths::DirectoryExists(MaterialXStdDataLibsPath))
 		{
 			FPlatformMisc::SetEnvironmentVar(TEXT("PXR_MTLX_STDLIB_SEARCH_PATHS"), *MaterialXStdDataLibsPath);
@@ -1004,10 +1021,10 @@ public:
 
 		// Combine our current plugins with any additional USD plugins the user may have set.
 		TArray<FString> PluginDirectories;
-		PluginDirectories.Add( UsdPluginsPath );
-		for ( const FDirectoryPath& Directory : GetDefault<UUsdProjectSettings>()->AdditionalPluginDirectories )
+		PluginDirectories.Add(UsdPluginsPath);
+		for (const FDirectoryPath& Directory : GetDefault<UUsdProjectSettings>()->AdditionalPluginDirectories)
 		{
-			if ( !Directory.Path.IsEmpty() )
+			if (!Directory.Path.IsEmpty())
 			{
 				// The directory chooser will try to enforce that the paths
 				// in the setting are within the game content dir, but it's
@@ -1015,36 +1032,36 @@ public:
 				// relative or an absolute path that points elsewhere, so aim
 				// to support any of these.
 				FString PluginDirPath = Directory.Path;
-				if ( FPaths::IsRelative( PluginDirPath ) )
+				if (FPaths::IsRelative(PluginDirPath))
 				{
-					PluginDirPath = FPaths::Combine( FPaths::ProjectContentDir(), PluginDirPath );
-					PluginDirPath = FPaths::ConvertRelativePathToFull( PluginDirPath );
+					PluginDirPath = FPaths::Combine(FPaths::ProjectContentDir(), PluginDirPath);
+					PluginDirPath = FPaths::ConvertRelativePathToFull(PluginDirPath);
 				}
-				FPaths::NormalizeDirectoryName( PluginDirPath );
-				FPaths::CollapseRelativeDirectories( PluginDirPath );
-				PluginDirectories.Add( PluginDirPath );
+				FPaths::NormalizeDirectoryName(PluginDirPath);
+				FPaths::CollapseRelativeDirectories(PluginDirPath);
+				PluginDirectories.Add(PluginDirPath);
 			}
 		}
 
 		{
 			FScopedUsdAllocs UsdAllocs;
 
-			std::vector< std::string > UsdPluginDirectories;
-			UsdPluginDirectories.reserve( PluginDirectories.Num() );
+			std::vector<std::string> UsdPluginDirectories;
+			UsdPluginDirectories.reserve(PluginDirectories.Num());
 
-			for ( const FString& Dir : PluginDirectories )
+			for (const FString& Dir : PluginDirectories)
 			{
-				UsdPluginDirectories.push_back( TCHAR_TO_UTF8( *Dir ) );
+				UsdPluginDirectories.push_back(TCHAR_TO_UTF8(*Dir));
 			}
 
-			PlugRegistry::GetInstance().RegisterPlugins( UsdPluginDirectories );
+			PlugRegistry::GetInstance().RegisterPlugins(UsdPluginDirectories);
 		}
 
 		// Set the default search path for USD's default resolver using any path specified in the settings.
 		const TArray<FDirectoryPath> SearchPath = GetDefault<UUsdProjectSettings>()->DefaultResolverSearchPath;
-		UnrealUSDWrapper::SetDefaultResolverDefaultSearchPath( SearchPath );
+		UnrealUSDWrapper::SetDefaultResolverDefaultSearchPath(SearchPath);
 
-#endif // USE_USD_SDK
+#endif	  // USE_USD_SDK
 
 		FUsdMemoryManager::Initialize();
 		UnrealUSDWrapper::SetupDiagnosticDelegate();
diff --git a/Engine/Plugins/Importers/USDImporter/Source/UnrealUSDWrapper/Private/UsdWrappers/SdfAttributeSpec.cpp b/Engine/Plugins/Importers/USDImporter/Source/UnrealUSDWrapper/Private/UsdWrappers/SdfAttributeSpec.cpp
--- a/Engine/Plugins/Importers/USDImporter/Source/UnrealUSDWrapper/Private/UsdWrappers/SdfAttributeSpec.cpp
+++ b/Engine/Plugins/Importers/USDImporter/Source/UnrealUSDWrapper/Private/UsdWrappers/SdfAttributeSpec.cpp
@@ -7,12 +7,10 @@
 #include "USDMemory.h"
 
 #if USE_USD_SDK
-
 #include "USDIncludesStart.h"
-	#include "pxr/usd/usd/attribute.h"
+#include "pxr/usd/usd/attribute.h"
 #include "USDIncludesEnd.h"
-
-#endif // #if USE_USD_SDK
+#endif	  // #if USE_USD_SDK
 
 namespace UE
 {
@@ -24,18 +22,18 @@ namespace UE
 			FSdfAttributeSpecImpl() = default;
 
 #if USE_USD_SDK
-			explicit FSdfAttributeSpecImpl( const pxr::SdfAttributeSpecHandle& InSdfAttributeSpec )
-				: PxrSdfAttributeSpec( InSdfAttributeSpec )
+			explicit FSdfAttributeSpecImpl(const pxr::SdfAttributeSpecHandle& InSdfAttributeSpec)
+				: PxrSdfAttributeSpec(InSdfAttributeSpec)
 			{
 			}
 
-			explicit FSdfAttributeSpecImpl( pxr::SdfAttributeSpecHandle&& InSdfAttributeSpec )
-				: PxrSdfAttributeSpec( MoveTemp( InSdfAttributeSpec ) )
+			explicit FSdfAttributeSpecImpl(pxr::SdfAttributeSpecHandle&& InSdfAttributeSpec)
+				: PxrSdfAttributeSpec(MoveTemp(InSdfAttributeSpec))
 			{
 			}
 
-			TUsdStore< pxr::SdfAttributeSpecHandle > PxrSdfAttributeSpec;
-#endif // #if USE_USD_SDK
+			TUsdStore<pxr::SdfAttributeSpecHandle> PxrSdfAttributeSpec;
+#endif	  // #if USE_USD_SDK
 		};
 	}
 }
@@ -45,18 +43,18 @@ namespace UE
 	FSdfAttributeSpec::FSdfAttributeSpec()
 	{
 		FScopedUnrealAllocs UnrealAllocs;
-		Impl =  MakeUnique< Internal::FSdfAttributeSpecImpl >();
+		Impl = MakeUnique<Internal::FSdfAttributeSpecImpl>();
 	}
 
-	FSdfAttributeSpec::FSdfAttributeSpec( const FSdfAttributeSpec& Other )
+	FSdfAttributeSpec::FSdfAttributeSpec(const FSdfAttributeSpec& Other)
 	{
 #if USE_USD_SDK
 		FScopedUnrealAllocs UnrealAllocs;
-		Impl = MakeUnique< Internal::FSdfAttributeSpecImpl >( Other.Impl->PxrSdfAttributeSpec.Get() );
-#endif // #if USE_USD_SDK
+		Impl = MakeUnique<Internal::FSdfAttributeSpecImpl>(Other.Impl->PxrSdfAttributeSpec.Get());
+#endif	  // #if USE_USD_SDK
 	}
 
-	FSdfAttributeSpec::FSdfAttributeSpec( FSdfAttributeSpec&& Other ) = default;
+	FSdfAttributeSpec::FSdfAttributeSpec(FSdfAttributeSpec&& Other) = default;
 
 	FSdfAttributeSpec::~FSdfAttributeSpec()
 	{
@@ -64,35 +62,35 @@ namespace UE
 		Impl.Reset();
 	}
 
-	FSdfAttributeSpec& FSdfAttributeSpec::operator=( const FSdfAttributeSpec& Other )
+	FSdfAttributeSpec& FSdfAttributeSpec::operator=(const FSdfAttributeSpec& Other)
 	{
 #if USE_USD_SDK
 		FScopedUnrealAllocs UnrealAllocs;
-		Impl = MakeUnique< Internal::FSdfAttributeSpecImpl >(  Other.Impl->PxrSdfAttributeSpec.Get() );
-#endif // #if USE_USD_SDK
+		Impl = MakeUnique<Internal::FSdfAttributeSpecImpl>(Other.Impl->PxrSdfAttributeSpec.Get());
+#endif	  // #if USE_USD_SDK
 		return *this;
 	}
 
-	FSdfAttributeSpec& FSdfAttributeSpec::operator=( FSdfAttributeSpec&& Other )
+	FSdfAttributeSpec& FSdfAttributeSpec::operator=(FSdfAttributeSpec&& Other)
 	{
 		FScopedUnrealAllocs UnrealAllocs;
-		Impl = MoveTemp( Other.Impl );
+		Impl = MoveTemp(Other.Impl);
 
 		return *this;
 	}
 
-	bool FSdfAttributeSpec::operator==( const FSdfAttributeSpec& Other ) const
+	bool FSdfAttributeSpec::operator==(const FSdfAttributeSpec& Other) const
 	{
 #if USE_USD_SDK
 		return Impl->PxrSdfAttributeSpec.Get() == Other.Impl->PxrSdfAttributeSpec.Get();
 #else
 		return false;
-#endif // #if USE_USD_SDK
+#endif	  // #if USE_USD_SDK
 	}
 
-	bool FSdfAttributeSpec::operator!=( const FSdfAttributeSpec& Other ) const
+	bool FSdfAttributeSpec::operator!=(const FSdfAttributeSpec& Other) const
 	{
-		return !( *this == Other );
+		return !(*this == Other);
 	}
 
 	FSdfAttributeSpec::operator bool() const
@@ -101,33 +99,33 @@ namespace UE
 		return (bool)Impl->PxrSdfAttributeSpec.Get();
 #else
 		return false;
-#endif // #if USE_USD_SDK
+#endif	  // #if USE_USD_SDK
 	}
 
 #if USE_USD_SDK
-	FSdfAttributeSpec::FSdfAttributeSpec( const pxr::SdfAttributeSpecHandle& InSdfAttributeSpec )
+	FSdfAttributeSpec::FSdfAttributeSpec(const pxr::SdfAttributeSpecHandle& InSdfAttributeSpec)
 	{
 		FScopedUnrealAllocs UnrealAllocs;
-		Impl =  MakeUnique< Internal::FSdfAttributeSpecImpl >( InSdfAttributeSpec );
+		Impl = MakeUnique<Internal::FSdfAttributeSpecImpl>(InSdfAttributeSpec);
 	}
 
-	FSdfAttributeSpec::FSdfAttributeSpec( pxr::SdfAttributeSpecHandle&& InSdfAttributeSpec )
+	FSdfAttributeSpec::FSdfAttributeSpec(pxr::SdfAttributeSpecHandle&& InSdfAttributeSpec)
 	{
 		FScopedUnrealAllocs UnrealAllocs;
-		Impl = MakeUnique< Internal::FSdfAttributeSpecImpl >( MoveTemp( InSdfAttributeSpec ) );
+		Impl = MakeUnique<Internal::FSdfAttributeSpecImpl>(MoveTemp(InSdfAttributeSpec));
 	}
 
-	FSdfAttributeSpec& FSdfAttributeSpec::operator=( const pxr::SdfAttributeSpecHandle& InSdfAttributeSpec )
+	FSdfAttributeSpec& FSdfAttributeSpec::operator=(const pxr::SdfAttributeSpecHandle& InSdfAttributeSpec)
 	{
 		FScopedUnrealAllocs UnrealAllocs;
-		Impl = MakeUnique< Internal::FSdfAttributeSpecImpl >( InSdfAttributeSpec );
+		Impl = MakeUnique<Internal::FSdfAttributeSpecImpl>(InSdfAttributeSpec);
 		return *this;
 	}
 
-	FSdfAttributeSpec& FSdfAttributeSpec::operator=( pxr::SdfAttributeSpecHandle&& InSdfAttributeSpec )
+	FSdfAttributeSpec& FSdfAttributeSpec::operator=(pxr::SdfAttributeSpecHandle&& InSdfAttributeSpec)
 	{
 		FScopedUnrealAllocs UnrealAllocs;
-		Impl = MakeUnique< Internal::FSdfAttributeSpecImpl >( MoveTemp( InSdfAttributeSpec ) );
+		Impl = MakeUnique<Internal::FSdfAttributeSpecImpl>(MoveTemp(InSdfAttributeSpec));
 		return *this;
 	}
 
@@ -140,35 +138,35 @@ namespace UE
 	{
 		return Impl->PxrSdfAttributeSpec.Get();
 	}
-#endif // #if USE_USD_SDK
+#endif	  // #if USE_USD_SDK
 
 	UE::FVtValue FSdfAttributeSpec::GetDefaultValue() const
 	{
 #if USE_USD_SDK
 		FScopedUsdAllocs UsdAllocs;
-		return UE::FVtValue{ Impl->PxrSdfAttributeSpec.Get()->GetDefaultValue() };
+		return UE::FVtValue{Impl->PxrSdfAttributeSpec.Get()->GetDefaultValue()};
 #else
 		return {};
-#endif // #if USE_USD_SDK
+#endif	  // #if USE_USD_SDK
 	}
 
-	bool FSdfAttributeSpec::SetDefaultValue( const UE::FVtValue& Value )
+	bool FSdfAttributeSpec::SetDefaultValue(const UE::FVtValue& Value)
 	{
 #if USE_USD_SDK
 		FScopedUsdAllocs UsdAllocs;
-		return Impl->PxrSdfAttributeSpec.Get()->SetDefaultValue( Value.GetUsdValue() );
+		return Impl->PxrSdfAttributeSpec.Get()->SetDefaultValue(Value.GetUsdValue());
 #else
 		return false;
-#endif // #if USE_USD_SDK
+#endif	  // #if USE_USD_SDK
 	}
 
 	FString FSdfAttributeSpec::GetName() const
 	{
 #if USE_USD_SDK
 		FScopedUsdAllocs UsdAllocs;
-		return ANSI_TO_TCHAR( Impl->PxrSdfAttributeSpec.Get()->GetName().c_str() );
+		return ANSI_TO_TCHAR(Impl->PxrSdfAttributeSpec.Get()->GetName().c_str());
 #else
 		return {};
-#endif // #if USE_USD_SDK
+#endif	  // #if USE_USD_SDK
 	}
-}
+}	 // namespace UE
diff --git a/Engine/Plugins/Importers/USDImporter/Source/UnrealUSDWrapper/Private/UsdWrappers/SdfLayer.cpp b/Engine/Plugins/Importers/USDImporter/Source/UnrealUSDWrapper/Private/UsdWrappers/SdfLayer.cpp
--- a/Engine/Plugins/Importers/USDImporter/Source/UnrealUSDWrapper/Private/UsdWrappers/SdfLayer.cpp
+++ b/Engine/Plugins/Importers/USDImporter/Source/UnrealUSDWrapper/Private/UsdWrappers/SdfLayer.cpp
@@ -8,14 +8,12 @@
 #include "USDMemory.h"
 
 #if USE_USD_SDK
-
 #include "USDIncludesStart.h"
-	#include "pxr/usd/sdf/layer.h"
-	#include "pxr/usd/sdf/layerUtils.h"
-	#include "pxr/usd/sdf/primSpec.h"
+#include "pxr/usd/sdf/layer.h"
+#include "pxr/usd/sdf/layerUtils.h"
+#include "pxr/usd/sdf/primSpec.h"
 #include "USDIncludesEnd.h"
-
-#endif // #if USE_USD_SDK
+#endif	  // #if USE_USD_SDK
 
 namespace UE
 {
@@ -27,13 +25,13 @@ namespace UE
 		public:
 			FSdfLayerImpl() = default;
 
-			explicit FSdfLayerImpl( const PtrType& InUsdPtr )
-				: UsdPtr( InUsdPtr )
+			explicit FSdfLayerImpl(const PtrType& InUsdPtr)
+				: UsdPtr(InUsdPtr)
 			{
 			}
 
-			explicit FSdfLayerImpl( PtrType&& InUsdPtr )
-				: UsdPtr( MoveTemp( InUsdPtr ) )
+			explicit FSdfLayerImpl(PtrType&& InUsdPtr)
+				: UsdPtr(MoveTemp(InUsdPtr))
 			{
 			}
 
@@ -47,7 +45,7 @@ namespace UE
 				return UsdPtr.Get();
 			}
 
-			TUsdStore< PtrType > UsdPtr;
+			TUsdStore<PtrType> UsdPtr;
 		};
 	}
 
@@ -55,66 +53,66 @@ namespace UE
 	FSdfLayerBase<PtrType>::FSdfLayerBase()
 	{
 		FScopedUnrealAllocs UnrealAllocs;
-		Impl = MakeUnique< Internal::FSdfLayerImpl<PtrType> >();
+		Impl = MakeUnique<Internal::FSdfLayerImpl<PtrType>>();
 	}
 
 	template<typename PtrType>
-	FSdfLayerBase<PtrType>::FSdfLayerBase( const FSdfLayer& Other )
+	FSdfLayerBase<PtrType>::FSdfLayerBase(const FSdfLayer& Other)
 	{
 		FScopedUnrealAllocs UnrealAllocs;
-		Impl = MakeUnique< Internal::FSdfLayerImpl<PtrType> >( Other.Impl->GetInner() );
+		Impl = MakeUnique<Internal::FSdfLayerImpl<PtrType>>(Other.Impl->GetInner());
 	}
 
 	template<typename PtrType>
-	FSdfLayerBase<PtrType>::FSdfLayerBase( FSdfLayer&& Other )
+	FSdfLayerBase<PtrType>::FSdfLayerBase(FSdfLayer&& Other)
 	{
 		FScopedUnrealAllocs UnrealAllocs;
-		Impl = MakeUnique< Internal::FSdfLayerImpl<PtrType> >( MoveTemp( Other.Impl->GetInner() ) );
+		Impl = MakeUnique<Internal::FSdfLayerImpl<PtrType>>(MoveTemp(Other.Impl->GetInner()));
 	}
 
 	template<typename PtrType>
-	FSdfLayerBase<PtrType>::FSdfLayerBase( const FSdfLayerWeak& Other )
+	FSdfLayerBase<PtrType>::FSdfLayerBase(const FSdfLayerWeak& Other)
 	{
 		FScopedUnrealAllocs UnrealAllocs;
-		Impl = MakeUnique< Internal::FSdfLayerImpl<PtrType> >( Other.Impl->GetInner() );
+		Impl = MakeUnique<Internal::FSdfLayerImpl<PtrType>>(Other.Impl->GetInner());
 	}
 
 	template<typename PtrType>
-	FSdfLayerBase<PtrType>::FSdfLayerBase( FSdfLayerWeak&& Other )
+	FSdfLayerBase<PtrType>::FSdfLayerBase(FSdfLayerWeak&& Other)
 	{
 		FScopedUnrealAllocs UnrealAllocs;
-		Impl = MakeUnique< Internal::FSdfLayerImpl<PtrType> >( MoveTemp( Other.Impl->GetInner() ) );
+		Impl = MakeUnique<Internal::FSdfLayerImpl<PtrType>>(MoveTemp(Other.Impl->GetInner()));
 	}
 
 	template<typename PtrType>
-	FSdfLayerBase<PtrType>& FSdfLayerBase<PtrType>::operator=( const FSdfLayer& Other )
+	FSdfLayerBase<PtrType>& FSdfLayerBase<PtrType>::operator=(const FSdfLayer& Other)
 	{
 		FScopedUnrealAllocs UnrealAllocs;
-		Impl = MakeUnique< Internal::FSdfLayerImpl<PtrType> >( Other.Impl->GetInner() );
+		Impl = MakeUnique<Internal::FSdfLayerImpl<PtrType>>(Other.Impl->GetInner());
 		return *this;
 	}
 
 	template<typename PtrType>
-	FSdfLayerBase<PtrType>& FSdfLayerBase<PtrType>::operator=( FSdfLayer&& Other )
+	FSdfLayerBase<PtrType>& FSdfLayerBase<PtrType>::operator=(FSdfLayer&& Other)
 	{
 		FScopedUnrealAllocs UnrealAllocs;
-		Impl = MakeUnique< Internal::FSdfLayerImpl<PtrType> >( MoveTemp( Other.Impl->GetInner() ) );
+		Impl = MakeUnique<Internal::FSdfLayerImpl<PtrType>>(MoveTemp(Other.Impl->GetInner()));
 		return *this;
 	}
 
 	template<typename PtrType>
-	FSdfLayerBase<PtrType>& FSdfLayerBase<PtrType>::operator=( const FSdfLayerWeak& Other )
+	FSdfLayerBase<PtrType>& FSdfLayerBase<PtrType>::operator=(const FSdfLayerWeak& Other)
 	{
 		FScopedUnrealAllocs UnrealAllocs;
-		Impl = MakeUnique< Internal::FSdfLayerImpl<PtrType> >( Other.Impl->GetInner() );
+		Impl = MakeUnique<Internal::FSdfLayerImpl<PtrType>>(Other.Impl->GetInner());
 		return *this;
 	}
 
 	template<typename PtrType>
-	FSdfLayerBase<PtrType>& FSdfLayerBase<PtrType>::operator=( FSdfLayerWeak&& Other )
+	FSdfLayerBase<PtrType>& FSdfLayerBase<PtrType>::operator=(FSdfLayerWeak&& Other)
 	{
 		FScopedUnrealAllocs UnrealAllocs;
-		Impl = MakeUnique< Internal::FSdfLayerImpl<PtrType> >( MoveTemp( Other.Impl->GetInner() ) );
+		Impl = MakeUnique<Internal::FSdfLayerImpl<PtrType>>(MoveTemp(Other.Impl->GetInner()));
 		return *this;
 	}
 
@@ -127,30 +125,32 @@ namespace UE
 
 	template<typename PtrType>
 	template<typename OtherPtrType>
-	bool FSdfLayerBase<PtrType>::operator==( const FSdfLayerBase<OtherPtrType>& Other ) const
+	bool FSdfLayerBase<PtrType>::operator==(const FSdfLayerBase<OtherPtrType>& Other) const
 	{
 		return Impl->GetInner() == Other.Impl->GetInner();
 	}
-	template bool UNREALUSDWRAPPER_API FSdfLayer::operator==( const FSdfLayer& Other ) const;
-	template bool UNREALUSDWRAPPER_API FSdfLayer::operator==( const FSdfLayerWeak& Other ) const;
-	template bool UNREALUSDWRAPPER_API FSdfLayerWeak::operator==( const FSdfLayer& Other ) const;
-	template bool UNREALUSDWRAPPER_API FSdfLayerWeak::operator==( const FSdfLayerWeak& Other ) const;
+
+	template bool UNREALUSDWRAPPER_API FSdfLayer::operator==(const FSdfLayer& Other) const;
+	template bool UNREALUSDWRAPPER_API FSdfLayer::operator==(const FSdfLayerWeak& Other) const;
+	template bool UNREALUSDWRAPPER_API FSdfLayerWeak::operator==(const FSdfLayer& Other) const;
+	template bool UNREALUSDWRAPPER_API FSdfLayerWeak::operator==(const FSdfLayerWeak& Other) const;
 
 	template<typename PtrType>
 	template<typename OtherPtrType>
-	bool FSdfLayerBase<PtrType>::operator!=( const FSdfLayerBase<OtherPtrType>& Other ) const
+	bool FSdfLayerBase<PtrType>::operator!=(const FSdfLayerBase<OtherPtrType>& Other) const
 	{
-		return !( *this == Other );
+		return !(*this == Other);
 	}
-	template bool UNREALUSDWRAPPER_API FSdfLayer::operator!=( const FSdfLayer& Other ) const;
-	template bool UNREALUSDWRAPPER_API FSdfLayer::operator!=( const FSdfLayerWeak& Other ) const;
-	template bool UNREALUSDWRAPPER_API FSdfLayerWeak::operator!=( const FSdfLayer& Other ) const;
-	template bool UNREALUSDWRAPPER_API FSdfLayerWeak::operator!=( const FSdfLayerWeak& Other ) const;
+
+	template bool UNREALUSDWRAPPER_API FSdfLayer::operator!=(const FSdfLayer& Other) const;
+	template bool UNREALUSDWRAPPER_API FSdfLayer::operator!=(const FSdfLayerWeak& Other) const;
+	template bool UNREALUSDWRAPPER_API FSdfLayerWeak::operator!=(const FSdfLayer& Other) const;
+	template bool UNREALUSDWRAPPER_API FSdfLayerWeak::operator!=(const FSdfLayerWeak& Other) const;
 
 	template<typename PtrType>
 	FSdfLayerBase<PtrType>::operator bool() const
 	{
-		return ( bool ) Impl->GetInner();
+		return (bool)Impl->GetInner();
 	}
 
 	UNREALUSDWRAPPER_API uint32 GetTypeHash(const UE::FSdfLayerWeak& Layer)
@@ -170,62 +170,62 @@ namespace UE
 
 #if USE_USD_SDK
 	template<typename PtrType>
-	FSdfLayerBase<PtrType>::FSdfLayerBase( const pxr::SdfLayerRefPtr& InSdfLayer )
+	FSdfLayerBase<PtrType>::FSdfLayerBase(const pxr::SdfLayerRefPtr& InSdfLayer)
 	{
 		FScopedUnrealAllocs UnrealAllocs;
-		Impl = MakeUnique< Internal::FSdfLayerImpl<PtrType> >( InSdfLayer );
+		Impl = MakeUnique<Internal::FSdfLayerImpl<PtrType>>(InSdfLayer);
 	}
 
 	template<typename PtrType>
-	FSdfLayerBase<PtrType>::FSdfLayerBase( pxr::SdfLayerRefPtr&& InSdfLayer )
+	FSdfLayerBase<PtrType>::FSdfLayerBase(pxr::SdfLayerRefPtr&& InSdfLayer)
 	{
 		FScopedUnrealAllocs UnrealAllocs;
-		Impl = MakeUnique< Internal::FSdfLayerImpl<PtrType> >( InSdfLayer );
+		Impl = MakeUnique<Internal::FSdfLayerImpl<PtrType>>(InSdfLayer);
 	}
 
 	template<typename PtrType>
-	FSdfLayerBase<PtrType>::FSdfLayerBase( const pxr::SdfLayerWeakPtr& InSdfLayer )
+	FSdfLayerBase<PtrType>::FSdfLayerBase(const pxr::SdfLayerWeakPtr& InSdfLayer)
 	{
 		FScopedUnrealAllocs UnrealAllocs;
-		Impl = MakeUnique< Internal::FSdfLayerImpl<PtrType> >( InSdfLayer );
+		Impl = MakeUnique<Internal::FSdfLayerImpl<PtrType>>(InSdfLayer);
 	}
 
 	template<typename PtrType>
-	FSdfLayerBase<PtrType>::FSdfLayerBase( pxr::SdfLayerWeakPtr&& InSdfLayer )
+	FSdfLayerBase<PtrType>::FSdfLayerBase(pxr::SdfLayerWeakPtr&& InSdfLayer)
 	{
 		FScopedUnrealAllocs UnrealAllocs;
-		Impl = MakeUnique< Internal::FSdfLayerImpl<PtrType> >( InSdfLayer );
+		Impl = MakeUnique<Internal::FSdfLayerImpl<PtrType>>(InSdfLayer);
 	}
 
 	template<typename PtrType>
-	FSdfLayerBase<PtrType>& FSdfLayerBase<PtrType>::operator=( const pxr::SdfLayerRefPtr& InSdfLayer )
+	FSdfLayerBase<PtrType>& FSdfLayerBase<PtrType>::operator=(const pxr::SdfLayerRefPtr& InSdfLayer)
 	{
 		FScopedUnrealAllocs UnrealAllocs;
-		Impl = MakeUnique< Internal::FSdfLayerImpl<PtrType> >( InSdfLayer );
+		Impl = MakeUnique<Internal::FSdfLayerImpl<PtrType>>(InSdfLayer);
 		return *this;
 	}
 
 	template<typename PtrType>
-	FSdfLayerBase<PtrType>& FSdfLayerBase<PtrType>::operator=( pxr::SdfLayerRefPtr&& InSdfLayer )
+	FSdfLayerBase<PtrType>& FSdfLayerBase<PtrType>::operator=(pxr::SdfLayerRefPtr&& InSdfLayer)
 	{
 		FScopedUnrealAllocs UnrealAllocs;
-		Impl = MakeUnique< Internal::FSdfLayerImpl<PtrType> >( InSdfLayer );
+		Impl = MakeUnique<Internal::FSdfLayerImpl<PtrType>>(InSdfLayer);
 		return *this;
 	}
 
 	template<typename PtrType>
-	FSdfLayerBase<PtrType>& FSdfLayerBase<PtrType>::operator=( const pxr::SdfLayerWeakPtr& InSdfLayer )
+	FSdfLayerBase<PtrType>& FSdfLayerBase<PtrType>::operator=(const pxr::SdfLayerWeakPtr& InSdfLayer)
 	{
 		FScopedUnrealAllocs UnrealAllocs;
-		Impl = MakeUnique< Internal::FSdfLayerImpl<PtrType> >( InSdfLayer );
+		Impl = MakeUnique<Internal::FSdfLayerImpl<PtrType>>(InSdfLayer);
 		return *this;
 	}
 
 	template<typename PtrType>
-	FSdfLayerBase<PtrType>& FSdfLayerBase<PtrType>::operator=( pxr::SdfLayerWeakPtr&& InSdfLayer )
+	FSdfLayerBase<PtrType>& FSdfLayerBase<PtrType>::operator=(pxr::SdfLayerWeakPtr&& InSdfLayer)
 	{
 		FScopedUnrealAllocs UnrealAllocs;
-		Impl = MakeUnique< Internal::FSdfLayerImpl<PtrType> >( InSdfLayer );
+		Impl = MakeUnique<Internal::FSdfLayerImpl<PtrType>>(InSdfLayer);
 		return *this;
 	}
 
@@ -252,42 +252,42 @@ namespace UE
 	{
 		return Impl->GetInner();
 	}
-#endif // #if USE_USD_SDK
+#endif	  // #if USE_USD_SDK
 
 	template<typename PtrType>
 	FString FSdfLayerBase<PtrType>::GetComment() const
 	{
 #if USE_USD_SDK
-		if ( const PtrType& Ptr = Impl->GetInner() )
+		if (const PtrType& Ptr = Impl->GetInner())
 		{
 			FScopedUsdAllocs UsdAllocs;
-			return ANSI_TO_TCHAR( Ptr->GetComment().c_str() );
+			return ANSI_TO_TCHAR(Ptr->GetComment().c_str());
 		}
-#endif // #if USE_USD_SDK
+#endif	  // #if USE_USD_SDK
 		return {};
 	}
 
 	template<typename PtrType>
-	void FSdfLayerBase<PtrType>::SetComment( const TCHAR* Comment ) const
+	void FSdfLayerBase<PtrType>::SetComment(const TCHAR* Comment) const
 	{
 #if USE_USD_SDK
-		if ( const PtrType& Ptr = Impl->GetInner() )
+		if (const PtrType& Ptr = Impl->GetInner())
 		{
 			FScopedUsdAllocs UsdAllocs;
-			Ptr->SetComment( TCHAR_TO_ANSI( Comment ) );
+			Ptr->SetComment(TCHAR_TO_ANSI(Comment));
 		}
-#endif // #if USE_USD_SDK
+#endif	  // #if USE_USD_SDK
 	}
 
 	template<typename PtrType>
-	void FSdfLayerBase<PtrType>::TransferContent( const FSdfLayer& SourceLayer )
+	void FSdfLayerBase<PtrType>::TransferContent(const FSdfLayer& SourceLayer)
 	{
 #if USE_USD_SDK
-		if ( const PtrType& Ptr = Impl->GetInner() )
+		if (const PtrType& Ptr = Impl->GetInner())
 		{
-			Ptr->TransferContent( pxr::SdfLayerRefPtr{ SourceLayer } );
+			Ptr->TransferContent(pxr::SdfLayerRefPtr{SourceLayer});
 		}
-#endif // #if USE_USD_SDK
+#endif	  // #if USE_USD_SDK
 	}
 
 	template<typename PtrType>
@@ -298,82 +298,83 @@ namespace UE
 		FScopedUsdAllocs UsdAllocs;
 
 		pxr::SdfLayerHandleSet LoadedSdfLayers = pxr::SdfLayer::GetLoadedLayers();
-		Result.Reserve( LoadedSdfLayers.size() );
+		Result.Reserve(LoadedSdfLayers.size());
 
-		for ( const pxr::SdfLayerHandle& LoadedSdfLayer : LoadedSdfLayers )
+		for (const pxr::SdfLayerHandle& LoadedSdfLayer : LoadedSdfLayers)
 		{
-			Result.Add( UE::FSdfLayerWeak{ LoadedSdfLayer } );
+			Result.Add(UE::FSdfLayerWeak{LoadedSdfLayer});
 		}
-#endif // #if USE_USD_SDK
+#endif	  // #if USE_USD_SDK
 		return Result;
 	}
 
 	template<typename PtrType>
-	FSdfLayer FSdfLayerBase<PtrType>::FindOrOpen( const TCHAR* Identifier, const TMap< FString, FString >& FileFormatArguments )
+	FSdfLayer FSdfLayerBase<PtrType>::FindOrOpen(const TCHAR* Identifier, const TMap<FString, FString>& FileFormatArguments)
 	{
 #if USE_USD_SDK
 		FScopedUsdAllocs Allocs;
 
-		std::map< std::string, std::string > UsdFileFormatArguments;
-		for ( const TPair< FString, FString >& Pair : FileFormatArguments )
+		std::map<std::string, std::string> UsdFileFormatArguments;
+		for (const TPair<FString, FString>& Pair : FileFormatArguments)
 		{
-			UsdFileFormatArguments.insert( { TCHAR_TO_ANSI( *Pair.Key ), TCHAR_TO_ANSI( *Pair.Value ) } );
+			UsdFileFormatArguments.insert({TCHAR_TO_ANSI(*Pair.Key), TCHAR_TO_ANSI(*Pair.Value)});
 		}
 
-		return FSdfLayer( pxr::SdfLayer::FindOrOpen( TCHAR_TO_ANSI( Identifier ), UsdFileFormatArguments ) );
+		return FSdfLayer(pxr::SdfLayer::FindOrOpen(TCHAR_TO_ANSI(Identifier), UsdFileFormatArguments));
 #else
 		return FSdfLayer();
-#endif // #if USE_USD_SDK
+#endif	  // #if USE_USD_SDK
 	}
 
 	template<typename PtrType>
-	FSdfLayer FSdfLayerBase<PtrType>::CreateNew( const TCHAR* Identifier, const TMap< FString, FString >& FileFormatArguments )
+	FSdfLayer FSdfLayerBase<PtrType>::CreateNew(const TCHAR* Identifier, const TMap<FString, FString>& FileFormatArguments)
 	{
 #if USE_USD_SDK
 		FScopedUsdAllocs Allocs;
 
-		std::map< std::string, std::string > UsdFileFormatArguments;
-		for ( const TPair< FString, FString >& Pair : FileFormatArguments )
+		std::map<std::string, std::string> UsdFileFormatArguments;
+		for (const TPair<FString, FString>& Pair : FileFormatArguments)
 		{
-			UsdFileFormatArguments.insert( { TCHAR_TO_ANSI( *Pair.Key ), TCHAR_TO_ANSI( *Pair.Value ) } );
+			UsdFileFormatArguments.insert({TCHAR_TO_ANSI(*Pair.Key), TCHAR_TO_ANSI(*Pair.Value)});
 		}
 
-		return FSdfLayer( pxr::SdfLayer::CreateNew( TCHAR_TO_ANSI( Identifier ), UsdFileFormatArguments ) );
+		return FSdfLayer(pxr::SdfLayer::CreateNew(TCHAR_TO_ANSI(Identifier), UsdFileFormatArguments));
 #else
 		return FSdfLayer();
-#endif // #if USE_USD_SDK
+#endif	  // #if USE_USD_SDK
 	}
 
 	template<typename PtrType>
-	TMap< FString, FString > FSdfLayerBase<PtrType>::GetFileFormatArguments() const
+	TMap<FString, FString> FSdfLayerBase<PtrType>::GetFileFormatArguments() const
 	{
-		TMap< FString, FString > FileFormatArguments;
+		TMap<FString, FString> FileFormatArguments;
 
 #if USE_USD_SDK
 		FScopedUsdAllocs Allocs;
 
-		if ( const PtrType& Ptr = Impl->GetInner() )
+		if (const PtrType& Ptr = Impl->GetInner())
 		{
-			const std::map< std::string, std::string >& UsdFileFormatArguments = Ptr->GetFileFormatArguments();
-			for ( std::map< std::string, std::string >::const_iterator Iter = UsdFileFormatArguments.cbegin(); Iter != UsdFileFormatArguments.cend(); ++Iter )
+			const std::map<std::string, std::string>& UsdFileFormatArguments = Ptr->GetFileFormatArguments();
+			for (std::map<std::string, std::string>::const_iterator Iter = UsdFileFormatArguments.cbegin(); Iter != UsdFileFormatArguments.cend();
+				 ++Iter)
 			{
-				FileFormatArguments.Add( ANSI_TO_TCHAR( Iter->first.c_str() ), ANSI_TO_TCHAR( Iter->second.c_str() ) );
+				FileFormatArguments.Add(ANSI_TO_TCHAR(Iter->first.c_str()), ANSI_TO_TCHAR(Iter->second.c_str()));
 			}
 		}
-#endif // #if USE_USD_SDK
+#endif	  // #if USE_USD_SDK
 
 		return FileFormatArguments;
 	}
 
 	template<typename PtrType>
-	bool FSdfLayerBase<PtrType>::Save( bool bForce ) const
+	bool FSdfLayerBase<PtrType>::Save(bool bForce) const
 	{
 #if USE_USD_SDK
-		if ( const PtrType& Ptr = Impl->GetInner() )
+		if (const PtrType& Ptr = Impl->GetInner())
 		{
-			return Ptr->Save( bForce );
+			return Ptr->Save(bForce);
 		}
-#endif // #if USE_USD_SDK
+#endif	  // #if USE_USD_SDK
 
 		return false;
 	}
@@ -383,35 +384,34 @@ namespace UE
 	{
 		TSet<FString> Results;
 #if USE_USD_SDK
-		if ( const PtrType& Ptr = Impl->GetInner() )
+		if (const PtrType& Ptr = Impl->GetInner())
 		{
 			FScopedUsdAllocs UsdAllocs;
 			const std::set<std::string> AssetDependencies = Ptr->GetCompositionAssetDependencies();
 
-			Results.Reserve( AssetDependencies.size() );
-			for ( const std::string& AssetDependency : AssetDependencies )
+			Results.Reserve(AssetDependencies.size());
+			for (const std::string& AssetDependency : AssetDependencies)
 			{
-				Results.Add( ANSI_TO_TCHAR( AssetDependency.c_str() ) );
+				Results.Add(ANSI_TO_TCHAR(AssetDependency.c_str()));
 			}
 		}
-#endif // #if USE_USD_SDK
+#endif	  // #if USE_USD_SDK
 		return Results;
 	}
 
 	template<typename PtrType>
-	bool FSdfLayerBase<PtrType>::UpdateCompositionAssetDependency(
-			const TCHAR* OldAssetPath,
-			const TCHAR* NewAssetPath)
+	bool FSdfLayerBase<PtrType>::UpdateCompositionAssetDependency(const TCHAR* OldAssetPath, const TCHAR* NewAssetPath)
 	{
 #if USE_USD_SDK
-		if ( const PtrType& Ptr = Impl->GetInner() )
+		if (const PtrType& Ptr = Impl->GetInner())
 		{
 			FScopedUsdAllocs UsdAllocs;
 			return Ptr->UpdateCompositionAssetDependency(
-				TCHAR_TO_ANSI( OldAssetPath ),
-				(NewAssetPath != nullptr) ? TCHAR_TO_ANSI( NewAssetPath ) : std::string() );
+				TCHAR_TO_ANSI(OldAssetPath),
+				(NewAssetPath != nullptr) ? TCHAR_TO_ANSI(NewAssetPath) : std::string()
+			);
 		}
-#endif // #if USE_USD_SDK
+#endif	  // #if USE_USD_SDK
 		return false;
 	}
 
@@ -419,12 +419,12 @@ namespace UE
 	FString FSdfLayerBase<PtrType>::GetRealPath() const
 	{
 #if USE_USD_SDK
-		if ( const PtrType& Ptr = Impl->GetInner() )
+		if (const PtrType& Ptr = Impl->GetInner())
 		{
 			FScopedUsdAllocs UsdAllocs;
-			return FString( ANSI_TO_TCHAR( Ptr->GetRealPath().c_str() ) );
+			return FString(ANSI_TO_TCHAR(Ptr->GetRealPath().c_str()));
 		}
-#endif // #if USE_USD_SDK
+#endif	  // #if USE_USD_SDK
 
 		return FString();
 	}
@@ -433,12 +433,12 @@ namespace UE
 	FString FSdfLayerBase<PtrType>::GetIdentifier() const
 	{
 #if USE_USD_SDK
-		if ( const PtrType& Ptr = Impl->GetInner() )
+		if (const PtrType& Ptr = Impl->GetInner())
 		{
 			FScopedUsdAllocs UsdAllocs;
-			return FString( ANSI_TO_TCHAR( Ptr->GetIdentifier().c_str() ) );
+			return FString(ANSI_TO_TCHAR(Ptr->GetIdentifier().c_str()));
 		}
-#endif // #if USE_USD_SDK
+#endif	  // #if USE_USD_SDK
 
 		return FString();
 	}
@@ -447,26 +447,26 @@ namespace UE
 	FString FSdfLayerBase<PtrType>::GetDisplayName() const
 	{
 #if USE_USD_SDK
-		if ( const PtrType& Ptr = Impl->GetInner() )
+		if (const PtrType& Ptr = Impl->GetInner())
 		{
 			FScopedUsdAllocs UsdAllocs;
-			return FString( ANSI_TO_TCHAR( Ptr->GetDisplayName().c_str() ) );
+			return FString(ANSI_TO_TCHAR(Ptr->GetDisplayName().c_str()));
 		}
-#endif // #if USE_USD_SDK
+#endif	  // #if USE_USD_SDK
 
 		return FString();
 	}
 
 	template<typename PtrType>
-	FString FSdfLayerBase<PtrType>::ComputeAbsolutePath( const FString& AssetPath ) const
+	FString FSdfLayerBase<PtrType>::ComputeAbsolutePath(const FString& AssetPath) const
 	{
 #if USE_USD_SDK
-		if ( const PtrType& Ptr = Impl->GetInner() )
+		if (const PtrType& Ptr = Impl->GetInner())
 		{
 			FScopedUsdAllocs UsdAllocs;
-			return FString( ANSI_TO_TCHAR( Ptr->ComputeAbsolutePath( TCHAR_TO_ANSI( *AssetPath ) ).c_str() ) );
+			return FString(ANSI_TO_TCHAR(Ptr->ComputeAbsolutePath(TCHAR_TO_ANSI(*AssetPath)).c_str()));
 		}
-#endif // #if USE_USD_SDK
+#endif	  // #if USE_USD_SDK
 
 		return FString();
 	}
@@ -475,11 +475,11 @@ namespace UE
 	bool FSdfLayerBase<PtrType>::HasStartTimeCode() const
 	{
 #if USE_USD_SDK
-		if ( const PtrType& Ptr = Impl->GetInner() )
+		if (const PtrType& Ptr = Impl->GetInner())
 		{
 			return Ptr->HasStartTimeCode();
 		}
-#endif // #if USE_USD_SDK
+#endif	  // #if USE_USD_SDK
 
 		return false;
 	}
@@ -488,11 +488,11 @@ namespace UE
 	bool FSdfLayerBase<PtrType>::HasEndTimeCode() const
 	{
 #if USE_USD_SDK
-		if ( const PtrType& Ptr = Impl->GetInner() )
+		if (const PtrType& Ptr = Impl->GetInner())
 		{
 			return Ptr->HasEndTimeCode();
 		}
-#endif // #if USE_USD_SDK
+#endif	  // #if USE_USD_SDK
 
 		return false;
 	}
@@ -501,11 +501,11 @@ namespace UE
 	double FSdfLayerBase<PtrType>::GetStartTimeCode() const
 	{
 #if USE_USD_SDK
-		if ( const PtrType& Ptr = Impl->GetInner() )
+		if (const PtrType& Ptr = Impl->GetInner())
 		{
 			return Ptr->GetStartTimeCode();
 		}
-#endif // #if USE_USD_SDK
+#endif	  // #if USE_USD_SDK
 
 		return 0.0;
 	}
@@ -514,52 +514,52 @@ namespace UE
 	double FSdfLayerBase<PtrType>::GetEndTimeCode() const
 	{
 #if USE_USD_SDK
-		if ( const PtrType& Ptr = Impl->GetInner() )
+		if (const PtrType& Ptr = Impl->GetInner())
 		{
 			return Ptr->GetEndTimeCode();
 		}
-#endif // #if USE_USD_SDK
+#endif	  // #if USE_USD_SDK
 
 		return 0.0;
 	}
 
 	template<typename PtrType>
-	void FSdfLayerBase<PtrType>::SetStartTimeCode( double TimeCode )
+	void FSdfLayerBase<PtrType>::SetStartTimeCode(double TimeCode)
 	{
 #if USE_USD_SDK
-		if ( PtrType& Ptr = Impl->GetInner() )
+		if (PtrType& Ptr = Impl->GetInner())
 		{
-			if ( !Ptr->HasStartTimeCode() || !FMath::IsNearlyEqual( TimeCode, Ptr->GetStartTimeCode() ) )
+			if (!Ptr->HasStartTimeCode() || !FMath::IsNearlyEqual(TimeCode, Ptr->GetStartTimeCode()))
 			{
-				Ptr->SetStartTimeCode( TimeCode );
+				Ptr->SetStartTimeCode(TimeCode);
 			}
 		}
-#endif // #if USE_USD_SDK
+#endif	  // #if USE_USD_SDK
 	}
 
 	template<typename PtrType>
-	void FSdfLayerBase<PtrType>::SetEndTimeCode( double TimeCode )
+	void FSdfLayerBase<PtrType>::SetEndTimeCode(double TimeCode)
 	{
 #if USE_USD_SDK
-		if ( PtrType& Ptr = Impl->GetInner() )
+		if (PtrType& Ptr = Impl->GetInner())
 		{
-			if ( !Ptr->HasEndTimeCode() || !FMath::IsNearlyEqual( TimeCode, Ptr->GetEndTimeCode() ) )
+			if (!Ptr->HasEndTimeCode() || !FMath::IsNearlyEqual(TimeCode, Ptr->GetEndTimeCode()))
 			{
-				Ptr->SetEndTimeCode( TimeCode );
+				Ptr->SetEndTimeCode(TimeCode);
 			}
 		}
-#endif // #if USE_USD_SDK
+#endif	  // #if USE_USD_SDK
 	}
 
 	template<typename PtrType>
 	bool FSdfLayerBase<PtrType>::HasTimeCodesPerSecond() const
 	{
 #if USE_USD_SDK
-		if ( const PtrType& Ptr = Impl->GetInner() )
+		if (const PtrType& Ptr = Impl->GetInner())
 		{
 			return Ptr->HasTimeCodesPerSecond();
 		}
-#endif // #if USE_USD_SDK
+#endif	  // #if USE_USD_SDK
 
 		return false;
 	}
@@ -568,35 +568,35 @@ namespace UE
 	double FSdfLayerBase<PtrType>::GetTimeCodesPerSecond() const
 	{
 #if USE_USD_SDK
-		if ( const PtrType& Ptr = Impl->GetInner() )
+		if (const PtrType& Ptr = Impl->GetInner())
 		{
 			return Ptr->GetTimeCodesPerSecond();
 		}
-#endif // #if USE_USD_SDK
+#endif	  // #if USE_USD_SDK
 
 		return 0.0;
 	}
 
 	template<typename PtrType>
-	void FSdfLayerBase<PtrType>::SetTimeCodesPerSecond( double TimeCodesPerSecond )
+	void FSdfLayerBase<PtrType>::SetTimeCodesPerSecond(double TimeCodesPerSecond)
 	{
 #if USE_USD_SDK
-		if ( PtrType& Ptr = Impl->GetInner() )
+		if (PtrType& Ptr = Impl->GetInner())
 		{
-			return Ptr->SetTimeCodesPerSecond( TimeCodesPerSecond );
+			return Ptr->SetTimeCodesPerSecond(TimeCodesPerSecond);
 		}
-#endif // #if USE_USD_SDK
+#endif	  // #if USE_USD_SDK
 	}
 
 	template<typename PtrType>
 	bool FSdfLayerBase<PtrType>::HasFramesPerSecond() const
 	{
 #if USE_USD_SDK
-		if ( const PtrType& Ptr = Impl->GetInner() )
+		if (const PtrType& Ptr = Impl->GetInner())
 		{
 			return Ptr->HasFramesPerSecond();
 		}
-#endif // #if USE_USD_SDK
+#endif	  // #if USE_USD_SDK
 
 		return false;
 	}
@@ -605,24 +605,24 @@ namespace UE
 	double FSdfLayerBase<PtrType>::GetFramesPerSecond() const
 	{
 #if USE_USD_SDK
-		if ( const PtrType& Ptr = Impl->GetInner() )
+		if (const PtrType& Ptr = Impl->GetInner())
 		{
 			return Ptr->GetFramesPerSecond();
 		}
-#endif // #if USE_USD_SDK
+#endif	  // #if USE_USD_SDK
 
 		return 0.0;
 	}
 
 	template<typename PtrType>
-	void FSdfLayerBase<PtrType>::SetFramesPerSecond( double FramesPerSecond )
+	void FSdfLayerBase<PtrType>::SetFramesPerSecond(double FramesPerSecond)
 	{
 #if USE_USD_SDK
-		if ( PtrType& Ptr = Impl->GetInner() )
+		if (PtrType& Ptr = Impl->GetInner())
 		{
-			return Ptr->SetFramesPerSecond( FramesPerSecond );
+			return Ptr->SetFramesPerSecond(FramesPerSecond);
 		}
-#endif // #if USE_USD_SDK
+#endif	  // #if USE_USD_SDK
 	}
 
 	template<typename PtrType>
@@ -641,7 +641,7 @@ namespace UE
 				SubLayerPaths.Emplace(ConvertedPath.Get());
 			}
 		}
-#endif // #if USE_USD_SDK
+#endif	  // #if USE_USD_SDK
 
 		return SubLayerPaths;
 	}
@@ -666,7 +666,7 @@ namespace UE
 
 			Ptr->SetSubLayerPaths(SubLayerPaths);
 		}
-#endif // #if USE_USD_SDK
+#endif	  // #if USE_USD_SDK
 	}
 
 	template<typename PtrType>
@@ -677,7 +677,7 @@ namespace UE
 		{
 			return (int64)Ptr->GetNumSubLayerPaths();
 		}
-#endif // #if USE_USD_SDK
+#endif	  // #if USE_USD_SDK
 
 		return 0;
 	}
@@ -694,7 +694,7 @@ namespace UE
 			std::string SubLayerPath = reinterpret_cast<const char*>(ConvertedPath.Get());
 			Ptr->InsertSubLayerPath(MoveTemp(SubLayerPath), Index);
 		}
-#endif // #if USE_USD_SDK
+#endif	  // #if USE_USD_SDK
 	}
 
 	template<typename PtrType>
@@ -705,7 +705,7 @@ namespace UE
 		{
 			Ptr->RemoveSubLayerPath(Index);
 		}
-#endif // #if USE_USD_SDK
+#endif	  // #if USE_USD_SDK
 	}
 
 	template<typename PtrType>
@@ -730,7 +730,7 @@ namespace UE
 				}
 			}
 		}
-#endif // #if USE_USD_SDK
+#endif	  // #if USE_USD_SDK
 
 		return SubLayerOffsets;
 	}
@@ -749,7 +749,7 @@ namespace UE
 				return FSdfLayerOffset(SubLayerOffset.GetOffset(), SubLayerOffset.GetScale());
 			}
 		}
-#endif // #if USE_USD_SDK
+#endif	  // #if USE_USD_SDK
 
 		return FSdfLayerOffset{};
 	}
@@ -765,18 +765,18 @@ namespace UE
 			pxr::SdfLayerOffset UsdLayerOffset(Offset.Offset, Offset.Scale);
 			Ptr->SetSubLayerOffset(MoveTemp(UsdLayerOffset), Index);
 		}
-#endif // #if USE_USD_SDK
+#endif	  // #if USE_USD_SDK
 	}
 
 	template<typename PtrType>
-	bool FSdfLayerBase<PtrType>::HasSpec( const FSdfPath& Path ) const
+	bool FSdfLayerBase<PtrType>::HasSpec(const FSdfPath& Path) const
 	{
 #if USE_USD_SDK
-		if ( const PtrType& Ptr = Impl->GetInner() )
+		if (const PtrType& Ptr = Impl->GetInner())
 		{
-			return Ptr->HasSpec( Path );
+			return Ptr->HasSpec(Path);
 		}
-#endif // #if USE_USD_SDK
+#endif	  // #if USE_USD_SDK
 
 		return false;
 	}
@@ -785,70 +785,70 @@ namespace UE
 	FSdfPrimSpec FSdfLayerBase<PtrType>::GetPseudoRoot() const
 	{
 #if USE_USD_SDK
-		if ( const PtrType& Ptr = Impl->GetInner() )
+		if (const PtrType& Ptr = Impl->GetInner())
 		{
-			return FSdfPrimSpec{ Ptr->GetPseudoRoot() };
+			return FSdfPrimSpec{Ptr->GetPseudoRoot()};
 		}
-#endif // #if USE_USD_SDK
+#endif	  // #if USE_USD_SDK
 
 		return FSdfPrimSpec{};
 	}
 
 	template<typename PtrType>
-	FSdfPrimSpec FSdfLayerBase<PtrType>::GetPrimAtPath( const FSdfPath& Path ) const
+	FSdfPrimSpec FSdfLayerBase<PtrType>::GetPrimAtPath(const FSdfPath& Path) const
 	{
 #if USE_USD_SDK
-		if ( const PtrType& Ptr = Impl->GetInner() )
+		if (const PtrType& Ptr = Impl->GetInner())
 		{
-			return FSdfPrimSpec{ Ptr->GetPrimAtPath( Path ) };
+			return FSdfPrimSpec{Ptr->GetPrimAtPath(Path)};
 		}
-#endif // #if USE_USD_SDK
+#endif	  // #if USE_USD_SDK
 
 		return FSdfPrimSpec{};
 	}
 
 	template<typename PtrType>
-	TSet< double > FSdfLayerBase<PtrType>::ListTimeSamplesForPath( const FSdfPath& Path ) const
+	TSet<double> FSdfLayerBase<PtrType>::ListTimeSamplesForPath(const FSdfPath& Path) const
 	{
-		TSet< double > TimeSamples;
+		TSet<double> TimeSamples;
 
 #if USE_USD_SDK
-		if ( const PtrType& Ptr = Impl->GetInner() )
+		if (const PtrType& Ptr = Impl->GetInner())
 		{
 			FScopedUsdAllocs UsdAllocs;
 
-			std::set< double > UsdTimeSamples = Ptr->ListTimeSamplesForPath( Path );
+			std::set<double> UsdTimeSamples = Ptr->ListTimeSamplesForPath(Path);
 
-			for ( double UsdTimeSample : UsdTimeSamples )
+			for (double UsdTimeSample : UsdTimeSamples)
 			{
-				TimeSamples.Add( UsdTimeSample );
+				TimeSamples.Add(UsdTimeSample);
 			}
 		}
-#endif // #if USE_USD_SDK
+#endif	  // #if USE_USD_SDK
 
 		return TimeSamples;
 	}
 
 	template<typename PtrType>
-	void FSdfLayerBase<PtrType>::EraseTimeSample( const FSdfPath& Path, double Time )
+	void FSdfLayerBase<PtrType>::EraseTimeSample(const FSdfPath& Path, double Time)
 	{
 #if USE_USD_SDK
-		if ( PtrType& Ptr = Impl->GetInner() )
+		if (PtrType& Ptr = Impl->GetInner())
 		{
-			return Ptr->EraseTimeSample( Path, Time );
+			return Ptr->EraseTimeSample(Path, Time);
 		}
-#endif // #if USE_USD_SDK
+#endif	  // #if USE_USD_SDK
 	}
 
 	template<typename PtrType>
 	bool FSdfLayerBase<PtrType>::IsDirty() const
 	{
 #if USE_USD_SDK
-		if ( const PtrType& Ptr = Impl->GetInner() )
+		if (const PtrType& Ptr = Impl->GetInner())
 		{
 			return Ptr->IsDirty();
 		}
-#endif // #if USE_USD_SDK
+#endif	  // #if USE_USD_SDK
 
 		return false;
 	}
@@ -857,11 +857,11 @@ namespace UE
 	bool FSdfLayerBase<PtrType>::IsEmpty() const
 	{
 #if USE_USD_SDK
-		if ( const PtrType& Ptr = Impl->GetInner() )
+		if (const PtrType& Ptr = Impl->GetInner())
 		{
 			return Ptr->IsEmpty();
 		}
-#endif // #if USE_USD_SDK
+#endif	  // #if USE_USD_SDK
 
 		return false;
 	}
@@ -870,32 +870,32 @@ namespace UE
 	bool FSdfLayerBase<PtrType>::IsAnonymous() const
 	{
 #if USE_USD_SDK
-		if ( const PtrType& Ptr = Impl->GetInner() )
+		if (const PtrType& Ptr = Impl->GetInner())
 		{
 			return Ptr->IsAnonymous();
 		}
-#endif // #if USE_USD_SDK
+#endif	  // #if USE_USD_SDK
 
 		return false;
 	}
 
 	template<typename PtrType>
-	bool FSdfLayerBase<PtrType>::Export( const TCHAR* Filename, const FString& Comment, const TMap< FString, FString >& FileFormatArguments ) const
+	bool FSdfLayerBase<PtrType>::Export(const TCHAR* Filename, const FString& Comment, const TMap<FString, FString>& FileFormatArguments) const
 	{
 #if USE_USD_SDK
-		if ( const PtrType& Ptr = Impl->GetInner() )
+		if (const PtrType& Ptr = Impl->GetInner())
 		{
 			FScopedUsdAllocs Allocs;
 
-			std::map< std::string, std::string > UsdFileFormatArguments;
-			for ( const TPair< FString, FString >& Pair : FileFormatArguments )
+			std::map<std::string, std::string> UsdFileFormatArguments;
+			for (const TPair<FString, FString>& Pair : FileFormatArguments)
 			{
-				UsdFileFormatArguments.insert( { TCHAR_TO_ANSI( *Pair.Key ), TCHAR_TO_ANSI( *Pair.Value ) } );
+				UsdFileFormatArguments.insert({TCHAR_TO_ANSI(*Pair.Key), TCHAR_TO_ANSI(*Pair.Value)});
 			}
 
-			return Ptr->Export( TCHAR_TO_ANSI( Filename ), TCHAR_TO_ANSI ( *Comment ), UsdFileFormatArguments);
+			return Ptr->Export(TCHAR_TO_ANSI(Filename), TCHAR_TO_ANSI(*Comment), UsdFileFormatArguments);
 		}
-#endif // #if USE_USD_SDK
+#endif	  // #if USE_USD_SDK
 
 		return false;
 	}
@@ -904,45 +904,45 @@ namespace UE
 	void FSdfLayerBase<PtrType>::Clear()
 	{
 #if USE_USD_SDK
-		if ( const PtrType& Ptr = Impl->GetInner() )
+		if (const PtrType& Ptr = Impl->GetInner())
 		{
 			return Ptr->Clear();
 		}
-#endif // #if USE_USD_SDK
+#endif	  // #if USE_USD_SDK
 	}
 
-	FString FSdfLayerUtils::SdfComputeAssetPathRelativeToLayer( const FSdfLayer& Anchor, const TCHAR* AssetPath )
+	FString FSdfLayerUtils::SdfComputeAssetPathRelativeToLayer(const FSdfLayer& Anchor, const TCHAR* AssetPath)
 	{
 #if USE_USD_SDK
 		FScopedUsdAllocs UsdAllocs;
-		return FString( ANSI_TO_TCHAR( pxr::SdfComputeAssetPathRelativeToLayer( pxr::SdfLayerRefPtr( Anchor ), TCHAR_TO_ANSI( AssetPath ) ).c_str() ) );
+		return FString(ANSI_TO_TCHAR(pxr::SdfComputeAssetPathRelativeToLayer(pxr::SdfLayerRefPtr(Anchor), TCHAR_TO_ANSI(AssetPath)).c_str()));
 #else
 		return FString();
-#endif // #if USE_USD_SDK
+#endif	  // #if USE_USD_SDK
 	}
 
 	template<typename PtrType>
 	bool FSdfLayerBase<PtrType>::IsMuted() const
 	{
 #if USE_USD_SDK
-		if ( const PtrType& Ptr = Impl->GetInner() )
+		if (const PtrType& Ptr = Impl->GetInner())
 		{
 			return Ptr->IsMuted();
 		}
-#endif // #if USE_USD_SDK
+#endif	  // #if USE_USD_SDK
 
 		return false;
 	}
 
 	template<typename PtrType>
-	void FSdfLayerBase<PtrType>::SetMuted( bool bMuted )
+	void FSdfLayerBase<PtrType>::SetMuted(bool bMuted)
 	{
 #if USE_USD_SDK
-		if ( PtrType& Ptr = Impl->GetInner() )
+		if (PtrType& Ptr = Impl->GetInner())
 		{
-			return Ptr->SetMuted( bMuted );
+			return Ptr->SetMuted(bMuted);
 		}
-#endif // #if USE_USD_SDK
+#endif	  // #if USE_USD_SDK
 	}
 
 #if USE_USD_SDK
@@ -951,5 +951,5 @@ namespace UE
 #else
 	template class UNREALUSDWRAPPER_API FSdfLayerBase<FDummyRefPtrType>;
 	template class UNREALUSDWRAPPER_API FSdfLayerBase<FDummyWeakPtrType>;
-#endif // #if USE_USD_SDK
-}
+#endif	  // #if USE_USD_SDK
+}	 // namespace UE
