USD: Reformat the UnrealUSDWrapper module.

diff --git a/Engine/Plugins/Importers/USDImporter/Source/UnrealUSDWrapper/Private/UsdWrappers/SdfPath.cpp b/Engine/Plugins/Importers/USDImporter/Source/UnrealUSDWrapper/Private/UsdWrappers/SdfPath.cpp
--- a/Engine/Plugins/Importers/USDImporter/Source/UnrealUSDWrapper/Private/UsdWrappers/SdfPath.cpp
+++ b/Engine/Plugins/Importers/USDImporter/Source/UnrealUSDWrapper/Private/UsdWrappers/SdfPath.cpp
@@ -5,12 +5,10 @@
 #include "USDMemory.h"
 
 #if USE_USD_SDK
-
 #include "USDIncludesStart.h"
-	#include "pxr/usd/sdf/path.h"
+#include "pxr/usd/sdf/path.h"
 #include "USDIncludesEnd.h"
-
-#endif // #if USE_USD_SDK
+#endif	  // #if USE_USD_SDK
 
 namespace UE
 {
@@ -26,68 +24,67 @@ namespace UE
 #if ENABLE_USD_DEBUG_PATH
 			FString DebugPath;
 #endif
-			explicit FSdfPathImpl( const pxr::SdfPath& InSdfPath )
-				: PxrSdfPath( InSdfPath )
+			explicit FSdfPathImpl(const pxr::SdfPath& InSdfPath)
+				: PxrSdfPath(InSdfPath)
 			{
 #if ENABLE_USD_DEBUG_PATH
 				DebugPath = FString(ANSI_TO_TCHAR(PxrSdfPath.Get().GetString().c_str()));
 #endif
 			}
 
-			explicit FSdfPathImpl( pxr::SdfPath&& InSdfPath )
-				: PxrSdfPath( MoveTemp( InSdfPath ) )
+			explicit FSdfPathImpl(pxr::SdfPath&& InSdfPath)
+				: PxrSdfPath(MoveTemp(InSdfPath))
 			{
 #if ENABLE_USD_DEBUG_PATH
 				DebugPath = FString(ANSI_TO_TCHAR(PxrSdfPath.Get().GetString().c_str()));
 #endif
 			}
 
-			TUsdStore< pxr::SdfPath > PxrSdfPath;
-#endif // #if USE_USD_SDK
+			TUsdStore<pxr::SdfPath> PxrSdfPath;
+#endif	  // #if USE_USD_SDK
 		};
-	}
+	}	  // namespace Internal
 
 	FSdfPath::FSdfPath()
 	{
 		FScopedUnrealAllocs UnrealAllocs;
-		Impl = MakeUnique< Internal::FSdfPathImpl >();
+		Impl = MakeUnique<Internal::FSdfPathImpl>();
 	}
 
-	FSdfPath::FSdfPath( const FSdfPath& Other )
+	FSdfPath::FSdfPath(const FSdfPath& Other)
 	{
 #if USE_USD_SDK
 		FScopedUnrealAllocs UnrealAllocs;
-		Impl =  MakeUnique< Internal::FSdfPathImpl >( Other.Impl->PxrSdfPath.Get() );
-#endif // #if USE_USD_SDK
+		Impl = MakeUnique<Internal::FSdfPathImpl>(Other.Impl->PxrSdfPath.Get());
+#endif	  // #if USE_USD_SDK
 	}
 
-	FSdfPath::FSdfPath( FSdfPath&& Other ) = default;
+	FSdfPath::FSdfPath(FSdfPath&& Other) = default;
 
-	FSdfPath::FSdfPath( const TCHAR* InPath )
+	FSdfPath::FSdfPath(const TCHAR* InPath)
 	{
 #if USE_USD_SDK
 		FScopedUnrealAllocs UnrealAllocs;
-		Impl = MakeUnique< Internal::FSdfPathImpl >(
-			FCString::Strlen( InPath ) == 0
-				? pxr::SdfPath() // USD gives a warning if we construct with an empty string
-				: pxr::SdfPath( TCHAR_TO_ANSI( InPath ) )
-			);
-#endif // #if USE_USD_SDK
+		Impl = MakeUnique<Internal::FSdfPathImpl>(
+			FCString::Strlen(InPath) == 0 ? pxr::SdfPath()	  // USD gives a warning if we construct with an empty string
+										  : pxr::SdfPath(TCHAR_TO_ANSI(InPath))
+		);
+#endif	  // #if USE_USD_SDK
 	}
 
-	FSdfPath& FSdfPath::operator=( const FSdfPath& Other )
+	FSdfPath& FSdfPath::operator=(const FSdfPath& Other)
 	{
 #if USE_USD_SDK
 		FScopedUnrealAllocs UnrealAllocs;
-		Impl = MakeUnique< Internal::FSdfPathImpl >( Other.Impl->PxrSdfPath.Get() );
-#endif // #if USE_USD_SDK
+		Impl = MakeUnique<Internal::FSdfPathImpl>(Other.Impl->PxrSdfPath.Get());
+#endif	  // #if USE_USD_SDK
 		return *this;
 	}
 
-	FSdfPath& FSdfPath::operator=( FSdfPath&& Other )
+	FSdfPath& FSdfPath::operator=(FSdfPath&& Other)
 	{
 		FScopedUnrealAllocs UnrealAllocs;
-		Impl = MoveTemp( Other.Impl );
+		Impl = MoveTemp(Other.Impl);
 
 		return *this;
 	}
@@ -98,18 +95,18 @@ namespace UE
 		Impl.Reset();
 	}
 
-	bool FSdfPath::operator==( const FSdfPath& Other ) const
+	bool FSdfPath::operator==(const FSdfPath& Other) const
 	{
 #if USE_USD_SDK
 		return Impl->PxrSdfPath.Get() == Other.Impl->PxrSdfPath.Get();
 #else
 		return false;
-#endif // #if USE_USD_SDK
+#endif	  // #if USE_USD_SDK
 	}
 
-	bool FSdfPath::operator!=( const FSdfPath& Other ) const
+	bool FSdfPath::operator!=(const FSdfPath& Other) const
 	{
-		return !( *this == Other );
+		return !(*this == Other);
 	}
 
 	bool FSdfPath::operator<(const FSdfPath& Other) const
@@ -118,7 +115,7 @@ namespace UE
 		return Impl->PxrSdfPath.Get() < Other.Impl->PxrSdfPath.Get();
 #else
 		return false;
-#endif // #if USE_USD_SDK
+#endif	  // #if USE_USD_SDK
 	}
 
 	bool FSdfPath::operator<=(const FSdfPath& Other) const
@@ -127,59 +124,59 @@ namespace UE
 		return Impl->PxrSdfPath.Get() <= Other.Impl->PxrSdfPath.Get();
 #else
 		return false;
-#endif // #if USE_USD_SDK
+#endif	  // #if USE_USD_SDK
 	}
 
-	uint32 GetTypeHash( const UE::FSdfPath& Path )
+	uint32 GetTypeHash(const UE::FSdfPath& Path)
 	{
 		uint32 Result = 0;
 #if USE_USD_SDK
 		FScopedUnrealAllocs UnrealAllocs;
-		Result = static_cast< uint32 >( Path.Impl->PxrSdfPath.Get().GetHash() );
-#endif // #if USE_USD_SDK
+		Result = static_cast<uint32>(Path.Impl->PxrSdfPath.Get().GetHash());
+#endif	  // #if USE_USD_SDK
 		return Result;
 	}
 
-	FArchive& operator<<( FArchive& Ar, UE::FSdfPath& Path )
+	FArchive& operator<<(FArchive& Ar, UE::FSdfPath& Path)
 	{
-		if ( Ar.IsSaving() )
+		if (Ar.IsSaving())
 		{
 			FString String = Path.GetString();
 			Ar << String;
 		}
-		else if ( Ar.IsLoading() )
+		else if (Ar.IsLoading())
 		{
 			FString String;
 			Ar << String;
 
-			Path = UE::FSdfPath{ *String };
+			Path = UE::FSdfPath{*String};
 		}
 
 		return Ar;
 	}
 
 #if USE_USD_SDK
-	FSdfPath::FSdfPath( const pxr::SdfPath& InSdfPath )
-		: Impl( MakeUnique< Internal::FSdfPathImpl >( InSdfPath ) )
+	FSdfPath::FSdfPath(const pxr::SdfPath& InSdfPath)
+		: Impl(MakeUnique<Internal::FSdfPathImpl>(InSdfPath))
 	{
 	}
 
-	FSdfPath::FSdfPath( pxr::SdfPath&& InSdfPath )
-		: Impl( MakeUnique< Internal::FSdfPathImpl >( MoveTemp( InSdfPath ) ) )
+	FSdfPath::FSdfPath(pxr::SdfPath&& InSdfPath)
+		: Impl(MakeUnique<Internal::FSdfPathImpl>(MoveTemp(InSdfPath)))
 	{
 	}
 
-	FSdfPath& FSdfPath::operator=(  const pxr::SdfPath& InSdfPath )
+	FSdfPath& FSdfPath::operator=(const pxr::SdfPath& InSdfPath)
 	{
 		FScopedUnrealAllocs UnrealAllocs;
-		Impl = MakeUnique< Internal::FSdfPathImpl >( InSdfPath );
+		Impl = MakeUnique<Internal::FSdfPathImpl>(InSdfPath);
 		return *this;
 	}
 
-	FSdfPath& FSdfPath::operator=( pxr::SdfPath&& InSdfPath )
+	FSdfPath& FSdfPath::operator=(pxr::SdfPath&& InSdfPath)
 	{
 		FScopedUnrealAllocs UnrealAllocs;
-		Impl = MakeUnique< Internal::FSdfPathImpl >( MoveTemp( InSdfPath ) );
+		Impl = MakeUnique<Internal::FSdfPathImpl>(MoveTemp(InSdfPath));
 		return *this;
 	}
 
@@ -192,17 +189,17 @@ namespace UE
 	{
 		return Impl->PxrSdfPath.Get();
 	}
-#endif // #if USE_USD_SDK
+#endif	  // #if USE_USD_SDK
 
 	const FSdfPath& FSdfPath::AbsoluteRootPath()
 	{
 #if USE_USD_SDK
-		static FSdfPath AbsoluteRootPath( pxr::SdfPath::AbsoluteRootPath() );
+		static FSdfPath AbsoluteRootPath(pxr::SdfPath::AbsoluteRootPath());
 		return AbsoluteRootPath;
 #else
 		static FSdfPath AbsoluteRootPath;
 		return AbsoluteRootPath;
-#endif // #if USE_USD_SDK
+#endif	  // #if USE_USD_SDK
 	}
 
 	bool FSdfPath::IsEmpty() const noexcept
@@ -211,7 +208,7 @@ namespace UE
 		return Impl->PxrSdfPath.Get().IsEmpty();
 #else
 		return true;
-#endif // #if USE_USD_SDK
+#endif	  // #if USE_USD_SDK
 	}
 
 	bool FSdfPath::IsAbsoluteRootPath() const
@@ -220,7 +217,7 @@ namespace UE
 		return Impl->PxrSdfPath.Get().IsAbsoluteRootPath();
 #else
 		return false;
-#endif // #if USE_USD_SDK
+#endif	  // #if USE_USD_SDK
 	}
 
 	bool FSdfPath::IsAbsoluteRootOrPrimPath() const
@@ -229,52 +226,52 @@ namespace UE
 		return Impl->PxrSdfPath.Get().IsAbsoluteRootOrPrimPath();
 #else
 		return false;
-#endif // #if USE_USD_SDK
+#endif	  // #if USE_USD_SDK
 	}
 
 	FString FSdfPath::GetName() const
 	{
 #if USE_USD_SDK
-		return FString( ANSI_TO_TCHAR( Impl->PxrSdfPath.Get().GetName().c_str() ) );
+		return FString(ANSI_TO_TCHAR(Impl->PxrSdfPath.Get().GetName().c_str()));
 #else
 		return FString();
-#endif // #if USE_USD_SDK
+#endif	  // #if USE_USD_SDK
 	}
 
 	FString FSdfPath::GetElementString() const
 	{
 #if USE_USD_SDK
-		return FString( ANSI_TO_TCHAR( Impl->PxrSdfPath.Get().GetElementString().c_str() ) );
+		return FString(ANSI_TO_TCHAR(Impl->PxrSdfPath.Get().GetElementString().c_str()));
 #else
 		return FString();
-#endif // #if USE_USD_SDK
+#endif	  // #if USE_USD_SDK
 	}
 
 	FSdfPath FSdfPath::GetAbsoluteRootOrPrimPath() const
 	{
 #if USE_USD_SDK
-		return FSdfPath( Impl->PxrSdfPath.Get().GetAbsoluteRootOrPrimPath() );
+		return FSdfPath(Impl->PxrSdfPath.Get().GetAbsoluteRootOrPrimPath());
 #else
 		return FSdfPath();
-#endif // #if USE_USD_SDK
+#endif	  // #if USE_USD_SDK
 	}
 
-	FSdfPath FSdfPath::ReplaceName( const TCHAR* NewLeafName ) const
+	FSdfPath FSdfPath::ReplaceName(const TCHAR* NewLeafName) const
 	{
 #if USE_USD_SDK
-		return FSdfPath( Impl->PxrSdfPath.Get().ReplaceName( pxr::TfToken( TCHAR_TO_ANSI( NewLeafName ) ) ) );
+		return FSdfPath(Impl->PxrSdfPath.Get().ReplaceName(pxr::TfToken(TCHAR_TO_ANSI(NewLeafName))));
 #else
 		return FSdfPath();
-#endif // #if USE_USD_SDK
+#endif	  // #if USE_USD_SDK
 	}
 
 	FSdfPath FSdfPath::GetParentPath() const
 	{
 #if USE_USD_SDK
-		return FSdfPath( Impl->PxrSdfPath.Get().GetParentPath() );
+		return FSdfPath(Impl->PxrSdfPath.Get().GetParentPath());
 #else
 		return FSdfPath();
-#endif // #if USE_USD_SDK
+#endif	  // #if USE_USD_SDK
 	}
 
 	FSdfPath FSdfPath::AppendPath(const UE::FSdfPath& NewRelativeSuffix) const
@@ -283,43 +280,43 @@ namespace UE
 		return FSdfPath(Impl->PxrSdfPath.Get().AppendPath(NewRelativeSuffix.Impl->PxrSdfPath.Get()));
 #else
 		return FSdfPath();
-#endif // #if USE_USD_SDK
+#endif	  // #if USE_USD_SDK
 	}
 
-	FSdfPath FSdfPath::AppendChild( const TCHAR* ChildName ) const
+	FSdfPath FSdfPath::AppendChild(const TCHAR* ChildName) const
 	{
 #if USE_USD_SDK
-		return FSdfPath( Impl->PxrSdfPath.Get().AppendChild( pxr::TfToken( TCHAR_TO_ANSI( ChildName ) ) ) );
+		return FSdfPath(Impl->PxrSdfPath.Get().AppendChild(pxr::TfToken(TCHAR_TO_ANSI(ChildName))));
 #else
 		return FSdfPath();
-#endif // #if USE_USD_SDK
+#endif	  // #if USE_USD_SDK
 	}
 
-	FSdfPath FSdfPath::AppendProperty( FName PropertyName ) const
+	FSdfPath FSdfPath::AppendProperty(FName PropertyName) const
 	{
 #if USE_USD_SDK
-		return FSdfPath( Impl->PxrSdfPath.Get().AppendProperty( pxr::TfToken( TCHAR_TO_ANSI( *PropertyName.ToString() ) ) ) );
+		return FSdfPath(Impl->PxrSdfPath.Get().AppendProperty(pxr::TfToken(TCHAR_TO_ANSI(*PropertyName.ToString()))));
 #else
 		return FSdfPath();
-#endif // #if USE_USD_SDK
+#endif	  // #if USE_USD_SDK
 	}
 
 	UE::FSdfPath FSdfPath::StripAllVariantSelections() const
 	{
 #if USE_USD_SDK
-		return UE::FSdfPath( Impl->PxrSdfPath.Get().StripAllVariantSelections() );
+		return UE::FSdfPath(Impl->PxrSdfPath.Get().StripAllVariantSelections());
 #else
 		return UE::FSdfPath();
-#endif // #if USE_USD_SDK
+#endif	  // #if USE_USD_SDK
 	}
 
 	FString FSdfPath::GetString() const
 	{
 #if USE_USD_SDK
-		return FString( ANSI_TO_TCHAR( Impl->PxrSdfPath.Get().GetString().c_str() ) );
+		return FString(ANSI_TO_TCHAR(Impl->PxrSdfPath.Get().GetString().c_str()));
 #else
 		return FString();
-#endif // #if USE_USD_SDK
+#endif	  // #if USE_USD_SDK
 	}
 
 	TArray<FSdfPath> FSdfPath::GetPrefixes() const
@@ -329,13 +326,13 @@ namespace UE
 		FScopedUsdAllocs Allocs;
 
 		std::vector<pxr::SdfPath> UsdPrefixes = Impl->PxrSdfPath.Get().GetPrefixes();
-		Result.Reserve( UsdPrefixes.size() );
+		Result.Reserve(UsdPrefixes.size());
 
-		for ( const pxr::SdfPath& UsdPrefix : UsdPrefixes )
+		for (const pxr::SdfPath& UsdPrefix : UsdPrefixes)
 		{
-			Result.Add( UE::FSdfPath{ ANSI_TO_TCHAR( UsdPrefix.GetString().c_str() ) } );
+			Result.Add(UE::FSdfPath{ANSI_TO_TCHAR(UsdPrefix.GetString().c_str())});
 		}
-#endif // #if USE_USD_SDK
+#endif	  // #if USE_USD_SDK
 		return Result;
 	}
 
@@ -345,7 +342,7 @@ namespace UE
 		return Impl->PxrSdfPath.Get().HasPrefix(Prefix);
 #else
 		return false;
-#endif // #if USE_USD_SDK
+#endif	  // #if USE_USD_SDK
 	}
 
 	FSdfPath FSdfPath::MakeAbsolutePath(const UE::FSdfPath& Anchor) const
@@ -354,7 +351,7 @@ namespace UE
 		return FSdfPath(Impl->PxrSdfPath.Get().MakeAbsolutePath(Anchor.Impl->PxrSdfPath.Get()));
 #else
 		return FSdfPath();
-#endif // #if USE_USD_SDK
+#endif	  // #if USE_USD_SDK
 	}
 
 	FSdfPath FSdfPath::MakeRelativePath(const UE::FSdfPath& Anchor) const
@@ -363,6 +360,6 @@ namespace UE
 		return FSdfPath(Impl->PxrSdfPath.Get().MakeRelativePath(Anchor.Impl->PxrSdfPath.Get()));
 #else
 		return FSdfPath();
-#endif // #if USE_USD_SDK
+#endif	  // #if USE_USD_SDK
 	}
-}
+}	 // namespace UE
diff --git a/Engine/Plugins/Importers/USDImporter/Source/UnrealUSDWrapper/Private/UsdWrappers/SdfPrimSpec.cpp b/Engine/Plugins/Importers/USDImporter/Source/UnrealUSDWrapper/Private/UsdWrappers/SdfPrimSpec.cpp
--- a/Engine/Plugins/Importers/USDImporter/Source/UnrealUSDWrapper/Private/UsdWrappers/SdfPrimSpec.cpp
+++ b/Engine/Plugins/Importers/USDImporter/Source/UnrealUSDWrapper/Private/UsdWrappers/SdfPrimSpec.cpp
@@ -9,13 +9,11 @@
 #include "UsdWrappers/SdfPath.h"
 
 #if USE_USD_SDK
-
 #include "USDIncludesStart.h"
-	#include "pxr/usd/sdf/attributeSpec.h"
-	#include "pxr/usd/sdf/primSpec.h"
+#include "pxr/usd/sdf/attributeSpec.h"
+#include "pxr/usd/sdf/primSpec.h"
 #include "USDIncludesEnd.h"
-
-#endif // #if USE_USD_SDK
+#endif	  // #if USE_USD_SDK
 
 namespace UE
 {
@@ -27,36 +25,36 @@ namespace UE
 			FSdfPrimSpecImpl() = default;
 
 #if USE_USD_SDK
-			explicit FSdfPrimSpecImpl( const pxr::SdfPrimSpecHandle& InSdfPrimSpec )
-				: PxrSdfPrimSpec( InSdfPrimSpec )
+			explicit FSdfPrimSpecImpl(const pxr::SdfPrimSpecHandle& InSdfPrimSpec)
+				: PxrSdfPrimSpec(InSdfPrimSpec)
 			{
 			}
 
-			explicit FSdfPrimSpecImpl( pxr::SdfPrimSpecHandle&& InSdfPrimSpec )
-				: PxrSdfPrimSpec( MoveTemp( InSdfPrimSpec ) )
+			explicit FSdfPrimSpecImpl(pxr::SdfPrimSpecHandle&& InSdfPrimSpec)
+				: PxrSdfPrimSpec(MoveTemp(InSdfPrimSpec))
 			{
 			}
 
-			TUsdStore< pxr::SdfPrimSpecHandle > PxrSdfPrimSpec;
-#endif // #if USE_USD_SDK
+			TUsdStore<pxr::SdfPrimSpecHandle> PxrSdfPrimSpec;
+#endif	  // #if USE_USD_SDK
 		};
 	}
 
 	FSdfPrimSpec::FSdfPrimSpec()
 	{
 		FScopedUnrealAllocs UnrealAllocs;
-		Impl = MakeUnique< Internal::FSdfPrimSpecImpl >();
+		Impl = MakeUnique<Internal::FSdfPrimSpecImpl>();
 	}
 
-	FSdfPrimSpec::FSdfPrimSpec( const FSdfPrimSpec& Other )
+	FSdfPrimSpec::FSdfPrimSpec(const FSdfPrimSpec& Other)
 	{
 #if USE_USD_SDK
 		FScopedUnrealAllocs UnrealAllocs;
-		Impl = MakeUnique< Internal::FSdfPrimSpecImpl >( Other.Impl->PxrSdfPrimSpec.Get() );
-#endif // #if USE_USD_SDK
+		Impl = MakeUnique<Internal::FSdfPrimSpecImpl>(Other.Impl->PxrSdfPrimSpec.Get());
+#endif	  // #if USE_USD_SDK
 	}
 
-	FSdfPrimSpec::FSdfPrimSpec( FSdfPrimSpec&& Other ) = default;
+	FSdfPrimSpec::FSdfPrimSpec(FSdfPrimSpec&& Other) = default;
 
 	FSdfPrimSpec::~FSdfPrimSpec()
 	{
@@ -64,19 +62,19 @@ namespace UE
 		Impl.Reset();
 	}
 
-	FSdfPrimSpec& FSdfPrimSpec::operator=( const FSdfPrimSpec& Other )
+	FSdfPrimSpec& FSdfPrimSpec::operator=(const FSdfPrimSpec& Other)
 	{
 #if USE_USD_SDK
 		FScopedUnrealAllocs UnrealAllocs;
-		Impl = MakeUnique< Internal::FSdfPrimSpecImpl >(  Other.Impl->PxrSdfPrimSpec.Get() );
-#endif // #if USE_USD_SDK
+		Impl = MakeUnique<Internal::FSdfPrimSpecImpl>(Other.Impl->PxrSdfPrimSpec.Get());
+#endif	  // #if USE_USD_SDK
 		return *this;
 	}
 
-	FSdfPrimSpec& FSdfPrimSpec::operator=( FSdfPrimSpec&& Other )
+	FSdfPrimSpec& FSdfPrimSpec::operator=(FSdfPrimSpec&& Other)
 	{
 		FScopedUnrealAllocs UnrealAllocs;
-		Impl = MoveTemp(  Other.Impl );
+		Impl = MoveTemp(Other.Impl);
 
 		return *this;
 	}
@@ -87,47 +85,47 @@ namespace UE
 		return (bool)Impl->PxrSdfPrimSpec.Get();
 #else
 		return false;
-#endif // #if USE_USD_SDK
+#endif	  // #if USE_USD_SDK
 	}
 
-	bool FSdfPrimSpec::operator==( const FSdfPrimSpec& Other ) const
+	bool FSdfPrimSpec::operator==(const FSdfPrimSpec& Other) const
 	{
 #if USE_USD_SDK
 		return Impl->PxrSdfPrimSpec.Get() == Other.Impl->PxrSdfPrimSpec.Get();
 #else
 		return false;
-#endif // #if USE_USD_SDK
+#endif	  // #if USE_USD_SDK
 	}
 
-	bool FSdfPrimSpec::operator!=( const FSdfPrimSpec& Other ) const
+	bool FSdfPrimSpec::operator!=(const FSdfPrimSpec& Other) const
 	{
-		return !( *this == Other );
+		return !(*this == Other);
 	}
 
 #if USE_USD_SDK
-	FSdfPrimSpec::FSdfPrimSpec( const pxr::SdfPrimSpecHandle& InSdfPrimSpec )
+	FSdfPrimSpec::FSdfPrimSpec(const pxr::SdfPrimSpecHandle& InSdfPrimSpec)
 	{
 		FScopedUnrealAllocs UnrealAllocs;
-		Impl =  MakeUnique< Internal::FSdfPrimSpecImpl >( InSdfPrimSpec );
+		Impl = MakeUnique<Internal::FSdfPrimSpecImpl>(InSdfPrimSpec);
 	}
 
-	FSdfPrimSpec::FSdfPrimSpec( pxr::SdfPrimSpecHandle&& InSdfPrimSpec )
+	FSdfPrimSpec::FSdfPrimSpec(pxr::SdfPrimSpecHandle&& InSdfPrimSpec)
 	{
 		FScopedUnrealAllocs UnrealAllocs;
-		Impl = MakeUnique< Internal::FSdfPrimSpecImpl >( MoveTemp( InSdfPrimSpec ) );
+		Impl = MakeUnique<Internal::FSdfPrimSpecImpl>(MoveTemp(InSdfPrimSpec));
 	}
 
-	FSdfPrimSpec& FSdfPrimSpec::operator=( const pxr::SdfPrimSpecHandle& InSdfPrimSpec )
+	FSdfPrimSpec& FSdfPrimSpec::operator=(const pxr::SdfPrimSpecHandle& InSdfPrimSpec)
 	{
 		FScopedUnrealAllocs UnrealAllocs;
-		Impl = MakeUnique< Internal::FSdfPrimSpecImpl >( InSdfPrimSpec );
+		Impl = MakeUnique<Internal::FSdfPrimSpecImpl>(InSdfPrimSpec);
 		return *this;
 	}
 
-	FSdfPrimSpec& FSdfPrimSpec::operator=( pxr::SdfPrimSpecHandle&& InSdfPrimSpec )
+	FSdfPrimSpec& FSdfPrimSpec::operator=(pxr::SdfPrimSpecHandle&& InSdfPrimSpec)
 	{
 		FScopedUnrealAllocs UnrealAllocs;
-		Impl = MakeUnique< Internal::FSdfPrimSpecImpl >( MoveTemp( InSdfPrimSpec ) );
+		Impl = MakeUnique<Internal::FSdfPrimSpecImpl>(MoveTemp(InSdfPrimSpec));
 		return *this;
 	}
 
@@ -140,92 +138,92 @@ namespace UE
 	{
 		return Impl->PxrSdfPrimSpec.Get();
 	}
-#endif // #if USE_USD_SDK
+#endif	  // #if USE_USD_SDK
 
 	ESdfSpecType FSdfPrimSpec::GetSpecType() const
 	{
 #if USE_USD_SDK
 		// Just in case they add another or reorder these
-		static_assert( ( int ) ESdfSpecType::SdfSpecTypeUnknown == ( int ) pxr::SdfSpecType::SdfSpecTypeUnknown );
-		static_assert( ( int ) ESdfSpecType::SdfSpecTypeAttribute == ( int ) pxr::SdfSpecType::SdfSpecTypeAttribute );
-		static_assert( ( int ) ESdfSpecType::SdfSpecTypeConnection == ( int ) pxr::SdfSpecType::SdfSpecTypeConnection );
-		static_assert( ( int ) ESdfSpecType::SdfSpecTypeExpression == ( int ) pxr::SdfSpecType::SdfSpecTypeExpression );
-		static_assert( ( int ) ESdfSpecType::SdfSpecTypeMapper == ( int ) pxr::SdfSpecType::SdfSpecTypeMapper );
-		static_assert( ( int ) ESdfSpecType::SdfSpecTypeMapperArg == ( int ) pxr::SdfSpecType::SdfSpecTypeMapperArg );
-		static_assert( ( int ) ESdfSpecType::SdfSpecTypePrim == ( int ) pxr::SdfSpecType::SdfSpecTypePrim );
-		static_assert( ( int ) ESdfSpecType::SdfSpecTypePseudoRoot == ( int ) pxr::SdfSpecType::SdfSpecTypePseudoRoot );
-		static_assert( ( int ) ESdfSpecType::SdfSpecTypeRelationship == ( int ) pxr::SdfSpecType::SdfSpecTypeRelationship );
-		static_assert( ( int ) ESdfSpecType::SdfSpecTypeRelationshipTarget == ( int ) pxr::SdfSpecType::SdfSpecTypeRelationshipTarget );
-		static_assert( ( int ) ESdfSpecType::SdfSpecTypeVariant == ( int ) pxr::SdfSpecType::SdfSpecTypeVariant );
-		static_assert( ( int ) ESdfSpecType::SdfSpecTypeVariantSet == ( int ) pxr::SdfSpecType::SdfSpecTypeVariantSet );
-		static_assert( ( int ) ESdfSpecType::SdfNumSpecTypes == ( int ) pxr::SdfSpecType::SdfNumSpecTypes );
-
-		return static_cast< ESdfSpecType >( Impl->PxrSdfPrimSpec.Get()->GetSpecType() );
+		static_assert((int)ESdfSpecType::SdfSpecTypeUnknown == (int)pxr::SdfSpecType::SdfSpecTypeUnknown);
+		static_assert((int)ESdfSpecType::SdfSpecTypeAttribute == (int)pxr::SdfSpecType::SdfSpecTypeAttribute);
+		static_assert((int)ESdfSpecType::SdfSpecTypeConnection == (int)pxr::SdfSpecType::SdfSpecTypeConnection);
+		static_assert((int)ESdfSpecType::SdfSpecTypeExpression == (int)pxr::SdfSpecType::SdfSpecTypeExpression);
+		static_assert((int)ESdfSpecType::SdfSpecTypeMapper == (int)pxr::SdfSpecType::SdfSpecTypeMapper);
+		static_assert((int)ESdfSpecType::SdfSpecTypeMapperArg == (int)pxr::SdfSpecType::SdfSpecTypeMapperArg);
+		static_assert((int)ESdfSpecType::SdfSpecTypePrim == (int)pxr::SdfSpecType::SdfSpecTypePrim);
+		static_assert((int)ESdfSpecType::SdfSpecTypePseudoRoot == (int)pxr::SdfSpecType::SdfSpecTypePseudoRoot);
+		static_assert((int)ESdfSpecType::SdfSpecTypeRelationship == (int)pxr::SdfSpecType::SdfSpecTypeRelationship);
+		static_assert((int)ESdfSpecType::SdfSpecTypeRelationshipTarget == (int)pxr::SdfSpecType::SdfSpecTypeRelationshipTarget);
+		static_assert((int)ESdfSpecType::SdfSpecTypeVariant == (int)pxr::SdfSpecType::SdfSpecTypeVariant);
+		static_assert((int)ESdfSpecType::SdfSpecTypeVariantSet == (int)pxr::SdfSpecType::SdfSpecTypeVariantSet);
+		static_assert((int)ESdfSpecType::SdfNumSpecTypes == (int)pxr::SdfSpecType::SdfNumSpecTypes);
+
+		return static_cast<ESdfSpecType>(Impl->PxrSdfPrimSpec.Get()->GetSpecType());
 #else
 		return ESdfSpecType::SdfSpecTypeUnknown;
-#endif // #if USE_USD_SDK
+#endif	  // #if USE_USD_SDK
 	}
 
 	FSdfLayerWeak FSdfPrimSpec::GetLayer() const
 	{
 #if USE_USD_SDK
-		return FSdfLayerWeak( Impl->PxrSdfPrimSpec.Get()->GetLayer() );
+		return FSdfLayerWeak(Impl->PxrSdfPrimSpec.Get()->GetLayer());
 #else
 		return FSdfLayerWeak();
-#endif // #if USE_USD_SDK
+#endif	  // #if USE_USD_SDK
 	}
 
 	FSdfPath FSdfPrimSpec::GetPath() const
 	{
 #if USE_USD_SDK
-		return FSdfPath( Impl->PxrSdfPrimSpec.Get()->GetPath() );
+		return FSdfPath(Impl->PxrSdfPrimSpec.Get()->GetPath());
 #else
 		return FSdfPath();
-#endif // #if USE_USD_SDK
+#endif	  // #if USE_USD_SDK
 	}
 
 	FSdfPrimSpec FSdfPrimSpec::GetRealNameParent() const
 	{
 #if USE_USD_SDK
-		return FSdfPrimSpec{ Impl->PxrSdfPrimSpec.Get()->GetRealNameParent() };
+		return FSdfPrimSpec{Impl->PxrSdfPrimSpec.Get()->GetRealNameParent()};
 #else
 		return FSdfPrimSpec{};
-#endif // #if USE_USD_SDK
+#endif	  // #if USE_USD_SDK
 	}
 
-	bool FSdfPrimSpec::RemoveNameChild( const FSdfPrimSpec& Child )
+	bool FSdfPrimSpec::RemoveNameChild(const FSdfPrimSpec& Child)
 	{
 #if USE_USD_SDK
-		return Impl->PxrSdfPrimSpec.Get()->RemoveNameChild( Child );
+		return Impl->PxrSdfPrimSpec.Get()->RemoveNameChild(Child);
 #else
 		return false;
-#endif // #if USE_USD_SDK
+#endif	  // #if USE_USD_SDK
 	}
 
-	UE::FSdfAttributeSpec FSdfPrimSpec::GetAttributeAtPath( const UE::FSdfPath& Path ) const
+	UE::FSdfAttributeSpec FSdfPrimSpec::GetAttributeAtPath(const UE::FSdfPath& Path) const
 	{
 #if USE_USD_SDK
-		return UE::FSdfAttributeSpec{ Impl->PxrSdfPrimSpec.Get()->GetAttributeAtPath( Path ) };
+		return UE::FSdfAttributeSpec{Impl->PxrSdfPrimSpec.Get()->GetAttributeAtPath(Path)};
 #else
 		return UE::FSdfAttributeSpec{};
-#endif // #if USE_USD_SDK
+#endif	  // #if USE_USD_SDK
 	}
 
 	FName FSdfPrimSpec::GetTypeName() const
 	{
 #if USE_USD_SDK
-		return FName( ANSI_TO_TCHAR( Impl->PxrSdfPrimSpec.Get()->GetTypeName().GetString().c_str() ) );
+		return FName(ANSI_TO_TCHAR(Impl->PxrSdfPrimSpec.Get()->GetTypeName().GetString().c_str()));
 #else
 		return FName();
-#endif // #if USE_USD_SDK
+#endif	  // #if USE_USD_SDK
 	}
 
 	FName FSdfPrimSpec::GetName() const
 	{
 #if USE_USD_SDK
-		return FName( ANSI_TO_TCHAR( Impl->PxrSdfPrimSpec.Get()->GetName().c_str() ) );
+		return FName(ANSI_TO_TCHAR(Impl->PxrSdfPrimSpec.Get()->GetName().c_str()));
 #else
 		return FName();
-#endif // #if USE_USD_SDK
+#endif	  // #if USE_USD_SDK
 	}
-}
+}	 // namespace UE
diff --git a/Engine/Plugins/Importers/USDImporter/Source/UnrealUSDWrapper/Private/UsdWrappers/UsdAttribute.cpp b/Engine/Plugins/Importers/USDImporter/Source/UnrealUSDWrapper/Private/UsdWrappers/UsdAttribute.cpp
--- a/Engine/Plugins/Importers/USDImporter/Source/UnrealUSDWrapper/Private/UsdWrappers/UsdAttribute.cpp
+++ b/Engine/Plugins/Importers/USDImporter/Source/UnrealUSDWrapper/Private/UsdWrappers/UsdAttribute.cpp
@@ -9,12 +9,10 @@
 #include "USDMemory.h"
 
 #if USE_USD_SDK
-
 #include "USDIncludesStart.h"
-	#include "pxr/usd/usd/attribute.h"
+#include "pxr/usd/usd/attribute.h"
 #include "USDIncludesEnd.h"
-
-#endif // #if USE_USD_SDK
+#endif	  // #if USE_USD_SDK
 
 namespace UE
 {
@@ -26,18 +24,18 @@ namespace UE
 			FUsdAttributeImpl() = default;
 
 #if USE_USD_SDK
-			explicit FUsdAttributeImpl( const pxr::UsdAttribute& InUsdAttribute )
-				: PxrUsdAttribute( InUsdAttribute )
+			explicit FUsdAttributeImpl(const pxr::UsdAttribute& InUsdAttribute)
+				: PxrUsdAttribute(InUsdAttribute)
 			{
 			}
 
-			explicit FUsdAttributeImpl( pxr::UsdAttribute&& InUsdAttribute )
-				: PxrUsdAttribute( MoveTemp( InUsdAttribute ) )
+			explicit FUsdAttributeImpl(pxr::UsdAttribute&& InUsdAttribute)
+				: PxrUsdAttribute(MoveTemp(InUsdAttribute))
 			{
 			}
 
-			TUsdStore< pxr::UsdAttribute > PxrUsdAttribute;
-#endif // #if USE_USD_SDK
+			TUsdStore<pxr::UsdAttribute> PxrUsdAttribute;
+#endif	  // #if USE_USD_SDK
 		};
 	}
 }
@@ -47,18 +45,18 @@ namespace UE
 	FUsdAttribute::FUsdAttribute()
 	{
 		FScopedUnrealAllocs UnrealAllocs;
-		Impl =  MakeUnique< Internal::FUsdAttributeImpl >();
+		Impl = MakeUnique<Internal::FUsdAttributeImpl>();
 	}
 
-	FUsdAttribute::FUsdAttribute( const FUsdAttribute& Other )
+	FUsdAttribute::FUsdAttribute(const FUsdAttribute& Other)
 	{
 #if USE_USD_SDK
 		FScopedUnrealAllocs UnrealAllocs;
-		Impl = MakeUnique< Internal::FUsdAttributeImpl >( Other.Impl->PxrUsdAttribute.Get() );
-#endif // #if USE_USD_SDK
+		Impl = MakeUnique<Internal::FUsdAttributeImpl>(Other.Impl->PxrUsdAttribute.Get());
+#endif	  // #if USE_USD_SDK
 	}
 
-	FUsdAttribute::FUsdAttribute( FUsdAttribute&& Other ) = default;
+	FUsdAttribute::FUsdAttribute(FUsdAttribute&& Other) = default;
 
 	FUsdAttribute::~FUsdAttribute()
 	{
@@ -66,35 +64,35 @@ namespace UE
 		Impl.Reset();
 	}
 
-	FUsdAttribute& FUsdAttribute::operator=( const FUsdAttribute& Other )
+	FUsdAttribute& FUsdAttribute::operator=(const FUsdAttribute& Other)
 	{
 #if USE_USD_SDK
 		FScopedUnrealAllocs UnrealAllocs;
-		Impl = MakeUnique< Internal::FUsdAttributeImpl >(  Other.Impl->PxrUsdAttribute.Get() );
-#endif // #if USE_USD_SDK
+		Impl = MakeUnique<Internal::FUsdAttributeImpl>(Other.Impl->PxrUsdAttribute.Get());
+#endif	  // #if USE_USD_SDK
 		return *this;
 	}
 
-	FUsdAttribute& FUsdAttribute::operator=( FUsdAttribute&& Other )
+	FUsdAttribute& FUsdAttribute::operator=(FUsdAttribute&& Other)
 	{
 		FScopedUnrealAllocs UnrealAllocs;
-		Impl = MoveTemp( Other.Impl );
+		Impl = MoveTemp(Other.Impl);
 
 		return *this;
 	}
 
-	bool FUsdAttribute::operator==( const FUsdAttribute& Other ) const
+	bool FUsdAttribute::operator==(const FUsdAttribute& Other) const
 	{
 #if USE_USD_SDK
 		return Impl->PxrUsdAttribute.Get() == Other.Impl->PxrUsdAttribute.Get();
 #else
 		return false;
-#endif // #if USE_USD_SDK
+#endif	  // #if USE_USD_SDK
 	}
 
-	bool FUsdAttribute::operator!=( const FUsdAttribute& Other ) const
+	bool FUsdAttribute::operator!=(const FUsdAttribute& Other) const
 	{
-		return !( *this == Other );
+		return !(*this == Other);
 	}
 
 	FUsdAttribute::operator bool() const
@@ -103,33 +101,33 @@ namespace UE
 		return (bool)Impl->PxrUsdAttribute.Get();
 #else
 		return false;
-#endif // #if USE_USD_SDK
+#endif	  // #if USE_USD_SDK
 	}
 
 #if USE_USD_SDK
-	FUsdAttribute::FUsdAttribute( const pxr::UsdAttribute& InUsdAttribute )
+	FUsdAttribute::FUsdAttribute(const pxr::UsdAttribute& InUsdAttribute)
 	{
 		FScopedUnrealAllocs UnrealAllocs;
-		Impl =  MakeUnique< Internal::FUsdAttributeImpl >( InUsdAttribute );
+		Impl = MakeUnique<Internal::FUsdAttributeImpl>(InUsdAttribute);
 	}
 
-	FUsdAttribute::FUsdAttribute( pxr::UsdAttribute&& InUsdAttribute )
+	FUsdAttribute::FUsdAttribute(pxr::UsdAttribute&& InUsdAttribute)
 	{
 		FScopedUnrealAllocs UnrealAllocs;
-		Impl = MakeUnique< Internal::FUsdAttributeImpl >( MoveTemp( InUsdAttribute ) );
+		Impl = MakeUnique<Internal::FUsdAttributeImpl>(MoveTemp(InUsdAttribute));
 	}
 
-	FUsdAttribute& FUsdAttribute::operator=( const pxr::UsdAttribute& InUsdAttribute )
+	FUsdAttribute& FUsdAttribute::operator=(const pxr::UsdAttribute& InUsdAttribute)
 	{
 		FScopedUnrealAllocs UnrealAllocs;
-		Impl = MakeUnique< Internal::FUsdAttributeImpl >( InUsdAttribute );
+		Impl = MakeUnique<Internal::FUsdAttributeImpl>(InUsdAttribute);
 		return *this;
 	}
 
-	FUsdAttribute& FUsdAttribute::operator=( pxr::UsdAttribute&& InUsdAttribute )
+	FUsdAttribute& FUsdAttribute::operator=(pxr::UsdAttribute&& InUsdAttribute)
 	{
 		FScopedUnrealAllocs UnrealAllocs;
-		Impl = MakeUnique< Internal::FUsdAttributeImpl >( MoveTemp( InUsdAttribute ) );
+		Impl = MakeUnique<Internal::FUsdAttributeImpl>(MoveTemp(InUsdAttribute));
 		return *this;
 	}
 
@@ -152,90 +150,90 @@ namespace UE
 	{
 		return Impl->PxrUsdAttribute.Get();
 	}
-#endif // #if USE_USD_SDK
+#endif	  // #if USE_USD_SDK
 
-	bool FUsdAttribute::GetMetadata( const TCHAR* Key, UE::FVtValue& Value ) const
+	bool FUsdAttribute::GetMetadata(const TCHAR* Key, UE::FVtValue& Value) const
 	{
 #if USE_USD_SDK
-		return Impl->PxrUsdAttribute.Get().GetMetadata( pxr::TfToken{ TCHAR_TO_ANSI( Key ) }, &Value.GetUsdValue() );
+		return Impl->PxrUsdAttribute.Get().GetMetadata(pxr::TfToken{TCHAR_TO_ANSI(Key)}, &Value.GetUsdValue());
 #else
 		return false;
-#endif // #if USE_USD_SDK
+#endif	  // #if USE_USD_SDK
 	}
 
-	bool FUsdAttribute::HasMetadata( const TCHAR* Key ) const
+	bool FUsdAttribute::HasMetadata(const TCHAR* Key) const
 	{
 #if USE_USD_SDK
-		return Impl->PxrUsdAttribute.Get().HasMetadata( pxr::TfToken{ TCHAR_TO_ANSI( Key ) } );
+		return Impl->PxrUsdAttribute.Get().HasMetadata(pxr::TfToken{TCHAR_TO_ANSI(Key)});
 #else
 		return false;
-#endif // #if USE_USD_SDK
+#endif	  // #if USE_USD_SDK
 	}
 
-	bool FUsdAttribute::SetMetadata( const TCHAR* Key, const UE::FVtValue& Value ) const
+	bool FUsdAttribute::SetMetadata(const TCHAR* Key, const UE::FVtValue& Value) const
 	{
 #if USE_USD_SDK
-		return Impl->PxrUsdAttribute.Get().SetMetadata( pxr::TfToken{ TCHAR_TO_ANSI( Key ) }, Value.GetUsdValue() );
+		return Impl->PxrUsdAttribute.Get().SetMetadata(pxr::TfToken{TCHAR_TO_ANSI(Key)}, Value.GetUsdValue());
 #else
 		return false;
-#endif // #if USE_USD_SDK
+#endif	  // #if USE_USD_SDK
 	}
 
-	bool FUsdAttribute::ClearMetadata( const TCHAR* Key ) const
+	bool FUsdAttribute::ClearMetadata(const TCHAR* Key) const
 	{
 #if USE_USD_SDK
-		return Impl->PxrUsdAttribute.Get().ClearMetadata( pxr::TfToken{ TCHAR_TO_ANSI( Key ) } );
+		return Impl->PxrUsdAttribute.Get().ClearMetadata(pxr::TfToken{TCHAR_TO_ANSI(Key)});
 #else
 		return false;
-#endif // #if USE_USD_SDK
+#endif	  // #if USE_USD_SDK
 	}
 
 	FName FUsdAttribute::GetName() const
 	{
 #if USE_USD_SDK
-		return FName( ANSI_TO_TCHAR( Impl->PxrUsdAttribute.Get().GetName().GetString().c_str() ) );
+		return FName(ANSI_TO_TCHAR(Impl->PxrUsdAttribute.Get().GetName().GetString().c_str()));
 #else
 		return FName();
-#endif // #if USE_USD_SDK
+#endif	  // #if USE_USD_SDK
 	}
 
 	FName FUsdAttribute::GetBaseName() const
 	{
 #if USE_USD_SDK
-		return FName( ANSI_TO_TCHAR( Impl->PxrUsdAttribute.Get().GetBaseName().GetString().c_str() ) );
+		return FName(ANSI_TO_TCHAR(Impl->PxrUsdAttribute.Get().GetBaseName().GetString().c_str()));
 #else
 		return FName();
-#endif // #if USE_USD_SDK
+#endif	  // #if USE_USD_SDK
 	}
 
 	FName FUsdAttribute::GetTypeName() const
 	{
 #if USE_USD_SDK
-		return FName( ANSI_TO_TCHAR( Impl->PxrUsdAttribute.Get().GetTypeName().GetAsToken().GetString().c_str() ) );
+		return FName(ANSI_TO_TCHAR(Impl->PxrUsdAttribute.Get().GetTypeName().GetAsToken().GetString().c_str()));
 #else
 		return FName();
-#endif // #if USE_USD_SDK
+#endif	  // #if USE_USD_SDK
 	}
 
-	bool FUsdAttribute::GetTimeSamples( TArray<double>& Times ) const
+	bool FUsdAttribute::GetTimeSamples(TArray<double>& Times) const
 	{
 #if USE_USD_SDK
 		FScopedUsdAllocs Allocs;
 
 		std::vector<double> UsdTimes;
-		bool bResult = Impl->PxrUsdAttribute.Get().GetTimeSamples( &UsdTimes );
-		if ( !bResult )
+		bool bResult = Impl->PxrUsdAttribute.Get().GetTimeSamples(&UsdTimes);
+		if (!bResult)
 		{
 			return false;
 		}
 
-		Times.SetNumUninitialized( UsdTimes.size() );
-		FMemory::Memcpy( Times.GetData(), UsdTimes.data(), UsdTimes.size() * sizeof( double ) );
+		Times.SetNumUninitialized(UsdTimes.size());
+		FMemory::Memcpy(Times.GetData(), UsdTimes.data(), UsdTimes.size() * sizeof(double));
 
 		return true;
 #else
 		return false;
-#endif // #if USE_USD_SDK
+#endif	  // #if USE_USD_SDK
 	}
 
 	size_t FUsdAttribute::GetNumTimeSamples() const
@@ -244,7 +242,7 @@ namespace UE
 		return Impl->PxrUsdAttribute.Get().GetNumTimeSamples();
 #else
 		return 0;
-#endif // #if USE_USD_SDK
+#endif	  // #if USE_USD_SDK
 	}
 
 	bool FUsdAttribute::HasValue() const
@@ -253,7 +251,7 @@ namespace UE
 		return Impl->PxrUsdAttribute.Get().HasValue();
 #else
 		return false;
-#endif // #if USE_USD_SDK
+#endif	  // #if USE_USD_SDK
 	}
 
 	bool FUsdAttribute::HasAuthoredValue() const
@@ -262,7 +260,7 @@ namespace UE
 		return Impl->PxrUsdAttribute.Get().HasAuthoredValue();
 #else
 		return false;
-#endif // #if USE_USD_SDK
+#endif	  // #if USE_USD_SDK
 	}
 
 	bool FUsdAttribute::HasFallbackValue() const
@@ -271,7 +269,7 @@ namespace UE
 		return Impl->PxrUsdAttribute.Get().HasFallbackValue();
 #else
 		return false;
-#endif // #if USE_USD_SDK
+#endif	  // #if USE_USD_SDK
 	}
 
 	bool FUsdAttribute::ValueMightBeTimeVarying() const
@@ -280,31 +278,31 @@ namespace UE
 		return Impl->PxrUsdAttribute.Get().ValueMightBeTimeVarying();
 #else
 		return false;
-#endif // #if USE_USD_SDK
+#endif	  // #if USE_USD_SDK
 	}
 
-	bool FUsdAttribute::Get( UE::FVtValue & Value, TOptional<double> Time /*= {} */ ) const
+	bool FUsdAttribute::Get(UE::FVtValue& Value, TOptional<double> Time /*= {} */) const
 	{
 #if USE_USD_SDK
 		FScopedUsdAllocs UsdAllocs;
 
 		pxr::UsdTimeCode TimeCode = Time.IsSet() ? Time.GetValue() : pxr::UsdTimeCode::Default();
-		return Impl->PxrUsdAttribute.Get().Get( &Value.GetUsdValue(), TimeCode );
+		return Impl->PxrUsdAttribute.Get().Get(&Value.GetUsdValue(), TimeCode);
 #else
 		return false;
-#endif // #if USE_USD_SDK
+#endif	  // #if USE_USD_SDK
 	}
 
-	bool FUsdAttribute::Set( const UE::FVtValue & Value, TOptional<double> Time /*= {} */ ) const
+	bool FUsdAttribute::Set(const UE::FVtValue& Value, TOptional<double> Time /*= {} */) const
 	{
 #if USE_USD_SDK
 		FScopedUsdAllocs UsdAllocs;
 
 		pxr::UsdTimeCode TimeCode = Time.IsSet() ? Time.GetValue() : pxr::UsdTimeCode::Default();
-		return Impl->PxrUsdAttribute.Get().Set( Value.GetUsdValue(), TimeCode );
+		return Impl->PxrUsdAttribute.Get().Set(Value.GetUsdValue(), TimeCode);
 #else
 		return false;
-#endif // #if USE_USD_SDK
+#endif	  // #if USE_USD_SDK
 	}
 
 	bool FUsdAttribute::Clear() const
@@ -313,19 +311,19 @@ namespace UE
 		return Impl->PxrUsdAttribute.Get().Clear();
 #else
 		return false;
-#endif // #if USE_USD_SDK
+#endif	  // #if USE_USD_SDK
 	}
 
-	bool FUsdAttribute::ClearAtTime( double Time ) const
+	bool FUsdAttribute::ClearAtTime(double Time) const
 	{
 #if USE_USD_SDK
-		return Impl->PxrUsdAttribute.Get().ClearAtTime( pxr::UsdTimeCode( Time ) );
+		return Impl->PxrUsdAttribute.Get().ClearAtTime(pxr::UsdTimeCode(Time));
 #else
 		return false;
-#endif // #if USE_USD_SDK
+#endif	  // #if USE_USD_SDK
 	}
 
-	bool FUsdAttribute::GetUnionedTimeSamples( const TArray<UE::FUsdAttribute>& Attrs, TArray<double>& OutTimes )
+	bool FUsdAttribute::GetUnionedTimeSamples(const TArray<UE::FUsdAttribute>& Attrs, TArray<double>& OutTimes)
 	{
 		bool bResult = false;
 
@@ -333,21 +331,21 @@ namespace UE
 		FScopedUsdAllocs UsdAllocs;
 
 		std::vector<pxr::UsdAttribute> UsdAttrs;
-		UsdAttrs.reserve( Attrs.Num() );
-		for ( const UE::FUsdAttribute& Attr : Attrs )
+		UsdAttrs.reserve(Attrs.Num());
+		for (const UE::FUsdAttribute& Attr : Attrs)
 		{
-			UsdAttrs.push_back( Attr );
+			UsdAttrs.push_back(Attr);
 		}
 
 		std::vector<double> UsdTimes;
 
-		bResult = pxr::UsdAttribute::GetUnionedTimeSamples( UsdAttrs, &UsdTimes );
-		if ( bResult )
+		bResult = pxr::UsdAttribute::GetUnionedTimeSamples(UsdAttrs, &UsdTimes);
+		if (bResult)
 		{
-			OutTimes.SetNumUninitialized( UsdTimes.size() );
-			FMemory::Memcpy( OutTimes.GetData(), UsdTimes.data(), OutTimes.Num() * OutTimes.GetTypeSize() );
+			OutTimes.SetNumUninitialized(UsdTimes.size());
+			FMemory::Memcpy(OutTimes.GetData(), UsdTimes.data(), OutTimes.Num() * OutTimes.GetTypeSize());
 		}
-#endif // #if USE_USD_SDK
+#endif	  // #if USE_USD_SDK
 
 		return bResult;
 	}
@@ -355,18 +353,18 @@ namespace UE
 	FSdfPath FUsdAttribute::GetPath() const
 	{
 #if USE_USD_SDK
-		return FSdfPath( Impl->PxrUsdAttribute.Get().GetPath() );
+		return FSdfPath(Impl->PxrUsdAttribute.Get().GetPath());
 #else
 		return FSdfPath();
-#endif // #if USE_USD_SDK
+#endif	  // #if USE_USD_SDK
 	}
 
 	FUsdPrim FUsdAttribute::GetPrim() const
 	{
 #if USE_USD_SDK
-		return FUsdPrim( Impl->PxrUsdAttribute.Get().GetPrim() );
+		return FUsdPrim(Impl->PxrUsdAttribute.Get().GetPrim());
 #else
 		return FUsdPrim();
-#endif // #if USE_USD_SDK
+#endif	  // #if USE_USD_SDK
 	}
-}
+}	 // namespace UE
diff --git a/Engine/Plugins/Importers/USDImporter/Source/UnrealUSDWrapper/Private/UsdWrappers/UsdEditContext.cpp b/Engine/Plugins/Importers/USDImporter/Source/UnrealUSDWrapper/Private/UsdWrappers/UsdEditContext.cpp
--- a/Engine/Plugins/Importers/USDImporter/Source/UnrealUSDWrapper/Private/UsdWrappers/UsdEditContext.cpp
+++ b/Engine/Plugins/Importers/USDImporter/Source/UnrealUSDWrapper/Private/UsdWrappers/UsdEditContext.cpp
@@ -6,13 +6,10 @@
 #include "UsdWrappers/UsdStage.h"
 
 #if USE_USD_SDK
-
 #include "USDIncludesStart.h"
-	#include "pxr/usd/sdf/layer.h"
-	#include "pxr/usd/usd/editContext.h"
+#include "pxr/usd/usd/editContext.h"
 #include "USDIncludesEnd.h"
-
-#endif // #if USE_USD_SDK
+#endif	  // #if USE_USD_SDK
 
 namespace UE
 {
@@ -22,31 +19,31 @@ namespace UE
 		{
 #if USE_USD_SDK
 		public:
-			explicit FUsdEditContextImpl( const pxr::UsdStageWeakPtr& StageWeakPtr, const pxr::UsdEditTarget& EditTarget = pxr::UsdEditTarget() )
-				: UsdEditContext( StageWeakPtr, EditTarget )
+			explicit FUsdEditContextImpl(const pxr::UsdStageWeakPtr& StageWeakPtr, const pxr::UsdEditTarget& EditTarget = pxr::UsdEditTarget())
+				: UsdEditContext(StageWeakPtr, EditTarget)
 			{
 			}
 
 		private:
 			pxr::UsdEditContext UsdEditContext;
-#endif // #if USE_USD_SDK
+#endif	  // #if USE_USD_SDK
 		};
 	}
 
-	FUsdEditContext::FUsdEditContext( const FUsdStageWeak& Stage )
+	FUsdEditContext::FUsdEditContext(const FUsdStageWeak& Stage)
 	{
 #if USE_USD_SDK
 		FScopedUsdAllocs Allocs;
-		Impl = MakeUnique< Internal::FUsdEditContextImpl >(Stage);
-#endif // #if USE_USD_SDK
+		Impl = MakeUnique<Internal::FUsdEditContextImpl>(Stage);
+#endif	  // #if USE_USD_SDK
 	}
 
-	FUsdEditContext::FUsdEditContext( const FUsdStageWeak& Stage, const FSdfLayer& EditTarget )
+	FUsdEditContext::FUsdEditContext(const FUsdStageWeak& Stage, const FSdfLayer& EditTarget)
 	{
 #if USE_USD_SDK
 		FScopedUsdAllocs Allocs;
-		Impl = MakeUnique< Internal::FUsdEditContextImpl >( Stage, pxr::SdfLayerWeakPtr{ EditTarget } );
-#endif // #if USE_USD_SDK
+		Impl = MakeUnique<Internal::FUsdEditContextImpl>(Stage, pxr::SdfLayerWeakPtr{EditTarget});
+#endif	  // #if USE_USD_SDK
 	}
 
 	FUsdEditContext::~FUsdEditContext()
@@ -54,6 +51,6 @@ namespace UE
 #if USE_USD_SDK
 		FScopedUsdAllocs Allocs;
 		Impl.Reset();
-#endif // #if USE_USD_SDK
+#endif	  // #if USE_USD_SDK
 	}
-}
\ No newline at end of file
+}	 // namespace UE
diff --git a/Engine/Plugins/Importers/USDImporter/Source/UnrealUSDWrapper/Private/UsdWrappers/UsdGeomBBoxCache.cpp b/Engine/Plugins/Importers/USDImporter/Source/UnrealUSDWrapper/Private/UsdWrappers/UsdGeomBBoxCache.cpp
--- a/Engine/Plugins/Importers/USDImporter/Source/UnrealUSDWrapper/Private/UsdWrappers/UsdGeomBBoxCache.cpp
+++ b/Engine/Plugins/Importers/USDImporter/Source/UnrealUSDWrapper/Private/UsdWrappers/UsdGeomBBoxCache.cpp
@@ -2,16 +2,13 @@
 
 #include "UsdWrappers/UsdGeomBBoxCache.h"
 
+#include "UnrealUSDWrapper.h"
 #include "USDMemory.h"
 
-#include "UsdWrappers/UsdPrim.h"
-
 #if USE_USD_SDK
-
 #include "USDIncludesStart.h"
 #include "pxr/usd/usdGeom/bboxCache.h"
 #include "USDIncludesEnd.h"
-
 #endif	  //  #if USE_USD_SDK
 
 namespace UE
@@ -170,7 +167,7 @@ namespace UE
 		FScopedUsdAllocs Allocs;
 		FWriteScopeLock ScopeLock{Lock};
 		Impl->PxrUsdGeomBBoxCache.Clear();
-#endif // #if USE_USD_SDK
+#endif	  // #if USE_USD_SDK
 	}
 
 	void FUsdGeomBBoxCache::SetIncludedPurposes(EUsdPurpose IncludedPurposeFlags)
@@ -199,7 +196,7 @@ namespace UE
 		}
 
 		Impl->PxrUsdGeomBBoxCache.SetIncludedPurposes(TokenVector);
-#endif // #if USE_USD_SDK
+#endif	  // #if USE_USD_SDK
 	}
 
 	EUsdPurpose FUsdGeomBBoxCache::GetIncludedPurposes() const
diff --git a/Engine/Plugins/Importers/USDImporter/Source/UnrealUSDWrapper/Private/UsdWrappers/UsdGeomXformable.cpp b/Engine/Plugins/Importers/USDImporter/Source/UnrealUSDWrapper/Private/UsdWrappers/UsdGeomXformable.cpp
--- a/Engine/Plugins/Importers/USDImporter/Source/UnrealUSDWrapper/Private/UsdWrappers/UsdGeomXformable.cpp
+++ b/Engine/Plugins/Importers/USDImporter/Source/UnrealUSDWrapper/Private/UsdWrappers/UsdGeomXformable.cpp
@@ -8,12 +8,10 @@
 #include "UsdWrappers/UsdPrim.h"
 
 #if USE_USD_SDK
-
 #include "USDIncludesStart.h"
-	#include "pxr/usd/usdGeom/xformable.h"
+#include "pxr/usd/usdGeom/xformable.h"
 #include "USDIncludesEnd.h"
-
-#endif //  #if USE_USD_SDK
+#endif	  //  #if USE_USD_SDK
 
 namespace UE
 {
@@ -25,74 +23,74 @@ namespace UE
 			FUsdGeomXformableImpl() = default;
 
 #if USE_USD_SDK
-			explicit FUsdGeomXformableImpl( const pxr::UsdGeomXformable& InUsdGeomXformable )
-				: PxrUsdGeomXformable( InUsdGeomXformable )
+			explicit FUsdGeomXformableImpl(const pxr::UsdGeomXformable& InUsdGeomXformable)
+				: PxrUsdGeomXformable(InUsdGeomXformable)
 			{
 			}
 
-			explicit FUsdGeomXformableImpl( pxr::UsdGeomXformable&& InUsdGeomXformable )
-				: PxrUsdGeomXformable( MoveTemp( InUsdGeomXformable ) )
+			explicit FUsdGeomXformableImpl(pxr::UsdGeomXformable&& InUsdGeomXformable)
+				: PxrUsdGeomXformable(MoveTemp(InUsdGeomXformable))
 			{
 			}
 
-			explicit FUsdGeomXformableImpl( const pxr::UsdPrim& InUsdPrim )
-				: PxrUsdGeomXformable( pxr::UsdGeomXformable( InUsdPrim ) )
+			explicit FUsdGeomXformableImpl(const pxr::UsdPrim& InUsdPrim)
+				: PxrUsdGeomXformable(pxr::UsdGeomXformable(InUsdPrim))
 			{
 			}
 
-			TUsdStore< pxr::UsdGeomXformable > PxrUsdGeomXformable;
-#endif // #if USE_USD_SDK
+			TUsdStore<pxr::UsdGeomXformable> PxrUsdGeomXformable;
+#endif	  // #if USE_USD_SDK
 		};
 	}
 
 	FUsdGeomXformable::FUsdGeomXformable()
 	{
 		FScopedUnrealAllocs UnrealAllocs;
-		Impl = MakeUnique< Internal::FUsdGeomXformableImpl >();
+		Impl = MakeUnique<Internal::FUsdGeomXformableImpl>();
 	}
 
-	FUsdGeomXformable::FUsdGeomXformable( const FUsdGeomXformable& Other )
-		: FUsdTyped( Other )
+	FUsdGeomXformable::FUsdGeomXformable(const FUsdGeomXformable& Other)
+		: FUsdTyped(Other)
 	{
 #if USE_USD_SDK
 		FScopedUnrealAllocs UnrealAllocs;
-		Impl = MakeUnique< Internal::FUsdGeomXformableImpl >( Other.Impl->PxrUsdGeomXformable.Get() );
-#endif // #if USE_USD_SDK
+		Impl = MakeUnique<Internal::FUsdGeomXformableImpl>(Other.Impl->PxrUsdGeomXformable.Get());
+#endif	  // #if USE_USD_SDK
 	}
 
-	FUsdGeomXformable::FUsdGeomXformable( FUsdGeomXformable&& Other )
-		: FUsdTyped( MoveTemp( Other ) )
-		, Impl( MoveTemp( Other.Impl ) )
+	FUsdGeomXformable::FUsdGeomXformable(FUsdGeomXformable&& Other)
+		: FUsdTyped(MoveTemp(Other))
+		, Impl(MoveTemp(Other.Impl))
 	{
 	}
 
-	FUsdGeomXformable::FUsdGeomXformable( const FUsdPrim& Prim )
-		: FUsdTyped( Prim )
+	FUsdGeomXformable::FUsdGeomXformable(const FUsdPrim& Prim)
+		: FUsdTyped(Prim)
 	{
 #if USE_USD_SDK
 		FScopedUnrealAllocs UnrealAllocs;
-		Impl = MakeUnique< Internal::FUsdGeomXformableImpl >( Prim );
-#endif // #if USE_USD_SDK
+		Impl = MakeUnique<Internal::FUsdGeomXformableImpl>(Prim);
+#endif	  // #if USE_USD_SDK
 	}
 
-	FUsdGeomXformable& FUsdGeomXformable::operator=( const FUsdGeomXformable& Other )
+	FUsdGeomXformable& FUsdGeomXformable::operator=(const FUsdGeomXformable& Other)
 	{
-		FUsdTyped::operator=( Other );
+		FUsdTyped::operator=(Other);
 
 #if USE_USD_SDK
 		FScopedUnrealAllocs UnrealAllocs;
-		Impl = MakeUnique< Internal::FUsdGeomXformableImpl >( Other.Impl->PxrUsdGeomXformable.Get() );
-#endif // #if USE_USD_SDK
+		Impl = MakeUnique<Internal::FUsdGeomXformableImpl>(Other.Impl->PxrUsdGeomXformable.Get());
+#endif	  // #if USE_USD_SDK
 
 		return *this;
 	}
 
-	FUsdGeomXformable& FUsdGeomXformable::operator=( FUsdGeomXformable&& Other )
+	FUsdGeomXformable& FUsdGeomXformable::operator=(FUsdGeomXformable&& Other)
 	{
-		FUsdTyped::operator=( MoveTemp( Other ) );
+		FUsdTyped::operator=(MoveTemp(Other));
 
 		FScopedUnrealAllocs UnrealAllocs;
-		Impl = MoveTemp( Other.Impl );
+		Impl = MoveTemp(Other.Impl);
 
 		return *this;
 	}
@@ -109,26 +107,26 @@ namespace UE
 		return (bool)Impl->PxrUsdGeomXformable.Get();
 #else
 		return false;
-#endif // #if USE_USD_SDK
+#endif	  // #if USE_USD_SDK
 	}
 
 #if USE_USD_SDK
-	FUsdGeomXformable::FUsdGeomXformable( const pxr::UsdGeomXformable& InUsdGeomXformable )
+	FUsdGeomXformable::FUsdGeomXformable(const pxr::UsdGeomXformable& InUsdGeomXformable)
 	{
 		FScopedUnrealAllocs UnrealAllocs;
-		Impl = MakeUnique< Internal::FUsdGeomXformableImpl >( InUsdGeomXformable );
+		Impl = MakeUnique<Internal::FUsdGeomXformableImpl>(InUsdGeomXformable);
 	}
 
-	FUsdGeomXformable::FUsdGeomXformable( pxr::UsdGeomXformable&& InUsdGeomXformable )
+	FUsdGeomXformable::FUsdGeomXformable(pxr::UsdGeomXformable&& InUsdGeomXformable)
 	{
 		FScopedUnrealAllocs UnrealAllocs;
-		Impl = MakeUnique< Internal::FUsdGeomXformableImpl >( MoveTemp( InUsdGeomXformable ) );
+		Impl = MakeUnique<Internal::FUsdGeomXformableImpl>(MoveTemp(InUsdGeomXformable));
 	}
 
-	FUsdGeomXformable::FUsdGeomXformable( const pxr::UsdPrim& Prim )
+	FUsdGeomXformable::FUsdGeomXformable(const pxr::UsdPrim& Prim)
 	{
 		FScopedUnrealAllocs UnrealAllocs;
-		Impl = MakeUnique< Internal::FUsdGeomXformableImpl >( Prim );
+		Impl = MakeUnique<Internal::FUsdGeomXformableImpl>(Prim);
 	}
 
 	FUsdGeomXformable::operator pxr::UsdGeomXformable&()
@@ -140,7 +138,7 @@ namespace UE
 	{
 		return Impl->PxrUsdGeomXformable.Get();
 	}
-#endif // #if USE_USD_SDK
+#endif	  // #if USE_USD_SDK
 
 	bool FUsdGeomXformable::GetResetXformStack() const
 	{
@@ -148,7 +146,7 @@ namespace UE
 		return Impl->PxrUsdGeomXformable.Get().GetResetXformStack();
 #else
 		return false;
-#endif // #if USE_USD_SDK
+#endif	  // #if USE_USD_SDK
 	}
 
 	bool FUsdGeomXformable::TransformMightBeTimeVarying() const
@@ -157,32 +155,32 @@ namespace UE
 		return Impl->PxrUsdGeomXformable.Get().TransformMightBeTimeVarying();
 #else
 		return false;
-#endif // #if USE_USD_SDK
+#endif	  // #if USE_USD_SDK
 	}
 
-	bool FUsdGeomXformable::GetTimeSamples( TArray< double >* Times ) const
+	bool FUsdGeomXformable::GetTimeSamples(TArray<double>* Times) const
 	{
-		if ( !Times )
+		if (!Times)
 		{
 			return false;
 		}
 
 #if USE_USD_SDK
 		FScopedUsdAllocs UsdAllocs;
-		std::vector< double > UsdTimes;
+		std::vector<double> UsdTimes;
 
-		if ( Impl->PxrUsdGeomXformable.Get().GetTimeSamples( &UsdTimes ) )
+		if (Impl->PxrUsdGeomXformable.Get().GetTimeSamples(&UsdTimes))
 		{
-			Times->Empty( UsdTimes.size() );
+			Times->Empty(UsdTimes.size());
 
-			for ( double UsdTime : UsdTimes )
+			for (double UsdTime : UsdTimes)
 			{
-				Times->Add( UsdTime );
+				Times->Add(UsdTime);
 			}
 
 			return true;
 		}
-#endif // #if USE_USD_SDK
+#endif	  // #if USE_USD_SDK
 
 		return false;
 	}
@@ -190,10 +188,10 @@ namespace UE
 	FUsdAttribute FUsdGeomXformable::GetXformOpOrderAttr() const
 	{
 #if USE_USD_SDK
-		return FUsdAttribute( Impl->PxrUsdGeomXformable.Get().GetXformOpOrderAttr() );
+		return FUsdAttribute(Impl->PxrUsdGeomXformable.Get().GetXformOpOrderAttr());
 #else
 		return FUsdAttribute();
-#endif // #if USE_USD_SDK
+#endif	  // #if USE_USD_SDK
 	}
 
 	bool FUsdGeomXformable::ClearXformOpOrder() const
@@ -202,6 +200,6 @@ namespace UE
 		return Impl->PxrUsdGeomXformable.Get().ClearXformOpOrder();
 #else
 		return false;
-#endif // #if USE_USD_SDK
+#endif	  // #if USE_USD_SDK
 	}
-}
+}	 // namespace UE
diff --git a/Engine/Plugins/Importers/USDImporter/Source/UnrealUSDWrapper/Private/UsdWrappers/UsdPayloads.cpp b/Engine/Plugins/Importers/USDImporter/Source/UnrealUSDWrapper/Private/UsdWrappers/UsdPayloads.cpp
--- a/Engine/Plugins/Importers/USDImporter/Source/UnrealUSDWrapper/Private/UsdWrappers/UsdPayloads.cpp
+++ b/Engine/Plugins/Importers/USDImporter/Source/UnrealUSDWrapper/Private/UsdWrappers/UsdPayloads.cpp
@@ -7,13 +7,11 @@
 #include "UsdWrappers/UsdPrim.h"
 
 #if USE_USD_SDK
-
 #include "USDIncludesStart.h"
-	#include "pxr/usd/usd/prim.h"
-	#include "pxr/usd/usd/payloads.h"
+#include "pxr/usd/usd/payloads.h"
+#include "pxr/usd/usd/prim.h"
 #include "USDIncludesEnd.h"
-
-#endif // #if USE_USD_SDK
+#endif	  // #if USE_USD_SDK
 
 namespace UE
 {
@@ -48,9 +46,9 @@ namespace UE
 
 			// pxr::UsdPayloads has no default constructor, so we store the prim itself.
 			TUsdStore<pxr::UsdPrim> PxrUsdPrim;
-#endif // #if USE_USD_SDK
+#endif	  // #if USE_USD_SDK
 		};
-	}
+	}	  // namespace Internal
 
 	FUsdPayloads::FUsdPayloads()
 	{
@@ -61,7 +59,7 @@ namespace UE
 	{
 #if USE_USD_SDK
 		Impl = MakeUnique<Internal::FUsdPayloadsImpl>(Other.Impl->PxrUsdPrim.Get());
-#endif // #if USE_USD_SDK
+#endif	  // #if USE_USD_SDK
 	}
 
 	FUsdPayloads::FUsdPayloads(FUsdPayloads&& Other) = default;
@@ -70,7 +68,7 @@ namespace UE
 	{
 #if USE_USD_SDK
 		Impl = MakeUnique<Internal::FUsdPayloadsImpl>(Other.Impl->PxrUsdPrim.Get());
-#endif // #if USE_USD_SDK
+#endif	  // #if USE_USD_SDK
 
 		return *this;
 	}
@@ -93,7 +91,7 @@ namespace UE
 		return (bool)Impl->PxrUsdPrim.Get();
 #else
 		return false;
-#endif // #if USE_USD_SDK
+#endif	  // #if USE_USD_SDK
 	}
 
 #if USE_USD_SDK
@@ -134,7 +132,7 @@ namespace UE
 	{
 		Impl = MakeUnique<Internal::FUsdPayloadsImpl>(InPrim);
 	}
-#endif // #if USE_USD_SDK
+#endif	  // #if USE_USD_SDK
 
 	bool FUsdPayloads::AddPayload(const FSdfPayload& Payload, EUsdListPosition Position)
 	{
@@ -149,20 +147,14 @@ namespace UE
 
 			const pxr::SdfPayload UsdPayload(AssetPath, UsdPrimPath, UsdLayerOffset);
 
-			return Impl->PxrUsdPrim.Get().GetPayloads().AddPayload(
-				UsdPayload,
-				static_cast<pxr::UsdListPosition>(Position));
+			return Impl->PxrUsdPrim.Get().GetPayloads().AddPayload(UsdPayload, static_cast<pxr::UsdListPosition>(Position));
 		}
-#endif // #if USE_USD_SDK
+#endif	  // #if USE_USD_SDK
 
 		return false;
 	}
 
-	bool FUsdPayloads::AddPayload(
-		const FString& Identifier,
-		const FSdfPath& PrimPath,
-		const FSdfLayerOffset& LayerOffset,
-		EUsdListPosition Position)
+	bool FUsdPayloads::AddPayload(const FString& Identifier, const FSdfPath& PrimPath, const FSdfLayerOffset& LayerOffset, EUsdListPosition Position)
 	{
 #if USE_USD_SDK
 		if (Impl->PxrUsdPrim.Get())
@@ -174,17 +166,18 @@ namespace UE
 			const pxr::SdfLayerOffset UsdLayerOffset(LayerOffset.Offset, LayerOffset.Scale);
 
 			return Impl->PxrUsdPrim.Get().GetPayloads().AddPayload(
-				UsdIdentifier, UsdPrimPath, UsdLayerOffset, static_cast<pxr::UsdListPosition>(Position));
+				UsdIdentifier,
+				UsdPrimPath,
+				UsdLayerOffset,
+				static_cast<pxr::UsdListPosition>(Position)
+			);
 		}
-#endif // #if USE_USD_SDK
+#endif	  // #if USE_USD_SDK
 
 		return false;
 	}
 
-	bool FUsdPayloads::AddPayload(
-		const FString& Identifier,
-		const FSdfLayerOffset& LayerOffset,
-		EUsdListPosition Position)
+	bool FUsdPayloads::AddPayload(const FString& Identifier, const FSdfLayerOffset& LayerOffset, EUsdListPosition Position)
 	{
 #if USE_USD_SDK
 		if (Impl->PxrUsdPrim.Get())
@@ -194,18 +187,14 @@ namespace UE
 			const std::string UsdIdentifier(TCHAR_TO_ANSI(*Identifier));
 			const pxr::SdfLayerOffset UsdLayerOffset(LayerOffset.Offset, LayerOffset.Scale);
 
-			return Impl->PxrUsdPrim.Get().GetPayloads().AddPayload(
-				UsdIdentifier, UsdLayerOffset, static_cast<pxr::UsdListPosition>(Position));
+			return Impl->PxrUsdPrim.Get().GetPayloads().AddPayload(UsdIdentifier, UsdLayerOffset, static_cast<pxr::UsdListPosition>(Position));
 		}
-#endif // #if USE_USD_SDK
+#endif	  // #if USE_USD_SDK
 
 		return false;
 	}
 
-	bool FUsdPayloads::AddInternalPayload(
-		const FSdfPath& PrimPath,
-		const FSdfLayerOffset& LayerOffset,
-		EUsdListPosition Position)
+	bool FUsdPayloads::AddInternalPayload(const FSdfPath& PrimPath, const FSdfLayerOffset& LayerOffset, EUsdListPosition Position)
 	{
 #if USE_USD_SDK
 		if (Impl->PxrUsdPrim.Get())
@@ -215,10 +204,9 @@ namespace UE
 			const pxr::SdfPath UsdPrimPath(PrimPath);
 			const pxr::SdfLayerOffset UsdLayerOffset(LayerOffset.Offset, LayerOffset.Scale);
 
-			return Impl->PxrUsdPrim.Get().GetPayloads().AddInternalPayload(
-				UsdPrimPath, UsdLayerOffset, static_cast<pxr::UsdListPosition>(Position));
+			return Impl->PxrUsdPrim.Get().GetPayloads().AddInternalPayload(UsdPrimPath, UsdLayerOffset, static_cast<pxr::UsdListPosition>(Position));
 		}
-#endif // #if USE_USD_SDK
+#endif	  // #if USE_USD_SDK
 
 		return false;
 	}
@@ -238,7 +226,7 @@ namespace UE
 
 			return Impl->PxrUsdPrim.Get().GetPayloads().RemovePayload(UsdPayload);
 		}
-#endif // #if USE_USD_SDK
+#endif	  // #if USE_USD_SDK
 
 		return false;
 	}
@@ -252,7 +240,7 @@ namespace UE
 
 			return Impl->PxrUsdPrim.Get().GetPayloads().ClearPayloads();
 		}
-#endif // #if USE_USD_SDK
+#endif	  // #if USE_USD_SDK
 
 		return false;
 	}
@@ -278,7 +266,7 @@ namespace UE
 
 			return Impl->PxrUsdPrim.Get().GetPayloads().SetPayloads(UsdPayloads);
 		}
-#endif // #if USE_USD_SDK
+#endif	  // #if USE_USD_SDK
 
 		return false;
 	}
@@ -290,8 +278,8 @@ namespace UE
 		{
 			return FUsdPrim(Impl->PxrUsdPrim.Get());
 		}
-#endif // #if USE_USD_SDK
+#endif	  // #if USE_USD_SDK
 
 		return FUsdPrim{};
 	}
-}
+}	 // namespace UE
diff --git a/Engine/Plugins/Importers/USDImporter/Source/UnrealUSDWrapper/Private/UsdWrappers/UsdPrim.cpp b/Engine/Plugins/Importers/USDImporter/Source/UnrealUSDWrapper/Private/UsdWrappers/UsdPrim.cpp
--- a/Engine/Plugins/Importers/USDImporter/Source/UnrealUSDWrapper/Private/UsdWrappers/UsdPrim.cpp
+++ b/Engine/Plugins/Importers/USDImporter/Source/UnrealUSDWrapper/Private/UsdWrappers/UsdPrim.cpp
@@ -12,15 +12,12 @@
 #include "UsdWrappers/UsdVariantSets.h"
 
 #if USE_USD_SDK
-
 #include "USDIncludesStart.h"
-	#include "pxr/usd/sdf/schema.h"
-	#include "pxr/usd/usd/attribute.h"
-	#include "pxr/usd/usd/prim.h"
-	#include "pxr/usd/usd/schemaBase.h"
+#include "pxr/usd/sdf/schema.h"
+#include "pxr/usd/usd/attribute.h"
+#include "pxr/usd/usd/prim.h"
 #include "USDIncludesEnd.h"
-
-#endif // #if USE_USD_SDK
+#endif	  // #if USE_USD_SDK
 
 namespace UE
 {
@@ -36,42 +33,42 @@ namespace UE
 #if ENABLE_USD_DEBUG_PATH
 			FSdfPath DebugPath;
 #endif
-			explicit FUsdPrimImpl( const pxr::UsdPrim& InUsdPrim )
-				: PxrUsdPrim( InUsdPrim )
+			explicit FUsdPrimImpl(const pxr::UsdPrim& InUsdPrim)
+				: PxrUsdPrim(InUsdPrim)
 			{
 #if ENABLE_USD_DEBUG_PATH
 				DebugPath = PxrUsdPrim.Get().GetPrimPath();
 #endif
 			}
 
-			explicit FUsdPrimImpl( pxr::UsdPrim&& InUsdPrim )
-				: PxrUsdPrim( MoveTemp( InUsdPrim ) )
+			explicit FUsdPrimImpl(pxr::UsdPrim&& InUsdPrim)
+				: PxrUsdPrim(MoveTemp(InUsdPrim))
 			{
 #if ENABLE_USD_DEBUG_PATH
 				DebugPath = PxrUsdPrim.Get().GetPrimPath();
 #endif
 			}
 
-			TUsdStore< pxr::UsdPrim > PxrUsdPrim;
-#endif // #if USE_USD_SDK
+			TUsdStore<pxr::UsdPrim> PxrUsdPrim;
+#endif	  // #if USE_USD_SDK
 		};
-	}
+	}	  // namespace Internal
 
 	FUsdPrim::FUsdPrim()
 	{
 		FScopedUnrealAllocs UnrealAllocs;
-		Impl = MakeUnique< Internal::FUsdPrimImpl >();
+		Impl = MakeUnique<Internal::FUsdPrimImpl>();
 	}
 
-	FUsdPrim::FUsdPrim( const FUsdPrim& Other )
+	FUsdPrim::FUsdPrim(const FUsdPrim& Other)
 	{
 #if USE_USD_SDK
 		FScopedUnrealAllocs UnrealAllocs;
-		Impl = MakeUnique< Internal::FUsdPrimImpl >( Other.Impl->PxrUsdPrim.Get() );
-#endif // #if USE_USD_SDK
+		Impl = MakeUnique<Internal::FUsdPrimImpl>(Other.Impl->PxrUsdPrim.Get());
+#endif	  // #if USE_USD_SDK
 	}
 
-	FUsdPrim::FUsdPrim( FUsdPrim&& Other ) = default;
+	FUsdPrim::FUsdPrim(FUsdPrim&& Other) = default;
 
 	FUsdPrim::~FUsdPrim()
 	{
@@ -79,19 +76,19 @@ namespace UE
 		Impl.Reset();
 	}
 
-	FUsdPrim& FUsdPrim::operator=( const FUsdPrim& Other )
+	FUsdPrim& FUsdPrim::operator=(const FUsdPrim& Other)
 	{
 #if USE_USD_SDK
 		FScopedUnrealAllocs UnrealAllocs;
-		Impl = MakeUnique< Internal::FUsdPrimImpl >(  Other.Impl->PxrUsdPrim.Get() );
-#endif // #if USE_USD_SDK
+		Impl = MakeUnique<Internal::FUsdPrimImpl>(Other.Impl->PxrUsdPrim.Get());
+#endif	  // #if USE_USD_SDK
 		return *this;
 	}
 
-	FUsdPrim& FUsdPrim::operator=( FUsdPrim&& Other )
+	FUsdPrim& FUsdPrim::operator=(FUsdPrim&& Other)
 	{
 		FScopedUnrealAllocs UnrealAllocs;
-		Impl = MoveTemp(  Other.Impl );
+		Impl = MoveTemp(Other.Impl);
 
 		return *this;
 	}
@@ -102,47 +99,47 @@ namespace UE
 		return (bool)Impl->PxrUsdPrim.Get();
 #else
 		return false;
-#endif // #if USE_USD_SDK
+#endif	  // #if USE_USD_SDK
 	}
 
-	bool FUsdPrim::operator==( const FUsdPrim& Other ) const
+	bool FUsdPrim::operator==(const FUsdPrim& Other) const
 	{
 #if USE_USD_SDK
 		return Impl->PxrUsdPrim.Get() == Other.Impl->PxrUsdPrim.Get();
 #else
 		return false;
-#endif // #if USE_USD_SDK
+#endif	  // #if USE_USD_SDK
 	}
 
-	bool FUsdPrim::operator!=( const FUsdPrim& Other ) const
+	bool FUsdPrim::operator!=(const FUsdPrim& Other) const
 	{
-		return !( *this == Other );
+		return !(*this == Other);
 	}
 
 #if USE_USD_SDK
-	FUsdPrim::FUsdPrim( const pxr::UsdPrim& InUsdPrim )
+	FUsdPrim::FUsdPrim(const pxr::UsdPrim& InUsdPrim)
 	{
 		FScopedUnrealAllocs UnrealAllocs;
-		Impl =  MakeUnique< Internal::FUsdPrimImpl >( InUsdPrim );
+		Impl = MakeUnique<Internal::FUsdPrimImpl>(InUsdPrim);
 	}
 
-	FUsdPrim::FUsdPrim( pxr::UsdPrim&& InUsdPrim )
+	FUsdPrim::FUsdPrim(pxr::UsdPrim&& InUsdPrim)
 	{
 		FScopedUnrealAllocs UnrealAllocs;
-		Impl = MakeUnique< Internal::FUsdPrimImpl >( MoveTemp( InUsdPrim ) );
+		Impl = MakeUnique<Internal::FUsdPrimImpl>(MoveTemp(InUsdPrim));
 	}
 
-	FUsdPrim& FUsdPrim::operator=( const pxr::UsdPrim& InUsdPrim )
+	FUsdPrim& FUsdPrim::operator=(const pxr::UsdPrim& InUsdPrim)
 	{
 		FScopedUnrealAllocs UnrealAllocs;
-		Impl = MakeUnique< Internal::FUsdPrimImpl >( InUsdPrim );
+		Impl = MakeUnique<Internal::FUsdPrimImpl>(InUsdPrim);
 		return *this;
 	}
 
-	FUsdPrim& FUsdPrim::operator=( pxr::UsdPrim&& InUsdPrim )
+	FUsdPrim& FUsdPrim::operator=(pxr::UsdPrim&& InUsdPrim)
 	{
 		FScopedUnrealAllocs UnrealAllocs;
-		Impl = MakeUnique< Internal::FUsdPrimImpl >( MoveTemp( InUsdPrim ) );
+		Impl = MakeUnique<Internal::FUsdPrimImpl>(MoveTemp(InUsdPrim));
 		return *this;
 	}
 
@@ -155,20 +152,20 @@ namespace UE
 	{
 		return Impl->PxrUsdPrim.Get();
 	}
-#endif // #if USE_USD_SDK
+#endif	  // #if USE_USD_SDK
 
-	bool FUsdPrim::SetSpecifier( ESdfSpecifier Specifier )
+	bool FUsdPrim::SetSpecifier(ESdfSpecifier Specifier)
 	{
 #if USE_USD_SDK
-		static_assert( ( int32 ) ESdfSpecifier::Def == ( int32 ) pxr::SdfSpecifierDef, "ESdfSpecifier enum doesn't match USD!" );
-		static_assert( ( int32 ) ESdfSpecifier::Over == ( int32 ) pxr::SdfSpecifierOver, "ESdfSpecifier enum doesn't match USD!" );
-		static_assert( ( int32 ) ESdfSpecifier::Class == ( int32 ) pxr::SdfSpecifierClass, "ESdfSpecifier enum doesn't match USD!" );
-		static_assert( ( int32 ) ESdfSpecifier::Num == ( int32 ) pxr::SdfNumSpecifiers, "ESdfSpecifier enum doesn't match USD!" );
+		static_assert((int32)ESdfSpecifier::Def == (int32)pxr::SdfSpecifierDef, "ESdfSpecifier enum doesn't match USD!");
+		static_assert((int32)ESdfSpecifier::Over == (int32)pxr::SdfSpecifierOver, "ESdfSpecifier enum doesn't match USD!");
+		static_assert((int32)ESdfSpecifier::Class == (int32)pxr::SdfSpecifierClass, "ESdfSpecifier enum doesn't match USD!");
+		static_assert((int32)ESdfSpecifier::Num == (int32)pxr::SdfNumSpecifiers, "ESdfSpecifier enum doesn't match USD!");
 
-		return Impl->PxrUsdPrim.Get().SetSpecifier( static_cast< pxr::SdfSpecifier > ( Specifier ) );
+		return Impl->PxrUsdPrim.Get().SetSpecifier(static_cast<pxr::SdfSpecifier>(Specifier));
 #else
 		return false;
-#endif // #if USE_USD_SDK
+#endif	  // #if USE_USD_SDK
 	}
 
 	bool FUsdPrim::IsActive() const
@@ -177,16 +174,16 @@ namespace UE
 		return Impl->PxrUsdPrim.Get().IsActive();
 #else
 		return false;
-#endif // #if USE_USD_SDK
+#endif	  // #if USE_USD_SDK
 	}
 
-	bool FUsdPrim::SetActive( bool bActive )
+	bool FUsdPrim::SetActive(bool bActive)
 	{
 #if USE_USD_SDK
-		return Impl->PxrUsdPrim.Get().SetActive( bActive );
+		return Impl->PxrUsdPrim.Get().SetActive(bActive);
 #else
 		return false;
-#endif // #if USE_USD_SDK
+#endif	  // #if USE_USD_SDK
 	}
 
 	bool FUsdPrim::IsValid() const
@@ -195,7 +192,7 @@ namespace UE
 		return Impl->PxrUsdPrim.Get().IsValid();
 #else
 		return false;
-#endif // #if USE_USD_SDK
+#endif	  // #if USE_USD_SDK
 	}
 
 	bool FUsdPrim::IsPseudoRoot() const
@@ -204,7 +201,7 @@ namespace UE
 		return Impl->PxrUsdPrim.Get().IsPseudoRoot();
 #else
 		return false;
-#endif // #if USE_USD_SDK
+#endif	  // #if USE_USD_SDK
 	}
 
 	bool FUsdPrim::IsModel() const
@@ -213,7 +210,7 @@ namespace UE
 		return Impl->PxrUsdPrim.Get().IsModel();
 #else
 		return false;
-#endif // #if USE_USD_SDK
+#endif	  // #if USE_USD_SDK
 	}
 
 	bool FUsdPrim::IsGroup() const
@@ -222,7 +219,7 @@ namespace UE
 		return Impl->PxrUsdPrim.Get().IsGroup();
 #else
 		return false;
-#endif // #if USE_USD_SDK
+#endif	  // #if USE_USD_SDK
 	}
 
 	TArray<FName> FUsdPrim::GetAppliedSchemas() const
@@ -235,202 +232,197 @@ namespace UE
 		std::vector<pxr::TfToken> UsdAppliedSchemas = Impl->PxrUsdPrim.Get().GetAppliedSchemas();
 		AppliedSchemas.Reserve(UsdAppliedSchemas.size());
 
-		for ( const pxr::TfToken& UsdSchema : UsdAppliedSchemas )
+		for (const pxr::TfToken& UsdSchema : UsdAppliedSchemas)
 		{
-			AppliedSchemas.Add( ANSI_TO_TCHAR( UsdSchema.GetString().c_str() ) );
+			AppliedSchemas.Add(ANSI_TO_TCHAR(UsdSchema.GetString().c_str()));
 		}
-#endif // #if USE_USD_SDK
+#endif	  // #if USE_USD_SDK
 
 		return AppliedSchemas;
 	}
 
-	bool FUsdPrim::IsA( FName SchemaIdentifier ) const
+	bool FUsdPrim::IsA(FName SchemaIdentifier) const
 	{
 #if USE_USD_SDK
 		FScopedUsdAllocs Allocs;
 
-		const pxr::TfToken UsdSchemaIdentifier( TCHAR_TO_ANSI( *SchemaIdentifier.ToString() ) );
+		const pxr::TfToken UsdSchemaIdentifier(TCHAR_TO_ANSI(*SchemaIdentifier.ToString()));
 
-		return Impl->PxrUsdPrim.Get().IsA( UsdSchemaIdentifier );
+		return Impl->PxrUsdPrim.Get().IsA(UsdSchemaIdentifier);
 #else
 		return false;
-#endif // #if USE_USD_SDK
+#endif	  // #if USE_USD_SDK
 	}
 
-	bool FUsdPrim::HasAPI( FName SchemaIdentifier ) const
+	bool FUsdPrim::HasAPI(FName SchemaIdentifier) const
 	{
 #if USE_USD_SDK
 		FScopedUsdAllocs Allocs;
 
-		const pxr::TfToken UsdSchemaIdentifier( TCHAR_TO_ANSI( *SchemaIdentifier.ToString() ) );
+		const pxr::TfToken UsdSchemaIdentifier(TCHAR_TO_ANSI(*SchemaIdentifier.ToString()));
 
-		return Impl->PxrUsdPrim.Get().HasAPI( UsdSchemaIdentifier );
+		return Impl->PxrUsdPrim.Get().HasAPI(UsdSchemaIdentifier);
 #else
 		return false;
-#endif // #if USE_USD_SDK
+#endif	  // #if USE_USD_SDK
 	}
 
-	bool FUsdPrim::HasAPI( FName SchemaIdentifier, FName InstanceName ) const
+	bool FUsdPrim::HasAPI(FName SchemaIdentifier, FName InstanceName) const
 	{
 #if USE_USD_SDK
 		FScopedUsdAllocs Allocs;
 
-		const pxr::TfToken UsdSchemaIdentifier( TCHAR_TO_ANSI( *SchemaIdentifier.ToString() ) );
-		const pxr::TfToken UsdInstanceName( TCHAR_TO_ANSI( *InstanceName.ToString() ) );
+		const pxr::TfToken UsdSchemaIdentifier(TCHAR_TO_ANSI(*SchemaIdentifier.ToString()));
+		const pxr::TfToken UsdInstanceName(TCHAR_TO_ANSI(*InstanceName.ToString()));
 
-		return Impl->PxrUsdPrim.Get().HasAPI( UsdSchemaIdentifier, UsdInstanceName );
+		return Impl->PxrUsdPrim.Get().HasAPI(UsdSchemaIdentifier, UsdInstanceName);
 #else
 		return false;
-#endif // #if USE_USD_SDK
+#endif	  // #if USE_USD_SDK
 	}
 
 	// Deprecated in 5.3
-	bool FUsdPrim::HasAPI( FName SchemaType, TOptional<FName> InstanceName ) const
+	bool FUsdPrim::HasAPI(FName SchemaType, TOptional<FName> InstanceName) const
 	{
-		if ( InstanceName.IsSet() )
+		if (InstanceName.IsSet())
 		{
-			return HasAPI( SchemaType, InstanceName.GetValue() );
+			return HasAPI(SchemaType, InstanceName.GetValue());
 		}
 
-		return HasAPI( SchemaType );
+		return HasAPI(SchemaType);
 	}
 
-	bool FUsdPrim::CanApplyAPI( FName SchemaIdentifier, FString* OutWhyNot ) const
+	bool FUsdPrim::CanApplyAPI(FName SchemaIdentifier, FString* OutWhyNot) const
 	{
 		bool bResult = false;
 
 #if USE_USD_SDK
 		FScopedUsdAllocs Allocs;
 
-		const pxr::TfToken UsdSchemaIdentifier( TCHAR_TO_ANSI( *SchemaIdentifier.ToString() ) );
+		const pxr::TfToken UsdSchemaIdentifier(TCHAR_TO_ANSI(*SchemaIdentifier.ToString()));
 
 		std::string WhyNot;
-		bResult = Impl->PxrUsdPrim.Get().CanApplyAPI(
-			UsdSchemaIdentifier,
-			OutWhyNot != nullptr ? &WhyNot : nullptr );
+		bResult = Impl->PxrUsdPrim.Get().CanApplyAPI(UsdSchemaIdentifier, OutWhyNot != nullptr ? &WhyNot : nullptr);
 
-		if ( !bResult && OutWhyNot != nullptr )
+		if (!bResult && OutWhyNot != nullptr)
 		{
-			*OutWhyNot = ANSI_TO_TCHAR( WhyNot.c_str() );
+			*OutWhyNot = ANSI_TO_TCHAR(WhyNot.c_str());
 		}
-#endif // #if USE_USD_SDK
+#endif	  // #if USE_USD_SDK
 
 		return bResult;
 	}
 
-	bool FUsdPrim::CanApplyAPI( FName SchemaIdentifier, FName InstanceName, FString* OutWhyNot ) const
+	bool FUsdPrim::CanApplyAPI(FName SchemaIdentifier, FName InstanceName, FString* OutWhyNot) const
 	{
 		bool bResult = false;
 
 #if USE_USD_SDK
 		FScopedUsdAllocs Allocs;
 
-		const pxr::TfToken UsdSchemaIdentifier( TCHAR_TO_ANSI( *SchemaIdentifier.ToString() ) );
-		const pxr::TfToken UsdInstanceName( TCHAR_TO_ANSI( *InstanceName.ToString() ) );
+		const pxr::TfToken UsdSchemaIdentifier(TCHAR_TO_ANSI(*SchemaIdentifier.ToString()));
+		const pxr::TfToken UsdInstanceName(TCHAR_TO_ANSI(*InstanceName.ToString()));
 
 		std::string WhyNot;
-		bResult = Impl->PxrUsdPrim.Get().CanApplyAPI(
-			UsdSchemaIdentifier,
-			UsdInstanceName,
-			OutWhyNot != nullptr ? &WhyNot : nullptr);
+		bResult = Impl->PxrUsdPrim.Get().CanApplyAPI(UsdSchemaIdentifier, UsdInstanceName, OutWhyNot != nullptr ? &WhyNot : nullptr);
 
-		if ( !bResult && OutWhyNot != nullptr )
+		if (!bResult && OutWhyNot != nullptr)
 		{
-			*OutWhyNot = ANSI_TO_TCHAR( WhyNot.c_str() );
+			*OutWhyNot = ANSI_TO_TCHAR(WhyNot.c_str());
 		}
-#endif // #if USE_USD_SDK
+#endif	  // #if USE_USD_SDK
 
 		return bResult;
 	}
 
-	bool FUsdPrim::ApplyAPI( FName SchemaIdentifier ) const
+	bool FUsdPrim::ApplyAPI(FName SchemaIdentifier) const
 	{
 #if USE_USD_SDK
 		FScopedUsdAllocs Allocs;
 
-		const pxr::TfToken UsdSchemaIdentifier( TCHAR_TO_ANSI( *SchemaIdentifier.ToString() ) );
+		const pxr::TfToken UsdSchemaIdentifier(TCHAR_TO_ANSI(*SchemaIdentifier.ToString()));
 
-		return Impl->PxrUsdPrim.Get().ApplyAPI( UsdSchemaIdentifier );
+		return Impl->PxrUsdPrim.Get().ApplyAPI(UsdSchemaIdentifier);
 #else
 		return false;
-#endif // #if USE_USD_SDK
+#endif	  // #if USE_USD_SDK
 	}
 
-	bool FUsdPrim::ApplyAPI( FName SchemaIdentifier, FName InstanceName ) const
+	bool FUsdPrim::ApplyAPI(FName SchemaIdentifier, FName InstanceName) const
 	{
 #if USE_USD_SDK
 		FScopedUsdAllocs Allocs;
 
-		const pxr::TfToken UsdSchemaIdentifier( TCHAR_TO_ANSI( *SchemaIdentifier.ToString() ) );
-		const pxr::TfToken UsdInstanceName( TCHAR_TO_ANSI( *InstanceName.ToString() ) );
+		const pxr::TfToken UsdSchemaIdentifier(TCHAR_TO_ANSI(*SchemaIdentifier.ToString()));
+		const pxr::TfToken UsdInstanceName(TCHAR_TO_ANSI(*InstanceName.ToString()));
 
-		return Impl->PxrUsdPrim.Get().ApplyAPI( UsdSchemaIdentifier, UsdInstanceName );
+		return Impl->PxrUsdPrim.Get().ApplyAPI(UsdSchemaIdentifier, UsdInstanceName);
 #else
 		return false;
-#endif // #if USE_USD_SDK
+#endif	  // #if USE_USD_SDK
 	}
 
-	bool FUsdPrim::RemoveAPI( FName SchemaIdentifier ) const
+	bool FUsdPrim::RemoveAPI(FName SchemaIdentifier) const
 	{
 #if USE_USD_SDK
 		FScopedUsdAllocs Allocs;
 
-		const pxr::TfToken UsdSchemaIdentifier( TCHAR_TO_ANSI( *SchemaIdentifier.ToString() ) );
+		const pxr::TfToken UsdSchemaIdentifier(TCHAR_TO_ANSI(*SchemaIdentifier.ToString()));
 
-		return Impl->PxrUsdPrim.Get().RemoveAPI( UsdSchemaIdentifier );
+		return Impl->PxrUsdPrim.Get().RemoveAPI(UsdSchemaIdentifier);
 #else
 		return false;
-#endif // #if USE_USD_SDK
+#endif	  // #if USE_USD_SDK
 	}
 
-	bool FUsdPrim::RemoveAPI( FName SchemaIdentifier, FName InstanceName ) const
+	bool FUsdPrim::RemoveAPI(FName SchemaIdentifier, FName InstanceName) const
 	{
 #if USE_USD_SDK
 		FScopedUsdAllocs Allocs;
 
-		const pxr::TfToken UsdSchemaIdentifier( TCHAR_TO_ANSI( *SchemaIdentifier.ToString() ) );
-		const pxr::TfToken UsdInstanceName( TCHAR_TO_ANSI( *InstanceName.ToString() ) );
+		const pxr::TfToken UsdSchemaIdentifier(TCHAR_TO_ANSI(*SchemaIdentifier.ToString()));
+		const pxr::TfToken UsdInstanceName(TCHAR_TO_ANSI(*InstanceName.ToString()));
 
-		return Impl->PxrUsdPrim.Get().RemoveAPI( UsdSchemaIdentifier, UsdInstanceName );
+		return Impl->PxrUsdPrim.Get().RemoveAPI(UsdSchemaIdentifier, UsdInstanceName);
 #else
 		return false;
-#endif // #if USE_USD_SDK
+#endif	  // #if USE_USD_SDK
 	}
 
 	const FSdfPath FUsdPrim::GetPrimPath() const
 	{
 #if USE_USD_SDK
-		return FSdfPath( Impl->PxrUsdPrim.Get().GetPrimPath() );
+		return FSdfPath(Impl->PxrUsdPrim.Get().GetPrimPath());
 #else
 		return FSdfPath();
-#endif // #if USE_USD_SDK
+#endif	  // #if USE_USD_SDK
 	}
 
 	FUsdStage FUsdPrim::GetStage() const
 	{
 #if USE_USD_SDK
-		return FUsdStage( Impl->PxrUsdPrim.Get().GetStage() );
+		return FUsdStage(Impl->PxrUsdPrim.Get().GetStage());
 #else
 		return FUsdStage();
-#endif // #if USE_USD_SDK
+#endif	  // #if USE_USD_SDK
 	}
 
 	FName FUsdPrim::GetName() const
 	{
 #if USE_USD_SDK
-		return FName( ANSI_TO_TCHAR( Impl->PxrUsdPrim.Get().GetName().GetString().c_str() ) );
+		return FName(ANSI_TO_TCHAR(Impl->PxrUsdPrim.Get().GetName().GetString().c_str()));
 #else
 		return FName();
-#endif // #if USE_USD_SDK
+#endif	  // #if USE_USD_SDK
 	}
 
 	FName FUsdPrim::GetTypeName() const
 	{
 #if USE_USD_SDK
-		return FName( ANSI_TO_TCHAR( Impl->PxrUsdPrim.Get().GetTypeName().GetString().c_str() ) );
+		return FName(ANSI_TO_TCHAR(Impl->PxrUsdPrim.Get().GetTypeName().GetString().c_str()));
 #else
 		return FName();
-#endif // #if USE_USD_SDK
+#endif	  // #if USE_USD_SDK
 	}
 
 	bool FUsdPrim::SetTypeName(FName TypeName) const
@@ -440,7 +432,7 @@ namespace UE
 		return Impl->PxrUsdPrim.Get().SetTypeName(pxr::TfToken(TCHAR_TO_ANSI(*TypeName.ToString())));
 #else
 		return false;
-#endif // #if USE_USD_SDK
+#endif	  // #if USE_USD_SDK
 	}
 
 	bool FUsdPrim::ClearTypeName() const
@@ -449,7 +441,7 @@ namespace UE
 		return Impl->PxrUsdPrim.Get().ClearTypeName();
 #else
 		return false;
-#endif // #if USE_USD_SDK
+#endif	  // #if USE_USD_SDK
 	}
 
 	bool FUsdPrim::HasAuthoredTypeName() const
@@ -458,57 +450,57 @@ namespace UE
 		return Impl->PxrUsdPrim.Get().HasAuthoredTypeName();
 #else
 		return false;
-#endif // #if USE_USD_SDK
+#endif	  // #if USE_USD_SDK
 	}
 
 	FUsdPrim FUsdPrim::GetParent() const
 	{
 #if USE_USD_SDK
-		return FUsdPrim( Impl->PxrUsdPrim.Get().GetParent() );
+		return FUsdPrim(Impl->PxrUsdPrim.Get().GetParent());
 #else
 		return FUsdPrim();
-#endif // #if USE_USD_SDK
+#endif	  // #if USE_USD_SDK
 	}
 
-	TArray< FUsdPrim > FUsdPrim::GetChildren() const
+	TArray<FUsdPrim> FUsdPrim::GetChildren() const
 	{
-		TArray< FUsdPrim > Children;
+		TArray<FUsdPrim> Children;
 
 #if USE_USD_SDK
 		FScopedUsdAllocs UsdAllocs;
 
 		pxr::UsdPrimSiblingRange PrimChildren = Impl->PxrUsdPrim.Get().GetChildren();
 
-		for ( const pxr::UsdPrim& Child : PrimChildren )
+		for (const pxr::UsdPrim& Child : PrimChildren)
 		{
-			Children.Emplace( Child );
+			Children.Emplace(Child);
 		}
-#endif // #if USE_USD_SDK
+#endif	  // #if USE_USD_SDK
 
 		return Children;
 	}
 
-	TArray< FUsdPrim > FUsdPrim::GetFilteredChildren( bool bTraverseInstanceProxies ) const
+	TArray<FUsdPrim> FUsdPrim::GetFilteredChildren(bool bTraverseInstanceProxies) const
 	{
-		TArray< FUsdPrim > Children;
+		TArray<FUsdPrim> Children;
 
 #if USE_USD_SDK
 		FScopedUsdAllocs UsdAllocs;
 
 		pxr::Usd_PrimFlagsPredicate Predicate = pxr::UsdPrimDefaultPredicate;
 
-		if ( bTraverseInstanceProxies )
+		if (bTraverseInstanceProxies)
 		{
-			Predicate = pxr::UsdTraverseInstanceProxies( Predicate ) ;
+			Predicate = pxr::UsdTraverseInstanceProxies(Predicate);
 		}
 
-		pxr::UsdPrimSiblingRange PrimChildren = Impl->PxrUsdPrim.Get().GetFilteredChildren( Predicate );
+		pxr::UsdPrimSiblingRange PrimChildren = Impl->PxrUsdPrim.Get().GetFilteredChildren(Predicate);
 
-		for ( const pxr::UsdPrim& Child : PrimChildren )
+		for (const pxr::UsdPrim& Child : PrimChildren)
 		{
-			Children.Emplace( Child );
+			Children.Emplace(Child);
 		}
-#endif // #if USE_USD_SDK
+#endif	  // #if USE_USD_SDK
 
 		return Children;
 	}
@@ -519,7 +511,7 @@ namespace UE
 		return FUsdVariantSets(Impl->PxrUsdPrim.Get());
 #else
 		return FUsdVariantSets{};
-#endif // #if USE_USD_SDK
+#endif	  // #if USE_USD_SDK
 	}
 
 	FUsdVariantSet FUsdPrim::GetVariantSet(const FString& VariantSetName) const
@@ -532,7 +524,7 @@ namespace UE
 		return FUsdVariantSet(Impl->PxrUsdPrim.Get(), UsdVariantSetName);
 #else
 		return FUsdVariantSet{};
-#endif // #if USE_USD_SDK
+#endif	  // #if USE_USD_SDK
 	}
 
 	bool FUsdPrim::HasVariantSets() const
@@ -541,7 +533,7 @@ namespace UE
 		return Impl->PxrUsdPrim.Get().HasVariantSets();
 #else
 		return false;
-#endif // #if USE_USD_SDK
+#endif	  // #if USE_USD_SDK
 	}
 
 	FUsdPayloads FUsdPrim::GetPayloads() const
@@ -550,7 +542,7 @@ namespace UE
 		return FUsdPayloads(Impl->PxrUsdPrim.Get());
 #else
 		return FUsdPayloads{};
-#endif // #if USE_USD_SDK
+#endif	  // #if USE_USD_SDK
 	}
 
 	bool FUsdPrim::HasAuthoredPayloads() const
@@ -559,7 +551,7 @@ namespace UE
 		return Impl->PxrUsdPrim.Get().HasAuthoredPayloads();
 #else
 		return false;
-#endif // #if USE_USD_SDK
+#endif	  // #if USE_USD_SDK
 	}
 
 	// Deprecated in 5.3
@@ -574,38 +566,41 @@ namespace UE
 		return Impl->PxrUsdPrim.Get().IsLoaded();
 #else
 		return false;
-#endif // #if USE_USD_SDK
+#endif	  // #if USE_USD_SDK
 	}
 
-	void FUsdPrim::Load( EUsdLoadPolicy Policy )
+	void FUsdPrim::Load(EUsdLoadPolicy Policy)
 	{
 #if USE_USD_SDK
-		static_assert( ( int32 ) EUsdLoadPolicy::UsdLoadWithDescendants == ( int32 ) pxr::UsdLoadWithDescendants, "EUsdLoadPolicy enum doesn't match USD!" );
-		static_assert( ( int32 ) EUsdLoadPolicy::UsdLoadWithoutDescendants == ( int32 ) pxr::UsdLoadWithoutDescendants, "EUsdLoadPolicy enum doesn't match USD!" );
+		static_assert((int32)EUsdLoadPolicy::UsdLoadWithDescendants == (int32)pxr::UsdLoadWithDescendants, "EUsdLoadPolicy enum doesn't match USD!");
+		static_assert(
+			(int32)EUsdLoadPolicy::UsdLoadWithoutDescendants == (int32)pxr::UsdLoadWithoutDescendants,
+			"EUsdLoadPolicy enum doesn't match USD!"
+		);
 
-		Impl->PxrUsdPrim.Get().Load( static_cast< pxr::UsdLoadPolicy > ( Policy ) );
-#endif // #if USE_USD_SDK
+		Impl->PxrUsdPrim.Get().Load(static_cast<pxr::UsdLoadPolicy>(Policy));
+#endif	  // #if USE_USD_SDK
 	}
 
 	void FUsdPrim::Unload()
 	{
 #if USE_USD_SDK
 		Impl->PxrUsdPrim.Get().Unload();
-#endif // #if USE_USD_SDK
+#endif	  // #if USE_USD_SDK
 	}
 
-	TArray< FUsdAttribute > FUsdPrim::GetAttributes() const
+	TArray<FUsdAttribute> FUsdPrim::GetAttributes() const
 	{
-		TArray< FUsdAttribute > Attributes;
+		TArray<FUsdAttribute> Attributes;
 
 #if USE_USD_SDK
 		FScopedUsdAllocs UsdAllocs;
 
-		for ( const pxr::UsdAttribute& Attribute : Impl->PxrUsdPrim.Get().GetAttributes() )
+		for (const pxr::UsdAttribute& Attribute : Impl->PxrUsdPrim.Get().GetAttributes())
 		{
-			Attributes.Emplace( Attribute );
+			Attributes.Emplace(Attribute);
 		}
-#endif // #if USE_USD_SDK
+#endif	  // #if USE_USD_SDK
 
 		return Attributes;
 	}
@@ -613,10 +608,10 @@ namespace UE
 	FUsdReferences FUsdPrim::GetReferences() const
 	{
 #if USE_USD_SDK
-		return FUsdReferences( Impl->PxrUsdPrim.Get() );
+		return FUsdReferences(Impl->PxrUsdPrim.Get());
 #else
 		return FUsdReferences{};
-#endif // #if USE_USD_SDK
+#endif	  // #if USE_USD_SDK
 	}
 
 	bool FUsdPrim::HasAuthoredReferences() const
@@ -625,7 +620,7 @@ namespace UE
 		return Impl->PxrUsdPrim.Get().HasAuthoredReferences();
 #else
 		return false;
-#endif // #if USE_USD_SDK
+#endif	  // #if USE_USD_SDK
 	}
 
 	FUsdAttribute FUsdPrim::GetAttribute(const TCHAR* AttrName) const
@@ -633,45 +628,43 @@ namespace UE
 #if USE_USD_SDK
 		FScopedUsdAllocs UsdAllocs;
 
-		return FUsdAttribute( Impl->PxrUsdPrim.Get().GetAttribute( pxr::TfToken( TCHAR_TO_ANSI( AttrName ) ) ) );
+		return FUsdAttribute(Impl->PxrUsdPrim.Get().GetAttribute(pxr::TfToken(TCHAR_TO_ANSI(AttrName))));
 #else
 		return FUsdAttribute{};
-#endif // #if USE_USD_SDK
+#endif	  // #if USE_USD_SDK
 	}
 
-	bool FUsdPrim::HasAttribute( const TCHAR* AttrName ) const
+	bool FUsdPrim::HasAttribute(const TCHAR* AttrName) const
 	{
 #if USE_USD_SDK
-		return Impl->PxrUsdPrim.Get().HasAttribute( pxr::TfToken( TCHAR_TO_ANSI( AttrName ) ) );
+		return Impl->PxrUsdPrim.Get().HasAttribute(pxr::TfToken(TCHAR_TO_ANSI(AttrName)));
 #else
 		return false;
-#endif // #if USE_USD_SDK
+#endif	  // #if USE_USD_SDK
 	}
 
-	FUsdAttribute FUsdPrim::CreateAttribute( const TCHAR* AttrName, FName TypeName ) const
+	FUsdAttribute FUsdPrim::CreateAttribute(const TCHAR* AttrName, FName TypeName) const
 	{
 #if USE_USD_SDK
 		FScopedUsdAllocs UsdAllocs;
 
-		return FUsdAttribute(
-			Impl->PxrUsdPrim.Get().CreateAttribute(
-				pxr::TfToken( TCHAR_TO_ANSI( AttrName ) ),
-				pxr::SdfSchema::GetInstance().FindType( TCHAR_TO_ANSI( *TypeName.ToString() ) )
-			)
-		);
+		return FUsdAttribute(Impl->PxrUsdPrim.Get().CreateAttribute(
+			pxr::TfToken(TCHAR_TO_ANSI(AttrName)),
+			pxr::SdfSchema::GetInstance().FindType(TCHAR_TO_ANSI(*TypeName.ToString()))
+		));
 #else
 		return FUsdAttribute{};
-#endif // #if USE_USD_SDK
+#endif	  // #if USE_USD_SDK
 	}
 
-	bool FUsdPrim::RemoveProperty( FName PropName ) const
+	bool FUsdPrim::RemoveProperty(FName PropName) const
 	{
 #if USE_USD_SDK
 		FScopedUsdAllocs UsdAllocs;
-		return Impl->PxrUsdPrim.Get().RemoveProperty( pxr::TfToken( TCHAR_TO_ANSI( *PropName.ToString() ) ) );
+		return Impl->PxrUsdPrim.Get().RemoveProperty(pxr::TfToken(TCHAR_TO_ANSI(*PropName.ToString())));
 #else
 		return false;
-#endif // #if USE_USD_SDK
+#endif	  // #if USE_USD_SDK
 	}
 
 	bool FUsdPrim::IsInstanceable() const
diff --git a/Engine/Plugins/Importers/USDImporter/Source/UnrealUSDWrapper/Private/UsdWrappers/UsdReferences.cpp b/Engine/Plugins/Importers/USDImporter/Source/UnrealUSDWrapper/Private/UsdWrappers/UsdReferences.cpp
--- a/Engine/Plugins/Importers/USDImporter/Source/UnrealUSDWrapper/Private/UsdWrappers/UsdReferences.cpp
+++ b/Engine/Plugins/Importers/USDImporter/Source/UnrealUSDWrapper/Private/UsdWrappers/UsdReferences.cpp
@@ -7,13 +7,11 @@
 #include "UsdWrappers/UsdPrim.h"
 
 #if USE_USD_SDK
-
 #include "USDIncludesStart.h"
-	#include "pxr/usd/usd/prim.h"
-	#include "pxr/usd/usd/references.h"
+#include "pxr/usd/usd/prim.h"
+#include "pxr/usd/usd/references.h"
 #include "USDIncludesEnd.h"
-
-#endif // #if USE_USD_SDK
+#endif	  // #if USE_USD_SDK
 
 namespace UE
 {
@@ -48,9 +46,9 @@ namespace UE
 
 			// pxr::UsdReferences has no default constructor, so we store the prim itself.
 			TUsdStore<pxr::UsdPrim> PxrUsdPrim;
-#endif // #if USE_USD_SDK
+#endif	  // #if USE_USD_SDK
 		};
-	}
+	}	  // namespace Internal
 
 	FUsdReferences::FUsdReferences()
 	{
@@ -61,7 +59,7 @@ namespace UE
 	{
 #if USE_USD_SDK
 		Impl = MakeUnique<Internal::FUsdReferencesImpl>(Other.Impl->PxrUsdPrim.Get());
-#endif // #if USE_USD_SDK
+#endif	  // #if USE_USD_SDK
 	}
 
 	FUsdReferences::FUsdReferences(FUsdReferences&& Other) = default;
@@ -70,7 +68,7 @@ namespace UE
 	{
 #if USE_USD_SDK
 		Impl = MakeUnique<Internal::FUsdReferencesImpl>(Other.Impl->PxrUsdPrim.Get());
-#endif // #if USE_USD_SDK
+#endif	  // #if USE_USD_SDK
 
 		return *this;
 	}
@@ -93,7 +91,7 @@ namespace UE
 		return (bool)Impl->PxrUsdPrim.Get();
 #else
 		return false;
-#endif // #if USE_USD_SDK
+#endif	  // #if USE_USD_SDK
 	}
 
 #if USE_USD_SDK
@@ -134,15 +132,27 @@ namespace UE
 	{
 		Impl = MakeUnique<Internal::FUsdReferencesImpl>(InPrim);
 	}
-#endif // #if USE_USD_SDK
+#endif	  // #if USE_USD_SDK
 
 	bool FUsdReferences::AddReference(const FSdfReference& Reference, EUsdListPosition Position)
 	{
 #if USE_USD_SDK
-		static_assert((int32)EUsdListPosition::FrontOfPrependList == (int32)pxr::UsdListPositionFrontOfPrependList, "EUsdListPosition enum doesn't match USD!");
-		static_assert((int32)EUsdListPosition::BackOfPrependList == (int32)pxr::UsdListPositionBackOfPrependList, "EUsdListPosition enum doesn't match USD!");
-		static_assert((int32)EUsdListPosition::FrontOfAppendList == (int32)pxr::UsdListPositionFrontOfAppendList, "EUsdListPosition enum doesn't match USD!");
-		static_assert((int32)EUsdListPosition::BackOfAppendList == (int32)pxr::UsdListPositionBackOfAppendList, "EUsdListPosition enum doesn't match USD!");
+		static_assert(
+			(int32)EUsdListPosition::FrontOfPrependList == (int32)pxr::UsdListPositionFrontOfPrependList,
+			"EUsdListPosition enum doesn't match USD!"
+		);
+		static_assert(
+			(int32)EUsdListPosition::BackOfPrependList == (int32)pxr::UsdListPositionBackOfPrependList,
+			"EUsdListPosition enum doesn't match USD!"
+		);
+		static_assert(
+			(int32)EUsdListPosition::FrontOfAppendList == (int32)pxr::UsdListPositionFrontOfAppendList,
+			"EUsdListPosition enum doesn't match USD!"
+		);
+		static_assert(
+			(int32)EUsdListPosition::BackOfAppendList == (int32)pxr::UsdListPositionBackOfAppendList,
+			"EUsdListPosition enum doesn't match USD!"
+		);
 
 		if (Impl->PxrUsdPrim.Get())
 		{
@@ -154,11 +164,9 @@ namespace UE
 
 			const pxr::SdfReference UsdReference(AssetPath, UsdPrimPath, UsdLayerOffset);
 
-			return Impl->PxrUsdPrim.Get().GetReferences().AddReference(
-				UsdReference,
-				static_cast<pxr::UsdListPosition>(Position));
+			return Impl->PxrUsdPrim.Get().GetReferences().AddReference(UsdReference, static_cast<pxr::UsdListPosition>(Position));
 		}
-#endif // #if USE_USD_SDK
+#endif	  // #if USE_USD_SDK
 
 		return false;
 	}
@@ -167,7 +175,8 @@ namespace UE
 		const FString& Identifier,
 		const FSdfPath& PrimPath,
 		const FSdfLayerOffset& LayerOffset,
-		EUsdListPosition Position)
+		EUsdListPosition Position
+	)
 	{
 #if USE_USD_SDK
 		if (Impl->PxrUsdPrim.Get())
@@ -179,17 +188,18 @@ namespace UE
 			const pxr::SdfLayerOffset UsdLayerOffset(LayerOffset.Offset, LayerOffset.Scale);
 
 			return Impl->PxrUsdPrim.Get().GetReferences().AddReference(
-				UsdIdentifier, UsdPrimPath, UsdLayerOffset, static_cast<pxr::UsdListPosition>(Position));
+				UsdIdentifier,
+				UsdPrimPath,
+				UsdLayerOffset,
+				static_cast<pxr::UsdListPosition>(Position)
+			);
 		}
-#endif // #if USE_USD_SDK
+#endif	  // #if USE_USD_SDK
 
 		return false;
 	}
 
-	bool FUsdReferences::AddReference(
-		const FString& Identifier,
-		const FSdfLayerOffset& LayerOffset,
-		EUsdListPosition Position)
+	bool FUsdReferences::AddReference(const FString& Identifier, const FSdfLayerOffset& LayerOffset, EUsdListPosition Position)
 	{
 #if USE_USD_SDK
 		if (Impl->PxrUsdPrim.Get())
@@ -199,18 +209,14 @@ namespace UE
 			const std::string UsdIdentifier(TCHAR_TO_ANSI(*Identifier));
 			const pxr::SdfLayerOffset UsdLayerOffset(LayerOffset.Offset, LayerOffset.Scale);
 
-			return Impl->PxrUsdPrim.Get().GetReferences().AddReference(
-				UsdIdentifier, UsdLayerOffset, static_cast<pxr::UsdListPosition>(Position));
+			return Impl->PxrUsdPrim.Get().GetReferences().AddReference(UsdIdentifier, UsdLayerOffset, static_cast<pxr::UsdListPosition>(Position));
 		}
-#endif // #if USE_USD_SDK
+#endif	  // #if USE_USD_SDK
 
 		return false;
 	}
 
-	bool FUsdReferences::AddInternalReference(
-		const FSdfPath& PrimPath,
-		const FSdfLayerOffset& LayerOffset,
-		EUsdListPosition Position)
+	bool FUsdReferences::AddInternalReference(const FSdfPath& PrimPath, const FSdfLayerOffset& LayerOffset, EUsdListPosition Position)
 	{
 #if USE_USD_SDK
 		if (Impl->PxrUsdPrim.Get())
@@ -221,9 +227,12 @@ namespace UE
 			const pxr::SdfLayerOffset UsdLayerOffset(LayerOffset.Offset, LayerOffset.Scale);
 
 			return Impl->PxrUsdPrim.Get().GetReferences().AddInternalReference(
-				UsdPrimPath, UsdLayerOffset, static_cast<pxr::UsdListPosition>(Position));
+				UsdPrimPath,
+				UsdLayerOffset,
+				static_cast<pxr::UsdListPosition>(Position)
+			);
 		}
-#endif // #if USE_USD_SDK
+#endif	  // #if USE_USD_SDK
 
 		return false;
 	}
@@ -243,7 +252,7 @@ namespace UE
 
 			return Impl->PxrUsdPrim.Get().GetReferences().RemoveReference(UsdReference);
 		}
-#endif // #if USE_USD_SDK
+#endif	  // #if USE_USD_SDK
 
 		return false;
 	}
@@ -257,7 +266,7 @@ namespace UE
 
 			return Impl->PxrUsdPrim.Get().GetReferences().ClearReferences();
 		}
-#endif // #if USE_USD_SDK
+#endif	  // #if USE_USD_SDK
 
 		return false;
 	}
@@ -283,7 +292,7 @@ namespace UE
 
 			return Impl->PxrUsdPrim.Get().GetReferences().SetReferences(UsdReferences);
 		}
-#endif // #if USE_USD_SDK
+#endif	  // #if USE_USD_SDK
 
 		return false;
 	}
@@ -295,8 +304,8 @@ namespace UE
 		{
 			return FUsdPrim(Impl->PxrUsdPrim.Get());
 		}
-#endif // #if USE_USD_SDK
+#endif	  // #if USE_USD_SDK
 
 		return FUsdPrim{};
 	}
-}
+}	 // namespace UE
